Control Rig: Further improvements to support LWC

diff --git a/Engine/Plugins/Animation/ControlRig/Source/ControlRig/Private/ControlRig.cpp b/Engine/Plugins/Animation/ControlRig/Source/ControlRig/Private/ControlRig.cpp
--- a/Engine/Plugins/Animation/ControlRig/Source/ControlRig/Private/ControlRig.cpp
+++ b/Engine/Plugins/Animation/ControlRig/Source/ControlRig/Private/ControlRig.cpp
@@ -939,8 +939,8 @@ void UControlRig::Execute(const EControlRigState InState, const FName& InEventNa
 					case ERigControlType::Vector2D:
 					{
 						Instruction.PrimitiveType = EControlRigDrawSettings::LineStrip;
-						FVector2D MinPos = Settings.MinimumValue.Get<FVector2D>();
-						FVector2D MaxPos = Settings.MaximumValue.Get<FVector2D>();
+						FVector3f MinPos = Settings.MinimumValue.Get<FVector3f>();
+						FVector3f MaxPos = Settings.MaximumValue.Get<FVector3f>();
 
 						switch (Settings.PrimaryAxis)
 						{
@@ -980,33 +980,33 @@ void UControlRig::Execute(const EControlRigState InState, const FName& InEventNa
 					case ERigControlType::TransformNoScale:
 					case ERigControlType::EulerTransform:
 					{
-						FVector MinPos, MaxPos;
+						FVector3f MinPos, MaxPos;
 
 						switch (Settings.ControlType)
 						{
 							case ERigControlType::Position:
 							case ERigControlType::Scale:
 							{
-								MinPos = Settings.MinimumValue.Get<FVector>();
-								MaxPos = Settings.MaximumValue.Get<FVector>();
+								MinPos = Settings.MinimumValue.Get<FVector3f>();
+								MaxPos = Settings.MaximumValue.Get<FVector3f>();
 								break;
 							}
 							case ERigControlType::Transform:
 							{
-								MinPos = Settings.MinimumValue.Get<FTransform>().GetLocation();
-								MaxPos = Settings.MaximumValue.Get<FTransform>().GetLocation();
+								MinPos = Settings.MinimumValue.Get<FRigControlValue::FTransform_Float>().GetTranslation();
+								MaxPos = Settings.MaximumValue.Get<FRigControlValue::FTransform_Float>().GetTranslation();
 								break;
 							}
 							case ERigControlType::TransformNoScale:
 							{
-								MinPos = Settings.MinimumValue.Get<FTransformNoScale>().Location;
-								MaxPos = Settings.MaximumValue.Get<FTransformNoScale>().Location;
+								MinPos = Settings.MinimumValue.Get<FRigControlValue::FTransformNoScale_Float>().GetTranslation();
+								MaxPos = Settings.MaximumValue.Get<FRigControlValue::FTransformNoScale_Float>().GetTranslation();
 								break;
 							}
 							case ERigControlType::EulerTransform:
 							{
-								MinPos = Settings.MinimumValue.Get<FEulerTransform>().Location;
-								MaxPos = Settings.MaximumValue.Get<FEulerTransform>().Location;
+								MinPos = Settings.MinimumValue.Get<FRigControlValue::FEulerTransform_Float>().GetTranslation();
+								MaxPos = Settings.MaximumValue.Get<FRigControlValue::FEulerTransform_Float>().GetTranslation();
 								break;
 							}
 						}
diff --git a/Engine/Plugins/Animation/ControlRig/Source/ControlRig/Private/ControlRigComponent.cpp b/Engine/Plugins/Animation/ControlRig/Source/ControlRig/Private/ControlRigComponent.cpp
--- a/Engine/Plugins/Animation/ControlRig/Source/ControlRig/Private/ControlRigComponent.cpp
+++ b/Engine/Plugins/Animation/ControlRig/Source/ControlRig/Private/ControlRigComponent.cpp
@@ -824,7 +824,8 @@ FVector2D UControlRigComponent::GetControlVector2D(FName ControlName)
 		{
 			if (ControlElement->Settings.ControlType == ERigControlType::Vector2D)
 			{
-				return CR->GetHierarchy()->GetControlValue(ControlElement, ERigControlValueType::Current).Get<FVector2D>();
+				const FVector3f Value = CR->GetHierarchy()->GetControlValue(ControlElement, ERigControlValueType::Current).Get<FVector3f>();
+				return FVector2D(Value.X, Value.Y);
 			}
 		}
 	}
diff --git a/Engine/Plugins/Animation/ControlRig/Source/ControlRig/Private/Rigs/FKControlRig.cpp b/Engine/Plugins/Animation/ControlRig/Source/ControlRig/Private/Rigs/FKControlRig.cpp
--- a/Engine/Plugins/Animation/ControlRig/Source/ControlRig/Private/Rigs/FKControlRig.cpp
+++ b/Engine/Plugins/Animation/ControlRig/Source/ControlRig/Private/Rigs/FKControlRig.cpp
@@ -411,8 +411,9 @@ void UFKControlRig::ToggleApplyMode()
         {
             if (ControlElement->Settings.ControlType == ERigControlType::EulerTransform)
 			{
-				const FEulerTransform InitValue = GetHierarchy()->GetControlValue(ControlElement, ERigControlValueType::Initial).Get<FEulerTransform>();
-				SetControlValue<FEulerTransform>(ControlElement->GetName(), InitValue, true, Context, bSetupUndo);
+				const FRigControlValue::FEulerTransform_Float InitValue =
+					GetHierarchy()->GetControlValue(ControlElement, ERigControlValueType::Initial).Get<FRigControlValue::FEulerTransform_Float>();
+				SetControlValue<FRigControlValue::FEulerTransform_Float>(ControlElement->GetName(), InitValue, true, Context, bSetupUndo);
 			}
 			else if (ControlElement->Settings.ControlType == ERigControlType::Float)
 			{
diff --git a/Engine/Plugins/Animation/ControlRig/Source/ControlRig/Private/Sequencer/MovieSceneControlRigParameterSection.cpp b/Engine/Plugins/Animation/ControlRig/Source/ControlRig/Private/Sequencer/MovieSceneControlRigParameterSection.cpp
--- a/Engine/Plugins/Animation/ControlRig/Source/ControlRig/Private/Sequencer/MovieSceneControlRigParameterSection.cpp
+++ b/Engine/Plugins/Animation/ControlRig/Source/ControlRig/Private/Sequencer/MovieSceneControlRigParameterSection.cpp
@@ -169,17 +169,19 @@ struct FParameterVectorChannelEditorData
 		
 				if (NumChannels == 2)
 				{
-					const FVector2D Vector = ControlRig->GetHierarchy()->GetControlValue(ControlElement, ERigControlValueType::Current).Get<FVector2D>();
+					const FVector3f Vector = ControlRig->GetHierarchy()->GetControlValue(ControlElement, ERigControlValueType::Current).Get<FVector3f>();
 					return FVector4(Vector.X, Vector.Y, 0.f, 0.f);
 				}
 				else if (NumChannels == 3)
 				{
-					const FVector Vector = ControlRig->GetHierarchy()->GetControlValue(ControlElement, ERigControlValueType::Current).Get<FVector>();
+					const FVector3f Vector = ControlRig->GetHierarchy()->GetControlValue(ControlElement, ERigControlValueType::Current).Get<FVector3f>();
 					return FVector4(Vector.X, Vector.Y, Vector.Z, 0.f);
 				}
 				else
 				{
-					return ControlRig->GetHierarchy()->GetControlValue(ControlElement, ERigControlValueType::Current).Get<FVector4>();
+					const FRigControlValue::FTransform_Float Storage = ControlRig->GetHierarchy()
+						->GetControlValue(ControlElement, ERigControlValueType::Current).Get<FRigControlValue::FTransform_Float>();
+					return FVector4(Storage.TranslationX, Storage.TranslationY, Storage.TranslationZ, Storage.TranslationW);
 				}
 			}
 		}
@@ -532,19 +534,24 @@ struct FParameterTransformChannelEditorData
 			{
 				if (ControlElement->Settings.ControlType == ERigControlType::Transform)
 				{
-					const FTransform Transform = ControlRig->GetHierarchy()->GetControlValue(ControlElement, ERigControlValueType::Current).Get<FTransform>();
-					return Transform.GetTranslation();
+					const FRigControlValue::FTransform_Float Transform = 
+						ControlRig->GetHierarchy()
+						->GetControlValue(ControlElement, ERigControlValueType::Current).Get<FRigControlValue::FTransform_Float>();
+					return FVector(Transform.GetTranslation());
 				}
 				else if  (ControlElement->Settings.ControlType == ERigControlType::TransformNoScale)
 				{
-					const FTransformNoScale NoScale = ControlRig->GetHierarchy()->GetControlValue(ControlElement, ERigControlValueType::Current).Get<FTransformNoScale>();
-					FTransform Transform = NoScale;
-					return Transform.GetTranslation();
+					const FRigControlValue::FTransformNoScale_Float Transform = 
+						ControlRig->GetHierarchy()
+						->GetControlValue(ControlElement, ERigControlValueType::Current).Get<FRigControlValue::FTransformNoScale_Float>();
+					return FVector(Transform.GetTranslation());
 				}
 				else if (ControlElement->Settings.ControlType == ERigControlType::EulerTransform)
 				{
-					const FEulerTransform Euler = ControlRig->GetHierarchy()->GetControlValue(ControlElement, ERigControlValueType::Current).Get<FEulerTransform>();
-					return Euler.Location;
+					const FRigControlValue::FEulerTransform_Float Euler = 
+						ControlRig->GetHierarchy()
+						->GetControlValue(ControlElement, ERigControlValueType::Current).Get<FRigControlValue::FEulerTransform_Float>();
+					return FVector(Euler.GetTranslation());
 				}
 			}
 		}
@@ -561,19 +568,24 @@ struct FParameterTransformChannelEditorData
 			{
 				if (ControlElement->Settings.ControlType == ERigControlType::Transform)
 				{
-					FTransform Transform = ControlRig->GetHierarchy()->GetControlValue(ControlElement, ERigControlValueType::Current).Get<FTransform>();
+					const FRigControlValue::FTransform_Float Transform = 
+						ControlRig->GetHierarchy()
+						->GetControlValue(ControlElement, ERigControlValueType::Current).Get<FRigControlValue::FTransform_Float>();
 					return Transform.GetRotation().Rotator();
 				}
 				else if (ControlElement->Settings.ControlType == ERigControlType::TransformNoScale)
 				{
-					FTransformNoScale NoScale = ControlRig->GetHierarchy()->GetControlValue(ControlElement, ERigControlValueType::Current).Get<FTransformNoScale>();
-					FTransform Transform = NoScale;
+					const FRigControlValue::FTransformNoScale_Float Transform = 
+						ControlRig->GetHierarchy()
+						->GetControlValue(ControlElement, ERigControlValueType::Current).Get<FRigControlValue::FTransformNoScale_Float>();
 					return Transform.GetRotation().Rotator();
 				}
 				else if (ControlElement->Settings.ControlType == ERigControlType::EulerTransform)
 				{
-					FEulerTransform Euler = ControlRig->GetHierarchy()->GetControlValue(ControlElement, ERigControlValueType::Current).Get<FEulerTransform>();
-					return Euler.Rotation;
+					const FRigControlValue::FEulerTransform_Float Transform = 
+						ControlRig->GetHierarchy()
+						->GetControlValue(ControlElement, ERigControlValueType::Current).Get<FRigControlValue::FEulerTransform_Float>();
+					return Transform.GetRotator();
 			}
 		}
 		}
@@ -589,13 +601,17 @@ struct FParameterTransformChannelEditorData
 			{
 				if (ControlElement->Settings.ControlType == ERigControlType::Transform)
 				{
-					FTransform Transform = ControlRig->GetHierarchy()->GetControlValue(ControlElement, ERigControlValueType::Current).Get<FTransform>();
-					return Transform.GetScale3D();
+					const FRigControlValue::FTransform_Float Transform = 
+						ControlRig->GetHierarchy()
+						->GetControlValue(ControlElement, ERigControlValueType::Current).Get<FRigControlValue::FTransform_Float>();
+					return FVector(Transform.GetScale3D());
 				}
 				else if (ControlElement->Settings.ControlType == ERigControlType::EulerTransform)
 				{
-					FEulerTransform Euler = ControlRig->GetHierarchy()->GetControlValue(ControlElement, ERigControlValueType::Current).Get<FEulerTransform>();
-					return Euler.Scale;
+					const FRigControlValue::FEulerTransform_Float Transform = 
+						ControlRig->GetHierarchy()
+						->GetControlValue(ControlElement, ERigControlValueType::Current).Get<FRigControlValue::FEulerTransform_Float>();
+					return FVector(Transform.GetScale3D());
 				}
 			}
 		}
@@ -1786,7 +1802,8 @@ void UMovieSceneControlRigParameterSection::RecreateWithThisControlRig(UControlR
 			if (bSetDefault)
 			{
 				//or use IntialValue?
-				DefaultValue = ControlRig->GetHierarchy()->GetControlValue(ControlElement, ERigControlValueType::Current).Get<FVector2D>();
+				const FVector3f TempValue = ControlRig->GetHierarchy()->GetControlValue(ControlElement, ERigControlValueType::Current).Get<FVector3f>();
+				DefaultValue = FVector2D(TempValue.X, TempValue.Y);
 			}
 			AddVector2DParameter(ControlElement->GetName(), DefaultValue, false);
 			break;
@@ -1800,7 +1817,7 @@ void UMovieSceneControlRigParameterSection::RecreateWithThisControlRig(UControlR
 			if (bSetDefault)
 			{
 				//or use IntialValue?
-				DefaultValue = ControlRig->GetHierarchy()->GetControlValue(ControlElement, ERigControlValueType::Current).Get<FVector>();
+				DefaultValue = ControlRig->GetHierarchy()->GetControlValue(ControlElement, ERigControlValueType::Current).Get<FVector3f>();
 			}
 			AddVectorParameter(ControlElement->GetName(), DefaultValue, false);
 			//mz todo specify rotator special so we can do quat interps
@@ -1815,17 +1832,23 @@ void UMovieSceneControlRigParameterSection::RecreateWithThisControlRig(UControlR
 			{
 				if (ControlElement->Settings.ControlType == ERigControlType::Transform)
 				{
-					DefaultValue = ControlRig->GetHierarchy()->GetControlValue(ControlElement, ERigControlValueType::Current).Get<FTransform>();
+					DefaultValue = 
+						ControlRig->GetHierarchy()->
+						GetControlValue(ControlElement, ERigControlValueType::Current).Get<FRigControlValue::FTransform_Float>().ToTransform();
 				}
 				else if (ControlElement->Settings.ControlType == ERigControlType::EulerTransform)
 
 				{
-					FEulerTransform Euler = ControlRig->GetHierarchy()->GetControlValue(ControlElement, ERigControlValueType::Current).Get<FEulerTransform>();
+					FEulerTransform Euler = 
+						ControlRig->GetHierarchy()->
+						GetControlValue(ControlElement, ERigControlValueType::Current).Get<FRigControlValue::FEulerTransform_Float>().ToTransform();
 					DefaultValue = Euler;
 				}
 				else
 				{
-					FTransformNoScale NoScale = ControlRig->GetHierarchy()->GetControlValue(ControlElement, ERigControlValueType::Current).Get<FTransformNoScale>();
+					FTransformNoScale NoScale = 
+						ControlRig->GetHierarchy()->
+						GetControlValue(ControlElement, ERigControlValueType::Current).Get<FRigControlValue::FTransformNoScale_Float>().ToTransform();
 					DefaultValue = NoScale;
 				}
 			}
@@ -1927,7 +1950,7 @@ void UMovieSceneControlRigParameterSection::RecordControlRigKey(FFrameNumber Fra
 				}
 				case ERigControlType::Vector2D:
 				{
-					FVector2D Val = ControlRig->GetHierarchy()->GetControlValue(ControlElement, ERigControlValueType::Current).Get<FVector2D>();
+					FVector3f Val = ControlRig->GetHierarchy()->GetControlValue(ControlElement, ERigControlValueType::Current).Get<FVector3f>();
 					if (bSetDefault)
 					{
 						FloatChannels[ChannelIndex]->SetDefault(Val.X);
@@ -1950,7 +1973,7 @@ void UMovieSceneControlRigParameterSection::RecordControlRigKey(FFrameNumber Fra
 				case ERigControlType::Scale:
 				case ERigControlType::Rotator:
 				{
-					FVector Val = ControlRig->GetHierarchy()->GetControlValue(ControlElement, ERigControlValueType::Current).Get<FVector>();
+					FVector3f Val = ControlRig->GetHierarchy()->GetControlValue(ControlElement, ERigControlValueType::Current).Get<FVector3f>();
 					if (ControlElement->Settings.ControlType == ERigControlType::Rotator &&
 						FloatChannels[ChannelIndex]->GetNumKeys() > 0)
 					{
@@ -1989,17 +2012,23 @@ void UMovieSceneControlRigParameterSection::RecordControlRigKey(FFrameNumber Fra
 					FTransform Val = FTransform::Identity;
 					if (ControlElement->Settings.ControlType == ERigControlType::TransformNoScale)
 					{
-						FTransformNoScale NoScale = ControlRig->GetHierarchy()->GetControlValue(ControlElement, ERigControlValueType::Current).Get<FTransformNoScale>();
+						FTransformNoScale NoScale = 
+							ControlRig->GetHierarchy()
+							->GetControlValue(ControlElement, ERigControlValueType::Current).Get<FRigControlValue::FTransformNoScale_Float>().ToTransform();
 						Val = NoScale;
 					}
 					else if (ControlElement->Settings.ControlType == ERigControlType::EulerTransform)
 					{
-						FEulerTransform Euler = ControlRig->GetHierarchy()->GetControlValue(ControlElement, ERigControlValueType::Current).Get < FEulerTransform >();
+						FEulerTransform Euler = 
+							ControlRig->GetHierarchy()
+							->GetControlValue(ControlElement, ERigControlValueType::Current).Get<FRigControlValue::FEulerTransform_Float>().ToTransform();
 						Val = Euler.ToFTransform();
 					}
 					else
 					{
-						Val = ControlRig->GetHierarchy()->GetControlValue(ControlElement, ERigControlValueType::Current).Get<FTransform>();
+						Val = 
+							ControlRig->GetHierarchy()
+							->GetControlValue(ControlElement, ERigControlValueType::Current).Get<FRigControlValue::FTransform_Float>().ToTransform();
 					}
 					FVector CurrentVector = Val.GetTranslation();
 					if (bSetDefault)
diff --git a/Engine/Plugins/Animation/ControlRig/Source/ControlRig/Private/Sequencer/MovieSceneControlRigParameterTemplate.cpp b/Engine/Plugins/Animation/ControlRig/Source/ControlRig/Private/Sequencer/MovieSceneControlRigParameterTemplate.cpp
--- a/Engine/Plugins/Animation/ControlRig/Source/ControlRig/Private/Sequencer/MovieSceneControlRigParameterTemplate.cpp
+++ b/Engine/Plugins/Animation/ControlRig/Source/ControlRig/Private/Sequencer/MovieSceneControlRigParameterTemplate.cpp
@@ -773,35 +773,41 @@ struct FControlRigParameterPreAnimatedTokenProducer : IMovieScenePreAnimatedToke
 					}
 					case ERigControlType::Vector2D:
 					{
-						const FVector2D Val = ControlRig->GetHierarchy()->GetControlValue(ControlElement, ERigControlValueType::Current).Get<FVector2D>();
-						Token.Vector2DValues.Add(TNameAndValue<FVector2D>{ ControlElement->GetName(), Val });
+						const FVector3f Val = ControlRig->GetHierarchy()->GetControlValue(ControlElement, ERigControlValueType::Current).Get<FVector3f>();
+						Token.Vector2DValues.Add(TNameAndValue<FVector2D>{ ControlElement->GetName(), FVector2D(Val.X, Val.Y) });
 						break;
 					}
 					case ERigControlType::Position:
 					case ERigControlType::Scale:
 					case ERigControlType::Rotator:
 					{
-						const FVector Val = ControlRig->GetHierarchy()->GetControlValue(ControlElement, ERigControlValueType::Current).Get<FVector>();
+						const FVector3f Val = ControlRig->GetHierarchy()->GetControlValue(ControlElement, ERigControlValueType::Current).Get<FVector3f>();
 						Token.VectorValues.Add(TNameAndValue<FVector>{ ControlElement->GetName(), Val });
 						//mz todo specify rotator special so we can do quat interps
 						break;
 					}
 					case ERigControlType::Transform:
 					{
-						const FTransform Val = ControlRig->GetHierarchy()->GetControlValue(ControlElement, ERigControlValueType::Current).Get<FTransform>();
+						const FTransform Val = 
+							ControlRig->GetHierarchy()
+							->GetControlValue(ControlElement, ERigControlValueType::Current).Get<FRigControlValue::FTransform_Float>().ToTransform();
 						Token.TransformValues.Add(TNameAndValue<FTransform>{ ControlElement->GetName(), Val });
 						break;
 					}
 					case ERigControlType::TransformNoScale:
 					{
-						const FTransformNoScale NoScale = ControlRig->GetHierarchy()->GetControlValue(ControlElement, ERigControlValueType::Current).Get<FTransformNoScale>();
+						const FTransformNoScale NoScale = 
+							ControlRig->GetHierarchy()
+							->GetControlValue(ControlElement, ERigControlValueType::Current).Get<FRigControlValue::FTransformNoScale_Float>().ToTransform();
 						FTransform Val = NoScale;
 						Token.TransformValues.Add(TNameAndValue<FTransform>{ ControlElement->GetName(), Val });
 						break;
 					}
 					case ERigControlType::EulerTransform:
 					{
-						const FEulerTransform Euler = ControlRig->GetHierarchy()->GetControlValue(ControlElement, ERigControlValueType::Current).Get<FEulerTransform>();
+						const FEulerTransform Euler = 
+							ControlRig->GetHierarchy()
+							->GetControlValue(ControlElement, ERigControlValueType::Current).Get<FRigControlValue::FEulerTransform_Float>().ToTransform();
 						FTransform Val = Euler.ToFTransform();
 						Token.TransformValues.Add(TNameAndValue<FTransform>{ ControlElement->GetName(), Val });
 						break;
@@ -1096,8 +1102,8 @@ struct TControlRigParameterActuatorVector2D : TMovieSceneBlendingActuator<FContr
 			FRigControlElement* ControlElement = Section->GetControlRig()->FindControl(ParameterName);
 			if (ControlElement && (ControlElement->Settings.ControlType == ERigControlType::Vector2D))
 			{
-				FVector2D Val = Section->GetControlRig()->GetHierarchy()->GetControlValue(ControlElement, ERigControlValueType::Current).Get<FVector2D>();
-				return FControlRigTrackTokenVector2D(Val);
+				FVector3f Val = Section->GetControlRig()->GetHierarchy()->GetControlValue(ControlElement, ERigControlValueType::Current).Get<FVector3f>();
+				return FControlRigTrackTokenVector2D(FVector2D(Val.X, Val.Y));
 			}
 		}
 		return FControlRigTrackTokenVector2D();
@@ -1119,7 +1125,8 @@ struct TControlRigParameterActuatorVector2D : TMovieSceneBlendingActuator<FContr
 				FRigControlElement* ControlElement = Section->GetControlRig()->FindControl(ParameterName);
 				if (ControlElement && (ControlElement->Settings.ControlType == ERigControlType::Vector2D))
 				{
-					Section->GetControlRig()->SetControlValue<FVector2D>(ParameterName, InFinalValue.Value, true, EControlRigSetKey::Never,bSetupUndo);
+					const FVector3f Value(InFinalValue.Value.X, InFinalValue.Value.Y, 0.f);
+					Section->GetControlRig()->SetControlValue<FVector3f>(ParameterName, Value, true, EControlRigSetKey::Never,bSetupUndo);
 				}
 			}
 
@@ -1159,7 +1166,7 @@ struct TControlRigParameterActuatorVector : TMovieSceneBlendingActuator<FControl
 			FRigControlElement* ControlElement = Section->GetControlRig()->FindControl(ParameterName);
 			if (ControlElement && (ControlElement->Settings.ControlType == ERigControlType::Position || ControlElement->Settings.ControlType == ERigControlType::Scale || ControlElement->Settings.ControlType == ERigControlType::Rotator))
 			{
-				FVector Val = Section->GetControlRig()->GetHierarchy()->GetControlValue(ControlElement, ERigControlValueType::Current).Get<FVector>();
+				FVector3f Val = Section->GetControlRig()->GetHierarchy()->GetControlValue(ControlElement, ERigControlValueType::Current).Get<FVector3f>();
 				return FControlRigTrackTokenVector(Val);
 			}
 		}
@@ -1181,7 +1188,7 @@ struct TControlRigParameterActuatorVector : TMovieSceneBlendingActuator<FControl
 				FRigControlElement* ControlElement = Section->GetControlRig()->FindControl(ParameterName);
 				if (ControlElement && (ControlElement->Settings.ControlType == ERigControlType::Position || ControlElement->Settings.ControlType == ERigControlType::Scale || ControlElement->Settings.ControlType == ERigControlType::Rotator))
 				{
-						Section->GetControlRig()->SetControlValue<FVector>(ParameterName, InFinalValue.Value, true, EControlRigSetKey::Never,bSetupUndo);
+						Section->GetControlRig()->SetControlValue<FVector3f>(ParameterName, InFinalValue.Value, true, EControlRigSetKey::Never,bSetupUndo);
 				}
 			}
 
@@ -1221,18 +1228,24 @@ struct TControlRigParameterActuatorTransform : TMovieSceneBlendingActuator<FCont
 			FRigControlElement* ControlElement = Section->GetControlRig()->FindControl(ParameterName);
 			if (ControlElement && ControlElement->Settings.ControlType == ERigControlType::Transform)
 			{
-				FTransform Val = Section->GetControlRig()->GetHierarchy()->GetControlValue(ControlElement, ERigControlValueType::Current).Get<FTransform>();
+				FTransform Val = 
+					Section->GetControlRig()->GetHierarchy()
+					->GetControlValue(ControlElement, ERigControlValueType::Current).Get<FRigControlValue::FTransform_Float>().ToTransform();
 				return FControlRigTrackTokenTransform(Val);
 			}
 			else if(ControlElement && ControlElement->Settings.ControlType == ERigControlType::TransformNoScale)
 			{
-				FTransformNoScale ValNoScale = Section->GetControlRig()->GetHierarchy()->GetControlValue(ControlElement, ERigControlValueType::Current).Get<FTransformNoScale>();
+				FTransformNoScale ValNoScale = 
+					Section->GetControlRig()->GetHierarchy()
+					->GetControlValue(ControlElement, ERigControlValueType::Current).Get<FRigControlValue::FTransformNoScale_Float>().ToTransform();
 				FTransform Val = ValNoScale;
 				return FControlRigTrackTokenTransform(Val);
 			}
 			else if (ControlElement && ControlElement->Settings.ControlType == ERigControlType::EulerTransform)
 			{
-				FEulerTransform Euler = Section->GetControlRig()->GetHierarchy()->GetControlValue(ControlElement, ERigControlValueType::Current).Get<FEulerTransform>();
+				FEulerTransform Euler = 
+					Section->GetControlRig()->GetHierarchy()
+					->GetControlValue(ControlElement, ERigControlValueType::Current).Get<FRigControlValue::FEulerTransform_Float>().ToTransform();
 				FTransform Val = Euler.ToFTransform();
 				return FControlRigTrackTokenTransform(Val);
 		}
@@ -1263,17 +1276,17 @@ struct TControlRigParameterActuatorTransform : TMovieSceneBlendingActuator<FCont
 				FRigControlElement* ControlElement = Section->GetControlRig()->FindControl(ParameterName);
 				if (ControlElement && ControlElement->Settings.ControlType == ERigControlType::Transform)
 				{
-						Section->GetControlRig()->SetControlValue<FTransform>(ParameterName, InFinalValue.Value,true, EControlRigSetKey::Never,bSetupUndo);
+						Section->GetControlRig()->SetControlValue<FRigControlValue::FTransform_Float>(ParameterName, InFinalValue.Value,true, EControlRigSetKey::Never,bSetupUndo);
 				}
 				else if (ControlElement && ControlElement->Settings.ControlType == ERigControlType::TransformNoScale)
 				{
 					FTransformNoScale NoScale = InFinalValue.Value;
-						Section->GetControlRig()->SetControlValue<FTransformNoScale>(ParameterName, NoScale, true, EControlRigSetKey::Never,bSetupUndo);
+						Section->GetControlRig()->SetControlValue<FRigControlValue::FTransformNoScale_Float>(ParameterName, NoScale, true, EControlRigSetKey::Never,bSetupUndo);
 				}
 				else if (ControlElement && ControlElement->Settings.ControlType == ERigControlType::EulerTransform)
 				{
 					FEulerTransform Euler = InFinalValue.Value;
-					Section->GetControlRig()->SetControlValue<FEulerTransform>(ParameterName, Euler, true, EControlRigSetKey::Never,bSetupUndo);
+					Section->GetControlRig()->SetControlValue<FRigControlValue::FEulerTransform_Float>(ParameterName, Euler, true, EControlRigSetKey::Never,bSetupUndo);
 				}
 			}
 
diff --git a/Engine/Plugins/Animation/ControlRig/Source/ControlRig/Private/Tools/ControlRigPose.cpp b/Engine/Plugins/Animation/ControlRig/Source/ControlRig/Private/Tools/ControlRigPose.cpp
--- a/Engine/Plugins/Animation/ControlRig/Source/ControlRig/Private/Tools/ControlRigPose.cpp
+++ b/Engine/Plugins/Animation/ControlRig/Source/ControlRig/Private/Tools/ControlRigPose.cpp
@@ -126,8 +126,8 @@ void FControlRigControlPose::PastePoseInternal(UControlRig* ControlRig, bool bDo
 			}
 			case ERigControlType::Vector2D:
 			{
-				FVector2D Val = CopyRigControl->Value.Get<FVector2D>();
-				ControlRig->SetControlValue<FVector2D>(ControlElement->GetName(), Val, true, Context,bSetupUndo);
+				FVector3f Val = CopyRigControl->Value.Get<FVector3f>();
+				ControlRig->SetControlValue<FVector3f>(ControlElement->GetName(), Val, true, Context,bSetupUndo);
 				break;
 			}
 			default:
diff --git a/Engine/Plugins/Animation/ControlRig/Source/ControlRig/Private/Units/Hierarchy/RigUnit_GetControlTransform.cpp b/Engine/Plugins/Animation/ControlRig/Source/ControlRig/Private/Units/Hierarchy/RigUnit_GetControlTransform.cpp
--- a/Engine/Plugins/Animation/ControlRig/Source/ControlRig/Private/Units/Hierarchy/RigUnit_GetControlTransform.cpp
+++ b/Engine/Plugins/Animation/ControlRig/Source/ControlRig/Private/Units/Hierarchy/RigUnit_GetControlTransform.cpp
@@ -121,9 +121,12 @@ FRigUnit_GetControlVector2D_Execute()
 				}
 				else
 				{
-					Vector = Hierarchy->GetControlValue(CachedControlIndex).Get<FVector2D>();
-					Minimum = Hierarchy->GetControlValue(CachedControlIndex, ERigControlValueType::Minimum).Get<FVector2D>();
-					Maximum = Hierarchy->GetControlValue(CachedControlIndex, ERigControlValueType::Maximum).Get<FVector2D>();
+					const FVector3f TempVector = Hierarchy->GetControlValue(CachedControlIndex).Get<FVector3f>();
+					const FVector3f TempMinimum = Hierarchy->GetControlValue(CachedControlIndex, ERigControlValueType::Minimum).Get<FVector3f>();
+					const FVector3f TempMaximum = Hierarchy->GetControlValue(CachedControlIndex, ERigControlValueType::Maximum).Get<FVector3f>();
+					Vector = FVector2D(TempVector.X, TempVector.Y);
+					Minimum = FVector2D(TempMinimum.X, TempMinimum.Y);
+					Maximum = FVector2D(TempMaximum.X, TempMaximum.Y);
 				}
 			}
 			default:
@@ -184,8 +187,8 @@ FRigUnit_GetControlVector_Execute()
 						Vector = Transform.GetScale3D();
 					}
 
-					Minimum = Hierarchy->GetControlValue(CachedControlIndex, ERigControlValueType::Minimum).Get<FVector>();
-					Maximum = Hierarchy->GetControlValue(CachedControlIndex, ERigControlValueType::Maximum).Get<FVector>();
+					Minimum = Hierarchy->GetControlValue(CachedControlIndex, ERigControlValueType::Minimum).Get<FVector3f>();
+					Maximum = Hierarchy->GetControlValue(CachedControlIndex, ERigControlValueType::Maximum).Get<FVector3f>();
 				}
 			}
 			default:
@@ -236,8 +239,11 @@ FRigUnit_GetControlRotator_Execute()
 					}
 
 					Rotator = Transform.GetRotation().Rotator();
-					Minimum = Hierarchy->GetControlValue(CachedControlIndex, ERigControlValueType::Minimum).Get<FRotator>();
-					Maximum = Hierarchy->GetControlValue(CachedControlIndex, ERigControlValueType::Maximum).Get<FRotator>();
+
+					const FVector3f MinimumVector = Hierarchy->GetControlValue(CachedControlIndex, ERigControlValueType::Minimum).Get<FVector3f>();
+					const FVector3f MaximumVector = Hierarchy->GetControlValue(CachedControlIndex, ERigControlValueType::Maximum).Get<FVector3f>();
+					Minimum = FRotator::MakeFromEuler(MinimumVector); 
+					Maximum = FRotator::MakeFromEuler(MaximumVector); 
 				}
 			}
 			default:
@@ -285,8 +291,9 @@ FRigUnit_GetControlTransform_Execute()
 							break;
 						}
 					}
-					Minimum = Hierarchy->GetControlValue(CachedControlIndex, ERigControlValueType::Minimum).Get<FTransform>();
-					Maximum = Hierarchy->GetControlValue(CachedControlIndex, ERigControlValueType::Maximum).Get<FTransform>();
+					
+					Minimum = Hierarchy->GetControlValue(CachedControlIndex, ERigControlValueType::Minimum).Get<FRigControlValue::FTransform_Float>().ToTransform();
+					Maximum = Hierarchy->GetControlValue(CachedControlIndex, ERigControlValueType::Maximum).Get<FRigControlValue::FTransform_Float>().ToTransform();
 				}
 			}
 			default:
diff --git a/Engine/Plugins/Animation/ControlRig/Source/ControlRig/Private/Units/Hierarchy/RigUnit_SetControlTransform.cpp b/Engine/Plugins/Animation/ControlRig/Source/ControlRig/Private/Units/Hierarchy/RigUnit_SetControlTransform.cpp
--- a/Engine/Plugins/Animation/ControlRig/Source/ControlRig/Private/Units/Hierarchy/RigUnit_SetControlTransform.cpp
+++ b/Engine/Plugins/Animation/ControlRig/Source/ControlRig/Private/Units/Hierarchy/RigUnit_SetControlTransform.cpp
@@ -212,14 +212,15 @@ FRigUnit_SetControlVector2D_Execute()
 				const FRigElementKey Key(Control, ERigElementType::Control);
 				if (CachedControlIndex.UpdateCache(Key, Hierarchy))
 				{
+					const FVector3f CurrentValue(Vector.X, Vector.Y, 0.f);
 					if(FMath::IsNearlyEqual(Weight, 1.f))
 					{
-						Hierarchy->SetControlValue(CachedControlIndex, FRigControlValue::Make<FVector2D>(Vector));
+						Hierarchy->SetControlValue(CachedControlIndex, FRigControlValue::Make<FVector3f>(CurrentValue));
 					}
 					else
 					{
-						FVector2D PreviousValue = Hierarchy->GetControlValue(CachedControlIndex).Get<FVector2D>();
-						Hierarchy->SetControlValue(CachedControlIndex, FRigControlValue::Make<FVector2D>(FMath::Lerp<FVector2D>(PreviousValue, Vector, FMath::Clamp<float>(Weight, 0.f, 1.f))));
+						const FVector3f PreviousValue = Hierarchy->GetControlValue(CachedControlIndex).Get<FVector3f>();
+						Hierarchy->SetControlValue(CachedControlIndex, FRigControlValue::Make<FVector3f>(FMath::Lerp<FVector3f>(PreviousValue, CurrentValue, FMath::Clamp<float>(Weight, 0.f, 1.f))));
 					}
 				}
 				break;
@@ -648,8 +649,10 @@ IMPLEMENT_RIGUNIT_AUTOMATION_TEST(FRigUnit_SetMultiControlVector2D)
 	Init();
 	Execute();
 
-	AddErrorIfFalse(Hierarchy->GetControlValue(FRigElementKey(TEXT("Control1"), ERigElementType::Control), ERigControlValueType::Current).Get<FVector2D>() == FVector2D(1.f, 1.f), TEXT("unexpected control value"));
-	AddErrorIfFalse(Hierarchy->GetControlValue(FRigElementKey(TEXT("Control2"), ERigElementType::Control), ERigControlValueType::Current).Get<FVector2D>() == FVector2D(2.f, 2.f), TEXT("unexpected control value"));
+	const FVector3f TempValue1 = Hierarchy->GetControlValue(FRigElementKey(TEXT("Control1"), ERigElementType::Control), ERigControlValueType::Current).Get<FVector3f>();
+	const FVector3f TempValue2 = Hierarchy->GetControlValue(FRigElementKey(TEXT("Control2"), ERigElementType::Control), ERigControlValueType::Current).Get<FVector3f>();
+	AddErrorIfFalse(FVector2D(TempValue1.X, TempValue1.Y) == FVector2D(1.f, 1.f), TEXT("unexpected control value"));
+	AddErrorIfFalse(FVector2D(TempValue2.X, TempValue2.Y) == FVector2D(2.f, 2.f), TEXT("unexpected control value"));
 
 	return true;
 }
@@ -682,9 +685,11 @@ IMPLEMENT_RIGUNIT_AUTOMATION_TEST(FRigUnit_SetMultiControlRotator)
 
 	Init();
 	Execute(); 
-	
-	AddErrorIfFalse(Hierarchy->GetControlValue(FRigElementKey(TEXT("Control1"), ERigElementType::Control), ERigControlValueType::Current).Get<FRotator>() != FRotator(0.f, 0.f, 0.f), TEXT("unexpected control value"));
-	AddErrorIfFalse(Hierarchy->GetControlValue(FRigElementKey(TEXT("Control2"), ERigElementType::Control), ERigControlValueType::Current).Get<FRotator>() != FRotator(0.f, 0.f, 0.f), TEXT("unexpected control value"));
+
+	const FVector3f TempValue1 = Hierarchy->GetControlValue(FRigElementKey(TEXT("Control1"), ERigElementType::Control), ERigControlValueType::Current).Get<FVector3f>();
+	const FVector3f TempValue2 = Hierarchy->GetControlValue(FRigElementKey(TEXT("Control2"), ERigElementType::Control), ERigControlValueType::Current).Get<FVector3f>();
+	AddErrorIfFalse(FRotator::MakeFromEuler(TempValue1) != FRotator(0.f, 0.f, 0.f), TEXT("unexpected control value"));
+	AddErrorIfFalse(FRotator::MakeFromEuler(TempValue2) != FRotator(0.f, 0.f, 0.f), TEXT("unexpected control value"));
 
 	return true;
 }
diff --git a/Engine/Plugins/Animation/ControlRig/Source/ControlRig/Public/Rigs/RigHierarchy.h b/Engine/Plugins/Animation/ControlRig/Source/ControlRig/Public/Rigs/RigHierarchy.h
--- a/Engine/Plugins/Animation/ControlRig/Source/ControlRig/Public/Rigs/RigHierarchy.h
+++ b/Engine/Plugins/Animation/ControlRig/Source/ControlRig/Public/Rigs/RigHierarchy.h
@@ -1957,7 +1957,7 @@ public:
 	UFUNCTION(BlueprintCallable, Category = URigHierarchy)
     static FORCEINLINE FRigControlValue MakeControlValueFromRotator(FRotator InValue)
 	{
-		return FRigControlValue::Make<FVector3f>(InValue.Vector());
+		return FRigControlValue::Make<FVector3f>(InValue.Euler());
 	}
 
 	/**
diff --git a/Engine/Plugins/Animation/ControlRig/Source/ControlRig/Public/Rigs/RigHierarchyDefines.h b/Engine/Plugins/Animation/ControlRig/Source/ControlRig/Public/Rigs/RigHierarchyDefines.h
--- a/Engine/Plugins/Animation/ControlRig/Source/ControlRig/Public/Rigs/RigHierarchyDefines.h
+++ b/Engine/Plugins/Animation/ControlRig/Source/ControlRig/Public/Rigs/RigHierarchyDefines.h
@@ -289,40 +289,40 @@ public:
 			case ERigControlType::Float: ValueStr = FString::Printf(TEXT("unreal.RigHierarchy.make_control_value_from_float(%.6f)"), Get<float>()); break;
 			case ERigControlType::Integer: ValueStr = FString::Printf(TEXT("unreal.RigHierarchy.make_control_value_from_int(%d)"), Get<int>()); break;
 			case ERigControlType::Position: ValueStr = FString::Printf(TEXT("unreal.RigHierarchy.make_control_value_from_vector(unreal.Vector(%.6f, %.6f, %.6f))"),
-				Get<FVector>().X, Get<FVector>().Y, Get<FVector>().Z); break;
+				Get<FVector3f>().X, Get<FVector3f>().Y, Get<FVector3f>().Z); break;
 			case ERigControlType::Rotator: ValueStr = FString::Printf(TEXT("unreal.RigHierarchy.make_control_value_from_rotator(unreal.Rotator(pitch=%.6f, roll=%.6f, yaw=%.6f))"),
-				Get<FRotator>().Pitch, Get<FRotator>().Roll, Get<FRotator>().Yaw); break;
+				Get<FVector3f>().X, Get<FVector3f>().Z, Get<FVector3f>().Y); break;
 			case ERigControlType::Scale: ValueStr = FString::Printf(TEXT("unreal.RigHierarchy.make_control_value_from_vector(unreal.Vector(%.6f, %.6f, %.6f))"),
-				Get<FVector>().X, Get<FVector>().Y, Get<FVector>().Z); break;
+				Get<FVector3f>().X, Get<FVector3f>().Y, Get<FVector3f>().Z); break;
 			case ERigControlType::Transform: ValueStr = FString::Printf(TEXT("unreal.RigHierarchy.make_control_value_from_transform(unreal.Transform(location=[%.6f,%.6f,%.6f],rotation=[%.6f,%.6f,%.6f],scale=[%.6f,%.6f,%.6f]))"),
-				Get<FTransform>().GetLocation().X,
-				Get<FTransform>().GetLocation().Y,
-				Get<FTransform>().GetLocation().Z,
-				Get<FTransform>().Rotator().Pitch,
-				Get<FTransform>().Rotator().Yaw,
-				Get<FTransform>().Rotator().Roll,
-				Get<FTransform>().GetScale3D().X,
-				Get<FTransform>().GetScale3D().Y,
-				Get<FTransform>().GetScale3D().Z); break;
+				Get<FTransform_Float>().ToTransform().GetLocation().X,
+				Get<FTransform_Float>().ToTransform().GetLocation().Y,
+				Get<FTransform_Float>().ToTransform().GetLocation().Z,
+				Get<FTransform_Float>().ToTransform().Rotator().Pitch,
+				Get<FTransform_Float>().ToTransform().Rotator().Yaw,
+				Get<FTransform_Float>().ToTransform().Rotator().Roll,
+				Get<FTransform_Float>().ToTransform().GetScale3D().X,
+				Get<FTransform_Float>().ToTransform().GetScale3D().Y,
+				Get<FTransform_Float>().ToTransform().GetScale3D().Z); break;
 			case ERigControlType::EulerTransform: ValueStr = FString::Printf(TEXT("unreal.RigHierarchy.make_control_value_from_euler_transform(unreal.EulerTransform(location=[%.6f,%.6f,%.6f],rotation=[%.6f,%.6f,%.6f],scale=[%.6f,%.6f,%.6f]))"),
-				Get<FEulerTransform>().Location.X,
-				Get<FEulerTransform>().Location.Y,
-				Get<FEulerTransform>().Location.Z,
-				Get<FEulerTransform>().Rotation.Euler().X,
-				Get<FEulerTransform>().Rotation.Euler().Y,
-				Get<FEulerTransform>().Rotation.Euler().Z,
-				Get<FEulerTransform>().Scale.X,
-				Get<FEulerTransform>().Scale.Y,
-				Get<FEulerTransform>().Scale.Z); break;
+				Get<FEulerTransform_Float>().ToTransform().Location.X,
+				Get<FEulerTransform_Float>().ToTransform().Location.Y,
+				Get<FEulerTransform_Float>().ToTransform().Location.Z,
+				Get<FEulerTransform_Float>().ToTransform().Rotation.Euler().X,
+				Get<FEulerTransform_Float>().ToTransform().Rotation.Euler().Y,
+				Get<FEulerTransform_Float>().ToTransform().Rotation.Euler().Z,
+				Get<FEulerTransform_Float>().ToTransform().Scale.X,
+				Get<FEulerTransform_Float>().ToTransform().Scale.Y,
+				Get<FEulerTransform_Float>().ToTransform().Scale.Z); break;
 			case ERigControlType::Vector2D: ValueStr = FString::Printf(TEXT("unreal.RigHierarchy.make_control_value_from_vector2d(unreal.Vector2D(%.6f, %.6f))"),
-				Get<FVector2D>().X, Get<FVector2D>().Y); break;
+				Get<FVector3f>().X, Get<FVector3f>().Y); break;
 			case ERigControlType::TransformNoScale: ValueStr = FString::Printf(TEXT("unreal.RigHierarchy.make_control_value_from_transform_no_scale(unreal.TransformNoScale(location=[%.6f,%.6f,%.6f],rotation=[%.6f,%.6f,%.6f]))"),
-				Get<FTransform>().GetLocation().X,
-				Get<FTransform>().GetLocation().Y,
-				Get<FTransform>().GetLocation().Z,
-				Get<FTransform>().Rotator().Pitch,
-				Get<FTransform>().Rotator().Yaw,
-				Get<FTransform>().Rotator().Roll); break;
+				Get<FTransformNoScale_Float>().ToTransform().ToFTransform().GetLocation().X,
+				Get<FTransformNoScale_Float>().ToTransform().ToFTransform().GetLocation().Y,
+				Get<FTransformNoScale_Float>().ToTransform().ToFTransform().GetLocation().Z,
+				Get<FTransformNoScale_Float>().ToTransform().ToFTransform().Rotator().Pitch,
+				Get<FTransformNoScale_Float>().ToTransform().ToFTransform().Rotator().Yaw,
+				Get<FTransformNoScale_Float>().ToTransform().ToFTransform().Rotator().Roll); break;
 			default: ensure(false);
 		}
 
@@ -429,7 +429,7 @@ public:
 			case ERigControlType::Rotator:
 			{
 				const FVector3f RotatorAxes = Get<FVector3f>();
-				Transform.SetRotation(FQuat(FRotator(RotatorAxes.X, RotatorAxes.Y, RotatorAxes.Z)));
+				Transform.SetRotation(FQuat(FRotator::MakeFromEuler(RotatorAxes)));
 				break;
 			}
 			case ERigControlType::Transform:
@@ -548,7 +548,7 @@ public:
 			case ERigControlType::Rotator:
 			{
 				//allow for values ><180/-180 by getting diff and adding that back in.
-				FRotator CurrentRotator = Get<FRotator>();
+				FRotator CurrentRotator = FRotator::MakeFromEuler(Get<FVector3f>());
 				FRotator CurrentRotWind, CurrentRotRem;
 				CurrentRotator.GetWindingAndRemainder(CurrentRotWind, CurrentRotRem);
 
@@ -559,7 +559,7 @@ public:
 
 				//Add Diff
 				CurrentRotator = CurrentRotator + DeltaRot;
-				Set<FVector3f>(CurrentRotator.Vector());
+				Set<FVector3f>(CurrentRotator.Euler());
 				break;
 			}
 			case ERigControlType::Transform:
@@ -578,7 +578,7 @@ public:
 				//Find Diff of the rotation from current and just add that instead of setting so we can go over/under -180
 				FEulerTransform NewTransform(InTransform);
 
-				const FEulerTransform CurrentEulerTransform = Get<FEulerTransform>();
+				const FEulerTransform CurrentEulerTransform = Get<FEulerTransform_Float>().ToTransform();
 				FRotator CurrentWinding;
 				FRotator CurrentRotRemainder;
 				CurrentEulerTransform.Rotation.GetWindingAndRemainder(CurrentWinding, CurrentRotRemainder);
@@ -775,9 +775,9 @@ public:
 				}
 				if (bLimitRotation)
 				{
-					ValueRef.RotationX = Clamp<float>(ValueRef.RotationX, Min.RotationX, Max.RotationX);
-					ValueRef.RotationY = Clamp<float>(ValueRef.RotationY, Min.RotationY, Max.RotationY);
-					ValueRef.RotationZ = Clamp<float>(ValueRef.RotationZ, Min.RotationZ, Max.RotationZ);
+					ValueRef.RotationPitch = Clamp<float>(ValueRef.RotationPitch, Min.RotationPitch, Max.RotationPitch);
+					ValueRef.RotationYaw = Clamp<float>(ValueRef.RotationYaw, Min.RotationYaw, Max.RotationYaw);
+					ValueRef.RotationRoll = Clamp<float>(ValueRef.RotationRoll, Min.RotationRoll, Max.RotationRoll);
 				}
 				if (bLimitScale)
 				{
@@ -844,6 +844,10 @@ public:
 			Transform.SetScale3D(FVector(ScaleX, ScaleY, ScaleZ));
 			return Transform;
 		}
+
+		FVector3f GetTranslation() const { return FVector3f(TranslationX, TranslationY, TranslationZ); }
+		FQuat GetRotation() const { return FQuat(RotationX, RotationY, RotationZ, RotationW); }
+		FVector3f GetScale3D() const { return FVector3f(ScaleX, ScaleY, ScaleZ); }
 	};
 
 	struct FTransformNoScale_Float
@@ -885,13 +889,16 @@ public:
 			Transform.Location = FVector(TranslationX, TranslationY, TranslationZ);
 			return Transform;
 		}
+
+		FVector3f GetTranslation() const { return FVector3f(TranslationX, TranslationY, TranslationZ); }
+		FQuat GetRotation() const { return FQuat(RotationX, RotationY, RotationZ, RotationW); }
 	};
 
 	struct FEulerTransform_Float
 	{
-		float RotationX;
-		float RotationY;
-		float RotationZ;
+		float RotationPitch;
+		float RotationYaw;
+		float RotationRoll;
 		float TranslationX;
 		float TranslationY;
 		float TranslationZ;
@@ -906,9 +913,9 @@ public:
 
 		FEulerTransform_Float(const FEulerTransform& InTransform)
 		{
-			RotationX = InTransform.Rotation.Pitch;
-			RotationY = InTransform.Rotation.Yaw;
-			RotationZ = InTransform.Rotation.Roll;
+			RotationPitch = InTransform.Rotation.Pitch;
+			RotationYaw = InTransform.Rotation.Yaw;
+			RotationRoll = InTransform.Rotation.Roll;
 			TranslationX = InTransform.Location.X;
 			TranslationY = InTransform.Location.Y;
 			TranslationZ = InTransform.Location.Z;
@@ -920,11 +927,15 @@ public:
 		FEulerTransform ToTransform() const
 		{
 			FEulerTransform Transform;
-			Transform.Rotation = FRotator(RotationX, RotationY, RotationZ);
+			Transform.Rotation = FRotator(RotationPitch, RotationYaw, RotationRoll);
 			Transform.Location = FVector(TranslationX, TranslationY, TranslationZ);
 			Transform.Scale = FVector(ScaleX, ScaleY, ScaleZ);
 			return Transform;
 		}
+		
+		FVector3f GetTranslation() const { return FVector3f(TranslationX, TranslationY, TranslationZ); }
+		FRotator GetRotator() const { return FRotator(RotationPitch, RotationYaw, RotationRoll); }
+		FVector3f GetScale3D() const { return FVector3f(ScaleX, ScaleY, ScaleZ); }
 	};
 
 private:
@@ -938,37 +949,29 @@ private:
 	friend struct FRigControlHierarchy;
 	friend class UControlRigBlueprint;
 	friend class URigHierarchyController;
+	friend class UControlRig;
 };
 
-#if UE_LARGE_WORLD_COORDINATES_DISABLED
-
-template<>
-inline FQuat FRigControlValue::SetFromString<FQuat>(const FString& InString)
-{
-	FQuat Value;
-	TBaseStructure<FQuat>::Get()->ImportText(*InString, &Value, nullptr, PPF_None, nullptr, TBaseStructure<FQuat>::Get()->GetName());
-	Set<FRotator>(Value.Rotator());
-	return Value;
-}
-
-#else
-
 template<>
 inline FRigControlValue FRigControlValue::Make(FVector2D InValue)
 {
 	return Make<FVector3f>(FVector3f(InValue.X, InValue.Y, 0.f));
 }
 
+#if !UE_LARGE_WORLD_COORDINATES_DISABLED
+
 template<>
 inline FRigControlValue FRigControlValue::Make(FVector InValue)
 {
 	return Make<FVector3f>(InValue);
 }
 
+#endif
+
 template<>
 inline FRigControlValue FRigControlValue::Make(FRotator InValue)
 {
-	return Make<FVector3f>(InValue.Vector());
+	return Make<FVector3f>(InValue.Euler());
 }
 
 template<>
@@ -989,6 +992,8 @@ inline FRigControlValue FRigControlValue::Make(FEulerTransform InValue)
 	return Make<FEulerTransform_Float>(InValue);
 }
 
+#if !UE_LARGE_WORLD_COORDINATES_DISABLED
+
 template<>
 inline FString FRigControlValue::ToString<FVector>() const
 {
@@ -998,6 +1003,8 @@ inline FString FRigControlValue::ToString<FVector>() const
 	return Result;
 }
 
+#endif
+
 template<>
 inline FString FRigControlValue::ToString<FRotator>() const
 {
@@ -1016,6 +1023,26 @@ inline FString FRigControlValue::ToString<FTransform>() const
 	return Result;
 }
 
+template<>
+inline FString FRigControlValue::ToString<FTransformNoScale>() const
+{
+	FTransformNoScale Value = GetRef<FTransformNoScale_Float>().ToTransform();
+	FString Result;
+	TBaseStructure<FTransformNoScale>::Get()->ExportText(Result, &Value, nullptr, nullptr, PPF_None, nullptr);
+	return Result;
+}
+
+template<>
+inline FString FRigControlValue::ToString<FEulerTransform>() const
+{
+	FEulerTransform Value = GetRef<FEulerTransform_Float>().ToTransform();
+	FString Result;
+	TBaseStructure<FEulerTransform>::Get()->ExportText(Result, &Value, nullptr, nullptr, PPF_None, nullptr);
+	return Result;
+}
+
+#if !UE_LARGE_WORLD_COORDINATES_DISABLED
+
 template<>
 inline FVector FRigControlValue::SetFromString<FVector>(const FString& InString)
 {
@@ -1025,12 +1052,14 @@ inline FVector FRigControlValue::SetFromString<FVector>(const FString& InString)
 	return Value;
 }
 
+#endif
+
 template<>
 inline FQuat FRigControlValue::SetFromString<FQuat>(const FString& InString)
 {
 	FQuat Value;
 	TBaseStructure<FQuat>::Get()->ImportText(*InString, &Value, nullptr, PPF_None, nullptr, TBaseStructure<FQuat>::Get()->GetName());
-	Set<FVector3f>(Value.Rotator().Vector());
+	Set<FVector3f>(Value.Rotator().Euler());
 	return Value;
 }
 
@@ -1039,7 +1068,7 @@ inline FRotator FRigControlValue::SetFromString<FRotator>(const FString& InStrin
 {
 	FRotator Value;
 	TBaseStructure<FRotator>::Get()->ImportText(*InString, &Value, nullptr, PPF_None, nullptr, TBaseStructure<FRotator>::Get()->GetName());
-	Set<FVector3f>(Value.Vector());
+	Set<FVector3f>(Value.Euler());
 	return Value;
 }
 
@@ -1053,35 +1082,47 @@ inline FTransform FRigControlValue::SetFromString<FTransform>(const FString& InS
 }
 
 template<>
-FVector2D& FRigControlValue::GetRef() = delete;
-template<>
-FVector& FRigControlValue::GetRef() = delete;
-template<>
-FRotator& FRigControlValue::GetRef() = delete;
-template<>
-FQuat& FRigControlValue::GetRef() = delete;
-template<>
-FTransform& FRigControlValue::GetRef() = delete;
-template<>
-FTransformNoScale& FRigControlValue::GetRef() = delete;
-template<>
-FEulerTransform& FRigControlValue::GetRef() = delete;
-template<>
-void FRigControlValue::Set(FVector2D InValue) = delete;
-template<>
-void FRigControlValue::Set(FVector InValue) = delete;
-template<>
-void FRigControlValue::Set(FRotator InValue) = delete;
-template<>
-void FRigControlValue::Set(FQuat InValue) = delete;
-template<>
-void FRigControlValue::Set(FTransform InValue) = delete;
-template<>
-void FRigControlValue::Set(FTransformNoScale InValue) = delete;
+inline FTransformNoScale FRigControlValue::SetFromString<FTransformNoScale>(const FString& InString)
+{
+	FTransformNoScale Value;
+	TBaseStructure<FTransformNoScale>::Get()->ImportText(*InString, &Value, nullptr, PPF_None, nullptr, TBaseStructure<FTransformNoScale>::Get()->GetName());
+	Set<FTransformNoScale_Float>(Value);
+	return Value;
+}
+
 template<>
-void FRigControlValue::Set(FEulerTransform InValue) = delete;
+inline FEulerTransform FRigControlValue::SetFromString<FEulerTransform>(const FString& InString)
+{
+	FEulerTransform Value;
+	TBaseStructure<FEulerTransform>::Get()->ImportText(*InString, &Value, nullptr, PPF_None, nullptr, TBaseStructure<FEulerTransform>::Get()->GetName());
+	Set<FEulerTransform_Float>(Value);
+	return Value;
+}
 
+#define UE_CONTROLRIG_VALUE_DELETE_TEMPLATE(T) \
+template<> \
+const T FRigControlValue::Get() const = delete;  \
+template<>  \
+const T& FRigControlValue::GetRef() const = delete; \
+template<> \
+T& FRigControlValue::GetRef() = delete; \
+template<> \
+void FRigControlValue::Set(T InValue) = delete;
+
+// We disable all of these to force Control Rig
+// to store values as floats. We've added our own
+// float variations purely for storage purposes,
+// from which we'll cast back and forth to FTransform etc.
+UE_CONTROLRIG_VALUE_DELETE_TEMPLATE(FVector2D)
+#if !UE_LARGE_WORLD_COORDINATES_DISABLED
+UE_CONTROLRIG_VALUE_DELETE_TEMPLATE(FVector)
 #endif
+UE_CONTROLRIG_VALUE_DELETE_TEMPLATE(FVector4)
+UE_CONTROLRIG_VALUE_DELETE_TEMPLATE(FRotator)
+UE_CONTROLRIG_VALUE_DELETE_TEMPLATE(FQuat)
+UE_CONTROLRIG_VALUE_DELETE_TEMPLATE(FTransform)
+UE_CONTROLRIG_VALUE_DELETE_TEMPLATE(FTransformNoScale)
+UE_CONTROLRIG_VALUE_DELETE_TEMPLATE(FEulerTransform)
 
 enum class EControlRigContextChannelToKey : uint32
 {
diff --git a/Engine/Plugins/Animation/ControlRig/Source/ControlRigEditor/Private/ControlRigElementDetails.cpp b/Engine/Plugins/Animation/ControlRig/Source/ControlRigEditor/Private/ControlRigElementDetails.cpp
--- a/Engine/Plugins/Animation/ControlRig/Source/ControlRigEditor/Private/ControlRigElementDetails.cpp
+++ b/Engine/Plugins/Animation/ControlRig/Source/ControlRigEditor/Private/ControlRigElementDetails.cpp
@@ -1297,6 +1297,92 @@ void FRigControlElementDetails_SetupFloatValueWidget(IDetailGroup& InGroup, IDet
 	.IsEnabled(EnabledAttribute);
 }
 
+template<typename T>
+T FRigControlElementDetails_ExtractValue(const FRigControlValue& InValue)
+{
+	return InValue.Get<T>();
+}
+
+template<typename T>
+FRigControlValue FRigControlElementDetails_PackageValue(const T& InValue)
+{
+	return FRigControlValue::Make<T>(InValue);
+}
+
+template<>
+FVector2D FRigControlElementDetails_ExtractValue(const FRigControlValue& InValue)
+{
+	const FVector3f TempValue = InValue.Get<FVector3f>();
+	return FVector2D(TempValue.X, TempValue.Y);
+}
+
+template<>
+FRigControlValue FRigControlElementDetails_PackageValue(const FVector2D& InValue)
+{
+	const FVector3f TempValue(InValue.X, InValue.Y, 0.f);
+	return FRigControlValue::Make<FVector3f>(TempValue);
+}
+
+template<>
+FVector FRigControlElementDetails_ExtractValue(const FRigControlValue& InValue)
+{
+	return InValue.Get<FVector3f>();
+}
+
+template<>
+FRigControlValue FRigControlElementDetails_PackageValue(const FVector& InValue)
+{
+	return FRigControlValue::Make<FVector3f>(InValue);
+}
+
+template<>
+FRotator FRigControlElementDetails_ExtractValue(const FRigControlValue& InValue)
+{
+	return FRotator::MakeFromEuler(InValue.Get<FVector3f>());
+}
+
+template<>
+FRigControlValue FRigControlElementDetails_PackageValue(const FRotator& InValue)
+{
+	return FRigControlValue::Make<FVector3f>(InValue.Euler());
+}
+
+template<>
+FTransform FRigControlElementDetails_ExtractValue(const FRigControlValue& InValue)
+{
+	return InValue.Get<FRigControlValue::FTransform_Float>().ToTransform();
+}
+
+template<>
+FRigControlValue FRigControlElementDetails_PackageValue(const FTransform& InValue)
+{
+	return FRigControlValue::Make<FRigControlValue::FTransform_Float>(InValue);
+}
+
+template<>
+FTransformNoScale FRigControlElementDetails_ExtractValue(const FRigControlValue& InValue)
+{
+	return InValue.Get<FRigControlValue::FTransformNoScale_Float>().ToTransform();
+}
+
+template<>
+FRigControlValue FRigControlElementDetails_PackageValue(const FTransformNoScale& InValue)
+{
+	return FRigControlValue::Make<FRigControlValue::FTransformNoScale_Float>(InValue);
+}
+
+template<>
+FEulerTransform FRigControlElementDetails_ExtractValue(const FRigControlValue& InValue)
+{
+	return InValue.Get<FRigControlValue::FEulerTransform_Float>().ToTransform();
+}
+
+template<>
+FRigControlValue FRigControlElementDetails_PackageValue(const FEulerTransform& InValue)
+{
+	return FRigControlValue::Make<FRigControlValue::FEulerTransform_Float>(InValue);
+}
+
 template<typename T>
 void FRigControlElementDetails_SetupStructValueWidget(IDetailGroup& InGroup, IDetailChildrenBuilder& InStructBuilder, ERigControlValueType InValueType, FRigControlElement* InControlElement, URigHierarchy* InHierarchy)
 {
@@ -1320,8 +1406,9 @@ void FRigControlElementDetails_SetupStructValueWidget(IDetailGroup& InGroup, IDe
             {
             	// update the struct with the current control value
             	uint8* StructMemory = StructToDisplay->GetStructMemory();
-            	const FRigControlValue& CurrentValue = HierarchyPtr->GetControlValue(Key, InValueType);            	
-            	FMemory::Memcpy(StructToDisplay->GetStructMemory(), &CurrentValue.GetRef<T>(), sizeof(T));
+            	const FRigControlValue& ControlValue = HierarchyPtr->GetControlValue(Key, InValueType);
+            	T ExtractedValue = FRigControlElementDetails_ExtractValue<T>(ControlValue);
+            	FMemory::Memcpy(StructToDisplay->GetStructMemory(), &ExtractedValue, sizeof(T));
 
             	return ControlElement->Settings.IsValueTypeEnabled(InValueType);
             }
@@ -1357,13 +1444,13 @@ void FRigControlElementDetails_SetupStructValueWidget(IDetailGroup& InGroup, IDe
 	{
 		if(HierarchyPtr.IsValid())
 		{
-            const FRigControlValue Value = FRigControlValue::Make(*(T*)StructToDisplay->GetStructMemory());
-			HierarchyPtr->SetControlValue(Key, Value, InValueType, true);
+            const FRigControlValue ControlValue = FRigControlElementDetails_PackageValue(*(T*)StructToDisplay->GetStructMemory());
+			HierarchyPtr->SetControlValue(Key, ControlValue, InValueType, true);
 			if(InValueType == ERigControlValueType::Initial)
 			{
                 if(UControlRigBlueprint* Blueprint = RigElementDetails_GetBlueprintFromHierarchy(HierarchyPtr.Get()))
                 {
-                    Blueprint->Hierarchy->SetControlValue(Key, Value, InValueType, true);
+                    Blueprint->Hierarchy->SetControlValue(Key, ControlValue, InValueType, true);
                 }
             }
 		}
diff --git a/Engine/Plugins/Animation/ControlRig/Source/ControlRigEditor/Private/ControlRigSequencerEditorLibrary.cpp b/Engine/Plugins/Animation/ControlRig/Source/ControlRigEditor/Private/ControlRigSequencerEditorLibrary.cpp
--- a/Engine/Plugins/Animation/ControlRig/Source/ControlRigEditor/Private/ControlRigSequencerEditorLibrary.cpp
+++ b/Engine/Plugins/Animation/ControlRig/Source/ControlRigEditor/Private/ControlRigSequencerEditorLibrary.cpp
@@ -1294,7 +1294,8 @@ FVector2D UControlRigSequencerEditorLibrary::GetLocalControlRigVector2D(ULevelSe
 		FRigControlValue RigValue;
 		if (GetControlRigValue(WeakSequencer.Pin().Get(), ControlRig, ControlName, TimeUnit, Frame, RigValue))
 		{
-			Value = RigValue.Get<FVector2D>();
+			const FVector3f TempValue = RigValue.Get<FVector3f>(); 
+			Value = FVector2D(TempValue.X, TempValue.Y);
 		}
 	}
 	return Value;
@@ -1321,8 +1322,8 @@ TArray<FVector2D> UControlRigSequencerEditorLibrary::GetLocalControlRigVector2Ds
 			Values.Reserve(RigValues.Num());
 			for (const FRigControlValue& RigValue : RigValues)
 			{
-				FVector2D Value = RigValue.Get<FVector2D>();
-				Values.Add(Value);
+				const FVector3f TempValue = RigValue.Get<FVector3f>(); 
+				Values.Add(FVector2D(TempValue.X, TempValue.Y));
 			}
 		}
 	}
@@ -1354,7 +1355,7 @@ void UControlRigSequencerEditorLibrary::SetLocalControlRigVector2D(ULevelSequenc
 		FRigControlModifiedContext Context;
 		Context.SetKey = bSetKey ? EControlRigSetKey::Always : EControlRigSetKey::DoNotCare;
 		Context.LocalTime = TickResolution.AsSeconds(FFrameTime(Frame));
-		ControlRig->SetControlValue<FVector2D>(ControlName, Value, true, Context);
+		ControlRig->SetControlValue<FVector3f>(ControlName, FVector3f(Value.X, Value.Y, 0.f), true, Context);
 	}
 }
 
@@ -1387,7 +1388,7 @@ void UControlRigSequencerEditorLibrary::SetLocalControlRigVector2Ds(ULevelSequen
 			}
 			FVector2D Value = Values[Index];
 			Context.LocalTime = TickResolution.AsSeconds(FFrameTime(Frame));
-			ControlRig->SetControlValue<FVector2D>(ControlName, Value, true, Context);
+			ControlRig->SetControlValue<FVector3f>(ControlName, FVector3f(Value.X, Value.Y, 0.f), true, Context);
 		}
 	}
 }
@@ -1411,7 +1412,7 @@ FVector UControlRigSequencerEditorLibrary::GetLocalControlRigPosition(ULevelSequ
 		FRigControlValue RigValue;
 		if (GetControlRigValue(WeakSequencer.Pin().Get(), ControlRig, ControlName, TimeUnit, Frame, RigValue))
 		{
-			Value = RigValue.Get<FVector>();
+			Value = RigValue.Get<FVector3f>();
 		}
 	}
 	return Value;
@@ -1438,7 +1439,7 @@ TArray<FVector> UControlRigSequencerEditorLibrary::GetLocalControlRigPositions(U
 			Values.Reserve(RigValues.Num());
 			for (const FRigControlValue& RigValue : RigValues)
 			{
-				FVector Value = RigValue.Get<FVector>();
+				FVector Value = RigValue.Get<FVector3f>();
 				Values.Add(Value);
 			}
 		}
@@ -1470,7 +1471,7 @@ void UControlRigSequencerEditorLibrary::SetLocalControlRigPosition(ULevelSequenc
 		FRigControlModifiedContext Context;
 		Context.SetKey = bSetKey ? EControlRigSetKey::Always : EControlRigSetKey::DoNotCare;
 		Context.LocalTime = TickResolution.AsSeconds(FFrameTime(Frame));
-		ControlRig->SetControlValue<FVector>(ControlName, Value, true, Context);
+		ControlRig->SetControlValue<FVector3f>(ControlName, Value, true, Context);
 	}
 }
 
@@ -1503,7 +1504,7 @@ void UControlRigSequencerEditorLibrary::SetLocalControlRigPositions(ULevelSequen
 			}
 			FVector Value = Values[Index];
 			Context.LocalTime = TickResolution.AsSeconds(FFrameTime(Frame));
-			ControlRig->SetControlValue<FVector>(ControlName, Value, true, Context);
+			ControlRig->SetControlValue<FVector3f>(ControlName, Value, true, Context);
 		}
 	}
 }
@@ -1527,7 +1528,7 @@ FRotator UControlRigSequencerEditorLibrary::GetLocalControlRigRotator(ULevelSequ
 		FRigControlValue RigValue;
 		if (GetControlRigValue(WeakSequencer.Pin().Get(), ControlRig, ControlName,TimeUnit, Frame, RigValue))
 		{
-			Value = RigValue.Get<FRotator>();
+			Value = FRotator::MakeFromEuler(RigValue.Get<FVector3f>());
 		}
 	}
 	return Value;
@@ -1554,7 +1555,7 @@ TArray<FRotator> UControlRigSequencerEditorLibrary::GetLocalControlRigRotators(U
 			Values.Reserve(RigValues.Num());
 			for (const FRigControlValue& RigValue : RigValues)
 			{
-				FRotator Value = RigValue.Get<FRotator>();
+				FRotator Value = FRotator::MakeFromEuler(RigValue.Get<FVector3f>());
 				Values.Add(Value);
 			}
 		}
@@ -1586,7 +1587,7 @@ void UControlRigSequencerEditorLibrary::SetLocalControlRigRotator(ULevelSequence
 		FRigControlModifiedContext Context;
 		Context.SetKey = bSetKey ? EControlRigSetKey::Always : EControlRigSetKey::DoNotCare;
 		Context.LocalTime = TickResolution.AsSeconds(FFrameTime(Frame));
-		ControlRig->SetControlValue<FRotator>(ControlName, Value, true, Context);
+		ControlRig->SetControlValue<FVector3f>(ControlName, Value.Euler(), true, Context);
 	}
 }
 
@@ -1619,7 +1620,7 @@ void UControlRigSequencerEditorLibrary::SetLocalControlRigRotators(ULevelSequenc
 			}
 			FRotator Value = Values[Index];
 			Context.LocalTime = TickResolution.AsSeconds(FFrameTime(Frame));
-			ControlRig->SetControlValue<FRotator>(ControlName, Value, true, Context);
+			ControlRig->SetControlValue<FVector3f>(ControlName, Value.Euler(), true, Context);
 		}
 	}
 }
@@ -1643,7 +1644,7 @@ FVector UControlRigSequencerEditorLibrary::GetLocalControlRigScale(ULevelSequenc
 		FRigControlValue RigValue;
 		if (GetControlRigValue(WeakSequencer.Pin().Get(), ControlRig, ControlName,TimeUnit, Frame, RigValue))
 		{
-			Value = RigValue.Get<FVector>();
+			Value = RigValue.Get<FVector3f>();
 		}
 	}
 	return Value;
@@ -1670,7 +1671,7 @@ TArray<FVector>UControlRigSequencerEditorLibrary::GetLocalControlRigScales(ULeve
 			Values.Reserve(RigValues.Num());
 			for (const FRigControlValue& RigValue : RigValues)
 			{
-				FVector Value = RigValue.Get<FVector>();
+				FVector Value = RigValue.Get<FVector3f>();
 				Values.Add(Value);
 			}
 		}
@@ -1702,7 +1703,7 @@ void UControlRigSequencerEditorLibrary::SetLocalControlRigScale(ULevelSequence*
 		FRigControlModifiedContext Context;
 		Context.SetKey = bSetKey ? EControlRigSetKey::Always : EControlRigSetKey::DoNotCare;
 		Context.LocalTime = TickResolution.AsSeconds(FFrameTime(Frame));
-		ControlRig->SetControlValue<FVector>(ControlName, Value, true, Context);
+		ControlRig->SetControlValue<FVector3f>(ControlName, Value, true, Context);
 	}
 }
 
@@ -1735,7 +1736,7 @@ void UControlRigSequencerEditorLibrary::SetLocalControlRigScales(ULevelSequence*
 			}
 			FVector Value = Values[Index];
 			Context.LocalTime = TickResolution.AsSeconds(FFrameTime(Frame));
-			ControlRig->SetControlValue<FVector>(ControlName, Value, true, Context);
+			ControlRig->SetControlValue<FVector3f>(ControlName, Value, true, Context);
 		}
 	}
 }
@@ -1759,7 +1760,7 @@ FEulerTransform UControlRigSequencerEditorLibrary::GetLocalControlRigEulerTransf
 		FRigControlValue RigValue;
 		if (GetControlRigValue(WeakSequencer.Pin().Get(), ControlRig, ControlName, TimeUnit, Frame, RigValue))
 		{
-			Value = RigValue.Get<FEulerTransform>();
+			Value = RigValue.Get<FRigControlValue::FEulerTransform_Float>().ToTransform();
 		}
 	}
 	return Value;
@@ -1786,7 +1787,7 @@ TArray<FEulerTransform> UControlRigSequencerEditorLibrary::GetLocalControlRigEul
 			Values.Reserve(RigValues.Num());
 			for (const FRigControlValue& RigValue : RigValues)
 			{
-				FEulerTransform Value = RigValue.Get<FEulerTransform>();
+				FEulerTransform Value = RigValue.Get<FRigControlValue::FEulerTransform_Float>().ToTransform();
 				Values.Add(Value);
 			}
 		}
@@ -1819,7 +1820,7 @@ void UControlRigSequencerEditorLibrary::SetLocalControlRigEulerTransform(ULevelS
 		FRigControlModifiedContext Context;
 		Context.SetKey = bSetKey ? EControlRigSetKey::Always : EControlRigSetKey::DoNotCare;
 		Context.LocalTime = TickResolution.AsSeconds(FFrameTime(Frame));
-		ControlRig->SetControlValue<FEulerTransform>(ControlName, Value, true, Context);
+		ControlRig->SetControlValue<FRigControlValue::FEulerTransform_Float>(ControlName, Value, true, Context);
 	}
 }
 
@@ -1852,7 +1853,7 @@ void UControlRigSequencerEditorLibrary::SetLocalControlRigEulerTransforms(ULevel
 			}
 			FEulerTransform Value = Values[Index];
 			Context.LocalTime = TickResolution.AsSeconds(FFrameTime(Frame));
-			ControlRig->SetControlValue<FEulerTransform>(ControlName, Value, true, Context);
+			ControlRig->SetControlValue<FRigControlValue::FEulerTransform_Float>(ControlName, Value, true, Context);
 		}
 	}
 }
@@ -1876,7 +1877,7 @@ FTransformNoScale UControlRigSequencerEditorLibrary::GetLocalControlRigTransform
 		FRigControlValue RigValue;
 		if (GetControlRigValue(WeakSequencer.Pin().Get(), ControlRig, ControlName,TimeUnit, Frame, RigValue))
 		{
-			Value = RigValue.Get<FTransformNoScale>();
+			Value = RigValue.Get<FRigControlValue::FTransformNoScale_Float>().ToTransform();
 		}
 	}
 	return Value;
@@ -1903,7 +1904,7 @@ TArray<FTransformNoScale> UControlRigSequencerEditorLibrary::GetLocalControlRigT
 			Values.Reserve(RigValues.Num());
 			for (const FRigControlValue& RigValue : RigValues)
 			{
-				FTransformNoScale Value = RigValue.Get<FTransformNoScale>();
+				FTransformNoScale Value = RigValue.Get<FRigControlValue::FTransformNoScale_Float>().ToTransform();
 				Values.Add(Value);
 			}
 		}
@@ -1936,7 +1937,7 @@ void UControlRigSequencerEditorLibrary::SetLocalControlRigTransformNoScale(ULeve
 		FRigControlModifiedContext Context;
 		Context.SetKey = bSetKey ? EControlRigSetKey::Always : EControlRigSetKey::DoNotCare;
 		Context.LocalTime = TickResolution.AsSeconds(FFrameTime(Frame));
-		ControlRig->SetControlValue<FTransformNoScale>(ControlName, Value, true, Context);
+		ControlRig->SetControlValue<FRigControlValue::FTransformNoScale_Float>(ControlName, Value, true, Context);
 	}
 }
 
@@ -1969,7 +1970,7 @@ void UControlRigSequencerEditorLibrary::SetLocalControlRigTransformNoScales(ULev
 			}
 			FTransformNoScale Value = Values[Index];
 			Context.LocalTime = TickResolution.AsSeconds(FFrameTime(Frame));
-			ControlRig->SetControlValue<FTransformNoScale>(ControlName, Value, true, Context);
+			ControlRig->SetControlValue<FRigControlValue::FTransformNoScale_Float>(ControlName, Value, true, Context);
 		}
 	}
 }
@@ -1993,7 +1994,7 @@ FTransform UControlRigSequencerEditorLibrary::GetLocalControlRigTransform(ULevel
 		FRigControlValue RigValue;
 		if (GetControlRigValue(WeakSequencer.Pin().Get(), ControlRig, ControlName, TimeUnit, Frame, RigValue))
 		{
-			Value = RigValue.Get<FTransform>();
+			Value = RigValue.Get<FRigControlValue::FTransform_Float>().ToTransform();
 		}
 	}
 	return Value;
@@ -2020,7 +2021,7 @@ TArray<FTransform> UControlRigSequencerEditorLibrary::GetLocalControlRigTransfor
 			Values.Reserve(RigValues.Num());
 			for (const FRigControlValue& RigValue : RigValues)
 			{
-				FTransform Value = RigValue.Get<FTransform>();
+				FTransform Value = RigValue.Get<FRigControlValue::FTransform_Float>().ToTransform();
 				Values.Add(Value);
 			}
 		}
@@ -2053,7 +2054,7 @@ void UControlRigSequencerEditorLibrary::SetLocalControlRigTransform(ULevelSequen
 		FRigControlModifiedContext Context;
 		Context.SetKey = bSetKey ? EControlRigSetKey::Always : EControlRigSetKey::DoNotCare;
 		Context.LocalTime = TickResolution.AsSeconds(FFrameTime(Frame));
-		ControlRig->SetControlValue<FTransform>(ControlName, Value, true, Context);
+		ControlRig->SetControlValue<FRigControlValue::FTransform_Float>(ControlName, Value, true, Context);
 	}
 }
 
@@ -2086,7 +2087,7 @@ void UControlRigSequencerEditorLibrary::SetLocalControlRigTransforms(ULevelSeque
 			}
 			FTransform Value = Values[Index];
 			Context.LocalTime = TickResolution.AsSeconds(FFrameTime(Frame));
-			ControlRig->SetControlValue<FTransform>(ControlName, Value, true, Context);
+			ControlRig->SetControlValue<FRigControlValue::FTransform_Float>(ControlName, Value, true, Context);
 		}
 	}
 }
diff --git a/Engine/Plugins/Animation/ControlRig/Source/ControlRigEditor/Private/EditMode/ControlRigControlsProxy.cpp b/Engine/Plugins/Animation/ControlRig/Source/ControlRigEditor/Private/EditMode/ControlRigControlsProxy.cpp
--- a/Engine/Plugins/Animation/ControlRig/Source/ControlRigEditor/Private/EditMode/ControlRigControlsProxy.cpp
+++ b/Engine/Plugins/Animation/ControlRig/Source/ControlRigEditor/Private/EditMode/ControlRigControlsProxy.cpp
@@ -151,7 +151,7 @@ void UControlRigTransformControlProxy::PostEditChangeProperty(struct FPropertyCh
 		{
 			//MUST set through ControlRig
 			FControlRigInteractionScope InteractionScope(ControlRig.Get());
-			ControlRig->SetControlValue<FTransform>(ControlName, Transform, true, EControlRigSetKey::DoNotCare,false);
+			ControlRig->SetControlValue<FRigControlValue::FTransform_Float>(ControlName, Transform, true, EControlRigSetKey::DoNotCare,false);
 			ControlRig->Evaluate_AnyThread();
 
 		}
@@ -166,7 +166,7 @@ void UControlRigTransformControlProxy::ValueChanged()
 		Modify();
 		const FName PropertyName("Transform");
 		FTrackInstancePropertyBindings Binding(PropertyName, PropertyName.ToString());
-		const FTransform NewTransform = ControlRig.Get()->GetHierarchy()->GetControlValue(ControlElement, ERigControlValueType::Current).Get<FTransform>();
+		const FTransform NewTransform = ControlRig.Get()->GetHierarchy()->GetControlValue(ControlElement, ERigControlValueType::Current).Get<FRigControlValue::FTransform_Float>().ToTransform();
 		Binding.CallFunction<FTransform>(*this, NewTransform);
 	}
 }
@@ -182,7 +182,7 @@ void UControlRigTransformControlProxy::PostEditUndo()
 			CheckEditModeOnSelectionChange(ControlRig.Get());
 		}
 		ControlRig->SelectControl(ControlName, bSelected);
-		ControlRig->SetControlValue<FTransform>(ControlName, Transform, true, EControlRigSetKey::Never,false);
+		ControlRig->SetControlValue<FRigControlValue::FTransform_Float>(ControlName, Transform, true, EControlRigSetKey::Never,false);
 	}
 }
 #endif
@@ -192,7 +192,7 @@ void UControlRigTransformControlProxy::SetKey(const IPropertyHandle& KeyedProper
 	FRigControlElement* ControlElement = GetControlElement();
 	if (ControlElement)
 	{
-		ControlRig->SetControlValue<FTransform>(ControlName, Transform, true, EControlRigSetKey::Always,false);
+		ControlRig->SetControlValue<FRigControlValue::FTransform_Float>(ControlName, Transform, true, EControlRigSetKey::Always,false);
 	}
 }
 
@@ -207,7 +207,7 @@ void UControlRigTransformNoScaleControlProxy::PostEditChangeProperty(struct FPro
 		{
 			//MUST set through ControlRig
 			FControlRigInteractionScope InteractionScope(ControlRig.Get());
-			ControlRig->SetControlValue<FTransformNoScale>(ControlName, Transform, true, EControlRigSetKey::DoNotCare,false);
+			ControlRig->SetControlValue<FRigControlValue::FTransformNoScale_Float>(ControlName, Transform, true, EControlRigSetKey::DoNotCare,false);
 			ControlRig->Evaluate_AnyThread();
 
 		}
@@ -222,7 +222,7 @@ void UControlRigTransformNoScaleControlProxy::ValueChanged()
 		Modify();
 		const FName PropertyName("Transform");
 		FTrackInstancePropertyBindings Binding(PropertyName, PropertyName.ToString());
-		const FTransformNoScale NewTransform = ControlRig.Get()->GetHierarchy()->GetControlValue(ControlElement, ERigControlValueType::Current).Get<FTransformNoScale>();
+		const FTransformNoScale NewTransform = ControlRig.Get()->GetHierarchy()->GetControlValue(ControlElement, ERigControlValueType::Current).Get<FRigControlValue::FTransformNoScale_Float>().ToTransform();
 		Binding.CallFunction<FTransformNoScale>(*this, NewTransform);
 	}
 }
@@ -238,7 +238,7 @@ void UControlRigTransformNoScaleControlProxy::PostEditUndo()
 			CheckEditModeOnSelectionChange(ControlRig.Get());
 		}
 		ControlRig->SelectControl(ControlName, bSelected);
-		ControlRig->SetControlValue<FTransformNoScale>(ControlName, Transform, true, EControlRigSetKey::Never,false);
+		ControlRig->SetControlValue<FRigControlValue::FTransformNoScale_Float>(ControlName, Transform, true, EControlRigSetKey::Never,false);
 	}
 }
 #endif
@@ -248,7 +248,7 @@ void UControlRigTransformNoScaleControlProxy::SetKey(const IPropertyHandle& Keye
 	FRigControlElement* ControlElement = GetControlElement();
 	if (ControlElement)
 	{
-		ControlRig->SetControlValue<FTransformNoScale>(ControlName, Transform, true, EControlRigSetKey::Always,false);
+		ControlRig->SetControlValue<FRigControlValue::FTransformNoScale_Float>(ControlName, Transform, true, EControlRigSetKey::Always,false);
 	}
 }
 
@@ -264,7 +264,7 @@ void UControlRigEulerTransformControlProxy::PostEditChangeProperty(struct FPrope
 		{
 			//MUST set through ControlRig
 			FControlRigInteractionScope InteractionScope(ControlRig.Get());
-			ControlRig->SetControlValue<FEulerTransform>(ControlName, Transform, true, EControlRigSetKey::DoNotCare,false);
+			ControlRig->SetControlValue<FRigControlValue::FEulerTransform_Float>(ControlName, Transform, true, EControlRigSetKey::DoNotCare,false);
 			ControlRig->Evaluate_AnyThread();
 
 		}
@@ -279,7 +279,7 @@ void UControlRigEulerTransformControlProxy::ValueChanged()
 		Modify();
 		const FName PropertyName("Transform");
 		FTrackInstancePropertyBindings Binding(PropertyName, PropertyName.ToString());
-		const FEulerTransform NewTransform = ControlRig.Get()->GetHierarchy()->GetControlValue(ControlElement, ERigControlValueType::Current).Get<FEulerTransform>();
+		const FEulerTransform NewTransform = ControlRig.Get()->GetHierarchy()->GetControlValue(ControlElement, ERigControlValueType::Current).Get<FRigControlValue::FEulerTransform_Float>().ToTransform();
 
 		Binding.CallFunction<FEulerTransform>(*this, NewTransform);
 
@@ -297,7 +297,7 @@ void UControlRigEulerTransformControlProxy::PostEditUndo()
 			CheckEditModeOnSelectionChange(ControlRig.Get());
 		}
 		ControlRig->SelectControl(ControlName, bSelected);
-		ControlRig->SetControlValue<FEulerTransform>(ControlName, Transform, true, EControlRigSetKey::Never,false);
+		ControlRig->SetControlValue<FRigControlValue::FEulerTransform_Float>(ControlName, Transform, true, EControlRigSetKey::Never,false);
 	}
 }
 #endif
@@ -307,7 +307,7 @@ void UControlRigEulerTransformControlProxy::SetKey(const IPropertyHandle& KeyedP
 	FRigControlElement* ControlElement = GetControlElement();
 	if (ControlElement)
 	{
-		ControlRig->SetControlValue<FEulerTransform>(ControlName, Transform, true, EControlRigSetKey::Always,false);
+		ControlRig->SetControlValue<FRigControlValue::FEulerTransform_Float>(ControlName, Transform, true, EControlRigSetKey::Always,false);
 	}
 }
 
@@ -506,7 +506,7 @@ void UControlRigVectorControlProxy::ValueChanged()
 		Modify();
 		const FName PropertyName("Vector");
 		FTrackInstancePropertyBindings Binding(PropertyName, PropertyName.ToString());
-		const FVector Val = ControlRig.Get()->GetHierarchy()->GetControlValue(ControlElement, ERigControlValueType::Current).Get<FVector>();
+		const FVector Val = ControlRig.Get()->GetHierarchy()->GetControlValue(ControlElement, ERigControlValueType::Current).Get<FVector3f>();
 		Binding.CallFunction<FVector>(*this, Val);
 	}
 }
@@ -522,7 +522,7 @@ void UControlRigVectorControlProxy::PostEditUndo()
 			CheckEditModeOnSelectionChange(ControlRig.Get());
 		}
 		ControlRig->SelectControl(ControlName, bSelected);
-		ControlRig->SetControlValue<FVector>(ControlName, Vector, true, EControlRigSetKey::Never,false);
+		ControlRig->SetControlValue<FVector3f>(ControlName, Vector, true, EControlRigSetKey::Never,false);
 	}
 }
 #endif
@@ -533,7 +533,7 @@ void UControlRigVectorControlProxy::SetKey(const IPropertyHandle& KeyedPropertyH
 	FRigControlElement* ControlElement = GetControlElement();
 	if (ControlElement)
 	{
-		ControlRig->SetControlValue<FVector>(ControlName, Vector, true, EControlRigSetKey::Always,false);
+		ControlRig->SetControlValue<FVector3f>(ControlName, Vector, true, EControlRigSetKey::Always,false);
 	}
 }
 
@@ -548,7 +548,7 @@ void UControlRigVector2DControlProxy::PostEditChangeProperty(struct FPropertyCha
 		{
 			//MUST set through ControlRig
 			FControlRigInteractionScope InteractionScope(ControlRig.Get());
-			ControlRig->SetControlValue<FVector2D>(ControlName, Vector2D, true, EControlRigSetKey::DoNotCare,false);
+			ControlRig->SetControlValue<FVector3f>(ControlName, FVector3f(Vector2D.X, Vector2D.Y, 0.f), true, EControlRigSetKey::DoNotCare,false);
 			ControlRig->Evaluate_AnyThread();
 
 		}
@@ -563,7 +563,8 @@ void UControlRigVector2DControlProxy::ValueChanged()
 		Modify();
 		const FName PropertyName("Vector2D");
 		FTrackInstancePropertyBindings Binding(PropertyName, PropertyName.ToString());
-		const FVector2D Val = ControlRig.Get()->GetHierarchy()->GetControlValue(ControlElement, ERigControlValueType::Current).Get<FVector2D>();
+		const FVector3f TempValue = ControlRig.Get()->GetHierarchy()->GetControlValue(ControlElement, ERigControlValueType::Current).Get<FVector3f>();
+		const FVector2D Val(TempValue.X, TempValue.Y);
 		Binding.CallFunction<FVector2D>(*this, Val);
 	}
 }
@@ -579,7 +580,7 @@ void UControlRigVector2DControlProxy::PostEditUndo()
 			CheckEditModeOnSelectionChange(ControlRig.Get());
 		}
 		ControlRig->SelectControl(ControlName, bSelected);
-		ControlRig->SetControlValue<FVector2D>(ControlName, Vector2D, true, EControlRigSetKey::Never,false);
+		ControlRig->SetControlValue<FVector3f>(ControlName, FVector3f(Vector2D.X, Vector2D.Y, 0.f), true, EControlRigSetKey::Never,false);
 	}
 }
 #endif
@@ -590,7 +591,7 @@ void UControlRigVector2DControlProxy::SetKey(const IPropertyHandle& KeyedPropert
 	FRigControlElement* ControlElement = GetControlElement();
 	if (ControlElement)
 	{
-		ControlRig->SetControlValue<FVector2D>(ControlName, Vector2D, true, EControlRigSetKey::Always,false);
+		ControlRig->SetControlValue<FVector3f>(ControlName, FVector3f(Vector2D.X, Vector2D.Y, 0.f), true, EControlRigSetKey::Always,false);
 	}
 }
 
diff --git a/Engine/Plugins/Animation/ControlRig/Source/ControlRigEditor/Private/Editor/ControlRigEditor.cpp b/Engine/Plugins/Animation/ControlRig/Source/ControlRigEditor/Private/Editor/ControlRigEditor.cpp
--- a/Engine/Plugins/Animation/ControlRig/Source/ControlRigEditor/Private/Editor/ControlRigEditor.cpp
+++ b/Engine/Plugins/Animation/ControlRig/Source/ControlRigEditor/Private/Editor/ControlRigEditor.cpp
@@ -5411,7 +5411,8 @@ void FControlRigEditor::HandleOnControlModified(UControlRig* Subject, FRigContro
 				}
 				case ERigControlType::Rotator:
 				{
-					FRotator Rotator = ControlValue.Get<FRotator>();
+					FVector3f RotatorAngles = ControlValue.Get<FVector3f>();
+					FRotator Rotator = FRotator::MakeFromEuler(RotatorAngles);
 					FRigControlValue RotatorValue = FRigControlValue::Make<FRotator>(Rotator);
 					NewDefaultValue = RotatorValue.ToString<FRotator>();
 					break;
@@ -5448,7 +5449,8 @@ void FControlRigEditor::HandleOnControlModified(UControlRig* Subject, FRigContro
 
 			if (ElementKey.Type == ERigElementType::Bone)
 			{
-				const FTransform Transform = ControlValue.Get<FTransform>() * Hierarchy->GetControlOffsetTransform(ControlElement, ERigTransformType::CurrentLocal);
+				const FTransform CurrentValue = ControlValue.Get<FRigControlValue::FTransform_Float>().ToTransform();
+				const FTransform Transform = CurrentValue * Hierarchy->GetControlOffsetTransform(ControlElement, ERigTransformType::CurrentLocal);
 				Blueprint->Hierarchy->SetLocalTransform(ElementKey, Transform);
 				Hierarchy->SetLocalTransform(ElementKey, Transform);
 
diff --git a/Engine/Plugins/Animation/ControlRig/Source/ControlRigEditor/Private/Sequencer/ControlRigParameterTrackEditor.cpp b/Engine/Plugins/Animation/ControlRig/Source/ControlRigEditor/Private/Sequencer/ControlRigParameterTrackEditor.cpp
--- a/Engine/Plugins/Animation/ControlRig/Source/ControlRigEditor/Private/Sequencer/ControlRigParameterTrackEditor.cpp
+++ b/Engine/Plugins/Animation/ControlRig/Source/ControlRigEditor/Private/Sequencer/ControlRigParameterTrackEditor.cpp
@@ -2096,7 +2096,7 @@ void FControlRigParameterTrackEditor::GetControlRigKeys(UControlRig* InControlRi
 		}
 		case ERigControlType::Vector2D:
 		{
-			FVector2D Val = ControlValue.Get<FVector2D>();
+			FVector3f Val = ControlValue.Get<FVector3f>();
 			OutGeneratedKeys.Add(FMovieSceneChannelValueSetter::Create<FMovieSceneFloatChannel>(ChannelIndex++, Val.X, bSetKey));
 			OutGeneratedKeys.Add(FMovieSceneChannelValueSetter::Create<FMovieSceneFloatChannel>(ChannelIndex++, Val.Y, bSetKey));
 			break;
@@ -2105,7 +2105,7 @@ void FControlRigParameterTrackEditor::GetControlRigKeys(UControlRig* InControlRi
 		case ERigControlType::Scale:
 		case ERigControlType::Rotator:
 		{
-			FVector3f Val = ControlValue.Get<FVector>();
+			FVector3f Val = ControlValue.Get<FVector3f>();
 			OutGeneratedKeys.Add(FMovieSceneChannelValueSetter::Create<FMovieSceneFloatChannel>(ChannelIndex++, Val.X, bSetKey));
 			OutGeneratedKeys.Add(FMovieSceneChannelValueSetter::Create<FMovieSceneFloatChannel>(ChannelIndex++, Val.Y, bSetKey));
 			OutGeneratedKeys.Add(FMovieSceneChannelValueSetter::Create<FMovieSceneFloatChannel>(ChannelIndex++, Val.Z, bSetKey));
@@ -2121,20 +2121,20 @@ void FControlRigParameterTrackEditor::GetControlRigKeys(UControlRig* InControlRi
 
 			if (ControlElement->Settings.ControlType == ERigControlType::TransformNoScale)
 			{
-				FTransformNoScale NoScale = ControlValue.Get<FTransformNoScale>();
+				FTransformNoScale NoScale = ControlValue.Get<FRigControlValue::FTransformNoScale_Float>().ToTransform();
 				Translation = NoScale.Location;
 				Rotation = NoScale.Rotation.Rotator();
 			}
 			else if (ControlElement->Settings.ControlType == ERigControlType::EulerTransform)
 			{
-				FEulerTransform Euler = ControlValue.Get < FEulerTransform >();
+				FEulerTransform Euler = ControlValue.Get<FRigControlValue::FEulerTransform_Float>().ToTransform();
 				Translation = Euler.Location;
 				Rotation = Euler.Rotation;
 				Scale = Euler.Scale;
 			}
 			else
 			{
-				FTransform Val = ControlValue.Get<FTransform>();
+				FTransform Val = ControlValue.Get<FRigControlValue::FTransform_Float>().ToTransform();
 				Translation = Val.GetTranslation();
 				Rotation = Val.GetRotation().Rotator();
 				Scale = Val.GetScale3D();
diff --git a/Engine/Plugins/Experimental/MotionTrailEditorMode/Source/MotionTrailEditorMode/Private/Sequencer/MovieSceneTransformTrail.cpp b/Engine/Plugins/Experimental/MotionTrailEditorMode/Source/MotionTrailEditorMode/Private/Sequencer/MovieSceneTransformTrail.cpp
--- a/Engine/Plugins/Experimental/MotionTrailEditorMode/Source/MotionTrailEditorMode/Private/Sequencer/MovieSceneTransformTrail.cpp
+++ b/Engine/Plugins/Experimental/MotionTrailEditorMode/Source/MotionTrailEditorMode/Private/Sequencer/MovieSceneTransformTrail.cpp
@@ -229,7 +229,9 @@ void FMovieSceneControlTransformTrail::UpdateCacheTimes(const FTrailEvaluateTime
 	check(Section);
 
 	const FRigElementKey ControlKey(ControlName, ERigElementType::Control);
-	const FTransform InitialTransform = Section->GetControlRig()->GetHierarchy()->GetControlValue(ControlKey, ERigControlValueType::Initial).Get<FTransform>();
+	const FTransform InitialTransform =
+		Section->GetControlRig()->GetHierarchy()
+		->GetControlValue(ControlKey, ERigControlValueType::Initial).Get<FRigControlValue::FTransform_Float>().ToTransform();
 
 	TArrayView<FMovieSceneFloatChannel*> FloatChannels = Section->GetChannelProxy().GetChannels<FMovieSceneFloatChannel>();
 	FloatChannels = FloatChannels.Slice(GetChannelOffset(), uint8(EMSTrailTransformChannel::MaxChannel) + 1);
@@ -249,7 +251,9 @@ UE::MovieScene::FIntermediate3DTransform FMovieSceneControlTransformTrail::Calcu
 	UMovieSceneControlRigParameterSection* Section = Cast<UMovieSceneControlRigParameterSection>(GetSection());
 
 	const FRigElementKey ControlKey(ControlName, ERigElementType::Control);
-	const FTransform InitialTransform = Section->GetControlRig()->GetHierarchy()->GetControlValue(ControlKey, ERigControlValueType::Initial).Get<FTransform>();
+	const FTransform InitialTransform =
+		Section->GetControlRig()->GetHierarchy()
+		->GetControlValue(ControlKey, ERigControlValueType::Initial).Get<FRigControlValue::FTransform_Float>().ToTransform();
 
 	const UE::MovieScene::FIntermediate3DTransform Delta = UE::MovieScene::FIntermediate3DTransform(
 		Current.GetTranslation() - Start.GetTranslation(),
