Control Rig: Additional prep work for Space Switching Demo

diff --git a/Engine/Plugins/Animation/ControlRig/Source/ControlRig/Private/Rigs/RigHierarchy.cpp b/Engine/Plugins/Animation/ControlRig/Source/ControlRig/Private/Rigs/RigHierarchy.cpp
--- a/Engine/Plugins/Animation/ControlRig/Source/ControlRig/Private/Rigs/RigHierarchy.cpp
+++ b/Engine/Plugins/Animation/ControlRig/Source/ControlRig/Private/Rigs/RigHierarchy.cpp
@@ -888,6 +888,31 @@ FRigElementWeight URigHierarchy::GetParentWeight(const FRigBaseElement* InChild,
 	return FRigElementWeight(FLT_MAX);
 }
 
+TArray<FRigElementWeight> URigHierarchy::GetParentWeightArray(FRigElementKey InChild, bool bInitial) const
+{
+	return GetParentWeightArray(Find(InChild), bInitial);
+}
+
+TArray<FRigElementWeight> URigHierarchy::GetParentWeightArray(const FRigBaseElement* InChild, bool bInitial) const
+{
+	TArray<FRigElementWeight> Weights;
+	if(const FRigMultiParentElement* MultiParentElement = Cast<FRigMultiParentElement>(InChild))
+	{
+		for(int32 ParentIndex = 0; ParentIndex < MultiParentElement->ParentConstraints.Num(); ParentIndex++)
+		{
+			if(bInitial)
+			{
+				Weights.Add(MultiParentElement->ParentConstraints[ParentIndex].InitialWeight);
+			}
+			else
+			{
+				Weights.Add(MultiParentElement->ParentConstraints[ParentIndex].Weight);
+			}
+		}
+	}
+	return Weights;
+}
+
 bool URigHierarchy::SetParentWeight(FRigElementKey InChild, FRigElementKey InParent, FRigElementWeight InWeight, bool bInitial, bool bAffectChildren)
 {
 	return SetParentWeight(Find(InChild), Find(InParent), InWeight, bInitial, bAffectChildren);
@@ -981,6 +1006,114 @@ bool URigHierarchy::SetParentWeight(FRigBaseElement* InChild, int32 InParentInde
 	return false;
 }
 
+bool URigHierarchy::SetParentWeightArray(FRigElementKey InChild, TArray<FRigElementWeight> InWeights, bool bInitial,
+	bool bAffectChildren)
+{
+	return SetParentWeightArray(Find(InChild), InWeights, bInitial, bAffectChildren);
+}
+
+bool URigHierarchy::SetParentWeightArray(FRigBaseElement* InChild, const TArray<FRigElementWeight>& InWeights,
+	bool bInitial, bool bAffectChildren)
+{
+	using namespace ERigTransformType;
+
+	if(FRigMultiParentElement* MultiParentElement = Cast<FRigMultiParentElement>(InChild))
+	{
+		if(MultiParentElement->ParentConstraints.Num() == InWeights.Num())
+		{
+			TArray<FRigElementWeight> InputWeights;
+			InputWeights.Reserve(InWeights.Num());
+
+			bool bFoundDifference = false;
+			for(int32 WeightIndex=0; WeightIndex < InWeights.Num(); WeightIndex++)
+			{
+				FRigElementWeight InputWeight = InWeights[WeightIndex];
+				InputWeight.Location = FMath::Max(InputWeight.Location, 0.f);
+				InputWeight.Rotation = FMath::Max(InputWeight.Rotation, 0.f);
+				InputWeight.Scale = FMath::Max(InputWeight.Scale, 0.f);
+				InputWeights.Add(InputWeight);
+
+				FRigElementWeight& TargetWeight = bInitial?
+					MultiParentElement->ParentConstraints[WeightIndex].InitialWeight :
+					MultiParentElement->ParentConstraints[WeightIndex].Weight;
+
+				if(!FMath::IsNearlyZero(InputWeight.Location - TargetWeight.Location) ||
+					!FMath::IsNearlyZero(InputWeight.Rotation - TargetWeight.Rotation) ||
+					!FMath::IsNearlyZero(InputWeight.Scale - TargetWeight.Scale))
+				{
+					bFoundDifference = true;
+				}
+			}
+
+			if(!bFoundDifference)
+			{
+				return false;
+			}
+			
+			const ERigTransformType::Type LocalType = bInitial ? InitialLocal : CurrentLocal;
+			const ERigTransformType::Type GlobalType = SwapLocalAndGlobal(LocalType);
+
+			if(bAffectChildren)
+			{
+				GetTransform(MultiParentElement, LocalType);
+				MultiParentElement->Pose.MarkDirty(GlobalType);
+			}
+			else
+			{
+				GetTransform(MultiParentElement, GlobalType);
+				MultiParentElement->Pose.MarkDirty(LocalType);
+			}
+
+			for(int32 WeightIndex=0; WeightIndex < InWeights.Num(); WeightIndex++)
+			{
+				if(bInitial)
+				{
+					MultiParentElement->ParentConstraints[WeightIndex].InitialWeight = InputWeights[WeightIndex];
+				}
+				else
+				{
+					MultiParentElement->ParentConstraints[WeightIndex].Weight = InputWeights[WeightIndex];
+				}
+			}
+
+			MultiParentElement->Parent.MarkDirty(GlobalType);
+
+			if(FRigControlElement* ControlElement = Cast<FRigControlElement>(MultiParentElement))
+			{
+				ControlElement->Offset.MarkDirty(GlobalType);
+			}
+
+			PropagateDirtyFlags(MultiParentElement, ERigTransformType::IsInitial(LocalType), bAffectChildren);
+			
+#if WITH_EDITOR
+			if (ensure(!bPropagatingChange))
+			{
+				TGuardValue<bool> bPropagatingChangeGuardValue(bPropagatingChange, true);
+			
+				for(FRigHierarchyListener& Listener : ListeningHierarchies)
+				{
+					if(!bForcePropagation && !Listener.ShouldReactToChange(LocalType))
+					{
+						continue;
+					}
+
+					URigHierarchy* ListeningHierarchy = Listener.Hierarchy.Get();
+					if (ListeningHierarchy)
+					{
+						if(FRigBaseElement* ListeningElement = ListeningHierarchy->Find(InChild->GetKey()))
+						{
+							ListeningHierarchy->SetParentWeightArray(ListeningElement, InWeights, bInitial, bAffectChildren);
+						}
+					}
+				}	
+			}
+#endif
+			return true;
+		}
+	}
+	return false;
+}
+
 TArray<FRigElementKey> URigHierarchy::GetAllKeys(bool bTraverse, ERigElementType InElementType) const
 {
 	TArray<FRigElementKey> Keys;
@@ -1201,6 +1334,19 @@ void URigHierarchy::SendEvent(const FRigEventContext& InEvent, bool bAsynchronou
 
 }
 
+void URigHierarchy::SendAutoKeyEvent(FRigElementKey InElement, float InOffsetInSeconds, bool bAsynchronous)
+{
+	FRigEventContext Context;
+	Context.Event = ERigEvent::RequestAutoKey;
+	Context.Key = InElement;
+	Context.LocalTime = InOffsetInSeconds;
+	if(UControlRig* Rig = Cast<UControlRig>(GetOuter()))
+	{
+		Context.LocalTime += Rig->AbsoluteTime;
+	}
+	SendEvent(Context, bAsynchronous);
+}
+
 URigHierarchyController* URigHierarchy::GetController(bool bCreateIfNeeded)
 {
 	if(LastControllerPtr.IsValid())
diff --git a/Engine/Plugins/Animation/ControlRig/Source/ControlRig/Private/Rigs/RigHierarchyController.cpp b/Engine/Plugins/Animation/ControlRig/Source/ControlRig/Private/Rigs/RigHierarchyController.cpp
--- a/Engine/Plugins/Animation/ControlRig/Source/ControlRig/Private/Rigs/RigHierarchyController.cpp
+++ b/Engine/Plugins/Animation/ControlRig/Source/ControlRig/Private/Rigs/RigHierarchyController.cpp
@@ -1032,7 +1032,7 @@ TArray<FRigElementKey> URigHierarchyController::ImportFromText(FString InContent
 
 					for(const FRigElementKey& NewParent : PerElementData.Parents)
 					{
-						AddParent(ExistingElement->GetKey(), NewParent, true, bSetupUndo);
+						AddParent(ExistingElement->GetKey(), NewParent, 0.f, true, bSetupUndo);
 					}
 				}
 				
@@ -1061,7 +1061,7 @@ TArray<FRigElementKey> URigHierarchyController::ImportFromText(FString InContent
 				Parent = *RemappedParent;
 			}
 
-			AddParent(NewElement->GetKey(), Parent, true, bSetupUndo);
+			AddParent(NewElement->GetKey(), Parent, 0.f, true, bSetupUndo);
 		}
 		
 		for(int32 ParentIndex = 0; ParentIndex < PerElementData.ParentWeights.Num(); ParentIndex++)
@@ -1778,7 +1778,7 @@ bool URigHierarchyController::RenameElement(FRigBaseElement* InElement, const FN
 	return true;
 }
 
-bool URigHierarchyController::AddParent(FRigElementKey InChild, FRigElementKey InParent, bool bMaintainGlobalTransform, bool bSetupUndo)
+bool URigHierarchyController::AddParent(FRigElementKey InChild, FRigElementKey InParent, float InWeight, bool bMaintainGlobalTransform, bool bSetupUndo)
 {
 	if(!IsValid())
 	{
@@ -1808,7 +1808,7 @@ bool URigHierarchyController::AddParent(FRigElementKey InChild, FRigElementKey I
 	}
 #endif
 
-	const bool bAdded = AddParent(Child, Parent, bMaintainGlobalTransform);
+	const bool bAdded = AddParent(Child, Parent, InWeight, bMaintainGlobalTransform);
 
 #if WITH_EDITOR
 	if(!bAdded && TransactionPtr.IsValid())
@@ -1821,7 +1821,7 @@ bool URigHierarchyController::AddParent(FRigElementKey InChild, FRigElementKey I
 	return bAdded;
 }
 
-bool URigHierarchyController::AddParent(FRigBaseElement* InChild, FRigBaseElement* InParent, bool bMaintainGlobalTransform, bool bRemoveAllParents)
+bool URigHierarchyController::AddParent(FRigBaseElement* InChild, FRigBaseElement* InParent, float InWeight, bool bMaintainGlobalTransform, bool bRemoveAllParents)
 {
 	if(InChild == nullptr || InParent == nullptr)
 	{
@@ -1892,6 +1892,8 @@ bool URigHierarchyController::AddParent(FRigBaseElement* InChild, FRigBaseElemen
 	{
 		return false;
 	}
+	Constraint.InitialWeight = InWeight;
+	Constraint.Weight = InWeight;
 
 	if(FRigSingleParentElement* SingleParentElement = Cast<FRigSingleParentElement>(InChild))
 	{
@@ -1916,23 +1918,29 @@ bool URigHierarchyController::AddParent(FRigBaseElement* InChild, FRigBaseElemen
 		const int32 ParentIndex = MultiParentElement->ParentConstraints.Add(Constraint);
 		MultiParentElement->IndexLookup.Add(Constraint.ParentElement->GetKey(), ParentIndex);
 
-		MultiParentElement->Parent.MarkDirty(ERigTransformType::CurrentGlobal);  
-		MultiParentElement->Parent.MarkDirty(ERigTransformType::InitialGlobal);
-
-		if(FRigControlElement* ControlElement = Cast<FRigControlElement>(MultiParentElement))
+		if(InWeight > SMALL_NUMBER)
 		{
-			ControlElement->Offset.MarkDirty(ERigTransformType::CurrentGlobal);  
-			ControlElement->Offset.MarkDirty(ERigTransformType::InitialGlobal);
-			ControlElement->Gizmo.MarkDirty(ERigTransformType::CurrentGlobal);  
-			ControlElement->Gizmo.MarkDirty(ERigTransformType::InitialGlobal);
+			MultiParentElement->Parent.MarkDirty(ERigTransformType::CurrentGlobal);  
+			MultiParentElement->Parent.MarkDirty(ERigTransformType::InitialGlobal);
+
+			if(FRigControlElement* ControlElement = Cast<FRigControlElement>(MultiParentElement))
+			{
+				ControlElement->Offset.MarkDirty(ERigTransformType::CurrentGlobal);  
+				ControlElement->Offset.MarkDirty(ERigTransformType::InitialGlobal);
+				ControlElement->Gizmo.MarkDirty(ERigTransformType::CurrentGlobal);  
+				ControlElement->Gizmo.MarkDirty(ERigTransformType::InitialGlobal);
+			}
 		}
 
 		Hierarchy->TopologyVersion++;
 
-		if(!bMaintainGlobalTransform)
+		if(InWeight > SMALL_NUMBER)
 		{
-			Hierarchy->PropagateDirtyFlags(MultiParentElement, true, true);
-			Hierarchy->PropagateDirtyFlags(MultiParentElement, false, true);
+			if(!bMaintainGlobalTransform)
+			{
+				Hierarchy->PropagateDirtyFlags(MultiParentElement, true, true);
+				Hierarchy->PropagateDirtyFlags(MultiParentElement, false, true);
+			}
 		}
 
 		Notify(ERigHierarchyNotification::ParentChanged, MultiParentElement);
@@ -2396,7 +2404,7 @@ bool URigHierarchyController::SetParent(FRigBaseElement* InChild, FRigBaseElemen
 	{
 		return false;
 	}
-	return AddParent(InChild, InParent, bMaintainGlobalTransform, true);
+	return AddParent(InChild, InParent, 1.f, bMaintainGlobalTransform, true);
 }
 
 void URigHierarchyController::AddElementToDirty(FRigBaseElement* InParent, FRigBaseElement* InElementToAdd, int32 InHierarchyDistance) const
diff --git a/Engine/Plugins/Animation/ControlRig/Source/ControlRig/Private/Rigs/RigHierarchyElements.cpp b/Engine/Plugins/Animation/ControlRig/Source/ControlRig/Private/Rigs/RigHierarchyElements.cpp
--- a/Engine/Plugins/Animation/ControlRig/Source/ControlRig/Private/Rigs/RigHierarchyElements.cpp
+++ b/Engine/Plugins/Animation/ControlRig/Source/ControlRig/Private/Rigs/RigHierarchyElements.cpp
@@ -328,6 +328,7 @@ void FRigMultiParentElement::CopyFrom(URigHierarchy* InHierarchy, FRigBaseElemen
 	const FRigMultiParentElement* Source = CastChecked<FRigMultiParentElement>(InOther);
 	Parent = Source->Parent;
 	ParentConstraints.Reset();
+	IndexLookup.Reset();
 
 	for(int32 ParentIndex = 0; ParentIndex < Source->ParentConstraints.Num(); ParentIndex++)
 	{
@@ -336,6 +337,7 @@ void FRigMultiParentElement::CopyFrom(URigHierarchy* InHierarchy, FRigBaseElemen
 		ParentConstraint.ParentElement = CastChecked<FRigTransformElement>(InHierarchy->Get(SourceParentElement->Index));
 		ParentConstraints.Add(ParentConstraint);
 		check(ParentConstraints[ParentIndex].ParentElement->GetKey() == SourceParentElement->GetKey());
+		IndexLookup.Add(ParentConstraint.ParentElement->GetKey(), ParentIndex);
 	}
 }
 
diff --git a/Engine/Plugins/Animation/ControlRig/Source/ControlRig/Private/Units/Math/RigUnit_MathBool.cpp b/Engine/Plugins/Animation/ControlRig/Source/ControlRig/Private/Units/Math/RigUnit_MathBool.cpp
--- a/Engine/Plugins/Animation/ControlRig/Source/ControlRig/Private/Units/Math/RigUnit_MathBool.cpp
+++ b/Engine/Plugins/Animation/ControlRig/Source/ControlRig/Private/Units/Math/RigUnit_MathBool.cpp
@@ -49,3 +49,28 @@ FRigUnit_MathBoolNotEquals_Execute()
 	Result = A != B;
 }
 
+FRigUnit_MathBoolToggled_Execute()
+{
+	DECLARE_SCOPE_HIERARCHICAL_COUNTER_RIGUNIT()
+
+	if(Context.State == EControlRigState::Init)
+	{
+		Initialized = false;
+		LastValue = Value;
+		Toggled = false;
+		return;
+	}
+
+	if(!Initialized)
+	{
+		Initialized = true;
+		Toggled = false;
+	}
+	else
+	{
+		Toggled = LastValue != Value;
+	}
+
+	LastValue = Value;
+}
+
diff --git a/Engine/Plugins/Animation/ControlRig/Source/ControlRig/Private/Units/Math/RigUnit_MathBool.h b/Engine/Plugins/Animation/ControlRig/Source/ControlRig/Private/Units/Math/RigUnit_MathBool.h
--- a/Engine/Plugins/Animation/ControlRig/Source/ControlRig/Private/Units/Math/RigUnit_MathBool.h
+++ b/Engine/Plugins/Animation/ControlRig/Source/ControlRig/Private/Units/Math/RigUnit_MathBool.h
@@ -195,3 +195,31 @@ struct CONTROLRIG_API FRigUnit_MathBoolNotEquals : public FRigUnit_MathBoolBase
 	UPROPERTY(meta=(Output))
 	bool Result;
 };
+
+/**
+* Returns true if the value has changed from the last run
+*/
+USTRUCT(meta=(DisplayName="Toggled", PrototypeName="Toggled", Keywords="Changed,Different"))
+struct CONTROLRIG_API FRigUnit_MathBoolToggled : public FRigUnit_MathBoolBase
+{
+	GENERATED_BODY()
+	RIGVM_METHOD()
+	virtual void Execute(const FRigUnitContext& Context) override;
+
+	FRigUnit_MathBoolToggled()
+	{
+		Value = Toggled = Initialized = LastValue = false;
+	}
+
+	UPROPERTY(meta=(Input))
+	bool Value;
+
+	UPROPERTY(meta=(Output))
+	bool Toggled;
+
+	UPROPERTY()
+	bool Initialized;
+
+	UPROPERTY()
+	bool LastValue;
+};
diff --git a/Engine/Plugins/Animation/ControlRig/Source/ControlRig/Public/ControlRig.h b/Engine/Plugins/Animation/ControlRig/Source/ControlRig/Public/ControlRig.h
--- a/Engine/Plugins/Animation/ControlRig/Source/ControlRig/Public/ControlRig.h
+++ b/Engine/Plugins/Animation/ControlRig/Source/ControlRig/Public/ControlRig.h
@@ -90,7 +90,11 @@ public:
 	/** Is valid for execution */
 	UFUNCTION(BlueprintPure, Category="Control Rig")
 	virtual bool CanExecute() const;
-	
+
+	/** Gets the current absolute time */
+	UFUNCTION(BlueprintPure, Category = "Control Rig")
+	float GetAbsoluteTime() const { return AbsoluteTime; }
+
 	/** Set the current delta time */
 	UFUNCTION(BlueprintCallable, Category="Control Rig")
 	void SetDeltaTime(float InDeltaTime);
diff --git a/Engine/Plugins/Animation/ControlRig/Source/ControlRig/Public/Rigs/RigHierarchy.h b/Engine/Plugins/Animation/ControlRig/Source/ControlRig/Public/Rigs/RigHierarchy.h
--- a/Engine/Plugins/Animation/ControlRig/Source/ControlRig/Public/Rigs/RigHierarchy.h
+++ b/Engine/Plugins/Animation/ControlRig/Source/ControlRig/Public/Rigs/RigHierarchy.h
@@ -1627,6 +1627,23 @@ public:
 	 */
 	FRigElementWeight GetParentWeight(const FRigBaseElement* InChild, int32 InParentIndex, bool bInitial = false) const;
 
+	/**
+	 * Returns the weights of all parents below a multi parent element
+	 * @param InChild The key of the multi parented element
+	 * @param bInitial If true the initial weights will be used
+	 * @return Returns the weight of a parent below a multi parent element, or FLT_MAX if the parent is invalid
+	 */
+	UFUNCTION(BlueprintCallable, Category = URigHierarchy)
+	TArray<FRigElementWeight> GetParentWeightArray(FRigElementKey InChild, bool bInitial = false) const;
+
+	/**
+	 * Returns the weights of all parents below a multi parent element
+	 * @param InChild The multi parented element
+	 * @param bInitial If true the initial weights will be used
+	 * @return Returns the weight of a parent below a multi parent element, or FLT_MAX if the parent is invalid
+	 */
+	TArray<FRigElementWeight> GetParentWeightArray(const FRigBaseElement* InChild, bool bInitial = false) const;
+
 	/**
 	 * Sets the weight of a parent below a multi parent element
 	 * @param InChild The key of the multi parented element
@@ -1661,6 +1678,27 @@ public:
 	 */
 	bool SetParentWeight(FRigBaseElement* InChild, int32 InParentIndex, FRigElementWeight InWeight, bool bInitial = false, bool bAffectChildren = true);
 
+	/**
+	 * Sets the all of the weights of the parents of a multi parent element
+	 * @param InChild The key of the multi parented element
+	 * @param InWeights The new weights to set for the parents
+	 * @param bInitial If true the initial weights will be used
+	 * @param bAffectChildren If set to false children will not move (maintain global).
+	 * @return Returns true if changing the weight was successful
+	 */
+	UFUNCTION(BlueprintCallable, Category = URigHierarchy)
+	bool SetParentWeightArray(FRigElementKey InChild, TArray<FRigElementWeight> InWeights, bool bInitial = false, bool bAffectChildren = true);
+
+	/**
+	 * Sets the all of the weights of the parents of a multi parent element
+	 * @param InChild The multi parented element
+	 * @param InWeights The new weights to set for the parents
+	 * @param bInitial If true the initial weights will be used
+	 * @param bAffectChildren If set to false children will not move (maintain global).
+	 * @return Returns true if changing the weight was successful
+	 */
+	bool SetParentWeightArray(FRigBaseElement* InChild,  const TArray<FRigElementWeight>& InWeights, bool bInitial = false, bool bAffectChildren = true);
+
 	/**
 	 * Returns true if an element is parented to another element
 	 * @param InChild The key of the child element to check for a parent
@@ -1765,6 +1803,15 @@ public:
 	 */
 	void SendEvent(const FRigEventContext& InEvent, bool bAsynchronous = true);
 
+	/**
+	* Sends an autokey event from the hierarchy to the world
+	* @param InElement The element to send the autokey for
+	* @param InOffsetInSeconds The time offset in seconds
+	* @param bAsynchronous If set to true the event will go on a thread safe queue
+	*/
+	UFUNCTION(BlueprintCallable, Category = URigHierarchy)
+	void SendAutoKeyEvent(FRigElementKey InElement, float InOffsetInSeconds = 0.f, bool bAsynchronous = true);
+
 	/**
 	 * Returns the delegate to listen to for events coming from this hierarchy
 	 * @return The delegate to listen to for events coming from this hierarchy
diff --git a/Engine/Plugins/Animation/ControlRig/Source/ControlRig/Public/Rigs/RigHierarchyController.h b/Engine/Plugins/Animation/ControlRig/Source/ControlRig/Public/Rigs/RigHierarchyController.h
--- a/Engine/Plugins/Animation/ControlRig/Source/ControlRig/Public/Rigs/RigHierarchyController.h
+++ b/Engine/Plugins/Animation/ControlRig/Source/ControlRig/Public/Rigs/RigHierarchyController.h
@@ -363,12 +363,13 @@ public:
 	 * Adds a new parent to an element. For elements that allow only one parent the parent will be replaced (Same as ::SetParent).
 	 * @param InChild The key of the element to add the parent for
 	 * @param InParent The key of the new parent to add
+	 * @param InWeight The initial weight to give to the parent
 	 * @param bMaintainGlobalTransform If set to true the child will stay in the same place spatially, otherwise it will maintain it's local transform (and potential move).
 	 * @param bSetupUndo If set to true the stack will record the change for undo / redo
 	 * @return Returns true if successful.
 	 */
 	UFUNCTION(BlueprintCallable, Category = URigHierarchyController)
-	bool AddParent(FRigElementKey InChild, FRigElementKey InParent, bool bMaintainGlobalTransform = true, bool bSetupUndo = false);
+	bool AddParent(FRigElementKey InChild, FRigElementKey InParent, float InWeight = 0.f, bool bMaintainGlobalTransform = true, bool bSetupUndo = false);
 
 	/**
 	 * Removes an existing parent from an element in the hierarchy. For elements that allow only one parent the element will be unparented (same as ::RemoveAllParents)
@@ -545,11 +546,12 @@ private:
 	 * Adds a new parent to an element. For elements that allow only one parent the parent will be replaced (Same as ::SetParent).
 	 * @param InChild The element to add the parent for
 	 * @param InParent The new parent to add
+	 * @param InWeight The initial weight to give to the parent
 	 * @param bMaintainGlobalTransform If set to true the child will stay in the same place spatially, otherwise it will maintain it's local transform (and potential move).
 	 * @param bRemoveAllParents If set to true all parents of the child will be removed first.
 	 * @return Returns true if successful.
 	 */
-	bool AddParent(FRigBaseElement* InChild, FRigBaseElement* InParent, bool bMaintainGlobalTransform = true, bool bRemoveAllParents = false);
+	bool AddParent(FRigBaseElement* InChild, FRigBaseElement* InParent, float InWeight = 0.f, bool bMaintainGlobalTransform = true, bool bRemoveAllParents = false);
 
 	/**
 	 * Removes an existing parent from an element in the hierarchy. For elements that allow only one parent the element will be unparented (same as ::RemoveAllParents)
diff --git a/Engine/Plugins/Animation/ControlRig/Source/ControlRig/Public/Rigs/RigHierarchyDefines.h b/Engine/Plugins/Animation/ControlRig/Source/ControlRig/Public/Rigs/RigHierarchyDefines.h
--- a/Engine/Plugins/Animation/ControlRig/Source/ControlRig/Public/Rigs/RigHierarchyDefines.h
+++ b/Engine/Plugins/Animation/ControlRig/Source/ControlRig/Public/Rigs/RigHierarchyDefines.h
@@ -1476,6 +1476,7 @@ struct CONTROLRIG_API FRigElementKeyCollection
 
 protected:
 
+	UPROPERTY(BlueprintReadWrite, EditAnywhere, Category = Collection, meta = (AllowPrivateAccess = "true"))
 	TArray<FRigElementKey> Keys;
 };
 
diff --git a/Engine/Plugins/Animation/ControlRig/Source/ControlRig/Public/Rigs/RigHierarchyElements.h b/Engine/Plugins/Animation/ControlRig/Source/ControlRig/Public/Rigs/RigHierarchyElements.h
--- a/Engine/Plugins/Animation/ControlRig/Source/ControlRig/Public/Rigs/RigHierarchyElements.h
+++ b/Engine/Plugins/Animation/ControlRig/Source/ControlRig/Public/Rigs/RigHierarchyElements.h
@@ -627,13 +627,13 @@ struct CONTROLRIG_API FRigElementWeight
 public:
 	GENERATED_BODY()
 
-	UPROPERTY()
+	UPROPERTY(BlueprintReadWrite, EditAnywhere, Category = Weight)
 	float Location;
 
-	UPROPERTY()
+	UPROPERTY(BlueprintReadWrite, EditAnywhere, Category = Weight)
 	float Rotation;
 
-	UPROPERTY()
+	UPROPERTY(BlueprintReadWrite, EditAnywhere, Category = Weight)
 	float Scale;
 
 	FORCEINLINE FRigElementWeight()
diff --git a/Engine/Plugins/Animation/ControlRig/Source/ControlRigEditor/Private/EditMode/ControlRigEditMode.cpp b/Engine/Plugins/Animation/ControlRig/Source/ControlRigEditor/Private/EditMode/ControlRigEditMode.cpp
--- a/Engine/Plugins/Animation/ControlRig/Source/ControlRigEditor/Private/EditMode/ControlRigEditMode.cpp
+++ b/Engine/Plugins/Animation/ControlRig/Source/ControlRigEditor/Private/EditMode/ControlRigEditMode.cpp
@@ -1437,10 +1437,6 @@ void FControlRigEditMode::ClearRigElementSelection(uint32 InTypes)
 	}
 
 	UControlRigBlueprint* Blueprint = Cast<UControlRigBlueprint>(InteractionRig->GetClass()->ClassGeneratedBy);
-	if (Blueprint)
-	{
-		Blueprint->GetHierarchyController()->ClearSelection();
-	}
 	if (IsInLevelEditor())
 	{
 		if(URigHierarchyController* Controller = InteractionRig->GetHierarchy()->GetController())
@@ -1448,6 +1444,10 @@ void FControlRigEditMode::ClearRigElementSelection(uint32 InTypes)
 			Controller->ClearSelection();
 		}
 	}
+	else if (Blueprint)
+	{
+		Blueprint->GetHierarchyController()->ClearSelection();
+	}
 }
 
 // internal private function that doesn't use guarding.
@@ -1460,10 +1460,6 @@ void FControlRigEditMode::SetRigElementSelectionInternal(ERigElementType Type, c
 	}
 
 	UControlRigBlueprint* Blueprint = Cast<UControlRigBlueprint>(InteractionRig->GetClass()->ClassGeneratedBy);
-	if (Blueprint)
-	{
-		Blueprint->GetHierarchyController()->SelectElement(FRigElementKey(InRigElementName, Type), bSelected);
-	}
 	if (IsInLevelEditor())
 	{
 		if(URigHierarchyController* Controller = InteractionRig->GetHierarchy()->GetController())
@@ -1471,6 +1467,10 @@ void FControlRigEditMode::SetRigElementSelectionInternal(ERigElementType Type, c
 			Controller->SelectElement(FRigElementKey(InRigElementName, Type), bSelected);
 		}
 	}
+	else if (Blueprint)
+	{
+		Blueprint->GetHierarchyController()->SelectElement(FRigElementKey(InRigElementName, Type), bSelected);
+	}
 }
 
 void FControlRigEditMode::SetRigElementSelection(ERigElementType Type, const FName& InRigElementName, bool bSelected)
