Control Rig: RigVM: Fix extraction of literal values for watch

diff --git a/Engine/Plugins/Animation/ControlRig/Source/ControlRigDeveloper/Private/ControlRigBlueprint.cpp b/Engine/Plugins/Animation/ControlRig/Source/ControlRigDeveloper/Private/ControlRigBlueprint.cpp
--- a/Engine/Plugins/Animation/ControlRig/Source/ControlRigDeveloper/Private/ControlRigBlueprint.cpp
+++ b/Engine/Plugins/Animation/ControlRig/Source/ControlRigDeveloper/Private/ControlRigBlueprint.cpp
@@ -2618,7 +2618,10 @@ void UControlRigBlueprint::HandleModifiedEvent(ERigVMGraphNotifType InNotifType,
 					bool bRequiresRecompile = false;
 
 					URigVMPin* RootPin = Pin->GetRootPin();
-					if (const FRigVMOperand* Operand = PinToOperandMap.Find(RootPin->GetPinPath()))
+					static const FString ConstSuffix = TEXT(":Const");
+					const FString PinHash = RootPin->GetPinPath(true) + ConstSuffix;
+					
+					if (const FRigVMOperand* Operand = PinToOperandMap.Find(PinHash))
 					{
 						FRigVMASTProxy RootPinProxy = FRigVMASTProxy::MakeFromUObject(RootPin);
 						if(const FRigVMExprAST* Expression = InGraph->GetRuntimeAST()->GetExprForSubject(RootPinProxy))
diff --git a/Engine/Plugins/Animation/ControlRig/Source/ControlRigEditor/Private/Graph/SControlRigGraphNode.cpp b/Engine/Plugins/Animation/ControlRig/Source/ControlRigEditor/Private/Graph/SControlRigGraphNode.cpp
--- a/Engine/Plugins/Animation/ControlRig/Source/ControlRigEditor/Private/Graph/SControlRigGraphNode.cpp
+++ b/Engine/Plugins/Animation/ControlRig/Source/ControlRigEditor/Private/Graph/SControlRigGraphNode.cpp
@@ -868,38 +868,61 @@ void SControlRigGraphNode::GetNodeInfoPopups(FNodeInfoContext* Context, TArray<F
 						PinName += UEdGraphSchema_K2::TypeToText(WatchPin->PinType).ToString();
 						PinName += TEXT(")");
 
-						FString WatchText;
-						FString PinHash = URigVMCompiler::GetPinHash(ModelPin, nullptr, true);
-						if (const FRigVMOperand* WatchOperand = RigBlueprint->PinToOperandMap.Find(PinHash))
+						TArray<FString> DefaultValues;
+					
+						if(ModelPin->GetDirection() == ERigVMPinDirection::Input ||
+							ModelPin->GetDirection() == ERigVMPinDirection::Visible)
 						{
-#if UE_RIGVM_UCLASS_BASED_STORAGE_DISABLED
-							FRigVMMemoryContainer& Memory = ActiveObject->GetVM()->GetDebugMemory();
-							TArray<FString> DefaultValues = Memory.GetRegisterValueAsString(*WatchOperand, ModelPin->GetCPPType(), ModelPin->GetCPPTypeObject());
-#else
-							URigVMMemoryStorage* Memory = ActiveObject->GetVM()->GetDebugMemory();
-							FString DebugValue = Memory->GetDataAsString(WatchOperand->GetRegisterIndex());
-							TArray<FString> DefaultValues = URigVMPin::SplitDefaultValue(DebugValue);
-#endif
-							
-							if (DefaultValues.Num() == 1)
-							{
-								WatchText = DefaultValues[0];
-							}
-							else if (DefaultValues.Num() > 1)
+							if(ModelPin->GetSourceLinks(true).Num() == 0)
 							{
-								WatchText = FString::Printf(TEXT("%s"), *FString::Join(DefaultValues, TEXT("\n")));
-							}
-							if (!WatchText.IsEmpty())
-							{
-								PinnedWatchText += FText::Format(LOCTEXT("WatchingAndValidFmt", "{0}\n\t{1}"), FText::FromString(PinName), FText::FromString(WatchText)).ToString();//@TODO: Print out object being debugged name?
+								const FString DefaultValue = ModelPin->GetDefaultValue();
+								if(ModelPin->IsArray())
+								{
+									DefaultValues = URigVMPin::SplitDefaultValue(DefaultValue);
+								}
+								else
+								{
+									DefaultValues.Add(DefaultValue);
+								}
 							}
-							else
+						}
+
+						if(DefaultValues.IsEmpty())
+						{
+							FString PinHash = URigVMCompiler::GetPinHash(ModelPin, nullptr, true);
+							if (const FRigVMOperand* WatchOperand = RigBlueprint->PinToOperandMap.Find(PinHash))
 							{
-								PinnedWatchText += FText::Format(LOCTEXT("InvalidPropertyFmt", "No watch found for {0}"), Schema->GetPinDisplayName(WatchPin)).ToString();//@TODO: Print out object being debugged name?
+#if UE_RIGVM_UCLASS_BASED_STORAGE_DISABLED
+								FRigVMMemoryContainer& Memory = ActiveObject->GetVM()->GetDebugMemory();
+								DefaultValues = Memory.GetRegisterValueAsString(*WatchOperand, ModelPin->GetCPPType(), ModelPin->GetCPPTypeObject());
+#else
+								URigVMMemoryStorage* Memory = ActiveObject->GetVM()->GetDebugMemory();
+								const FString DebugValue = Memory->GetDataAsString(WatchOperand->GetRegisterIndex());
+								DefaultValues = URigVMPin::SplitDefaultValue(DebugValue);
+#endif
 							}
+						}
+
+						FString WatchText;
 
-							ValidWatchCount++;
+						if (DefaultValues.Num() == 1)
+						{
+							WatchText = DefaultValues[0];
+						}
+						else if (DefaultValues.Num() > 1)
+						{
+							WatchText = FString::Printf(TEXT("%s"), *FString::Join(DefaultValues, TEXT("\n")));
+						}
+
+						if (!WatchText.IsEmpty())
+						{
+							PinnedWatchText += FText::Format(LOCTEXT("WatchingAndValidFmt", "{0}\n\t{1}"), FText::FromString(PinName), FText::FromString(WatchText)).ToString();//@TODO: Print out object being debugged name?
+						}
+						else
+						{
+							PinnedWatchText += FText::Format(LOCTEXT("InvalidPropertyFmt", "No watch found for {0}"), Schema->GetPinDisplayName(WatchPin)).ToString();//@TODO: Print out object being debugged name?
 						}
+						ValidWatchCount++;
 					}
 				}
 			}
