Control Rig: Fix multi edit for gizmo transform and other properties

diff --git a/Engine/Plugins/Animation/ControlRig/Source/ControlRigEditor/Private/Editor/ControlRigEditor.cpp b/Engine/Plugins/Animation/ControlRig/Source/ControlRigEditor/Private/Editor/ControlRigEditor.cpp
--- a/Engine/Plugins/Animation/ControlRig/Source/ControlRigEditor/Private/Editor/ControlRigEditor.cpp
+++ b/Engine/Plugins/Animation/ControlRig/Source/ControlRigEditor/Private/Editor/ControlRigEditor.cpp
@@ -1422,7 +1422,7 @@ void FControlRigEditor::SetDetailObjects(const TArray<UObject*>& InObjects)
 	{
 		if(UDetailsViewWrapperObject* WrapperObject = Cast<UDetailsViewWrapperObject>(FilteredObject))
 		{
-			WrapperObjects.Add(WrapperObject);
+			WrapperObjects.Add(TStrongObjectPtr<UDetailsViewWrapperObject>(WrapperObject));
 		}
 	}
 
@@ -1601,7 +1601,7 @@ void FControlRigEditor::ClearDetailObject()
 		return;
 	}
 
-	for(TWeakObjectPtr<UDetailsViewWrapperObject> WrapperObjectPtr : WrapperObjects)
+	for(const TStrongObjectPtr<UDetailsViewWrapperObject>& WrapperObjectPtr : WrapperObjects)
 	{
 		if(WrapperObjectPtr.IsValid())
 		{
@@ -4119,6 +4119,9 @@ void FControlRigEditor::OnFinishedChangingProperties(const FPropertyChangedEvent
 void FControlRigEditor::OnWrappedPropertyChangedChainEvent(UDetailsViewWrapperObject* InWrapperObject, const FString& InPropertyPath, FPropertyChangedChainEvent& InPropertyChangedChainEvent)
 {
 	check(InWrapperObject);
+	check(!WrapperObjects.IsEmpty());
+
+	TGuardValue<bool> SuspendDetailsPanelRefresh(bSuspendDetailsPanelRefresh, true);
 
 	const FString StoredStructString = TEXT("StoredStruct->");
 
@@ -4127,7 +4130,17 @@ void FControlRigEditor::OnWrappedPropertyChangedChainEvent(UDetailsViewWrapperOb
 	FString PropertyPath = InPropertyPath;
 	if(InWrapperObject->GetWrappedStruct()->IsChildOf(FRigBaseElement::StaticStruct()))
 	{
+		check(InWrapperObject->GetWrappedStruct() == WrapperObjects[0]->GetWrappedStruct());
+
 		URigHierarchy* Hierarchy = CastChecked<URigHierarchy>(InWrapperObject->GetOuter());
+		FRigBaseElement* WrappedElement = InWrapperObject->GetContent<FRigBaseElement>();
+		FRigBaseElement* FirstWrappedElement = WrapperObjects[0]->GetContent<FRigBaseElement>();
+		FRigElementKey Key = WrappedElement->GetKey();
+		if(!Hierarchy->Contains(Key))
+		{
+			return;
+		}
+		
 		check(PropertyPath.RemoveFromStart(StoredStructString));
 
 		const FString PoseString = TEXT("Pose->");
@@ -4171,7 +4184,6 @@ void FControlRigEditor::OnWrappedPropertyChangedChainEvent(UDetailsViewWrapperOb
 		};
 
 		bool bIsInitial = false;
-		FRigBaseElement* WrappedElement = InWrapperObject->GetContent<FRigBaseElement>();
 		if(PropertyPath.RemoveFromStart(PoseString))
 		{
 			const ERigTransformType::Type TransformType = Local::GetTransformTypeFromPath(PropertyPath);
@@ -4188,7 +4200,7 @@ void FControlRigEditor::OnWrappedPropertyChangedChainEvent(UDetailsViewWrapperOb
 				return;
 			}
 			
-			const FTransform Transform = CastChecked<FRigTransformElement>(WrappedElement)->Pose.Get(TransformType);
+			const FTransform Transform = CastChecked<FRigTransformElement>(FirstWrappedElement)->Pose.Get(TransformType);
 
 			if(ERigTransformType::IsLocal(TransformType) && TransformElement->IsA<FRigControlElement>())
 			{
@@ -4218,7 +4230,7 @@ void FControlRigEditor::OnWrappedPropertyChangedChainEvent(UDetailsViewWrapperOb
 			
 			ERigTransformType::Type TransformType = Local::GetTransformTypeFromPath(PropertyPath);
 			bIsInitial = bIsInitial || ERigTransformType::IsInitial(TransformType);
-			const FTransform Transform = CastChecked<FRigControlElement>(WrappedElement)->Offset.Get(TransformType);
+			const FTransform Transform = CastChecked<FRigControlElement>(FirstWrappedElement)->Offset.Get(TransformType);
 			
 			ControlRigBP->Hierarchy->SetControlOffsetTransform(ControlElement, Transform, ERigTransformType::MakeInitial(TransformType), true, true, false);
 		}
@@ -4232,7 +4244,7 @@ void FControlRigEditor::OnWrappedPropertyChangedChainEvent(UDetailsViewWrapperOb
 
 			ERigTransformType::Type TransformType = Local::GetTransformTypeFromPath(PropertyPath);
 			bIsInitial = bIsInitial || ERigTransformType::IsInitial(TransformType);
-			const FTransform Transform = CastChecked<FRigControlElement>(WrappedElement)->Gizmo.Get(TransformType);
+			const FTransform Transform = CastChecked<FRigControlElement>(FirstWrappedElement)->Gizmo.Get(TransformType);
 			
 			ControlRigBP->Hierarchy->SetControlGizmoTransform(ControlElement, Transform, ERigTransformType::MakeInitial(TransformType), true, false);
 		}
@@ -4259,17 +4271,18 @@ void FControlRigEditor::OnWrappedPropertyChangedChainEvent(UDetailsViewWrapperOb
 	
 	if(InWrapperObject->GetWrappedStruct()->IsChildOf(FRigUnit::StaticStruct()))
 	{
-		uint8* PropertyStorage = InWrapperObject->GetStructProperty()->ContainerPtrToValuePtr<uint8>(InWrapperObject);
+		check(InWrapperObject->GetWrappedStruct() == WrapperObjects[0]->GetWrappedStruct());
+		
+		uint8* FirstPropertyStorage = WrapperObjects[0]->GetStructProperty()->ContainerPtrToValuePtr<uint8>(WrapperObjects[0].Get());
 		UScriptStruct* UnitStruct = InWrapperObject->GetWrappedStruct();
 		FName RootPinName = InPropertyChangedChainEvent.PropertyChain.GetHead()->GetNextNode()->GetValue()->GetFName();
 		FProperty* TargetProperty = UnitStruct->FindPropertyByName(RootPinName);
 
 		URigVMNode* Node = CastChecked<URigVMNode>(InWrapperObject->GetOuter());
 
-		PropertyStorage += TargetProperty->GetOffset_ReplaceWith_ContainerPtrToValuePtr();
-		FString DefaultValue = FRigVMStruct::ExportToFullyQualifiedText(TargetProperty, PropertyStorage);
+		FirstPropertyStorage += TargetProperty->GetOffset_ReplaceWith_ContainerPtrToValuePtr();
+		FString DefaultValue = FRigVMStruct::ExportToFullyQualifiedText(TargetProperty, FirstPropertyStorage);
 
-		TargetProperty->ImportText(*DefaultValue, PropertyStorage, PPF_None, nullptr);
 		if (!DefaultValue.IsEmpty())
 		{
 			FString PinPath = FString::Printf(TEXT("%s.%s"), *Node->GetName(), *RootPinName.ToString());
diff --git a/Engine/Plugins/Animation/ControlRig/Source/ControlRigEditor/Private/Editor/ControlRigEditor.h b/Engine/Plugins/Animation/ControlRig/Source/ControlRigEditor/Private/Editor/ControlRigEditor.h
--- a/Engine/Plugins/Animation/ControlRig/Source/ControlRigEditor/Private/Editor/ControlRigEditor.h
+++ b/Engine/Plugins/Animation/ControlRig/Source/ControlRigEditor/Private/Editor/ControlRigEditor.h
@@ -463,7 +463,7 @@ protected:
 	bool bSuspendDetailsPanelRefresh;
 
 	TMap<int32, int32> CapsuleToHierarchyIndex;
-	TArray<TWeakObjectPtr<UDetailsViewWrapperObject>> WrapperObjects;
+	TArray<TStrongObjectPtr<UDetailsViewWrapperObject>> WrapperObjects;
 	
 	friend class FControlRigEditorMode;
 	friend class SControlRigStackView;
