Sequencer: Control Rig:  Fix for CR Component evaluation. Previously with Sequencer and CR Components we would assign the CR to a Component that the sequencer was using, in any world. This looks like it was causing issues with two worlds running with pre-forward solve events so now we only do that if in non-game, and if in game (which includes PIE), we don't re-set the CR Component CR, but instead grab the CR from it and then use that for evaluation.

Also fixes a possible issue when re-assigning a CR to a CR Component, need to make sure previous CR delegate's get unbound.

diff --git a/Engine/Plugins/Animation/ControlRig/Source/ControlRig/Private/ControlRigComponent.cpp b/Engine/Plugins/Animation/ControlRig/Source/ControlRig/Private/ControlRigComponent.cpp
--- a/Engine/Plugins/Animation/ControlRig/Source/ControlRig/Private/ControlRigComponent.cpp
+++ b/Engine/Plugins/Animation/ControlRig/Source/ControlRig/Private/ControlRigComponent.cpp
@@ -1144,6 +1144,15 @@ UControlRig* UControlRigComponent::SetupControlRigIfRequired()
 }
 void UControlRigComponent::SetControlRig(UControlRig* InControlRig)
 {
+	if (ControlRig)
+	{
+		ControlRig->OnInitialized_AnyThread().RemoveAll(this);
+		ControlRig->OnPreSetup_AnyThread().RemoveAll(this);
+		ControlRig->OnPostSetup_AnyThread().RemoveAll(this);
+		ControlRig->OnPreForwardsSolve_AnyThread().RemoveAll(this);
+		ControlRig->OnPostForwardsSolve_AnyThread().RemoveAll(this);
+		ControlRig->OnExecuted_AnyThread().RemoveAll(this);
+	}
 	ControlRig = InControlRig;
 	ControlRig->OnInitialized_AnyThread().AddUObject(this, &UControlRigComponent::HandleControlRigInitializedEvent);
 	ControlRig->OnPreSetup_AnyThread().AddUObject(this, &UControlRigComponent::HandleControlRigPreSetupEvent);
diff --git a/Engine/Plugins/Animation/ControlRig/Source/ControlRig/Private/Sequencer/MovieSceneControlRigParameterTemplate.cpp b/Engine/Plugins/Animation/ControlRig/Source/ControlRig/Private/Sequencer/MovieSceneControlRigParameterTemplate.cpp
--- a/Engine/Plugins/Animation/ControlRig/Source/ControlRig/Private/Sequencer/MovieSceneControlRigParameterTemplate.cpp
+++ b/Engine/Plugins/Animation/ControlRig/Source/ControlRig/Private/Sequencer/MovieSceneControlRigParameterTemplate.cpp
@@ -822,6 +822,17 @@ struct FControlRigParameterPreAnimatedTokenProducer : IMovieScenePreAnimatedToke
 					if (ControlRigComponent->GetControlRig() != ControlRig)
 					{
 						ControlRigComponent->Initialize();
+						/*
+						Previously with Sequencer and CR Components we would assign the CR to a Component 
+						that the sequencer was using, in any world. This looks like it was causing issues 
+						with two worlds running with pre-forward solve events so now we only do that if 
+						in non-game, and if in game (which includes PIE), we don't re-set the 
+						CR Component's CR, but instead grab the CR from it and then use that for evaluation.
+						*/
+						if (ControlRigComponent->GetWorld()->IsGameWorld() == false)
+						{
+							ControlRigComponent->SetControlRig(ControlRig);
+						}
 					}
 					else
 					{
@@ -848,7 +859,72 @@ struct FControlRigParameterPreAnimatedTokenProducer : IMovieScenePreAnimatedToke
 
 };
 
-
+static UControlRig* GetControlRig(const UMovieSceneControlRigParameterSection* Section,UObject* BoundObject)
+{
+	UControlRig* ControlRig = Section->GetControlRig();
+	if (ControlRig->GetObjectBinding())
+	{
+		if (UControlRigComponent* ControlRigComponent = Cast<UControlRigComponent>(ControlRig->GetObjectBinding()->GetBoundObject()))
+		{
+			if (AActor* Actor = Cast<AActor>(BoundObject))
+			{
+				if (UControlRigComponent* NewControlRigComponent = Actor->FindComponentByClass<UControlRigComponent>())
+				{
+					if (NewControlRigComponent->GetWorld()->IsGameWorld())
+					{
+						ControlRig = NewControlRigComponent->GetControlRig();
+						if (!ControlRig)
+						{
+							NewControlRigComponent->Initialize();
+							ControlRig = NewControlRigComponent->GetControlRig();
+						}
+						if (ControlRig->GetObjectBinding() == nullptr)
+						{
+							ControlRig->SetObjectBinding(MakeShared<FControlRigObjectBinding>());
+						}
+						if (ControlRig->GetObjectBinding()->GetBoundObject() != NewControlRigComponent)
+						{
+							ControlRig->GetObjectBinding()->BindToObject(BoundObject);
+						}
+					}
+					else if (NewControlRigComponent != ControlRigComponent)
+					{
+						NewControlRigComponent->SetControlRig(ControlRig);
+					}
+				}
+			}
+		}
+		else if (UControlRigComponent* NewControlRigComponent = Cast<UControlRigComponent>(BoundObject))
+		{
+			if (NewControlRigComponent->GetWorld()->IsGameWorld())
+			{
+				ControlRig = NewControlRigComponent->GetControlRig();
+				if (!ControlRig)
+				{
+					NewControlRigComponent->Initialize();
+					ControlRig = NewControlRigComponent->GetControlRig();
+				}
+				if (ControlRig->GetObjectBinding() == nullptr)
+				{
+					ControlRig->SetObjectBinding(MakeShared<FControlRigObjectBinding>());
+				}
+				if (ControlRig->GetObjectBinding()->GetBoundObject() != NewControlRigComponent)
+				{
+					ControlRig->GetObjectBinding()->BindToObject(BoundObject);
+				}
+			}
+			else if (NewControlRigComponent != ControlRigComponent)
+			{
+				NewControlRigComponent->SetControlRig(ControlRig);
+			}
+		}
+	}
+	else
+	{
+		return nullptr;
+	}
+	return ControlRig;
+}
 
 /* Simple token used for non-blendables*/
 struct FControlRigParameterExecutionToken : IMovieSceneExecutionToken
@@ -908,43 +984,7 @@ struct FControlRigParameterExecutionToken : IMovieSceneExecutionToken
 
 				if (ControlRig->GetObjectBinding())
 				{
-					if (UControlRigComponent* ControlRigComponent = Cast<UControlRigComponent>(ControlRig->GetObjectBinding()->GetBoundObject()))
-					{
-						if (AActor* Actor = Cast<AActor>(BoundObjects[0].Get()))
-						{
-							if (UControlRigComponent* NewControlRigComponent = Actor->FindComponentByClass<UControlRigComponent>())
-							{
-								if (NewControlRigComponent != ControlRigComponent)
-								{
-									ControlRig->GetObjectBinding()->BindToObject(BoundObjects[0].Get());
-									if (NewControlRigComponent->GetControlRig() != ControlRig)
-									{
-										NewControlRigComponent->Initialize();
-									}
-									else
-									{
-										ControlRig->Initialize();
-									}
-								}
-							}
-						}
-						else if (UControlRigComponent* NewControlRigComponent = Cast<UControlRigComponent>(BoundObjects[0].Get()))
-						{
-							if (NewControlRigComponent != ControlRigComponent)
-							{
-								ControlRig->GetObjectBinding()->BindToObject(BoundObjects[0].Get());
-								if (NewControlRigComponent->GetControlRig() != ControlRig)
-								{
-									NewControlRigComponent->Initialize();
-								}
-								else
-								{
-									ControlRig->Initialize();
-								}
-							}
-						}
-					}
-					else if (USkeletalMeshComponent* SkeletalMeshComponent = Cast<USkeletalMeshComponent>(ControlRig->GetObjectBinding()->GetBoundObject()))
+					if (USkeletalMeshComponent* SkeletalMeshComponent = Cast<USkeletalMeshComponent>(ControlRig->GetObjectBinding()->GetBoundObject()))
 					{
 						if (UControlRigLayerInstance* AnimInstance = Cast<UControlRigLayerInstance>(SkeletalMeshComponent->GetAnimInstance()))
 						{
@@ -955,6 +995,10 @@ struct FControlRigParameterExecutionToken : IMovieSceneExecutionToken
 							AnimInstance->UpdateControlRigTrack(ControlRig->GetUniqueID(), Weight, InputSettings, true);
 						}
 					}
+					else
+					{
+						ControlRig =GetControlRig(Section, BoundObjects[0].Get());
+					}
 				}
 			}		
 		}
@@ -967,18 +1011,18 @@ struct FControlRigParameterExecutionToken : IMovieSceneExecutionToken
 			bWasDoNotKey = Section->GetDoNotKey();
 			Section->SetDoNotKey(true);
 
-			if (Section->GetControlRig())
+			if (ControlRig)
 			{
 				const bool bSetupUndo = false;
-				Section->GetControlRig()->SetAbsoluteTime((float)Context.GetFrameRate().AsSeconds(Context.GetTime()));
+				ControlRig->SetAbsoluteTime((float)Context.GetFrameRate().AsSeconds(Context.GetTime()));
 				for (const FBoolParameterStringAndValue& BoolNameAndValue : BoolValues)
 				{
 					if (Section->ControlsToSet.Num() == 0 || Section->ControlsToSet.Contains(BoolNameAndValue.ParameterName))
 					{
-						FRigControlElement* RigControl = Section->GetControlRig()->FindControl(BoolNameAndValue.ParameterName);
+						FRigControlElement* RigControl = ControlRig->FindControl(BoolNameAndValue.ParameterName);
 						if (RigControl && RigControl->Settings.ControlType == ERigControlType::Bool)
 						{
-							Section->GetControlRig()->SetControlValue<bool>(BoolNameAndValue.ParameterName, BoolNameAndValue.Value, true, EControlRigSetKey::Never,bSetupUndo);
+							ControlRig->SetControlValue<bool>(BoolNameAndValue.ParameterName, BoolNameAndValue.Value, true, EControlRigSetKey::Never,bSetupUndo);
 						}
 					}
 				}
@@ -987,10 +1031,10 @@ struct FControlRigParameterExecutionToken : IMovieSceneExecutionToken
 				{
 					if (Section->ControlsToSet.Num() == 0 || Section->ControlsToSet.Contains(IntegerNameAndValue.ParameterName))
 					{
-						FRigControlElement* RigControl = Section->GetControlRig()->FindControl(IntegerNameAndValue.ParameterName);
+						FRigControlElement* RigControl = ControlRig->FindControl(IntegerNameAndValue.ParameterName);
 						if (RigControl && RigControl->Settings.ControlType == ERigControlType::Integer)
 						{
-							Section->GetControlRig()->SetControlValue<int32>(IntegerNameAndValue.ParameterName, IntegerNameAndValue.Value, true, EControlRigSetKey::Never,bSetupUndo);
+							ControlRig->SetControlValue<int32>(IntegerNameAndValue.ParameterName, IntegerNameAndValue.Value, true, EControlRigSetKey::Never,bSetupUndo);
 						}
 					}
 				}
@@ -1031,12 +1075,14 @@ struct TControlRigParameterActuatorFloat : TMovieSceneBlendingActuator<FControlR
 	{
 		const UMovieSceneControlRigParameterSection* Section = SectionData.Get();
 
-		if (Section && Section->GetControlRig())
+		UControlRig *ControlRig = GetControlRig(Section, InObject);
+
+		if (Section && ControlRig)
 		{
-			FRigControlElement* ControlElement = Section->GetControlRig()->FindControl(ParameterName);
+			FRigControlElement* ControlElement = ControlRig->FindControl(ParameterName);
 			if (ControlElement && ControlElement->Settings.ControlType == ERigControlType::Float)
 			{
-				const float Val = Section->GetControlRig()->GetHierarchy()->GetControlValue(ControlElement, ERigControlValueType::Current).Get<float>();
+				const float Val = ControlRig->GetHierarchy()->GetControlValue(ControlElement, ERigControlValueType::Current).Get<float>();
 				return FControlRigTrackTokenFloat(Val);
 			}
 
@@ -1055,12 +1101,14 @@ struct TControlRigParameterActuatorFloat : TMovieSceneBlendingActuator<FControlR
 			bWasDoNotKey = Section->GetDoNotKey();
 			Section->SetDoNotKey(true);
 
-			if (Section->GetControlRig() && (Section->ControlsToSet.Num() == 0 || Section->ControlsToSet.Contains(ParameterName)))
+			UControlRig* ControlRig = GetControlRig(Section, InObject);
+
+			if (ControlRig && (Section->ControlsToSet.Num() == 0 || Section->ControlsToSet.Contains(ParameterName)))
 			{
-				FRigControlElement* ControlElement = Section->GetControlRig()->FindControl(ParameterName);
+				FRigControlElement* ControlElement = ControlRig->FindControl(ParameterName);
 				if (ControlElement && ControlElement->Settings.ControlType == ERigControlType::Float)
 				{
-					Section->GetControlRig()->SetControlValue<float>(ParameterName, InFinalValue.Value, true, EControlRigSetKey::Never,bSetupUndo);
+					ControlRig->SetControlValue<float>(ParameterName, InFinalValue.Value, true, EControlRigSetKey::Never,bSetupUndo);
 				}
 			}
 
@@ -1097,12 +1145,14 @@ struct TControlRigParameterActuatorVector2D : TMovieSceneBlendingActuator<FContr
 	{
 		const UMovieSceneControlRigParameterSection* Section = SectionData.Get();
 
-		if (Section && Section->GetControlRig())
+		UControlRig* ControlRig = GetControlRig(Section, InObject);
+
+		if (Section && ControlRig)
 		{
-			FRigControlElement* ControlElement = Section->GetControlRig()->FindControl(ParameterName);
+			FRigControlElement* ControlElement = ControlRig->FindControl(ParameterName);
 			if (ControlElement && (ControlElement->Settings.ControlType == ERigControlType::Vector2D))
 			{
-				FVector3f Val = Section->GetControlRig()->GetHierarchy()->GetControlValue(ControlElement, ERigControlValueType::Current).Get<FVector3f>();
+				FVector3f Val = ControlRig->GetHierarchy()->GetControlValue(ControlElement, ERigControlValueType::Current).Get<FVector3f>();
 				return FControlRigTrackTokenVector2D(FVector2D(Val.X, Val.Y));
 			}
 		}
@@ -1120,13 +1170,16 @@ struct TControlRigParameterActuatorVector2D : TMovieSceneBlendingActuator<FContr
 			bWasDoNotKey = Section->GetDoNotKey();
 			Section->SetDoNotKey(true);
 
-			if (Section->GetControlRig() && (Section->ControlsToSet.Num() == 0 || Section->ControlsToSet.Contains(ParameterName)))
+			UControlRig* ControlRig = GetControlRig(Section, InObject);
+
+
+			if (ControlRig && (Section->ControlsToSet.Num() == 0 || Section->ControlsToSet.Contains(ParameterName)))
 			{
-				FRigControlElement* ControlElement = Section->GetControlRig()->FindControl(ParameterName);
+				FRigControlElement* ControlElement = ControlRig->FindControl(ParameterName);
 				if (ControlElement && (ControlElement->Settings.ControlType == ERigControlType::Vector2D))
 				{
 					const FVector3f Value(InFinalValue.Value.X, InFinalValue.Value.Y, 0.f);
-					Section->GetControlRig()->SetControlValue<FVector3f>(ParameterName, Value, true, EControlRigSetKey::Never,bSetupUndo);
+					ControlRig->SetControlValue<FVector3f>(ParameterName, Value, true, EControlRigSetKey::Never,bSetupUndo);
 				}
 			}
 
@@ -1160,13 +1213,14 @@ struct TControlRigParameterActuatorVector : TMovieSceneBlendingActuator<FControl
 	FControlRigTrackTokenVector RetrieveCurrentValue(UObject* InObject, IMovieScenePlayer* Player) const
 	{
 		const UMovieSceneControlRigParameterSection* Section = SectionData.Get();
+		UControlRig* ControlRig = GetControlRig(Section, InObject);
 
-		if (Section->GetControlRig())
+		if (ControlRig)
 		{
-			FRigControlElement* ControlElement = Section->GetControlRig()->FindControl(ParameterName);
+			FRigControlElement* ControlElement = ControlRig->FindControl(ParameterName);
 			if (ControlElement && (ControlElement->Settings.ControlType == ERigControlType::Position || ControlElement->Settings.ControlType == ERigControlType::Scale || ControlElement->Settings.ControlType == ERigControlType::Rotator))
-			{
-				FVector3f Val = Section->GetControlRig()->GetHierarchy()->GetControlValue(ControlElement, ERigControlValueType::Current).Get<FVector3f>();
+			{;
+				FVector3f Val = ControlRig->GetHierarchy()->GetControlValue(ControlElement, ERigControlValueType::Current).Get<FVector3f>();
 				return FControlRigTrackTokenVector(Val);
 			}
 		}
@@ -1183,12 +1237,14 @@ struct TControlRigParameterActuatorVector : TMovieSceneBlendingActuator<FControl
 			const bool bSetupUndo = false;
 			bWasDoNotKey = Section->GetDoNotKey();
 			Section->SetDoNotKey(true);
-			if (Section->GetControlRig() && (Section->ControlsToSet.Num() == 0 || Section->ControlsToSet.Contains(ParameterName)))
+			UControlRig* ControlRig = GetControlRig(Section, InObject);
+
+			if (ControlRig && (Section->ControlsToSet.Num() == 0 || Section->ControlsToSet.Contains(ParameterName)))
 			{
-				FRigControlElement* ControlElement = Section->GetControlRig()->FindControl(ParameterName);
+				FRigControlElement* ControlElement = ControlRig->FindControl(ParameterName);
 				if (ControlElement && (ControlElement->Settings.ControlType == ERigControlType::Position || ControlElement->Settings.ControlType == ERigControlType::Scale || ControlElement->Settings.ControlType == ERigControlType::Rotator))
 				{
-						Section->GetControlRig()->SetControlValue<FVector3f>(ParameterName, InFinalValue.Value, true, EControlRigSetKey::Never,bSetupUndo);
+						ControlRig->SetControlValue<FVector3f>(ParameterName, InFinalValue.Value, true, EControlRigSetKey::Never,bSetupUndo);
 				}
 			}
 
@@ -1223,29 +1279,25 @@ struct TControlRigParameterActuatorTransform : TMovieSceneBlendingActuator<FCont
 	{
 		const UMovieSceneControlRigParameterSection* Section = SectionData.Get();
 
-		if (Section->GetControlRig() && (Section->ControlsToSet.Num() == 0 || Section->ControlsToSet.Contains(ParameterName)))
+		UControlRig* ControlRig = GetControlRig(Section, InObject);
+
+		if (ControlRig && (Section->ControlsToSet.Num() == 0 || Section->ControlsToSet.Contains(ParameterName)))
 		{
-			FRigControlElement* ControlElement = Section->GetControlRig()->FindControl(ParameterName);
+			FRigControlElement* ControlElement = ControlRig->FindControl(ParameterName);
 			if (ControlElement && ControlElement->Settings.ControlType == ERigControlType::Transform)
 			{
-				FTransform Val = 
-					Section->GetControlRig()->GetHierarchy()
-					->GetControlValue(ControlElement, ERigControlValueType::Current).Get<FRigControlValue::FTransform_Float>().ToTransform();
+				FTransform Val = Section->GetHierarchy()->GetControlValue(ControlElement, ERigControlValueType::Current).Get<FRigControlValue::FTransform_Float>().ToTransform();
 				return FControlRigTrackTokenTransform(Val);
 			}
 			else if(ControlElement && ControlElement->Settings.ControlType == ERigControlType::TransformNoScale)
 			{
-				FTransformNoScale ValNoScale = 
-					Section->GetControlRig()->GetHierarchy()
-					->GetControlValue(ControlElement, ERigControlValueType::Current).Get<FRigControlValue::FTransformNoScale_Float>().ToTransform();
+				FTransformNoScale ValNoScale = ControlRig->GetHierarchy()->GetControlValue(ControlElement, ERigControlValueType::Current).Get<FRigControlValue::FTransformNoScale_Float>().ToTransform();
 				FTransform Val = ValNoScale;
 				return FControlRigTrackTokenTransform(Val);
 			}
 			else if (ControlElement && ControlElement->Settings.ControlType == ERigControlType::EulerTransform)
 			{
-				FEulerTransform Euler = 
-					Section->GetControlRig()->GetHierarchy()
-					->GetControlValue(ControlElement, ERigControlValueType::Current).Get<FRigControlValue::FEulerTransform_Float>().ToTransform();
+				FEulerTransform Euler = ControlRig->GetHierarchy()->GetControlValue(ControlElement, ERigControlValueType::Current).Get<FRigControlValue::FEulerTransform_Float>().ToTransform();
 				FTransform Val = Euler.ToFTransform();
 				return FControlRigTrackTokenTransform(Val);
 		}
@@ -1271,22 +1323,24 @@ struct TControlRigParameterActuatorTransform : TMovieSceneBlendingActuator<FCont
 			bWasDoNotKey = Section->GetDoNotKey();
 			Section->SetDoNotKey(true);
 			const bool bSetupUndo = false;
-			if (Section->GetControlRig() && (Section->ControlsToSet.Num() == 0 || Section->ControlsToSet.Contains(ParameterName)))
+			UControlRig* ControlRig = GetControlRig(Section, InObject);
+
+			if (ControlRig && (Section->ControlsToSet.Num() == 0 || Section->ControlsToSet.Contains(ParameterName)))
 			{
-				FRigControlElement* ControlElement = Section->GetControlRig()->FindControl(ParameterName);
+				FRigControlElement* ControlElement = ControlRig->FindControl(ParameterName);
 				if (ControlElement && ControlElement->Settings.ControlType == ERigControlType::Transform)
 				{
-						Section->GetControlRig()->SetControlValue<FRigControlValue::FTransform_Float>(ParameterName, InFinalValue.Value,true, EControlRigSetKey::Never,bSetupUndo);
+						ControlRig->SetControlValue<FRigControlValue::FTransform_Float>(ParameterName, InFinalValue.Value,true, EControlRigSetKey::Never,bSetupUndo);
 				}
 				else if (ControlElement && ControlElement->Settings.ControlType == ERigControlType::TransformNoScale)
 				{
 					FTransformNoScale NoScale = InFinalValue.Value;
-						Section->GetControlRig()->SetControlValue<FRigControlValue::FTransformNoScale_Float>(ParameterName, NoScale, true, EControlRigSetKey::Never,bSetupUndo);
+						ControlRig->SetControlValue<FRigControlValue::FTransformNoScale_Float>(ParameterName, NoScale, true, EControlRigSetKey::Never,bSetupUndo);
 				}
 				else if (ControlElement && ControlElement->Settings.ControlType == ERigControlType::EulerTransform)
 				{
 					FEulerTransform Euler = InFinalValue.Value;
-					Section->GetControlRig()->SetControlValue<FRigControlValue::FEulerTransform_Float>(ParameterName, Euler, true, EControlRigSetKey::Never,bSetupUndo);
+					ControlRig->SetControlValue<FRigControlValue::FEulerTransform_Float>(ParameterName, Euler, true, EControlRigSetKey::Never,bSetupUndo);
 				}
 			}
 
