Build Fix from Bad Merge: Control Rig

diff --git a/Engine/Plugins/Animation/ControlRig/Source/ControlRig/Private/Sequencer/MovieSceneControlRigParameterTemplate.cpp b/Engine/Plugins/Animation/ControlRig/Source/ControlRig/Private/Sequencer/MovieSceneControlRigParameterTemplate.cpp
--- a/Engine/Plugins/Animation/ControlRig/Source/ControlRig/Private/Sequencer/MovieSceneControlRigParameterTemplate.cpp
+++ b/Engine/Plugins/Animation/ControlRig/Source/ControlRig/Private/Sequencer/MovieSceneControlRigParameterTemplate.cpp
@@ -1286,7 +1286,7 @@ struct TControlRigParameterActuatorTransform : TMovieSceneBlendingActuator<FCont
 			FRigControlElement* ControlElement = ControlRig->FindControl(ParameterName);
 			if (ControlElement && ControlElement->Settings.ControlType == ERigControlType::Transform)
 			{
-				FTransform Val = Section->GetHierarchy()->GetControlValue(ControlElement, ERigControlValueType::Current).Get<FRigControlValue::FTransform_Float>().ToTransform();
+				FTransform Val = ControlRig->GetHierarchy()->GetControlValue(ControlElement, ERigControlValueType::Current).Get<FRigControlValue::FTransform_Float>().ToTransform();
 				return FControlRigTrackTokenTransform(Val);
 			}
 			else if(ControlElement && ControlElement->Settings.ControlType == ERigControlType::TransformNoScale)
