Control Rig: Follow up to 17034624 - null check

diff --git a/Engine/Plugins/Animation/ControlRig/Source/ControlRig/Private/Sequencer/MovieSceneControlRigParameterTemplate.cpp b/Engine/Plugins/Animation/ControlRig/Source/ControlRig/Private/Sequencer/MovieSceneControlRigParameterTemplate.cpp
--- a/Engine/Plugins/Animation/ControlRig/Source/ControlRig/Private/Sequencer/MovieSceneControlRigParameterTemplate.cpp
+++ b/Engine/Plugins/Animation/ControlRig/Source/ControlRig/Private/Sequencer/MovieSceneControlRigParameterTemplate.cpp
@@ -878,13 +878,16 @@ static UControlRig* GetControlRig(const UMovieSceneControlRigParameterSection* S
 							NewControlRigComponent->Initialize();
 							ControlRig = NewControlRigComponent->GetControlRig();
 						}
-						if (ControlRig->GetObjectBinding() == nullptr)
+						if (ControlRig)
 						{
-							ControlRig->SetObjectBinding(MakeShared<FControlRigObjectBinding>());
-						}
-						if (ControlRig->GetObjectBinding()->GetBoundObject() != NewControlRigComponent)
-						{
-							ControlRig->GetObjectBinding()->BindToObject(BoundObject);
+							if (ControlRig->GetObjectBinding() == nullptr)
+							{
+								ControlRig->SetObjectBinding(MakeShared<FControlRigObjectBinding>());
+							}
+							if (ControlRig->GetObjectBinding()->GetBoundObject() != NewControlRigComponent)
+							{
+								ControlRig->GetObjectBinding()->BindToObject(BoundObject);							
+							}
 						}
 					}
 					else if (NewControlRigComponent != ControlRigComponent)
@@ -904,13 +907,16 @@ static UControlRig* GetControlRig(const UMovieSceneControlRigParameterSection* S
 					NewControlRigComponent->Initialize();
 					ControlRig = NewControlRigComponent->GetControlRig();
 				}
-				if (ControlRig->GetObjectBinding() == nullptr)
+				if (ControlRig)
 				{
-					ControlRig->SetObjectBinding(MakeShared<FControlRigObjectBinding>());
-				}
-				if (ControlRig->GetObjectBinding()->GetBoundObject() != NewControlRigComponent)
-				{
-					ControlRig->GetObjectBinding()->BindToObject(BoundObject);
+					if (ControlRig->GetObjectBinding() == nullptr)
+					{
+						ControlRig->SetObjectBinding(MakeShared<FControlRigObjectBinding>());
+					}
+					if (ControlRig->GetObjectBinding()->GetBoundObject() != NewControlRigComponent)
+					{
+						ControlRig->GetObjectBinding()->BindToObject(BoundObject);
+					}
 				}
 			}
 			else if (NewControlRigComponent != ControlRigComponent)
@@ -997,7 +1003,7 @@ struct FControlRigParameterExecutionToken : IMovieSceneExecutionToken
 					}
 					else
 					{
-						ControlRig =GetControlRig(Section, BoundObjects[0].Get());
+						ControlRig = GetControlRig(Section, BoundObjects[0].Get());
 					}
 				}
 			}		
