Control Rig: Control Actor support for Control Rig Component

diff --git a/Engine/Plugins/Animation/ControlRig/Source/ControlRig/Private/ControlRigControlActor.cpp b/Engine/Plugins/Animation/ControlRig/Source/ControlRig/Private/ControlRigControlActor.cpp
--- a/Engine/Plugins/Animation/ControlRig/Source/ControlRig/Private/ControlRigControlActor.cpp
+++ b/Engine/Plugins/Animation/ControlRig/Source/ControlRig/Private/ControlRigControlActor.cpp
@@ -147,6 +147,17 @@ void AControlRigControlActor::Refresh()
 					}
 				}
 			}
+
+			if(ControlRig == nullptr)
+			{
+				UObject* Outer = RigInstance->GetOuter();
+				if(RigInstance->IsInOuter(ActorToTrack))
+				{
+					ControlRig = RigInstance;
+					RemoveUnbindDelegate();
+					break;
+				}
+			}
 		}
 
 		if (ControlRig == nullptr)
