Control Rig: Avoid dependency cycle between rig and hierarchy

diff --git a/Engine/Plugins/Animation/ControlRig/Source/ControlRig/Public/ControlRig.h b/Engine/Plugins/Animation/ControlRig/Source/ControlRig/Public/ControlRig.h
--- a/Engine/Plugins/Animation/ControlRig/Source/ControlRig/Public/ControlRig.h
+++ b/Engine/Plugins/Animation/ControlRig/Source/ControlRig/Public/ControlRig.h
@@ -207,9 +207,6 @@ public:
 	{
 		if(DynamicHierarchy != nullptr && DynamicHierarchy->GetOuter() != this)
 		{
-#if WITH_EDITOR
-			DynamicHierarchy->OnUndoRedo().RemoveAll(this);
-#endif
 			DynamicHierarchy = nullptr;
 		}
 		
@@ -221,7 +218,8 @@ public:
 				DynamicHierarchy->SetFlags(RF_Transient | RF_Transactional);
 			}
 #if WITH_EDITOR
-			DynamicHierarchy->OnUndoRedo().AddUObject(this, &UControlRig::OnHierarchyTransformUndoRedo);
+			const TWeakObjectPtr<UControlRig> WeakThis = this;
+			DynamicHierarchy->OnUndoRedo().AddStatic(&UControlRig::OnHierarchyTransformUndoRedoWeak, WeakThis);
 #endif
 		}
 		return DynamicHierarchy;
@@ -658,6 +656,16 @@ protected:
 		return InterRigSyncBracket > 0;
 	}
 
+#if WITH_EDITOR
+	FORCEINLINE static void OnHierarchyTransformUndoRedoWeak(URigHierarchy* InHierarchy, const FRigElementKey& InKey, ERigTransformType::Type InTransformType, const FTransform& InTransform, bool bIsUndo, TWeakObjectPtr<UControlRig> WeakThis)
+	{
+		if(WeakThis.IsValid() && InHierarchy != nullptr)
+		{
+			WeakThis->OnHierarchyTransformUndoRedo(InHierarchy, InKey, InTransformType, InTransform, bIsUndo);
+		}
+	}
+#endif
+	
 	void OnHierarchyTransformUndoRedo(URigHierarchy* InHierarchy, const FRigElementKey& InKey, ERigTransformType::Type InTransformType, const FTransform& InTransform, bool bIsUndo);
 
 	FFilterControlEvent OnFilterControl;
