Control Rig: Fix copy/paste nodes to different editor not working

diff --git a/Engine/Plugins/Animation/ControlRig/Source/ControlRigEditor/Private/Editor/DetailsViewWrapperObject.cpp b/Engine/Plugins/Animation/ControlRig/Source/ControlRigEditor/Private/Editor/DetailsViewWrapperObject.cpp
--- a/Engine/Plugins/Animation/ControlRig/Source/ControlRigEditor/Private/Editor/DetailsViewWrapperObject.cpp
+++ b/Engine/Plugins/Animation/ControlRig/Source/ControlRigEditor/Private/Editor/DetailsViewWrapperObject.cpp
@@ -152,7 +152,7 @@ UDetailsViewWrapperObject* UDetailsViewWrapperObject::MakeInstance(UScriptStruct
 		return nullptr;
 	}
 
-	UDetailsViewWrapperObject* Instance = NewObject<UDetailsViewWrapperObject>(InOuter, WrapperClass, NAME_None, RF_Public | RF_Transient);
+	UDetailsViewWrapperObject* Instance = NewObject<UDetailsViewWrapperObject>(InOuter, WrapperClass, NAME_None, RF_Public | RF_Transient | RF_TextExportTransient | RF_DuplicateTransient);
 	Instance->SetContent(InStructMemory, InStruct->GetStructureSize());
 	return Instance;
 }
