New "Spherical Pose Reader" control rig node.

diff --git a/Engine/Plugins/Animation/ControlRig/Source/ControlRig/Private/Units/Highlevel/Hierarchy/RigUnit_SphericalPoseReader.cpp b/Engine/Plugins/Animation/ControlRig/Source/ControlRig/Private/Units/Highlevel/Hierarchy/RigUnit_SphericalPoseReader.cpp
new file mode 100644
--- /dev/null
+++ b/Engine/Plugins/Animation/ControlRig/Source/ControlRig/Private/Units/Highlevel/Hierarchy/RigUnit_SphericalPoseReader.cpp
@@ -0,0 +1,230 @@
+// Copyright Epic Games, Inc. All Rights Reserved.
+
+#include "RigUnit_SphericalPoseReader.h"
+#include "Kismet/KismetMathLibrary.h"
+#include "Units/RigUnitContext.h"
+
+FRigUnit_SphericalPoseReader_Execute()
+{
+	DECLARE_SCOPE_HIERARCHICAL_COUNTER_RIGUNIT()
+
+	URigHierarchy* Hierarchy = Context.Hierarchy;
+	if (Hierarchy == nullptr)
+	{
+		return;
+	}
+
+	if (Context.State == EControlRigState::Init)
+	{
+		return;
+	}
+
+	if (!DriverItem.IsValid())
+	{
+		return;
+	}
+
+	// remap/clamp inputs
+	RemapAndConvertInputs(
+		InnerRegion, OuterRegion,
+		ActiveRegionSize, PositiveWidth, NegativeWidth, PositiveHeight, NegativeHeight,
+		FalloffSize, PositiveWidthFalloff, NegativeWidthFalloff, PositiveHeightFalloff, NegativeHeightFalloff);
+	
+	// get the world offset
+	const FTransform LocalDriverTransformInit = Hierarchy->GetInitialLocalTransform(DriverItem);
+	const FTransform GlobalDriverParentTransform = Hierarchy->GetParentTransform(DriverItem);
+	FTransform WorldOffset = LocalDriverTransformInit * GlobalDriverParentTransform;
+	// rotate by optional static offset
+	FQuat RotationOffsetQuat = FQuat::MakeFromEuler(RotationOffset);
+	WorldOffset.SetRotation(WorldOffset.GetRotation()*RotationOffsetQuat);
+
+	// get driver axis
+	const FTransform GlobalDriverTransform = Hierarchy->GetGlobalTransform(DriverItem);
+	const FVector CurrentGlobalDriverAxis = GlobalDriverTransform.GetRotation().RotateVector(DriverAxis);
+	DriverNormal = WorldOffset.InverseTransformVectorNoScale(CurrentGlobalDriverAxis).GetSafeNormal();
+
+	// evaluate the interpolation of the output param
+	const float AngleFromForward = FMath::Acos(FVector::DotProduct(DriverNormal, FVector::ZAxisVector));
+	const float Mag = RemapRange(AngleFromForward, 0.0f, PI, 0.0f, 1.0f);
+	Driver2D = DriverNormal;
+	if (Mag > SMALL_NUMBER)
+	{
+		Driver2D.Z = 0.0f;
+		Driver2D.Normalize();
+		Driver2D *= Mag;
+	}
+	Driver2D.Z = -1.0f;
+
+	if (FMath::IsNearlyZero(Mag))
+	{
+		// avoid NaNs from DistanceToEllipse, and guaranteed to be inside inner ellipse
+		OutputParam = 1.0f;
+		Debug.DrawDebug(WorldOffset, Context.DrawInterface, InnerRegion, OuterRegion, DriverNormal, Driver2D, OutputParam);
+		return;
+	}
+
+	const float PointX = Driver2D.X;
+	const float PointY = Driver2D.Y;
+
+	float EllipseWidth;
+	float EllipseHeight;
+
+	// query inner ellipse
+	InnerRegion.GetEllipseWidthAndHeight(PointX, PointY, EllipseWidth, EllipseHeight);
+	FEllipseQuery InnerEllipseResults;
+	DistanceToEllipse(PointX, PointY, EllipseWidth, EllipseHeight, InnerEllipseResults);
+
+	// query outer ellipse
+	OuterRegion.GetEllipseWidthAndHeight(PointX, PointY, EllipseWidth, EllipseHeight);
+	FEllipseQuery OuterEllipseResults;
+	DistanceToEllipse(PointX, PointY, EllipseWidth, EllipseHeight, OuterEllipseResults);
+
+	// calc output param
+	OutputParam = CalcOutputParam(InnerEllipseResults, OuterEllipseResults);
+	
+	// do all debug drawing
+	Debug.DrawDebug(WorldOffset, Context.DrawInterface, InnerRegion, OuterRegion, DriverNormal, Driver2D, OutputParam);
+}
+
+void FRigUnit_SphericalPoseReader::RemapAndConvertInputs(
+	FSphericalRegion& InnerRegion,
+	FSphericalRegion& OuterRegion,
+	const float ActiveRegionSize,
+	const float PositiveWidth,
+	const float NegativeWidth,
+	const float PositiveHeight,
+	const float NegativeHeight,
+	const float FalloffSize,
+	const float PositiveWidthFalloff,
+	const float NegativeWidthFalloff,
+	const float PositiveHeightFalloff,
+	const float NegativeHeightFalloff)
+{
+	// remap normalized inputs to angles
+	float RegionAngle = ActiveRegionSize * 180.0f;
+	RegionAngle = FMath::Min(FMath::Max(0.5f, RegionAngle), 178.0f);
+
+	InnerRegion.RegionAngleRadians = FMath::DegreesToRadians(RegionAngle);
+	InnerRegion.PosWidth = PositiveWidth;
+	InnerRegion.NegWidth = NegativeWidth;
+	InnerRegion.PosHeight = PositiveHeight;
+	InnerRegion.NegHeight = NegativeHeight;
+
+	// clamp outer falloff angle to always be greater than inner
+	float FalloffAngle = RegionAngle + (FalloffSize * 180.0f);
+	FalloffAngle = FMath::Min(FMath::Max(1.0f, FalloffAngle), 179.0f);
+	OuterRegion.RegionAngleRadians = FMath::DegreesToRadians(FalloffAngle);
+
+	// clamp falloff to always be outside inner angle
+	const float InvOuterAngleRadians = 1.0f / OuterRegion.RegionAngleRadians;
+	//
+	const float PosWidthMin = (InnerRegion.RegionAngleRadians * PositiveWidth) * InvOuterAngleRadians;
+	OuterRegion.PosWidth = FMath::Lerp(PosWidthMin, 1.0f, PositiveWidthFalloff);
+	//
+	const float NegWidthMin = (InnerRegion.RegionAngleRadians * NegativeWidth) * InvOuterAngleRadians;
+	OuterRegion.NegWidth = FMath::Lerp(NegWidthMin, 1.0f, NegativeWidthFalloff);
+	//
+	const float PosHeightMin = (InnerRegion.RegionAngleRadians * PositiveHeight) * InvOuterAngleRadians;
+	OuterRegion.PosHeight = FMath::Lerp(PosHeightMin, 1.0f, PositiveHeightFalloff);
+	//
+	const float NegHeightMin = (InnerRegion.RegionAngleRadians * NegativeHeight) * InvOuterAngleRadians;
+	OuterRegion.NegHeight = FMath::Lerp(NegHeightMin, 1.0f, NegativeHeightFalloff);
+}
+
+float FRigUnit_SphericalPoseReader::CalcOutputParam(
+	const FEllipseQuery& InnerEllipseResults,
+	const FEllipseQuery& OuterEllipseResults)
+{
+    if (InnerEllipseResults.IsInside)
+    {
+        return 1.0f; // inside inner ellipse
+    }
+
+    if (!OuterEllipseResults.IsInside)
+    {
+        return 0.0f; // outside outer ellipse
+    }
+
+    // we're between outer and inner ellipse, calculate falloff
+    const float DistanceToOuter = FMath::Sqrt(OuterEllipseResults.DistSq);
+    const float DistanceToInner = FMath::Sqrt(InnerEllipseResults.DistSq);
+    const float TotalDistance = DistanceToInner + DistanceToOuter;
+    if (TotalDistance < 0.0001f)
+    {
+        return 0.0f; // don't lerp when outer is VERY close to inner (avoid div by zero)
+    }
+
+    return 1.0f - (DistanceToInner / TotalDistance);
+}
+
+void FRigUnit_SphericalPoseReader::DistanceToEllipse(
+	const float InX,
+	const float InY,
+	const float SizeX,
+	const float SizeY,
+	FEllipseQuery& OutEllipseQuery)
+{
+	if (SizeX <= KINDA_SMALL_NUMBER || SizeY <= KINDA_SMALL_NUMBER)
+	{
+		return; // degenerate ellipse
+	}
+	
+    const float PX = FMath::Abs(InX);
+    const float PY = FMath::Abs(InY);
+
+	const float SizeXSq = SizeX * SizeX;
+	const float SizeYSq = SizeY * SizeY;
+
+	const float InvSizeX = 1.0f / SizeX;
+	const float InvSizeY = 1.0f / SizeY;
+	
+    float TX = 0.70710678118f;
+    float TY = 0.70710678118f;
+
+    const int Iterations = 2; // this could be higher for greatly quality
+    for (int i=0; i<Iterations; ++i)
+    {
+        const float ScaledX = SizeX * TX;
+        const float ScaledY = SizeY * TY;
+
+        const float EX = (SizeXSq - SizeYSq) * (TX * TX * TX) * InvSizeX;
+        const float EY = (SizeYSq - SizeXSq) * (TY * TY * TY) * InvSizeY;
+
+        const float RX = ScaledX - EX;
+        const float RY = ScaledY - EY;
+
+        const float QX = PX - EX;
+        const float QY = PY - EY;
+
+        const float R = FMath::Sqrt(RX * RX + RY * RY);
+        const float Q = FMath::Sqrt(QY * QY + QX * QX);
+
+        TX = FMath::Min(1.f, FMath::Max(0.f, (QX * R / Q + EX) * InvSizeX));
+        TY = FMath::Min(1.f, FMath::Max(0.f, (QY * R / Q + EY) * InvSizeY));
+
+        const float InvT = 1.0f / FMath::Sqrt(TX * TX + TY * TY);
+
+        TX *= InvT;
+        TY *= InvT;
+    }
+
+    OutEllipseQuery.ClosestX = SizeX * (InX < 0 ? -TX : TX);
+    OutEllipseQuery.ClosestY = SizeY * (InY < 0 ? -TY : TY);
+    const float ToClosestX = OutEllipseQuery.ClosestX - InX;
+    const float ToClosestY = OutEllipseQuery.ClosestY - InY;
+    OutEllipseQuery.DistSq = FMath::Pow(ToClosestX, 2.f) + FMath::Pow(ToClosestY, 2.f);
+    const float CenterToClosestDistSq = FMath::Pow(OutEllipseQuery.ClosestX, 2.f) + FMath::Pow(OutEllipseQuery.ClosestY, 2.f);
+    const float CenterToInputDistSq = FMath::Pow(InX, 2.f) + FMath::Pow(InY, 2.f);
+    OutEllipseQuery.IsInside = CenterToClosestDistSq > CenterToInputDistSq;
+}
+
+float FRigUnit_SphericalPoseReader::RemapRange(
+	const float T,
+	const float AStart,
+	const float AEnd,
+	const float BStart,
+	const float BEnd)
+{
+	check(FMath::Abs(AEnd - AStart) > 0.0f);
+	return BStart + (T - AStart) * (BEnd - BStart) / (AEnd - AStart);
+}
diff --git a/Engine/Plugins/Animation/ControlRig/Source/ControlRig/Private/Units/Highlevel/Hierarchy/RigUnit_SphericalPoseReader.h b/Engine/Plugins/Animation/ControlRig/Source/ControlRig/Private/Units/Highlevel/Hierarchy/RigUnit_SphericalPoseReader.h
new file mode 100644
--- /dev/null
+++ b/Engine/Plugins/Animation/ControlRig/Source/ControlRig/Private/Units/Highlevel/Hierarchy/RigUnit_SphericalPoseReader.h
@@ -0,0 +1,367 @@
+// Copyright Epic Games, Inc. All Rights Reserved.
+
+#pragma once
+
+#include "Drawing/ControlRigDrawInterface.h"
+#include "Units/Highlevel/RigUnit_HighlevelBase.h"
+#include "RigUnit_SphericalPoseReader.generated.h"
+
+
+struct FEllipseQuery
+{
+	float ClosestX;
+	float ClosestY;
+	float DistSq;
+	bool IsInside;
+};
+
+USTRUCT()
+struct CONTROLRIG_API FSphericalRegion
+{
+	GENERATED_BODY()
+
+	FSphericalRegion()
+	: RegionAngleRadians(0.1f),
+	PosWidth(1.0f),
+	NegWidth(1.0f),
+	PosHeight(1.0f),
+	NegHeight(1.0f)
+	{
+	}
+	
+	float RegionAngleRadians;
+	float PosWidth;
+	float NegWidth;
+	float PosHeight;
+	float NegHeight;
+
+	void GetEllipseWidthAndHeight(const float PointX, const float PointY, float& OutWidth, float& OutHeight)
+	{
+		OutWidth = ((PointX > 0 ? PosWidth : NegWidth) * RegionAngleRadians) * INV_PI;
+		OutHeight = ((PointY > 0 ? PosHeight : NegHeight) * RegionAngleRadians) * INV_PI;
+	}
+};
+
+USTRUCT()
+struct CONTROLRIG_API FSphericalPoseReaderDebugSettings
+{
+	GENERATED_BODY()
+
+	FSphericalPoseReaderDebugSettings()
+		: bDrawDebug(true),
+		bDraw2D(false),
+		DebugScale(25.0f),
+		DebugSegments(20),
+		DebugThickness(0.25f)
+	{
+	}
+
+	UPROPERTY(meta = (Input))
+	bool bDrawDebug = true;
+
+	UPROPERTY(meta = (Input))
+	bool bDraw2D = false;
+	
+	UPROPERTY(meta = (Input, ClampMin = "0.0", UIMin = "0.01", UIMax = "100.0"))
+	float DebugScale = 25.0f;
+	
+	UPROPERTY(meta = (Input, ClampMin = "0", ClampMax = "50", UIMin = "0", UIMax = "10"))
+	int DebugSegments = 20;
+
+	UPROPERTY(meta = (Input, ClampMin = "0", ClampMax = "50", UIMin = "0.01", UIMax = "2.0"))
+	float DebugThickness = 0.25f;
+	
+	void DrawDebug(
+		const FTransform WorldOffset,
+		FControlRigDrawInterface* DrawInterface,
+		const FSphericalRegion& InnerRegion,
+		const FSphericalRegion& OuterRegion,
+		const FVector DriverNormal,
+		const FVector Driver2D,
+		const float OutputParam) const
+    {
+		if (DrawInterface == nullptr || !bDrawDebug)
+		{
+			return;
+		}
+
+		// make colors
+		FLinearColor DriverColor = FLinearColor::Red;
+		FLinearColor Yellow = FLinearColor::Yellow;
+		FLinearColor Green = FLinearColor::Green;
+		FLinearColor DarkYellow = FLinearColor::Yellow * .4f;
+		FLinearColor DarkGreen = FLinearColor::Green * .2f;
+		FLinearColor Blue = FLinearColor::Blue;
+		if (OutputParam < SMALL_NUMBER)
+		{
+			// outside range
+			DriverColor = FLinearColor::Red;
+			Yellow = DarkYellow;
+			Green = DarkGreen;
+		}else if (OutputParam >= 0.99f)
+		{
+			// in green range
+			DriverColor = Blue;
+			Yellow = DarkYellow;
+		}else
+		{
+			// in yellow range
+			DriverColor = Blue;
+			Green = DarkGreen;
+		}
+		
+        // draw inner manifold
+        TArray<FVector> InnerConePoints;
+        TArray<FVector> FlatInnerConePoints;
+        DrawManifold(InnerConePoints, FlatInnerConePoints, InnerRegion);
+		if (bDraw2D)
+		{
+			// flat
+			DrawInterface->DrawLineStrip(WorldOffset, FlatInnerConePoints, Green, DebugThickness);
+		}else
+		{
+			// spherical
+			DrawInterface->DrawLineStrip(WorldOffset, InnerConePoints, Green, DebugThickness);
+		}
+		
+        // draw outer manifold
+        TArray<FVector> OuterConePoints;
+        TArray<FVector> FlatOuterConePoints;
+        DrawManifold(OuterConePoints, FlatOuterConePoints, OuterRegion);
+		if (bDraw2D)
+		{
+			// flat
+			DrawInterface->DrawLineStrip(WorldOffset, FlatOuterConePoints, Yellow, DebugThickness);
+		}else
+		{
+			// spherical
+			DrawInterface->DrawLineStrip(WorldOffset, OuterConePoints, Yellow, DebugThickness);
+		}
+
+		if (!bDraw2D)
+		{
+			// draw connector lines between inner/outer
+			for (int I = 0; I < InnerConePoints.Num(); ++I)
+			{
+				for (int P = 0; P < DebugSegments; ++P)
+				{
+					float StartPercent = P / (float)DebugSegments;
+					float EndPercent = (P + 1) / (float)DebugSegments;
+					FVector StartPoint = FMath::Lerp(InnerConePoints[I], OuterConePoints[I], StartPercent).GetSafeNormal();
+					FVector EndPoint = FMath::Lerp(InnerConePoints[I], OuterConePoints[I], EndPercent).GetSafeNormal();
+					DrawInterface->DrawLine(WorldOffset, StartPoint * DebugScale, EndPoint * DebugScale, Yellow, DebugThickness);
+				}
+			}
+		
+			// draw connector lines from active region to origin
+			for (int I = 0; I < InnerConePoints.Num(); ++I)
+			{
+				for (int P = 0; P < DebugSegments; ++P)
+				{
+					float StartPercent = P / (float)DebugSegments;
+					float EndPercent = (P + 1) / (float)DebugSegments;
+					FVector StartPoint = FMath::Lerp(FVector::UpVector * DebugScale, InnerConePoints[I], StartPercent).GetSafeNormal();
+					FVector EndPoint = FMath::Lerp(FVector::UpVector * DebugScale, InnerConePoints[I], EndPercent).GetSafeNormal();
+					DrawInterface->DrawLine(WorldOffset, StartPoint * DebugScale, EndPoint * DebugScale, Green, DebugThickness);
+				}
+			}
+		}
+
+		if (bDraw2D)
+		{
+			// draw driver vector
+			DrawInterface->DrawPoint(WorldOffset, Driver2D * DebugScale, DebugThickness*5.0f, DriverColor);
+		}else
+		{
+			// draw driver vector
+			DrawInterface->DrawLine(WorldOffset, FVector::ZeroVector, DriverNormal * DebugScale, DriverColor, DebugThickness*2.0f);
+			DrawInterface->DrawPoint(WorldOffset, DriverNormal * DebugScale, DebugThickness*5.0f, DriverColor);
+		}   
+    }
+	
+	void DrawManifold(TArray<FVector>& SphericalPoints, TArray<FVector>& FlatPoints, const FSphericalRegion& Region) const
+	{
+		const float AngleStep = (PI*2.0f) / FMath::Max(4.0f, (static_cast<float>(DebugSegments) - 1.0f));
+		for (int i = 0; i < DebugSegments; ++i)
+		{
+			// spin around forward vector, drawing each span
+			float SpinAngle = (AngleStep * i);
+			// get 2d projection in normalized space 0-1
+			FVector PointOnPlane = CalcSphericalCoordinates(SpinAngle, Region);
+			// get 3d spherical point
+			FVector PointOnSphere = SphericalCoordinatesToNormal(PointOnPlane);
+
+			// store points
+			PointOnPlane.Z = -1.0f; // move 2d debug view to base of sphere
+			SphericalPoints.Add(PointOnSphere * DebugScale);
+			FlatPoints.Add(PointOnPlane * DebugScale);
+		}
+	}
+
+	static FVector CalcSphericalCoordinates(const float SpinAngleRadians, const FSphericalRegion& Region)
+	{
+		// uniform angle
+		const FQuat SpinRot = FQuat(FVector::ZAxisVector, SpinAngleRadians);
+		const FVector SpanAxis = SpinRot * FVector::XAxisVector;
+
+		// scale span axis length proportional to amount of rotation towards back of sphere
+		FVector Point = SpanAxis * (Region.RegionAngleRadians / PI);
+
+		// stretch cone point non-uniformly
+		const float Width = Point.X > 0 ? Region.PosWidth : Region.NegWidth;
+		const float Height = Point.Y > 0 ? Region.PosHeight : Region.NegHeight;
+		Point.X *= Width;
+		Point.Y *= Height;
+		Point.Z = 0.0f;
+
+		return Point;
+	}
+
+	static float SignedAngleBetweenNormals(const FVector& From, const FVector& To, const FVector& Axis)
+	{
+		const float FromDotTo = FVector::DotProduct(From, To);
+		const float Angle = FMath::Acos(FromDotTo);
+		const FVector Cross = FVector::CrossProduct(From, To);
+		const float Dot = FVector::DotProduct(Cross, Axis);
+		return Dot >= 0 ? Angle : -Angle;
+	}
+
+	static FVector SphericalCoordinatesToNormal(FVector Point2D)
+	{
+		const float AngleToSquashedPoint = SignedAngleBetweenNormals(Point2D.GetSafeNormal(), FVector::XAxisVector, -FVector::ZAxisVector);
+		const FQuat SquashedSpanRot = FQuat(FVector::ZAxisVector, AngleToSquashedPoint);
+		const FVector SquashedSpanAxis = SquashedSpanRot * FVector::YAxisVector;
+		const float NormalizedMag = Point2D.Size();
+		const FQuat ConeRot = FQuat(SquashedSpanAxis, (NormalizedMag * PI));
+		const FVector Point3D = ConeRot * FVector::ZAxisVector;
+		return Point3D;
+	}
+};
+
+USTRUCT(meta=(DisplayName="Spherical Pose Reader", Category="Hierarchy", Keywords="Pose Reader"))
+struct CONTROLRIG_API FRigUnit_SphericalPoseReader: public FRigUnit_HighlevelBaseMutable
+{
+	GENERATED_BODY()
+
+	FRigUnit_SphericalPoseReader()
+			: OutputParam(0.0f),
+			DriverItem(FRigElementKey(NAME_None, ERigElementType::Bone)),
+			DriverAxis(FVector(1.0f,0.0f,0.0f)),
+			RotationOffset(FVector(90.0f,0.0f,0.0f)),
+			ActiveRegionSize(0.1f),
+			PositiveWidth(1.0f),
+			NegativeWidth(1.0f),
+			PositiveHeight(1.0f),
+			NegativeHeight(1.0f),
+			FalloffSize(0.2f),
+			PositiveWidthFalloff(1.0f),
+			NegativeWidthFalloff(1.0f),
+			PositiveHeightFalloff(1.0f),
+			NegativeHeightFalloff(1.0f),
+			Debug(FSphericalPoseReaderDebugSettings()),
+			InnerRegion(FSphericalRegion()),
+			OuterRegion(FSphericalRegion()),
+			DriverNormal(FVector::Zero()),
+			Driver2D(FVector::Zero())
+	{
+	}
+	
+	RIGVM_METHOD()
+    virtual void Execute(const FRigUnitContext& Context) override;
+
+	static void RemapAndConvertInputs(
+		FSphericalRegion& InnerRegion,
+		FSphericalRegion& OuterRegion,
+		const float ActiveRegionSize,
+		const float PositiveWidth,
+		const float NegativeWidth,
+		const float PositiveHeight,
+		const float NegativeHeight,
+		const float FalloffSize,
+		const float PositiveWidthFalloff,
+		const float NegativeWidthFalloff,
+		const float PositiveHeightFalloff,
+		const float NegativeHeightFalloff);
+	static void EvaluateInterpolation(FRigUnit_SphericalPoseReader& Node, const FTransform& DriverTransform);
+	
+	static void DrawDebug(FRigUnit_SphericalPoseReader& PoseReader);
+
+	static float CalcOutputParam(
+		const FEllipseQuery& InnerEllipseResults,
+		const FEllipseQuery& OuterEllipseResults);
+	
+	static void DistanceToEllipse(
+		const float InX,
+		const float InY,
+		const float SizeX,
+		const float SizeY,
+		FEllipseQuery& OutEllipseQuery);
+	
+	static float RemapRange(
+		const float T,
+		const float AStart,
+		const float AEnd,
+		const float BStart,
+		const float BEnd);
+
+	// The normalized output parameter; ranges from 0 (when outside yellow region) to 1 (in the green region) and smoothly blends from 0-1 in the yellow region.
+	UPROPERTY(meta = (Output))
+	float OutputParam;
+	
+	// Rotate the entire falloff region to align with the desired area of effect.
+	UPROPERTY(meta = (Input, ExpandByDefault))
+    FRigElementKey DriverItem;
+
+	// The axis of the driver transform that is compared against the falloff regions. Typically the axis that is pointing at the child; usually X by convention. Default is X-axis (1,0,0).
+	UPROPERTY(meta = (Input))
+	FVector DriverAxis = FVector(1.0f,0.0f,0.0f);
+
+	// Rotate the entire falloff region to align with the desired area of effect.
+	UPROPERTY(meta = (Input))
+	FVector RotationOffset = FVector(90.0f,0.0f,0.0f);
+
+	// The size of the active region (green) that outputs the full value (1.0). Range is 0-1. Default is 0.1.
+	UPROPERTY(meta = (Input, ClampMin = "0.0", ClampMax = "1.0", UIMin = "0.0", UIMax = "1.0"))
+    float ActiveRegionSize = 0.1f;
+	// Scale the active region in the POSITIVE width direction. Range is 0-1. Default is 1.0.
+	UPROPERTY(meta = (Input, ClampMin = "0.0", ClampMax = "1.0", UIMin = "0.0", UIMax = "1.0"))
+    float PositiveWidth = 1.0f;
+	// Scale the active region in the NEGATIVE width direction. Range is 0-1. Default is 1.0.
+	UPROPERTY(meta = (Input, ClampMin = "0.0", ClampMax = "1.0", UIMin = "0.0", UIMax = "1.0"))
+    float NegativeWidth = 1.0f;
+	// Scale the active region in the POSITIVE height direction. Range is 0-1. Default is 1.0.
+    UPROPERTY(meta = (Input, ClampMin = "0.0", ClampMax = "1.0", UIMin = "0.0", UIMax = "1.0"))
+    float PositiveHeight = 1.0f;
+	// Scale the active region in the NEGATIVE height direction. Range is 0-1. Default is 1.0.
+    UPROPERTY(meta = (Input, ClampMin = "0.0", ClampMax = "1.0", UIMin = "0.0", UIMax = "1.0"))
+    float NegativeHeight = 1.0f;
+
+	// The size of the falloff region (yellow) that defines the start of the output range. A value of 1 wraps the entire sphere with falloff. Range is 0-1. Default is 0.2.
+	UPROPERTY(meta = (Input, ClampMin = "0.0", ClampMax = "1.0", UIMin = "0.0", UIMax = "1.0"))
+    float FalloffSize = 0.2f;
+	// Scale the falloff region in the POSITIVE width direction. Range is 0-1. Default is 1.0.
+	UPROPERTY(meta = (Input, ClampMin = "0.0", ClampMax = "1.0", UIMin = "0.0", UIMax = "1.0"))
+    float PositiveWidthFalloff = 1.0f;
+	// Scale the falloff region in the NEGATIVE width direction. Range is 0-1. Default is 1.0.
+	UPROPERTY(meta = (Input, ClampMin = "0.0", ClampMax = "1.0", UIMin = "0.0", UIMax = "1.0"))
+    float NegativeWidthFalloff = 1.0f;
+	// Scale the falloff region in the POSITIVE height direction. Range is 0-1. Default is 1.0.
+	UPROPERTY(meta = (Input, ClampMin = "0.0", ClampMax = "1.0", UIMin = "0.0", UIMax = "1.0"))
+    float PositiveHeightFalloff = 1.0f;
+	// Scale the falloff region in the NEGATIVE height direction. Range is 0-1. Default is 1.0.
+	UPROPERTY(meta = (Input, ClampMin = "0.0", ClampMax = "1.0", UIMin = "0.0", UIMax = "1.0"))
+    float NegativeHeightFalloff = 1.0f;
+
+	UPROPERTY(meta = (Input))
+	FSphericalPoseReaderDebugSettings Debug;
+
+	// private work data
+	UPROPERTY(transient)
+	FSphericalRegion InnerRegion;
+	UPROPERTY(transient)
+	FSphericalRegion OuterRegion;
+	UPROPERTY(transient)
+	FVector DriverNormal;
+	UPROPERTY(transient)
+	FVector Driver2D;
+};
\ No newline at end of file
