Control Rig: Fix against crashes when inspect memory tabs are open

diff --git a/Engine/Plugins/Animation/ControlRig/Source/ControlRigEditor/Private/Editor/ControlRigEditor.cpp b/Engine/Plugins/Animation/ControlRig/Source/ControlRigEditor/Private/Editor/ControlRigEditor.cpp
--- a/Engine/Plugins/Animation/ControlRig/Source/ControlRigEditor/Private/Editor/ControlRigEditor.cpp
+++ b/Engine/Plugins/Animation/ControlRig/Source/ControlRigEditor/Private/Editor/ControlRigEditor.cpp
@@ -90,6 +90,7 @@
 #include "Components/StaticMeshComponent.h"
 #include "ControlRig/Private/Units/Execution/RigUnit_SequenceExecution.h"
 #include "ControlRigContextMenuContext.h"
+#include "Types/ISlateMetaData.h"
 
 #define LOCTEXT_NAMESPACE "ControlRigEditor"
 
@@ -1302,6 +1303,22 @@ UBlueprint* FControlRigEditor::GetBlueprintObj() const
 	return nullptr;
 }
 
+#if !UE_RIGVM_UCLASS_BASED_STORAGE_DISABLED
+
+class FMemoryTypeMetaData : public ISlateMetaData
+{
+public:
+	SLATE_METADATA_TYPE(FMemoryTypeMetaData, ISlateMetaData)
+
+		FMemoryTypeMetaData(ERigVMMemoryType InMemoryType)
+		: MemoryType(InMemoryType)
+	{
+	}
+	ERigVMMemoryType MemoryType;
+};
+
+#endif
+
 void FControlRigEditor::SetDetailObjects(const TArray<UObject*>& InObjects)
 {
 	if(bSuspendDetailsPanelRefresh)
@@ -1323,14 +1340,21 @@ void FControlRigEditor::SetDetailObjects(const TArray<UObject*>& InObjects)
 			TSharedRef<IDetailsView> DetailsView = EditModule.CreateDetailView( DetailsViewArgs );
 			TSharedRef<SDockTab> DockTab = SNew(SDockTab)
 			.Label( LOCTEXT("ControlRigMemoryDetails", "Control Rig Memory Details") )
+			.AddMetaData<FMemoryTypeMetaData>(FMemoryTypeMetaData(Memory->GetMemoryType()))
 			.TabRole(ETabRole::NomadTab)
 			[
 				DetailsView
 			];
 
+			FName TabId = *FString::Printf(TEXT("ControlRigMemoryDetails_%d"), (int32)Memory->GetMemoryType());
+			if(TSharedPtr<SDockTab> ActiveTab = GetTabManager()->FindExistingLiveTab(TabId))
+			{
+				ActiveTab->RequestCloseTab();
+			}
+
 			GetTabManager()->InsertNewDocumentTab(
 				FBlueprintEditorTabs::DetailsID,
-				TEXT("ControlRigMemoryDetails"),
+				TabId,
 				FTabManager::FLastMajorOrNomadTab(TEXT("ControlRigMemoryDetails")),
 				DockTab
 			);
@@ -2611,6 +2635,29 @@ void FControlRigEditor::HandleVMCompiledEvent(UBlueprint* InBlueprint, URigVM* I
 		RigBlueprint->CompileLog.Messages.Reset();
 		RigBlueprint->CompileLog.NumErrors = RigBlueprint->CompileLog.NumWarnings = 0;  
 	}
+
+#if !UE_RIGVM_UCLASS_BASED_STORAGE_DISABLED
+	
+	TArray<FName> TabIds;
+	TabIds.Add(*FString::Printf(TEXT("ControlRigMemoryDetails_%d"), (int32)ERigVMMemoryType::Literal));
+	TabIds.Add(*FString::Printf(TEXT("ControlRigMemoryDetails_%d"), (int32)ERigVMMemoryType::Work));
+	TabIds.Add(*FString::Printf(TEXT("ControlRigMemoryDetails_%d"), (int32)ERigVMMemoryType::Debug));
+
+	for (const FName& TabId : TabIds)
+	{
+		TSharedPtr<SDockTab> ActiveTab = GetTabManager()->FindExistingLiveTab(TabId);
+		if(ActiveTab)
+		{
+			if(ActiveTab->GetMetaData<FMemoryTypeMetaData>().IsValid())
+			{
+				ERigVMMemoryType MemoryType = ActiveTab->GetMetaData<FMemoryTypeMetaData>()->MemoryType;			
+				TSharedRef<IDetailsView> DetailsView = StaticCastSharedRef<IDetailsView>(ActiveTab->GetContent());
+				DetailsView->SetObject(InVM->GetMemoryByType(MemoryType));
+			}
+		}
+	}
+	
+#endif
 }
 
 void FControlRigEditor::HandleControlRigExecutedEvent(UControlRig* InControlRig, const EControlRigState InState, const FName& InEventName)
