Sequencer: Control Rig: FK Control Rig Improvements, filter based off of selection and add hit testing of displayed hierarchy.

We Make it easier to show active channels on Control Rigs by letting user use selected controls to filter the section. Also fixed issues where filter/mask would get lost when leaving the section so we only reset the mask on new control rigs or if number of controls has changed. Also maks now work with non-zero default control rigs and additive sections.

For hit testing of the FK rig, we add some hit proxies to get set when we draw the bone hierarchy.

diff --git a/Engine/Plugins/Animation/ControlRig/Source/ControlRig/Private/Sequencer/MovieSceneControlRigParameterSection.cpp b/Engine/Plugins/Animation/ControlRig/Source/ControlRig/Private/Sequencer/MovieSceneControlRigParameterSection.cpp
--- a/Engine/Plugins/Animation/ControlRig/Source/ControlRig/Private/Sequencer/MovieSceneControlRigParameterSection.cpp
+++ b/Engine/Plugins/Animation/ControlRig/Source/ControlRig/Private/Sequencer/MovieSceneControlRigParameterSection.cpp
@@ -1722,6 +1722,7 @@ float UMovieSceneControlRigParameterSection::GetTotalWeightValue(FFrameTime InTi
 
 void UMovieSceneControlRigParameterSection::RecreateWithThisControlRig(UControlRig* InControlRig, bool bSetDefault)
 {
+	bool bSameControlRig = (ControlRig == InControlRig);
 	SetControlRig(InControlRig);
 	/* Don't delete old tracks but eventually show that they aren't associated.. but
 	then how to delete?
@@ -1734,10 +1735,15 @@ void UMovieSceneControlRigParameterSection::RecreateWithThisControlRig(UControlR
 	ColorParameterNamesAndCurves.Empty();
 	TransformParameterNamesAndCurves.Empty();
 	*/
-	TArray<bool> OnArray;
-	OnArray.Init(true, ControlRig->AvailableControls().Num());
-	SetControlsMask(OnArray);
 
+	//if we had the same with same number of controls keep the mask otherwise reset it.
+	if (!bSameControlRig || ControlRig->AvailableControls().Num() != ControlsMask.Num())
+	{
+		TArray<bool> OnArray;
+		OnArray.Init(true, ControlRig->AvailableControls().Num());
+		SetControlsMask(OnArray);
+	}
+	
 	TArray<FRigControlElement*> SortedControls;
 	ControlRig->GetControlsInOrder(SortedControls);
 
diff --git a/Engine/Plugins/Animation/ControlRig/Source/ControlRig/Private/Sequencer/MovieSceneControlRigParameterTemplate.cpp b/Engine/Plugins/Animation/ControlRig/Source/ControlRig/Private/Sequencer/MovieSceneControlRigParameterTemplate.cpp
--- a/Engine/Plugins/Animation/ControlRig/Source/ControlRig/Private/Sequencer/MovieSceneControlRigParameterTemplate.cpp
+++ b/Engine/Plugins/Animation/ControlRig/Source/ControlRig/Private/Sequencer/MovieSceneControlRigParameterTemplate.cpp
@@ -1508,84 +1508,138 @@ void FMovieSceneControlRigParameterTemplate::EvaluateCurvesWithMasks(const FMovi
 		// Populate each of the output arrays in turn
 		for (int32 Index = 0; Index < Scalars.Num(); ++Index)
 		{
+			const FScalarParameterNameAndCurve& Scalar = this->Scalars[Index];
 			float Value = 0;
 
-			const FScalarParameterNameAndCurve& Scalar = this->Scalars[Index];
 			if (HACK_ChannelMasks->ScalarCurveMask[Index])
 			{
 				Scalar.ParameterCurve.Evaluate(Time, Value);
 			}
-			
+			else
+			{
+				Value = (Section->GetBlendType().Get() == EMovieSceneBlendType::Additive
+					|| Scalar.ParameterCurve.GetDefault().IsSet() == false) ? 0.0f :
+					Scalar.ParameterCurve.GetDefault().GetValue();
+			}
+
 			Values.ScalarValues.Emplace(Scalar.ParameterName, Value);
 		}
 
 		for (int32 Index = 0; Index < Bools.Num(); ++Index)
 		{
 			bool Value = false;
-
 			const FBoolParameterNameAndCurve& Bool = Bools[Index];
 			if (HACK_ChannelMasks->BoolCurveMask[Index])
 			{
 				Bool.ParameterCurve.Evaluate(Time, Value);
 			}
+			else
+			{
+				Value = (Section->GetBlendType().Get() == EMovieSceneBlendType::Additive
+					|| Bool.ParameterCurve.GetDefault().IsSet() == false) ? false :
+					Bool.ParameterCurve.GetDefault().GetValue();
+			}
 
 			Values.BoolValues.Emplace(Bool.ParameterName, Value);
 		}
 		for (int32 Index = 0; Index < Integers.Num(); ++Index)
 		{
-			int32 Value = 0;
-
+			int32  Value = 0;
 			const FIntegerParameterNameAndCurve& Integer = Integers[Index];
 			if (HACK_ChannelMasks->IntegerCurveMask[Index])
 			{
 				Integer.ParameterCurve.Evaluate(Time, Value);
 			}
+			else
+			{
+				Value = (Section->GetBlendType().Get() == EMovieSceneBlendType::Additive
+					|| Integer.ParameterCurve.GetDefault().IsSet() == false) ? 0 :
+					Integer.ParameterCurve.GetDefault().GetValue();
+			}
 
 			Values.IntegerValues.Emplace(Integer.ParameterName, Value);
 		}
 		for (int32 Index = 0; Index < Enums.Num(); ++Index)
 		{
-			uint8 Value = 0;
-
+			uint8  Value = 0;
 			const FEnumParameterNameAndCurve& Enum = Enums[Index];
 			if (HACK_ChannelMasks->EnumCurveMask[Index])
 			{
 				Enum.ParameterCurve.Evaluate(Time, Value);
 			}
+			else
+			{
+				Value = (Section->GetBlendType().Get() == EMovieSceneBlendType::Additive
+					|| Enum.ParameterCurve.GetDefault().IsSet() == false) ? 0 :
+					Enum.ParameterCurve.GetDefault().GetValue();
+
+			}
 			Values.IntegerValues.Emplace(Enum.ParameterName, (int32)Value);
 
 		}
 		for (int32 Index = 0; Index < Vector2Ds.Num(); ++Index)
 		{
 			FVector2D Value(ForceInitToZero);
-
 			const FVector2DParameterNameAndCurves& Vector2D = Vector2Ds[Index];
+
 			if (HACK_ChannelMasks->Vector2DCurveMask[Index])
 			{
 				Vector2D.XCurve.Evaluate(Time, Value.X);
 				Vector2D.YCurve.Evaluate(Time, Value.Y);
 			}
+			else
+			{
+				if (Section->GetBlendType().Get() != EMovieSceneBlendType::Additive)
+				{
+					if (Vector2D.XCurve.GetDefault().IsSet())
+					{
+						Value.X = Vector2D.XCurve.GetDefault().GetValue();
+					}
+					if (Vector2D.YCurve.GetDefault().IsSet())
+					{
+						Value.Y = Vector2D.YCurve.GetDefault().GetValue();
+					}
+				}
+			}
 
 			Values.Vector2DValues.Emplace(Vector2D.ParameterName, Value);
 		}
+
 		for (int32 Index = 0; Index < Vectors.Num(); ++Index)
 		{
-			FVector3f Value(ForceInitToZero);
-
+			FVector Value(ForceInitToZero);
 			const FVectorParameterNameAndCurves& Vector = Vectors[Index];
+
 			if (HACK_ChannelMasks->VectorCurveMask[Index])
 			{
 				Vector.XCurve.Evaluate(Time, Value.X);
 				Vector.YCurve.Evaluate(Time, Value.Y);
 				Vector.ZCurve.Evaluate(Time, Value.Z);
 			}
+			else
+			{
+				if (Section->GetBlendType().Get() != EMovieSceneBlendType::Additive)
+				{
+					if (Vector.XCurve.GetDefault().IsSet())
+					{
+						Value.X = Vector.XCurve.GetDefault().GetValue();
+					}
+					if (Vector.YCurve.GetDefault().IsSet())
+					{
+						Value.Y = Vector.YCurve.GetDefault().GetValue();
+					}
+					if (Vector.ZCurve.GetDefault().IsSet())
+					{
+						Value.Z = Vector.ZCurve.GetDefault().GetValue();
+					}
+				}
+			}
 
 			Values.VectorValues.Emplace(Vector.ParameterName, Value);
 		}
 		for (int32 Index = 0; Index < Colors.Num(); ++Index)
 		{
 			FLinearColor ColorValue = FLinearColor::White;
-
 			const FColorParameterNameAndCurves& Color = Colors[Index];
 			if (HACK_ChannelMasks->ColorCurveMask[Index])
 			{
@@ -1594,7 +1648,29 @@ void FMovieSceneControlRigParameterTemplate::EvaluateCurvesWithMasks(const FMovi
 				Color.BlueCurve.Evaluate(Time, ColorValue.B);
 				Color.AlphaCurve.Evaluate(Time, ColorValue.A);
 			}
-		
+			else
+			{
+				if (Section->GetBlendType().Get() != EMovieSceneBlendType::Additive)
+				{
+					if (Color.RedCurve.GetDefault().IsSet())
+					{
+						ColorValue.R = Color.RedCurve.GetDefault().GetValue();
+					}
+					if (Color.GreenCurve.GetDefault().IsSet())
+					{
+						ColorValue.G = Color.GreenCurve.GetDefault().GetValue();
+					}
+					if (Color.BlueCurve.GetDefault().IsSet())
+					{
+						ColorValue.B = Color.BlueCurve.GetDefault().GetValue();
+					}
+					if (Color.AlphaCurve.GetDefault().IsSet())
+					{
+						ColorValue.A = Color.AlphaCurve.GetDefault().GetValue();
+					}
+				}
+			}
+
 			Values.ColorValues.Emplace(Color.ParameterName, ColorValue);
 		}
 
@@ -1602,6 +1678,11 @@ void FMovieSceneControlRigParameterTemplate::EvaluateCurvesWithMasks(const FMovi
 		for (int32 Index = 0; Index < Transforms.Num(); ++Index)
 		{
 			FVector3f Translation(ForceInitToZero), Scale(FVector3f::OneVector);
+			if (Section->GetBlendType().Get() == EMovieSceneBlendType::Additive)
+			{
+				Scale = FVector3f(0.0f, 0.0f, 0.0f);
+			}
+
 			FRotator Rotator(0.0f, 0.0f, 0.0f);
 
 			const FTransformParameterNameAndCurves& Transform = Transforms[Index];
@@ -1611,39 +1692,151 @@ void FMovieSceneControlRigParameterTemplate::EvaluateCurvesWithMasks(const FMovi
 				{
 					Transform.Translation[0].Evaluate(Time, Translation[0]);
 				}
+				else
+				{
+					if (Section->GetBlendType().Get() != EMovieSceneBlendType::Additive && Transform.Translation[0].GetDefault().IsSet())
+					{
+						Translation[0] = Transform.Translation[0].GetDefault().GetValue();
+					}
+				}
 				if (EnumHasAllFlags(ChannelMask, EMovieSceneTransformChannel::TranslationY))
 				{
 					Transform.Translation[1].Evaluate(Time, Translation[1]);
 				}
+				else
+				{
+					if (Section->GetBlendType().Get() != EMovieSceneBlendType::Additive && Transform.Translation[1].GetDefault().IsSet())
+					{
+						Translation[1] = Transform.Translation[1].GetDefault().GetValue();
+					}
+				}
 				if (EnumHasAllFlags(ChannelMask, EMovieSceneTransformChannel::TranslationZ))
 				{
 					Transform.Translation[2].Evaluate(Time, Translation[2]);
 				}
+				else
+				{
+					if (Section->GetBlendType().Get() != EMovieSceneBlendType::Additive && Transform.Translation[2].GetDefault().IsSet())
+					{
+						Translation[2] = Transform.Translation[2].GetDefault().GetValue();
+					}
+				}
 				if (EnumHasAllFlags(ChannelMask, EMovieSceneTransformChannel::RotationX))
 				{
 					Transform.Rotation[0].Evaluate(Time, Rotator.Roll);
 				}
+				else
+				{
+					if (Section->GetBlendType().Get() != EMovieSceneBlendType::Additive && Transform.Rotation[0].GetDefault().IsSet())
+					{
+						Rotator.Roll = Transform.Rotation[0].GetDefault().GetValue();
+					}
+				}
 				if (EnumHasAllFlags(ChannelMask, EMovieSceneTransformChannel::RotationY))
 				{
 					Transform.Rotation[1].Evaluate(Time, Rotator.Pitch);
 				}
+				else
+				{
+					if (Section->GetBlendType().Get() != EMovieSceneBlendType::Additive && Transform.Rotation[1].GetDefault().IsSet())
+					{
+						Rotator.Pitch = Transform.Rotation[1].GetDefault().GetValue();
+					}
+				}
 				if (EnumHasAllFlags(ChannelMask, EMovieSceneTransformChannel::RotationZ))
 				{
 					Transform.Rotation[2].Evaluate(Time, Rotator.Yaw);
 				}
+				else
+				{
+					if (Section->GetBlendType().Get() != EMovieSceneBlendType::Additive && Transform.Rotation[2].GetDefault().IsSet())
+					{
+						Rotator.Yaw = Transform.Rotation[2].GetDefault().GetValue();
+					}
+				}
 				//mz todo quat interp...
 				if (EnumHasAllFlags(ChannelMask, EMovieSceneTransformChannel::ScaleX))
 				{
 					Transform.Scale[0].Evaluate(Time, Scale[0]);
 				}
+				else
+				{
+					if (Section->GetBlendType().Get() != EMovieSceneBlendType::Additive && Transform.Scale[0].GetDefault().IsSet())
+					{
+						Scale[0] = Transform.Scale[0].GetDefault().GetValue();
+					}
+				}
 				if (EnumHasAllFlags(ChannelMask, EMovieSceneTransformChannel::ScaleY))
 				{
 					Transform.Scale[1].Evaluate(Time, Scale[1]);
 				}
+				else
+				{
+					if (Section->GetBlendType().Get() != EMovieSceneBlendType::Additive && Transform.Scale[1].GetDefault().IsSet())
+					{
+						Scale[1] = Transform.Scale[1].GetDefault().GetValue();
+					}
+				}
 				if (EnumHasAllFlags(ChannelMask, EMovieSceneTransformChannel::ScaleZ))
 				{
 					Transform.Scale[2].Evaluate(Time, Scale[2]);
 				}
+				else
+				{
+					if (Section->GetBlendType().Get() != EMovieSceneBlendType::Additive && Transform.Scale[2].GetDefault().IsSet())
+					{
+						Scale[2] = Transform.Scale[2].GetDefault().GetValue();
+					}
+				}
+			}
+			else //completely masked use default or zeroed, which is already set if additive
+			{
+				if (Section->GetBlendType().Get() != EMovieSceneBlendType::Additive)
+				{
+
+					if (Transform.Translation[0].GetDefault().IsSet())
+					{
+						Translation[0] = Transform.Translation[0].GetDefault().GetValue();
+					}
+					if (Transform.Translation[1].GetDefault().IsSet())
+					{
+						Translation[1] = Transform.Translation[1].GetDefault().GetValue();
+					}
+					if (Transform.Translation[2].GetDefault().IsSet())
+					{
+						Translation[2] = Transform.Translation[2].GetDefault().GetValue();
+					}
+
+					if (Transform.Rotation[0].GetDefault().IsSet())
+					{
+						Rotator.Roll = Transform.Rotation[0].GetDefault().GetValue();
+					}
+
+					if (Transform.Rotation[1].GetDefault().IsSet())
+					{
+						Rotator.Pitch = Transform.Rotation[1].GetDefault().GetValue();
+					}
+
+					if (Transform.Rotation[2].GetDefault().IsSet())
+					{
+						Rotator.Yaw = Transform.Rotation[2].GetDefault().GetValue();
+					}
+
+					if (Transform.Scale[0].GetDefault().IsSet())
+					{
+						Scale[0] = Transform.Scale[0].GetDefault().GetValue();
+					}
+
+					if (Transform.Scale[1].GetDefault().IsSet())
+					{
+						Scale[1] = Transform.Scale[1].GetDefault().GetValue();
+					}
+
+					if (Transform.Scale[2].GetDefault().IsSet())
+					{
+						Scale[2] = Transform.Scale[2].GetDefault().GetValue();
+					}
+				}
 			}
 			FTransformParameterStringAndValue NameAndValue(Transform.ParameterName, Translation, Rotator, Scale);
 			Values.TransformValues.Emplace(NameAndValue);
diff --git a/Engine/Plugins/Animation/ControlRig/Source/ControlRigEditor/Private/EditMode/ControlRigEditMode.cpp b/Engine/Plugins/Animation/ControlRig/Source/ControlRigEditor/Private/EditMode/ControlRigEditMode.cpp
--- a/Engine/Plugins/Animation/ControlRig/Source/ControlRigEditor/Private/EditMode/ControlRigEditMode.cpp
+++ b/Engine/Plugins/Animation/ControlRig/Source/ControlRigEditor/Private/EditMode/ControlRigEditMode.cpp
@@ -53,6 +53,9 @@
 #include "Editor/UnrealEdEngine.h"
 #include "Settings/ControlRigSettings.h"
 #include "ToolMenus.h"
+#include "Sequencer/MovieSceneControlRigParameterTrack.h"
+#include "Sequencer/MovieSceneControlRigParameterSection.h"
+
 
 void UControlRigEditModeDelegateHelper::OnPoseInitialized()
 {
@@ -513,6 +516,101 @@ void FControlRigEditMode::Tick(FEditorViewportClient* ViewportClient, float Delt
 		}
 	}
 }
+//Hit proxy for FK Rigs and bones.
+struct  HFKRigBoneProxy : public HHitProxy
+{
+	DECLARE_HIT_PROXY()
+
+	FName BoneName;
+	UControlRig* ControlRig;
+
+	HFKRigBoneProxy()
+		: HHitProxy(HPP_Foreground)
+		, BoneName(NAME_None)
+		, ControlRig(nullptr)
+	{}
+
+	HFKRigBoneProxy(FName InBoneName, UControlRig *InControlRig)
+		: HHitProxy(HPP_Foreground)
+		, BoneName(InBoneName)
+		, ControlRig(InControlRig)
+	{
+	}
+
+	// HHitProxy interface
+	virtual EMouseCursor::Type GetMouseCursor() override { return EMouseCursor::Crosshairs; }
+	// End of HHitProxy interface
+};
+
+IMPLEMENT_HIT_PROXY(HFKRigBoneProxy, HHitProxy)
+
+
+TSet<FName> FControlRigEditMode::GetActiveControlsFromSequencer(UControlRig* ControlRig)
+{
+	TSet<FName> ActiveControls;
+	if (WeakSequencer.IsValid() == false)
+	{
+		return ActiveControls;
+	}
+	if (TSharedPtr<IControlRigObjectBinding> ObjectBinding = ControlRig->GetObjectBinding())
+	{
+		USceneComponent* Component = Cast<USceneComponent>(ObjectBinding->GetBoundObject());
+		if (!Component)
+		{
+			return ActiveControls;
+		}
+		const bool bCreateHandleIfMissing = false;
+		FName CreatedFolderName = NAME_None;
+		FGuid ObjectHandle = WeakSequencer.Pin()->GetHandleToObject(Component, bCreateHandleIfMissing);
+		if (!ObjectHandle.IsValid())
+		{
+			UObject* ActorObject = Component->GetOwner();
+			ObjectHandle = WeakSequencer.Pin()->GetHandleToObject(ActorObject, bCreateHandleIfMissing);
+			if (!ObjectHandle.IsValid())
+			{
+				return ActiveControls;
+			}
+		}
+		bool bCreateTrack = false;
+		UMovieScene* MovieScene = WeakSequencer.Pin()->GetFocusedMovieSceneSequence()->GetMovieScene();
+		if (!MovieScene)
+		{
+			return ActiveControls;
+		}
+		if (FMovieSceneBinding* Binding = MovieScene->FindBinding(ObjectHandle))
+		{
+			for (UMovieSceneTrack* Track : Binding->GetTracks())
+			{
+				if (UMovieSceneControlRigParameterTrack* ControlRigParameterTrack = Cast<UMovieSceneControlRigParameterTrack>(Track))
+				{
+					if (ControlRigParameterTrack->GetControlRig() == ControlRig)
+					{
+						UMovieSceneControlRigParameterSection* ActiveSection = Cast<UMovieSceneControlRigParameterSection>(ControlRigParameterTrack->GetSectionToKey());
+						if (ActiveSection)
+						{
+							TArray<FRigControlElement*> Controls;
+							ControlRig->GetControlsInOrder(Controls);
+							TArray<bool> Mask = ActiveSection->GetControlsMask();
+
+							TArray<FName> Names;
+							int Index = 0;
+							for (FRigControlElement* ControlElement : Controls)
+							{
+								if (Mask[Index])
+								{
+									ActiveControls.Add(ControlElement->GetName());
+								}
+								++Index;
+							}
+						}
+					}
+				}
+			}
+		}
+	}
+	return ActiveControls;
+}
+
 
 void FControlRigEditMode::Render(const FSceneView* View, FViewport* Viewport, FPrimitiveDrawInterface* PDI)
 {	
@@ -539,7 +637,14 @@ void FControlRigEditMode::Render(const FSceneView* View, FViewport* Viewport, FP
 		URigHierarchy* Hierarchy = ControlRig->GetHierarchy();
 		if (Settings->bDisplayHierarchy)
 		{
-			Hierarchy->ForEach<FRigTransformElement>([PDI, Hierarchy, ComponentTransform](FRigTransformElement* TransformElement) -> bool
+			const bool bHasFKRig = (ControlRig->IsA<UAdditiveControlRig>() || ControlRig->IsA<UFKControlRig>());
+			const bool bBoolSetHitProxies = PDI && PDI->IsHitTesting() && bHasFKRig;
+			TSet<FName> ActiveControlName;
+			if (bHasFKRig)
+			{
+				ActiveControlName = GetActiveControlsFromSequencer(ControlRig);
+			}
+			Hierarchy->ForEach<FRigTransformElement>([PDI, Hierarchy, ComponentTransform,ControlRig, bHasFKRig, bBoolSetHitProxies, ActiveControlName](FRigTransformElement* TransformElement) -> bool
 			{
 				const FTransform Transform = Hierarchy->GetTransform(TransformElement, ERigTransformType::CurrentGlobal);
 
@@ -548,12 +653,57 @@ void FControlRigEditMode::Render(const FSceneView* View, FViewport* Viewport, FP
 				{
 					if(FRigTransformElement* ParentTransformElement = Cast<FRigTransformElement>(ParentElement))
 					{
+						FLinearColor Color = FLinearColor::White;
+						if (bHasFKRig)
+						{
+							FName ControlName = UFKControlRig::GetControlName(ParentTransformElement->GetName());
+							if (ActiveControlName.Num() > 0 && ActiveControlName.Contains(ControlName) == false)
+							{
+								continue;
+							}
+							if (ControlRig->IsControlSelected(ControlName))
+							{
+								Color = FLinearColor::Yellow;
+							}
+						}
 						const FTransform ParentTransform = Hierarchy->GetTransform(ParentTransformElement, ERigTransformType::CurrentGlobal);
-                        PDI->DrawLine(ComponentTransform.TransformPosition(Transform.GetLocation()),ComponentTransform.TransformPosition(ParentTransform.GetLocation()), FLinearColor::White, SDPG_Foreground);
+						const bool bHitTesting = bBoolSetHitProxies && (ParentTransformElement->GetType() == ERigElementType::Bone);
+						if (bHitTesting)
+						{
+							PDI->SetHitProxy(new HFKRigBoneProxy(ParentTransformElement->GetName(), ControlRig));
+						}
+                        PDI->DrawLine(ComponentTransform.TransformPosition(Transform.GetLocation()),ComponentTransform.TransformPosition(ParentTransform.GetLocation()), Color, SDPG_Foreground);
+						if (bHitTesting)
+						{
+							PDI->SetHitProxy(nullptr);
+						}
 					}
                 }
 
-				PDI->DrawPoint(ComponentTransform.TransformPosition(Transform.GetLocation()), FLinearColor::White, 5.0f, SDPG_Foreground);
+				FLinearColor Color = FLinearColor::White;
+				if (bHasFKRig)
+				{
+					FName ControlName = UFKControlRig::GetControlName(TransformElement->GetName());
+					if (ActiveControlName.Num() > 0 && ActiveControlName.Contains(ControlName) == false)
+					{
+						return true;
+					}
+					if (ControlRig->IsControlSelected(ControlName))
+					{
+						Color = FLinearColor::Yellow;
+					}
+				}
+				const bool bHitTesting = PDI && PDI->IsHitTesting() && bBoolSetHitProxies && (TransformElement->GetType() == ERigElementType::Bone);
+				if (bHitTesting)
+				{
+					PDI->SetHitProxy(new HFKRigBoneProxy(TransformElement->GetName(), ControlRig));
+				}
+				PDI->DrawPoint(ComponentTransform.TransformPosition(Transform.GetLocation()), Color, 5.0f, SDPG_Foreground);
+
+				if (bHitTesting)
+				{
+					PDI->SetHitProxy(nullptr);
+				}
 
 				return true;
 			});
@@ -919,10 +1069,19 @@ bool FControlRigEditMode::HandleClick(FEditorViewportClient* InViewportClient, H
 								{
 									FScopedTransaction ScopedTransaction(LOCTEXT("SelectControlTransaction", "Select Control"), IsInLevelEditor() && !GIsTransacting);
 
-									if (Click.IsShiftDown() || Click.IsControlDown())
+									if (Click.IsShiftDown()) //guess we just select
 									{
 										SetRigElementSelection(ERigElementType::Control, ControlName, true);
 									}
+									else if (Click.IsControlDown()) //if ctrl we toggle selection
+									{
+										UControlRig* InteractionRig = GetControlRig(true);
+										if (InteractionRig)
+										{
+											bool bIsSelected = InteractionRig->IsControlSelected(ControlName);
+											SetRigElementSelection(ERigElementType::Control, ControlName, !bIsSelected);
+										}
+									}
 									else
 									{
 										ClearRigElementSelection(FRigElementTypeHelper::ToMask(ERigElementType::Control));
@@ -937,6 +1096,34 @@ bool FControlRigEditMode::HandleClick(FEditorViewportClient* InViewportClient, H
 			}
 		}
 	}
+	else if (HFKRigBoneProxy* FKBoneProxy = HitProxyCast<HFKRigBoneProxy>(HitProxy))
+	{
+		FName ControlName(*(FKBoneProxy->BoneName.ToString() + TEXT("_CONTROL")));
+		if (FKBoneProxy->ControlRig->FindControl(ControlName))
+		{
+			FScopedTransaction ScopedTransaction(LOCTEXT("SelectControlTransaction", "Select Control"), IsInLevelEditor() && !GIsTransacting);
+
+			if (Click.IsShiftDown()) //guess we just select
+			{
+				SetRigElementSelection(ERigElementType::Control, ControlName, true);
+			}
+			else if (Click.IsControlDown()) //if ctrl we toggle selection
+			{
+				UControlRig* InteractionRig = GetControlRig(true);
+				if (InteractionRig)
+				{
+					bool bIsSelected = InteractionRig->IsControlSelected(ControlName);
+					SetRigElementSelection(ERigElementType::Control, ControlName, !bIsSelected);
+				}
+			}
+			else
+			{
+				ClearRigElementSelection(FRigElementTypeHelper::ToMask(ERigElementType::Control));
+				SetRigElementSelection(ERigElementType::Control, ControlName, true);
+			}
+			return true;
+		}
+	}
 	else if(HPersonaSelectionHitProxy* CapsuleHitProxy = HitProxyCast<HPersonaSelectionHitProxy>(HitProxy))
 	{
 		static_cast<HPersonaSelectionHitProxy*>(HitProxy)->BroadcastClicked();
@@ -1682,7 +1869,7 @@ void FControlRigEditMode::HandleSelectionChanged()
 			bIsChangingGizmoTransform = false;
 		}
 	}
-	
+
 	// update the pivot transform of our selected objects (they could be animating)
 	RecalcPivotTransform();
 }
diff --git a/Engine/Plugins/Animation/ControlRig/Source/ControlRigEditor/Private/EditMode/ControlRigEditMode.h b/Engine/Plugins/Animation/ControlRig/Source/ControlRigEditor/Private/EditMode/ControlRigEditMode.h
--- a/Engine/Plugins/Animation/ControlRig/Source/ControlRigEditor/Private/EditMode/ControlRigEditMode.h
+++ b/Engine/Plugins/Animation/ControlRig/Source/ControlRigEditor/Private/EditMode/ControlRigEditMode.h
@@ -122,7 +122,7 @@ public:
 	void RefreshObjects();
 
 	/** Get the settings we are using */
-	const UControlRigEditModeSettings* GetSettings() { return Settings; }
+	UControlRigEditModeSettings* GetSettings() const { return Settings; }
 
 	/** Find the edit mode corresponding to the specified world context */
 	static FControlRigEditMode* GetEditModeFromWorldContext(UWorld* InWorldContext);
@@ -326,7 +326,7 @@ public:
 	bool GetOnlySelectRigControls()const;
 
 private:
-
+	TSet<FName> GetActiveControlsFromSequencer(UControlRig* ControlRig);
 	bool CreateGizmoActors(UWorld* World);
 	void DestroyGizmosActors();
 
diff --git a/Engine/Plugins/Animation/ControlRig/Source/ControlRigEditor/Private/Sequencer/ControlRigParameterTrackEditor.cpp b/Engine/Plugins/Animation/ControlRig/Source/ControlRigEditor/Private/Sequencer/ControlRigParameterTrackEditor.cpp
--- a/Engine/Plugins/Animation/ControlRig/Source/ControlRigEditor/Private/Sequencer/ControlRigParameterTrackEditor.cpp
+++ b/Engine/Plugins/Animation/ControlRig/Source/ControlRigEditor/Private/Sequencer/ControlRigParameterTrackEditor.cpp
@@ -76,7 +76,7 @@
 #include "UnrealEdGlobals.h"
 #include "Editor/UnrealEdEngine.h"
 #include "Toolkits/IToolkitHost.h"
-
+#include "ControlRigEditModeSettings.h"
 #define LOCTEXT_NAMESPACE "FControlRigParameterTrackEditor"
 
 static USkeletalMeshComponent* AcquireSkeletalMeshFromObject(UObject* BoundObject, TSharedPtr<ISequencer> SequencerPtr)
@@ -1939,6 +1939,17 @@ void FControlRigParameterTrackEditor::HandleControlSelected(UControlRig* Subject
 		});
 	}
 
+	//if fk rig show hierarchy.
+	if (Subject->IsA<UFKControlRig>())
+	{
+		FControlRigEditMode* EditMode = GetEditMode();
+		if (EditMode)
+		{
+			EditMode->GetSettings()->bDisplayHierarchy = true;
+		}
+
+	}
+
 	if (bIsDoingSelection)
 	{
 		return;
@@ -1964,7 +1975,7 @@ void FControlRigParameterTrackEditor::HandleControlSelected(UControlRig* Subject
 		{
 			return;
 		}
-
+		
 		FFindOrCreateTrackResult TrackResult = FindOrCreateControlRigTrackForObject(ObjectHandle, Subject, ControlRigName, bCreateTrack);
 		UMovieSceneControlRigParameterTrack* Track = CastChecked<UMovieSceneControlRigParameterTrack>(TrackResult.Track, ECastCheckedType::NullAllowed);
 		if (Track)
@@ -3099,24 +3110,19 @@ void FControlRigParameterSection::BuildSectionContextMenu(FMenuBuilder& MenuBuil
 		};
 		MenuBuilder.BeginSection(NAME_None, LOCTEXT("RigSectionActiveChannels", "Active Channels"));
 		{
-			MenuBuilder.AddSubMenu(
-				LOCTEXT("ToggleRigControlsText", "Rig Controls"), LOCTEXT("ToggleRigControlsText_Tooltip", "Causes this section to affect all rig controls"),
-				FNewMenuDelegate::CreateLambda([=](FMenuBuilder& SubMenuBuilder) {
-				int32 Index = 0;
-				for (FRigControlElement* ControlElement : Controls)
-				{
-					const FName RigName = ControlElement->GetName();
-					FText Name = FText::FromName(RigName);
-					FText Text = FText::Format(LOCTEXT("RigControlToggle", "{0}"), Name);
-					FText TooltipText = FText::Format(LOCTEXT("RigControlToggleTooltip", "Causes this section to affect rig control {0}"), Name);
-					SubMenuBuilder.AddMenuEntry(
-						Text, TooltipText,
-						FSlateIcon(), ToggleControls(Index++), NAME_None, EUserInterfaceActionType::ToggleButton);
-				}
-			}),
-				ToggleControls(-1),
-				NAME_None,
-				EUserInterfaceActionType::ToggleButton);
+			MenuBuilder.AddMenuEntry(
+				LOCTEXT("SetFromSelectedControls", "Set From Selected Controls"),
+				LOCTEXT("SetFromSelectedControls_ToolTip", "Set active channels from the current control selection"),
+				FSlateIcon(),
+				FUIAction(FExecuteAction::CreateLambda([=] { ShowSelectedControlsChannels(); }))
+			);
+
+			MenuBuilder.AddMenuEntry(
+				LOCTEXT("ShowAllControls", "Show All Controls"),
+				LOCTEXT("ShowAllControls_ToolTip", "Set active channels from all controls"),
+				FSlateIcon(),
+				FUIAction(FExecuteAction::CreateLambda([=] { return ShowAllControlsChannels(); }))
+			);
 
 			MenuBuilder.AddSubMenu(
 				LOCTEXT("AllTranslation", "Translation"), LOCTEXT("AllTranslation_ToolTip", "Causes this section to affect the translation of rig control transforms"),
@@ -3178,6 +3184,47 @@ void FControlRigParameterSection::BuildSectionContextMenu(FMenuBuilder& MenuBuil
 	}
 }
 
+void FControlRigParameterSection::ShowSelectedControlsChannels()
+{
+	UMovieSceneControlRigParameterSection* ParameterSection = CastChecked<UMovieSceneControlRigParameterSection>(WeakSection.Get());
+	TSharedPtr<ISequencer> SequencerPtr = WeakSequencer.Pin();
+	UControlRig* ControlRig = ParameterSection ? ParameterSection->GetControlRig() : nullptr;
+
+	if (ParameterSection && ControlRig && SequencerPtr.IsValid())
+	{
+		FScopedTransaction Transaction(LOCTEXT("ShowSelecedControlChannels", "Show Selected Control Channels"));
+		ParameterSection->Modify();
+		ParameterSection->FillControlsMask(false);
+
+		TArray<FRigControlElement*> Controls;
+		ControlRig->GetControlsInOrder(Controls);
+		int32 Index = 0;
+		for (const FRigControlElement* RigControl : Controls)
+		{
+			const FName RigName = RigControl->GetName();
+			if (ControlRig->IsControlSelected(RigName))
+			{
+				ParameterSection->SetControlsMask(Index, true);
+			}
+			++Index;
+		}
+		SequencerPtr->NotifyMovieSceneDataChanged(EMovieSceneDataChangeType::MovieSceneStructureItemsChanged);
+	}
+}
+
+void FControlRigParameterSection::ShowAllControlsChannels()
+{
+	UMovieSceneControlRigParameterSection* ParameterSection = CastChecked<UMovieSceneControlRigParameterSection>(WeakSection.Get());
+	TSharedPtr<ISequencer> SequencerPtr = WeakSequencer.Pin();
+	if (ParameterSection && SequencerPtr.IsValid())
+	{
+		FScopedTransaction Transaction(LOCTEXT("ShowAllControlChannels", "Show All Control Channels"));
+		ParameterSection->Modify();
+		ParameterSection->FillControlsMask(true);
+		SequencerPtr->NotifyMovieSceneDataChanged(EMovieSceneDataChangeType::MovieSceneStructureItemsChanged);
+	}
+}
+
 //mz todo
 bool FControlRigParameterSection::RequestDeleteCategory(const TArray<FName>& CategoryNamePaths)
 {
diff --git a/Engine/Plugins/Animation/ControlRig/Source/ControlRigEditor/Private/Sequencer/ControlRigParameterTrackEditor.h b/Engine/Plugins/Animation/ControlRig/Source/ControlRigEditor/Private/Sequencer/ControlRigParameterTrackEditor.h
--- a/Engine/Plugins/Animation/ControlRig/Source/ControlRigEditor/Private/Sequencer/ControlRigParameterTrackEditor.h
+++ b/Engine/Plugins/Animation/ControlRig/Source/ControlRigEditor/Private/Sequencer/ControlRigParameterTrackEditor.h
@@ -247,6 +247,10 @@ protected:
 	/** Animation asset enter pressed */
 	void OnAnimationAssetEnterPressedForFK(const TArray<FAssetData>& AssetData, FGuid ObjectBinding, UMovieSceneControlRigParameterSection* Section);
 
+	/** Select controls channels from selected pressed*/
+	void ShowSelectedControlsChannels();
+	/** Show all control channels*/
+	void ShowAllControlsChannels();
 
 	/** The sequencer which is controlling this section. */
 	TWeakPtr<ISequencer> WeakSequencer;
