LWC Horde fix, FVector3f needed.
#trivial

diff --git a/Engine/Plugins/Animation/ControlRig/Source/ControlRig/Private/Sequencer/MovieSceneControlRigParameterTemplate.cpp b/Engine/Plugins/Animation/ControlRig/Source/ControlRig/Private/Sequencer/MovieSceneControlRigParameterTemplate.cpp
--- a/Engine/Plugins/Animation/ControlRig/Source/ControlRig/Private/Sequencer/MovieSceneControlRigParameterTemplate.cpp
+++ b/Engine/Plugins/Animation/ControlRig/Source/ControlRig/Private/Sequencer/MovieSceneControlRigParameterTemplate.cpp
@@ -1607,7 +1607,7 @@ void FMovieSceneControlRigParameterTemplate::EvaluateCurvesWithMasks(const FMovi
 
 		for (int32 Index = 0; Index < Vectors.Num(); ++Index)
 		{
-			FVector Value(ForceInitToZero);
+			FVector3f Value(ForceInitToZero);
 			const FVectorParameterNameAndCurves& Vector = Vectors[Index];
 
 			if (HACK_ChannelMasks->VectorCurveMask[Index])
