Control Rig: Fix static analysis

diff --git a/Engine/Plugins/Animation/ControlRig/Source/ControlRig/Private/Sequencer/MovieSceneControlRigParameterTemplate.cpp b/Engine/Plugins/Animation/ControlRig/Source/ControlRig/Private/Sequencer/MovieSceneControlRigParameterTemplate.cpp
--- a/Engine/Plugins/Animation/ControlRig/Source/ControlRig/Private/Sequencer/MovieSceneControlRigParameterTemplate.cpp
+++ b/Engine/Plugins/Animation/ControlRig/Source/ControlRig/Private/Sequencer/MovieSceneControlRigParameterTemplate.cpp
@@ -1081,9 +1081,9 @@ struct TControlRigParameterActuatorFloat : TMovieSceneBlendingActuator<FControlR
 	{
 		const UMovieSceneControlRigParameterSection* Section = SectionData.Get();
 
-		UControlRig *ControlRig = GetControlRig(Section, InObject);
+		UControlRig* ControlRig = Section ? GetControlRig(Section, InObject) : nullptr;
 
-		if (Section && ControlRig)
+		if (ControlRig)
 		{
 			FRigControlElement* ControlElement = ControlRig->FindControl(ParameterName);
 			if (ControlElement && ControlElement->Settings.ControlType == ERigControlType::Float)
@@ -1151,9 +1151,9 @@ struct TControlRigParameterActuatorVector2D : TMovieSceneBlendingActuator<FContr
 	{
 		const UMovieSceneControlRigParameterSection* Section = SectionData.Get();
 
-		UControlRig* ControlRig = GetControlRig(Section, InObject);
+		UControlRig* ControlRig = Section ? GetControlRig(Section, InObject) : nullptr;
 
-		if (Section && ControlRig)
+		if (ControlRig)
 		{
 			FRigControlElement* ControlElement = ControlRig->FindControl(ParameterName);
 			if (ControlElement && (ControlElement->Settings.ControlType == ERigControlType::Vector2D))
@@ -1219,7 +1219,8 @@ struct TControlRigParameterActuatorVector : TMovieSceneBlendingActuator<FControl
 	FControlRigTrackTokenVector RetrieveCurrentValue(UObject* InObject, IMovieScenePlayer* Player) const
 	{
 		const UMovieSceneControlRigParameterSection* Section = SectionData.Get();
-		UControlRig* ControlRig = GetControlRig(Section, InObject);
+
+		UControlRig* ControlRig = Section ? GetControlRig(Section, InObject) : nullptr;
 
 		if (ControlRig)
 		{
@@ -1285,9 +1286,9 @@ struct TControlRigParameterActuatorTransform : TMovieSceneBlendingActuator<FCont
 	{
 		const UMovieSceneControlRigParameterSection* Section = SectionData.Get();
 
-		UControlRig* ControlRig = GetControlRig(Section, InObject);
+		UControlRig* ControlRig = Section ? GetControlRig(Section, InObject) : nullptr;
 
-		if (ControlRig && (Section->ControlsToSet.Num() == 0 || Section->ControlsToSet.Contains(ParameterName)))
+		if (ControlRig && Section && (Section->ControlsToSet.Num() == 0 || Section->ControlsToSet.Contains(ParameterName)))
 		{
 			FRigControlElement* ControlElement = ControlRig->FindControl(ParameterName);
 			if (ControlElement && ControlElement->Settings.ControlType == ERigControlType::Transform)
