Control Rig: Fix crash when reseting gizmo name to default

diff --git a/Engine/Plugins/Animation/ControlRig/Source/ControlRigEditor/Private/Editor/SControlRigGizmoNameList.cpp b/Engine/Plugins/Animation/ControlRig/Source/ControlRigEditor/Private/Editor/SControlRigGizmoNameList.cpp
--- a/Engine/Plugins/Animation/ControlRig/Source/ControlRigEditor/Private/Editor/SControlRigGizmoNameList.cpp
+++ b/Engine/Plugins/Animation/ControlRig/Source/ControlRigEditor/Private/Editor/SControlRigGizmoNameList.cpp
@@ -128,8 +128,15 @@ void SControlRigGizmoNameList::OnNameListChanged(TSharedPtr<FString> NewSelectio
 {
 	if (SelectInfo != ESelectInfo::Direct)
 	{
-		FString NewValue = *NewSelection.Get();
-		SetNameListText(FText::FromString(NewValue), ETextCommit::OnEnter);
+		if (NewSelection.IsValid())
+		{
+			FString NewValue = *NewSelection.Get();
+			SetNameListText(FText::FromString(NewValue), ETextCommit::OnEnter);
+		}
+		else
+		{
+			SetNameListText(FText(), ETextCommit::OnEnter);
+		}
 	}
 }
 
