Control Rig: Fix build error in non-editor builds

diff --git a/Engine/Plugins/Animation/ControlRig/Source/ControlRig/Private/ControlRig.cpp b/Engine/Plugins/Animation/ControlRig/Source/ControlRig/Private/ControlRig.cpp
--- a/Engine/Plugins/Animation/ControlRig/Source/ControlRig/Private/ControlRig.cpp
+++ b/Engine/Plugins/Animation/ControlRig/Source/ControlRig/Private/ControlRig.cpp
@@ -1102,10 +1102,12 @@ void UControlRig::SetVM(URigVM* NewVM)
 		}
 
 		// setup array handling and error reporting on the VM
+#if WITH_EDITOR
 		VMRuntimeSettings.LogFunction = [this](EMessageSeverity::Type InSeverity, int32 InInstructionIndex, const FString& Message)
 		{
 			LogOnce(InSeverity, InInstructionIndex, Message);
 		};
+#endif
 		NewVM->SetRuntimeSettings(VMRuntimeSettings);
 	}
 
