Control Rig: Fix expanse state not working after undo of add rig element

diff --git a/Engine/Plugins/Animation/ControlRig/Source/ControlRigEditor/Private/Editor/SRigHierarchy.h b/Engine/Plugins/Animation/ControlRig/Source/ControlRigEditor/Private/Editor/SRigHierarchy.h
--- a/Engine/Plugins/Animation/ControlRig/Source/ControlRigEditor/Private/Editor/SRigHierarchy.h
+++ b/Engine/Plugins/Animation/ControlRig/Source/ControlRigEditor/Private/Editor/SRigHierarchy.h
@@ -130,7 +130,11 @@ public:
 	/** Save a snapshot of the internal map that tracks item expansion before tree reconstruction */
 	void SaveAndClearSparseItemInfos()
 	{
-		OldSparseItemInfos = SparseItemInfos;
+		// Only save the info if there is something to save (do not overwrite info with an empty map)
+		if (!SparseItemInfos.IsEmpty())
+		{
+			OldSparseItemInfos = SparseItemInfos;
+		}
 		ClearExpandedItems();
 	}
 
