Control Rig: Offer inspection of live values

diff --git a/Engine/Plugins/Animation/ControlRig/Source/ControlRig/Public/ControlRig.h b/Engine/Plugins/Animation/ControlRig/Source/ControlRig/Public/ControlRig.h
--- a/Engine/Plugins/Animation/ControlRig/Source/ControlRig/Public/ControlRig.h
+++ b/Engine/Plugins/Animation/ControlRig/Source/ControlRig/Public/ControlRig.h
@@ -701,6 +701,7 @@ private:
 public:
 
 	void SetIsInDebugMode(const bool bValue) { bIsInDebugMode = bValue; }
+	bool IsInDebugMode() const { return bIsInDebugMode; }
 	
 	/** Adds a breakpoint in the VM at the InstructionIndex for the Node */
 	void AddBreakpoint(int32 InstructionIndex, URigVMNode* InNode, uint16 InDepth);
diff --git a/Engine/Plugins/Animation/ControlRig/Source/ControlRigEditor/Private/ControlRigElementDetails.cpp b/Engine/Plugins/Animation/ControlRig/Source/ControlRigEditor/Private/ControlRigElementDetails.cpp
--- a/Engine/Plugins/Animation/ControlRig/Source/ControlRigEditor/Private/ControlRigElementDetails.cpp
+++ b/Engine/Plugins/Animation/ControlRig/Source/ControlRigEditor/Private/ControlRigElementDetails.cpp
@@ -188,6 +188,7 @@ void FRigElementKeyDetails::CustomizeHeader(TSharedRef<IPropertyHandle> InStruct
 				.OptionsSource(&ElementNameList)
 				.OnSelectionChanged(this, &FRigElementKeyDetails::OnElementNameChanged)
 				.OnGenerateWidget(this, &FRigElementKeyDetails::OnGetElementNameWidget)
+				.IsEnabled(!NameHandle->IsEditConst())
 				.Content()
 				[
 					SNew(STextBlock)
diff --git a/Engine/Plugins/Animation/ControlRig/Source/ControlRigEditor/Private/ControlRigGraphDetails.cpp b/Engine/Plugins/Animation/ControlRig/Source/ControlRigEditor/Private/ControlRigGraphDetails.cpp
--- a/Engine/Plugins/Animation/ControlRig/Source/ControlRigEditor/Private/ControlRigGraphDetails.cpp
+++ b/Engine/Plugins/Animation/ControlRig/Source/ControlRigEditor/Private/ControlRigGraphDetails.cpp
@@ -22,6 +22,7 @@
 #include "Kismet2/BlueprintEditorUtils.h"
 #include "Graph/ControlRigGraphSchema.h"
 #include "EditorCategoryUtils.h"
+#include "IPropertyUtilities.h"
 
 #define LOCTEXT_NAMESPACE "ControlRigGraphDetails"
 
@@ -1058,4 +1059,246 @@ TSharedRef<ITableRow> FControlRigGraphDetails::HandleGenerateRowAccessSpecifier(
         ];
 }
 
+#if !UE_RIGVM_UCLASS_BASED_STORAGE_DISABLED
+
+TSharedRef<IDetailCustomization> FControlRigWrappedNodeDetails::MakeInstance()
+{
+	return MakeShareable(new FControlRigWrappedNodeDetails);
+}
+
+void FControlRigWrappedNodeDetails::CustomizeDetails(IDetailLayoutBuilder& DetailLayout)
+{
+	TArray<TWeakObjectPtr<UObject>> Objects;
+	DetailLayout.GetObjectsBeingCustomized(Objects);
+
+	TArray<UDetailsViewWrapperObject*> Wrappers;
+	TArray<URigVMNode*> Nodes;
+	for(const TWeakObjectPtr<UObject>& Object : Objects)
+	{
+		UDetailsViewWrapperObject* Wrapper = CastChecked<UDetailsViewWrapperObject>(Object);
+		Wrappers.Add(Wrapper);
+
+		URigVMNode* Node = Wrapper->GetTypedOuter<URigVMNode>();
+		if(Node == nullptr)
+		{
+			return;
+		}
+		Nodes.Add(Node);
+	}
+
+	UDetailsViewWrapperObject* FirstWrapper = Wrappers[0];
+	URigVMNode* FirstNode = Nodes[0];
+
+	TArray<FName> Categories;
+	DetailLayout.GetCategoryNames(Categories);
+	
+	IDetailCategoryBuilder& DefaultsCategory = DetailLayout.EditCategory(Categories[0]);
+
+	TSharedPtr<IPropertyHandle> StoredStructHandle = DetailLayout.GetProperty(TEXT("StoredStruct"));
+	for(URigVMPin* Pin : FirstNode->GetPins())
+	{
+		bool bHasAnyInputLink = false;
+		for(URigVMNode* EachNode : Nodes)
+		{
+			if(URigVMPin* EachPin = EachNode->FindPin(Pin->GetName()))
+			{
+				if(EachPin->GetSourceLinks(true).Num() > 0)
+				{
+					bHasAnyInputLink = true;
+					break;
+				}
+			}
+		}
+
+		TSharedPtr<IPropertyHandle> PinHandle = StoredStructHandle->GetChildHandle(Pin->GetFName());
+		if(PinHandle.IsValid())
+		{
+			static const TCHAR DisabledFormat[] = TEXT("%s\n\nNote: Editing disabled since pin has an input link.");
+
+			DefaultsCategory.AddProperty(PinHandle)
+			.DisplayName(FText::FromName(Pin->GetDisplayName()))
+			.IsEnabled(!bHasAnyInputLink)
+			.ToolTip(FText::FromString(FString::Printf(DisabledFormat, *Pin->GetToolTipText().ToString())));
+		}
+	}
+
+	DetailLayout.HideProperty(StoredStructHandle);
+
+	UControlRigBlueprint* Blueprint = FirstNode->GetTypedOuter<UControlRigBlueprint>();
+	if(Blueprint == nullptr)
+	{
+		return;
+	}
+
+	if(Objects.Num() > 1)
+	{
+		return;
+	}
+
+	UControlRig* DebuggedRig = Cast<UControlRig>(Blueprint->GetObjectBeingDebugged());
+	if(DebuggedRig == nullptr)
+	{
+		return;
+	}
+
+	URigVM* VM = DebuggedRig->GetVM();
+	if(VM == nullptr)
+	{
+		return;
+	}
+
+	TSharedPtr<FRigVMParserAST> AST = FirstNode->GetGraph()->GetRuntimeAST(Blueprint->VMCompileSettings.ASTSettings, false);
+	if(!AST.IsValid())
+	{
+		return;
+	}
+
+	FRigVMByteCode& ByteCode = VM->GetByteCode();
+	if(ByteCode.GetFirstInstructionIndexForSubject(FirstNode) == INDEX_NONE)
+	{
+		return;
+	}
+	
+	IDetailCategoryBuilder& DebugCategory = DetailLayout.EditCategory("DebugLiveValues", LOCTEXT("DebugLiveValues", "Inspect Live Values"), ECategoryPriority::Uncommon);
+	DebugCategory.InitiallyCollapsed(true);
+
+	for(URigVMPin* Pin : FirstNode->GetPins())
+	{
+		// only show hidden pins in debug mode
+		if(Pin->GetDirection() == ERigVMPinDirection::Hidden)
+		{
+			if(!DebuggedRig->IsInDebugMode())
+			{
+				continue;
+			}
+		}
+		
+		URigVMPin* SourcePin = Pin;
+		if(Blueprint->VMCompileSettings.ASTSettings.bFoldAssignments)
+		{
+			do
+			{
+				TArray<URigVMPin*> SourcePins = Pin->GetLinkedSourcePins(false);
+				if(SourcePins.Num() > 0)
+				{
+					SourcePin = SourcePins[0];
+				}
+				else
+				{
+					break;
+				}
+			}
+			while(SourcePin->GetNode()->IsA<URigVMRerouteNode>());
+		}
+		
+		TArray<const FRigVMExprAST*> Expressions = AST->GetExpressionsForSubject(SourcePin);
+		if(Expressions.Num() == 0 && SourcePin != Pin)
+		{
+			SourcePin = Pin;
+			Expressions = AST->GetExpressionsForSubject(Pin);
+		}
+
+		bool bHasVar = false;
+		for(const FRigVMExprAST* Expression : Expressions)
+		{
+			if(Expression->IsA(FRigVMExprAST::EType::Literal))
+			{
+				continue;
+			}
+			else if(Expression->IsA(FRigVMExprAST::EType::Var))
+			{
+				bHasVar = true;
+				break;
+			}
+		}
+
+		TArray<const FRigVMExprAST*> FilteredExpressions;
+		for(const FRigVMExprAST* Expression : Expressions)
+		{
+			if(Expression->IsA(FRigVMExprAST::EType::Literal))
+			{
+				if(bHasVar)
+				{
+					continue;
+				}
+				FilteredExpressions.Add(Expression);
+			}
+			else if(Expression->IsA(FRigVMExprAST::EType::Var))
+			{
+				FilteredExpressions.Add(Expression);
+			}
+			else if(Expression->IsA(FRigVMExprAST::EType::CachedValue))
+			{
+				const FRigVMCachedValueExprAST* CachedValueExpr = Expression->To<FRigVMCachedValueExprAST>();
+				FilteredExpressions.Add(CachedValueExpr->GetVarExpr());
+			}
+		}
+
+		bool bAddedProperty = false;
+		int32 SuffixIndex = 1;
+		FString NameSuffix;
+
+		TArray<FRigVMOperand> KnownOperands; 
+		for(const FRigVMExprAST* Expression : FilteredExpressions)
+		{
+			const FRigVMVarExprAST* VarExpr = Expression->To<FRigVMVarExprAST>();
+
+			FString PinHash = URigVMCompiler::GetPinHash(SourcePin, VarExpr, false);
+			const FRigVMOperand* Operand = Blueprint->PinToOperandMap.Find(PinHash);
+			if(Operand)
+			{
+				if(Operand->GetRegisterOffset() != INDEX_NONE)
+				{
+					continue;
+				}
+				if(KnownOperands.Contains(*Operand))
+				{
+					continue;
+				}
+				
+				URigVMMemoryStorage* Memory = VM->GetMemoryByType(Operand->GetMemoryType());
+				if(Memory == nullptr)
+				{
+					continue;
+				}
+
+				if(!Memory->IsValidIndex(Operand->GetRegisterIndex()))
+				{
+					continue;
+				}
+				
+				const FProperty* Property = Memory->GetProperty(Operand->GetRegisterIndex());
+				if(Property == nullptr)
+				{
+					continue;
+				}
+
+				TArray<UObject*> ExternalObjects;
+				ExternalObjects.Add(Memory);
+				DebugCategory.AddExternalObjectProperty(ExternalObjects, Property->GetFName(), EPropertyLocation::Default, FAddPropertyParams().ForceShowProperty())
+				->DisplayName(FText::FromString(FString::Printf(TEXT("%s%s"), *Pin->GetName(), *NameSuffix)))
+				.IsEnabled(false);
+
+				SuffixIndex++;
+				bAddedProperty = true;
+				NameSuffix = FString::Printf(TEXT("_%d"), SuffixIndex);
+				KnownOperands.Add(*Operand);
+			}
+		}
+
+		if(!bAddedProperty)
+		{
+			TSharedPtr<IPropertyHandle> PinHandle = StoredStructHandle->GetChildHandle(Pin->GetFName());
+			if(PinHandle.IsValid())
+			{
+				DefaultsCategory.AddProperty(PinHandle)
+				.DisplayName(FText::FromName(Pin->GetDisplayName()))
+				.IsEnabled(false);
+			}
+		}
+	}
+}
+
+#endif
+
 #undef LOCTEXT_NAMESPACE
diff --git a/Engine/Plugins/Animation/ControlRig/Source/ControlRigEditor/Private/ControlRigGraphDetails.h b/Engine/Plugins/Animation/ControlRig/Source/ControlRigEditor/Private/ControlRigGraphDetails.h
--- a/Engine/Plugins/Animation/ControlRig/Source/ControlRigEditor/Private/ControlRigGraphDetails.h
+++ b/Engine/Plugins/Animation/ControlRig/Source/ControlRigEditor/Private/ControlRigGraphDetails.h
@@ -12,6 +12,7 @@
 #include "SGraphPin.h"
 #include "Graph/SControlRigGraphNode.h"
 #include "Widgets/Colors/SColorBlock.h"
+#include "DetailsViewWrapperObject.h"
 
 class IDetailLayoutBuilder;
 
@@ -196,4 +197,20 @@ private:
 	bool bIsPickingColor;
 
 	static TArray<TSharedPtr<FString>> AccessSpecifierStrings;
-};
\ No newline at end of file
+};
+
+#if !UE_RIGVM_UCLASS_BASED_STORAGE_DISABLED
+
+/** Customization for editing a Control Rig node */
+class FControlRigWrappedNodeDetails : public IDetailCustomization
+{
+public:
+
+	/** Makes a new instance of this detail layout class for a specific detail view requesting it */
+	static TSharedRef<IDetailCustomization> MakeInstance();
+
+	// IDetailCustomization interface
+	virtual void CustomizeDetails(IDetailLayoutBuilder& DetailLayout) override;
+};
+
+#endif
\ No newline at end of file
diff --git a/Engine/Plugins/Animation/ControlRig/Source/ControlRigEditor/Private/Editor/ControlRigEditor.cpp b/Engine/Plugins/Animation/ControlRig/Source/ControlRigEditor/Private/Editor/ControlRigEditor.cpp
--- a/Engine/Plugins/Animation/ControlRig/Source/ControlRigEditor/Private/Editor/ControlRigEditor.cpp
+++ b/Engine/Plugins/Animation/ControlRig/Source/ControlRigEditor/Private/Editor/ControlRigEditor.cpp
@@ -91,6 +91,7 @@
 #include "ControlRig/Private/Units/Execution/RigUnit_SequenceExecution.h"
 #include "ControlRigContextMenuContext.h"
 #include "Types/ISlateMetaData.h"
+#include "ControlRigGraphDetails.h"
 
 #define LOCTEXT_NAMESPACE "ControlRigEditor"
 
@@ -981,6 +982,8 @@ void FControlRigEditor::SetExecutionMode(const EControlRigExecutionModeType InEx
 		HaltedAtNode->SetExecutionIsHaltedAtThisNode(false);
 		HaltedAtNode = nullptr;
 	}
+
+	RefreshDetailView();
 }
 
 int32 FControlRigEditor::GetExecutionModeComboValue() const
@@ -1268,6 +1271,8 @@ void FControlRigEditor::HandleSetObjectBeingDebugged(UObject* InObject)
 			EditMode->SetObjects(nullptr,  nullptr,nullptr);
 		}
 	}
+
+	RefreshDetailView();
 }
 
 FString FControlRigEditor::GetCustomDebugObjectLabel(UObject* ObjectBeingDebugged) const
@@ -1401,11 +1406,31 @@ void FControlRigEditor::SetDetailObjects(const TArray<UObject*>& InObjects)
 				// create a struct matching this node
 				TSharedPtr<FStructOnScope> StructInstance  = UnitNode->ConstructStructInstance(false);
 
+#if !UE_RIGVM_UCLASS_BASED_STORAGE_DISABLED
+
+				// check if we know the dynamic class already
+				bool bClassCreated = UDetailsViewWrapperObject::GetClassForStruct(NodeDetailStruct, false) == nullptr;
+
+#endif
+
 				// create the wrapper object
 				UDetailsViewWrapperObject* WrapperObject = UDetailsViewWrapperObject::MakeInstance(
 					NodeDetailStruct, StructInstance->GetStructMemory(), UnitNode);
 				WrapperObject->GetWrappedPropertyChangedChainEvent().AddSP(this, &FControlRigEditor::OnWrappedPropertyChangedChainEvent);
 				WrapperObject->AddToRoot();
+
+#if !UE_RIGVM_UCLASS_BASED_STORAGE_DISABLED
+
+				if(bClassCreated)
+				{
+					UClass* WrapperClass = UDetailsViewWrapperObject::GetClassForStruct(NodeDetailStruct, false);
+					check(WrapperClass);
+
+					FPropertyEditorModule& EditModule = FModuleManager::Get().GetModuleChecked<FPropertyEditorModule>("PropertyEditor");
+					EditModule.RegisterCustomClassLayout(WrapperClass->GetFName(), FOnGetDetailCustomizationInstance::CreateStatic(&FControlRigWrappedNodeDetails::MakeInstance));
+				}
+
+#endif
 				
 				FilteredObjects.Add(WrapperObject);
 				continue;
@@ -2633,9 +2658,11 @@ void FControlRigEditor::HandleVMCompiledEvent(UBlueprint* InBlueprint, URigVM* I
 		CompilerResultsListing->ClearMessages();
 		CompilerResultsListing->AddMessages(RigBlueprint->CompileLog.Messages);
 		RigBlueprint->CompileLog.Messages.Reset();
-		RigBlueprint->CompileLog.NumErrors = RigBlueprint->CompileLog.NumWarnings = 0;  
+		RigBlueprint->CompileLog.NumErrors = RigBlueprint->CompileLog.NumWarnings = 0;
 	}
 
+	RefreshDetailView();
+	
 #if !UE_RIGVM_UCLASS_BASED_STORAGE_DISABLED
 	
 	TArray<FName> TabIds;
diff --git a/Engine/Plugins/Animation/ControlRig/Source/ControlRigEditor/Private/Editor/DetailsViewWrapperObject.cpp b/Engine/Plugins/Animation/ControlRig/Source/ControlRigEditor/Private/Editor/DetailsViewWrapperObject.cpp
--- a/Engine/Plugins/Animation/ControlRig/Source/ControlRigEditor/Private/Editor/DetailsViewWrapperObject.cpp
+++ b/Engine/Plugins/Animation/ControlRig/Source/ControlRigEditor/Private/Editor/DetailsViewWrapperObject.cpp
@@ -12,7 +12,7 @@
 TMap<UScriptStruct*, UClass*> UDetailsViewWrapperObject::StructToClass;
 TMap<UClass*, UScriptStruct*> UDetailsViewWrapperObject::ClassToStruct;
 
-UClass* UDetailsViewWrapperObject::GetClassForStruct(UScriptStruct* InStruct)
+UClass* UDetailsViewWrapperObject::GetClassForStruct(UScriptStruct* InStruct, bool bCreateIfNeeded)
 {
 	check(InStruct != nullptr);
 
@@ -22,6 +22,11 @@ UClass* UDetailsViewWrapperObject::GetClassForStruct(UScriptStruct* InStruct)
 		return *ExistingClass;
 	}
 
+	if(!bCreateIfNeeded)
+	{
+		return nullptr;
+	}
+
 	UClass *SuperClass = UDetailsViewWrapperObject::StaticClass();
 	const FName WrapperClassName(FString::Printf(TEXT("%s_WrapperObject"), *InStruct->GetStructCPPName()));
 
@@ -79,16 +84,24 @@ UClass* UDetailsViewWrapperObject::GetClassForStruct(UScriptStruct* InStruct)
 	
 	// Make sure the variables show up in the details panel
 	Property->SetPropertyFlags(CPF_Edit);
+#if UE_RIGVM_UCLASS_BASED_STORAGE_DISABLED
 	Property->SetMetaData(TEXT("ShowOnlyInnerProperties"), TEXT("true"));
+#else
+	// for non-nodes - show the properties directly under the details panel main category 
+	if(!InStruct->IsChildOf(FRigUnit::StaticStruct()))
+	{
+		Property->SetMetaData(TEXT("ShowOnlyInnerProperties"), TEXT("true"));
+	}
+#endif
 
-	// For rig units mark all inputs as edit
+	// For rig units
 	if(InStruct->IsChildOf(FRigUnit::StaticStruct()))
 	{
 		// mark all input properties with editanywhere
 		for (TFieldIterator<FProperty> PropertyIt(InStruct); PropertyIt; ++PropertyIt)
 		{
 			FProperty* ChildProperty = *PropertyIt;
-			if (!ChildProperty ->HasMetaData(TEXT("Input")))
+			if (!ChildProperty->HasMetaData(TEXT("Input")) && !ChildProperty->HasMetaData(TEXT("Visible")))
 			{
 				continue;
 			}
@@ -101,12 +114,6 @@ UClass* UDetailsViewWrapperObject::GetClassForStruct(UScriptStruct* InStruct)
 					continue;
 				}
 			}
-
-			// only do this for input pins
-			if(!ChildProperty->HasMetaData(TEXT("Input")))
-			{
-				continue;
-			}
 			
 			ChildProperty->SetPropertyFlags(ChildProperty->GetPropertyFlags() | CPF_Edit);
 		}
diff --git a/Engine/Plugins/Animation/ControlRig/Source/ControlRigEditor/Private/Editor/DetailsViewWrapperObject.h b/Engine/Plugins/Animation/ControlRig/Source/ControlRigEditor/Private/Editor/DetailsViewWrapperObject.h
--- a/Engine/Plugins/Animation/ControlRig/Source/ControlRigEditor/Private/Editor/DetailsViewWrapperObject.h
+++ b/Engine/Plugins/Animation/ControlRig/Source/ControlRigEditor/Private/Editor/DetailsViewWrapperObject.h
@@ -15,7 +15,7 @@ class UDetailsViewWrapperObject : public UObject
 public:
 	GENERATED_BODY()
 
-	static UClass* GetClassForStruct(UScriptStruct* InStruct);
+	static UClass* GetClassForStruct(UScriptStruct* InStruct, bool bCreateIfNeeded = true);
 	static UDetailsViewWrapperObject* MakeInstance(UScriptStruct* InStruct, uint8* InStructMemory, UObject* InOuter = nullptr);
 	UScriptStruct* GetWrappedStruct() const;
 	FStructProperty* GetStructProperty() const;
