ControlRig: RigHierarchy: Guard against notifies during garbage collection

diff --git a/Engine/Plugins/Animation/ControlRig/Source/ControlRig/Private/Rigs/RigHierarchy.cpp b/Engine/Plugins/Animation/ControlRig/Source/ControlRig/Private/Rigs/RigHierarchy.cpp
--- a/Engine/Plugins/Animation/ControlRig/Source/ControlRig/Private/Rigs/RigHierarchy.cpp
+++ b/Engine/Plugins/Animation/ControlRig/Source/ControlRig/Private/Rigs/RigHierarchy.cpp
@@ -255,7 +255,10 @@ void URigHierarchy::Reset()
 	ResetPoseHash = INDEX_NONE;
 	ResetPoseHasFilteredChildren.Reset();
 
-	Notify(ERigHierarchyNotification::HierarchyReset, nullptr);
+	if(!IsGarbageCollecting())
+	{
+		Notify(ERigHierarchyNotification::HierarchyReset, nullptr);
+	}
 }
 
 void URigHierarchy::CopyHierarchy(URigHierarchy* InHierarchy)
