Control Rig: Fix crash when deleting function

diff --git a/Engine/Plugins/Animation/ControlRig/Source/ControlRigDeveloper/Private/Graph/ControlRigGraphSchema.cpp b/Engine/Plugins/Animation/ControlRig/Source/ControlRigDeveloper/Private/Graph/ControlRigGraphSchema.cpp
--- a/Engine/Plugins/Animation/ControlRig/Source/ControlRigDeveloper/Private/Graph/ControlRigGraphSchema.cpp
+++ b/Engine/Plugins/Animation/ControlRig/Source/ControlRigDeveloper/Private/Graph/ControlRigGraphSchema.cpp
@@ -948,8 +948,41 @@ bool UControlRigGraphSchema::TryDeleteGraph(UEdGraph* GraphToDelete) const
 								}
 							}
 						}
+
+						bool bSetupUndoRedo = true;
+
+						// if the element to remove is a function, check if it is public and referenced. If so,
+						// warn the user about a bulk remove
+						if (URigVMFunctionLibrary* Library = Cast<URigVMFunctionLibrary>(LibraryNode->GetGraph()))
+						{
+							const FName& FunctionName = LibraryNode->GetFName();
+							if (RigBlueprint->IsFunctionPublic(FunctionName))
+							{
+								for (auto Reference : Library->GetReferencesForFunction(FunctionName))
+								{
+									if (Reference.IsValid())
+									{
+										UControlRigBlueprint* OtherBlueprint = Reference->GetTypedOuter<UControlRigBlueprint>(); 
+										if (OtherBlueprint != RigBlueprint)
+										{											
+											if(RigBlueprint->OnRequestBulkEditDialog().IsBound())
+											{
+												URigVMController* FunctionController = RigBlueprint->GetController(LibraryNode->GetContainedGraph());
+												FRigVMController_BulkEditResult Result = RigBlueprint->OnRequestBulkEditDialog().Execute(RigBlueprint, FunctionController, LibraryNode, ERigVMControllerBulkEditType::RemoveFunction);
+												if(Result.bCanceled)
+												{
+													return false;
+												}
+												bSetupUndoRedo = Result.bSetupUndoRedo;
+											}
+											break;	
+										}
+									}
+								}
+							}
+						}
 						
-						return Controller->RemoveNode(LibraryNode, true, false, true);
+						return Controller->RemoveNode(LibraryNode, bSetupUndoRedo, false, true);
 					}
 				}
 			}
diff --git a/Engine/Source/Developer/RigVMDeveloper/Private/RigVMModel/RigVMController.cpp b/Engine/Source/Developer/RigVMDeveloper/Private/RigVMModel/RigVMController.cpp
--- a/Engine/Source/Developer/RigVMDeveloper/Private/RigVMModel/RigVMController.cpp
+++ b/Engine/Source/Developer/RigVMDeveloper/Private/RigVMModel/RigVMController.cpp
@@ -4335,7 +4335,7 @@ bool URigVMController::RemoveNode(URigVMNode* InNode, bool bSetupUndoRedo, bool
 								TSoftObjectPtr<URigVMLibraryNode>());
 
 							FRigVMControllerGraphGuard GraphGuard(this, FunctionReferencePtr->GetGraph(), false);
-							RepopulatePinsOnNode(FunctionReferencePtr.Get(), false, true);
+							RemoveNode(FunctionReferencePtr.Get());
 						}
 						FunctionReferencePtr->ReferencedNodePtr.ResetWeakPtr();
 					}
diff --git a/Engine/Source/Developer/RigVMDeveloper/Public/RigVMModel/RigVMController.h b/Engine/Source/Developer/RigVMDeveloper/Public/RigVMModel/RigVMController.h
--- a/Engine/Source/Developer/RigVMDeveloper/Public/RigVMModel/RigVMController.h
+++ b/Engine/Source/Developer/RigVMDeveloper/Public/RigVMModel/RigVMController.h
@@ -31,6 +31,7 @@ enum class ERigVMControllerBulkEditType : uint8
     RemoveVariable,
 	RenameVariable,
     ChangeVariableType,
+	RemoveFunction,
     Max UMETA(Hidden),
 };
 
