Control Rig: Crash when watching pins

diff --git a/Engine/Plugins/Animation/ControlRig/Source/ControlRigDeveloper/Private/ControlRigBlueprint.cpp b/Engine/Plugins/Animation/ControlRig/Source/ControlRigDeveloper/Private/ControlRigBlueprint.cpp
--- a/Engine/Plugins/Animation/ControlRig/Source/ControlRigDeveloper/Private/ControlRigBlueprint.cpp
+++ b/Engine/Plugins/Animation/ControlRig/Source/ControlRigDeveloper/Private/ControlRigBlueprint.cpp
@@ -2862,7 +2862,7 @@ void UControlRigBlueprint::HandleModifiedEvent(ERigVMGraphNotifType InNotifType,
 #if UE_RIGVM_UCLASS_BASED_STORAGE_DISABLED
 						Compiler->RemoveDebugRegister(Pin, CR->GetVM(), &PinToOperandMap, RuntimeAST);
 #else
-						Compiler->MarkDebugWatch(true, Pin, CR->GetVM(), &PinToOperandMap, RuntimeAST);
+						Compiler->MarkDebugWatch(false, Pin, CR->GetVM(), &PinToOperandMap, RuntimeAST);
 #endif
 					}
 				}
diff --git a/Engine/Plugins/Animation/ControlRig/Source/ControlRigEditor/Private/Graph/SControlRigGraphNode.cpp b/Engine/Plugins/Animation/ControlRig/Source/ControlRigEditor/Private/Graph/SControlRigGraphNode.cpp
--- a/Engine/Plugins/Animation/ControlRig/Source/ControlRigEditor/Private/Graph/SControlRigGraphNode.cpp
+++ b/Engine/Plugins/Animation/ControlRig/Source/ControlRigEditor/Private/Graph/SControlRigGraphNode.cpp
@@ -897,8 +897,11 @@ void SControlRigGraphNode::GetNodeInfoPopups(FNodeInfoContext* Context, TArray<F
 								DefaultValues = Memory.GetRegisterValueAsString(*WatchOperand, ModelPin->GetCPPType(), ModelPin->GetCPPTypeObject());
 #else
 								URigVMMemoryStorage* Memory = ActiveObject->GetVM()->GetDebugMemory();
-								const FString DebugValue = Memory->GetDataAsString(WatchOperand->GetRegisterIndex());
-								DefaultValues = URigVMPin::SplitDefaultValue(DebugValue);
+								const FString DebugValue = Memory->GetDataAsStringSafe(WatchOperand->GetRegisterIndex());
+								if(!DebugValue.IsEmpty())
+								{
+									DefaultValues = URigVMPin::SplitDefaultValue(DebugValue);
+								}
 #endif
 							}
 						}
diff --git a/Engine/Source/Runtime/RigVM/Private/RigVMCore/RigVM.cpp b/Engine/Source/Runtime/RigVM/Private/RigVMCore/RigVM.cpp
--- a/Engine/Source/Runtime/RigVM/Private/RigVMCore/RigVM.cpp
+++ b/Engine/Source/Runtime/RigVM/Private/RigVMCore/RigVM.cpp
@@ -3289,6 +3289,18 @@ void URigVM::CopyOperandForDebuggingImpl(const FRigVMOperand& InArg, const FRigV
 	TargetPtr = ArrayHelper.GetRawPtr(AddedIndex);
 	TargetProperty = TargetArrayProperty->Inner;
 
+	if(InArg.GetMemoryType() == ERigVMMemoryType::External)
+	{
+		if(ExternalVariables.IsValidIndex(InArg.GetRegisterIndex()))
+		{
+			FRigVMExternalVariable& ExternalVariable = ExternalVariables[InArg.GetRegisterIndex()];
+			const FProperty* SourceProperty = ExternalVariable.Property;
+			const uint8* SourcePtr = ExternalVariable.Memory;
+			URigVMMemoryStorage::CopyProperty(TargetProperty, TargetPtr, SourceProperty, SourcePtr);
+		}
+		return;
+	}
+
 	URigVMMemoryStorage* SourceMemory = GetMemoryByType(InArg.GetMemoryType());
 	if(SourceMemory == nullptr)
 	{
