Sequencer: Control Rig : Pose Library, allow user to select controls from a pose with the edit mode not active. While looking into this found some cases where we weren't clearing out a delgate which could cause crash.

diff --git a/Engine/Plugins/Animation/ControlRig/Source/ControlRigEditor/Private/EditMode/ControlRigEditMode.cpp b/Engine/Plugins/Animation/ControlRig/Source/ControlRigEditor/Private/EditMode/ControlRigEditMode.cpp
--- a/Engine/Plugins/Animation/ControlRig/Source/ControlRigEditor/Private/EditMode/ControlRigEditMode.cpp
+++ b/Engine/Plugins/Animation/ControlRig/Source/ControlRigEditor/Private/EditMode/ControlRigEditMode.cpp
@@ -176,6 +176,8 @@ FControlRigEditMode::~FControlRigEditMode()
 #if WITH_EDITOR
 	FCoreUObjectDelegates::OnObjectsReplaced.RemoveAll(this);
 #endif
+
+	OnControlRigAddedOrRemovedDelegate.Clear();
 }
 
 void FControlRigEditMode::SetObjects(const TWeakObjectPtr<>& InSelectedObject,  UObject* BindingObject, TWeakPtr<ISequencer> InSequencer)
@@ -387,6 +389,7 @@ void FControlRigEditMode::Exit()
 		DelegateHelper->RemoveDelegates();
 		DelegateHelper.Reset();
 	}
+	OnControlRigAddedOrRemovedDelegate.Clear();
 
 	//clear delegates
 	FEditorModeTools* ModeManager = GetModeManager();
diff --git a/Engine/Plugins/Animation/ControlRig/Source/ControlRigEditor/Private/EditMode/SControlRigBaseListWidget.cpp b/Engine/Plugins/Animation/ControlRig/Source/ControlRigEditor/Private/EditMode/SControlRigBaseListWidget.cpp
--- a/Engine/Plugins/Animation/ControlRig/Source/ControlRigEditor/Private/EditMode/SControlRigBaseListWidget.cpp
+++ b/Engine/Plugins/Animation/ControlRig/Source/ControlRigEditor/Private/EditMode/SControlRigBaseListWidget.cpp
@@ -39,6 +39,10 @@
 #include "SControlRigUpdatePose.h"
 #include "SControlRigRenamePoseControls.h"
 #include "Dialogs/Dialogs.h"
+#include "ControlRigSequencerEditorLibrary.h"
+#include "LevelSequence.h"
+#include "LevelSequenceEditorBlueprintLibrary.h"
+
 
 #define LOCTEXT_NAMESPACE "ControlRigBaseListWidget"
 
@@ -1320,6 +1324,22 @@ void SControlRigBaseListWidget::ExecuteSelectControls(UControlRigPoseAsset* Pose
 		ControlRigEditMode->GetControlRig(true)->Modify();
 		PoseAsset->SelectControls(ControlRigEditMode->GetControlRig(true));
 	}
+	else
+	{
+		ULevelSequence *LevelSequence = ULevelSequenceEditorBlueprintLibrary::GetCurrentLevelSequence();
+		if (LevelSequence)
+		{
+			TArray<FControlRigSequencerBindingProxy> Proxies = UControlRigSequencerEditorLibrary::GetControlRigs(LevelSequence);
+			if (Proxies.Num() > 0)
+			{
+				//MZ TODO when we have Mutliple Control Rig's active select more than one.
+				if (Proxies[0].ControlRig)
+				{
+					PoseAsset->SelectControls(Proxies[0].ControlRig);
+				}
+			}
+		}
+	}
 }
 
 void SControlRigBaseListWidget::ExecutePasteMirrorPose(UControlRigPoseAsset* PoseAsset)
diff --git a/Engine/Plugins/Animation/ControlRig/Source/ControlRigEditor/Private/EditMode/SControlRigControlViews.cpp b/Engine/Plugins/Animation/ControlRig/Source/ControlRigEditor/Private/EditMode/SControlRigControlViews.cpp
--- a/Engine/Plugins/Animation/ControlRig/Source/ControlRigEditor/Private/EditMode/SControlRigControlViews.cpp
+++ b/Engine/Plugins/Animation/ControlRig/Source/ControlRigEditor/Private/EditMode/SControlRigControlViews.cpp
@@ -455,6 +455,13 @@ SControlRigPoseView::~SControlRigPoseView()
 			EditMode->GetControlRig(true)->ControlSelected().RemoveAll(this);
 		}
 	}
+	else
+	{
+		if (CurrentControlRig.IsValid())
+		{
+			(CurrentControlRig.Get())->ControlSelected().RemoveAll(this);
+		}
+	}
 }
 
 void SControlRigPoseView::NotifyPostChange(const FPropertyChangedEvent& PropertyChangedEvent, FProperty* PropertyThatChanged)
@@ -613,11 +620,21 @@ TSharedRef<SWidget> SControlRigPoseView::GetThumbnailWidget()
 
 UControlRig* SControlRigPoseView::GetControlRig()
 {
+	UControlRig* NewControlRig = nullptr;
 	if (FControlRigEditMode* EditMode = static_cast<FControlRigEditMode*>(GLevelEditorModeTools().GetActiveMode(FControlRigEditMode::ModeName)))
 	{
-		return EditMode->GetControlRig(true);
+		NewControlRig =  EditMode->GetControlRig(true);
+	}
+	if (NewControlRig != CurrentControlRig)
+	{
+		if (CurrentControlRig.IsValid())
+		{
+			(CurrentControlRig.Get())->ControlSelected().RemoveAll(this);
+
+		}
 	}
-	return nullptr;
+	CurrentControlRig = NewControlRig;
+	return NewControlRig;
 }
 
 /* We may want to list the Controls in it (design said no but animators said yes)
@@ -641,6 +658,7 @@ void SControlRigPoseView::HandleControlAdded(UControlRig* ControlRig, bool bIsAd
 		{
 			(ControlRig)->ControlSelected().RemoveAll(this);
 			(ControlRig)->ControlSelected().AddRaw(this, &SControlRigPoseView::HandleControlSelected);
+			CurrentControlRig = ControlRig;
 		}
 		else
 		{
diff --git a/Engine/Plugins/Animation/ControlRig/Source/ControlRigEditor/Private/EditMode/SControlRigControlViews.h b/Engine/Plugins/Animation/ControlRig/Source/ControlRigEditor/Private/EditMode/SControlRigControlViews.h
--- a/Engine/Plugins/Animation/ControlRig/Source/ControlRigEditor/Private/EditMode/SControlRigControlViews.h
+++ b/Engine/Plugins/Animation/ControlRig/Source/ControlRigEditor/Private/EditMode/SControlRigControlViews.h
@@ -74,7 +74,8 @@ private:
 	TSharedRef<SWidget> GetThumbnailWidget();
 
 	TWeakObjectPtr<UControlRigPoseAsset> PoseAsset;
-	
+	TWeakObjectPtr<UControlRig> CurrentControlRig;
+
 	TSharedPtr<STextBlock> TextStatusBlock1;
 	TSharedPtr<STextBlock> TextStatusBlock2;
 
