Fix for RigVM/RigHierachy crash due to CheckNoEntry check when serializing.
Load EEP_WakeUP_Replication
Load crowd0x0 level sequence

diff --git a/Engine/Plugins/Animation/ControlRig/Source/ControlRig/Private/Rigs/RigHierarchy.cpp b/Engine/Plugins/Animation/ControlRig/Source/ControlRig/Private/Rigs/RigHierarchy.cpp
--- a/Engine/Plugins/Animation/ControlRig/Source/ControlRig/Private/Rigs/RigHierarchy.cpp
+++ b/Engine/Plugins/Animation/ControlRig/Source/ControlRig/Private/Rigs/RigHierarchy.cpp
@@ -124,7 +124,7 @@ void URigHierarchy::Serialize(FArchive& Ar)
 	}
 	else
 	{
-		checkNoEntry();
+		// remove due to FPIEFixupSerializer hitting this checkNoEntry();
 	}
 }
 
diff --git a/Engine/Plugins/Animation/ControlRig/Source/ControlRig/Private/Rigs/RigHierarchyDefines.cpp b/Engine/Plugins/Animation/ControlRig/Source/ControlRig/Private/Rigs/RigHierarchyDefines.cpp
--- a/Engine/Plugins/Animation/ControlRig/Source/ControlRig/Private/Rigs/RigHierarchyDefines.cpp
+++ b/Engine/Plugins/Animation/ControlRig/Source/ControlRig/Private/Rigs/RigHierarchyDefines.cpp
@@ -19,7 +19,7 @@ void FRigElementKey::Serialize(FArchive& Ar)
 	}
 	else
 	{
-		checkNoEntry();
+		// remove due to FPIEFixupSerializer hitting this checkNoEntry();
 	}
 }
 
diff --git a/Engine/Plugins/Animation/ControlRig/Source/ControlRig/Private/Rigs/RigHierarchyElements.cpp b/Engine/Plugins/Animation/ControlRig/Source/ControlRig/Private/Rigs/RigHierarchyElements.cpp
--- a/Engine/Plugins/Animation/ControlRig/Source/ControlRig/Private/Rigs/RigHierarchyElements.cpp
+++ b/Engine/Plugins/Animation/ControlRig/Source/ControlRig/Private/Rigs/RigHierarchyElements.cpp
@@ -59,7 +59,7 @@ void FRigBaseElement::Serialize(FArchive& Ar, URigHierarchy* Hierarchy, ESeriali
 	}
 	else
 	{
-		checkNoEntry();
+		// remove due to FPIEFixupSerializer hitting this checkNoEntry();
 	}
 }
 
diff --git a/Engine/Source/Runtime/RigVM/Private/RigVMCore/RigVM.cpp b/Engine/Source/Runtime/RigVM/Private/RigVMCore/RigVM.cpp
--- a/Engine/Source/Runtime/RigVM/Private/RigVMCore/RigVM.cpp
+++ b/Engine/Source/Runtime/RigVM/Private/RigVMCore/RigVM.cpp
@@ -25,7 +25,7 @@ void FRigVMParameter::Serialize(FArchive& Ar)
 	}
 	else
 	{
-		checkNoEntry();
+		// remove due to FPIEFixupSerializer hitting this checkNoEntry();
 	}
 }
 
@@ -126,7 +126,7 @@ void URigVM::Serialize(FArchive& Ar)
 	}
 	else
 	{
-		checkNoEntry();
+		// remove due to FPIEFixupSerializer hitting this checkNoEntry();
 	}
 }
 
