Control Rig Editor: Guard against details panel crashes

diff --git a/Engine/Plugins/Animation/ControlRig/Source/ControlRigEditor/Private/ControlRigGraphDetails.cpp b/Engine/Plugins/Animation/ControlRig/Source/ControlRigEditor/Private/ControlRigGraphDetails.cpp
--- a/Engine/Plugins/Animation/ControlRig/Source/ControlRigEditor/Private/ControlRigGraphDetails.cpp
+++ b/Engine/Plugins/Animation/ControlRig/Source/ControlRigEditor/Private/ControlRigGraphDetails.cpp
@@ -1255,33 +1255,65 @@ void FControlRigWrappedNodeDetails::CustomizeDetails(IDetailLayoutBuilder& Detai
 				{
 					continue;
 				}
-				
-				URigVMMemoryStorage* Memory = VM->GetMemoryByType(Operand->GetMemoryType());
-				if(Memory == nullptr)
-				{
-					continue;
-				}
 
-				if(!Memory->IsValidIndex(Operand->GetRegisterIndex()))
+				const FProperty* Property = nullptr;
+				TArray<UObject*> ExternalObjects;
+
+				if(Operand->GetMemoryType() == ERigVMMemoryType::External)
 				{
-					continue;
+					if(!VM->GetExternalVariables().IsValidIndex(Operand->GetRegisterIndex()))
+					{
+						continue;
+					}
+					ExternalObjects.Add(DebuggedRig);
+					Property = VM->GetExternalVariables()[Operand->GetRegisterIndex()].Property; 
 				}
-				
-				const FProperty* Property = Memory->GetProperty(Operand->GetRegisterIndex());
-				if(Property == nullptr)
+				else
 				{
-					continue;
+					URigVMMemoryStorage* Memory = VM->GetMemoryByType(Operand->GetMemoryType());
+					if(Memory == nullptr)
+					{
+						continue;
+					}
+
+					if(Memory->GetOuter() == GetTransientPackage())
+					{
+						continue;
+					}
+
+					if(Memory->GetClass()->GetOuter() == GetTransientPackage())
+					{
+						continue;
+					}
+
+					if(!Memory->IsValidIndex(Operand->GetRegisterIndex()))
+					{
+						continue;
+					}
+					
+					Property = Memory->GetProperty(Operand->GetRegisterIndex());
+					if(Property == nullptr)
+					{
+						continue;
+					}
+
+					ExternalObjects.Add(Memory);
 				}
 
-				TArray<UObject*> ExternalObjects;
-				ExternalObjects.Add(Memory);
-				DebugCategory.AddExternalObjectProperty(ExternalObjects, Property->GetFName(), EPropertyLocation::Default, FAddPropertyParams().ForceShowProperty())
-				->DisplayName(FText::FromString(FString::Printf(TEXT("%s%s"), *Pin->GetName(), *NameSuffix)))
-				.IsEnabled(false);
+				check(ExternalObjects.Num() > 0);
+				check(Property);
+
+				IDetailPropertyRow* PropertyRow = DebugCategory.AddExternalObjectProperty(ExternalObjects, Property->GetFName(), EPropertyLocation::Default, FAddPropertyParams().ForceShowProperty());
+				if(PropertyRow)
+				{
+					PropertyRow->DisplayName(FText::FromString(FString::Printf(TEXT("%s%s"), *Pin->GetName(), *NameSuffix)));
+					PropertyRow->IsEnabled(false);
+
+					SuffixIndex++;
+					bAddedProperty = true;
+					NameSuffix = FString::Printf(TEXT("_%d"), SuffixIndex);
+				}
 
-				SuffixIndex++;
-				bAddedProperty = true;
-				NameSuffix = FString::Printf(TEXT("_%d"), SuffixIndex);
 				KnownOperands.Add(*Operand);
 			}
 		}
@@ -1291,7 +1323,7 @@ void FControlRigWrappedNodeDetails::CustomizeDetails(IDetailLayoutBuilder& Detai
 			TSharedPtr<IPropertyHandle> PinHandle = StoredStructHandle->GetChildHandle(Pin->GetFName());
 			if(PinHandle.IsValid())
 			{
-				DefaultsCategory.AddProperty(PinHandle)
+				DebugCategory.AddProperty(PinHandle)
 				.DisplayName(FText::FromName(Pin->GetDisplayName()))
 				.IsEnabled(false);
 			}
