Control Rig: Move Python related buttons to their own Python category

diff --git a/Engine/Plugins/Animation/ControlRig/Source/ControlRigDeveloper/Public/ControlRigBlueprint.h b/Engine/Plugins/Animation/ControlRig/Source/ControlRigDeveloper/Public/ControlRigBlueprint.h
--- a/Engine/Plugins/Animation/ControlRig/Source/ControlRigDeveloper/Public/ControlRigBlueprint.h
+++ b/Engine/Plugins/Animation/ControlRig/Source/ControlRigDeveloper/Public/ControlRigBlueprint.h
@@ -138,6 +138,16 @@ struct CONTROLRIGDEVELOPER_API FRigGraphDisplaySettings
 	int32 NodeRunLimit;
 };
 
+USTRUCT()
+struct CONTROLRIGDEVELOPER_API FControlRigPythonSettings
+{
+	GENERATED_BODY();
+
+	FControlRigPythonSettings()
+	{
+	}
+};
+
 UCLASS(BlueprintType, meta=(IgnoreClassThumbnail))
 class CONTROLRIGDEVELOPER_API UControlRigBlueprint : public UBlueprint, public IInterface_PreviewMeshProvider
 {
@@ -286,6 +296,9 @@ public:
 	UPROPERTY(EditAnywhere, BlueprintReadWrite, Category = "VM")
 	FRigVMCompileSettings VMCompileSettings;
 
+	UPROPERTY(EditAnywhere, Category = "Python Log Settings")
+	FControlRigPythonSettings PythonLogSettings;
+
 protected:
 
 	UPROPERTY(BlueprintReadOnly, Category = "VM")
diff --git a/Engine/Plugins/Animation/ControlRig/Source/ControlRigEditor/Private/ControlRigCompilerDetails.cpp b/Engine/Plugins/Animation/ControlRig/Source/ControlRigEditor/Private/ControlRigCompilerDetails.cpp
--- a/Engine/Plugins/Animation/ControlRig/Source/ControlRigEditor/Private/ControlRigCompilerDetails.cpp
+++ b/Engine/Plugins/Animation/ControlRig/Source/ControlRigEditor/Private/ControlRigCompilerDetails.cpp
@@ -126,30 +126,6 @@ void FRigVMCompileSettingsDetails::CustomizeChildren(TSharedRef<IPropertyHandle>
 						.Text(LOCTEXT("CopyByteCodeToClipboard", "Copy ByteCode"))
 					]
 				]
-				+ SVerticalBox::Slot()
-				[
-					SNew(SButton)
-					.OnClicked(this, &FRigVMCompileSettingsDetails::OnCopyPythonScriptClicked)
-					.ContentPadding(FMargin(2))
-					.Content()
-					[
-						SNew(STextBlock)
-						.Justification(ETextJustify::Center)
-						.Text(LOCTEXT("CopyPythonScript", "Copy Python Script"))
-					]
-				]
-				+ SVerticalBox::Slot()
-				[
-					SNew(SButton)
-					.OnClicked(this, &FRigVMCompileSettingsDetails::OnRunPythonContextClicked)
-					.ContentPadding(FMargin(2))
-					.Content()
-					[
-						SNew(STextBlock)
-						.Justification(ETextJustify::Center)
-						.Text(LOCTEXT("RunPythonContext", "Run Python Context"))
-					]
-				]
 			];
 	}
 }
@@ -203,52 +179,4 @@ FReply FRigVMCompileSettingsDetails::OnCopyByteCodeClicked()
 	return FReply::Handled();
 }
 
-FReply FRigVMCompileSettingsDetails::OnCopyPythonScriptClicked()
-{
-	if (BlueprintBeingCustomized)
-	{
-		FString NewName = BlueprintBeingCustomized->GetPathName();
-		int32 DotIndex = NewName.Find(TEXT("."), ESearchCase::IgnoreCase, ESearchDir::FromEnd);
-		if (DotIndex != INDEX_NONE)
-		{
-			NewName = NewName.Left(DotIndex);
-		}
-		
-		TArray<FString> Commands = BlueprintBeingCustomized->GeneratePythonCommands(NewName);
-		FString FullScript = FString::Join(Commands, TEXT("\n"));
-		FPlatformApplicationMisc::ClipboardCopy(*FullScript);
-	}
-	return FReply::Handled();
-}
-
-FReply FRigVMCompileSettingsDetails::OnRunPythonContextClicked()
-{
-	if (BlueprintBeingCustomized)
-	{
-		FString BlueprintName = BlueprintBeingCustomized->GetPathName();
-		int32 DotIndex = BlueprintName.Find(TEXT("."), ESearchCase::IgnoreCase, ESearchDir::FromEnd);
-		if (DotIndex != INDEX_NONE)
-		{
-			BlueprintName = BlueprintName.Left(DotIndex);
-		}
-		
-		TArray<FString> PyCommands = {
-			TEXT("import unreal"),
-			FString::Printf(TEXT("blueprint = unreal.load_object(name = '%s', outer = None)"), *BlueprintName),
-			TEXT("library = blueprint.get_local_function_library()"),
-			TEXT("library_controller = blueprint.get_controller(library)"),
-			TEXT("hierarchy = blueprint.hierarchy"),
-			TEXT("hierarchy_controller = hierarchy.get_controller()")};
-
-		for (FString& Command : PyCommands)
-		{
-			RigVMPythonUtils::Print(BlueprintBeingCustomized->GetFName().ToString(), Command);
-		
-			// Run the Python commands
-			IPythonScriptPlugin::Get()->ExecPythonCommand(*Command);
-		}
-	}
-	return FReply::Handled();
-}
-
 #undef LOCTEXT_NAMESPACE
diff --git a/Engine/Plugins/Animation/ControlRig/Source/ControlRigEditor/Private/ControlRigCompilerDetails.h b/Engine/Plugins/Animation/ControlRig/Source/ControlRigEditor/Private/ControlRigCompilerDetails.h
--- a/Engine/Plugins/Animation/ControlRig/Source/ControlRigEditor/Private/ControlRigCompilerDetails.h
+++ b/Engine/Plugins/Animation/ControlRig/Source/ControlRigEditor/Private/ControlRigCompilerDetails.h
@@ -29,8 +29,6 @@ protected:
 #endif
 	FReply OnCopyASTClicked();
 	FReply OnCopyByteCodeClicked();
-	FReply OnCopyPythonScriptClicked();
-	FReply OnRunPythonContextClicked();
-
+	
 	UControlRigBlueprint* BlueprintBeingCustomized;
 };
diff --git a/Engine/Plugins/Animation/ControlRig/Source/ControlRigEditor/Private/ControlRigEditorModule.cpp b/Engine/Plugins/Animation/ControlRig/Source/ControlRigEditor/Private/ControlRigEditorModule.cpp
--- a/Engine/Plugins/Animation/ControlRig/Source/ControlRigEditor/Private/ControlRigEditorModule.cpp
+++ b/Engine/Plugins/Animation/ControlRig/Source/ControlRigEditor/Private/ControlRigEditorModule.cpp
@@ -111,6 +111,7 @@
 #include "Graph/SControlRigGraphPinVariableBinding.h"
 #include "AssetTypeActions_ControlRigPose.h"
 #include "ControlRigBlueprintFactory.h"
+#include "ControlRigPythonLogDetails.h"
 #include "EditMode/SControlRigBaseListWidget.h"
 #include "EditMode/SControlRigTweenWidget.h"
 #include "EditMode/SControlRigSnapper.h"
@@ -209,6 +210,9 @@ void FControlRigEditorModule::StartupModule()
 	PropertiesToUnregisterOnShutdown.Add(FRigVMCompileSettings::StaticStruct()->GetFName());
 	PropertyEditorModule.RegisterCustomPropertyTypeLayout(PropertiesToUnregisterOnShutdown.Last(), FOnGetPropertyTypeCustomizationInstance::CreateStatic(&FRigVMCompileSettingsDetails::MakeInstance));
 
+	PropertiesToUnregisterOnShutdown.Add(FControlRigPythonSettings::StaticStruct()->GetFName());
+	PropertyEditorModule.RegisterCustomPropertyTypeLayout(PropertiesToUnregisterOnShutdown.Last(), FOnGetPropertyTypeCustomizationInstance::CreateStatic(&FControlRigPythonLogDetails::MakeInstance));
+
 	PropertiesToUnregisterOnShutdown.Add(FControlRigDrawContainer::StaticStruct()->GetFName());
 	PropertyEditorModule.RegisterCustomPropertyTypeLayout(PropertiesToUnregisterOnShutdown.Last(), FOnGetPropertyTypeCustomizationInstance::CreateStatic(&FControlRigDrawContainerDetails::MakeInstance));
 
diff --git a/Engine/Plugins/Animation/ControlRig/Source/ControlRigEditor/Private/ControlRigPythonLogDetails.cpp b/Engine/Plugins/Animation/ControlRig/Source/ControlRigEditor/Private/ControlRigPythonLogDetails.cpp
new file mode 100644
--- /dev/null
+++ b/Engine/Plugins/Animation/ControlRig/Source/ControlRigEditor/Private/ControlRigPythonLogDetails.cpp
@@ -0,0 +1,142 @@
+﻿// Copyright Epic Games, Inc. All Rights Reserved.
+
+#include "ControlRigPythonLogDetails.h"
+#include "Widgets/SWidget.h"
+#include "DetailLayoutBuilder.h"
+#include "DetailCategoryBuilder.h"
+#include "DetailWidgetRow.h"
+#include "IDetailChildrenBuilder.h"
+#include "Widgets/Input/SButton.h"
+#include "Widgets/Text/STextBlock.h"
+#include "HAL/PlatformApplicationMisc.h"
+#include "ControlRig.h"
+#include "IPropertyUtilities.h"
+#include "IPythonScriptPlugin.h"
+#include "RigVMPythonUtils.h"
+
+#define LOCTEXT_NAMESPACE "ControlRigCompilerDetails"
+
+void FControlRigPythonLogDetails::CustomizeHeader(TSharedRef<IPropertyHandle> InStructPropertyHandle, FDetailWidgetRow& HeaderRow, IPropertyTypeCustomizationUtils& StructCustomizationUtils)
+{
+	HeaderRow
+	.NameContent()
+	[
+		InStructPropertyHandle->CreatePropertyNameWidget()
+	]
+	.ValueContent()
+	[
+		InStructPropertyHandle->CreatePropertyValueWidget()
+	];
+
+	TArray<UObject*> Objects;
+	InStructPropertyHandle->GetOuterObjects(Objects);
+	ensure(Objects.Num() == 1); // This is in here to ensure we are only showing the modifier details in the blueprint editor
+
+	for (UObject* Object : Objects)
+	{
+		if (Object->IsA<UControlRigBlueprint>())
+		{
+			BlueprintBeingCustomized = Cast<UControlRigBlueprint>(Object);
+		}
+	}
+}
+
+void FControlRigPythonLogDetails::CustomizeChildren(TSharedRef<IPropertyHandle> InStructPropertyHandle, IDetailChildrenBuilder& StructBuilder, IPropertyTypeCustomizationUtils& StructCustomizationUtils)
+{
+	if (InStructPropertyHandle->IsValidHandle())
+	{
+		uint32 NumChildren = 0;
+		InStructPropertyHandle->GetNumChildren(NumChildren);
+
+		for (uint32 ChildIndex = 0; ChildIndex < NumChildren; ChildIndex++)
+		{
+			StructBuilder.AddProperty(InStructPropertyHandle->GetChildHandle(ChildIndex).ToSharedRef());
+		}
+
+		StructBuilder.AddCustomRow(LOCTEXT("Commands", "Python Commands"))
+			.NameContent()
+			[
+				SNew(STextBlock)
+				.Text(FText::FromString(TEXT("Commands")))
+			.Font(IDetailLayoutBuilder::GetDetailFont())
+			]
+			.ValueContent()
+			[
+				SNew(SVerticalBox)
+				+ SVerticalBox::Slot()
+				[
+					SNew(SButton)
+					.OnClicked(this, &FControlRigPythonLogDetails::OnCopyPythonScriptClicked)
+					.ContentPadding(FMargin(2))
+					.Content()
+					[
+						SNew(STextBlock)
+						.Justification(ETextJustify::Center)
+						.Text(LOCTEXT("CopyPythonScript", "Copy Python Script"))
+					]
+				]
+				+ SVerticalBox::Slot()
+				[
+					SNew(SButton)
+					.OnClicked(this, &FControlRigPythonLogDetails::OnRunPythonContextClicked)
+					.ContentPadding(FMargin(2))
+					.Content()
+					[
+						SNew(STextBlock)
+						.Justification(ETextJustify::Center)
+						.Text(LOCTEXT("RunPythonContext", "Run Python Context"))
+					]
+				]
+			];
+	}
+}
+
+FReply FControlRigPythonLogDetails::OnCopyPythonScriptClicked()
+{
+	if (BlueprintBeingCustomized)
+	{
+		FString NewName = BlueprintBeingCustomized->GetPathName();
+		int32 DotIndex = NewName.Find(TEXT("."), ESearchCase::IgnoreCase, ESearchDir::FromEnd);
+		if (DotIndex != INDEX_NONE)
+		{
+			NewName = NewName.Left(DotIndex);
+		}
+		
+		TArray<FString> Commands = BlueprintBeingCustomized->GeneratePythonCommands(NewName);
+		FString FullScript = FString::Join(Commands, TEXT("\n"));
+		FPlatformApplicationMisc::ClipboardCopy(*FullScript);
+	}
+	return FReply::Handled();
+}
+
+FReply FControlRigPythonLogDetails::OnRunPythonContextClicked()
+{
+	if (BlueprintBeingCustomized)
+	{
+		FString BlueprintName = BlueprintBeingCustomized->GetPathName();
+		int32 DotIndex = BlueprintName.Find(TEXT("."), ESearchCase::IgnoreCase, ESearchDir::FromEnd);
+		if (DotIndex != INDEX_NONE)
+		{
+			BlueprintName = BlueprintName.Left(DotIndex);
+		}
+		
+		TArray<FString> PyCommands = {
+			TEXT("import unreal"),
+			FString::Printf(TEXT("blueprint = unreal.load_object(name = '%s', outer = None)"), *BlueprintName),
+			TEXT("library = blueprint.get_local_function_library()"),
+			TEXT("library_controller = blueprint.get_controller(library)"),
+			TEXT("hierarchy = blueprint.hierarchy"),
+			TEXT("hierarchy_controller = hierarchy.get_controller()")};
+
+		for (FString& Command : PyCommands)
+		{
+			RigVMPythonUtils::Print(BlueprintBeingCustomized->GetFName().ToString(), Command);
+		
+			// Run the Python commands
+			IPythonScriptPlugin::Get()->ExecPythonCommand(*Command);
+		}
+	}
+	return FReply::Handled();
+}
+
+#undef LOCTEXT_NAMESPACE
diff --git a/Engine/Plugins/Animation/ControlRig/Source/ControlRigEditor/Private/ControlRigPythonLogDetails.h b/Engine/Plugins/Animation/ControlRig/Source/ControlRigEditor/Private/ControlRigPythonLogDetails.h
new file mode 100644
--- /dev/null
+++ b/Engine/Plugins/Animation/ControlRig/Source/ControlRigEditor/Private/ControlRigPythonLogDetails.h
@@ -0,0 +1,30 @@
+﻿// Copyright Epic Games, Inc. All Rights Reserved.
+
+#pragma once
+
+#include "CoreMinimal.h"
+#include "IPropertyTypeCustomization.h"
+#include "ControlRigBlueprint.h"
+
+class IPropertyHandle;
+
+class FControlRigPythonLogDetails : public IPropertyTypeCustomization
+{
+public:
+
+	static TSharedRef<IPropertyTypeCustomization> MakeInstance()
+	{
+		return MakeShareable(new FControlRigPythonLogDetails);
+	}
+
+	/** IPropertyTypeCustomization interface */
+	virtual void CustomizeHeader(TSharedRef<class IPropertyHandle> InStructPropertyHandle, class FDetailWidgetRow& HeaderRow, IPropertyTypeCustomizationUtils& StructCustomizationUtils) override;
+	virtual void CustomizeChildren(TSharedRef<class IPropertyHandle> InStructPropertyHandle, class IDetailChildrenBuilder& StructBuilder, IPropertyTypeCustomizationUtils& StructCustomizationUtils) override;
+
+protected:
+
+	FReply OnCopyPythonScriptClicked();
+	FReply OnRunPythonContextClicked();
+
+	UControlRigBlueprint* BlueprintBeingCustomized;
+};
