Control Rig: Fix crash when ControlRigEditor is closed after importing hierarchy

diff --git a/Engine/Plugins/Animation/ControlRig/Source/ControlRigEditor/Private/Editor/ControlRigEditor.cpp b/Engine/Plugins/Animation/ControlRig/Source/ControlRigEditor/Private/Editor/ControlRigEditor.cpp
--- a/Engine/Plugins/Animation/ControlRig/Source/ControlRigEditor/Private/Editor/ControlRigEditor.cpp
+++ b/Engine/Plugins/Animation/ControlRig/Source/ControlRigEditor/Private/Editor/ControlRigEditor.cpp
@@ -1511,7 +1511,13 @@ void FControlRigEditor::SetDetailViewForRigElements()
 			continue;
 		}
 
-		UDetailsViewWrapperObject* WrapperObject = UDetailsViewWrapperObject::MakeInstance(Element->GetElementStruct(), (uint8*)Element, Hierarchy);
+		URigHierarchy* SelectionHierarchy = Hierarchy;
+		if (Hierarchy->HierarchyForSelectionPtr.IsValid())
+		{
+			SelectionHierarchy = Hierarchy->HierarchyForSelectionPtr.Get();
+		}
+
+		UDetailsViewWrapperObject* WrapperObject = UDetailsViewWrapperObject::MakeInstance(Element->GetElementStruct(), (uint8*)Element, SelectionHierarchy);
 		WrapperObject->GetWrappedPropertyChangedChainEvent().AddSP(this, &FControlRigEditor::OnWrappedPropertyChangedChainEvent);
 		WrapperObject->AddToRoot();
 
