Add If check for invalid world on component for Chance.
#trivial

diff --git a/Engine/Plugins/Animation/ControlRig/Source/ControlRig/Private/Sequencer/MovieSceneControlRigParameterTemplate.cpp b/Engine/Plugins/Animation/ControlRig/Source/ControlRig/Private/Sequencer/MovieSceneControlRigParameterTemplate.cpp
--- a/Engine/Plugins/Animation/ControlRig/Source/ControlRig/Private/Sequencer/MovieSceneControlRigParameterTemplate.cpp
+++ b/Engine/Plugins/Animation/ControlRig/Source/ControlRig/Private/Sequencer/MovieSceneControlRigParameterTemplate.cpp
@@ -870,58 +870,64 @@ static UControlRig* GetControlRig(const UMovieSceneControlRigParameterSection* S
 			{
 				if (UControlRigComponent* NewControlRigComponent = Actor->FindComponentByClass<UControlRigComponent>())
 				{
-					if (NewControlRigComponent->GetWorld()->IsGameWorld())
+					if (NewControlRigComponent->GetWorld())
 					{
-						ControlRig = NewControlRigComponent->GetControlRig();
-						if (!ControlRig)
+						if (NewControlRigComponent->GetWorld()->IsGameWorld())
 						{
-							NewControlRigComponent->Initialize();
 							ControlRig = NewControlRigComponent->GetControlRig();
-						}
-						if (ControlRig)
-						{
-							if (ControlRig->GetObjectBinding() == nullptr)
+							if (!ControlRig)
 							{
-								ControlRig->SetObjectBinding(MakeShared<FControlRigObjectBinding>());
+								NewControlRigComponent->Initialize();
+								ControlRig = NewControlRigComponent->GetControlRig();
 							}
-							if (ControlRig->GetObjectBinding()->GetBoundObject() != NewControlRigComponent)
+							if (ControlRig)
 							{
-								ControlRig->GetObjectBinding()->BindToObject(BoundObject);							
+								if (ControlRig->GetObjectBinding() == nullptr)
+								{
+									ControlRig->SetObjectBinding(MakeShared<FControlRigObjectBinding>());
+								}
+								if (ControlRig->GetObjectBinding()->GetBoundObject() != NewControlRigComponent)
+								{
+									ControlRig->GetObjectBinding()->BindToObject(BoundObject);
+								}
 							}
 						}
-					}
-					else if (NewControlRigComponent != ControlRigComponent)
-					{
-						NewControlRigComponent->SetControlRig(ControlRig);
+						else if (NewControlRigComponent != ControlRigComponent)
+						{
+							NewControlRigComponent->SetControlRig(ControlRig);
+						}
 					}
 				}
 			}
 		}
 		else if (UControlRigComponent* NewControlRigComponent = Cast<UControlRigComponent>(BoundObject))
 		{
-			if (NewControlRigComponent->GetWorld()->IsGameWorld())
+			if (NewControlRigComponent->GetWorld())
 			{
-				ControlRig = NewControlRigComponent->GetControlRig();
-				if (!ControlRig)
+				if (NewControlRigComponent->GetWorld()->IsGameWorld())
 				{
-					NewControlRigComponent->Initialize();
 					ControlRig = NewControlRigComponent->GetControlRig();
-				}
-				if (ControlRig)
-				{
-					if (ControlRig->GetObjectBinding() == nullptr)
+					if (!ControlRig)
 					{
-						ControlRig->SetObjectBinding(MakeShared<FControlRigObjectBinding>());
+						NewControlRigComponent->Initialize();
+						ControlRig = NewControlRigComponent->GetControlRig();
 					}
-					if (ControlRig->GetObjectBinding()->GetBoundObject() != NewControlRigComponent)
+					if (ControlRig)
 					{
-						ControlRig->GetObjectBinding()->BindToObject(BoundObject);
+						if (ControlRig->GetObjectBinding() == nullptr)
+						{
+							ControlRig->SetObjectBinding(MakeShared<FControlRigObjectBinding>());
+						}
+						if (ControlRig->GetObjectBinding()->GetBoundObject() != NewControlRigComponent)
+						{
+							ControlRig->GetObjectBinding()->BindToObject(BoundObject);
+						}
 					}
 				}
-			}
-			else if (NewControlRigComponent != ControlRigComponent)
-			{
-				NewControlRigComponent->SetControlRig(ControlRig);
+				else if (NewControlRigComponent != ControlRigComponent)
+				{
+					NewControlRigComponent->SetControlRig(ControlRig);
+				}
 			}
 		}
 	}
