Add one more check for invalid world on component for Chance.
#trivial

diff --git a/Engine/Plugins/Animation/ControlRig/Source/ControlRig/Private/Sequencer/MovieSceneControlRigParameterTemplate.cpp b/Engine/Plugins/Animation/ControlRig/Source/ControlRig/Private/Sequencer/MovieSceneControlRigParameterTemplate.cpp
--- a/Engine/Plugins/Animation/ControlRig/Source/ControlRig/Private/Sequencer/MovieSceneControlRigParameterTemplate.cpp
+++ b/Engine/Plugins/Animation/ControlRig/Source/ControlRig/Private/Sequencer/MovieSceneControlRigParameterTemplate.cpp
@@ -829,9 +829,12 @@ struct FControlRigParameterPreAnimatedTokenProducer : IMovieScenePreAnimatedToke
 						in non-game, and if in game (which includes PIE), we don't re-set the 
 						CR Component's CR, but instead grab the CR from it and then use that for evaluation.
 						*/
-						if (ControlRigComponent->GetWorld()->IsGameWorld() == false)
+						if (ControlRigComponent->GetWorld())
 						{
-							ControlRigComponent->SetControlRig(ControlRig);
+							if (ControlRigComponent->GetWorld()->IsGameWorld() == false)
+							{
+								ControlRigComponent->SetControlRig(ControlRig);
+							}
 						}
 					}
 					else
