Control Rig: initialize VM memory proactively to avoid creating memory storage uobjects in worker thread

diff --git a/Engine/Plugins/Animation/ControlRig/Source/ControlRig/Private/ControlRig.cpp b/Engine/Plugins/Animation/ControlRig/Source/ControlRig/Private/ControlRig.cpp
--- a/Engine/Plugins/Animation/ControlRig/Source/ControlRig/Private/ControlRig.cpp
+++ b/Engine/Plugins/Animation/ControlRig/Source/ControlRig/Private/ControlRig.cpp
@@ -1109,6 +1109,14 @@ void UControlRig::SetVM(URigVM* NewVM)
 		};
 #endif
 		NewVM->SetRuntimeSettings(VMRuntimeSettings);
+
+#if !UE_RIGVM_UCLASS_BASED_STORAGE_DISABLED
+		// ensure memory storage objects are created during initialization in game thread
+		// delaying this step may lead to async uobject creation and causing GC to complain
+		NewVM->GetMemoryByType(ERigVMMemoryType::Work, true);
+		NewVM->GetMemoryByType(ERigVMMemoryType::Literal, true);
+		NewVM->GetMemoryByType(ERigVMMemoryType::Debug, true);
+#endif
 	}
 
 	VM = NewVM;
diff --git a/Engine/Source/Runtime/RigVM/Private/RigVMCore/RigVM.cpp b/Engine/Source/Runtime/RigVM/Private/RigVMCore/RigVM.cpp
--- a/Engine/Source/Runtime/RigVM/Private/RigVMCore/RigVM.cpp
+++ b/Engine/Source/Runtime/RigVM/Private/RigVMCore/RigVM.cpp
@@ -817,7 +817,7 @@ URigVMMemoryStorage* URigVM::GetMemoryByType(ERigVMMemoryType InMemoryType, bool
 					}
 				}
 			}
-			check(WorkMemoryStorageObject->GetOuter() == this);
+			check(DebugMemoryStorageObject->GetOuter() == this);
 			return DebugMemoryStorageObject;
 		}
 		default:
