Control Rig: Mark PreviewMesh as EditorOnlyData

diff --git a/Engine/Plugins/Animation/ControlRig/Source/ControlRigDeveloper/Private/ControlRigBlueprint.cpp b/Engine/Plugins/Animation/ControlRig/Source/ControlRigDeveloper/Private/ControlRigBlueprint.cpp
--- a/Engine/Plugins/Animation/ControlRig/Source/ControlRigDeveloper/Private/ControlRigBlueprint.cpp
+++ b/Engine/Plugins/Animation/ControlRig/Source/ControlRigDeveloper/Private/ControlRigBlueprint.cpp
@@ -257,22 +257,28 @@ USkeletalMesh* UControlRigBlueprint::GetPreviewMesh() const
 {
 	DECLARE_SCOPE_HIERARCHICAL_COUNTER_FUNC()
 
+#if WITH_EDITORONLY_DATA
 	if (!PreviewSkeletalMesh.IsValid())
 	{
 		PreviewSkeletalMesh.LoadSynchronous();
 	}
 
 	return PreviewSkeletalMesh.Get();
+#else
+	return nullptr;
+#endif
 }
 
 void UControlRigBlueprint::SetPreviewMesh(USkeletalMesh* PreviewMesh, bool bMarkAsDirty/*=true*/)
 {
+#if WITH_EDITORONLY_DATA
 	if(bMarkAsDirty)
 	{
 		Modify();
 	}
 
 	PreviewSkeletalMesh = PreviewMesh;
+#endif
 }
 
 void UControlRigBlueprint::Serialize(FArchive& Ar)
@@ -1448,9 +1454,11 @@ TArray<FString> UControlRigBlueprint::GeneratePythonCommands(const FString InNew
 		}
 	}
 
+#if WITH_EDITORONLY_DATA
 	FString PreviewMeshPath = GetPreviewMesh()->GetPathName();
 	Commands.Add(FString::Printf(TEXT("blueprint.set_preview_mesh(unreal.load_object(name='%s', outer=None))"),
 		*PreviewMeshPath));
+#endif
 
 	return Commands;
 }
diff --git a/Engine/Plugins/Animation/ControlRig/Source/ControlRigDeveloper/Public/ControlRigBlueprint.h b/Engine/Plugins/Animation/ControlRig/Source/ControlRigDeveloper/Public/ControlRigBlueprint.h
--- a/Engine/Plugins/Animation/ControlRig/Source/ControlRigDeveloper/Public/ControlRigBlueprint.h
+++ b/Engine/Plugins/Animation/ControlRig/Source/ControlRigDeveloper/Public/ControlRigBlueprint.h
@@ -423,8 +423,10 @@ private:
 	bool bSupportsControls;
 
 	/** The default skeletal mesh to use when previewing this asset */
+#if WITH_EDITORONLY_DATA
 	UPROPERTY(AssetRegistrySearchable)
 	TSoftObjectPtr<USkeletalMesh> PreviewSkeletalMesh;
+#endif
 
 	/** The skeleton from import into a hierarchy */
 	UPROPERTY(DuplicateTransient, AssetRegistrySearchable)
