Control Rig: delay rig vm memory intiialization from control rig constructor to initializer() since in packaged build during control rig constructor the memory class may not be ready

diff --git a/Engine/Plugins/Animation/ControlRig/Source/ControlRig/Private/ControlRig.cpp b/Engine/Plugins/Animation/ControlRig/Source/ControlRig/Private/ControlRig.cpp
--- a/Engine/Plugins/Animation/ControlRig/Source/ControlRig/Private/ControlRig.cpp
+++ b/Engine/Plugins/Animation/ControlRig/Source/ControlRig/Private/ControlRig.cpp
@@ -467,6 +467,14 @@ void UControlRig::InstantiateVMFromCDO()
 	if (VM == nullptr || VM->GetOuter() != this)
 	{
 		SetVM(NewObject<URigVM>(this, TEXT("VM")));
+		
+#if !UE_RIGVM_UCLASS_BASED_STORAGE_DISABLED
+		// ensure memory storage objects are created during initialization in game thread
+		// delaying this step may lead to async uobject creation and causing GC to complain
+		VM->GetMemoryByType(ERigVMMemoryType::Work, true);
+		VM->GetMemoryByType(ERigVMMemoryType::Literal, true);
+		VM->GetMemoryByType(ERigVMMemoryType::Debug, true);
+#endif
 	}
 	
 	if (!HasAnyFlags(RF_ClassDefaultObject))
@@ -1110,13 +1118,6 @@ void UControlRig::SetVM(URigVM* NewVM)
 #endif
 		NewVM->SetRuntimeSettings(VMRuntimeSettings);
 
-#if !UE_RIGVM_UCLASS_BASED_STORAGE_DISABLED
-		// ensure memory storage objects are created during initialization in game thread
-		// delaying this step may lead to async uobject creation and causing GC to complain
-		NewVM->GetMemoryByType(ERigVMMemoryType::Work, true);
-		NewVM->GetMemoryByType(ERigVMMemoryType::Literal, true);
-		NewVM->GetMemoryByType(ERigVMMemoryType::Debug, true);
-#endif
 	}
 
 	VM = NewVM;
