Control Rig: change a check to an ensure to work around a crash when loading a level sequence in PIE, proper fix will follow

diff --git a/Engine/Plugins/Animation/ControlRig/Source/ControlRig/Private/ControlRig.cpp b/Engine/Plugins/Animation/ControlRig/Source/ControlRig/Private/ControlRig.cpp
--- a/Engine/Plugins/Animation/ControlRig/Source/ControlRig/Private/ControlRig.cpp
+++ b/Engine/Plugins/Animation/ControlRig/Source/ControlRig/Private/ControlRig.cpp
@@ -2654,8 +2654,10 @@ void UControlRig::SetBoneInitialTransformsFromSkeletalMeshComponent(USkeletalMes
 
 void UControlRig::SetBoneInitialTransformsFromSkeletalMesh(USkeletalMesh* InSkeletalMesh)
 {
-	check(InSkeletalMesh);
-	SetBoneInitialTransformsFromRefSkeleton(InSkeletalMesh->GetRefSkeleton());
+	if (ensure(InSkeletalMesh))
+	{ 
+		SetBoneInitialTransformsFromRefSkeleton(InSkeletalMesh->GetRefSkeleton());
+	}
 }
 
 void UControlRig::SetBoneInitialTransformsFromRefSkeleton(const FReferenceSkeleton& InReferenceSkeleton)
