Control Rig: Avoid MemAllocator for RBF nodes

diff --git a/Engine/Plugins/Animation/ControlRig/Source/ControlRig/Private/Units/Math/RigUnit_MathRBFInterpolate.cpp b/Engine/Plugins/Animation/ControlRig/Source/ControlRig/Private/Units/Math/RigUnit_MathRBFInterpolate.cpp
--- a/Engine/Plugins/Animation/ControlRig/Source/ControlRig/Private/Units/Math/RigUnit_MathRBFInterpolate.cpp
+++ b/Engine/Plugins/Animation/ControlRig/Source/ControlRig/Private/Units/Math/RigUnit_MathRBFInterpolate.cpp
@@ -177,10 +177,10 @@ void FRigUnit_MathRBFInterpolateQuatBase::GetInterpolatedWeights(
 	float SmoothingAngle,
 	bool bNormalizeOutput,
 	FVector TwistAxis,
-	TArray<float, TMemStackAllocator<>> &Weights
+	TArray<float> &Weights
 )
 {
-	TArray<FQuat, TMemStackAllocator<> > Quats;
+	TArray<FQuat > Quats;
 	Quats.Reserve(Targets.Num());
 
 	for (const auto& Target : Targets)
@@ -240,10 +240,10 @@ void FRigUnit_MathRBFInterpolateVectorBase::GetInterpolatedWeights(
 	ERBFKernelType SmoothingFunction,
 	float SmoothingRadius,
 	bool bNormalizeOutput,
-	TArray<float, TMemStackAllocator<>>& Weights
+	TArray<float>& Weights
 )
 {
-	TArray<FVector, TMemStackAllocator<> > Vectors;
+	TArray<FVector > Vectors;
 	Vectors.Reserve(Targets.Num());
 
 	for (const auto& Target : Targets)
@@ -297,7 +297,7 @@ FRigUnit_MathRBFInterpolateQuatFloat_Execute()
 	DECLARE_SCOPE_HIERARCHICAL_COUNTER_RIGUNIT()
 
 	const TArrayView<const FMathRBFInterpolateQuatFloat_Target> TargetsView(Targets.GetData(), Targets.Num());
-	TArray<float, TMemStackAllocator<>> Weights;
+	TArray<float> Weights;
 	GetInterpolatedWeights<FMathRBFInterpolateQuatFloat_Target>(
 		Context.State, WorkData, TargetsView, Input,  
 		DistanceFunction, SmoothingFunction, SmoothingAngle, 
@@ -316,7 +316,7 @@ FRigUnit_MathRBFInterpolateQuatVector_Execute()
 	DECLARE_SCOPE_HIERARCHICAL_COUNTER_RIGUNIT()
 
 	const TArrayView<const FMathRBFInterpolateQuatVector_Target> TargetsView(Targets.GetData(), Targets.Num());
-	TArray<float, TMemStackAllocator<>> Weights;
+	TArray<float> Weights;
 	GetInterpolatedWeights<FMathRBFInterpolateQuatVector_Target>(
 		Context.State, WorkData, TargetsView, Input,
 		DistanceFunction, SmoothingFunction, SmoothingAngle, 
@@ -335,7 +335,7 @@ FRigUnit_MathRBFInterpolateQuatColor_Execute()
 	DECLARE_SCOPE_HIERARCHICAL_COUNTER_RIGUNIT()
 
 	const TArrayView<const FMathRBFInterpolateQuatColor_Target> TargetsView(Targets.GetData(), Targets.Num());
-	TArray<float, TMemStackAllocator<>> Weights;
+	TArray<float> Weights;
 	GetInterpolatedWeights<FMathRBFInterpolateQuatColor_Target>(
 		Context.State, WorkData, TargetsView, Input,
 		DistanceFunction, SmoothingFunction, SmoothingAngle, 
@@ -353,7 +353,7 @@ FRigUnit_MathRBFInterpolateQuatQuat_Execute()
 	DECLARE_SCOPE_HIERARCHICAL_COUNTER_RIGUNIT()
 
 	const TArrayView<const FMathRBFInterpolateQuatQuat_Target> TargetsView(Targets.GetData(), Targets.Num());
-	TArray<float, TMemStackAllocator<>> Weights;
+	TArray<float> Weights;
 	GetInterpolatedWeights<FMathRBFInterpolateQuatQuat_Target>(
 		Context.State, WorkData, TargetsView, Input,
 		DistanceFunction, SmoothingFunction, SmoothingAngle,
@@ -373,7 +373,7 @@ FRigUnit_MathRBFInterpolateQuatXform_Execute()
 	DECLARE_SCOPE_HIERARCHICAL_COUNTER_RIGUNIT()
 
 	const TArrayView<const FMathRBFInterpolateQuatXform_Target> TargetsView(Targets.GetData(), Targets.Num());
-	TArray<float, TMemStackAllocator<>> Weights;
+	TArray<float> Weights;
 	GetInterpolatedWeights<FMathRBFInterpolateQuatXform_Target>(
 		Context.State, WorkData, TargetsView, Input,
 		DistanceFunction, SmoothingFunction, SmoothingAngle,
@@ -403,7 +403,7 @@ FRigUnit_MathRBFInterpolateVectorFloat_Execute()
 	DECLARE_SCOPE_HIERARCHICAL_COUNTER_RIGUNIT()
 
 	const TArrayView<const FMathRBFInterpolateVectorFloat_Target> TargetsView(Targets.GetData(), Targets.Num());
-	TArray<float, TMemStackAllocator<>> Weights;
+	TArray<float> Weights;
 	GetInterpolatedWeights<FMathRBFInterpolateVectorFloat_Target>(
 		Context.State, WorkData, TargetsView, Input,
 		DistanceFunction, SmoothingFunction, SmoothingRadius, 
@@ -422,7 +422,7 @@ FRigUnit_MathRBFInterpolateVectorVector_Execute()
 	DECLARE_SCOPE_HIERARCHICAL_COUNTER_RIGUNIT()
 
 	const TArrayView<const FMathRBFInterpolateVectorVector_Target> TargetsView(Targets.GetData(), Targets.Num());
-	TArray<float, TMemStackAllocator<>> Weights;
+	TArray<float> Weights;
 	GetInterpolatedWeights<FMathRBFInterpolateVectorVector_Target>(
 		Context.State, WorkData, TargetsView, Input,
 		DistanceFunction, SmoothingFunction, SmoothingRadius, 
@@ -441,7 +441,7 @@ FRigUnit_MathRBFInterpolateVectorColor_Execute()
 	DECLARE_SCOPE_HIERARCHICAL_COUNTER_RIGUNIT()
 
 	const TArrayView<const FMathRBFInterpolateVectorColor_Target> TargetsView(Targets.GetData(), Targets.Num());
-	TArray<float, TMemStackAllocator<>> Weights;
+	TArray<float> Weights;
 	GetInterpolatedWeights<FMathRBFInterpolateVectorColor_Target>(
 		Context.State, WorkData, TargetsView, Input,
 		DistanceFunction, SmoothingFunction, SmoothingRadius, 
@@ -460,7 +460,7 @@ FRigUnit_MathRBFInterpolateVectorQuat_Execute()
 	DECLARE_SCOPE_HIERARCHICAL_COUNTER_RIGUNIT()
 
 	const TArrayView<const FMathRBFInterpolateVectorQuat_Target> TargetsView(Targets.GetData(), Targets.Num());
-	TArray<float, TMemStackAllocator<>> Weights;
+	TArray<float> Weights;
 	GetInterpolatedWeights<FMathRBFInterpolateVectorQuat_Target>(
 		Context.State, WorkData, TargetsView, Input,
 		DistanceFunction, SmoothingFunction, SmoothingRadius, 
@@ -480,7 +480,7 @@ FRigUnit_MathRBFInterpolateVectorXform_Execute()
 	DECLARE_SCOPE_HIERARCHICAL_COUNTER_RIGUNIT()
 
 	const TArrayView<const FMathRBFInterpolateVectorXform_Target> TargetsView(Targets.GetData(), Targets.Num());
-	TArray<float, TMemStackAllocator<>> Weights;
+	TArray<float> Weights;
 	GetInterpolatedWeights<FMathRBFInterpolateVectorXform_Target>(
 		Context.State, WorkData, TargetsView, Input,
 		DistanceFunction, SmoothingFunction, SmoothingRadius, 
diff --git a/Engine/Plugins/Animation/ControlRig/Source/ControlRig/Private/Units/Math/RigUnit_MathRBFInterpolate.h b/Engine/Plugins/Animation/ControlRig/Source/ControlRig/Private/Units/Math/RigUnit_MathRBFInterpolate.h
--- a/Engine/Plugins/Animation/ControlRig/Source/ControlRig/Private/Units/Math/RigUnit_MathRBFInterpolate.h
+++ b/Engine/Plugins/Animation/ControlRig/Source/ControlRig/Private/Units/Math/RigUnit_MathRBFInterpolate.h
@@ -119,7 +119,7 @@ protected:
 		float SmoothingAngle,
 		bool bNormalizeOutput,
 		FVector TwistAxis,
-		TArray<float, TMemStackAllocator<>>& Weights
+		TArray<float>& Weights
 	);
 
 	static uint64 HashTargets(const TArrayView<const FQuat>& Targets);
@@ -160,7 +160,7 @@ protected:
 		ERBFKernelType SmoothingFunction,
 		float SmoothingRadius,
 		bool bNormalizeOutput,
-		TArray<float, TMemStackAllocator<>>& Weights
+		TArray<float>& Weights
 	);
 
 	static uint64 HashTargets(const TArrayView<const FVector>& Targets);
