Ensure that linked animation instances are not lost when running a ControlRigLayerInstance for use with sequencer

diff --git a/Engine/Plugins/Animation/ControlRig/Source/ControlRig/Private/Sequencer/ControlRigLayerInstance.cpp b/Engine/Plugins/Animation/ControlRig/Source/ControlRig/Private/Sequencer/ControlRigLayerInstance.cpp
--- a/Engine/Plugins/Animation/ControlRig/Source/ControlRig/Private/Sequencer/ControlRigLayerInstance.cpp
+++ b/Engine/Plugins/Animation/ControlRig/Source/ControlRig/Private/Sequencer/ControlRigLayerInstance.cpp
@@ -63,21 +63,19 @@ void UControlRigLayerInstance::SetSourceAnimInstance(UAnimInstance* SourceAnimIn
 
 	if (SourceAnimInstance)
 	{
-		// we only have one instance fir this layer
-		FLinkedInstancesAdapter::ResetLinkedInstance(MeshComponent);
+		// Add the current animation instance as a linked instance
 		FLinkedInstancesAdapter::AddLinkedInstance(MeshComponent, SourceAnimInstance);
-	}
-	else
-	{
-		FLinkedInstancesAdapter::ResetLinkedInstance(MeshComponent);
-	}
 
-	if (SourceAnimInstance)
-	{
+		// Direct the control rig instance to the current animation instance to evaluate as its source (input pose)
 		GetProxyOnGameThread<FControlRigLayerInstanceProxy>().SetSourceAnimInstance(SourceAnimInstance, UAnimInstance::GetProxyOnGameThreadStatic<FAnimInstanceProxy>(SourceAnimInstance));
 	}
 	else
 	{
+		UAnimInstance* CurrentSourceAnimInstance = GetProxyOnGameThread<FControlRigLayerInstanceProxy>().GetSourceAnimInstance();		
+		// Remove the original instances from the linked instances as it should be reinstated as the main anim instance
+		FLinkedInstancesAdapter::RemoveLinkedInstance(MeshComponent, CurrentSourceAnimInstance);
+
+		// Null out the animation instance used to as the input source for the control rig instance
 		GetProxyOnGameThread<FControlRigLayerInstanceProxy>().SetSourceAnimInstance(nullptr, nullptr);
 	}
 }
