Sequencer: Control Rig: Select Controls from Pose with no Control Rig Edit mode.  Missed this button in last submission.
#trivial

diff --git a/Engine/Plugins/Animation/ControlRig/Source/ControlRigEditor/Private/EditMode/SControlRigControlViews.cpp b/Engine/Plugins/Animation/ControlRig/Source/ControlRigEditor/Private/EditMode/SControlRigControlViews.cpp
--- a/Engine/Plugins/Animation/ControlRig/Source/ControlRigEditor/Private/EditMode/SControlRigControlViews.cpp
+++ b/Engine/Plugins/Animation/ControlRig/Source/ControlRigEditor/Private/EditMode/SControlRigControlViews.cpp
@@ -19,6 +19,9 @@
 #include "EditorModeManager.h"
 #include "Widgets/Input/SButton.h"
 #include "LevelEditor.h"
+#include "ControlRigSequencerEditorLibrary.h"
+#include "LevelSequence.h"
+#include "LevelSequenceEditorBlueprintLibrary.h"
 
 #define LOCTEXT_NAMESPACE "ControlRigBaseListWidget"
 
@@ -529,9 +532,26 @@ FReply SControlRigPoseView::OnPastePose()
 
 FReply SControlRigPoseView::OnSelectControls()
 {
-	if (GetControlRig() && PoseAsset.IsValid())
+	if (PoseAsset.IsValid())
 	{
-		PoseAsset->SelectControls(GetControlRig(), bIsMirror);
+		UControlRig* ControlRig = GetControlRig();
+		if (ControlRig == nullptr)
+		{
+			ULevelSequence* LevelSequence = ULevelSequenceEditorBlueprintLibrary::GetCurrentLevelSequence();
+			if (LevelSequence)
+			{
+				TArray<FControlRigSequencerBindingProxy> Proxies = UControlRigSequencerEditorLibrary::GetControlRigs(LevelSequence);
+				if (Proxies.Num() > 0)
+				{
+					//MZ TODO when we have Mutliple Control Rig's active select more than one.
+					ControlRig = Proxies[0].ControlRig;
+				}
+			}
+		}
+		if (ControlRig)
+		{
+			PoseAsset->SelectControls(ControlRig, bIsMirror);
+		}
 	}
 	return FReply::Handled();
 }
