Control Rig: Performance improvement for UpdateComponentToWorld

diff --git a/Engine/Plugins/Animation/ControlRig/Source/ControlRig/Private/ControlRigComponent.cpp b/Engine/Plugins/Animation/ControlRig/Source/ControlRig/Private/ControlRigComponent.cpp
--- a/Engine/Plugins/Animation/ControlRig/Source/ControlRig/Private/ControlRigComponent.cpp
+++ b/Engine/Plugins/Animation/ControlRig/Source/ControlRig/Private/ControlRigComponent.cpp
@@ -1511,7 +1511,7 @@ void UControlRigComponent::HandleControlRigPreSetupEvent(UControlRig* InControlR
 
 		if (USkeletalMeshComponent* Component = Cast< USkeletalMeshComponent>(MappedElement.SceneComponent))
 		{
-			ComponentsToTick.Add(Component);
+			ComponentsToTick.AddUnique(Component);
 		}
 	}
 
