Control Rig: Fix crash when collapsing nodes several times

diff --git a/Engine/Plugins/Animation/ControlRig/Source/ControlRigEditor/Private/Editor/ControlRigEditor.cpp b/Engine/Plugins/Animation/ControlRig/Source/ControlRigEditor/Private/Editor/ControlRigEditor.cpp
--- a/Engine/Plugins/Animation/ControlRig/Source/ControlRigEditor/Private/Editor/ControlRigEditor.cpp
+++ b/Engine/Plugins/Animation/ControlRig/Source/ControlRigEditor/Private/Editor/ControlRigEditor.cpp
@@ -2522,11 +2522,13 @@ bool FControlRigEditor::IsSectionVisible(NodeSectionID::Type InSectionID) const
 		}
 		case NodeSectionID::LOCAL_VARIABLE:
 		{
-			const URigVMGraph* Graph = GetFocusedModel();
-			const URigVMGraph* ParentGraph = Graph->GetParentGraph();				
-			if (ParentGraph && ParentGraph->IsA<URigVMFunctionLibrary>())
+			if(const URigVMGraph* Graph = GetFocusedModel())
 			{
-				return true;
+				const URigVMGraph* ParentGraph = Graph->GetParentGraph();				
+				if (ParentGraph && ParentGraph->IsA<URigVMFunctionLibrary>())
+				{
+					return true;
+				}
 			}
 		}
 		default:
