Control Rig: Set up central editor close delegate to notify all widgets referencing the blueprint

diff --git a/Engine/Plugins/Animation/ControlRig/Source/ControlRigEditor/Private/Editor/ControlRigEditor.cpp b/Engine/Plugins/Animation/ControlRig/Source/ControlRigEditor/Private/Editor/ControlRigEditor.cpp
--- a/Engine/Plugins/Animation/ControlRig/Source/ControlRigEditor/Private/Editor/ControlRigEditor.cpp
+++ b/Engine/Plugins/Animation/ControlRig/Source/ControlRigEditor/Private/Editor/ControlRigEditor.cpp
@@ -218,6 +218,9 @@ FControlRigEditor::~FControlRigEditor()
 	ClearDetailObject();
 
 	UControlRigBlueprint* RigBlueprint = GetControlRigBlueprint();
+
+	ControlRigEditorClosedDelegate.Broadcast(this, RigBlueprint);
+
 	if (RigBlueprint)
 	{
 		// clear editor related data from the debugged control rig instance 
diff --git a/Engine/Plugins/Animation/ControlRig/Source/ControlRigEditor/Private/Editor/ControlRigEditor.h b/Engine/Plugins/Animation/ControlRig/Source/ControlRigEditor/Private/Editor/ControlRigEditor.h
--- a/Engine/Plugins/Animation/ControlRig/Source/ControlRigEditor/Private/Editor/ControlRigEditor.h
+++ b/Engine/Plugins/Animation/ControlRig/Source/ControlRigEditor/Private/Editor/ControlRigEditor.h
@@ -29,6 +29,9 @@ class SBorder;
 class USkeletalMesh;
 class FStructOnScope;
 class UToolMenu;
+class FControlRigEditor;
+
+DECLARE_MULTICAST_DELEGATE_TwoParams(FControlRigEditorClosed, const FControlRigEditor*, UControlRigBlueprint*);
 
 UENUM()
 enum class EControlRigEditorEventQueue : uint8
@@ -92,6 +95,12 @@ public:
 
 	int32 GetRigHierarchyTabCount() const { return RigHierarchyTabCount; }
 
+	FControlRigEditorClosed& OnControlRigEditorClosed() { return ControlRigEditorClosedDelegate; }
+
+private:
+
+	FControlRigEditorClosed ControlRigEditorClosedDelegate;
+
 public:
 	// IToolkit Interface
 	virtual FName GetToolkitFName() const override;
diff --git a/Engine/Plugins/Animation/ControlRig/Source/ControlRigEditor/Private/Editor/SRigCurveContainer.cpp b/Engine/Plugins/Animation/ControlRig/Source/ControlRigEditor/Private/Editor/SRigCurveContainer.cpp
--- a/Engine/Plugins/Animation/ControlRig/Source/ControlRigEditor/Private/Editor/SRigCurveContainer.cpp
+++ b/Engine/Plugins/Animation/ControlRig/Source/ControlRigEditor/Private/Editor/SRigCurveContainer.cpp
@@ -201,6 +201,11 @@ void SRigCurveContainer::Construct(const FArguments& InArgs, TSharedRef<FControl
 	];
 
 	CreateRigCurveList();
+
+	if(ControlRigEditor.IsValid())
+	{
+		ControlRigEditor.Pin()->OnControlRigEditorClosed().AddSP(this, &SRigCurveContainer::OnControlRigEditorClose);
+	}
 }
 
 SRigCurveContainer::~SRigCurveContainer()
@@ -208,15 +213,7 @@ SRigCurveContainer::~SRigCurveContainer()
 	// Make sure we don't get called on editor exit if we've already been shut down.
 	GEditor->OnEditorClose().RemoveAll(this);
 
-	if (ControlRigEditor.IsValid())
-	{
-		ControlRigBlueprint = ControlRigEditor.Pin()->GetControlRigBlueprint();
-		if (ControlRigBlueprint.IsValid())
-		{
-			ControlRigBlueprint->Hierarchy->OnModified().RemoveAll(this);
-			ControlRigBlueprint->OnRefreshEditor().RemoveAll(this);
-		}
-	}
+	OnEditorClose();
 }
 
 FReply SRigCurveContainer::OnKeyDown(const FGeometry& MyGeometry, const FKeyEvent& InKeyEvent)
@@ -316,10 +313,25 @@ TSharedPtr<SWidget> SRigCurveContainer::OnGetContextMenuContent() const
 
 void SRigCurveContainer::OnEditorClose()
 {
+	if (ControlRigEditor.IsValid())
+	{
+		ControlRigBlueprint = ControlRigEditor.Pin()->GetControlRigBlueprint();
+		if (ControlRigBlueprint.IsValid())
+		{
+			ControlRigBlueprint->Hierarchy->OnModified().RemoveAll(this);
+			ControlRigBlueprint->OnRefreshEditor().RemoveAll(this);
+		}
+	}
+
 	ControlRigBlueprint.Reset();
 	ControlRigEditor.Reset();
 }
 
+void SRigCurveContainer::OnControlRigEditorClose(const FControlRigEditor* InEditor, UControlRigBlueprint* InBlueprint)
+{
+	OnEditorClose();
+}
+
 void SRigCurveContainer::OnRenameClicked()
 {
 	TArray< FDisplayedRigCurveInfoPtr > SelectedItems = RigCurveListView->GetSelectedItems();
@@ -377,6 +389,12 @@ void SRigCurveContainer::CreateNewNameEntry(const FText& CommittedText, ETextCom
 
 void SRigCurveContainer::CreateRigCurveList( const FString& SearchText )
 {
+	if(bIsChangingRigHierarchy)
+	{
+		return;
+	}
+	TGuardValue<bool> GuardReentry(bIsChangingRigHierarchy, true);
+
 	if(!ControlRigBlueprint.IsValid())
 	{
 		return;
diff --git a/Engine/Plugins/Animation/ControlRig/Source/ControlRigEditor/Private/Editor/SRigCurveContainer.h b/Engine/Plugins/Animation/ControlRig/Source/ControlRigEditor/Private/Editor/SRigCurveContainer.h
--- a/Engine/Plugins/Animation/ControlRig/Source/ControlRigEditor/Private/Editor/SRigCurveContainer.h
+++ b/Engine/Plugins/Animation/ControlRig/Source/ControlRigEditor/Private/Editor/SRigCurveContainer.h
@@ -207,6 +207,7 @@ private:
 	TSharedPtr<SWidget> OnGetContextMenuContent() const;
 
 	void OnEditorClose();
+	void OnControlRigEditorClose(const FControlRigEditor* InEditor, UControlRigBlueprint* InBlueprint);
 
 	/**
 	* Clears and rebuilds the table, according to an optional search string
diff --git a/Engine/Plugins/Animation/ControlRig/Source/ControlRigEditor/Private/Editor/SRigHierarchy.cpp b/Engine/Plugins/Animation/ControlRig/Source/ControlRigEditor/Private/Editor/SRigHierarchy.cpp
--- a/Engine/Plugins/Animation/ControlRig/Source/ControlRigEditor/Private/Editor/SRigHierarchy.cpp
+++ b/Engine/Plugins/Animation/ControlRig/Source/ControlRigEditor/Private/Editor/SRigHierarchy.cpp
@@ -288,18 +288,7 @@ void SRigHierarchyItem::OnNameCommitted(const FText& InText, ETextCommit::Type I
 
 SRigHierarchy::~SRigHierarchy()
 {
-	if (ControlRigEditor.IsValid())
-	{
-		ControlRigEditor.Pin()->GetKeyDownDelegate().Unbind();
-		ControlRigEditor.Pin()->OnGetViewportContextMenu().Unbind();
-		ControlRigEditor.Pin()->OnViewportContextMenuCommands().Unbind();
-	}
-
-	if (ControlRigBlueprint.IsValid())
-	{
-		ControlRigBlueprint->Hierarchy->OnModified().RemoveAll(this);
-		ControlRigBlueprint->OnRefreshEditor().RemoveAll(this);
-	}
+	OnEditorClose();
 }
 
 void SRigHierarchy::Construct(const FArguments& InArgs, TSharedRef<FControlRigEditor> InControlRigEditor)
@@ -465,9 +454,34 @@ void SRigHierarchy::Construct(const FArguments& InArgs, TSharedRef<FControlRigEd
 		});
 		ControlRigEditor.Pin()->OnGetViewportContextMenu().BindSP(this, &SRigHierarchy::GetOrCreateContextMenu);
 		ControlRigEditor.Pin()->OnViewportContextMenuCommands().BindSP(this, &SRigHierarchy::GetContextMenuCommands);
+		ControlRigEditor.Pin()->OnControlRigEditorClosed().AddSP(this, &SRigHierarchy::OnControlRigEditorClose);
 	}
 }
 
+void SRigHierarchy::OnEditorClose()
+{
+	if (ControlRigEditor.IsValid())
+	{
+		ControlRigEditor.Pin()->GetKeyDownDelegate().Unbind();
+		ControlRigEditor.Pin()->OnGetViewportContextMenu().Unbind();
+		ControlRigEditor.Pin()->OnViewportContextMenuCommands().Unbind();
+	}
+
+	if (ControlRigBlueprint.IsValid())
+	{
+		ControlRigBlueprint->Hierarchy->OnModified().RemoveAll(this);
+		ControlRigBlueprint->OnRefreshEditor().RemoveAll(this);
+	}
+	
+	ControlRigEditor.Reset();
+	ControlRigBlueprint.Reset();
+}
+
+void SRigHierarchy::OnControlRigEditorClose(const FControlRigEditor* InEditor, UControlRigBlueprint* InBlueprint)
+{
+	OnEditorClose();
+}
+
 void SRigHierarchy::BindCommands()
 {
 	// create new command
diff --git a/Engine/Plugins/Animation/ControlRig/Source/ControlRigEditor/Private/Editor/SRigHierarchy.h b/Engine/Plugins/Animation/ControlRig/Source/ControlRigEditor/Private/Editor/SRigHierarchy.h
--- a/Engine/Plugins/Animation/ControlRig/Source/ControlRigEditor/Private/Editor/SRigHierarchy.h
+++ b/Engine/Plugins/Animation/ControlRig/Source/ControlRigEditor/Private/Editor/SRigHierarchy.h
@@ -178,6 +178,10 @@ public:
 	}
 
 private:
+
+	void OnEditorClose();
+	void OnControlRigEditorClose(const FControlRigEditor* InEditor, UControlRigBlueprint* InBlueprint);
+
 	/** Bind commands that this widget handles */
 	void BindCommands();
 
