Control Rig: Fix editor freeze when instruction double clicked in execution stack

diff --git a/Engine/Plugins/Animation/ControlRig/Source/ControlRigEditor/Private/ControlRigGraphDetails.cpp b/Engine/Plugins/Animation/ControlRig/Source/ControlRigEditor/Private/ControlRigGraphDetails.cpp
--- a/Engine/Plugins/Animation/ControlRig/Source/ControlRigEditor/Private/ControlRigGraphDetails.cpp
+++ b/Engine/Plugins/Animation/ControlRig/Source/ControlRigEditor/Private/ControlRigGraphDetails.cpp
@@ -1178,7 +1178,7 @@ void FControlRigWrappedNodeDetails::CustomizeDetails(IDetailLayoutBuilder& Detai
 		{
 			do
 			{
-				TArray<URigVMPin*> SourcePins = Pin->GetLinkedSourcePins(false);
+				TArray<URigVMPin*> SourcePins = SourcePin->GetLinkedSourcePins(false);
 				if(SourcePins.Num() > 0)
 				{
 					SourcePin = SourcePins[0];
