Sequencer: Fix reconstruct channel proxy so that it's always invoked on serialize. There was a case where it was disabled for undo and selection, but that's incorrect. It needs to be invoked on undo so that the channel proxy is up to date (in case undo removes a channel)

diff --git a/Engine/Plugins/Animation/ControlRig/Source/ControlRig/Private/Sequencer/MovieSceneControlRigParameterSection.cpp b/Engine/Plugins/Animation/ControlRig/Source/ControlRig/Private/Sequencer/MovieSceneControlRigParameterSection.cpp
--- a/Engine/Plugins/Animation/ControlRig/Source/ControlRig/Private/Sequencer/MovieSceneControlRigParameterSection.cpp
+++ b/Engine/Plugins/Animation/ControlRig/Source/ControlRig/Private/Sequencer/MovieSceneControlRigParameterSection.cpp
@@ -787,7 +787,7 @@ void UMovieSceneControlRigParameterSection::PostEditImport()
 	{
 		SetControlRig(Track->GetControlRig());
 	}
-	ReconstructChannelProxy(true);
+	ReconstructChannelProxy();
 }
 
 void UMovieSceneControlRigParameterSection::PostLoad()
@@ -908,7 +908,7 @@ void UMovieSceneControlRigParameterSection::AddScalarParameter(FName InParameter
 		}
 		if (bReconstructChannel)
 		{
-			ReconstructChannelProxy(true);
+			ReconstructChannelProxy();
 		}
 	}
 }
@@ -931,7 +931,7 @@ void UMovieSceneControlRigParameterSection::AddBoolParameter(FName InParameterNa
 		}
 		if (bReconstructChannel)
 		{
-			ReconstructChannelProxy(true);
+			ReconstructChannelProxy();
 		}
 	}
 }
@@ -953,7 +953,7 @@ void UMovieSceneControlRigParameterSection::AddEnumParameter(FName InParameterNa
 		ExistingChannel->SetEnum(Enum);
 		if (bReconstructChannel)
 		{
-			ReconstructChannelProxy(true);
+			ReconstructChannelProxy();
 		}
 	}
 }
@@ -975,7 +975,7 @@ void UMovieSceneControlRigParameterSection::AddIntegerParameter(FName InParamete
 		}
 		if (bReconstructChannel)
 		{
-			ReconstructChannelProxy(true);
+			ReconstructChannelProxy();
 		}
 	}
 }
@@ -1000,7 +1000,7 @@ void UMovieSceneControlRigParameterSection::AddVector2DParameter(FName InParamet
 		}
 		if (bReconstructChannel)
 		{
-			ReconstructChannelProxy(true);
+			ReconstructChannelProxy();
 		}
 	}
 }
@@ -1028,7 +1028,7 @@ void UMovieSceneControlRigParameterSection::AddVectorParameter(FName InParameter
 		}
 		if (bReconstructChannel)
 		{
-			ReconstructChannelProxy(true);
+			ReconstructChannelProxy();
 		}
 	}
 }
@@ -1057,7 +1057,7 @@ void UMovieSceneControlRigParameterSection::AddColorParameter(FName InParameterN
 		}
 		if (bReconstructChannel)
 		{
-			ReconstructChannelProxy(true);
+			ReconstructChannelProxy();
 		}
 	}
 }
@@ -1105,144 +1105,199 @@ void UMovieSceneControlRigParameterSection::AddTransformParameter(FName InParame
 		}
 		if(bReconstructChannel)
 		{
-			ReconstructChannelProxy(true);
+			ReconstructChannelProxy();
 		}
 	}
 }
 
 
-void UMovieSceneControlRigParameterSection::ReconstructChannelProxy(bool bForce)
+void UMovieSceneControlRigParameterSection::ReconstructChannelProxy()
 {
 	FMovieSceneChannelProxyData Channels;
-	if(bForce || ControlsMask!= OldControlsMask)
+	ControlChannelMap.Empty();
+	// Need to create the channels in sorted orders
+	if (ControlRig)
 	{
-		ControlChannelMap.Empty();
-		OldControlsMask = ControlsMask;
-		// Need to create the channels in sorted orders
-		if (ControlRig)
+		TArray<FRigControlElement*> SortedControls;
+		ControlRig->GetControlsInOrder(SortedControls);
+		if (ControlsMask.Num() != SortedControls.Num())
 		{
-			TArray<FRigControlElement*> SortedControls;
-			ControlRig->GetControlsInOrder(SortedControls);
-			if (ControlsMask.Num() != SortedControls.Num())
-			{
-				TArray<bool> OnArray;
-				OnArray.Init(true, ControlRig->AvailableControls().Num());
-				SetControlsMask(OnArray);
-				OldControlsMask = ControlsMask;
-			}
+			TArray<bool> OnArray;
+			OnArray.Init(true, ControlRig->AvailableControls().Num());
+			SetControlsMask(OnArray);
+		}
 			
-			int32 ControlIndex = 0; 
-			int32 TotalIndex = 0; 
-			int32 FloatChannelIndex = 0;
-			int32 BoolChannelIndex = 0;
-			int32 EnumChannelIndex = 0;
-			int32 IntegerChannelIndex = 0;
-			const FName BoolChannelTypeName = FMovieSceneBoolChannel::StaticStruct()->GetFName();
-			const FName EnumChannelTypeName = FMovieSceneByteChannel::StaticStruct()->GetFName();
-			const FName IntegerChannelTypeName = FMovieSceneIntegerChannel::StaticStruct()->GetFName();
+		int32 ControlIndex = 0; 
+		int32 TotalIndex = 0; 
+		int32 FloatChannelIndex = 0;
+		int32 BoolChannelIndex = 0;
+		int32 EnumChannelIndex = 0;
+		int32 IntegerChannelIndex = 0;
+		const FName BoolChannelTypeName = FMovieSceneBoolChannel::StaticStruct()->GetFName();
+		const FName EnumChannelTypeName = FMovieSceneByteChannel::StaticStruct()->GetFName();
+		const FName IntegerChannelTypeName = FMovieSceneIntegerChannel::StaticStruct()->GetFName();
 			
-			for (FRigControlElement* ControlElement : SortedControls)
+		for (FRigControlElement* ControlElement : SortedControls)
+		{
+			if (!ControlElement->Settings.bAnimatable)
 			{
-				if (!ControlElement->Settings.bAnimatable)
-				{
-					continue;
-				}
+				continue;
+			}
 
-				FRigControlElement* ParentControlElement = Cast<FRigControlElement>(ControlRig->GetHierarchy()->GetFirstParent(ControlElement));
-				FName ParentControlName = NAME_None;
-				if(ParentControlElement)
-				{
-					ParentControlName = ParentControlElement->GetName();
-				}
+			FRigControlElement* ParentControlElement = Cast<FRigControlElement>(ControlRig->GetHierarchy()->GetFirstParent(ControlElement));
+			FName ParentControlName = NAME_None;
+			if(ParentControlElement)
+			{
+				ParentControlName = ParentControlElement->GetName();
+			}
 #if WITH_EDITOR
 
-				switch (ControlElement->Settings.ControlType)
-				{
-				case ERigControlType::Float:
+			switch (ControlElement->Settings.ControlType)
+			{
+			case ERigControlType::Float:
+			{
+				for (FScalarParameterNameAndCurve& Scalar : GetScalarParameterNamesAndCurves())
 				{
-					for (FScalarParameterNameAndCurve& Scalar : GetScalarParameterNamesAndCurves())
+					if (ControlElement->GetName() == Scalar.ParameterName)
 					{
-						if (ControlElement->GetName() == Scalar.ParameterName)
-						{
-							bool bEnabled = ControlsMask[ControlIndex];
+						bool bEnabled = ControlsMask[ControlIndex];
 
-							FText Group;
-							if(ParentControlElement)
+						FText Group;
+						if(ParentControlElement)
+						{
+							switch(ParentControlElement->Settings.ControlType)
 							{
-								switch(ParentControlElement->Settings.ControlType)
+								case ERigControlType::Position:
+								case ERigControlType::Scale:
+								case ERigControlType::Rotator:
+								case ERigControlType::Transform:
+								case ERigControlType::EulerTransform:
+								case ERigControlType::TransformNoScale:
 								{
-									case ERigControlType::Position:
-									case ERigControlType::Scale:
-									case ERigControlType::Rotator:
-									case ERigControlType::Transform:
-									case ERigControlType::EulerTransform:
-									case ERigControlType::TransformNoScale:
-									{
-										Group = FText::FromName(ParentControlElement->GetDisplayName());
-										break;
-									}
+									Group = FText::FromName(ParentControlElement->GetDisplayName());
+									break;
 								}
 							}
+						}
 
-							if (Group.IsEmpty())
+						if (Group.IsEmpty())
+						{
+							ControlChannelMap.Add(Scalar.ParameterName, FChannelMapInfo(ControlIndex, TotalIndex, FloatChannelIndex));
+							Group = FText::FromName(ControlElement->GetDisplayName());
+							ControlIndex++;  //up the index only if no parent
+						}
+						else
+						{
+							FChannelMapInfo* pChannelIndex = ControlChannelMap.Find(ParentControlName);
+							if (pChannelIndex)
 							{
-								ControlChannelMap.Add(Scalar.ParameterName, FChannelMapInfo(ControlIndex, TotalIndex, FloatChannelIndex));
-								Group = FText::FromName(ControlElement->GetDisplayName());
-								ControlIndex++;  //up the index only if no parent
+								ControlChannelMap.Add(Scalar.ParameterName, FChannelMapInfo(ControlIndex, TotalIndex, FloatChannelIndex,pChannelIndex->ControlIndex));
 							}
 							else
 							{
-								FChannelMapInfo* pChannelIndex = ControlChannelMap.Find(ParentControlName);
-								if (pChannelIndex)
-								{
-									ControlChannelMap.Add(Scalar.ParameterName, FChannelMapInfo(ControlIndex, TotalIndex, FloatChannelIndex,pChannelIndex->ControlIndex));
-								}
-								else
-								{
-									ControlChannelMap.Add(Scalar.ParameterName, FChannelMapInfo(ControlIndex, TotalIndex, FloatChannelIndex));
+								ControlChannelMap.Add(Scalar.ParameterName, FChannelMapInfo(ControlIndex, TotalIndex, FloatChannelIndex));
 
+							}
+						}
+
+						FParameterFloatChannelEditorData EditorData(ControlRig, Scalar.ParameterName, bEnabled, Group, TotalIndex);
+						EditorData.MetaData.DisplayText = FText::FromName(ControlElement->GetDisplayName());
+						Channels.Add(Scalar.ParameterCurve, EditorData.MetaData, EditorData.ExternalValues);
+						FloatChannelIndex += 1;
+						TotalIndex += 1;
+						break;
+					}
+				}
+				break;
+			}
+			case ERigControlType::Bool:
+			{
+				for (FBoolParameterNameAndCurve& Bool : GetBoolParameterNamesAndCurves())
+				{
+					if (ControlElement->GetName() == Bool.ParameterName)
+					{
+						bool bEnabled = ControlsMask[ControlIndex];
+
+						FText Group;
+						if(ParentControlElement)
+						{
+							switch(ParentControlElement->Settings.ControlType)
+							{
+								case ERigControlType::Position:
+								case ERigControlType::Scale:
+								case ERigControlType::Rotator:
+								case ERigControlType::Transform:
+								case ERigControlType::EulerTransform:
+								case ERigControlType::TransformNoScale:
+								{
+									Group = FText::FromName(ParentControlElement->GetDisplayName());
+									break;
 								}
 							}
+						}
 
-							FParameterFloatChannelEditorData EditorData(ControlRig, Scalar.ParameterName, bEnabled, Group, TotalIndex);
-							EditorData.MetaData.DisplayText = FText::FromName(ControlElement->GetDisplayName());
-							Channels.Add(Scalar.ParameterCurve, EditorData.MetaData, EditorData.ExternalValues);
-							FloatChannelIndex += 1;
-							TotalIndex += 1;
-							break;
+						if (Group.IsEmpty())
+						{
+							ControlChannelMap.Add(Bool.ParameterName, FChannelMapInfo(ControlIndex, TotalIndex, BoolChannelIndex, INDEX_NONE,BoolChannelTypeName));
+							Group = FText::FromName(ControlElement->GetDisplayName());
+							ControlIndex++;  //up the index only if no parent
+						}
+						else
+						{
+							FChannelMapInfo* pChannelIndex = ControlChannelMap.Find(ParentControlName);
+							if (pChannelIndex)
+							{
+								ControlChannelMap.Add(Bool.ParameterName, FChannelMapInfo(ControlIndex, TotalIndex, BoolChannelIndex,pChannelIndex->ControlIndex,BoolChannelTypeName));
+							}
+							else
+							{
+								ControlChannelMap.Add(Bool.ParameterName, FChannelMapInfo(ControlIndex, TotalIndex, BoolChannelIndex,INDEX_NONE, BoolChannelTypeName));
+							}
 						}
+
+						FMovieSceneChannelMetaData MetaData(Bool.ParameterName, Group, Group, bEnabled);
+						MetaData.DisplayText = FText::FromName(ControlElement->GetDisplayName());
+						MetaData.SortOrder = TotalIndex++;
+						BoolChannelIndex += 1;
+						// Prevent single channels from collapsing to the track node
+						MetaData.bCanCollapseToTrack = false;
+						Channels.Add(Bool.ParameterCurve, MetaData, TMovieSceneExternalValue<bool>());
+						break;
 					}
-					break;
 				}
-				case ERigControlType::Bool:
+				break;
+			}
+			case ERigControlType::Integer:
+			{
+				if (ControlElement->Settings.ControlEnum)
 				{
-					for (FBoolParameterNameAndCurve& Bool : GetBoolParameterNamesAndCurves())
+					for (FEnumParameterNameAndCurve& Enum : GetEnumParameterNamesAndCurves())
 					{
-						if (ControlElement->GetName() == Bool.ParameterName)
+						if (ControlElement->GetName() == Enum.ParameterName)
 						{
 							bool bEnabled = ControlsMask[ControlIndex];
 
 							FText Group;
-							if(ParentControlElement)
+							if (ParentControlElement)
 							{
-								switch(ParentControlElement->Settings.ControlType)
+								switch (ParentControlElement->Settings.ControlType)
+								{
+								case ERigControlType::Position:
+								case ERigControlType::Scale:
+								case ERigControlType::Rotator:
+								case ERigControlType::Transform:
+								case ERigControlType::EulerTransform:
+								case ERigControlType::TransformNoScale:
 								{
-									case ERigControlType::Position:
-									case ERigControlType::Scale:
-									case ERigControlType::Rotator:
-									case ERigControlType::Transform:
-									case ERigControlType::EulerTransform:
-									case ERigControlType::TransformNoScale:
-									{
-										Group = FText::FromName(ParentControlElement->GetDisplayName());
-										break;
-									}
+									Group = FText::FromName(ParentControlElement->GetDisplayName());
+									break;
+								}
 								}
 							}
 
 							if (Group.IsEmpty())
 							{
-								ControlChannelMap.Add(Bool.ParameterName, FChannelMapInfo(ControlIndex, TotalIndex, BoolChannelIndex, INDEX_NONE,BoolChannelTypeName));
+								ControlChannelMap.Add(Enum.ParameterName, FChannelMapInfo(ControlIndex, TotalIndex, EnumChannelIndex,INDEX_NONE, EnumChannelTypeName));
 								Group = FText::FromName(ControlElement->GetDisplayName());
 								ControlIndex++;  //up the index only if no parent
 							}
@@ -1251,431 +1306,371 @@ void UMovieSceneControlRigParameterSection::ReconstructChannelProxy(bool bForce)
 								FChannelMapInfo* pChannelIndex = ControlChannelMap.Find(ParentControlName);
 								if (pChannelIndex)
 								{
-									ControlChannelMap.Add(Bool.ParameterName, FChannelMapInfo(ControlIndex, TotalIndex, BoolChannelIndex,pChannelIndex->ControlIndex,BoolChannelTypeName));
+									ControlChannelMap.Add(Enum.ParameterName, FChannelMapInfo(ControlIndex, TotalIndex, EnumChannelIndex, pChannelIndex->ControlIndex, EnumChannelTypeName));
 								}
 								else
 								{
-									ControlChannelMap.Add(Bool.ParameterName, FChannelMapInfo(ControlIndex, TotalIndex, BoolChannelIndex,INDEX_NONE, BoolChannelTypeName));
+									ControlChannelMap.Add(Enum.ParameterName, FChannelMapInfo(ControlIndex, TotalIndex, EnumChannelIndex, INDEX_NONE, EnumChannelTypeName));
 								}
 							}
 
-							FMovieSceneChannelMetaData MetaData(Bool.ParameterName, Group, Group, bEnabled);
+							FMovieSceneChannelMetaData MetaData(Enum.ParameterName, Group, Group, bEnabled);
 							MetaData.DisplayText = FText::FromName(ControlElement->GetDisplayName());
+							EnumChannelIndex += 1;
 							MetaData.SortOrder = TotalIndex++;
-							BoolChannelIndex += 1;
 							// Prevent single channels from collapsing to the track node
 							MetaData.bCanCollapseToTrack = false;
-							Channels.Add(Bool.ParameterCurve, MetaData, TMovieSceneExternalValue<bool>());
+							Channels.Add(Enum.ParameterCurve, MetaData, TMovieSceneExternalValue<uint8>());
 							break;
 						}
 					}
-					break;
 				}
-				case ERigControlType::Integer:
+				else
 				{
-					if (ControlElement->Settings.ControlEnum)
+					for (FIntegerParameterNameAndCurve& Integer : GetIntegerParameterNamesAndCurves())
 					{
-						for (FEnumParameterNameAndCurve& Enum : GetEnumParameterNamesAndCurves())
+						if (ControlElement->GetName() == Integer.ParameterName)
 						{
-							if (ControlElement->GetName() == Enum.ParameterName)
-							{
-								bool bEnabled = ControlsMask[ControlIndex];
+							bool bEnabled = ControlsMask[ControlIndex];
 
-								FText Group;
-								if (ParentControlElement)
+							FText Group;
+							if (ParentControlElement)
+							{
+								switch (ParentControlElement->Settings.ControlType)
 								{
-									switch (ParentControlElement->Settings.ControlType)
-									{
-									case ERigControlType::Position:
-									case ERigControlType::Scale:
-									case ERigControlType::Rotator:
-									case ERigControlType::Transform:
-									case ERigControlType::EulerTransform:
-									case ERigControlType::TransformNoScale:
-									{
-										Group = FText::FromName(ParentControlElement->GetDisplayName());
-										break;
-									}
-									}
-								}
-
-								if (Group.IsEmpty())
+								case ERigControlType::Position:
+								case ERigControlType::Scale:
+								case ERigControlType::Rotator:
+								case ERigControlType::Transform:
+								case ERigControlType::EulerTransform:
+								case ERigControlType::TransformNoScale:
 								{
-									ControlChannelMap.Add(Enum.ParameterName, FChannelMapInfo(ControlIndex, TotalIndex, EnumChannelIndex,INDEX_NONE, EnumChannelTypeName));
-									Group = FText::FromName(ControlElement->GetDisplayName());
-									ControlIndex++;  //up the index only if no parent
+									Group = FText::FromName(ParentControlElement->GetDisplayName());
+									break;
 								}
-								else
-								{
-									FChannelMapInfo* pChannelIndex = ControlChannelMap.Find(ParentControlName);
-									if (pChannelIndex)
-									{
-										ControlChannelMap.Add(Enum.ParameterName, FChannelMapInfo(ControlIndex, TotalIndex, EnumChannelIndex, pChannelIndex->ControlIndex, EnumChannelTypeName));
-									}
-									else
-									{
-										ControlChannelMap.Add(Enum.ParameterName, FChannelMapInfo(ControlIndex, TotalIndex, EnumChannelIndex, INDEX_NONE, EnumChannelTypeName));
-									}
 								}
-
-								FMovieSceneChannelMetaData MetaData(Enum.ParameterName, Group, Group, bEnabled);
-								MetaData.DisplayText = FText::FromName(ControlElement->GetDisplayName());
-								EnumChannelIndex += 1;
-								MetaData.SortOrder = TotalIndex++;
-								// Prevent single channels from collapsing to the track node
-								MetaData.bCanCollapseToTrack = false;
-								Channels.Add(Enum.ParameterCurve, MetaData, TMovieSceneExternalValue<uint8>());
-								break;
 							}
-						}
-					}
-					else
-					{
-						for (FIntegerParameterNameAndCurve& Integer : GetIntegerParameterNamesAndCurves())
-						{
-							if (ControlElement->GetName() == Integer.ParameterName)
+							if (Group.IsEmpty())
 							{
-								bool bEnabled = ControlsMask[ControlIndex];
-
-								FText Group;
-								if (ParentControlElement)
-								{
-									switch (ParentControlElement->Settings.ControlType)
-									{
-									case ERigControlType::Position:
-									case ERigControlType::Scale:
-									case ERigControlType::Rotator:
-									case ERigControlType::Transform:
-									case ERigControlType::EulerTransform:
-									case ERigControlType::TransformNoScale:
-									{
-										Group = FText::FromName(ParentControlElement->GetDisplayName());
-										break;
-									}
-									}
-								}
-								if (Group.IsEmpty())
+								ControlChannelMap.Add(Integer.ParameterName, FChannelMapInfo(ControlIndex, TotalIndex, IntegerChannelIndex,INDEX_NONE,IntegerChannelTypeName));
+								Group = FText::FromName(ControlElement->GetDisplayName());
+								ControlIndex++;  //up the index only if no parent
+							}
+							else
+							{
+								FChannelMapInfo* pChannelIndex = ControlChannelMap.Find(ParentControlName);
+								if (pChannelIndex)
 								{
-									ControlChannelMap.Add(Integer.ParameterName, FChannelMapInfo(ControlIndex, TotalIndex, IntegerChannelIndex,INDEX_NONE,IntegerChannelTypeName));
-									Group = FText::FromName(ControlElement->GetDisplayName());
-									ControlIndex++;  //up the index only if no parent
+									ControlChannelMap.Add(Integer.ParameterName, FChannelMapInfo(ControlIndex, TotalIndex, IntegerChannelIndex, pChannelIndex->ControlIndex,IntegerChannelTypeName));
 								}
 								else
 								{
-									FChannelMapInfo* pChannelIndex = ControlChannelMap.Find(ParentControlName);
-									if (pChannelIndex)
-									{
-										ControlChannelMap.Add(Integer.ParameterName, FChannelMapInfo(ControlIndex, TotalIndex, IntegerChannelIndex, pChannelIndex->ControlIndex,IntegerChannelTypeName));
-									}
-									else
-									{
-										ControlChannelMap.Add(Integer.ParameterName, FChannelMapInfo(ControlIndex, TotalIndex, IntegerChannelIndex, INDEX_NONE, IntegerChannelTypeName));
-									}
+									ControlChannelMap.Add(Integer.ParameterName, FChannelMapInfo(ControlIndex, TotalIndex, IntegerChannelIndex, INDEX_NONE, IntegerChannelTypeName));
 								}
-
-								FMovieSceneChannelMetaData MetaData(Integer.ParameterName, Group, Group, bEnabled);
-								MetaData.DisplayText = FText::FromName(ControlElement->GetDisplayName());
-								IntegerChannelIndex += 1;
-								MetaData.SortOrder = TotalIndex++;
-								// Prevent single channels from collapsing to the track node
-								MetaData.bCanCollapseToTrack = false;
-								Channels.Add(Integer.ParameterCurve, MetaData, TMovieSceneExternalValue<int32>());
-								break;
 							}
-						}
 
+							FMovieSceneChannelMetaData MetaData(Integer.ParameterName, Group, Group, bEnabled);
+							MetaData.DisplayText = FText::FromName(ControlElement->GetDisplayName());
+							IntegerChannelIndex += 1;
+							MetaData.SortOrder = TotalIndex++;
+							// Prevent single channels from collapsing to the track node
+							MetaData.bCanCollapseToTrack = false;
+							Channels.Add(Integer.ParameterCurve, MetaData, TMovieSceneExternalValue<int32>());
+							break;
+						}
 					}
-					break;
+
 				}
-				case ERigControlType::Vector2D:
+				break;
+			}
+			case ERigControlType::Vector2D:
+			{
+				for (FVector2DParameterNameAndCurves& Vector2D : GetVector2DParameterNamesAndCurves())
 				{
-					for (FVector2DParameterNameAndCurves& Vector2D : GetVector2DParameterNamesAndCurves())
+					if (ControlElement->GetName() == Vector2D.ParameterName)
 					{
-						if (ControlElement->GetName() == Vector2D.ParameterName)
-						{
-							ControlChannelMap.Add(Vector2D.ParameterName, FChannelMapInfo(ControlIndex, TotalIndex, FloatChannelIndex));
-							bool bEnabled = ControlsMask[ControlIndex++];
-							FText Group = FText::FromName(ControlElement->GetDisplayName());
-							FParameterVectorChannelEditorData EditorData(ControlRig, Vector2D.ParameterName, bEnabled, Group, TotalIndex, 2);
-							Channels.Add(Vector2D.XCurve, EditorData.MetaData[0], EditorData.ExternalValues[0]);
-							Channels.Add(Vector2D.YCurve, EditorData.MetaData[1], EditorData.ExternalValues[1]);
-							FloatChannelIndex += 2;
-							TotalIndex += 2;
-							break;
-						}
+						ControlChannelMap.Add(Vector2D.ParameterName, FChannelMapInfo(ControlIndex, TotalIndex, FloatChannelIndex));
+						bool bEnabled = ControlsMask[ControlIndex++];
+						FText Group = FText::FromName(ControlElement->GetDisplayName());
+						FParameterVectorChannelEditorData EditorData(ControlRig, Vector2D.ParameterName, bEnabled, Group, TotalIndex, 2);
+						Channels.Add(Vector2D.XCurve, EditorData.MetaData[0], EditorData.ExternalValues[0]);
+						Channels.Add(Vector2D.YCurve, EditorData.MetaData[1], EditorData.ExternalValues[1]);
+						FloatChannelIndex += 2;
+						TotalIndex += 2;
+						break;
 					}
-					break;
 				}
-				case ERigControlType::Position:
-				case ERigControlType::Scale:
-				case ERigControlType::Rotator:
+				break;
+			}
+			case ERigControlType::Position:
+			case ERigControlType::Scale:
+			case ERigControlType::Rotator:
+			{
+				for (FVectorParameterNameAndCurves& Vector : GetVectorParameterNamesAndCurves())
 				{
-					for (FVectorParameterNameAndCurves& Vector : GetVectorParameterNamesAndCurves())
+					if (ControlElement->GetName() == Vector.ParameterName)
 					{
-						if (ControlElement->GetName() == Vector.ParameterName)
+						if (ControlElement->Settings.ControlType == ERigControlType::Scale)
 						{
-							if (ControlElement->Settings.ControlType == ERigControlType::Scale)
+							if (BlendType == EMovieSceneBlendType::Additive)
 							{
-								if (BlendType == EMovieSceneBlendType::Additive)
-								{
-									Vector.XCurve.SetDefault(0.0f);
-									Vector.YCurve.SetDefault(0.0f);
-									Vector.ZCurve.SetDefault(0.0f);
-								}
-								else
-								{
-									Vector.XCurve.SetDefault(1.0f);
-									Vector.YCurve.SetDefault(1.0f);
-									Vector.ZCurve.SetDefault(1.0f);
-								}
+								Vector.XCurve.SetDefault(0.0f);
+								Vector.YCurve.SetDefault(0.0f);
+								Vector.ZCurve.SetDefault(0.0f);
+							}
+							else
+							{
+								Vector.XCurve.SetDefault(1.0f);
+								Vector.YCurve.SetDefault(1.0f);
+								Vector.ZCurve.SetDefault(1.0f);
 							}
-							ControlChannelMap.Add(Vector.ParameterName, FChannelMapInfo(ControlIndex, TotalIndex, FloatChannelIndex));
-							bool bEnabled = ControlsMask[ControlIndex++];
-							FText Group = FText::FromName(ControlElement->GetDisplayName());
-							FParameterVectorChannelEditorData EditorData(ControlRig, Vector.ParameterName, bEnabled, Group, TotalIndex, 3);
-							Channels.Add(Vector.XCurve, EditorData.MetaData[0], EditorData.ExternalValues[0]);
-							Channels.Add(Vector.YCurve, EditorData.MetaData[1], EditorData.ExternalValues[1]);
-							Channels.Add(Vector.ZCurve, EditorData.MetaData[2], EditorData.ExternalValues[2]);
-							FloatChannelIndex += 3;
-							TotalIndex += 3;
-							break;
 						}
+						ControlChannelMap.Add(Vector.ParameterName, FChannelMapInfo(ControlIndex, TotalIndex, FloatChannelIndex));
+						bool bEnabled = ControlsMask[ControlIndex++];
+						FText Group = FText::FromName(ControlElement->GetDisplayName());
+						FParameterVectorChannelEditorData EditorData(ControlRig, Vector.ParameterName, bEnabled, Group, TotalIndex, 3);
+						Channels.Add(Vector.XCurve, EditorData.MetaData[0], EditorData.ExternalValues[0]);
+						Channels.Add(Vector.YCurve, EditorData.MetaData[1], EditorData.ExternalValues[1]);
+						Channels.Add(Vector.ZCurve, EditorData.MetaData[2], EditorData.ExternalValues[2]);
+						FloatChannelIndex += 3;
+						TotalIndex += 3;
+						break;
 					}
-					break;
 				}
+				break;
+			}
 
-				case ERigControlType::TransformNoScale:
-				case ERigControlType::Transform:
-				case ERigControlType::EulerTransform:
+			case ERigControlType::TransformNoScale:
+			case ERigControlType::Transform:
+			case ERigControlType::EulerTransform:
+			{
+				for (FTransformParameterNameAndCurves& Transform : GetTransformParameterNamesAndCurves())
 				{
-					for (FTransformParameterNameAndCurves& Transform : GetTransformParameterNamesAndCurves())
+					if (ControlElement->GetName() == Transform.ParameterName)
 					{
-						if (ControlElement->GetName() == Transform.ParameterName)
-						{
-							ControlChannelMap.Add(Transform.ParameterName, FChannelMapInfo(ControlIndex, TotalIndex, FloatChannelIndex));
-							bool bEnabled = ControlsMask[ControlIndex++];
-							FText Group = FText::FromName(ControlElement->GetDisplayName());
+						ControlChannelMap.Add(Transform.ParameterName, FChannelMapInfo(ControlIndex, TotalIndex, FloatChannelIndex));
+						bool bEnabled = ControlsMask[ControlIndex++];
+						FText Group = FText::FromName(ControlElement->GetDisplayName());
 
-							FParameterTransformChannelEditorData EditorData(ControlRig, Transform.ParameterName, bEnabled, TransformMask.GetChannels(), Group, 
-								TotalIndex);
+						FParameterTransformChannelEditorData EditorData(ControlRig, Transform.ParameterName, bEnabled, TransformMask.GetChannels(), Group, 
+							TotalIndex);
 
-							Channels.Add(Transform.Translation[0], EditorData.MetaData[0], EditorData.ExternalValues[0]);
-							Channels.Add(Transform.Translation[1], EditorData.MetaData[1], EditorData.ExternalValues[1]);
-							Channels.Add(Transform.Translation[2], EditorData.MetaData[2], EditorData.ExternalValues[2]);
+						Channels.Add(Transform.Translation[0], EditorData.MetaData[0], EditorData.ExternalValues[0]);
+						Channels.Add(Transform.Translation[1], EditorData.MetaData[1], EditorData.ExternalValues[1]);
+						Channels.Add(Transform.Translation[2], EditorData.MetaData[2], EditorData.ExternalValues[2]);
 
-							Channels.Add(Transform.Rotation[0], EditorData.MetaData[3], EditorData.ExternalValues[3]);
-							Channels.Add(Transform.Rotation[1], EditorData.MetaData[4], EditorData.ExternalValues[4]);
-							Channels.Add(Transform.Rotation[2], EditorData.MetaData[5], EditorData.ExternalValues[5]);
+						Channels.Add(Transform.Rotation[0], EditorData.MetaData[3], EditorData.ExternalValues[3]);
+						Channels.Add(Transform.Rotation[1], EditorData.MetaData[4], EditorData.ExternalValues[4]);
+						Channels.Add(Transform.Rotation[2], EditorData.MetaData[5], EditorData.ExternalValues[5]);
 
-							if (ControlElement->Settings.ControlType == ERigControlType::Transform ||
-								ControlElement->Settings.ControlType == ERigControlType::EulerTransform)
+						if (ControlElement->Settings.ControlType == ERigControlType::Transform ||
+							ControlElement->Settings.ControlType == ERigControlType::EulerTransform)
+						{
+							if (BlendType == EMovieSceneBlendType::Additive)
 							{
-								if (BlendType == EMovieSceneBlendType::Additive)
-								{
-									Transform.Scale[0].SetDefault(0.0f);
-									Transform.Scale[1].SetDefault(0.0f);
-									Transform.Scale[2].SetDefault(0.0f);
-								}
-								else
-								{
-									Transform.Scale[0].SetDefault(1.0f);
-									Transform.Scale[1].SetDefault(1.0f);
-									Transform.Scale[2].SetDefault(1.0f);
-								}
-								Channels.Add(Transform.Scale[0], EditorData.MetaData[6], EditorData.ExternalValues[6]);
-								Channels.Add(Transform.Scale[1], EditorData.MetaData[7], EditorData.ExternalValues[7]);
-								Channels.Add(Transform.Scale[2], EditorData.MetaData[8], EditorData.ExternalValues[8]);
-								FloatChannelIndex += 9;
-								TotalIndex += 9;
-
+								Transform.Scale[0].SetDefault(0.0f);
+								Transform.Scale[1].SetDefault(0.0f);
+								Transform.Scale[2].SetDefault(0.0f);
 							}
 							else
 							{
-								FloatChannelIndex += 6;
-								TotalIndex += 6;
-
+								Transform.Scale[0].SetDefault(1.0f);
+								Transform.Scale[1].SetDefault(1.0f);
+								Transform.Scale[2].SetDefault(1.0f);
 							}
-							break;
+							Channels.Add(Transform.Scale[0], EditorData.MetaData[6], EditorData.ExternalValues[6]);
+							Channels.Add(Transform.Scale[1], EditorData.MetaData[7], EditorData.ExternalValues[7]);
+							Channels.Add(Transform.Scale[2], EditorData.MetaData[8], EditorData.ExternalValues[8]);
+							FloatChannelIndex += 9;
+							TotalIndex += 9;
+
 						}
+						else
+						{
+							FloatChannelIndex += 6;
+							TotalIndex += 6;
+
+						}
+						break;
 					}
 				}
-				default:
-					break;
-				}
+			}
+			default:
+				break;
+			}
 #else
-				switch (ControlElement->Settings.ControlType)
-				{
-				case ERigControlType::Float:
+			switch (ControlElement->Settings.ControlType)
+			{
+			case ERigControlType::Float:
+			{
+				for (FScalarParameterNameAndCurve& Scalar : GetScalarParameterNamesAndCurves())
 				{
-					for (FScalarParameterNameAndCurve& Scalar : GetScalarParameterNamesAndCurves())
+					if (ControlElement->GetName() == Scalar.ParameterName)
 					{
-						if (ControlElement->GetName() == Scalar.ParameterName)
-						{
-							ControlChannelMap.Add(Scalar.ParameterName, FChannelMapInfo(ControlIndex, TotalIndex,FloatChannelIndex));
-							Channels.Add(Scalar.ParameterCurve);
-							FloatChannelIndex += 1;
-							TotalIndex += 1;
-							ControlIndex++;
-							break;
-						}
+						ControlChannelMap.Add(Scalar.ParameterName, FChannelMapInfo(ControlIndex, TotalIndex,FloatChannelIndex));
+						Channels.Add(Scalar.ParameterCurve);
+						FloatChannelIndex += 1;
+						TotalIndex += 1;
+						ControlIndex++;
+						break;
 					}
-					break;
 				}
-				case ERigControlType::Bool:
+				break;
+			}
+			case ERigControlType::Bool:
+			{
+				for (FBoolParameterNameAndCurve& Bool : GetBoolParameterNamesAndCurves())
 				{
-					for (FBoolParameterNameAndCurve& Bool : GetBoolParameterNamesAndCurves())
+					if (ControlElement->GetName() == Bool.ParameterName)
 					{
-						if (ControlElement->GetName() == Bool.ParameterName)
-						{
-							ControlChannelMap.Add(Bool.ParameterName, FChannelMapInfo(ControlIndex, TotalIndex, BoolChannelIndex));
-							Channels.Add(Bool.ParameterCurve);
-							BoolChannelIndex += 1;
-							TotalIndex += 1;
-							ControlIndex++;
-							break;
-						}
+						ControlChannelMap.Add(Bool.ParameterName, FChannelMapInfo(ControlIndex, TotalIndex, BoolChannelIndex));
+						Channels.Add(Bool.ParameterCurve);
+						BoolChannelIndex += 1;
+						TotalIndex += 1;
+						ControlIndex++;
+						break;
 					}
-					break;
 				}
-				case ERigControlType::Integer:
+				break;
+			}
+			case ERigControlType::Integer:
+			{
+				if (ControlElement->Settings.ControlEnum)
 				{
-					if (ControlElement->Settings.ControlEnum)
-					{
-						for (FEnumParameterNameAndCurve& Enum : GetEnumParameterNamesAndCurves())
-						{
-							if (ControlElement->GetName() == Enum.ParameterName)
-							{
-								ControlChannelMap.Add(Enum.ParameterName, FChannelMapInfo(ControlIndex, TotalIndex, EnumChannelIndex));
-								Channels.Add(Enum.ParameterCurve);
-								EnumChannelIndex += 1;
-								TotalIndex += 1;
-								ControlIndex++;
-								break;
-							}
-						}
-					}
-					else
+					for (FEnumParameterNameAndCurve& Enum : GetEnumParameterNamesAndCurves())
 					{
-						for (FIntegerParameterNameAndCurve& Integer : GetIntegerParameterNamesAndCurves())
+						if (ControlElement->GetName() == Enum.ParameterName)
 						{
-							if (ControlElement->GetName() == Integer.ParameterName)
-							{
-								ControlChannelMap.Add(Integer.ParameterName, FChannelMapInfo(ControlIndex, TotalIndex, IntegerChannelIndex));
-								Channels.Add(Integer.ParameterCurve);
-								IntegerChannelIndex += 1;
-								TotalIndex += 1;
-								ControlIndex++;
-								break;
-							}
+							ControlChannelMap.Add(Enum.ParameterName, FChannelMapInfo(ControlIndex, TotalIndex, EnumChannelIndex));
+							Channels.Add(Enum.ParameterCurve);
+							EnumChannelIndex += 1;
+							TotalIndex += 1;
+							ControlIndex++;
+							break;
 						}
 					}
-					break;
 				}
-				case ERigControlType::Vector2D:
+				else
 				{
-					for (FVector2DParameterNameAndCurves& Vector2D : GetVector2DParameterNamesAndCurves())
+					for (FIntegerParameterNameAndCurve& Integer : GetIntegerParameterNamesAndCurves())
 					{
-						if (ControlElement->GetName() == Vector2D.ParameterName)
+						if (ControlElement->GetName() == Integer.ParameterName)
 						{
-							ControlChannelMap.Add(Vector2D.ParameterName, FChannelMapInfo(ControlIndex, TotalIndex, FloatChannelIndex));
-							Channels.Add(Vector2D.XCurve);
-							Channels.Add(Vector2D.YCurve);
-							FloatChannelIndex += 2;
-							TotalIndex += 2;
+							ControlChannelMap.Add(Integer.ParameterName, FChannelMapInfo(ControlIndex, TotalIndex, IntegerChannelIndex));
+							Channels.Add(Integer.ParameterCurve);
+							IntegerChannelIndex += 1;
+							TotalIndex += 1;
 							ControlIndex++;
 							break;
 						}
 					}
-					break;
 				}
-				case ERigControlType::Position:
-				case ERigControlType::Scale:
-				case ERigControlType::Rotator:
+				break;
+			}
+			case ERigControlType::Vector2D:
+			{
+				for (FVector2DParameterNameAndCurves& Vector2D : GetVector2DParameterNamesAndCurves())
 				{
-					for (FVectorParameterNameAndCurves& Vector : GetVectorParameterNamesAndCurves())
+					if (ControlElement->GetName() == Vector2D.ParameterName)
 					{
-						if (ControlElement->GetName() == Vector.ParameterName)
-						{
-							ControlChannelMap.Add(Vector.ParameterName, FChannelMapInfo(ControlIndex, TotalIndex, FloatChannelIndex));
-							Channels.Add(Vector.XCurve);
-							Channels.Add(Vector.YCurve);
-							Channels.Add(Vector.ZCurve);
-							FloatChannelIndex += 3;
-							TotalIndex += 3;
-							ControlIndex++;
-							break;
-						}
+						ControlChannelMap.Add(Vector2D.ParameterName, FChannelMapInfo(ControlIndex, TotalIndex, FloatChannelIndex));
+						Channels.Add(Vector2D.XCurve);
+						Channels.Add(Vector2D.YCurve);
+						FloatChannelIndex += 2;
+						TotalIndex += 2;
+						ControlIndex++;
+						break;
 					}
-					break;
 				}
-				/*
-				for (FColorParameterNameAndCurves& Color : GetColorParameterNamesAndCurves())
+				break;
+			}
+			case ERigControlType::Position:
+			case ERigControlType::Scale:
+			case ERigControlType::Rotator:
+			{
+				for (FVectorParameterNameAndCurves& Vector : GetVectorParameterNamesAndCurves())
 				{
-					Channels.Add(Color.RedCurve);
-					Channels.Add(Color.GreenCurve);
-					Channels.Add(Color.BlueCurve);
-					Channels.Add(Color.AlphaCurve);
-					break
+					if (ControlElement->GetName() == Vector.ParameterName)
+					{
+						ControlChannelMap.Add(Vector.ParameterName, FChannelMapInfo(ControlIndex, TotalIndex, FloatChannelIndex));
+						Channels.Add(Vector.XCurve);
+						Channels.Add(Vector.YCurve);
+						Channels.Add(Vector.ZCurve);
+						FloatChannelIndex += 3;
+						TotalIndex += 3;
+						ControlIndex++;
+						break;
+					}
 				}
-				*/
-				case ERigControlType::TransformNoScale:
-				case ERigControlType::Transform:
-				case ERigControlType::EulerTransform:
+				break;
+			}
+			/*
+			for (FColorParameterNameAndCurves& Color : GetColorParameterNamesAndCurves())
+			{
+				Channels.Add(Color.RedCurve);
+				Channels.Add(Color.GreenCurve);
+				Channels.Add(Color.BlueCurve);
+				Channels.Add(Color.AlphaCurve);
+				break
+			}
+			*/
+			case ERigControlType::TransformNoScale:
+			case ERigControlType::Transform:
+			case ERigControlType::EulerTransform:
+			{
+				for (FTransformParameterNameAndCurves& Transform : GetTransformParameterNamesAndCurves())
 				{
-					for (FTransformParameterNameAndCurves& Transform : GetTransformParameterNamesAndCurves())
+					if (ControlElement->GetName() == Transform.ParameterName)
 					{
-						if (ControlElement->GetName() == Transform.ParameterName)
-						{
-							ControlChannelMap.Add(Transform.ParameterName, FChannelMapInfo(ControlIndex, TotalIndex,FloatChannelIndex));
-							Channels.Add(Transform.Translation[0]);
-							Channels.Add(Transform.Translation[1]);
-							Channels.Add(Transform.Translation[2]);
+						ControlChannelMap.Add(Transform.ParameterName, FChannelMapInfo(ControlIndex, TotalIndex,FloatChannelIndex));
+						Channels.Add(Transform.Translation[0]);
+						Channels.Add(Transform.Translation[1]);
+						Channels.Add(Transform.Translation[2]);
 
-							Channels.Add(Transform.Rotation[0]);
-							Channels.Add(Transform.Rotation[1]);
-							Channels.Add(Transform.Rotation[2]);
+						Channels.Add(Transform.Rotation[0]);
+						Channels.Add(Transform.Rotation[1]);
+						Channels.Add(Transform.Rotation[2]);
 
-							if (ControlElement->Settings.ControlType == ERigControlType::Transform || ControlElement->Settings.ControlType == ERigControlType::EulerTransform)
-							{
-								Channels.Add(Transform.Scale[0]);
-								Channels.Add(Transform.Scale[1]);
-								Channels.Add(Transform.Scale[2]);
-								FloatChannelIndex += 9;
-								TotalIndex += 9;
-							}
-							else
-							{
-								FloatChannelIndex += 6;
-								TotalIndex += 6;
-							}
-
-							ControlIndex++;
-							break;
+						if (ControlElement->Settings.ControlType == ERigControlType::Transform || ControlElement->Settings.ControlType == ERigControlType::EulerTransform)
+						{
+							Channels.Add(Transform.Scale[0]);
+							Channels.Add(Transform.Scale[1]);
+							Channels.Add(Transform.Scale[2]);
+							FloatChannelIndex += 9;
+							TotalIndex += 9;
 						}
+						else
+						{
+							FloatChannelIndex += 6;
+							TotalIndex += 6;
+						}
+
+						ControlIndex++;
+						break;
 					}
-					break;
 				}
-				}
-#endif
+				break;
+			}
 			}
+#endif
+		}
 		
 #if WITH_EDITOR
-			FMovieSceneChannelMetaData      MetaData;
-			MetaData.SetIdentifiers("Weight", NSLOCTEXT("MovieSceneTransformSection", "Weight", "Weight"));
-			MetaData.bEnabled = EnumHasAllFlags(TransformMask.GetChannels(), EMovieSceneTransformChannel::Weight);
-			MetaData.SortOrder = TotalIndex++;
-			MetaData.bCanCollapseToTrack = false;
-			TMovieSceneExternalValue<float> ExVal;
-			Channels.Add(Weight, MetaData, ExVal);
+		FMovieSceneChannelMetaData      MetaData;
+		MetaData.SetIdentifiers("Weight", NSLOCTEXT("MovieSceneTransformSection", "Weight", "Weight"));
+		MetaData.bEnabled = EnumHasAllFlags(TransformMask.GetChannels(), EMovieSceneTransformChannel::Weight);
+		MetaData.SortOrder = TotalIndex++;
+		MetaData.bCanCollapseToTrack = false;
+		TMovieSceneExternalValue<float> ExVal;
+		Channels.Add(Weight, MetaData, ExVal);
 #else
-			Channels.Add(Weight);
+		Channels.Add(Weight);
 
 #endif
-		}
+	}
 
 
-		ChannelProxy = MakeShared<FMovieSceneChannelProxy>(MoveTemp(Channels));
-	}
+	ChannelProxy = MakeShared<FMovieSceneChannelProxy>(MoveTemp(Channels));
 }
 
 FMovieSceneInterrogationKey UMovieSceneControlRigParameterSection::GetFloatInterrogationKey()
@@ -1866,7 +1861,7 @@ void UMovieSceneControlRigParameterSection::RecreateWithThisControlRig(UControlR
 			break;
 		}
 	}
-	ReconstructChannelProxy(true);
+	ReconstructChannelProxy();
 }
 
 void UMovieSceneControlRigParameterSection::SetControlRig(UControlRig* InControlRig)
@@ -2268,7 +2263,7 @@ void UMovieSceneControlRigParameterSection::AddEnumParameterKey(FName InParamete
 		const int32 NewIndex = EnumParameterNamesAndCurves.Add(FEnumParameterNameAndCurve(InParameterName));
 		ExistingChannel = &EnumParameterNamesAndCurves[NewIndex].ParameterCurve;
 
-		ReconstructChannelProxy(true);
+		ReconstructChannelProxy();
 	}
 
 	ExistingChannel->GetData().UpdateOrAddKey(InTime, InValue);
@@ -2296,7 +2291,7 @@ void UMovieSceneControlRigParameterSection::AddIntegerParameterKey(FName InParam
 		const int32 NewIndex = IntegerParameterNamesAndCurves.Add(FIntegerParameterNameAndCurve(InParameterName));
 		ExistingChannel = &IntegerParameterNamesAndCurves[NewIndex].ParameterCurve;
 
-		ReconstructChannelProxy(true);
+		ReconstructChannelProxy();
 	}
 
 	ExistingChannel->GetData().UpdateOrAddKey(InTime, InValue);
@@ -2314,7 +2309,7 @@ bool UMovieSceneControlRigParameterSection::RemoveEnumParameter(FName InParamete
 		if (EnumParameterNamesAndCurves[i].ParameterName == InParameterName)
 		{
 			EnumParameterNamesAndCurves.RemoveAt(i);
-			ReconstructChannelProxy(true);
+			ReconstructChannelProxy();
 			return true;
 		}
 	}
@@ -2328,7 +2323,7 @@ bool UMovieSceneControlRigParameterSection::RemoveIntegerParameter(FName InParam
 		if (IntegerParameterNamesAndCurves[i].ParameterName == InParameterName)
 		{
 			IntegerParameterNamesAndCurves.RemoveAt(i);
-			ReconstructChannelProxy(true);
+			ReconstructChannelProxy();
 			return true;
 		}
 	}
diff --git a/Engine/Plugins/Animation/ControlRig/Source/ControlRig/Private/Sequencer/MovieSceneControlRigParameterTrack.cpp b/Engine/Plugins/Animation/ControlRig/Source/ControlRig/Private/Sequencer/MovieSceneControlRigParameterTrack.cpp
--- a/Engine/Plugins/Animation/ControlRig/Source/ControlRig/Private/Sequencer/MovieSceneControlRigParameterTrack.cpp
+++ b/Engine/Plugins/Animation/ControlRig/Source/ControlRig/Private/Sequencer/MovieSceneControlRigParameterTrack.cpp
@@ -75,7 +75,7 @@ void UMovieSceneControlRigParameterTrack::AddSection(UMovieSceneSection& Section
 		{
 			CRSection->SetControlRig(ControlRig);
 		}
-		CRSection->ReconstructChannelProxy(true);
+		CRSection->ReconstructChannelProxy();
 	}
 
 	if (Sections.Num() > 1)
diff --git a/Engine/Plugins/Animation/ControlRig/Source/ControlRig/Public/Sequencer/MovieSceneControlRigParameterSection.h b/Engine/Plugins/Animation/ControlRig/Source/ControlRig/Public/Sequencer/MovieSceneControlRigParameterSection.h
--- a/Engine/Plugins/Animation/ControlRig/Source/ControlRig/Public/Sequencer/MovieSceneControlRigParameterSection.h
+++ b/Engine/Plugins/Animation/ControlRig/Source/ControlRig/Public/Sequencer/MovieSceneControlRigParameterSection.h
@@ -195,11 +195,6 @@ protected:
 	/*Integer Curves*/
 	UPROPERTY()
 	TArray<FIntegerParameterNameAndCurve> IntegerParameterNamesAndCurves;
-private:
-	/** Copy of Mask for controls, checked when reconstructing*/
-	TArray<bool> OldControlsMask;
-
-
 
 public:
 
@@ -232,7 +227,7 @@ public:
 	void SetControlsMask(const TArray<bool>& InMask)
 	{
 		ControlsMask = InMask;
-		ReconstructChannelProxy(true);
+		ReconstructChannelProxy();
 	}
 
 	void SetControlsMask(int32 Index, bool Val)
@@ -241,14 +236,13 @@ public:
 		{
 			ControlsMask[Index] = Val;
 		}
-		ReconstructChannelProxy(true);
+		ReconstructChannelProxy();
 	}
 
 	void FillControlsMask(bool Val)
 	{
 		ControlsMask.Init(Val, ControlsMask.Num());
-		ReconstructChannelProxy(true);
-
+		ReconstructChannelProxy();
 	}
 
 	/**
@@ -265,7 +259,7 @@ public:
 	void SetTransformMask(FMovieSceneTransformMask NewMask)
 	{
 		TransformMask = NewMask;
-		ReconstructChannelProxy(true);
+		ReconstructChannelProxy();
 	}
 
 public:
@@ -343,7 +337,7 @@ public:
 	 static FMovieSceneInterrogationKey GetVectorInterrogationKey();
 	 static FMovieSceneInterrogationKey GetTransformInterrogationKey();
 
-	virtual void ReconstructChannelProxy(bool bForce ) override;
+	virtual void ReconstructChannelProxy() override;
 
 protected:
 
diff --git a/Engine/Plugins/Experimental/MotionTrailEditorMode/Source/MotionTrailEditorMode/Private/Sequencer/SequencerTrailHierarchy.cpp b/Engine/Plugins/Experimental/MotionTrailEditorMode/Source/MotionTrailEditorMode/Private/Sequencer/SequencerTrailHierarchy.cpp
--- a/Engine/Plugins/Experimental/MotionTrailEditorMode/Source/MotionTrailEditorMode/Private/Sequencer/SequencerTrailHierarchy.cpp
+++ b/Engine/Plugins/Experimental/MotionTrailEditorMode/Source/MotionTrailEditorMode/Private/Sequencer/SequencerTrailHierarchy.cpp
@@ -690,7 +690,7 @@ void FSequencerTrailHierarchy::AddControlsToHierarchy(class USkeletalMeshCompone
 		RegisterControlRigDelegates(CompToAdd, CRParamSection);
 	}
 
-	CRParamSection->ReconstructChannelProxy(true);
+	CRParamSection->ReconstructChannelProxy();
 
 	TArray<FRigControlElement*> SortedControls;
 	CRParamSection->GetControlRig()->GetControlsInOrder(SortedControls);
diff --git a/Engine/Source/Runtime/MovieSceneTracks/Private/Sections/MovieSceneParameterSection.cpp b/Engine/Source/Runtime/MovieSceneTracks/Private/Sections/MovieSceneParameterSection.cpp
--- a/Engine/Source/Runtime/MovieSceneTracks/Private/Sections/MovieSceneParameterSection.cpp
+++ b/Engine/Source/Runtime/MovieSceneTracks/Private/Sections/MovieSceneParameterSection.cpp
@@ -52,8 +52,7 @@ void UMovieSceneParameterSection::Serialize(FArchive& Ar)
 
 	if (Ar.IsLoading())
 	{
-		//Don't force if transacting, since it may not be a channel creation/deletion change
-		ReconstructChannelProxy(!Ar.IsTransacting());
+		ReconstructChannelProxy();
 	}
 }
 
@@ -61,13 +60,12 @@ void UMovieSceneParameterSection::PostEditImport()
 {
 	Super::PostEditImport();
 
-	ReconstructChannelProxy(true);
+	ReconstructChannelProxy();
 }
 
 
-void UMovieSceneParameterSection::ReconstructChannelProxy(bool bForce)
+void UMovieSceneParameterSection::ReconstructChannelProxy()
 {
-
 	FMovieSceneChannelProxyData Channels;
 
 #if WITH_EDITOR
@@ -224,7 +222,6 @@ void UMovieSceneParameterSection::ReconstructChannelProxy(bool bForce)
 #endif
 
 	ChannelProxy = MakeShared<FMovieSceneChannelProxy>(MoveTemp(Channels));
-	
 }
 
 void UMovieSceneParameterSection::AddScalarParameterKey( FName InParameterName, FFrameNumber InTime, float InValue )
@@ -243,7 +240,7 @@ void UMovieSceneParameterSection::AddScalarParameterKey( FName InParameterName,
 		const int32 NewIndex = ScalarParameterNamesAndCurves.Add( FScalarParameterNameAndCurve( InParameterName ) );
 		ExistingChannel = &ScalarParameterNamesAndCurves[NewIndex].ParameterCurve;
 
-		ReconstructChannelProxy(true);
+		ReconstructChannelProxy();
 	}
 
 	ExistingChannel->AddCubicKey(InTime, InValue);
@@ -270,7 +267,7 @@ void UMovieSceneParameterSection::AddBoolParameterKey(FName InParameterName, FFr
 		const int32 NewIndex = BoolParameterNamesAndCurves.Add(FBoolParameterNameAndCurve(InParameterName));
 		ExistingChannel = &BoolParameterNamesAndCurves[NewIndex].ParameterCurve;
 
-		ReconstructChannelProxy(true);
+		ReconstructChannelProxy();
 	}
 
 	ExistingChannel->GetData().UpdateOrAddKey(InTime, InValue);
@@ -298,7 +295,7 @@ void UMovieSceneParameterSection::AddVector2DParameterKey(FName InParameterName,
 		int32 NewIndex = Vector2DParameterNamesAndCurves.Add(FVector2DParameterNameAndCurves(InParameterName));
 		ExistingCurves = &Vector2DParameterNamesAndCurves[NewIndex];
 
-		ReconstructChannelProxy(true);
+		ReconstructChannelProxy();
 	}
 
 	ExistingCurves->XCurve.AddCubicKey(InTime, InValue.X);
@@ -326,7 +323,7 @@ void UMovieSceneParameterSection::AddVectorParameterKey( FName InParameterName,
 		int32 NewIndex = VectorParameterNamesAndCurves.Add( FVectorParameterNameAndCurves( InParameterName ) );
 		ExistingCurves = &VectorParameterNamesAndCurves[NewIndex];
 
-		ReconstructChannelProxy(true);
+		ReconstructChannelProxy();
 	}
 
 	ExistingCurves->XCurve.AddCubicKey(InTime, InValue.X);
@@ -355,7 +352,7 @@ void UMovieSceneParameterSection::AddColorParameterKey( FName InParameterName, F
 		int32 NewIndex = ColorParameterNamesAndCurves.Add( FColorParameterNameAndCurves( InParameterName ) );
 		ExistingCurves = &ColorParameterNamesAndCurves[NewIndex];
 
-		ReconstructChannelProxy(true);
+		ReconstructChannelProxy();
 	}
 
 	ExistingCurves->RedCurve.AddCubicKey(   InTime, InValue.R );
@@ -385,7 +382,7 @@ void UMovieSceneParameterSection::AddTransformParameterKey(FName InParameterName
 		int32 NewIndex = TransformParameterNamesAndCurves.Add(FTransformParameterNameAndCurves(InParameterName));
 		ExistingCurves = &TransformParameterNamesAndCurves[NewIndex];
 
-		ReconstructChannelProxy(true);
+		ReconstructChannelProxy();
 	}
 	FVector Translation = InValue.GetTranslation();
 	FRotator Rotator = InValue.GetRotation().Rotator();
@@ -415,7 +412,7 @@ bool UMovieSceneParameterSection::RemoveScalarParameter( FName InParameterName )
 		if ( ScalarParameterNamesAndCurves[i].ParameterName == InParameterName )
 		{
 			ScalarParameterNamesAndCurves.RemoveAt(i);
-			ReconstructChannelProxy(true);
+			ReconstructChannelProxy();
 			return true;
 		}
 	}
@@ -429,7 +426,7 @@ bool UMovieSceneParameterSection::RemoveBoolParameter(FName InParameterName)
 		if (BoolParameterNamesAndCurves[i].ParameterName == InParameterName)
 		{
 			BoolParameterNamesAndCurves.RemoveAt(i);
-			ReconstructChannelProxy(true);
+			ReconstructChannelProxy();
 			return true;
 		}
 	}
@@ -443,7 +440,7 @@ bool UMovieSceneParameterSection::RemoveVector2DParameter(FName InParameterName)
 		if (Vector2DParameterNamesAndCurves[i].ParameterName == InParameterName)
 		{
 			Vector2DParameterNamesAndCurves.RemoveAt(i);
-			ReconstructChannelProxy(true);
+			ReconstructChannelProxy();
 			return true;
 		}
 	}
@@ -457,7 +454,7 @@ bool UMovieSceneParameterSection::RemoveVectorParameter( FName InParameterName )
 		if ( VectorParameterNamesAndCurves[i].ParameterName == InParameterName )
 		{
 			VectorParameterNamesAndCurves.RemoveAt( i );
-			ReconstructChannelProxy(true);
+			ReconstructChannelProxy();
 			return true;
 		}
 	}
@@ -471,7 +468,7 @@ bool UMovieSceneParameterSection::RemoveColorParameter( FName InParameterName )
 		if ( ColorParameterNamesAndCurves[i].ParameterName == InParameterName )
 		{
 			ColorParameterNamesAndCurves.RemoveAt( i );
-			ReconstructChannelProxy(true);
+			ReconstructChannelProxy();
 			return true;
 		}
 	}
@@ -485,7 +482,7 @@ bool UMovieSceneParameterSection::RemoveTransformParameter(FName InParameterName
 		if (TransformParameterNamesAndCurves[i].ParameterName == InParameterName)
 		{
 			TransformParameterNamesAndCurves.RemoveAt(i);
-			ReconstructChannelProxy(true);
+			ReconstructChannelProxy();
 			return true;
 		}
 	}
diff --git a/Engine/Source/Runtime/MovieSceneTracks/Public/Sections/MovieSceneParameterSection.h b/Engine/Source/Runtime/MovieSceneTracks/Public/Sections/MovieSceneParameterSection.h
--- a/Engine/Source/Runtime/MovieSceneTracks/Public/Sections/MovieSceneParameterSection.h
+++ b/Engine/Source/Runtime/MovieSceneTracks/Public/Sections/MovieSceneParameterSection.h
@@ -423,7 +423,7 @@ protected:
 	//~ UMovieSceneSection interface
 	virtual void Serialize(FArchive& Ar) override;
 	virtual void PostEditImport() override;
-	virtual void ReconstructChannelProxy(bool bForce);
+	virtual void ReconstructChannelProxy();
 
 protected:
 	/** The bool parameter names and their associated curves. */
