Control Rig: if we are in the game thread, there is no need to defer copy the VM, if the copy is deferred, and the VM's mem class is marked pending kill, it will crash when details view attemps to access it for live inspection.

diff --git a/Engine/Plugins/Animation/ControlRig/Source/ControlRig/Private/ControlRig.cpp b/Engine/Plugins/Animation/ControlRig/Source/ControlRig/Private/ControlRig.cpp
--- a/Engine/Plugins/Animation/ControlRig/Source/ControlRig/Private/ControlRig.cpp
+++ b/Engine/Plugins/Animation/ControlRig/Source/ControlRig/Private/ControlRig.cpp
@@ -483,7 +483,9 @@ void UControlRig::InstantiateVMFromCDO()
 		if (VM && CDO && CDO->VM)
 		{
 			// reference the literal memory + byte code
-			VM->CopyFrom(CDO->VM, true, true);
+			// only defer if called from worker thread,
+			// which should be unlikely
+			VM->CopyFrom(CDO->VM, !IsInGameThread(), true);
 		}
 		else
 		{
diff --git a/Engine/Plugins/Animation/ControlRig/Source/ControlRigEditor/Private/ControlRigGraphDetails.cpp b/Engine/Plugins/Animation/ControlRig/Source/ControlRigEditor/Private/ControlRigGraphDetails.cpp
--- a/Engine/Plugins/Animation/ControlRig/Source/ControlRigEditor/Private/ControlRigGraphDetails.cpp
+++ b/Engine/Plugins/Animation/ControlRig/Source/ControlRigEditor/Private/ControlRigGraphDetails.cpp
@@ -1281,6 +1281,14 @@ void FControlRigWrappedNodeDetails::CustomizeDetails(IDetailLayoutBuilder& Detai
 						continue;
 					}
 
+					// the UClass must be alive for the details view to access it
+					// this ensure can fail if VM memory is not updated immediately after compile
+					// because of deferred copy
+					if(!ensure(!Memory->GetClass()->IsPendingKill()))
+					{
+						continue;
+					}
+
 					if(Memory->GetClass()->GetOuter() == GetTransientPackage())
 					{
 						continue;
