Control Rig: Pose Cache to allow access to transform array

diff --git a/Engine/Plugins/Animation/ControlRig/Source/ControlRig/Private/Units/Execution/RigUnit_Hierarchy.cpp b/Engine/Plugins/Animation/ControlRig/Source/ControlRig/Private/Units/Execution/RigUnit_Hierarchy.cpp
--- a/Engine/Plugins/Animation/ControlRig/Source/ControlRig/Private/Units/Execution/RigUnit_Hierarchy.cpp
+++ b/Engine/Plugins/Animation/ControlRig/Source/ControlRig/Private/Units/Execution/RigUnit_Hierarchy.cpp
@@ -435,6 +435,32 @@ FRigUnit_PoseGetTransform_Execute()
 	}
 }
 
+FRigUnit_PoseGetTransformArray_Execute()
+{
+	// set up defaults
+	Valid = false;
+	Transforms.Reset();
+
+	Valid = true;
+
+	Transforms.SetNum(Pose.Num());
+
+	if(Space == EBoneGetterSetterMode::GlobalSpace)
+	{
+		for(int32 Index=0; Index<Pose.Num(); Index++)
+		{
+			Transforms[Index] = Pose[Index].GlobalTransform;
+		}
+	}
+	else
+	{
+		for(int32 Index=0; Index<Pose.Num(); Index++)
+		{
+			Transforms[Index] = Pose[Index].LocalTransform;
+		}
+	}
+}
+
 FRigUnit_PoseGetCurve_Execute()
 {
 	if (Context.State == EControlRigState::Init)
diff --git a/Engine/Plugins/Animation/ControlRig/Source/ControlRig/Private/Units/Execution/RigUnit_Hierarchy.h b/Engine/Plugins/Animation/ControlRig/Source/ControlRig/Private/Units/Execution/RigUnit_Hierarchy.h
--- a/Engine/Plugins/Animation/ControlRig/Source/ControlRig/Private/Units/Execution/RigUnit_Hierarchy.h
+++ b/Engine/Plugins/Animation/ControlRig/Source/ControlRig/Private/Units/Execution/RigUnit_Hierarchy.h
@@ -403,6 +403,41 @@ struct CONTROLRIG_API FRigUnit_PoseGetTransform : public FRigUnit_HierarchyBase
 	int32 CachedPoseHash;
 };
 
+/**
+* Returns the hierarchy's pose transform
+*/
+USTRUCT(meta=(DisplayName="Get Pose Cache Transform Array", Keywords="Hierarchy,Pose,State", Varying))
+struct CONTROLRIG_API FRigUnit_PoseGetTransformArray : public FRigUnit_HierarchyBase
+{
+	GENERATED_BODY()
+
+	FRigUnit_PoseGetTransformArray()
+	{
+		Pose = FRigPose();
+		Space = EBoneGetterSetterMode::GlobalSpace;
+		Valid = false;
+		Transforms.Reset();
+	}
+
+	RIGVM_METHOD()
+	virtual void Execute(const FRigUnitContext& Context) override;
+
+	UPROPERTY(meta = (Input))
+	FRigPose Pose;
+
+	/**
+	* Defines if the transform should be retrieved in local or global space
+	*/ 
+	UPROPERTY(meta = (Input))
+	EBoneGetterSetterMode Space;
+
+	UPROPERTY(meta = (Output))
+	bool Valid;
+	
+	UPROPERTY(meta = (Output))
+	TArray<FTransform> Transforms;
+};
+
 /**
 * Returns the hierarchy's pose curve value
 */
