Control Rig: Fix random nodes to work with arrays / loops

diff --git a/Engine/Plugins/Animation/ControlRig/Source/ControlRig/Private/Units/Math/RigUnit_Random.cpp b/Engine/Plugins/Animation/ControlRig/Source/ControlRig/Private/Units/Math/RigUnit_Random.cpp
--- a/Engine/Plugins/Animation/ControlRig/Source/ControlRig/Private/Units/Math/RigUnit_Random.cpp
+++ b/Engine/Plugins/Animation/ControlRig/Source/ControlRig/Private/Units/Math/RigUnit_Random.cpp
@@ -20,11 +20,16 @@ FRigUnit_RandomFloat_Execute()
     DECLARE_SCOPE_HIERARCHICAL_COUNTER_RIGUNIT()
 	if (Context.State == EControlRigState::Init)
 	{
-		LastSeed = Seed;
+		LastSeed = BaseSeed = Seed;
 		TimeLeft = 0.f;
 		return;
 	}
 
+	if(Seed != BaseSeed)
+	{
+		LastSeed = BaseSeed = Seed;
+	}
+
 	TimeLeft = TimeLeft - Context.DeltaTime;
 	if (TimeLeft > 0.f)
 	{
@@ -43,11 +48,16 @@ FRigUnit_RandomVector_Execute()
     DECLARE_SCOPE_HIERARCHICAL_COUNTER_RIGUNIT()
 	if (Context.State == EControlRigState::Init)
 	{
-		LastSeed = Seed;
+		LastSeed = BaseSeed = Seed;
 		TimeLeft = 0.f;
 		return;
 	}
 
+	if(Seed != BaseSeed)
+	{
+		LastSeed = BaseSeed = Seed;
+	}
+
 	TimeLeft = TimeLeft - Context.DeltaTime;
 	if (TimeLeft > 0.f)
 	{
diff --git a/Engine/Plugins/Animation/ControlRig/Source/ControlRig/Private/Units/Math/RigUnit_Random.h b/Engine/Plugins/Animation/ControlRig/Source/ControlRig/Private/Units/Math/RigUnit_Random.h
--- a/Engine/Plugins/Animation/ControlRig/Source/ControlRig/Private/Units/Math/RigUnit_Random.h
+++ b/Engine/Plugins/Animation/ControlRig/Source/ControlRig/Private/Units/Math/RigUnit_Random.h
@@ -15,7 +15,7 @@ struct CONTROLRIG_API FRigUnit_RandomFloat : public FRigUnit_MathBase
 
 	FRigUnit_RandomFloat()
 	{
-		Seed = LastSeed = 217;
+		Seed = BaseSeed = LastSeed = 217;
 		Minimum = Result = LastResult = 0.f;
 		Maximum = 1.f;
 		Duration = 0.f;
@@ -47,6 +47,9 @@ struct CONTROLRIG_API FRigUnit_RandomFloat : public FRigUnit_MathBase
 	UPROPERTY()
 	int32 LastSeed;
 
+	UPROPERTY()
+	int32 BaseSeed;
+
 	UPROPERTY()
 	float TimeLeft;
 };
@@ -61,7 +64,7 @@ struct CONTROLRIG_API FRigUnit_RandomVector: public FRigUnit_MathBase
 
 	FRigUnit_RandomVector()
 	{
-		Seed = LastSeed = 217;
+		Seed = BaseSeed = LastSeed = 217;
 		Minimum = 0.f;
 		Maximum = 1.f;
 		Duration = 0.f;
@@ -94,6 +97,9 @@ struct CONTROLRIG_API FRigUnit_RandomVector: public FRigUnit_MathBase
 	UPROPERTY()
 	int32 LastSeed;
 
+	UPROPERTY()
+	int32 BaseSeed;
+
 	UPROPERTY()
 	float TimeLeft;
 };
\ No newline at end of file
