Control Rig: Add transform array math nodes + drawing

preflight: https://horde.devtools.epicgames.com/job/611d39425e73720001d79b0e

diff --git a/Engine/Plugins/Animation/ControlRig/Source/ControlRig/Private/Units/Debug/RigUnit_DebugTransform.cpp b/Engine/Plugins/Animation/ControlRig/Source/ControlRig/Private/Units/Debug/RigUnit_DebugTransform.cpp
--- a/Engine/Plugins/Animation/ControlRig/Source/ControlRig/Private/Units/Debug/RigUnit_DebugTransform.cpp
+++ b/Engine/Plugins/Animation/ControlRig/Source/ControlRig/Private/Units/Debug/RigUnit_DebugTransform.cpp
@@ -151,4 +151,51 @@ FRigUnit_DebugTransformArrayMutable_Execute()
 			}
 		}
 	}
-}
\ No newline at end of file
+}
+
+FRigUnit_DebugTransformArrayMutableItemSpace_Execute()
+{
+	DECLARE_SCOPE_HIERARCHICAL_COUNTER_RIGUNIT()
+	if (Context.State == EControlRigState::Init)
+	{
+		return;
+	}
+
+	if (Context.DrawInterface == nullptr || !bEnabled || Transforms.Num() == 0)
+	{
+		return;
+	}
+
+	for(const FTransform& Transform : Transforms)
+	{
+		FRigUnit_DebugTransformMutableItemSpace::StaticExecute(
+			RigVMExecuteContext, 
+			Transform,
+			Mode,
+			Color,
+			Thickness,
+			Scale,
+			Space, 
+			WorldOffset, 
+			bEnabled,
+			ExecuteContext, 
+			Context);
+	}
+
+	if(ParentIndices.Num() == Transforms.Num())
+	{
+		for(int32 Index = 0; Index < ParentIndices.Num(); Index++)
+		{
+			if(Transforms.IsValidIndex(ParentIndices[Index]))
+			{
+				Context.DrawInterface->DrawLine(
+					WorldOffset,
+					Transforms[Index].GetTranslation(),
+					Transforms[ParentIndices[Index]].GetTranslation(),
+					Color,
+					Thickness
+				);
+			}
+		}
+	}
+}
diff --git a/Engine/Plugins/Animation/ControlRig/Source/ControlRig/Private/Units/Debug/RigUnit_DebugTransform.h b/Engine/Plugins/Animation/ControlRig/Source/ControlRig/Private/Units/Debug/RigUnit_DebugTransform.h
--- a/Engine/Plugins/Animation/ControlRig/Source/ControlRig/Private/Units/Debug/RigUnit_DebugTransform.h
+++ b/Engine/Plugins/Animation/ControlRig/Source/ControlRig/Private/Units/Debug/RigUnit_DebugTransform.h
@@ -207,4 +207,54 @@ struct CONTROLRIG_API FRigUnit_DebugTransformArrayMutable : public FRigUnit_Debu
 
 	UPROPERTY(transient)
 	FRigUnit_DebugTransformArrayMutable_WorkData WorkData;
-};
\ No newline at end of file
+};
+
+/**
+* Given a transform array, will draw a point, axis, or a box in the viewport
+*/
+USTRUCT(meta=(DisplayName="Draw Transform Array"))
+struct CONTROLRIG_API FRigUnit_DebugTransformArrayMutableItemSpace : public FRigUnit_DebugBaseMutable
+{
+	GENERATED_BODY()
+
+	FRigUnit_DebugTransformArrayMutableItemSpace()
+	{
+		WorldOffset = FTransform::Identity;
+		Mode = ERigUnitDebugTransformMode::Axes;
+		Color = FLinearColor::White;
+		Thickness = 0.f;
+		Scale = 10.f;
+		Space = FRigElementKey(NAME_None, ERigElementType::Bone);
+		bEnabled = true;
+	}
+
+	RIGVM_METHOD()
+	virtual void Execute(const FRigUnitContext& Context) override;
+
+	UPROPERTY(meta = (Input))
+	TArray<FTransform> Transforms;
+
+	UPROPERTY(meta = (Input))
+	TArray<int32> ParentIndices;
+
+	UPROPERTY(meta = (Input))
+	ERigUnitDebugTransformMode Mode;
+
+	UPROPERTY(meta = (Input))
+	FLinearColor Color;
+
+	UPROPERTY(meta = (Input))
+	float Thickness;
+
+	UPROPERTY(meta = (Input))
+	float Scale;
+
+	UPROPERTY(meta = (Input))
+	FRigElementKey Space;
+
+	UPROPERTY(meta = (Input))
+	FTransform WorldOffset;
+
+	UPROPERTY(meta = (Input, Constant))
+	bool bEnabled;
+};
diff --git a/Engine/Plugins/Animation/ControlRig/Source/ControlRig/Private/Units/Execution/RigUnit_Collection.cpp b/Engine/Plugins/Animation/ControlRig/Source/ControlRig/Private/Units/Execution/RigUnit_Collection.cpp
--- a/Engine/Plugins/Animation/ControlRig/Source/ControlRig/Private/Units/Execution/RigUnit_Collection.cpp
+++ b/Engine/Plugins/Animation/ControlRig/Source/ControlRig/Private/Units/Execution/RigUnit_Collection.cpp
@@ -212,6 +212,70 @@ FRigUnit_CollectionGetItems_Execute()
 	Items = Collection.GetKeys();
 }
 
+FRigUnit_CollectionGetParentIndices_Execute()
+{
+	DECLARE_SCOPE_HIERARCHICAL_COUNTER_RIGUNIT()
+
+	if(Context.Hierarchy == nullptr)
+	{
+		ParentIndices.Reset();
+		return;
+	}
+	
+	ParentIndices.SetNumUninitialized(Collection.Num());
+
+	for(int32 Index=0;Index<Collection.Num();Index++)
+	{
+		ParentIndices[Index] = INDEX_NONE;
+
+		const int32 ItemIndex = Context.Hierarchy->GetIndex(Collection[Index]);
+		if(ItemIndex == INDEX_NONE)
+		{
+			continue;;
+		}
+
+		switch(Collection[Index].Type)
+		{
+			case ERigElementType::Curve:
+			{
+				continue;
+			}
+			case ERigElementType::Bone:
+			{
+				ParentIndices[Index] = Context.Hierarchy->GetFirstParent(ItemIndex);
+				break;
+			}
+			default:
+			{
+				if(const FRigBaseElement* ChildElement = Context.Hierarchy->Get(ItemIndex))
+				{
+					TArray<int32> ItemParents = Context.Hierarchy->GetParents(ItemIndex);
+					for(int32 ParentIndex = 0; ParentIndex < ItemParents.Num(); ParentIndex++)
+					{
+						const FRigElementWeight Weight = Context.Hierarchy->GetParentWeight(ChildElement, ParentIndex, false);
+						if(!Weight.IsAlmostZero())
+						{
+							ParentIndices[Index] = ItemParents[ParentIndex];
+						}
+					}
+				}
+				break;
+			}
+		}
+
+		if(ParentIndices[Index] != INDEX_NONE)
+		{
+			int32 ParentIndex = ParentIndices[Index];
+			ParentIndices[Index] = INDEX_NONE;
+
+			while(ParentIndices[Index] == INDEX_NONE && ParentIndex != INDEX_NONE)
+			{
+				ParentIndices[Index] = Collection.GetKeys().Find(Context.Hierarchy->GetKey(ParentIndex));
+				ParentIndex = Context.Hierarchy->GetFirstParent(ParentIndex);
+			}
+		}
+	}
+}
 
 FRigUnit_CollectionUnion_Execute()
 {
diff --git a/Engine/Plugins/Animation/ControlRig/Source/ControlRig/Private/Units/Execution/RigUnit_Collection.h b/Engine/Plugins/Animation/ControlRig/Source/ControlRig/Private/Units/Execution/RigUnit_Collection.h
--- a/Engine/Plugins/Animation/ControlRig/Source/ControlRig/Private/Units/Execution/RigUnit_Collection.h
+++ b/Engine/Plugins/Animation/ControlRig/Source/ControlRig/Private/Units/Execution/RigUnit_Collection.h
@@ -227,6 +227,32 @@ struct CONTROLRIG_API FRigUnit_CollectionGetItems : public FRigUnit_CollectionBa
 	TArray<FRigElementKey> Items;
 };
 
+/**
+ * Returns an array of relative parent indices for each item. Several options here
+ * a) If an item has multiple parents the major parent (based on the weights) will be returned.
+ * b) If an item has a parent that's not part of the collection INDEX_NONE will be returned.
+ * c) If an item has a parent that's not part of the collection, but a grand parent is we'll use that index instead.
+ */
+USTRUCT(meta = (DisplayName = "Get Parent Indices", Keywords = "Collection,Array", Varying))
+struct CONTROLRIG_API FRigUnit_CollectionGetParentIndices : public FRigUnit_CollectionBase
+{
+	GENERATED_BODY()
+
+	FRigUnit_CollectionGetParentIndices()
+	{
+	}
+
+	RIGVM_METHOD()
+	virtual void Execute(const FRigUnitContext& Context) override;
+
+	UPROPERTY(meta = (Input))
+	FRigElementKeyCollection Collection;
+
+	UPROPERTY(meta = (Output))
+	TArray<int32> ParentIndices;
+};
+
+
 /**
  * Returns the union of two provided collections
  * (the combination of all items from both A and B).
diff --git a/Engine/Plugins/Animation/ControlRig/Source/ControlRig/Private/Units/Hierarchy/RigUnit_GetTransform.cpp b/Engine/Plugins/Animation/ControlRig/Source/ControlRig/Private/Units/Hierarchy/RigUnit_GetTransform.cpp
--- a/Engine/Plugins/Animation/ControlRig/Source/ControlRig/Private/Units/Hierarchy/RigUnit_GetTransform.cpp
+++ b/Engine/Plugins/Animation/ControlRig/Source/ControlRig/Private/Units/Hierarchy/RigUnit_GetTransform.cpp
@@ -79,3 +79,21 @@ FRigUnit_GetTransform_Execute()
 		}
 	}
 }
+
+FRigUnit_GetTransformArray_Execute()
+{
+    DECLARE_SCOPE_HIERARCHICAL_COUNTER_RIGUNIT()
+
+	if(CachedIndex.Num() != Items.Num())
+	{
+		CachedIndex.Reset();
+		CachedIndex.SetNum(Items.Num());
+	}
+
+	Transforms.SetNum(Items.Num());
+	for(int32 Index=0;Index<Items.Num();Index++)
+	{
+		FRigUnit_GetTransform::StaticExecute(RigVMExecuteContext, Items[Index], Space, bInitial, Transforms[Index], CachedIndex[Index], Context);
+
+	}
+}
diff --git a/Engine/Plugins/Animation/ControlRig/Source/ControlRig/Private/Units/Hierarchy/RigUnit_GetTransform.h b/Engine/Plugins/Animation/ControlRig/Source/ControlRig/Private/Units/Hierarchy/RigUnit_GetTransform.h
--- a/Engine/Plugins/Animation/ControlRig/Source/ControlRig/Private/Units/Hierarchy/RigUnit_GetTransform.h
+++ b/Engine/Plugins/Animation/ControlRig/Source/ControlRig/Private/Units/Hierarchy/RigUnit_GetTransform.h
@@ -65,3 +65,50 @@ struct CONTROLRIG_API FRigUnit_GetTransform : public FRigUnit
 	UPROPERTY()
 	FCachedRigElement CachedIndex;
 };
+
+/**
+* GetTransformArray is used to retrieve an array of transforms from the hierarchy.
+*/
+USTRUCT(meta=(DisplayName="Get Transform Array", Category="Hierarchy", DocumentationPolicy = "Strict", Keywords="GetBoneTransform,GetControlTransform,GetInitialTransform,GetSpaceTransform,GetTransform", NodeColor="0.462745, 1,0, 0.329412",Varying))
+struct CONTROLRIG_API FRigUnit_GetTransformArray : public FRigUnit
+{
+	GENERATED_BODY()
+
+	FRigUnit_GetTransformArray()
+		: Items()
+		, Space(EBoneGetterSetterMode::GlobalSpace)
+		, bInitial(false)
+		, Transforms()
+		, CachedIndex()
+	{}
+
+	RIGVM_METHOD()
+	virtual void Execute(const FRigUnitContext& Context) override;
+
+	/**
+	* The items to retrieve the transforms for
+	*/
+	UPROPERTY(meta = (Input))
+	FRigElementKeyCollection Items;
+
+	/**
+	* Defines if the transforms should be retrieved in local or global space
+	*/ 
+	UPROPERTY(meta = (Input))
+	EBoneGetterSetterMode Space;
+
+	/**
+	* Defines if the transforms should be retrieved as current (false) or initial (true).
+	* Initial transforms for bones and other elements in the hierarchy represent the reference pose's value.
+	*/ 
+	UPROPERTY(meta = (Input))
+	bool bInitial;
+
+	// The current transform of the given item - or identity in case it wasn't found.
+	UPROPERTY(meta=(Output))
+	TArray<FTransform> Transforms;
+
+	// Used to cache the internally
+	UPROPERTY()
+	TArray<FCachedRigElement> CachedIndex;
+};
diff --git a/Engine/Plugins/Animation/ControlRig/Source/ControlRig/Private/Units/Hierarchy/RigUnit_SetTransform.cpp b/Engine/Plugins/Animation/ControlRig/Source/ControlRig/Private/Units/Hierarchy/RigUnit_SetTransform.cpp
--- a/Engine/Plugins/Animation/ControlRig/Source/ControlRig/Private/Units/Hierarchy/RigUnit_SetTransform.cpp
+++ b/Engine/Plugins/Animation/ControlRig/Source/ControlRig/Private/Units/Hierarchy/RigUnit_SetTransform.cpp
@@ -157,3 +157,26 @@ FRigUnit_SetScale_Execute()
 	Transform.SetScale3D(Scale);
 	FRigUnit_SetTransform::StaticExecute(RigVMExecuteContext, Item, Space, false, Transform, Weight, bPropagateToChildren, CachedIndex, ExecuteContext, Context);
 }
+
+
+FRigUnit_SetTransformArray_Execute()
+{
+	DECLARE_SCOPE_HIERARCHICAL_COUNTER_RIGUNIT()
+
+	if(CachedIndex.Num() != Items.Num())
+	{
+		CachedIndex.Reset();
+		CachedIndex.SetNum(Items.Num());
+	}
+
+	if(Transforms.Num() != Items.Num())
+	{
+		UE_CONTROLRIG_RIGUNIT_REPORT_WARNING(TEXT("The number of transforms (%d) doesn't match the size of the collection (%d)."), Transforms.Num(), Items.Num());
+		return;
+	}
+
+	for(int32 Index=0;Index<Items.Num();Index++)
+	{
+		FRigUnit_SetTransform::StaticExecute(RigVMExecuteContext, Items[Index], Space, bInitial, Transforms[Index], Weight, bPropagateToChildren, CachedIndex[Index], ExecuteContext, Context);
+	}
+}
\ No newline at end of file
diff --git a/Engine/Plugins/Animation/ControlRig/Source/ControlRig/Private/Units/Math/RigUnit_MathBase.h b/Engine/Plugins/Animation/ControlRig/Source/ControlRig/Private/Units/Math/RigUnit_MathBase.h
--- a/Engine/Plugins/Animation/ControlRig/Source/ControlRig/Private/Units/Math/RigUnit_MathBase.h
+++ b/Engine/Plugins/Animation/ControlRig/Source/ControlRig/Private/Units/Math/RigUnit_MathBase.h
@@ -10,3 +10,9 @@ struct CONTROLRIG_API FRigUnit_MathBase : public FRigUnit
 {
 	GENERATED_BODY()
 };
+
+USTRUCT(meta=(Abstract, NodeColor = "0.05 0.25 0.05"))
+struct CONTROLRIG_API FRigUnit_MathMutableBase : public FRigUnitMutable
+{
+	GENERATED_BODY()
+};
diff --git a/Engine/Plugins/Animation/ControlRig/Source/ControlRig/Private/Units/Math/RigUnit_MathTransform.cpp b/Engine/Plugins/Animation/ControlRig/Source/ControlRig/Private/Units/Math/RigUnit_MathTransform.cpp
--- a/Engine/Plugins/Animation/ControlRig/Source/ControlRig/Private/Units/Math/RigUnit_MathTransform.cpp
+++ b/Engine/Plugins/Animation/ControlRig/Source/ControlRig/Private/Units/Math/RigUnit_MathTransform.cpp
@@ -37,6 +37,62 @@ FRigUnit_MathTransformMakeAbsolute_Execute()
 	Global.NormalizeRotation();
 }
 
+FRigUnit_MathTransformAccumulateArray_Execute()
+{
+	DECLARE_SCOPE_HIERARCHICAL_COUNTER_RIGUNIT()
+	if(Transforms.Num() == 0)
+	{
+		return;
+	}
+
+	if(ParentIndices.Num() > 0 && ParentIndices.Num() != Transforms.Num())
+	{
+		UE_CONTROLRIG_RIGUNIT_REPORT_ERROR(TEXT("If the indices are specified their num (%d) has to match the transforms (%d)."), ParentIndices.Num(), Transforms.Num());
+		return;
+	}
+
+	if(TargetSpace == EBoneGetterSetterMode::LocalSpace)
+	{
+		if(ParentIndices.IsEmpty())
+		{
+			for(int32 Index=Transforms.Num()-1; Index>=0;Index--)
+			{
+				const FTransform& ParentTransform = (Index == 0) ? Root : Transforms[Index - 1];
+				Transforms[Index] = Transforms[Index].GetRelativeTransform(ParentTransform);
+			}
+		}
+		else
+		{
+			for(int32 Index=Transforms.Num()-1; Index>=0;Index--)
+			{
+				const int32 ParentIndex = ParentIndices[Index];
+				const FTransform& ParentTransform = (ParentIndex == INDEX_NONE || ParentIndex >= Index) ? Root : Transforms[ParentIndex];
+				Transforms[Index] = Transforms[Index].GetRelativeTransform(ParentTransform);
+			}
+		}
+	}
+	else
+	{
+		if(ParentIndices.IsEmpty())
+		{
+			for(int32 Index=0; Index<Transforms.Num(); Index++)
+			{
+				const FTransform& ParentTransform = (Index == 0) ? Root : Transforms[Index - 1];
+				Transforms[Index] = Transforms[Index] * ParentTransform;
+			}
+		}
+		else
+		{
+			for(int32 Index=0; Index<Transforms.Num(); Index++)
+			{
+				const int32 ParentIndex = ParentIndices[Index];
+				const FTransform& ParentTransform = (ParentIndex == INDEX_NONE || ParentIndex >= Index) ? Root : Transforms[ParentIndex];
+				Transforms[Index] = Transforms[Index] * ParentTransform;
+			}
+		}
+	}
+}
+
 FRigUnit_MathTransformInverse_Execute()
 {
     DECLARE_SCOPE_HIERARCHICAL_COUNTER_RIGUNIT()
diff --git a/Engine/Plugins/Animation/ControlRig/Source/ControlRig/Private/Units/Math/RigUnit_MathTransform.h b/Engine/Plugins/Animation/ControlRig/Source/ControlRig/Private/Units/Math/RigUnit_MathTransform.h
--- a/Engine/Plugins/Animation/ControlRig/Source/ControlRig/Private/Units/Math/RigUnit_MathTransform.h
+++ b/Engine/Plugins/Animation/ControlRig/Source/ControlRig/Private/Units/Math/RigUnit_MathTransform.h
@@ -13,6 +13,12 @@ struct CONTROLRIG_API FRigUnit_MathTransformBase : public FRigUnit_MathBase
 	GENERATED_BODY()
 };
 
+USTRUCT(meta=(Abstract, Category="Math|Transform", MenuDescSuffix="(Transform)"))
+struct CONTROLRIG_API FRigUnit_MathTransformMutableBase : public FRigUnit_MathMutableBase
+{
+	GENERATED_BODY()
+};
+
 USTRUCT(meta=(Abstract))
 struct CONTROLRIG_API FRigUnit_MathTransformUnaryOp : public FRigUnit_MathTransformBase
 {
@@ -161,6 +167,48 @@ struct CONTROLRIG_API FRigUnit_MathTransformMakeAbsolute : public FRigUnit_MathT
 	FTransform Global;
 };
 
+/**
+* Treats the provided transforms as a chain with global / local transforms, and
+* projects each transform into the target space. Optionally you can provide
+* a custom parent indices array, with which you can represent more than just chains.
+*/
+USTRUCT(meta=(DisplayName="Accumulate Transform Array", Keywords="Local,Global,Absolute,Array"))
+struct CONTROLRIG_API FRigUnit_MathTransformAccumulateArray : public FRigUnit_MathTransformMutableBase
+{
+	GENERATED_BODY()
+
+	FRigUnit_MathTransformAccumulateArray()
+	{
+		TargetSpace = EBoneGetterSetterMode::GlobalSpace;
+		Root = FTransform::Identity;
+	}
+
+	RIGVM_METHOD()
+	virtual void Execute(const FRigUnitContext& Context) override;
+
+	UPROPERTY(meta=(Input, Output))
+	TArray<FTransform> Transforms;
+
+	/**
+	* Defines the space to project to
+	*/ 
+	UPROPERTY(meta = (Input))
+	EBoneGetterSetterMode TargetSpace;
+
+	/**
+	 * Provides the parent transform for the root
+	 */
+	UPROPERTY(meta=(Input))
+	FTransform Root;
+
+	/**
+	 * If this array is the same size as the transforms array the indices will be used
+	 * to look up each transform's parent. They are expected to be in order.
+	 */
+	UPROPERTY(meta=(Input))
+	TArray<int32> ParentIndices;
+};
+
 /**
  * Returns the negative value
  */
diff --git a/Engine/Plugins/Animation/ControlRig/Source/ControlRig/Private/Units/Simulation/RigUnit_SpringInterp.cpp b/Engine/Plugins/Animation/ControlRig/Source/ControlRig/Private/Units/Simulation/RigUnit_SpringInterp.cpp
--- a/Engine/Plugins/Animation/ControlRig/Source/ControlRig/Private/Units/Simulation/RigUnit_SpringInterp.cpp
+++ b/Engine/Plugins/Animation/ControlRig/Source/ControlRig/Private/Units/Simulation/RigUnit_SpringInterp.cpp
@@ -15,6 +15,7 @@ FRigUnit_SpringInterp_Execute()
 	if (Context.State == EControlRigState::Init)
 	{
 		SpringState.Reset();
+		Result = Target;
 	}
 	else
 	{
@@ -30,10 +31,12 @@ FRigUnit_SpringInterp_Execute()
 		{
 			SpringState.Velocity += Force * (Context.DeltaTime / RigUnitSpringInterpConstants::Mass);
 		}
-		Result = UKismetMathLibrary::FloatSpringInterp(
-			bUseCurrentInput ? Current : Result, AdjustedTarget, SpringState, Stiffness, CriticalDamping,
+		SimulatedResult = UKismetMathLibrary::FloatSpringInterp(
+			bUseCurrentInput ? Current : SimulatedResult, AdjustedTarget, SpringState, Stiffness, CriticalDamping,
 			Context.DeltaTime, RigUnitSpringInterpConstants::Mass, TargetVelocityAmount, 
 			false, 0.0f, 0.0f, !bUseCurrentInput || bInitializeFromTarget);
+
+		Result = SimulatedResult;
 		Velocity = SpringState.Velocity;
 	}
 }
@@ -45,6 +48,7 @@ FRigUnit_SpringInterpVector_Execute()
 	if (Context.State == EControlRigState::Init)
 	{
 		SpringState.Reset();
+		Result = Target;
 	}
 	else
 	{
@@ -60,10 +64,11 @@ FRigUnit_SpringInterpVector_Execute()
 		{
 			SpringState.Velocity += Force * (Context.DeltaTime / RigUnitSpringInterpConstants::Mass);
 		}
-		Result = UKismetMathLibrary::VectorSpringInterp(
-			bUseCurrentInput ? Current : Result, AdjustedTarget, SpringState, Stiffness, CriticalDamping,
+		SimulatedResult = UKismetMathLibrary::VectorSpringInterp(
+			bUseCurrentInput ? Current : SimulatedResult, AdjustedTarget, SpringState, Stiffness, CriticalDamping,
 			Context.DeltaTime, RigUnitSpringInterpConstants::Mass, TargetVelocityAmount,
 			false, FVector(), FVector(), !bUseCurrentInput || bInitializeFromTarget);
+		Result = SimulatedResult;
 	}
 	Velocity = SpringState.Velocity;
 }
@@ -75,6 +80,7 @@ FRigUnit_SpringInterpQuaternion_Execute()
 	if (Context.State == EControlRigState::Init)
 	{
 		SpringState.Reset();
+		Result = Target;
 	}
 	else
 	{
@@ -82,10 +88,11 @@ FRigUnit_SpringInterpQuaternion_Execute()
 		float AngularFrequency = Strength * 2.0f * PI;
 		float Stiffness = AngularFrequency * AngularFrequency;
 		SpringState.AngularVelocity += Torque * (Context.DeltaTime / RigUnitSpringInterpConstants::Mass);
-		Result = UKismetMathLibrary::QuaternionSpringInterp(
-			bUseCurrentInput ? Current : Result, Target, SpringState, Stiffness, CriticalDamping,
+		SimulatedResult = UKismetMathLibrary::QuaternionSpringInterp(
+			bUseCurrentInput ? Current : SimulatedResult, Target, SpringState, Stiffness, CriticalDamping,
 			Context.DeltaTime, RigUnitSpringInterpConstants::Mass, TargetVelocityAmount, 
 			!bUseCurrentInput || bInitializeFromTarget);
+		Result = SimulatedResult;
 		AngularVelocity = SpringState.AngularVelocity;
 	}
 }
diff --git a/Engine/Plugins/Animation/ControlRig/Source/ControlRig/Private/Units/Simulation/RigUnit_SpringInterp.h b/Engine/Plugins/Animation/ControlRig/Source/ControlRig/Private/Units/Simulation/RigUnit_SpringInterp.h
--- a/Engine/Plugins/Animation/ControlRig/Source/ControlRig/Private/Units/Simulation/RigUnit_SpringInterp.h
+++ b/Engine/Plugins/Animation/ControlRig/Source/ControlRig/Private/Units/Simulation/RigUnit_SpringInterp.h
@@ -17,7 +17,7 @@ struct CONTROLRIG_API FRigUnit_SpringInterp : public FRigUnit_SimBase
 	FRigUnit_SpringInterp()
 	{
 		bUseCurrentInput = false;
-		Current = Target = Velocity = Result = 0.0f;
+		Current = Target = Velocity = Result = SimulatedResult = 0.0f;
 		Strength = 4.0f;
 		CriticalDamping = 1.0f;
 		TargetVelocityAmount = 0.0f;
@@ -89,6 +89,9 @@ struct CONTROLRIG_API FRigUnit_SpringInterp : public FRigUnit_SimBase
 	UPROPERTY(meta = (Output))
 	float Velocity;
 
+	UPROPERTY()
+	float SimulatedResult;
+
 	UPROPERTY()
 	FFloatSpringState SpringState;
 };
@@ -104,7 +107,7 @@ struct CONTROLRIG_API FRigUnit_SpringInterpVector : public FRigUnit_SimBase
 	FRigUnit_SpringInterpVector()
 	{
 		bUseCurrentInput = false;
-		Current = Target = Velocity = Result = FVector::ZeroVector;
+		Current = Target = Velocity = Result = SimulatedResult = FVector::ZeroVector;
 		Strength = 4.0f;
 		CriticalDamping = 1.0f;
 		TargetVelocityAmount = 0.0f;
@@ -176,6 +179,9 @@ struct CONTROLRIG_API FRigUnit_SpringInterpVector : public FRigUnit_SimBase
 	UPROPERTY(meta = (Output))
 	FVector Velocity;
 
+	UPROPERTY()
+	FVector SimulatedResult;
+
 	UPROPERTY()
 	FVectorSpringState SpringState;
 };
@@ -191,7 +197,7 @@ struct CONTROLRIG_API FRigUnit_SpringInterpQuaternion : public FRigUnit_SimBase
 	FRigUnit_SpringInterpQuaternion()
 	{
 		bUseCurrentInput = false;
-		Current = Target = Result = FQuat::Identity;
+		Current = Target = Result = SimulatedResult = FQuat::Identity;
 		AngularVelocity = FVector::ZeroVector;
 		Strength = 4.0f;
 		CriticalDamping = 1.0f;
@@ -264,6 +270,9 @@ struct CONTROLRIG_API FRigUnit_SpringInterpQuaternion : public FRigUnit_SimBase
 	UPROPERTY(meta = (Output))
 	FVector AngularVelocity;
 
+	UPROPERTY(meta=(Output))
+	FQuat SimulatedResult;
+
 	UPROPERTY()
 	FQuaternionSpringState SpringState;
 };
diff --git a/Engine/Plugins/Animation/ControlRig/Source/ControlRig/Public/Units/Hierarchy/RigUnit_SetTransform.h b/Engine/Plugins/Animation/ControlRig/Source/ControlRig/Public/Units/Hierarchy/RigUnit_SetTransform.h
--- a/Engine/Plugins/Animation/ControlRig/Source/ControlRig/Public/Units/Hierarchy/RigUnit_SetTransform.h
+++ b/Engine/Plugins/Animation/ControlRig/Source/ControlRig/Public/Units/Hierarchy/RigUnit_SetTransform.h
@@ -261,3 +261,65 @@ struct FRigUnit_SetScale : public FRigUnitMutable
 	UPROPERTY()
 	FCachedRigElement CachedIndex;
 };
+
+/**
+ * SetTransformArray is used to set an array of transforms on the hierarchy.
+ * 
+ * Note: For Controls when setting the initial transform this node
+ * actually sets the Control's offset transform and resets the local
+ * values to (0, 0, 0).
+ */
+USTRUCT(meta=(DisplayName="Set Transform Array", Category="Hierarchy", DocumentationPolicy = "Strict", Keywords="SetBoneTransform,SetControlTransform,SetInitialTransform,SetSpaceTransform", NodeColor="0, 0.364706, 1.0", Varying))
+struct CONTROLRIG_API FRigUnit_SetTransformArray : public FRigUnitMutable
+{
+	GENERATED_BODY()
+
+	FRigUnit_SetTransformArray()
+		: Items()
+		, Space(EBoneGetterSetterMode::GlobalSpace)
+		, bInitial(false)
+		, Transforms()
+		, Weight(1.f)
+		, bPropagateToChildren(true)
+		, CachedIndex()
+	{}
+
+	RIGVM_METHOD()
+	virtual void Execute(const FRigUnitContext& Context) override;
+
+	/**
+	 * The item to set the transform for
+	 */
+	UPROPERTY(meta = (Input))
+	FRigElementKeyCollection Items;
+
+	/**
+	 * Defines if the transform should be set in local or global space
+	 */ 
+	UPROPERTY(meta = (Input))
+	EBoneGetterSetterMode Space;
+
+	/**
+	 * Defines if the transform should be set as current (false) or initial (true).
+	 * Initial transforms for bones and other elements in the hierarchy represent the reference pose's value.
+	 */ 
+	UPROPERTY(meta = (Input))
+	bool bInitial;
+
+	// The new transform of the given item
+	UPROPERTY(meta=(Input))
+	TArray<FTransform> Transforms;
+
+	// Defines how much the change will be applied
+	UPROPERTY(meta = (Input, UIMin = "0.0", UIMax = "1.0"))
+	float Weight;
+
+	// If set to true children of affected items in the hierarchy
+	// will follow the transform change - otherwise only the parent will move.
+	UPROPERTY(meta=(Input))
+	bool bPropagateToChildren;
+
+	// Used to cache the internally
+	UPROPERTY()
+	TArray<FCachedRigElement> CachedIndex;
+};
