BP Editor: added a virtual function to allow derived BP editors to decide whether to load BP func lib for autocomplete or not. In case of Control Rig BP, since the VMs are different, there is no need to load BP func lib. In some projects a BP func lib can reference many assets, causing Control Rig BP Editor to take a long time to load even an empty Control Rig BP

diff --git a/Engine/Plugins/Animation/ControlRig/Source/ControlRigEditor/Private/Editor/ControlRigEditor.h b/Engine/Plugins/Animation/ControlRig/Source/ControlRigEditor/Private/Editor/ControlRigEditor.h
--- a/Engine/Plugins/Animation/ControlRig/Source/ControlRigEditor/Private/Editor/ControlRigEditor.h
+++ b/Engine/Plugins/Animation/ControlRig/Source/ControlRigEditor/Private/Editor/ControlRigEditor.h
@@ -130,6 +130,9 @@ public:
 
 	virtual FReply OnSpawnGraphNodeByShortcut(FInputChord InChord, const FVector2D& InPosition, UEdGraph* InGraph) override;
 
+	// Control Rig BP does not use regular BP function libraries, no need to load them.
+	virtual bool ShouldLoadBPLibrariesFromAssetRegistry() override { return false; }
+
 	// FEditorUndoClient Interface
 	virtual void PostUndo(bool bSuccess) override;
 	virtual void PostRedo(bool bSuccess) override;
diff --git a/Engine/Source/Editor/Kismet/Private/BlueprintEditor.cpp b/Engine/Source/Editor/Kismet/Private/BlueprintEditor.cpp
--- a/Engine/Source/Editor/Kismet/Private/BlueprintEditor.cpp
+++ b/Engine/Source/Editor/Kismet/Private/BlueprintEditor.cpp
@@ -1856,7 +1856,7 @@ void FBlueprintEditor::CommonInitialization(const TArray<UBlueprint*>& InitBluep
 	if (InitBlueprints.Num() == 1)
 	{
 		// Load blueprint libraries
-		if (!bShouldOpenInDefaultsMode)
+		if (!bShouldOpenInDefaultsMode && ShouldLoadBPLibrariesFromAssetRegistry())
 		{
 			LoadLibrariesFromAssetRegistry();
 		}
diff --git a/Engine/Source/Editor/Kismet/Public/BlueprintEditor.h b/Engine/Source/Editor/Kismet/Public/BlueprintEditor.h
--- a/Engine/Source/Editor/Kismet/Public/BlueprintEditor.h
+++ b/Engine/Source/Editor/Kismet/Public/BlueprintEditor.h
@@ -493,6 +493,9 @@ public:
 	// Should be called when initializing any editor built off this foundation
 	void CommonInitialization(const TArray<UBlueprint*>& InitBlueprints, bool bShouldOpenInDefaultsMode);
 
+	// Other types of BP editors that don't use BP VM/function libs can choose to not load BP function libraries, for example: Control Rig Blueprints
+	virtual bool ShouldLoadBPLibrariesFromAssetRegistry() { return true; }
+
 	// Should be called when initializing an editor that has a blueprint, after layout (tab spawning) is done
 	void PostLayoutBlueprintEditorInitialization();
 
