Control Rig: Performance improvements

* use inline allocators to avoid allocations for minimum size cases
* batch allocate elements for copy hierarchy and use placement new

diff --git a/Engine/Plugins/Animation/ControlRig/Source/ControlRig/Private/ControlRig.cpp b/Engine/Plugins/Animation/ControlRig/Source/ControlRig/Private/ControlRig.cpp
--- a/Engine/Plugins/Animation/ControlRig/Source/ControlRig/Private/ControlRig.cpp
+++ b/Engine/Plugins/Animation/ControlRig/Source/ControlRig/Private/ControlRig.cpp
@@ -2755,7 +2755,7 @@ UControlRig::FPoseScope::FPoseScope(UControlRig* InControlRig)
 {
 	check(InControlRig);
 
-	CachedPose = InControlRig->GetHierarchy()->GetPose();
+	CachedPose = InControlRig->GetHierarchy()->GetPose(false, ERigElementType::Control, FRigElementKeyCollection());
 
 	ControlRig = InControlRig;
 }
diff --git a/Engine/Plugins/Animation/ControlRig/Source/ControlRig/Private/Drawing/ControlRigDrawInterface.cpp b/Engine/Plugins/Animation/ControlRig/Source/ControlRig/Private/Drawing/ControlRigDrawInterface.cpp
--- a/Engine/Plugins/Animation/ControlRig/Source/ControlRig/Private/Drawing/ControlRigDrawInterface.cpp
+++ b/Engine/Plugins/Animation/ControlRig/Source/ControlRig/Private/Drawing/ControlRigDrawInterface.cpp
@@ -346,7 +346,7 @@ void FControlRigDrawInterface::DrawHierarchy(const FTransform& WorldOffset, URig
 				InstructionZ.Positions.Add(P0);
 				InstructionZ.Positions.Add(PZ);
 
-				TArray<FRigBaseElement*> Parents = Hierarchy->GetParents(Child);
+				FRigBaseElementParentArray Parents = Hierarchy->GetParents(Child);
 				for (FRigBaseElement* Parent : Parents)
 				{
 					if(FRigTransformElement* ParentTransformElement = Cast<FRigTransformElement>(Parent))
diff --git a/Engine/Plugins/Animation/ControlRig/Source/ControlRig/Private/Rigs/RigHierarchy.cpp b/Engine/Plugins/Animation/ControlRig/Source/ControlRig/Private/Rigs/RigHierarchy.cpp
--- a/Engine/Plugins/Animation/ControlRig/Source/ControlRig/Private/Rigs/RigHierarchy.cpp
+++ b/Engine/Plugins/Animation/ControlRig/Source/ControlRig/Private/Rigs/RigHierarchy.cpp
@@ -79,7 +79,7 @@ FAutoConsoleCommandWithWorldAndArgs FCmdControlRigHierarchyTraceFrames
 // URigHierarchy
 ////////////////////////////////////////////////////////////////////////////////
 
-const TArray<FRigBaseElement*> URigHierarchy::EmptyElementArray;
+const FRigBaseElementChildrenArray URigHierarchy::EmptyElementArray;
 
 URigHierarchy::URigHierarchy()
 : TopologyVersion(0)
@@ -196,7 +196,7 @@ void URigHierarchy::Load(FArchive& Ar)
 
 		Element->SubIndex = Num(Key.Type);
 		Element->Index = Elements.Add(Element);
-		ElementsPerType[(int32)Key.Type].Add(Element);
+		ElementsPerType[RigElementTypeToFlatIndex(Key.Type)].Add(Element);
 		IndexLookup.Add(Key, Element->Index);
 		
 		Element->Load(Ar, this, FRigBaseElement::StaticData);
@@ -217,9 +217,9 @@ void URigHierarchy::Load(FArchive& Ar)
 		if(FRigTransformElement* TransformElement = Cast<FRigTransformElement>(Elements[ElementIndex]))
 		{
 #if URIGHIERARCHY_RECURSIVE_DIRTY_PROPAGATION
-			TArray<FRigBaseElement*> CurrentParents = GetParents(TransformElement, false);
+			FRigBaseElementParentArray CurrentParents = GetParents(TransformElement, false);
 #else
-			TArray<FRigBaseElement*> CurrentParents = GetParents(TransformElement, true);
+			FRigBaseElementParentArray CurrentParents = GetParents(TransformElement, true);
 #endif
 			for (FRigBaseElement* CurrentParent : CurrentParents)
 			{
@@ -240,13 +240,14 @@ void URigHierarchy::Reset()
 	TopologyVersion = 0;
 	bEnableDirtyPropagation = true;
 
-	for(int32 ElementIndex=0; ElementIndex<Elements.Num(); ElementIndex++)
+	// walk in reverse since certain elements might not have been allocated themselves
+	for(int32 ElementIndex = Elements.Num() - 1; ElementIndex >= 0; ElementIndex--)
 	{
-		delete Elements[ElementIndex];
+		DestroyElement(Elements[ElementIndex]);
 	}
 	Elements.Reset();
 	ElementsPerType.Reset();
-	for(int32 TypeIndex=0;TypeIndex<(int32)ERigElementType::Last;TypeIndex++)
+	for(int32 TypeIndex=0;TypeIndex<RigElementTypeToFlatIndex(ERigElementType::Last);TypeIndex++)
 	{
 		ElementsPerType.Add(TArray<FRigBaseElement*>());
 	}
@@ -267,18 +268,51 @@ void URigHierarchy::CopyHierarchy(URigHierarchy* InHierarchy)
 	
 	Reset();
 
+	// Allocate the elements in batches to improve performance
+	TArray<uint8*> NewElementsPerType;
+	TArray<int32> StructureSizePerType;
+	for(int32 ElementTypeIndex = 0; ElementTypeIndex < InHierarchy->ElementsPerType.Num(); ElementTypeIndex++)
+	{
+		const ERigElementType ElementType = FlatIndexToRigElementType(ElementTypeIndex);
+		int32 StructureSize = 0;
+
+		const int32 Count = InHierarchy->ElementsPerType[ElementTypeIndex].Num();
+		if(Count)
+		{
+			FRigBaseElement* ElementMemory = MakeElement(ElementType, Count, &StructureSize);
+			NewElementsPerType.Add((uint8*)ElementMemory);
+		}
+		else
+		{
+			NewElementsPerType.Add(nullptr);
+		}
+
+		StructureSizePerType.Add(StructureSize);
+		ElementsPerType[ElementTypeIndex].Reserve(Count);
+	}
+
+	Elements.Reserve(InHierarchy->Elements.Num());
+	IndexLookup.Reserve(InHierarchy->IndexLookup.Num());
+	
 	for(int32 Index = 0; Index < InHierarchy->Num(); Index++)
 	{
 		FRigBaseElement* Source = InHierarchy->Get(Index);
 		const FRigElementKey& Key = Source->Key;
 
-		FRigBaseElement* Target = MakeElement(Key.Type);
+		const int32 ElementTypeIndex = RigElementTypeToFlatIndex(Key.Type);
+		
+		const int32 SubIndex = Num(Key.Type);
+
+		const int32 StructureSize = StructureSizePerType[ElementTypeIndex];
+		check(NewElementsPerType[ElementTypeIndex] != nullptr);
+		FRigBaseElement* Target = (FRigBaseElement*)&NewElementsPerType[ElementTypeIndex][StructureSize * SubIndex];
+		//FRigBaseElement* Target = MakeElement(Key.Type);
 		
 		Target->Key = Key;
-		Target->SubIndex = Num(Target->GetType());
+		Target->SubIndex = SubIndex;
 		Target->Index = Elements.Add(Target);
 
-		ElementsPerType[(int32)Key.Type].Add(Target);
+		ElementsPerType[ElementTypeIndex].Add(Target);
 		IndexLookup.Add(Key, Target->Index);
 
 		check(Source->Index == Index);
@@ -415,7 +449,7 @@ void URigHierarchy::ResetPoseToInitial(ERigElementType InTypeFilter)
 				const bool bFilteredOut = (!InElement->IsTypeOf(InTypeFilter)) || ResetPoseHasFilteredChildren[InElement->Index];
 				if(bFilteredOut)
 				{
-					const TArray<FRigBaseElement*> Parents = GetParents(InElement);
+					const FRigBaseElementParentArray Parents = GetParents(InElement);
 					for(const FRigBaseElement* Parent : Parents)
 					{
 						// only mark this up if the parent is not filtered out / 
@@ -513,7 +547,7 @@ void URigHierarchy::ResetCurveValues()
 
 int32 URigHierarchy::Num(ERigElementType InElementType) const
 {
-	return ElementsPerType[(int32)InElementType].Num();
+	return ElementsPerType[RigElementTypeToFlatIndex(InElementType)].Num();
 }
 
 TArray<FRigBaseElement*> URigHierarchy::GetSelectedElements(ERigElementType InTypeFilter) const
@@ -662,8 +696,8 @@ FName URigHierarchy::GetSafeNewName(const FString& InPotentialNewName, ERigEleme
 
 TArray<FRigElementKey> URigHierarchy::GetChildren(FRigElementKey InKey, bool bRecursive) const
 {
-	TArray<FRigBaseElement*> LocalChildren;
-	const TArray<FRigBaseElement*>* ChildrenPtr = nullptr;
+	FRigBaseElementChildrenArray LocalChildren;
+	const FRigBaseElementChildrenArray* ChildrenPtr = nullptr;
 	if(bRecursive)
 	{
 		LocalChildren = GetChildren(Find(InKey), true);
@@ -674,7 +708,7 @@ TArray<FRigElementKey> URigHierarchy::GetChildren(FRigElementKey InKey, bool bRe
 		ChildrenPtr = &GetChildren(Find(InKey));
 	}
 
-	const TArray<FRigBaseElement*>& Children = *ChildrenPtr;
+	const FRigBaseElementChildrenArray& Children = *ChildrenPtr;
 
 	TArray<FRigElementKey> Keys;
 	for(const FRigBaseElement* Child : Children)
@@ -686,8 +720,8 @@ TArray<FRigElementKey> URigHierarchy::GetChildren(FRigElementKey InKey, bool bRe
 
 TArray<int32> URigHierarchy::GetChildren(int32 InIndex, bool bRecursive) const
 {
-	TArray<FRigBaseElement*> LocalChildren;
-	const TArray<FRigBaseElement*>* ChildrenPtr = nullptr;
+	FRigBaseElementChildrenArray LocalChildren;
+	const FRigBaseElementChildrenArray* ChildrenPtr = nullptr;
 	if(bRecursive)
 	{
 		LocalChildren = GetChildren(Get(InIndex), true);
@@ -698,7 +732,7 @@ TArray<int32> URigHierarchy::GetChildren(int32 InIndex, bool bRecursive) const
 		ChildrenPtr = &GetChildren(Get(InIndex));
 	}
 
-	const TArray<FRigBaseElement*>& Children = *ChildrenPtr;
+	const FRigBaseElementChildrenArray& Children = *ChildrenPtr;
 
 	TArray<int32> Indices;
 	for(const FRigBaseElement* Child : Children)
@@ -708,7 +742,7 @@ TArray<int32> URigHierarchy::GetChildren(int32 InIndex, bool bRecursive) const
 	return Indices;
 }
 
-const TArray<FRigBaseElement*>& URigHierarchy::GetChildren(const FRigBaseElement* InElement) const
+const FRigBaseElementChildrenArray& URigHierarchy::GetChildren(const FRigBaseElement* InElement) const
 {
 	if(InElement)
 	{
@@ -718,10 +752,10 @@ const TArray<FRigBaseElement*>& URigHierarchy::GetChildren(const FRigBaseElement
 	return EmptyElementArray;
 }
 
-TArray<FRigBaseElement*> URigHierarchy::GetChildren(const FRigBaseElement* InElement, bool bRecursive) const
+FRigBaseElementChildrenArray URigHierarchy::GetChildren(const FRigBaseElement* InElement, bool bRecursive) const
 {
 	// call the non-recursive variation
-	TArray<FRigBaseElement*> Children = GetChildren(InElement);
+	FRigBaseElementChildrenArray Children = GetChildren(InElement);
 	
 	if(bRecursive)
 	{
@@ -736,7 +770,7 @@ TArray<FRigBaseElement*> URigHierarchy::GetChildren(const FRigBaseElement* InEle
 
 TArray<FRigElementKey> URigHierarchy::GetParents(FRigElementKey InKey, bool bRecursive) const
 {
-	const TArray<FRigBaseElement*>& Parents = GetParents(Find(InKey), bRecursive);
+	const FRigBaseElementParentArray& Parents = GetParents(Find(InKey), bRecursive);
 	TArray<FRigElementKey> Keys;
 	for(const FRigBaseElement* Parent : Parents)
 	{
@@ -747,7 +781,7 @@ TArray<FRigElementKey> URigHierarchy::GetParents(FRigElementKey InKey, bool bRec
 
 TArray<int32> URigHierarchy::GetParents(int32 InIndex, bool bRecursive) const
 {
-	const TArray<FRigBaseElement*>& Parents = GetParents(Get(InIndex), bRecursive);
+	const FRigBaseElementParentArray& Parents = GetParents(Get(InIndex), bRecursive);
 	TArray<int32> Indices;
 	for(const FRigBaseElement* Parent : Parents)
 	{
@@ -756,22 +790,23 @@ TArray<int32> URigHierarchy::GetParents(int32 InIndex, bool bRecursive) const
 	return Indices;
 }
 
-TArray<FRigBaseElement*> URigHierarchy::GetParents(const FRigBaseElement* InElement, bool bRecursive) const
+FRigBaseElementParentArray URigHierarchy::GetParents(const FRigBaseElement* InElement, bool bRecursive) const
 {
-	TArray<FRigBaseElement*> Parents;
+	FRigBaseElementParentArray Parents;
 
 	if(const FRigSingleParentElement* SingleParentElement = Cast<FRigSingleParentElement>(InElement))
 	{
 		if(SingleParentElement->ParentElement)
 		{
-			Parents.AddUnique(SingleParentElement->ParentElement);
+			Parents.Add(SingleParentElement->ParentElement);
 		}
 	}
 	else if(const FRigMultiParentElement* MultiParentElement = Cast<FRigMultiParentElement>(InElement))
 	{
+		Parents.Reserve(MultiParentElement->ParentConstraints.Num());
 		for(const FRigElementParentConstraint& ParentConstraint : MultiParentElement->ParentConstraints)
 		{
-			Parents.AddUnique(ParentConstraint.ParentElement);
+			Parents.Add(ParentConstraint.ParentElement);
 		}
 	}
 
@@ -780,7 +815,7 @@ TArray<FRigBaseElement*> URigHierarchy::GetParents(const FRigBaseElement* InElem
 		const int32 CurrentNumberParents = Parents.Num();
 		for(int32 ParentIndex = 0;ParentIndex < CurrentNumberParents; ParentIndex++)
 		{
-			const TArray<FRigBaseElement*> GrandParents = GetParents(Parents[ParentIndex], bRecursive);
+			const FRigBaseElementParentArray GrandParents = GetParents(Parents[ParentIndex], bRecursive);
 			for (FRigBaseElement* GrandParent : GrandParents)
 			{
 				Parents.AddUnique(GrandParent);
@@ -1169,7 +1204,7 @@ void URigHierarchy::Traverse(FRigBaseElement* InElement, bool bTowardsChildren,
 	{
 		if(bTowardsChildren)
 		{
-			const TArray<FRigBaseElement*>& Children = GetChildren(InElement);
+			const FRigBaseElementChildrenArray& Children = GetChildren(InElement);
 			for (FRigBaseElement* Child : Children)
 			{
 				Traverse(Child, true, PerElementFunction);
@@ -1177,7 +1212,7 @@ void URigHierarchy::Traverse(FRigBaseElement* InElement, bool bTowardsChildren,
 		}
 		else
 		{
-			TArray<FRigBaseElement*> Parents = GetParents(InElement);
+			FRigBaseElementParentArray Parents = GetParents(InElement);
 			for (FRigBaseElement* Parent : Parents)
 			{
 				Traverse(Parent, false, PerElementFunction);
@@ -1895,7 +1930,7 @@ FTransform URigHierarchy::GetParentTransform(FRigBaseElement* InElement, const E
 			FTransform OutputTransform = FTransform::Identity;
 
 			const bool bInitial = IsInitial(InTransformType);
-			const TArray<FRigElementParentConstraint>& ParentConstraints = MultiParentElement->ParentConstraints;
+			const FRigElementParentConstraintArray& ParentConstraints = MultiParentElement->ParentConstraints;
 
 			struct Local
 			{
@@ -2821,39 +2856,95 @@ void URigHierarchy::UpdateAllCachedChildren() const
 	}
 }
 
-FRigBaseElement* URigHierarchy::MakeElement(ERigElementType InElementType)
+FRigBaseElement* URigHierarchy::MakeElement(ERigElementType InElementType, int32 InCount, int32* OutStructureSize)
 {
+	check(InCount > 0);
+	
 	FRigBaseElement* Element = nullptr;
 	switch(InElementType)
 	{
 		case ERigElementType::Bone:
 		{
-			Element = new FRigBoneElement();
+			if(OutStructureSize)
+			{
+				*OutStructureSize = sizeof(FRigBoneElement);
+			}
+			FRigBoneElement* Elements = (FRigBoneElement*)FMemory::Malloc(sizeof(FRigBoneElement) * InCount);
+			for(int32 Index=0;Index<InCount;Index++)
+			{
+				new(&Elements[Index]) FRigBoneElement(); 
+			} 
+			Element = Elements;
 			break;
 		}
 		case ERigElementType::Null:
 		{
-			Element = new FRigNullElement();
+			if(OutStructureSize)
+			{
+				*OutStructureSize = sizeof(FRigNullElement);
+			}
+			FRigNullElement* Elements = (FRigNullElement*)FMemory::Malloc(sizeof(FRigNullElement) * InCount);
+			for(int32 Index=0;Index<InCount;Index++)
+			{
+				new(&Elements[Index]) FRigNullElement(); 
+			} 
+			Element = Elements;
 			break;
 		}
 		case ERigElementType::Control:
 		{
-			Element = new FRigControlElement();
+			if(OutStructureSize)
+			{
+				*OutStructureSize = sizeof(FRigControlElement);
+			}
+			FRigControlElement* Elements = (FRigControlElement*)FMemory::Malloc(sizeof(FRigControlElement) * InCount);
+			for(int32 Index=0;Index<InCount;Index++)
+			{
+				new(&Elements[Index]) FRigControlElement(); 
+			} 
+			Element = Elements;
 			break;
 		}
 		case ERigElementType::Curve:
 		{
-			Element = new FRigCurveElement();
+			if(OutStructureSize)
+			{
+				*OutStructureSize = sizeof(FRigCurveElement);
+			}
+			FRigCurveElement* Elements = (FRigCurveElement*)FMemory::Malloc(sizeof(FRigCurveElement) * InCount);
+			for(int32 Index=0;Index<InCount;Index++)
+			{
+				new(&Elements[Index]) FRigCurveElement(); 
+			} 
+			Element = Elements;
 			break;
 		}
 		case ERigElementType::RigidBody:
 		{
-			Element = new FRigRigidBodyElement();
+			if(OutStructureSize)
+			{
+				*OutStructureSize = sizeof(FRigRigidBodyElement);
+			}
+			FRigRigidBodyElement* Elements = (FRigRigidBodyElement*)FMemory::Malloc(sizeof(FRigRigidBodyElement) * InCount);
+			for(int32 Index=0;Index<InCount;Index++)
+			{
+				new(&Elements[Index]) FRigRigidBodyElement(); 
+			} 
+			Element = Elements;
 			break;
 		}
 		case ERigElementType::Socket:
 		{
-			Element = new FRigSocketElement();
+			if(OutStructureSize)
+			{
+				*OutStructureSize = sizeof(FRigSocketElement);
+			}
+			FRigSocketElement* Elements = (FRigSocketElement*)FMemory::Malloc(sizeof(FRigSocketElement) * InCount);
+			for(int32 Index=0;Index<InCount;Index++)
+			{
+				new(&Elements[Index]) FRigSocketElement(); 
+			} 
+			Element = Elements;
 			break;
 		}
 		default:
@@ -2861,9 +2952,91 @@ FRigBaseElement* URigHierarchy::MakeElement(ERigElementType InElementType)
 			ensure(false);
 		}
 	}
+
+	if(Element)
+	{
+		Element->OwnedInstances = InCount;
+	}
 	return Element;
 }
 
+void URigHierarchy::DestroyElement(FRigBaseElement*& InElement)
+{
+	check(InElement != nullptr);
+
+	if(InElement->OwnedInstances == 0)
+	{
+		return;
+	}
+
+	const int32 Count = InElement->OwnedInstances;
+	switch(InElement->GetType())
+	{
+		case ERigElementType::Bone:
+		{
+			FRigBoneElement* Elements = Cast<FRigBoneElement>(InElement);
+			for(int32 Index=0;Index<Count;Index++)
+			{
+				Elements[Index].~FRigBoneElement(); 
+			}
+			break;
+		}
+		case ERigElementType::Null:
+		{
+			FRigNullElement* Elements = Cast<FRigNullElement>(InElement);
+			for(int32 Index=0;Index<Count;Index++)
+			{
+				Elements[Index].~FRigNullElement(); 
+			}
+			break;
+		}
+		case ERigElementType::Control:
+		{
+			FRigControlElement* Elements = Cast<FRigControlElement>(InElement);
+			for(int32 Index=0;Index<Count;Index++)
+			{
+				Elements[Index].~FRigControlElement(); 
+			}
+			break;
+		}
+		case ERigElementType::Curve:
+		{
+			FRigCurveElement* Elements = Cast<FRigCurveElement>(InElement);
+			for(int32 Index=0;Index<Count;Index++)
+			{
+				Elements[Index].~FRigCurveElement(); 
+			}
+			break;
+		}
+		case ERigElementType::RigidBody:
+		{
+			FRigRigidBodyElement* Elements = Cast<FRigRigidBodyElement>(InElement);
+			for(int32 Index=0;Index<Count;Index++)
+			{
+				Elements[Index].~FRigRigidBodyElement(); 
+			}
+			break;
+		}
+		case ERigElementType::Socket:
+		{
+			FRigSocketElement* Elements = Cast<FRigSocketElement>(InElement);
+			for(int32 Index=0;Index<Count;Index++)
+			{
+				Elements[Index].~FRigSocketElement(); 
+			}
+			break;
+		}
+		default:
+		{
+			ensure(false);
+			return;
+		}
+	}
+
+	FMemory::Free(InElement);
+	InElement = nullptr;
+}
+
 #if URIGHIERARCHY_RECURSIVE_DIRTY_PROPAGATION
 void URigHierarchy::PropagateDirtyFlags(FRigTransformElement* InTransformElement, bool bInitial, bool bAffectChildren, bool bComputeOpposed, bool bMarkDirty) const
 #else
diff --git a/Engine/Plugins/Animation/ControlRig/Source/ControlRig/Private/Rigs/RigHierarchyController.cpp b/Engine/Plugins/Animation/ControlRig/Source/ControlRig/Private/Rigs/RigHierarchyController.cpp
--- a/Engine/Plugins/Animation/ControlRig/Source/ControlRig/Private/Rigs/RigHierarchyController.cpp
+++ b/Engine/Plugins/Animation/ControlRig/Source/ControlRig/Private/Rigs/RigHierarchyController.cpp
@@ -1525,7 +1525,7 @@ int32 URigHierarchyController::AddElement(FRigBaseElement* InElementToAdd, FRigB
 
 	InElementToAdd->SubIndex = Hierarchy->Num(InElementToAdd->Key.Type);
 	InElementToAdd->Index = Hierarchy->Elements.Add(InElementToAdd);
-	Hierarchy->ElementsPerType[(int32)InElementToAdd->GetKey().Type].Add(InElementToAdd);
+	Hierarchy->ElementsPerType[URigHierarchy::RigElementTypeToFlatIndex(InElementToAdd->GetKey().Type)].Add(InElementToAdd);
 
 	Hierarchy->IndexLookup.Add(InElementToAdd->Key, InElementToAdd->Index);
 	Hierarchy->TopologyVersion++;
@@ -1603,7 +1603,7 @@ bool URigHierarchyController::RemoveElement(FRigBaseElement* InElement)
 	// if this is a transform element - make sure to allow dependents to store their global transforms
 	if(FRigTransformElement* TransformElement = Cast<FRigTransformElement>(InElement))
 	{
-		TArray<FRigTransformElement::FElementToDirty> PreviousElementsToDirty = TransformElement->ElementsToDirty; 
+		FRigTransformElement::FElementsToDirtyArray PreviousElementsToDirty = TransformElement->ElementsToDirty; 
 		for(const FRigTransformElement::FElementToDirty& ElementToDirty : PreviousElementsToDirty)
 		{
 			if(FRigSingleParentElement* SingleParentElement = Cast<FRigSingleParentElement>(ElementToDirty.Element))
@@ -1630,7 +1630,7 @@ bool URigHierarchyController::RemoveElement(FRigBaseElement* InElement)
 	const int32 NumElementsRemoved = Hierarchy->Elements.Remove(InElement);
 	ensure(NumElementsRemoved == 1);
 
-	const int32 NumTypeElementsRemoved = Hierarchy->ElementsPerType[(int32)InElement->GetKey().Type].Remove(InElement);
+	const int32 NumTypeElementsRemoved = Hierarchy->ElementsPerType[URigHierarchy::RigElementTypeToFlatIndex(InElement->GetKey().Type)].Remove(InElement);
 	ensure(NumTypeElementsRemoved == 1);
 
 	const int32 NumLookupsRemoved = Hierarchy->IndexLookup.Remove(InElement->Key);
@@ -2183,7 +2183,7 @@ bool URigHierarchyController::RemoveAllParents(FRigBaseElement* InChild, bool bM
 	{
 		bool bSuccess = true;
 
-		TArray<FRigElementParentConstraint> ParentConstraints = MultiParentElement->ParentConstraints;
+		FRigElementParentConstraintArray ParentConstraints = MultiParentElement->ParentConstraints;
 		for(const FRigElementParentConstraint& ParentConstraint : ParentConstraints)
 		{
 			if(!RemoveParent(MultiParentElement, ParentConstraint.ParentElement, bMaintainGlobalTransform))
diff --git a/Engine/Plugins/Animation/ControlRig/Source/ControlRig/Private/Rigs/RigHierarchyElements.cpp b/Engine/Plugins/Animation/ControlRig/Source/ControlRig/Private/Rigs/RigHierarchyElements.cpp
--- a/Engine/Plugins/Animation/ControlRig/Source/ControlRig/Private/Rigs/RigHierarchyElements.cpp
+++ b/Engine/Plugins/Animation/ControlRig/Source/ControlRig/Private/Rigs/RigHierarchyElements.cpp
@@ -182,6 +182,8 @@ void FRigTransformElement::CopyFrom(URigHierarchy* InHierarchy, FRigBaseElement*
 	Pose = SourceTransform->Pose;
 
 	ElementsToDirty.Reset();
+	ElementsToDirty.Reserve(SourceTransform->ElementsToDirty.Num());
+	
 	for(int32 ElementToDirtyIndex = 0; ElementToDirtyIndex < SourceTransform->ElementsToDirty.Num(); ElementToDirtyIndex++)
 	{
 		const FElementToDirty& Source = SourceTransform->ElementsToDirty[ElementToDirtyIndex];
@@ -328,7 +330,9 @@ void FRigMultiParentElement::CopyFrom(URigHierarchy* InHierarchy, FRigBaseElemen
 	const FRigMultiParentElement* Source = CastChecked<FRigMultiParentElement>(InOther);
 	Parent = Source->Parent;
 	ParentConstraints.Reset();
+	ParentConstraints.Reserve(Source->ParentConstraints.Num());
 	IndexLookup.Reset();
+	IndexLookup.Reserve(Source->IndexLookup.Num());
 
 	for(int32 ParentIndex = 0; ParentIndex < Source->ParentConstraints.Num(); ParentIndex++)
 	{
diff --git a/Engine/Plugins/Animation/ControlRig/Source/ControlRig/Public/Rigs/RigHierarchy.h b/Engine/Plugins/Animation/ControlRig/Source/ControlRig/Public/Rigs/RigHierarchy.h
--- a/Engine/Plugins/Animation/ControlRig/Source/ControlRig/Public/Rigs/RigHierarchy.h
+++ b/Engine/Plugins/Animation/ControlRig/Source/ControlRig/Public/Rigs/RigHierarchy.h
@@ -1520,7 +1520,7 @@ public:
 	 * @param InElement The element to retrieve the children for
 	 * @return Returns the child elements
 	 */
-	const TArray<FRigBaseElement*>& GetChildren(const FRigBaseElement* InElement) const;
+	const FRigBaseElementChildrenArray& GetChildren(const FRigBaseElement* InElement) const;
 
 	/**
 	 * Returns the child elements of a given element
@@ -1528,7 +1528,7 @@ public:
 	 * @param bRecursive If set to true grand-children will also be returned etc
 	 * @return Returns the child elements
 	 */
-	TArray<FRigBaseElement*> GetChildren(const FRigBaseElement* InElement, bool bRecursive) const;
+	FRigBaseElementChildrenArray GetChildren(const FRigBaseElement* InElement, bool bRecursive) const;
 
 	/**
 	 * Returns the parent elements of a given element key
@@ -1553,7 +1553,7 @@ public:
 	 * @param bRecursive If set to true parents of parents will also be returned
 	 * @return Returns the parent elements
 	 */
-	TArray<FRigBaseElement*> GetParents(const FRigBaseElement* InElement, bool bRecursive = false) const;
+	FRigBaseElementParentArray GetParents(const FRigBaseElement* InElement, bool bRecursive = false) const;
 
 	/**
 	 * Returns the first parent element of a given element key
@@ -2311,7 +2311,12 @@ private:
 	/*
 	 * Helper function to create an element for a given type
 	 */
-	static FRigBaseElement* MakeElement(ERigElementType InElementType); 
+	static FRigBaseElement* MakeElement(ERigElementType InElementType, int32 InCount = 1, int32* OutStructureSize = nullptr); 
+
+	/*
+	* Helper function to create an element for a given type
+	*/
+	static void DestroyElement(FRigBaseElement*& InElement); 
 
 	/**
 	 * Marks all affected elements of a given element as dirty
@@ -2348,7 +2353,7 @@ private:
 	TMap<FRigElementKey, int32> IndexLookup;
 
 	// Static empty element array used for ref returns
-	static const TArray<FRigBaseElement*> EmptyElementArray;
+	static const FRigBaseElementChildrenArray EmptyElementArray;
 
 	///////////////////////////////////////////////
 	/// Undo redo related
@@ -2496,6 +2501,91 @@ protected:
 
 #endif
 
+	FORCEINLINE static int32 RigElementTypeToFlatIndex(ERigElementType InElementType)
+	{
+		switch(InElementType)
+		{
+			case ERigElementType::Bone:
+			{
+				return 0;
+			}
+			case ERigElementType::Null:
+			{
+				return 1;
+			}
+			case ERigElementType::Control:
+			{
+				return 2;
+			}
+			case ERigElementType::Curve:
+			{
+				return 3;
+			}
+			case ERigElementType::RigidBody:
+			{
+				return 4;
+			}
+			case ERigElementType::Socket:
+			{
+				return 5;
+			}
+			case ERigElementType::Last:
+			{
+				return 6;
+			}
+			case ERigElementType::All:
+			default:
+			{
+				checkNoEntry();
+				break;
+			}
+		}
+
+		return INDEX_NONE;
+	}
+
+	FORCEINLINE static ERigElementType FlatIndexToRigElementType(int32 InIndex)
+	{
+		switch(InIndex)
+		{
+			case 0:
+			{
+				return ERigElementType::Bone;
+			}
+			case 1:
+			{
+				return ERigElementType::Null;
+			}
+			case 2:
+			{
+				return ERigElementType::Control;
+			}
+			case 3:
+			{
+				return ERigElementType::Curve;
+			}
+			case 4:
+			{
+				return ERigElementType::RigidBody;
+			}
+			case 5:
+			{
+				return ERigElementType::Socket;
+			}
+			case 6:
+			{
+				return ERigElementType::Last;
+			}
+			default:
+			{
+				checkNoEntry();
+				break;
+			}
+		}
+
+		return ERigElementType::None;
+	}
+
 	friend class URigHierarchyController;
 	friend class UControlRig;
 	friend class FControlRigEditor;
diff --git a/Engine/Plugins/Animation/ControlRig/Source/ControlRig/Public/Rigs/RigHierarchyElements.h b/Engine/Plugins/Animation/ControlRig/Source/ControlRig/Public/Rigs/RigHierarchyElements.h
--- a/Engine/Plugins/Animation/ControlRig/Source/ControlRig/Public/Rigs/RigHierarchyElements.h
+++ b/Engine/Plugins/Animation/ControlRig/Source/ControlRig/Public/Rigs/RigHierarchyElements.h
@@ -396,6 +396,12 @@ struct CONTROLRIG_API FRigCurrentAndInitialTransform
 	FRigLocalAndGlobalTransform Initial;
 };
 
+struct FRigBaseElement;
+//typedef TArray<FRigBaseElement*> FRigBaseElementChildrenArray;
+typedef TArray<FRigBaseElement*, TInlineAllocator<3>> FRigBaseElementChildrenArray;
+//typedef TArray<FRigBaseElement*> FRigBaseElementParentArray;
+typedef TArray<FRigBaseElement*, TInlineAllocator<1>> FRigBaseElementParentArray;
+
 USTRUCT(BlueprintType)
 struct CONTROLRIG_API FRigBaseElement
 {
@@ -409,6 +415,7 @@ public:
 	, SubIndex(INDEX_NONE)
 	, bSelected(false)
 	, TopologyVersion(0)
+	, OwnedInstances(0)
 	{}
 
 	virtual ~FRigBaseElement(){}
@@ -420,7 +427,7 @@ public:
 	};
 
 protected:
-	
+
 	UPROPERTY(BlueprintReadOnly, VisibleAnywhere, Category = RigElement, meta = (AllowPrivateAccess = "true"))
 	FRigElementKey Key;
 
@@ -508,10 +515,13 @@ public:
 protected:
 
 	// helper function to be called as part of URigHierarchy::CopyHierarchy
-	virtual void CopyFrom(URigHierarchy* InHierarchy, FRigBaseElement* InOther, URigHierarchy* InOtherHierarchy) {}
+	virtual void  CopyFrom(URigHierarchy* InHierarchy, FRigBaseElement* InOther, URigHierarchy* InOtherHierarchy) {}
 
 	mutable uint16 TopologyVersion;
-	mutable TArray<FRigBaseElement*> CachedChildren;
+	mutable FRigBaseElementChildrenArray CachedChildren;
+
+	// used for constructing / destructing the memory. typically == 1
+	int32 OwnedInstances;
 
 	friend class URigHierarchy;
 	friend class URigHierarchyController;
@@ -560,8 +570,10 @@ protected:
 		FRigTransformElement* Element;
 		int32 HierarchyDistance;
 	};
-	
-	TArray<FElementToDirty> ElementsToDirty;
+
+	//typedef TArray<FElementToDirty> FElementsToDirtyArray;
+	typedef TArray<FElementToDirty, TInlineAllocator<3>> FElementsToDirtyArray;  
+	FElementsToDirtyArray ElementsToDirty;
 
 	FORCEINLINE static bool IsClassOf(const FRigBaseElement* InElement)
 	{
@@ -714,6 +726,9 @@ public:
 	}
 };
 
+//typedef TArray<FRigElementParentConstraint> FRigElementParentConstraintArray;
+typedef TArray<FRigElementParentConstraint, TInlineAllocator<1>> FRigElementParentConstraintArray;
+
 USTRUCT(BlueprintType)
 struct CONTROLRIG_API FRigMultiParentElement : public FRigTransformElement
 {
@@ -734,9 +749,9 @@ public:
 	UPROPERTY(BlueprintReadOnly, Category = RigElement)
 	FRigCurrentAndInitialTransform Parent;
 
-	TArray<FRigElementParentConstraint> ParentConstraints;
+	FRigElementParentConstraintArray ParentConstraints;
 	TMap<FRigElementKey, int32> IndexLookup;
-	
+
 protected:
 
 	virtual void CopyFrom(URigHierarchy* InHierarchy, FRigBaseElement* InOther, URigHierarchy* InOtherHierarchy) override;
diff --git a/Engine/Plugins/Animation/ControlRig/Source/ControlRigEditor/Private/EditMode/ControlRigEditMode.cpp b/Engine/Plugins/Animation/ControlRig/Source/ControlRigEditor/Private/EditMode/ControlRigEditMode.cpp
--- a/Engine/Plugins/Animation/ControlRig/Source/ControlRigEditor/Private/EditMode/ControlRigEditMode.cpp
+++ b/Engine/Plugins/Animation/ControlRig/Source/ControlRigEditor/Private/EditMode/ControlRigEditMode.cpp
@@ -652,7 +652,7 @@ void FControlRigEditMode::Render(const FSceneView* View, FViewport* Viewport, FP
 			{
 				const FTransform Transform = Hierarchy->GetTransform(TransformElement, ERigTransformType::CurrentGlobal);
 
-				TArray<FRigBaseElement*> Parents = Hierarchy->GetParents(TransformElement);
+				FRigBaseElementParentArray Parents = Hierarchy->GetParents(TransformElement);
 				for(FRigBaseElement* ParentElement : Parents)
 				{
 					if(FRigTransformElement* ParentTransformElement = Cast<FRigTransformElement>(ParentElement))
diff --git a/Engine/Source/Runtime/RigVM/Private/RigVMCore/RigVMMemoryStorage.cpp b/Engine/Source/Runtime/RigVM/Private/RigVMCore/RigVMMemoryStorage.cpp
--- a/Engine/Source/Runtime/RigVM/Private/RigVMCore/RigVMMemoryStorage.cpp
+++ b/Engine/Source/Runtime/RigVM/Private/RigVMCore/RigVMMemoryStorage.cpp
@@ -200,9 +200,17 @@ void URigVMMemoryStorageGeneratorClass::PostLoad()
 	RefreshPropertyPaths();
 }
 
-FString URigVMMemoryStorageGeneratorClass::GetClassName(ERigVMMemoryType InMemoryType)
+const FString& URigVMMemoryStorageGeneratorClass::GetClassName(ERigVMMemoryType InMemoryType)
 {
-	return TEXT("RigVMMemory_") + StaticEnum<ERigVMMemoryType>()->GetDisplayNameTextByValue((int64)InMemoryType).ToString();
+	static TArray<FString> ClassNames;
+	if(ClassNames.IsEmpty())
+	{
+		for(int32 MemoryTypeIndex=0;MemoryTypeIndex<(int32)ERigVMMemoryType::Invalid;MemoryTypeIndex++)
+		{
+			ClassNames.Add(TEXT("RigVMMemory_") + StaticEnum<ERigVMMemoryType>()->GetDisplayNameTextByValue((int64)MemoryTypeIndex).ToString());
+		}
+	}
+	return ClassNames[(int32)InMemoryType];
 }
 
 URigVMMemoryStorageGeneratorClass* URigVMMemoryStorageGeneratorClass::GetStorageClass(UObject* InOuter, ERigVMMemoryType InMemoryType)
@@ -227,7 +235,7 @@ URigVMMemoryStorageGeneratorClass* URigVMMemoryStorageGeneratorClass::GetStorage
 	}
 	while (Outer);
 
-	const FString ClassName = GetClassName(InMemoryType);
+	const FString& ClassName = GetClassName(InMemoryType);
 
 	for(UObject* PotentialClassContainer : PotentialClassContainers)
 	{
@@ -251,7 +259,7 @@ URigVMMemoryStorageGeneratorClass* URigVMMemoryStorageGeneratorClass::CreateStor
 	UPackage* Package = InOuter->GetOutermost();
 	UClass *SuperClass = URigVMMemoryStorage::StaticClass();
 
-	const FString ClassName = GetClassName(InMemoryType);
+	const FString& ClassName = GetClassName(InMemoryType);
 
 	// if there's an old class - remove it from the package and mark it to destroy
 	URigVMMemoryStorageGeneratorClass* OldClass = FindObject<URigVMMemoryStorageGeneratorClass>(Package, *ClassName);
diff --git a/Engine/Source/Runtime/RigVM/Public/RigVMCore/RigVMMemoryStorage.h b/Engine/Source/Runtime/RigVM/Public/RigVMCore/RigVMMemoryStorage.h
--- a/Engine/Source/Runtime/RigVM/Public/RigVMCore/RigVMMemoryStorage.h
+++ b/Engine/Source/Runtime/RigVM/Public/RigVMCore/RigVMMemoryStorage.h
@@ -228,7 +228,7 @@ public:
 	virtual void PostLoad() override;
 
 	// Returns the name of a storage class for a given memory type. The name is unique within the package.
-	static FString GetClassName(ERigVMMemoryType InMemoryType);
+	static const FString& GetClassName(ERigVMMemoryType InMemoryType);
 
 	// Returns an existing class for a memory type within the package (or nullptr) 
 	static URigVMMemoryStorageGeneratorClass* GetStorageClass(UObject* InOuter, ERigVMMemoryType InMemoryType);
