Control Rig: Lazy evaluation for the control rig component

diff --git a/Engine/Plugins/Animation/ControlRig/Source/ControlRig/Private/ControlRigComponent.cpp b/Engine/Plugins/Animation/ControlRig/Source/ControlRig/Private/ControlRigComponent.cpp
--- a/Engine/Plugins/Animation/ControlRig/Source/ControlRig/Private/ControlRigComponent.cpp
+++ b/Engine/Plugins/Animation/ControlRig/Source/ControlRig/Private/ControlRigComponent.cpp
@@ -2,6 +2,7 @@
 
 #include "ControlRigComponent.h"
 #include "Units/Execution/RigUnit_BeginExecution.h"
+#include "ControlRig/Private/Units/Execution/RigUnit_Hierarchy.h"
 
 #include "SkeletalDebugRendering.h"
 #include "Components/InstancedStaticMeshComponent.h"
@@ -56,6 +57,12 @@ UControlRigComponent::UControlRigComponent(const FObjectInitializer& ObjectIniti
 	bShowDebugDrawing = false;
 	bIsInsideInitializeBracket = false;
 	bWantsInitializeComponent = true;
+	
+	bEnableLazyEvaluation = false;
+	LazyEvaluationPositionThreshold = 0.1f;
+	LazyEvaluationRotationThreshold = 0.5f;
+	LazyEvaluationScaleThreshold = 0.01f;
+	bNeedsEvaluation = true;
 }
 
 #if WITH_EDITOR
@@ -367,27 +374,30 @@ void UControlRigComponent::Update(float DeltaTime)
 
 			TransferInputs();
 
-#if WITH_EDITOR
-			if(URigHierarchy* Hierarchy = CR->GetHierarchy())
+			if(bNeedsEvaluation)
 			{
-				if(Hierarchy->IsTracingChanges())
+#if WITH_EDITOR
+				if(URigHierarchy* Hierarchy = CR->GetHierarchy())
 				{
-					Hierarchy->StorePoseForTrace(TEXT("UControlRigComponent::BeforeEvaluate"));
+					if(Hierarchy->IsTracingChanges())
+					{
+						Hierarchy->StorePoseForTrace(TEXT("UControlRigComponent::BeforeEvaluate"));
+					}
 				}
-			}
 #endif
 
-			CR->Evaluate_AnyThread();
+				CR->Evaluate_AnyThread();
 
 #if WITH_EDITOR
-			if(URigHierarchy* Hierarchy = CR->GetHierarchy())
-			{
-				if(Hierarchy->IsTracingChanges())
+				if(URigHierarchy* Hierarchy = CR->GetHierarchy())
 				{
-					Hierarchy->StorePoseForTrace(TEXT("UControlRigComponent::AfterEvaluate"));
+					if(Hierarchy->IsTracingChanges())
+					{
+						Hierarchy->StorePoseForTrace(TEXT("UControlRigComponent::AfterEvaluate"));
+					}
 				}
-			}
 #endif
+			}
 		} 
 
 		if (bShowDebugDrawing)
@@ -1283,8 +1293,12 @@ void UControlRigComponent::ValidateMappingData()
 
 void UControlRigComponent::TransferInputs()
 {
+	bNeedsEvaluation = true;
 	if (ControlRig)
 	{
+		InputElementIndices.Reset();
+		InputTransforms.Reset();
+		
 		for (FControlRigComponentMappedElement& MappedElement : MappedElements)
 		{
 			if (MappedElement.ElementIndex == INDEX_NONE || MappedElement.Direction == EControlRigComponentMapDirection::Output)
@@ -1323,8 +1337,47 @@ void UControlRigComponent::TransferInputs()
 			Transform = Transform * MappedElement.Offset;
 
 			ConvertTransformToRigSpace(Transform, MappedElement.Space);
-			ControlRig->GetHierarchy()->SetGlobalTransform(MappedElement.ElementIndex, Transform);
+
+			InputElementIndices.Add(MappedElement.ElementIndex);
+			InputTransforms.Add(Transform);
 		}
+
+		if(bEnableLazyEvaluation && InputTransforms.Num() > 0)
+		{
+			if(LastInputTransforms.Num() == InputTransforms.Num())
+			{
+				bNeedsEvaluation = false;
+
+				const float PositionU = FMath::Abs(LazyEvaluationPositionThreshold);
+				const float RotationU = FMath::Abs(LazyEvaluationRotationThreshold);
+				const float ScaleU = FMath::Abs(LazyEvaluationScaleThreshold);
+
+				for(int32 Index=0;Index<InputElementIndices.Num();Index++)
+				{
+					if(!FRigUnit_PoseGetDelta::AreTransformsEqual(
+						InputTransforms[Index],
+						LastInputTransforms[Index],
+						PositionU,
+						RotationU,
+						ScaleU))
+					{
+						bNeedsEvaluation = true;
+						break;
+					}
+				}
+			}
+
+			LastInputTransforms = InputTransforms;
+		}
+
+		if(bNeedsEvaluation)
+		{
+			for(int32 Index=0;Index<InputElementIndices.Num();Index++)
+			{
+				ControlRig->GetHierarchy()->SetGlobalTransform(InputElementIndices[Index], InputTransforms[Index]);
+			}
+		}
+
 #if WITH_EDITOR
 		if(URigHierarchy* Hierarchy = ControlRig->GetHierarchy())
 		{
diff --git a/Engine/Plugins/Animation/ControlRig/Source/ControlRig/Private/Units/Execution/RigUnit_Hierarchy.cpp b/Engine/Plugins/Animation/ControlRig/Source/ControlRig/Private/Units/Execution/RigUnit_Hierarchy.cpp
--- a/Engine/Plugins/Animation/ControlRig/Source/ControlRig/Private/Units/Execution/RigUnit_Hierarchy.cpp
+++ b/Engine/Plugins/Animation/ControlRig/Source/ControlRig/Private/Units/Execution/RigUnit_Hierarchy.cpp
@@ -215,75 +215,6 @@ FRigUnit_PoseGetDelta_Execute()
 	const float ScaleU = FMath::Abs(ScaleThreshold);
 	const float CurveU = FMath::Abs(CurveThreshold);
 
-	auto ArePoseElementsEqual = [
-		PositionU,
-		RotationU,
-		ScaleU,
-		CurveU,
-		Space
-		] (
-			const FRigPoseElement& A,
-			const FRigPoseElement& B
-		) -> bool
-	{
-		const FRigElementKey& KeyA = A.Index.GetKey();
-		const FRigElementKey& KeyB = B.Index.GetKey();
-		check(KeyA == KeyB);
-		
-		if(KeyA.Type == ERigElementType::Curve)
-		{
-			if(CurveU > SMALL_NUMBER)
-			{
-				return FMath::Abs(A.CurveValue - B.CurveValue) < CurveU;
-			}
-			return true;
-		}
-		
-		const FTransform& TransformA = (Space == EBoneGetterSetterMode::GlobalSpace) ? A.GlobalTransform : A.LocalTransform;
-		const FTransform& TransformB = (Space == EBoneGetterSetterMode::GlobalSpace) ? B.GlobalTransform : B.LocalTransform;
-		
-		if(PositionU > SMALL_NUMBER)
-		{
-			const FVector PositionA = TransformA.GetLocation();
-			const FVector PositionB = TransformB.GetLocation();
-			const FVector Delta = (PositionA - PositionB).GetAbs();
-			if( (Delta.X >= PositionU) ||
-				(Delta.Y >= PositionU) ||
-				(Delta.Z >= PositionU))
-			{
-				return false;
-			}
-		}
-
-		if(RotationU > SMALL_NUMBER)
-		{
-			const FVector RotationA = TransformA.GetRotation().Rotator().Euler();
-			const FVector RotationB = TransformB.GetRotation().Rotator().Euler();
-			const FVector Delta = (RotationA - RotationB).GetAbs();
-			if( (Delta.X >= RotationU) ||
-				(Delta.Y >= RotationU) ||
-				(Delta.Z >= RotationU))
-			{
-				return false;
-			}
-		}
-
-		if(ScaleU > SMALL_NUMBER)
-		{
-			const FVector ScaleA = TransformA.GetScale3D();
-			const FVector ScaleB = TransformB.GetScale3D();
-			const FVector Delta = (ScaleA - ScaleB).GetAbs();
-			if( (Delta.X >= ScaleU) ||
-				(Delta.Y >= ScaleU) ||
-				(Delta.Z >= ScaleU))
-			{
-				return false;
-			}
-		}
-
-		return true;
-	};
-
 	// if we should compare a sub set of things
 	if(!ItemsToCompare.IsEmpty())
 	{
@@ -314,7 +245,7 @@ FRigUnit_PoseGetDelta_Execute()
 			const FRigPoseElement& A = PoseA[IndexA];
 			const FRigPoseElement& B = PoseB[IndexB];
 
-			if(!ArePoseElementsEqual(A, B))
+			if(!ArePoseElementsEqual(A, B, Space, PositionU, RotationU, ScaleU, CurveU))
 			{
 				PosesAreEqual = false;
 				ItemsWithDelta.Add(Key);
@@ -339,7 +270,7 @@ FRigUnit_PoseGetDelta_Execute()
 				continue;
 			}
 
-			if(!ArePoseElementsEqual(A, B))
+			if(!ArePoseElementsEqual(A, B, Space, PositionU, RotationU, ScaleU, CurveU))
 			{
 				PosesAreEqual = false;
 				ItemsWithDelta.Add(KeyA);
@@ -370,7 +301,7 @@ FRigUnit_PoseGetDelta_Execute()
 
 			const FRigPoseElement& B = PoseB[IndexB];
 			
-			if(!ArePoseElementsEqual(A, B))
+			if(!ArePoseElementsEqual(A, B, Space, PositionU, RotationU, ScaleU, CurveU))
 			{
 				PosesAreEqual = false;
 				ItemsWithDelta.Add(KeyA);
@@ -399,6 +330,78 @@ FRigUnit_PoseGetDelta_Execute()
 	}
 }
 
+bool FRigUnit_PoseGetDelta::ArePoseElementsEqual(const FRigPoseElement& A, const FRigPoseElement& B,
+	EBoneGetterSetterMode Space, float PositionU, float RotationU, float ScaleU, float CurveU)
+{
+	const FRigElementKey& KeyA = A.Index.GetKey();
+	const FRigElementKey& KeyB = B.Index.GetKey();
+	check(KeyA == KeyB);
+		
+	if(KeyA.Type == ERigElementType::Curve)
+	{
+		return AreCurvesEqual(A.CurveValue, B.CurveValue, CurveU);
+	}
+
+	return AreTransformsEqual(
+		(Space == EBoneGetterSetterMode::GlobalSpace) ? A.GlobalTransform : A.LocalTransform,
+		(Space == EBoneGetterSetterMode::GlobalSpace) ? B.GlobalTransform : B.LocalTransform,
+		PositionU, RotationU, ScaleU);
+}
+
+bool FRigUnit_PoseGetDelta::AreTransformsEqual(const FTransform& A, const FTransform& B, float PositionU,
+	float RotationU, float ScaleU)
+{
+	if(PositionU > SMALL_NUMBER)
+	{
+		const FVector PositionA = A.GetLocation();
+		const FVector PositionB = B.GetLocation();
+		const FVector Delta = (PositionA - PositionB).GetAbs();
+		if( (Delta.X >= PositionU) ||
+			(Delta.Y >= PositionU) ||
+			(Delta.Z >= PositionU))
+		{
+			return false;
+		}
+	}
+
+	if(RotationU > SMALL_NUMBER)
+	{
+		const FVector RotationA = A.GetRotation().Rotator().Euler();
+		const FVector RotationB = B.GetRotation().Rotator().Euler();
+		const FVector Delta = (RotationA - RotationB).GetAbs();
+		if( (Delta.X >= RotationU) ||
+			(Delta.Y >= RotationU) ||
+			(Delta.Z >= RotationU))
+		{
+			return false;
+		}
+	}
+
+	if(ScaleU > SMALL_NUMBER)
+	{
+		const FVector ScaleA = A.GetScale3D();
+		const FVector ScaleB = B.GetScale3D();
+		const FVector Delta = (ScaleA - ScaleB).GetAbs();
+		if( (Delta.X >= ScaleU) ||
+			(Delta.Y >= ScaleU) ||
+			(Delta.Z >= ScaleU))
+		{
+			return false;
+		}
+	}
+
+	return true;
+}
+
+bool FRigUnit_PoseGetDelta::AreCurvesEqual(float A, float B, float CurveU)
+{
+	if(CurveU > SMALL_NUMBER)
+	{
+		return FMath::Abs(A - B) < CurveU;
+	}
+	return true;
+}
+
 FRigUnit_PoseGetTransform_Execute()
 {
 	if (Context.State == EControlRigState::Init)
@@ -510,4 +513,4 @@ FRigUnit_PoseLoop_Execute()
 		GlobalTransform = LocalTransform = FTransform::Identity;
 		CurveValue = 0.f;
 	}
-}
\ No newline at end of file
+}
diff --git a/Engine/Plugins/Animation/ControlRig/Source/ControlRig/Private/Units/Execution/RigUnit_Hierarchy.h b/Engine/Plugins/Animation/ControlRig/Source/ControlRig/Private/Units/Execution/RigUnit_Hierarchy.h
--- a/Engine/Plugins/Animation/ControlRig/Source/ControlRig/Private/Units/Execution/RigUnit_Hierarchy.h
+++ b/Engine/Plugins/Animation/ControlRig/Source/ControlRig/Private/Units/Execution/RigUnit_Hierarchy.h
@@ -351,6 +351,27 @@ struct CONTROLRIG_API FRigUnit_PoseGetDelta : public FRigUnit_HierarchyBase
 
 	UPROPERTY(meta = (Output))
 	FRigElementKeyCollection ItemsWithDelta;
+
+	static bool ArePoseElementsEqual(
+		const FRigPoseElement& A,
+		const FRigPoseElement& B,
+		EBoneGetterSetterMode Space,
+		float PositionU,
+		float RotationU,
+		float ScaleU,
+		float CurveU);
+
+	static bool AreTransformsEqual(
+		const FTransform& A,
+		const FTransform& B,
+		float PositionU,
+		float RotationU,
+		float ScaleU);
+
+	static bool AreCurvesEqual(
+		float A,
+		float B,
+		float CurveU);
 };
 
 /**
diff --git a/Engine/Plugins/Animation/ControlRig/Source/ControlRig/Public/ControlRigComponent.h b/Engine/Plugins/Animation/ControlRig/Source/ControlRig/Public/ControlRigComponent.h
--- a/Engine/Plugins/Animation/ControlRig/Source/ControlRig/Public/ControlRigComponent.h
+++ b/Engine/Plugins/Animation/ControlRig/Source/ControlRig/Public/ControlRigComponent.h
@@ -295,6 +295,22 @@ public:
 	UFUNCTION(BlueprintCallable, Category = "ControlRig", meta = (DisplayName = "Set Bone Initial Transforms From Skeletal Mesh"))
 	void SetBoneInitialTransformsFromSkeletalMesh(USkeletalMesh* InSkeletalMesh);
 
+	/** When checked the rig will only run if any of the mapped inputs has changed */
+	UPROPERTY(EditAnywhere, Category = "Lazy Evaluation")
+	bool bEnableLazyEvaluation;
+
+	// The delta threshold for a translation / position difference. 0.0 disables position differences.
+	UPROPERTY(EditAnywhere, Category = "Lazy Evaluation", meta = (DisplayName = "Position Threshold"))
+	float LazyEvaluationPositionThreshold;
+
+	// The delta threshold for a rotation difference (in degrees). 0.0 disables rotation differences.
+	UPROPERTY(EditAnywhere, Category = "Lazy Evaluation", meta = (DisplayName = "Rotation Threshold"))
+	float LazyEvaluationRotationThreshold;
+
+	// The delta threshold for a scale difference. 0.0 disables scale differences.
+	UPROPERTY(EditAnywhere, Category = "Lazy Evaluation", meta = (DisplayName = "Scale Threshold"))
+	float LazyEvaluationScaleThreshold;
+
 	/** When checked the transforms are reset before a tick / update of the rig */
 	UPROPERTY(EditAnywhere, Category = "Animation")
 	bool bResetTransformBeforeTick;
@@ -638,6 +654,11 @@ private:
 
 	FControlRigComponentEvent ControlRigCreatedEvent;
 	bool bIsInsideInitializeBracket;
+	bool bNeedsEvaluation;
+
+	TArray<int32> InputElementIndices;
+	TArray<FTransform> InputTransforms;
+	TArray<FTransform> LastInputTransforms;
 
 	friend class FControlRigSceneProxy;
 };
