Control Rig: Fix crash when using debug in a simulated PIE session

diff --git a/Engine/Plugins/Animation/ControlRig/Source/ControlRig/Private/ControlRig.cpp b/Engine/Plugins/Animation/ControlRig/Source/ControlRig/Private/ControlRig.cpp
--- a/Engine/Plugins/Animation/ControlRig/Source/ControlRig/Private/ControlRig.cpp
+++ b/Engine/Plugins/Animation/ControlRig/Source/ControlRig/Private/ControlRig.cpp
@@ -548,9 +548,13 @@ void UControlRig::Execute(const EControlRigState InState, const FName& InEventNa
 				GetSnapshotVM();
 			}
 		}
-		VM->SetDebugInfo(&DebugInfo);
+
+		if(VM)
+		{
+			VM->SetDebugInfo(&DebugInfo);
+		}
 	}
-	else
+	else if(VM)
 	{
 		VM->SetDebugInfo(nullptr);
 	}
@@ -2755,7 +2759,7 @@ UControlRig::FPoseScope::FPoseScope(UControlRig* InControlRig)
 {
 	check(InControlRig);
 
-	CachedPose = InControlRig->GetHierarchy()->GetPose(false, ERigElementType::Control, FRigElementKeyCollection());
+	CachedPose = InControlRig->GetHierarchy()->GetPose();
 
 	ControlRig = InControlRig;
 }
