Control Rig: Allow double to float links and vice versa when building context menu content

diff --git a/Engine/Plugins/Animation/ControlRig/Source/ControlRigDeveloper/Private/Graph/ControlRigGraphSchema.cpp b/Engine/Plugins/Animation/ControlRig/Source/ControlRigDeveloper/Private/Graph/ControlRigGraphSchema.cpp
--- a/Engine/Plugins/Animation/ControlRig/Source/ControlRigDeveloper/Private/Graph/ControlRigGraphSchema.cpp
+++ b/Engine/Plugins/Animation/ControlRig/Source/ControlRigDeveloper/Private/Graph/ControlRigGraphSchema.cpp
@@ -1209,6 +1209,22 @@ bool UControlRigGraphSchema::ArePinsCompatible(const UEdGraphPin* PinA, const UE
 		return true;
 	}
 
+	// if large world coordinate support is enabled we should allow connections
+	// between float and double
+#if !UE_LARGE_WORLD_COORDINATES_DISABLED
+	if(PinA->PinType.ContainerType == EPinContainerType::None &&
+		PinB->PinType.ContainerType == EPinContainerType::None)
+	{
+		if((PinA->PinType.PinCategory == UEdGraphSchema_K2::PC_Float &&
+			PinB->PinType.PinCategory == UEdGraphSchema_K2::PC_Double) ||
+			(PinA->PinType.PinCategory == UEdGraphSchema_K2::PC_Double &&
+			PinB->PinType.PinCategory == UEdGraphSchema_K2::PC_Float))
+		{
+			return true;
+		}
+	}
+#endif
+
 	struct Local
 	{
 		static FString GetCPPTypeFromPinType(const FEdGraphPinType& InPinType)
