Control Rig: Avoid crash when reimporting mesh

diff --git a/Engine/Plugins/Animation/ControlRig/Source/ControlRigEditor/Private/Editor/SRigCurveContainer.cpp b/Engine/Plugins/Animation/ControlRig/Source/ControlRigEditor/Private/Editor/SRigCurveContainer.cpp
--- a/Engine/Plugins/Animation/ControlRig/Source/ControlRigEditor/Private/Editor/SRigCurveContainer.cpp
+++ b/Engine/Plugins/Animation/ControlRig/Source/ControlRigEditor/Private/Editor/SRigCurveContainer.cpp
@@ -146,7 +146,11 @@ void SRigCurveContainer::Construct(const FArguments& InArgs, TSharedRef<FControl
 	bIsChangingRigHierarchy = false;
 
 #if WITH_EDITOR
-	GEditor->OnEditorClose().AddRaw(this, &SRigCurveContainer::OnEditorClose);
+	GEditor->OnEditorClose().AddLambda([&]()
+	{
+		const FControlRigEditor* Editor = ControlRigEditor.IsValid() ? ControlRigEditor.Pin().Get() : nullptr;
+		OnEditorClose(Editor, ControlRigBlueprint.Get());
+	});
 #endif
 
 	ControlRigBlueprint->Hierarchy->OnModified().AddRaw(this, &SRigCurveContainer::OnHierarchyModified);
@@ -204,7 +208,7 @@ void SRigCurveContainer::Construct(const FArguments& InArgs, TSharedRef<FControl
 
 	if(ControlRigEditor.IsValid())
 	{
-		ControlRigEditor.Pin()->OnControlRigEditorClosed().AddSP(this, &SRigCurveContainer::OnControlRigEditorClose);
+		ControlRigEditor.Pin()->OnControlRigEditorClosed().AddSP(this, &SRigCurveContainer::OnEditorClose);
 	}
 }
 
@@ -213,7 +217,8 @@ SRigCurveContainer::~SRigCurveContainer()
 	// Make sure we don't get called on editor exit if we've already been shut down.
 	GEditor->OnEditorClose().RemoveAll(this);
 
-	OnEditorClose();
+	const FControlRigEditor* Editor = ControlRigEditor.IsValid() ? ControlRigEditor.Pin().Get() : nullptr;
+	OnEditorClose(Editor, ControlRigBlueprint.Get());
 }
 
 FReply SRigCurveContainer::OnKeyDown(const FGeometry& MyGeometry, const FKeyEvent& InKeyEvent)
@@ -311,27 +316,18 @@ TSharedPtr<SWidget> SRigCurveContainer::OnGetContextMenuContent() const
 	return MenuBuilder.MakeWidget();
 }
 
-void SRigCurveContainer::OnEditorClose()
+void SRigCurveContainer::OnEditorClose(const FControlRigEditor* InEditor, UControlRigBlueprint* InBlueprint)
 {
-	if (ControlRigEditor.IsValid())
+	if (InBlueprint)
 	{
-		ControlRigBlueprint = ControlRigEditor.Pin()->GetControlRigBlueprint();
-		if (ControlRigBlueprint.IsValid())
-		{
-			ControlRigBlueprint->Hierarchy->OnModified().RemoveAll(this);
-			ControlRigBlueprint->OnRefreshEditor().RemoveAll(this);
-		}
+		InBlueprint->Hierarchy->OnModified().RemoveAll(this);
+		InBlueprint->OnRefreshEditor().RemoveAll(this);
 	}
 
 	ControlRigBlueprint.Reset();
 	ControlRigEditor.Reset();
 }
 
-void SRigCurveContainer::OnControlRigEditorClose(const FControlRigEditor* InEditor, UControlRigBlueprint* InBlueprint)
-{
-	OnEditorClose();
-}
-
 void SRigCurveContainer::OnRenameClicked()
 {
 	TArray< FDisplayedRigCurveInfoPtr > SelectedItems = RigCurveListView->GetSelectedItems();
diff --git a/Engine/Plugins/Animation/ControlRig/Source/ControlRigEditor/Private/Editor/SRigCurveContainer.h b/Engine/Plugins/Animation/ControlRig/Source/ControlRigEditor/Private/Editor/SRigCurveContainer.h
--- a/Engine/Plugins/Animation/ControlRig/Source/ControlRigEditor/Private/Editor/SRigCurveContainer.h
+++ b/Engine/Plugins/Animation/ControlRig/Source/ControlRigEditor/Private/Editor/SRigCurveContainer.h
@@ -206,8 +206,7 @@ private:
 	/** Handler for context menus */
 	TSharedPtr<SWidget> OnGetContextMenuContent() const;
 
-	void OnEditorClose();
-	void OnControlRigEditorClose(const FControlRigEditor* InEditor, UControlRigBlueprint* InBlueprint);
+	void OnEditorClose(const FControlRigEditor* InEditor, UControlRigBlueprint* InBlueprint);
 
 	/**
 	* Clears and rebuilds the table, according to an optional search string
diff --git a/Engine/Plugins/Animation/ControlRig/Source/ControlRigEditor/Private/Editor/SRigHierarchy.cpp b/Engine/Plugins/Animation/ControlRig/Source/ControlRigEditor/Private/Editor/SRigHierarchy.cpp
--- a/Engine/Plugins/Animation/ControlRig/Source/ControlRigEditor/Private/Editor/SRigHierarchy.cpp
+++ b/Engine/Plugins/Animation/ControlRig/Source/ControlRigEditor/Private/Editor/SRigHierarchy.cpp
@@ -288,7 +288,8 @@ void SRigHierarchyItem::OnNameCommitted(const FText& InText, ETextCommit::Type I
 
 SRigHierarchy::~SRigHierarchy()
 {
-	OnEditorClose();
+	const FControlRigEditor* Editor = ControlRigEditor.IsValid() ? ControlRigEditor.Pin().Get() : nullptr;
+	OnEditorClose(Editor, ControlRigBlueprint.Get());
 }
 
 void SRigHierarchy::Construct(const FArguments& InArgs, TSharedRef<FControlRigEditor> InControlRigEditor)
@@ -454,34 +455,30 @@ void SRigHierarchy::Construct(const FArguments& InArgs, TSharedRef<FControlRigEd
 		});
 		ControlRigEditor.Pin()->OnGetViewportContextMenu().BindSP(this, &SRigHierarchy::GetOrCreateContextMenu);
 		ControlRigEditor.Pin()->OnViewportContextMenuCommands().BindSP(this, &SRigHierarchy::GetContextMenuCommands);
-		ControlRigEditor.Pin()->OnControlRigEditorClosed().AddSP(this, &SRigHierarchy::OnControlRigEditorClose);
+		ControlRigEditor.Pin()->OnControlRigEditorClosed().AddSP(this, &SRigHierarchy::OnEditorClose);
 	}
 }
 
-void SRigHierarchy::OnEditorClose()
+void SRigHierarchy::OnEditorClose(const FControlRigEditor* InEditor, UControlRigBlueprint* InBlueprint)
 {
-	if (ControlRigEditor.IsValid())
+	if (InEditor)
 	{
-		ControlRigEditor.Pin()->GetKeyDownDelegate().Unbind();
-		ControlRigEditor.Pin()->OnGetViewportContextMenu().Unbind();
-		ControlRigEditor.Pin()->OnViewportContextMenuCommands().Unbind();
+		FControlRigEditor* Editor = (FControlRigEditor*)InEditor;  
+		Editor->GetKeyDownDelegate().Unbind();
+		Editor->OnGetViewportContextMenu().Unbind();
+		Editor->OnViewportContextMenuCommands().Unbind();
 	}
 
-	if (ControlRigBlueprint.IsValid())
+	if (InBlueprint)
 	{
-		ControlRigBlueprint->Hierarchy->OnModified().RemoveAll(this);
-		ControlRigBlueprint->OnRefreshEditor().RemoveAll(this);
+		InBlueprint->Hierarchy->OnModified().RemoveAll(this);
+		InBlueprint->OnRefreshEditor().RemoveAll(this);
 	}
 	
 	ControlRigEditor.Reset();
 	ControlRigBlueprint.Reset();
 }
 
-void SRigHierarchy::OnControlRigEditorClose(const FControlRigEditor* InEditor, UControlRigBlueprint* InBlueprint)
-{
-	OnEditorClose();
-}
-
 void SRigHierarchy::BindCommands()
 {
 	// create new command
diff --git a/Engine/Plugins/Animation/ControlRig/Source/ControlRigEditor/Private/Editor/SRigHierarchy.h b/Engine/Plugins/Animation/ControlRig/Source/ControlRigEditor/Private/Editor/SRigHierarchy.h
--- a/Engine/Plugins/Animation/ControlRig/Source/ControlRigEditor/Private/Editor/SRigHierarchy.h
+++ b/Engine/Plugins/Animation/ControlRig/Source/ControlRigEditor/Private/Editor/SRigHierarchy.h
@@ -179,8 +179,7 @@ public:
 
 private:
 
-	void OnEditorClose();
-	void OnControlRigEditorClose(const FControlRigEditor* InEditor, UControlRigBlueprint* InBlueprint);
+	void OnEditorClose(const FControlRigEditor* InEditor, UControlRigBlueprint* InBlueprint);
 
 	/** Bind commands that this widget handles */
 	void BindCommands();
