Control Rig Profiler: Add total micro seconds counter

diff --git a/Engine/Plugins/Animation/ControlRig/Source/ControlRigDeveloper/Public/ControlRigBlueprint.h b/Engine/Plugins/Animation/ControlRig/Source/ControlRigDeveloper/Public/ControlRigBlueprint.h
--- a/Engine/Plugins/Animation/ControlRig/Source/ControlRigDeveloper/Public/ControlRigBlueprint.h
+++ b/Engine/Plugins/Animation/ControlRig/Source/ControlRigDeveloper/Public/ControlRigBlueprint.h
@@ -114,6 +114,7 @@ struct CONTROLRIGDEVELOPER_API FRigGraphDisplaySettings
 		, NodeRunLimit(64)
 		, MinMicroSeconds(0.0)
 		, MaxMicroSeconds(1.0)
+		, TotalMicroSeconds(0.0)
 		, bAutoDetermineRange(true)
 		, LastMinMicroSeconds(0.0)
 		, LastMaxMicroSeconds(1.0)
@@ -152,6 +153,10 @@ struct CONTROLRIGDEVELOPER_API FRigGraphDisplaySettings
 	UPROPERTY(EditAnywhere, Category = "Graph Display Settings", transient, meta = (EditCondition = "!bAutoDetermineRange"))
 	double MaxMicroSeconds;
 
+	// The total duration of the last execution of the rig
+	UPROPERTY(VisibleAnywhere, Category = "Graph Display Settings", transient)
+	double TotalMicroSeconds;
+
 	UPROPERTY(EditAnywhere, Category = "Graph Display Settings")
 	bool bAutoDetermineRange;
 
diff --git a/Engine/Plugins/Animation/ControlRig/Source/ControlRigEditor/Private/Editor/ControlRigEditor.cpp b/Engine/Plugins/Animation/ControlRig/Source/ControlRigEditor/Private/Editor/ControlRigEditor.cpp
--- a/Engine/Plugins/Animation/ControlRig/Source/ControlRigEditor/Private/Editor/ControlRigEditor.cpp
+++ b/Engine/Plugins/Animation/ControlRig/Source/ControlRigEditor/Private/Editor/ControlRigEditor.cpp
@@ -2703,15 +2703,15 @@ void FControlRigEditor::HandleControlRigExecutedEvent(UControlRig* InControlRig,
 {
 	if (UControlRigBlueprint* ControlRigBP = GetControlRigBlueprint())
 	{
-		URigHierarchy* Hierarchy = ControlRigBP->Hierarchy; 
+		UControlRig* DebuggedControlRig = Cast<UControlRig>(ControlRigBP->GetObjectBeingDebugged());
+		if (DebuggedControlRig == nullptr)
+		{
+			DebuggedControlRig = ControlRig;
+		}
 
+		URigHierarchy* Hierarchy = ControlRigBP->Hierarchy; 
 		if(!bSetupModeEnabled)
 		{
-			UControlRig* DebuggedControlRig = Cast<UControlRig>(ControlRigBP->GetObjectBeingDebugged());
-			if (DebuggedControlRig == nullptr)
-			{
-				DebuggedControlRig = ControlRig;
-			}
 			if(DebuggedControlRig)
 			{
 				Hierarchy = DebuggedControlRig->GetHierarchy();
@@ -2761,7 +2761,7 @@ void FControlRigEditor::HandleControlRigExecutedEvent(UControlRig* InControlRig,
 
 		if(ControlRigBP->RigGraphDisplaySettings.NodeRunLimit > 1)
 		{
-			if(UControlRig* DebuggedControlRig = Cast<UControlRig>(ControlRigBP->GetObjectBeingDebugged()))
+			if(DebuggedControlRig)
 			{
 				if(URigVM* VM = DebuggedControlRig->GetVM())
 				{
@@ -2830,6 +2830,11 @@ void FControlRigEditor::HandleControlRigExecutedEvent(UControlRig* InControlRig,
 
 		if(ControlRigBP->VMRuntimeSettings.bEnableProfiling)
 		{
+			if(DebuggedControlRig)
+			{
+				ControlRigBP->RigGraphDisplaySettings.TotalMicroSeconds = DebuggedControlRig->GetVM()->GetContext().LastExecutionMicroSeconds;
+			}
+
 			if(ControlRigBP->RigGraphDisplaySettings.bAutoDetermineRange)
 			{
 				if(ControlRigBP->RigGraphDisplaySettings.LastMaxMicroSeconds < 0.0)
diff --git a/Engine/Source/Runtime/RigVM/Private/RigVMCore/RigVM.cpp b/Engine/Source/Runtime/RigVM/Private/RigVMCore/RigVM.cpp
--- a/Engine/Source/Runtime/RigVM/Private/RigVMCore/RigVM.cpp
+++ b/Engine/Source/Runtime/RigVM/Private/RigVMCore/RigVM.cpp
@@ -1902,6 +1902,7 @@ bool URigVM::Execute(TArrayView<URigVMMemoryStorage*> Memory, TArrayView<void*>
 
 #if WITH_EDITOR
 	uint64 StartCycles = 0;
+	uint64 OverallCycles = 0;
 	if(Context.RuntimeSettings.bEnableProfiling)
 	{
 		StartCycles = FPlatformTime::Cycles64();
@@ -2382,6 +2383,9 @@ bool URigVM::Execute(TArrayView<URigVMMemoryStorage*> Memory, TArrayView<void*>
 			}
 			case ERigVMOpCode::Exit:
 			{
+#if WITH_EDITOR
+				Context.LastExecutionMicroSeconds = OverallCycles * FPlatformTime::GetSecondsPerCycle() * 1000.0 * 1000.0;
+#endif
 				ExecutionReachedExit().Broadcast(InEntryName);
 #if WITH_EDITOR					
 				if (HaltedAtBreakpoint != nullptr)
@@ -2976,6 +2980,7 @@ bool URigVM::Execute(TArrayView<URigVMMemoryStorage*> Memory, TArrayView<void*>
 			}
 
 			StartCycles = EndCycles;
+			OverallCycles += Cycles;
 		}
 #endif
 	}
diff --git a/Engine/Source/Runtime/RigVM/Public/RigVMCore/RigVM.h b/Engine/Source/Runtime/RigVM/Public/RigVMCore/RigVM.h
--- a/Engine/Source/Runtime/RigVM/Public/RigVMCore/RigVM.h
+++ b/Engine/Source/Runtime/RigVM/Public/RigVMCore/RigVM.h
@@ -745,6 +745,7 @@ public:
 #endif
 
 	uint32 GetNumExecutions() const { return NumExecutions; }
+	const FRigVMExecuteContext& GetContext() const { return Context; }
 
 private:
 
diff --git a/Engine/Source/Runtime/RigVM/Public/RigVMCore/RigVMExecuteContext.h b/Engine/Source/Runtime/RigVM/Public/RigVMCore/RigVMExecuteContext.h
--- a/Engine/Source/Runtime/RigVM/Public/RigVMCore/RigVMExecuteContext.h
+++ b/Engine/Source/Runtime/RigVM/Public/RigVMCore/RigVMExecuteContext.h
@@ -192,7 +192,8 @@ struct RIGVM_API FRigVMExecuteContext
 		: EventName(NAME_None)
 		, FunctionName(NAME_None)
 		, InstructionIndex(0)
-		, RuntimeSettings() 
+		, RuntimeSettings()
+		, LastExecutionMicroSeconds() 
 	{
 		Reset();
 	}
@@ -346,4 +347,5 @@ struct RIGVM_API FRigVMExecuteContext
 	TArray<FRigVMExternalVariable> ExternalVariables;
 	TArray<FRigVMSlice> Slices;
 	TArray<uint16> SliceOffsets;
+	double LastExecutionMicroSeconds;
 };
