Control Rig: Space switching high level API

diff --git a/Engine/Plugins/Animation/ControlRig/Source/ControlRig/Private/Drawing/ControlRigDrawInterface.cpp b/Engine/Plugins/Animation/ControlRig/Source/ControlRig/Private/Drawing/ControlRigDrawInterface.cpp
--- a/Engine/Plugins/Animation/ControlRig/Source/ControlRig/Private/Drawing/ControlRigDrawInterface.cpp
+++ b/Engine/Plugins/Animation/ControlRig/Source/ControlRig/Private/Drawing/ControlRigDrawInterface.cpp
@@ -347,10 +347,20 @@ void FControlRigDrawInterface::DrawHierarchy(const FTransform& WorldOffset, URig
 				InstructionZ.Positions.Add(PZ);
 
 				FRigBaseElementParentArray Parents = Hierarchy->GetParents(Child);
-				for (FRigBaseElement* Parent : Parents)
+				TArray<FRigElementWeight> Weights = Hierarchy->GetParentWeightArray(Child);
+				
+				for (int32 ParentIndex = 0; ParentIndex < Parents.Num(); ParentIndex++)
 				{
-					if(FRigTransformElement* ParentTransformElement = Cast<FRigTransformElement>(Parent))
+					if(FRigTransformElement* ParentTransformElement = Cast<FRigTransformElement>(Parents[ParentIndex]))
 					{
+						if(Weights.IsValidIndex(ParentIndex))
+						{
+							if(Weights[ParentIndex].IsAlmostZero())
+							{
+								continue;
+							}
+						}
+						
 						FTransform ParentTransform = FTransform::Identity;
 						GetTransformLambda(ParentTransformElement, ParentTransform);
 
diff --git a/Engine/Plugins/Animation/ControlRig/Source/ControlRig/Private/Rigs/RigHierarchy.cpp b/Engine/Plugins/Animation/ControlRig/Source/ControlRig/Private/Rigs/RigHierarchy.cpp
--- a/Engine/Plugins/Animation/ControlRig/Source/ControlRig/Private/Rigs/RigHierarchy.cpp
+++ b/Engine/Plugins/Animation/ControlRig/Source/ControlRig/Private/Rigs/RigHierarchy.cpp
@@ -1152,6 +1152,98 @@ bool URigHierarchy::SetParentWeightArray(FRigBaseElement* InChild, const TArray<
 	return false;
 }
 
+bool URigHierarchy::SwitchToParent(FRigElementKey InChild, FRigElementKey InParent, bool bInitial, bool bAffectChildren)
+{
+	return SwitchToParent(Find(InChild), Find(InParent), bInitial, bAffectChildren);
+}
+
+bool URigHierarchy::SwitchToParent(FRigBaseElement* InChild, FRigBaseElement* InParent, bool bInitial,
+	bool bAffectChildren)
+{
+	if(const FRigMultiParentElement* MultiParentElement = Cast<FRigMultiParentElement>(InChild))
+	{
+		int32 ParentIndex = INDEX_NONE;
+		if(const int32* ParentIndexPtr = MultiParentElement->IndexLookup.Find(InParent->GetKey()))
+		{
+			ParentIndex = *ParentIndexPtr;
+		}
+		else
+		{
+			if(URigHierarchyController* Controller = GetController(true))
+			{
+				if(Controller->AddParent(InChild, InParent, 0.f, true, false))
+				{
+					ParentIndex = MultiParentElement->IndexLookup.FindChecked(InParent->GetKey());
+				}
+			}
+		}
+
+		return SwitchToParent(InChild, ParentIndex, bInitial, bAffectChildren);
+	}
+	return false;
+}
+
+bool URigHierarchy::SwitchToParent(FRigBaseElement* InChild, int32 InParentIndex, bool bInitial, bool bAffectChildren)
+{
+	TArray<FRigElementWeight> Weights = GetParentWeightArray(InChild, bInitial);
+	if(Weights.IsValidIndex(InParentIndex))
+	{
+		FMemory::Memzero(Weights.GetData(), Weights.GetAllocatedSize());
+		Weights[InParentIndex] = 1.f;
+		return SetParentWeightArray(InChild, Weights, bInitial, bAffectChildren);
+	}
+	return false;
+}
+
+bool URigHierarchy::SwitchToDefaultParent(FRigElementKey InChild, bool bInitial, bool bAffectChildren)
+{
+	return SwitchToDefaultParent(Find(InChild), bInitial, bAffectChildren);
+}
+
+bool URigHierarchy::SwitchToDefaultParent(FRigBaseElement* InChild, bool bInitial, bool bAffectChildren)
+{
+	// we assume that the first stored parent is the default parent
+	return SwitchToParent(InChild, 0, bInitial, bAffectChildren);
+}
+
+bool URigHierarchy::SwitchToWorldSpace(FRigElementKey InChild, bool bInitial, bool bAffectChildren)
+{
+	return SwitchToWorldSpace(Find(InChild), bInitial, bAffectChildren);
+}
+
+bool URigHierarchy::SwitchToWorldSpace(FRigBaseElement* InChild, bool bInitial, bool bAffectChildren)
+{
+	FRigElementKey WorldSocket = GetOrAddWorldSpaceSocket();
+	if(FRigBaseElement* Parent = Find(WorldSocket))
+	{
+		return SwitchToParent(InChild, Parent, bInitial, bAffectChildren);
+	}
+	return false;
+}
+
+FRigElementKey URigHierarchy::GetOrAddWorldSpaceSocket()
+{
+	static const FName WorldSpaceSocketName = TEXT("WorldSpace");
+	const FRigElementKey WorldSpaceSocketKey(WorldSpaceSocketName, ERigElementType::Socket);
+
+	FRigBaseElement* Parent = Find(WorldSpaceSocketKey);
+	if(Parent)
+	{
+		return Parent->GetKey();
+	}
+
+	if(URigHierarchyController* Controller = GetController(true))
+	{
+		return Controller->AddSocket(
+			WorldSpaceSocketName,
+			FRigElementKey(),
+			FRigSocketGetWorldTransformDelegate::CreateUObject(this, &URigHierarchy::GetWorldTransformForSocket),
+			false);
+	}
+
+	return FRigElementKey();
+}
+
 TArray<FRigElementKey> URigHierarchy::GetAllKeys(bool bTraverse, ERigElementType InElementType) const
 {
 	TArray<FRigElementKey> Keys;
@@ -3371,3 +3463,11 @@ void URigHierarchy::ComputeAllTransforms()
 	}
 }
 
+FTransform URigHierarchy::GetWorldTransformForSocket(const FRigUnitContext* InContext, const FRigElementKey& InKey, bool bInitial)
+{
+	if(const USceneComponent* OuterSceneComponent = GetTypedOuter<USceneComponent>())
+	{
+		return OuterSceneComponent->GetComponentToWorld();
+	}
+	return FTransform::Identity;
+}
diff --git a/Engine/Plugins/Animation/ControlRig/Source/ControlRig/Private/Rigs/RigHierarchyController.cpp b/Engine/Plugins/Animation/ControlRig/Source/ControlRig/Private/Rigs/RigHierarchyController.cpp
--- a/Engine/Plugins/Animation/ControlRig/Source/ControlRig/Private/Rigs/RigHierarchyController.cpp
+++ b/Engine/Plugins/Animation/ControlRig/Source/ControlRig/Private/Rigs/RigHierarchyController.cpp
@@ -1675,6 +1675,14 @@ bool URigHierarchyController::RemoveElement(FRigBaseElement* InElement)
 		}
 	}
 
+	for(FRigBaseElement* Element : Hierarchy->Elements)
+	{
+		if(FRigControlElement* ControlElement = Cast<FRigControlElement>(Element))
+		{
+			ControlElement->Settings.SpaceFavorites.Remove(InElement->GetKey());
+		}
+	}
+	
 	Hierarchy->TopologyVersion++;
 
 	Notify(ERigHierarchyNotification::ElementRemoved, InElement);
@@ -1769,6 +1777,17 @@ bool URigHierarchyController::RenameElement(FRigBaseElement* InElement, const FN
 				MultiParentElement->IndexLookup.Add(NewKey, ExistingIndex);
 			}
 		}
+
+		if(FRigControlElement* ControlElement = Cast<FRigControlElement>(Element))
+		{
+			for(FRigElementKey& Favorite : ControlElement->Settings.SpaceFavorites)
+			{
+				if(Favorite == OldKey)
+				{
+					Favorite.Name = NewKey.Name;
+				}
+			}
+		}
 	}
 	
 	Hierarchy->PreviousNameMap.FindOrAdd(NewKey) = OldKey;
@@ -1860,32 +1879,35 @@ bool URigHierarchyController::AddParent(FRigBaseElement* InChild, FRigBaseElemen
 		RemoveAllParents(InChild, bMaintainGlobalTransform);		
 	}
 
-	if(FRigTransformElement* TransformElement = Cast<FRigTransformElement>(InChild))
+	if(InWeight > SMALL_NUMBER || bRemoveAllParents)
 	{
-		if(bMaintainGlobalTransform)
+		if(FRigTransformElement* TransformElement = Cast<FRigTransformElement>(InChild))
 		{
-			Hierarchy->GetTransform(TransformElement, ERigTransformType::CurrentGlobal);
-			Hierarchy->GetTransform(TransformElement, ERigTransformType::InitialGlobal);
-			TransformElement->Pose.MarkDirty(ERigTransformType::CurrentLocal);
-			TransformElement->Pose.MarkDirty(ERigTransformType::InitialLocal);
+			if(bMaintainGlobalTransform)
+			{
+				Hierarchy->GetTransform(TransformElement, ERigTransformType::CurrentGlobal);
+				Hierarchy->GetTransform(TransformElement, ERigTransformType::InitialGlobal);
+				TransformElement->Pose.MarkDirty(ERigTransformType::CurrentLocal);
+				TransformElement->Pose.MarkDirty(ERigTransformType::InitialLocal);
+			}
+			else
+			{
+				Hierarchy->GetTransform(TransformElement, ERigTransformType::CurrentLocal);
+				Hierarchy->GetTransform(TransformElement, ERigTransformType::InitialLocal);
+				TransformElement->Pose.MarkDirty(ERigTransformType::CurrentGlobal);
+				TransformElement->Pose.MarkDirty(ERigTransformType::InitialGlobal);
+			}
 		}
-		else
+
+		if(FRigControlElement* ControlElement = Cast<FRigControlElement>(InChild))
 		{
-			Hierarchy->GetTransform(TransformElement, ERigTransformType::CurrentLocal);
-			Hierarchy->GetTransform(TransformElement, ERigTransformType::InitialLocal);
-			TransformElement->Pose.MarkDirty(ERigTransformType::CurrentGlobal);
-			TransformElement->Pose.MarkDirty(ERigTransformType::InitialGlobal);
+			Hierarchy->GetControlOffsetTransform(ControlElement, ERigTransformType::CurrentLocal);
+			Hierarchy->GetControlOffsetTransform(ControlElement, ERigTransformType::InitialLocal);
+			Hierarchy->GetControlGizmoTransform(ControlElement, ERigTransformType::CurrentLocal);
+			Hierarchy->GetControlGizmoTransform(ControlElement, ERigTransformType::InitialLocal);
 		}
 	}
 
-	if(FRigControlElement* ControlElement = Cast<FRigControlElement>(InChild))
-	{
-		Hierarchy->GetControlOffsetTransform(ControlElement, ERigTransformType::CurrentLocal);
-		Hierarchy->GetControlOffsetTransform(ControlElement, ERigTransformType::InitialLocal);
-		Hierarchy->GetControlGizmoTransform(ControlElement, ERigTransformType::CurrentLocal);
-		Hierarchy->GetControlGizmoTransform(ControlElement, ERigTransformType::InitialLocal);
-	}
-
 	FRigElementParentConstraint Constraint;
 	Constraint.ParentElement = Cast<FRigTransformElement>(InParent);
 	if(Constraint.ParentElement == nullptr)
diff --git a/Engine/Plugins/Animation/ControlRig/Source/ControlRig/Private/Rigs/RigHierarchyElements.cpp b/Engine/Plugins/Animation/ControlRig/Source/ControlRig/Private/Rigs/RigHierarchyElements.cpp
--- a/Engine/Plugins/Animation/ControlRig/Source/ControlRig/Private/Rigs/RigHierarchyElements.cpp
+++ b/Engine/Plugins/Animation/ControlRig/Source/ControlRig/Private/Rigs/RigHierarchyElements.cpp
@@ -421,11 +421,14 @@ FRigControlSettings::FRigControlSettings()
 , GizmoColor(FLinearColor::Red)
 , bIsTransientControl(false)
 , ControlEnum(nullptr)
+, SpaceFavorites()
 {
 }
 
 void FRigControlSettings::Save(FArchive& Ar)
 {
+	Ar.UsingCustomVersion(FControlRigObjectVersion::GUID);
+
 	static const UEnum* ControlTypeEnum = StaticEnum<ERigControlType>();
 	static const UEnum* ControlAxisEnum = StaticEnum<ERigControlAxis>();
 
@@ -456,10 +459,13 @@ void FRigControlSettings::Save(FArchive& Ar)
 	Ar << GizmoColor;
 	Ar << bIsTransientControl;
 	Ar << ControlEnumPathName;
+	Ar << SpaceFavorites;
 }
 
 void FRigControlSettings::Load(FArchive& Ar)
 {
+	Ar.UsingCustomVersion(FControlRigObjectVersion::GUID);
+
 	static const UEnum* ControlTypeEnum = StaticEnum<ERigControlType>();
 	static const UEnum* ControlAxisEnum = StaticEnum<ERigControlAxis>();
 
@@ -495,6 +501,15 @@ void FRigControlSettings::Load(FArchive& Ar)
 	{
 		ControlEnum = FindObject<UEnum>(ANY_PACKAGE, *ControlEnumPathName);
 	}
+
+	if (Ar.CustomVer(FControlRigObjectVersion::GUID) >= FControlRigObjectVersion::RigHierarchyControlSpaceFavorites)
+	{
+		Ar << SpaceFavorites;
+	}
+	else
+	{
+		SpaceFavorites.Reset();
+	}
 }
 
 ////////////////////////////////////////////////////////////////////////////////
@@ -563,6 +578,10 @@ bool FRigControlSettings::operator==(const FRigControlSettings& InOther) const
 	{
 		return false;
 	}
+	if(SpaceFavorites != InOther.SpaceFavorites)
+	{
+		return false;
+	}
 
 	const FTransform MinimumTransform = MinimumValue.GetAsTransform(ControlType, PrimaryAxis);
 	const FTransform OtherMinimumTransform = InOther.MinimumValue.GetAsTransform(ControlType, PrimaryAxis);
diff --git a/Engine/Plugins/Animation/ControlRig/Source/ControlRig/Public/ControlRigObjectVersion.h b/Engine/Plugins/Animation/ControlRig/Source/ControlRig/Public/ControlRigObjectVersion.h
--- a/Engine/Plugins/Animation/ControlRig/Source/ControlRig/Public/ControlRigObjectVersion.h
+++ b/Engine/Plugins/Animation/ControlRig/Source/ControlRig/Public/ControlRigObjectVersion.h
@@ -35,7 +35,10 @@ struct CONTROLRIG_API FControlRigObjectVersion
 
 		// RigHierarchy to support multi component parent constraints
 		RigHierarchyMultiParentConstraints,
-		
+
+		// RigHierarchy now supports space favorites per control
+		RigHierarchyControlSpaceFavorites,
+
 		// -----<new versions can be added above this line>-------------------------------------------------
 		VersionPlusOne,
 		LatestVersion = VersionPlusOne - 1
diff --git a/Engine/Plugins/Animation/ControlRig/Source/ControlRig/Public/Rigs/RigHierarchy.h b/Engine/Plugins/Animation/ControlRig/Source/ControlRig/Public/Rigs/RigHierarchy.h
--- a/Engine/Plugins/Animation/ControlRig/Source/ControlRig/Public/Rigs/RigHierarchy.h
+++ b/Engine/Plugins/Animation/ControlRig/Source/ControlRig/Public/Rigs/RigHierarchy.h
@@ -1699,6 +1699,90 @@ public:
 	 */
 	bool SetParentWeightArray(FRigBaseElement* InChild,  const TArray<FRigElementWeight>& InWeights, bool bInitial = false, bool bAffectChildren = true);
 
+	/**
+		* Switches a multi parent element to a single parent.
+		* This sets the new parent's weight to 1.0 and disables
+		* weights for all other potential parents.
+		* @param InChild The key of the multi parented element
+		* @param InParent The key of the parent to look up the weight for
+		* @param bInitial If true the initial weights will be used
+		* @param bAffectChildren If set to false children will not move (maintain global).
+		* @return Returns true if changing the weight was successful
+		*/
+	UFUNCTION(BlueprintCallable, Category = URigHierarchy)
+	bool SwitchToParent(FRigElementKey InChild, FRigElementKey InParent, bool bInitial = false, bool bAffectChildren = true);
+
+	/**
+	 * Switches a multi parent element to a single parent.
+	 * This sets the new parent's weight to 1.0 and disables
+	 * weights for all other potential parents.
+	 * @param InChild The multi parented element
+	 * @param InParent The parent to look up the weight for
+	 * @param bInitial If true the initial weights will be used
+	 * @param bAffectChildren If set to false children will not move (maintain global).
+	 * @return Returns true if changing the weight was successful
+	 */
+	bool SwitchToParent(FRigBaseElement* InChild, FRigBaseElement* InParent, bool bInitial = false, bool bAffectChildren = true);
+
+	/**
+	 * Switches a multi parent element to a single parent.
+	 * This sets the new parent's weight to 1.0 and disables
+	 * weights for all other potential parents.
+	 * @param InChild The multi parented element
+	 * @param InParentIndex The index of the parent inside of the multi parent element
+	 * @param bInitial If true the initial weights will be used
+	 * @param bAffectChildren If set to false children will not move (maintain global).
+	 * @return Returns true if changing the weight was successful
+	 */
+	bool SwitchToParent(FRigBaseElement* InChild, int32 InParentIndex, bool bInitial = false, bool bAffectChildren = true);
+
+	/**
+	 * Switches a multi parent element to its first parent
+	 * @param InChild The key of the multi parented element
+	 * @param bInitial If true the initial weights will be used
+	 * @param bAffectChildren If set to false children will not move (maintain global).
+	 * @return Returns true if changing the weight was successful
+	 */
+	UFUNCTION(BlueprintCallable, Category = URigHierarchy)
+	bool SwitchToDefaultParent(FRigElementKey InChild, bool bInitial = false, bool bAffectChildren = true);
+
+	/**
+	 * Switches a multi parent element to its first parent
+	 * This sets the new parent's weight to 1.0 and disables
+	 * weights for all other potential parents.
+	 * @param InChild The multi parented element
+	 * @param bInitial If true the initial weights will be used
+	 * @param bAffectChildren If set to false children will not move (maintain global).
+	 * @return Returns true if changing the weight was successful
+	 */
+	bool SwitchToDefaultParent(FRigBaseElement* InChild, bool bInitial = false, bool bAffectChildren = true);
+
+	/**
+	 * Switches a multi parent element to world space.
+	 * This injects a world space socket.
+	 * @param InChild The key of the multi parented element
+	 * @param bInitial If true the initial weights will be used
+	 * @param bAffectChildren If set to false children will not move (maintain global).
+	 * @return Returns true if changing the weight was successful
+	 */
+	UFUNCTION(BlueprintCallable, Category = URigHierarchy)
+	bool SwitchToWorldSpace(FRigElementKey InChild, bool bInitial = false, bool bAffectChildren = true);
+
+	/**
+	* Switches a multi parent element to world space.
+	* This injects a world space socket.
+	* @param InChild The multi parented element
+	* @param bInitial If true the initial weights will be used
+	* @param bAffectChildren If set to false children will not move (maintain global).
+	* @return Returns true if changing the weight was successful
+	*/
+	bool SwitchToWorldSpace(FRigBaseElement* InChild, bool bInitial = false, bool bAffectChildren = true);
+
+	/**
+	 * Adds the world space socket or returns it
+	 */
+	FRigElementKey GetOrAddWorldSpaceSocket();
+
 	/**
 	 * Returns true if an element is parented to another element
 	 * @param InChild The key of the child element to check for a parent
@@ -2586,6 +2670,8 @@ protected:
 		return ERigElementType::None;
 	}
 
+	FTransform GetWorldTransformForSocket(const FRigUnitContext* InContext, const FRigElementKey& InKey, bool bInitial);
+	
 	friend class URigHierarchyController;
 	friend class UControlRig;
 	friend class FControlRigEditor;
diff --git a/Engine/Plugins/Animation/ControlRig/Source/ControlRig/Public/Rigs/RigHierarchyController.h b/Engine/Plugins/Animation/ControlRig/Source/ControlRig/Public/Rigs/RigHierarchyController.h
--- a/Engine/Plugins/Animation/ControlRig/Source/ControlRig/Public/Rigs/RigHierarchyController.h
+++ b/Engine/Plugins/Animation/ControlRig/Source/ControlRig/Public/Rigs/RigHierarchyController.h
@@ -606,4 +606,5 @@ private:
 	bool bSuspendNotifications;
 
 	friend class UControlRig;
+	friend class URigHierarchy;
 };
diff --git a/Engine/Plugins/Animation/ControlRig/Source/ControlRig/Public/Rigs/RigHierarchyElements.h b/Engine/Plugins/Animation/ControlRig/Source/ControlRig/Public/Rigs/RigHierarchyElements.h
--- a/Engine/Plugins/Animation/ControlRig/Source/ControlRig/Public/Rigs/RigHierarchyElements.h
+++ b/Engine/Plugins/Animation/ControlRig/Source/ControlRig/Public/Rigs/RigHierarchyElements.h
@@ -908,6 +908,12 @@ struct CONTROLRIG_API FRigControlSettings
 	UPROPERTY(BlueprintReadOnly, EditAnywhere, Category = Control)
 	TObjectPtr<UEnum> ControlEnum;
 
+	/** The favorite spaces for this control.
+	 * This will be used as the default content for the space picker
+	 */
+	UPROPERTY(BlueprintReadOnly, EditAnywhere, Category = Animation)
+	TArray<FRigElementKey> SpaceFavorites;
+
 	/** Applies the limits expressed by these settings to a value */
 	FORCEINLINE void ApplyLimits(FRigControlValue& InOutValue) const
 	{
diff --git a/Engine/Plugins/Animation/ControlRig/Source/ControlRigEditor/Private/ControlRigElementDetails.cpp b/Engine/Plugins/Animation/ControlRig/Source/ControlRigEditor/Private/ControlRigElementDetails.cpp
--- a/Engine/Plugins/Animation/ControlRig/Source/ControlRigEditor/Private/ControlRigElementDetails.cpp
+++ b/Engine/Plugins/Animation/ControlRig/Source/ControlRigEditor/Private/ControlRigElementDetails.cpp
@@ -1907,6 +1907,9 @@ void FRigControlElementDetails::CustomizeChildren(TSharedRef<class IPropertyHand
 		));
 		
 	}
+
+	IDetailGroup* AnimationGroup = &StructBuilder.AddGroup(TEXT("Animation"), LOCTEXT("Animation", "Animation"));
+	AnimationGroup->AddPropertyRow(SettingsHandle->GetChildHandle(TEXT("SpaceFavorites")).ToSharedRef());
 }
 
 bool FRigControlElementDetails::IsGizmoEnabled() const
