Control Rig: Recover fix for regression in PoseScope

diff --git a/Engine/Plugins/Animation/ControlRig/Source/ControlRig/Private/ControlRig.cpp b/Engine/Plugins/Animation/ControlRig/Source/ControlRig/Private/ControlRig.cpp
--- a/Engine/Plugins/Animation/ControlRig/Source/ControlRig/Private/ControlRig.cpp
+++ b/Engine/Plugins/Animation/ControlRig/Source/ControlRig/Private/ControlRig.cpp
@@ -2773,7 +2773,7 @@ UControlRig::FPoseScope::FPoseScope(UControlRig* InControlRig)
 {
 	check(InControlRig);
 
-	CachedPose = InControlRig->GetHierarchy()->GetPose(false, ERigElementType::Control, FRigElementKeyCollection());
+	CachedPose = InControlRig->GetHierarchy()->GetPose();
 
 	ControlRig = InControlRig;
 }
