Fix for crashes/ensurse when binding a sequencer/control rig track to a skeletal mesh component w/ no skeletal mesh

diff --git a/Engine/Plugins/Animation/ControlRig/Source/ControlRig/Private/Sequencer/MovieSceneControlRigParameterTemplate.cpp b/Engine/Plugins/Animation/ControlRig/Source/ControlRig/Private/Sequencer/MovieSceneControlRigParameterTemplate.cpp
--- a/Engine/Plugins/Animation/ControlRig/Source/ControlRig/Private/Sequencer/MovieSceneControlRigParameterTemplate.cpp
+++ b/Engine/Plugins/Animation/ControlRig/Source/ControlRig/Private/Sequencer/MovieSceneControlRigParameterTemplate.cpp
@@ -554,17 +554,20 @@ void FControlRigBindingHelper::BindToSequencerInstance(UControlRig* ControlRig)
 	check(ControlRig);
 	if (USkeletalMeshComponent* SkeletalMeshComponent = Cast<USkeletalMeshComponent>(ControlRig->GetObjectBinding()->GetBoundObject()))
 	{
-		bool bWasCreated = false;
-		if (UControlRigLayerInstance* AnimInstance = FAnimCustomInstanceHelper::BindToSkeletalMeshComponent<UControlRigLayerInstance>(SkeletalMeshComponent, bWasCreated))
+		if(SkeletalMeshComponent->SkeletalMesh)
 		{
-			if (bWasCreated || !AnimInstance->HasControlRigTrack(ControlRig->GetUniqueID()))
+			bool bWasCreated = false;
+			if (UControlRigLayerInstance* AnimInstance = FAnimCustomInstanceHelper::BindToSkeletalMeshComponent<UControlRigLayerInstance>(SkeletalMeshComponent, bWasCreated))
 			{
-				AnimInstance->RecalcRequiredBones();
-				AnimInstance->AddControlRigTrack(ControlRig->GetUniqueID(), ControlRig);
-				ControlRig->Initialize();
-				ControlRig->RequestInit();
-				ControlRig->SetBoneInitialTransformsFromSkeletalMeshComponent(SkeletalMeshComponent, true);
-				ControlRig->Evaluate_AnyThread();
+				if (bWasCreated || !AnimInstance->HasControlRigTrack(ControlRig->GetUniqueID()))
+				{
+					AnimInstance->RecalcRequiredBones();
+					AnimInstance->AddControlRigTrack(ControlRig->GetUniqueID(), ControlRig);
+					ControlRig->Initialize();
+					ControlRig->RequestInit();
+					ControlRig->SetBoneInitialTransformsFromSkeletalMeshComponent(SkeletalMeshComponent, true);
+					ControlRig->Evaluate_AnyThread();
+				}
 			}
 		}
 	}
diff --git a/Engine/Source/Runtime/MovieSceneTracks/Private/Evaluation/MovieSceneSkeletalAnimationTemplate.cpp b/Engine/Source/Runtime/MovieSceneTracks/Private/Evaluation/MovieSceneSkeletalAnimationTemplate.cpp
--- a/Engine/Source/Runtime/MovieSceneTracks/Private/Evaluation/MovieSceneSkeletalAnimationTemplate.cpp
+++ b/Engine/Source/Runtime/MovieSceneTracks/Private/Evaluation/MovieSceneSkeletalAnimationTemplate.cpp
@@ -277,7 +277,7 @@ namespace MovieScene
 			ensureMsgf(InObject, TEXT("Attempting to evaluate an Animation track with a null object."));
 
 			USkeletalMeshComponent* SkeletalMeshComponent = SkeletalMeshComponentFromObject(InObject);
-			if (!SkeletalMeshComponent)
+			if (!SkeletalMeshComponent || !SkeletalMeshComponent->SkeletalMesh)
 			{
 				return;
 			}
