Control Rig: Fix python logging of change type of variable to handle arrays

diff --git a/Engine/Plugins/Animation/ControlRig/Source/ControlRigDeveloper/Private/ControlRigBlueprint.cpp b/Engine/Plugins/Animation/ControlRig/Source/ControlRigDeveloper/Private/ControlRigBlueprint.cpp
--- a/Engine/Plugins/Animation/ControlRig/Source/ControlRigDeveloper/Private/ControlRigBlueprint.cpp
+++ b/Engine/Plugins/Animation/ControlRig/Source/ControlRigDeveloper/Private/ControlRigBlueprint.cpp
@@ -3723,7 +3723,7 @@ void UControlRigBlueprint::OnVariableTypeChanged(const FName& InVarName, FEdGrap
 	RigVMPythonUtils::Print(GetFName().ToString(),
 		FString::Printf(TEXT("blueprint.change_member_variable_type('%s', '%s')"),
 		*InVarName.ToString(),
-		*PinType));
+		NewVariable.bIsArray ? *RigVMUtilities::ArrayTypeFromBaseType(PinType) : *PinType));
 
 	BroadcastExternalVariablesChangedEvent();
 }
diff --git a/Engine/Source/Runtime/RigVM/Public/RigVMCore/RigVMUtilities.h b/Engine/Source/Runtime/RigVM/Public/RigVMCore/RigVMUtilities.h
--- a/Engine/Source/Runtime/RigVM/Public/RigVMCore/RigVMUtilities.h
+++ b/Engine/Source/Runtime/RigVM/Public/RigVMCore/RigVMUtilities.h
@@ -11,4 +11,9 @@ namespace RigVMUtilities
 	{
 		return InCPPType.StartsWith(TEXT("TArray<"));
 	}
+
+	FORCEINLINE FString ArrayTypeFromBaseType(const FString& InCPPType)
+	{
+		return "TArray<" + InCPPType + ">";
+	}
 }
\ No newline at end of file
