Control Rig: Fixed Control Pin Value stops right after moving the transient control. What was happening was that after transient control moves, interaction bracket was closed -> VM recompile -> ControlRig::InstantiateVMFromCDO -> RequestSetup -> In next execute, before setup, current CR's hierarchy is copied from CDO and thus we lost the old transient controls. The approach here is to preserve transient control if there is one.

diff --git a/Engine/Plugins/Animation/ControlRig/Source/ControlRig/Private/ControlRig.cpp b/Engine/Plugins/Animation/ControlRig/Source/ControlRig/Private/ControlRig.cpp
--- a/Engine/Plugins/Animation/ControlRig/Source/ControlRig/Private/ControlRig.cpp
+++ b/Engine/Plugins/Animation/ControlRig/Source/ControlRig/Private/ControlRig.cpp
@@ -702,6 +702,9 @@ void UControlRig::Execute(const EControlRigState InState, const FName& InEventNa
 				{
 					if(CDO->GetHierarchy()->GetTopologyVersion()!= GetHierarchy()->GetTopologyVersion())
 					{
+#if WITH_EDITOR
+						FTransientControlScope TransientControlScope(GetHierarchy());
+#endif
 						GetHierarchy()->CopyHierarchy(CDO->GetHierarchy());
 					}
 				}
@@ -2784,6 +2787,46 @@ UControlRig::FPoseScope::~FPoseScope()
 
 	ControlRig->GetHierarchy()->SetPose(CachedPose);
 }
+
+#if WITH_EDITOR
+
+UControlRig::FTransientControlScope::FTransientControlScope(TObjectPtr<URigHierarchy> InHierarchy)
+	:Hierarchy(InHierarchy)
+{
+	for (FRigControlElement* Control : Hierarchy->GetTransientControls())
+	{
+		FTransientControlInfo Info;
+		Info.Name = Control->GetName();
+		Info.Parent = Hierarchy->GetFirstParent(Control->GetKey());
+		Info.Settings = Control->Settings;
+		// preserve whatever value that was produced by this transient control at the moment
+		Info.Value = Hierarchy->GetControlValue(Control->GetKey(),ERigControlValueType::Current);
+		Info.OffsetTransform = Hierarchy->GetControlOffsetTransform(Control, ERigTransformType::CurrentLocal);
+		Info.GizmoTransform = Hierarchy->GetControlGizmoTransform(Control, ERigTransformType::CurrentLocal);
+				
+		SavedTransientControls.Add(Info);
+	}
+}
+
+UControlRig::FTransientControlScope::~FTransientControlScope()
+{
+	URigHierarchyController* Controller = Hierarchy->GetController();
+	for (const FTransientControlInfo& Info : SavedTransientControls)
+	{
+		Controller->AddControl(
+            Info.Name,
+            Info.Parent,
+            Info.Settings,
+            Info.Value,
+            Info.OffsetTransform,
+            Info.GizmoTransform,
+            false,
+            false
+        );
+	}
+}
+
+#endif
  
 #undef LOCTEXT_NAMESPACE
 
diff --git a/Engine/Plugins/Animation/ControlRig/Source/ControlRig/Public/ControlRig.h b/Engine/Plugins/Animation/ControlRig/Source/ControlRig/Public/ControlRig.h
--- a/Engine/Plugins/Animation/ControlRig/Source/ControlRig/Public/ControlRig.h
+++ b/Engine/Plugins/Animation/ControlRig/Source/ControlRig/Public/ControlRig.h
@@ -735,6 +735,35 @@ private:
 		FRigPose CachedPose;
 	};
 
+#if WITH_EDITOR
+
+	// Class used to temporarily cache current transient controls
+	// restore them after a CopyHierarchy call
+	class FTransientControlScope
+	{
+	public:
+		FTransientControlScope(TObjectPtr<URigHierarchy> InHierarchy);
+		~FTransientControlScope();
+	
+	private:
+		// used to match URigHierarchyController::AddControl(...)
+		struct FTransientControlInfo
+		{
+			FName Name;
+			// transient control should only have 1 parent, with weight = 1.0
+			FRigElementKey Parent;
+			FRigControlSettings Settings;
+			FRigControlValue Value;
+			FTransform OffsetTransform;
+			FTransform GizmoTransform;
+		};
+		
+		TArray<FTransientControlInfo> SavedTransientControls;
+		TObjectPtr<URigHierarchy> Hierarchy;
+	};
+
+#endif
+
 	friend class FControlRigBlueprintCompilerContext;
 	friend struct FRigHierarchyRef;
 	friend class FControlRigEditor;
