Control Rig: Refactor Rig hierarchy widget, split it up and reuse for edit mode

diff --git a/Engine/Plugins/Animation/ControlRig/Source/ControlRig/Private/Rigs/RigHierarchy.cpp b/Engine/Plugins/Animation/ControlRig/Source/ControlRig/Private/Rigs/RigHierarchy.cpp
--- a/Engine/Plugins/Animation/ControlRig/Source/ControlRig/Private/Rigs/RigHierarchy.cpp
+++ b/Engine/Plugins/Animation/ControlRig/Source/ControlRig/Private/Rigs/RigHierarchy.cpp
@@ -1038,6 +1038,8 @@ bool URigHierarchy::SetParentWeight(FRigBaseElement* InChild, int32 InParentInde
 				}	
 			}
 #endif
+
+			Notify(ERigHierarchyNotification::ParentWeightsChanged, MultiParentElement);
 			return true;
 		}
 	}
@@ -1146,6 +1148,9 @@ bool URigHierarchy::SetParentWeightArray(FRigBaseElement* InChild, const TArray<
 				}	
 			}
 #endif
+
+			Notify(ERigHierarchyNotification::ParentWeightsChanged, MultiParentElement);
+
 			return true;
 		}
 	}
@@ -1223,8 +1228,7 @@ bool URigHierarchy::SwitchToWorldSpace(FRigBaseElement* InChild, bool bInitial,
 
 FRigElementKey URigHierarchy::GetOrAddWorldSpaceSocket()
 {
-	static const FName WorldSpaceSocketName = TEXT("WorldSpace");
-	const FRigElementKey WorldSpaceSocketKey(WorldSpaceSocketName, ERigElementType::Socket);
+	const FRigElementKey WorldSpaceSocketKey = GetWorldSpaceSocketKey();
 
 	FRigBaseElement* Parent = Find(WorldSpaceSocketKey);
 	if(Parent)
@@ -1235,7 +1239,7 @@ FRigElementKey URigHierarchy::GetOrAddWorldSpaceSocket()
 	if(URigHierarchyController* Controller = GetController(true))
 	{
 		return Controller->AddSocket(
-			WorldSpaceSocketName,
+			WorldSpaceSocketKey.Name,
 			FRigElementKey(),
 			FRigSocketGetWorldTransformDelegate::CreateUObject(this, &URigHierarchy::GetWorldTransformForSocket),
 			false);
@@ -1244,6 +1248,12 @@ FRigElementKey URigHierarchy::GetOrAddWorldSpaceSocket()
 	return FRigElementKey();
 }
 
+FRigElementKey URigHierarchy::GetWorldSpaceSocketKey() const
+{
+	static const FName WorldSpaceSocketName = TEXT("WorldSpace");
+	return FRigElementKey(WorldSpaceSocketName, ERigElementType::Socket); 
+}
+
 TArray<FRigElementKey> URigHierarchy::GetAllKeys(bool bTraverse, ERigElementType InElementType) const
 {
 	TArray<FRigElementKey> Keys;
diff --git a/Engine/Plugins/Animation/ControlRig/Source/ControlRig/Public/Rigs/RigHierarchy.h b/Engine/Plugins/Animation/ControlRig/Source/ControlRig/Public/Rigs/RigHierarchy.h
--- a/Engine/Plugins/Animation/ControlRig/Source/ControlRig/Public/Rigs/RigHierarchy.h
+++ b/Engine/Plugins/Animation/ControlRig/Source/ControlRig/Public/Rigs/RigHierarchy.h
@@ -1782,6 +1782,7 @@ public:
 	 * Adds the world space socket or returns it
 	 */
 	FRigElementKey GetOrAddWorldSpaceSocket();
+	FRigElementKey GetWorldSpaceSocketKey() const;
 
 	/**
 	 * Returns true if an element is parented to another element
diff --git a/Engine/Plugins/Animation/ControlRig/Source/ControlRig/Public/Rigs/RigHierarchyDefines.h b/Engine/Plugins/Animation/ControlRig/Source/ControlRig/Public/Rigs/RigHierarchyDefines.h
--- a/Engine/Plugins/Animation/ControlRig/Source/ControlRig/Public/Rigs/RigHierarchyDefines.h
+++ b/Engine/Plugins/Animation/ControlRig/Source/ControlRig/Public/Rigs/RigHierarchyDefines.h
@@ -48,6 +48,7 @@ enum class ERigHierarchyNotification : uint8
 	ControlSettingChanged,
 	ControlVisibilityChanged,
 	ControlGizmoTransformChanged,
+	ParentWeightsChanged,
 
 	/** MAX - invalid */
 	Max UMETA(Hidden),
diff --git a/Engine/Plugins/Animation/ControlRig/Source/ControlRigEditor/Private/ControlRigHierarchyCommands.cpp b/Engine/Plugins/Animation/ControlRig/Source/ControlRigEditor/Private/ControlRigHierarchyCommands.cpp
--- a/Engine/Plugins/Animation/ControlRig/Source/ControlRigEditor/Private/ControlRigHierarchyCommands.cpp
+++ b/Engine/Plugins/Animation/ControlRig/Source/ControlRigEditor/Private/ControlRigHierarchyCommands.cpp
@@ -37,6 +37,7 @@ void FControlRigHierarchyCommands::RegisterCommands()
 	UI_COMMAND(ShowSockets, "Show Sockets", "Whether to show or hide sockets", EUserInterfaceActionType::ToggleButton, FInputChord());
 	UI_COMMAND(ShowDynamicHierarchy, "Dynamic Hierarchy", "Whether to show the static or dynamic hierarchy", EUserInterfaceActionType::ToggleButton, FInputChord());
 	UI_COMMAND(ToggleGizmoTransformEdit, "Toggle Gizmo Transform Edit", "Toggle Editing Selected Control's Gizmo Transform", EUserInterfaceActionType::ToggleButton, FInputChord(EKeys::Period, EModifierKey::Control)); 
+	UI_COMMAND(TestSpaceSwitching, "Test Space Switching", "Try out space switching on the control", EUserInterfaceActionType::Button, FInputChord(EKeys::H, EModifierKey::Control)); 
 }
 
 #undef LOCTEXT_NAMESPACE
diff --git a/Engine/Plugins/Animation/ControlRig/Source/ControlRigEditor/Private/ControlRigHierarchyCommands.h b/Engine/Plugins/Animation/ControlRig/Source/ControlRigEditor/Private/ControlRigHierarchyCommands.h
--- a/Engine/Plugins/Animation/ControlRig/Source/ControlRigEditor/Private/ControlRigHierarchyCommands.h
+++ b/Engine/Plugins/Animation/ControlRig/Source/ControlRigEditor/Private/ControlRigHierarchyCommands.h
@@ -110,7 +110,10 @@ public:
 
 	/** Toggle Gizmo Transform Edit*/
 	TSharedPtr< FUICommandInfo > ToggleGizmoTransformEdit;
-	
+
+	/** Test out space switching as it would look like for animator */
+	TSharedPtr< FUICommandInfo > TestSpaceSwitching;
+
 	/**
 	 * Initialize commands
 	 */
diff --git a/Engine/Plugins/Animation/ControlRig/Source/ControlRigEditor/Private/EditMode/SControlHierarchy.cpp b/Engine/Plugins/Animation/ControlRig/Source/ControlRigEditor/Private/EditMode/SControlHierarchy.cpp
deleted file mode 100644
diff --git a/Engine/Plugins/Animation/ControlRig/Source/ControlRigEditor/Private/EditMode/SControlHierarchy.h b/Engine/Plugins/Animation/ControlRig/Source/ControlRigEditor/Private/EditMode/SControlHierarchy.h
deleted file mode 100644
diff --git a/Engine/Plugins/Animation/ControlRig/Source/ControlRigEditor/Private/EditMode/SControlRigEditModeTools.cpp b/Engine/Plugins/Animation/ControlRig/Source/ControlRigEditor/Private/EditMode/SControlRigEditModeTools.cpp
--- a/Engine/Plugins/Animation/ControlRig/Source/ControlRigEditor/Private/EditMode/SControlRigEditModeTools.cpp
+++ b/Engine/Plugins/Animation/ControlRig/Source/ControlRigEditor/Private/EditMode/SControlRigEditModeTools.cpp
@@ -16,7 +16,6 @@
 #include "Widgets/Layout/SScrollBox.h"
 #include "MovieSceneSequence.h"
 #include "MovieScene.h"
-#include "SControlHierarchy.h"
 #include "SControlPicker.h"
 #include "IDetailCustomization.h"
 #include "DetailLayoutBuilder.h"
@@ -193,6 +192,11 @@ protected:
 
 void SControlRigEditModeTools::SetControlRig(UControlRig* ControlRig)
 {
+	if(ViewportRig.IsValid())
+	{
+		ViewportRig->ControlSelected().RemoveAll(this);
+	}
+
 	SequencerRig = ControlRig;
 	ViewportRig = ControlRig;
 	if (SequencerRig.IsValid())
@@ -207,11 +211,31 @@ void SControlRigEditModeTools::SetControlRig(UControlRig* ControlRig)
 	Objects.Add(SequencerRig);
 	RigOptionsDetailsView->SetObjects(Objects);
 
-	ControlHierarchy->SetControlRig(ViewportRig.Get());
+	HierarchyTreeView->RefreshTreeView(true);
+
+	if (ViewportRig.IsValid())
+	{
+		ViewportRig->ControlSelected().AddRaw(this, &SControlRigEditModeTools::OnRigElementSelected);
+	}
+}
+
+const URigHierarchy* SControlRigEditModeTools::GetHierarchy() const
+{
+	if(ViewportRig.IsValid())
+	{
+		return ViewportRig->GetHierarchy();
+	}
+	if(SequencerRig.IsValid())
+	{
+		return ViewportRig->GetHierarchy();
+	}
+	return nullptr;
 }
 
 void SControlRigEditModeTools::Construct(const FArguments& InArgs, FControlRigEditMode& InEditMode,UWorld* InWorld)
 {
+	bIsChangingRigHierarchy = false;
+	
 	// initialize settings view
 	FDetailsViewArgs DetailsViewArgs;
 	{
@@ -241,6 +265,19 @@ void SControlRigEditModeTools::Construct(const FArguments& InArgs, FControlRigEd
 	RigOptionsDetailsView->SetKeyframeHandler(SharedThis(this));
 	RigOptionsDetailsView->OnFinishedChangingProperties().AddSP(this, &SControlRigEditModeTools::OnRigOptionFinishedChange);
 
+	DisplaySettings.bShowBones = false;
+	DisplaySettings.bShowControls = true;
+	DisplaySettings.bShowNulls = false;
+	DisplaySettings.bShowSockets = false;
+	DisplaySettings.bShowRigidBodies = false;
+	DisplaySettings.bHideParentsOnFilter = true;
+	DisplaySettings.bFlattenHierarchyOnFilter = true;
+
+	FRigTreeDelegates RigTreeDelegates;
+	RigTreeDelegates.OnGetHierarchy = FOnGetRigTreeHierarchy::CreateSP(this, &SControlRigEditModeTools::GetHierarchy);
+	RigTreeDelegates.OnGetDisplaySettings = FOnGetRigTreeDisplaySettings::CreateSP(this, &SControlRigEditModeTools::GetDisplaySettings);
+	RigTreeDelegates.OnSelectionChanged = FOnRigTreeSelectionChanged::CreateSP(this, &SControlRigEditModeTools::HandleSelectionChanged);
+
 	ChildSlot
 	[
 		SNew(SScrollBox)
@@ -258,7 +295,8 @@ void SControlRigEditModeTools::Construct(const FArguments& InArgs, FControlRigEd
 				.BorderBackgroundColor(FLinearColor(.6f, .6f, .6f))
 				.BodyContent()
 				[
-					SAssignNew(ControlHierarchy, SControlHierarchy, InEditMode.GetControlRig(true))
+					SAssignNew(HierarchyTreeView, SRigHierarchyTreeView)
+					.RigTreeDelegates(RigTreeDelegates)
 				]
 			]
 
@@ -284,6 +322,8 @@ void SControlRigEditModeTools::Construct(const FArguments& InArgs, FControlRigEd
 			]
 		]
 	];
+
+	HierarchyTreeView->RefreshTreeView(true);
 }
 
 void SControlRigEditModeTools::SetDetailsObjects(const TArray<TWeakObjectPtr<>>& InObjects)
@@ -491,6 +531,53 @@ void SControlRigEditModeTools::HandleModifiedEvent(ERigVMGraphNotifType InNotifT
 	}
 }
 
+void SControlRigEditModeTools::HandleSelectionChanged(TSharedPtr<FRigTreeElement> Selection, ESelectInfo::Type SelectInfo)
+{
+	if(bIsChangingRigHierarchy)
+	{
+		return;
+	}
+	
+	const URigHierarchy* Hierarchy = GetHierarchy();
+	if (Hierarchy)
+	{
+		URigHierarchyController* Controller = ((URigHierarchy*)Hierarchy)->GetController(true);
+		check(Controller);
+		
+		TGuardValue<bool> GuardRigHierarchyChanges(bIsChangingRigHierarchy, true);
+
+		const TArray<FRigElementKey> NewSelection = HierarchyTreeView->GetSelectedKeys();
+		if(!Controller->SetSelection(NewSelection))
+		{
+			return;
+		}
+	}
+}
+
+void SControlRigEditModeTools::OnRigElementSelected(UControlRig* Subject, FRigControlElement* ControlElement, bool bSelected)
+{
+	const FRigElementKey Key = ControlElement->GetKey();
+	for (int32 RootIndex = 0; RootIndex < HierarchyTreeView->GetRootElements().Num(); ++RootIndex)
+	{
+		TSharedPtr<FRigTreeElement> Found = HierarchyTreeView->FindElement(Key, HierarchyTreeView->GetRootElements()[RootIndex]);
+		if (Found.IsValid())
+		{
+			HierarchyTreeView->SetItemSelection(Found, bSelected, ESelectInfo::OnNavigation);
+			
+			TArray<TSharedPtr<FRigTreeElement>> SelectedItems = HierarchyTreeView->GetSelectedItems();
+			for (TSharedPtr<FRigTreeElement> SelectedItem : SelectedItems)
+			{
+				HierarchyTreeView->SetExpansionRecursive(SelectedItem, false, true);
+			}
+
+			if (SelectedItems.Num() > 0)
+			{
+				HierarchyTreeView->RequestScrollIntoView(SelectedItems.Last());
+			}
+		}
+	}
+}
+
 EVisibility SControlRigEditModeTools::GetRigOptionExpanderVisibility() const
 {
 	if (UControlRig* ControlRig = SequencerRig.Get())
diff --git a/Engine/Plugins/Animation/ControlRig/Source/ControlRigEditor/Private/EditMode/SControlRigEditModeTools.h b/Engine/Plugins/Animation/ControlRig/Source/ControlRigEditor/Private/EditMode/SControlRigEditModeTools.h
--- a/Engine/Plugins/Animation/ControlRig/Source/ControlRigEditor/Private/EditMode/SControlRigEditModeTools.h
+++ b/Engine/Plugins/Animation/ControlRig/Source/ControlRigEditor/Private/EditMode/SControlRigEditModeTools.h
@@ -6,6 +6,7 @@
 #include "Widgets/SCompoundWidget.h"
 #include "IDetailKeyframeHandler.h"
 #include "RigVMModel/RigVMGraph.h"
+#include "SRigHierarchyTreeView.h"
 #include "Widgets/DeclarativeSyntaxSupport.h"
 
 class FControlRigEditMode;
@@ -13,8 +14,9 @@ class IDetailsView;
 class ISequencer;
 class SControlPicker;
 class SExpandableArea;
-class SControlHierarchy;
+class SRigHierarchyTreeView;
 class UControlRig;
+class URigHierarchy;
 class FToolBarBuilder;
 class FEditorModeTools;
 
@@ -35,6 +37,9 @@ public:
 	/** Set The Control Rig we are using*/
 	void SetControlRig(UControlRig* ControlRig);
 
+	/** Returns the hierarchy currently being used */
+	const URigHierarchy* GetHierarchy() const;
+
 	// IDetailKeyframeHandler interface
 	virtual bool IsPropertyKeyable(const UClass* InObjectClass, const class IPropertyHandle& PropertyHandle) const override;
 	virtual bool IsPropertyKeyingEnabled() const override;
@@ -52,7 +57,7 @@ private:
 	TSharedPtr<IDetailsView> RigOptionsDetailsView;
 
 	/** Hierarchy picker for controls*/
-	TSharedPtr<SControlHierarchy> ControlHierarchy;
+	TSharedPtr<SRigHierarchyTreeView> HierarchyTreeView;
 
 	/** Special picker for controls, no longer used */
 	TSharedPtr<SControlPicker> ControlPicker;
@@ -70,6 +75,8 @@ private:
 	void OnManipulatorsPicked(const TArray<FName>& Manipulators);
 
 	void HandleModifiedEvent(ERigVMGraphNotifType InNotifType, URigVMGraph* InGraph, UObject* InSubject);
+	void HandleSelectionChanged(TSharedPtr<FRigTreeElement> Selection, ESelectInfo::Type SelectInfo);
+	void OnRigElementSelected(UControlRig* Subject, FRigControlElement* ControlElement, bool bSelected);
 
 	EVisibility GetRigOptionExpanderVisibility() const;
 
@@ -88,6 +95,10 @@ private:
 	//TWeakPtr<SWindow> SelectionSetWindow;
 
 	FEditorModeTools* ModeTools = nullptr;
+	FRigTreeDisplaySettings DisplaySettings;
+	const FRigTreeDisplaySettings& GetDisplaySettings() const { return DisplaySettings; }
+	bool bIsChangingRigHierarchy;
+
 public:
 	/** Modes Panel Header Information **/
 	void CustomizeToolBarPalette(FToolBarBuilder& ToolBarBuilder);
diff --git a/Engine/Plugins/Animation/ControlRig/Source/ControlRigEditor/Private/Editor/SRigHierarchy.cpp b/Engine/Plugins/Animation/ControlRig/Source/ControlRigEditor/Private/Editor/SRigHierarchy.cpp
--- a/Engine/Plugins/Animation/ControlRig/Source/ControlRigEditor/Private/Editor/SRigHierarchy.cpp
+++ b/Engine/Plugins/Animation/ControlRig/Source/ControlRigEditor/Private/Editor/SRigHierarchy.cpp
@@ -40,47 +40,10 @@
 #include "EditMode/ControlRigEditMode.h"
 #include "ToolMenus.h"
 #include "ControlRigContextMenuContext.h"
+#include "SRigSpacePickerWidget.h"
 
 #define LOCTEXT_NAMESPACE "SRigHierarchy"
 
-//////////////////////////////////////////////////////////////
-/// FRigTreeElement
-///////////////////////////////////////////////////////////
-FRigTreeElement::FRigTreeElement(const FRigElementKey& InKey, TWeakPtr<SRigHierarchy> InHierarchyHandler)
-{
-	Key = InKey;
-	bIsTransient = false;
-
-	if(InHierarchyHandler.IsValid())
-	{
-		if(URigHierarchy* Hierarchy = InHierarchyHandler.Pin()->GetDebuggedHierarchy())
-		{
-			if(FRigControlElement* ControlElement = Hierarchy->Find<FRigControlElement>(Key))
-			{
-				bIsTransient = ControlElement->Settings.bIsTransientControl;
-			}
-		}
-	}
-}
-
-
-TSharedRef<ITableRow> FRigTreeElement::MakeTreeRowWidget(TSharedPtr<FControlRigEditor> InControlRigEditor, const TSharedRef<STableViewBase>& InOwnerTable, TSharedRef<FRigTreeElement> InRigTreeElement, TSharedRef<FUICommandList> InCommandList, TSharedPtr<SRigHierarchy> InHierarchy)
-{
-	if (InRigTreeElement->Key.IsValid())
-	{
-		return SNew(SRigHierarchyItem, InControlRigEditor, InOwnerTable, InRigTreeElement, InCommandList, InHierarchy)
-			.OnRenameElement(InHierarchy.Get(), &SRigHierarchy::RenameElement)
-			.OnVerifyElementNameChanged(InHierarchy.Get(), &SRigHierarchy::OnVerifyNameChanged);
-	}
-
-	return SNew(SRigHierarchyItem, InControlRigEditor, InOwnerTable, InRigTreeElement, InCommandList, InHierarchy);
-}
-
-void FRigTreeElement::RequestRename()
-{
-	OnRenameRequested.ExecuteIfBound();
-}
-
 //////////////////////////////////////////////////////////////
 /// FRigElementHierarchyDragDropOp
 ///////////////////////////////////////////////////////////
@@ -114,176 +77,6 @@ FString FRigElementHierarchyDragDropOp::GetJoinedElementNames() const
 	return FString::Join(ElementNameStrings, TEXT(","));
 }
 
-//////////////////////////////////////////////////////////////
-/// SRigHierarchyItem
-///////////////////////////////////////////////////////////
-void SRigHierarchyItem::Construct(const FArguments& InArgs, TSharedPtr<FControlRigEditor> InControlRigEditor, const TSharedRef<STableViewBase>& OwnerTable, TSharedRef<FRigTreeElement> InRigTreeElement, TSharedRef<FUICommandList> InCommandList, TSharedPtr<SRigHierarchy> InHierarchy)
-{
-	WeakRigTreeElement = InRigTreeElement;
-	WeakCommandList = InCommandList;
-	ControlRigEditor = InControlRigEditor;
-
-	OnVerifyElementNameChanged = InArgs._OnVerifyElementNameChanged;
-	OnRenameElement = InArgs._OnRenameElement;
-
-	if (!InRigTreeElement->Key.IsValid())
-	{
-		STableRow<TSharedPtr<FRigTreeElement>>::Construct(
-			STableRow<TSharedPtr<FRigTreeElement>>::FArguments()
-			.ShowSelection(false)
-			.OnCanAcceptDrop(InHierarchy.Get(), &SRigHierarchy::OnCanAcceptDrop)
-			.OnAcceptDrop(InHierarchy.Get(), &SRigHierarchy::OnAcceptDrop)
-			.Content()
-			[
-				SNew(SVerticalBox)
-				+ SVerticalBox::Slot()
-				.FillHeight(200.f)
-				[
-					SNew(SSpacer)
-				]
-			], OwnerTable);
-		return;
-	}
-
-	TSharedPtr< SInlineEditableTextBlock > InlineWidget;
-
-	const FSlateBrush* Brush = nullptr;
-	switch (InRigTreeElement->Key.Type)
-	{
-		case ERigElementType::Control:
-		{
-			Brush = FControlRigEditorStyle::Get().GetBrush("ControlRig.Tree.Control");
-			break;
-		}
-		case ERigElementType::Null:
-		{
-			Brush = FControlRigEditorStyle::Get().GetBrush("ControlRig.Tree.Null");
-			break;
-		}
-		case ERigElementType::Bone:
-		{
-			ERigBoneType BoneType = ERigBoneType::User;
-
-			FRigBoneElement* BoneElement = InHierarchy->GetHierarchyForTopology()->Find<FRigBoneElement>(InRigTreeElement->Key);
-			if(BoneElement)
-			{
-				BoneType = BoneElement->BoneType;
-			}
-
-			switch (BoneType)
-			{
-				case ERigBoneType::Imported:
-				{
-					Brush = FControlRigEditorStyle::Get().GetBrush("ControlRig.Tree.BoneImported");
-					break;
-				}
-				case ERigBoneType::User:
-				default:
-				{
-					Brush = FControlRigEditorStyle::Get().GetBrush("ControlRig.Tree.BoneUser");
-					break;
-				}
-			}
-
-			break;
-		}
-		case ERigElementType::RigidBody:
-		{
-			Brush = FControlRigEditorStyle::Get().GetBrush("ControlRig.Tree.RigidBody");
-			break;
-		}
-		case ERigElementType::Socket:
-		{
-			Brush = FControlRigEditorStyle::Get().GetBrush("ControlRig.Tree.Socket");
-			break;
-		}
-		default:
-		{
-			break;
-		}
-	}
-
-	STableRow<TSharedPtr<FRigTreeElement>>::Construct(
-		STableRow<TSharedPtr<FRigTreeElement>>::FArguments()
-		.OnDragDetected(InHierarchy.Get(), &SRigHierarchy::OnDragDetected)
-		.OnCanAcceptDrop(InHierarchy.Get(), &SRigHierarchy::OnCanAcceptDrop)
-		.OnAcceptDrop(InHierarchy.Get(), &SRigHierarchy::OnAcceptDrop)
-		.ShowWires(true)
-		.Content()
-		[
-			SNew(SHorizontalBox)
-			+ SHorizontalBox::Slot()
-			.MaxWidth(18)
-			.FillWidth(1.0)
-			.HAlign(HAlign_Left)
-			.VAlign(VAlign_Center)
-			[
-				SNew(SImage)
-				.Image(Brush)
-			]
-			+ SHorizontalBox::Slot()
-			.AutoWidth()
-			.VAlign(VAlign_Center)
-			[
-				SAssignNew(InlineWidget, SInlineEditableTextBlock)
-				.Text(this, &SRigHierarchyItem::GetName)
-				.OnVerifyTextChanged(this, &SRigHierarchyItem::OnVerifyNameChanged)
-				.OnTextCommitted(this, &SRigHierarchyItem::OnNameCommitted)
-				.MultiLine(false)
-			]
-		], OwnerTable);
-
-	InRigTreeElement->OnRenameRequested.BindSP(InlineWidget.Get(), &SInlineEditableTextBlock::EnterEditingMode);
-}
-
-FText SRigHierarchyItem::GetName() const
-{
-	if(WeakRigTreeElement.Pin()->bIsTransient)
-	{
-		static const FText TemporaryControl = FText::FromString(TEXT("Temporary Control"));
-		return TemporaryControl;
-	}
-	return (FText::FromName(WeakRigTreeElement.Pin()->Key.Name));
-}
-
-bool SRigHierarchyItem::OnVerifyNameChanged(const FText& InText, FText& OutErrorMessage)
-{
-	FString NewName = InText.ToString();
-	if (OnVerifyElementNameChanged.IsBound())
-	{
-		return OnVerifyElementNameChanged.Execute(WeakRigTreeElement.Pin()->Key, NewName, OutErrorMessage);
-	}
-
-	// if not bound, just allow
-	return true;
-}
-
-void SRigHierarchyItem::OnNameCommitted(const FText& InText, ETextCommit::Type InCommitType) const
-{
-	// for now only allow enter
-	// because it is important to keep the unique names per pose
-	if (InCommitType == ETextCommit::OnEnter)
-	{
-		FString NewName = InText.ToString();
-		FRigElementKey OldKey = WeakRigTreeElement.Pin()->Key;
-
-		if (OnRenameElement.IsBound())
-		{
-			FName NewSanitizedName = OnRenameElement.Execute(OldKey, NewName);
-			if (NewSanitizedName.IsNone())
-			{
-				return;
-			}
-			NewName = NewSanitizedName.ToString();
-		}
-
-		if (WeakRigTreeElement.IsValid())
-		{
-			WeakRigTreeElement.Pin()->Key.Name = *NewName;
-		}
-	}
-}
-
 ///////////////////////////////////////////////////////////
 
 SRigHierarchy::~SRigHierarchy()
@@ -313,6 +106,21 @@ void SRigHierarchy::Construct(const FArguments& InArgs, TSharedRef<FControlRigEd
 
 	BindCommands();
 
+	// setup all delegates for the rig hierarchy widget
+	FRigTreeDelegates Delegates;
+	Delegates.OnGetHierarchy = FOnGetRigTreeHierarchy::CreateSP(this, &SRigHierarchy::GetHierarchyForTopology);
+	Delegates.OnGetDisplaySettings = FOnGetRigTreeDisplaySettings::CreateSP(this, &SRigHierarchy::GetDisplaySettings);
+	Delegates.OnRenameElement = FOnRigTreeRenameElement::CreateSP(this, &SRigHierarchy::HandleRenameElement);
+	Delegates.OnVerifyElementNameChanged = FOnRigTreeVerifyElementNameChanged::CreateSP(this, &SRigHierarchy::HandleVerifyNameChanged);
+	Delegates.OnSelectionChanged = FOnRigTreeSelectionChanged::CreateSP(this, &SRigHierarchy::OnSelectionChanged);
+	Delegates.OnContextMenuOpening = FOnContextMenuOpening::CreateSP(this, &SRigHierarchy::CreateContextMenuWidget);
+	Delegates.OnMouseButtonClick = FOnRigTreeMouseButtonClick::CreateSP(this, &SRigHierarchy::OnItemClicked);
+	Delegates.OnMouseButtonDoubleClick = FOnRigTreeMouseButtonClick::CreateSP(this, &SRigHierarchy::OnItemDoubleClicked);
+	Delegates.OnSetExpansionRecursive = FOnRigTreeSetExpansionRecursive::CreateSP(this, &SRigHierarchy::OnSetExpansionRecursive);
+	Delegates.OnCanAcceptDrop = FOnRigTreeCanAcceptDrop::CreateSP(this, &SRigHierarchy::OnCanAcceptDrop);
+	Delegates.OnAcceptDrop = FOnRigTreeAcceptDrop::CreateSP(this, &SRigHierarchy::OnAcceptDrop);
+	Delegates.OnDragDetected = FOnDragDetected::CreateSP(this, &SRigHierarchy::OnDragDetected);
+
 	ChildSlot
 	[
 		SNew(SVerticalBox)
@@ -407,17 +215,7 @@ void SRigHierarchy::Construct(const FArguments& InArgs, TSharedRef<FControlRigEd
 			.BorderImage(FEditorStyle::GetBrush("SCSEditor.TreePanel"))
 			[
 				SAssignNew(TreeView, SRigHierarchyTreeView)
-				.TreeItemsSource(&RootElements)
-				.SelectionMode(ESelectionMode::Multi)
-				.OnGenerateRow(this, &SRigHierarchy::MakeTableRowWidget)
-				.OnGetChildren(this, &SRigHierarchy::HandleGetChildrenForTree)
-				.OnSelectionChanged(this, &SRigHierarchy::OnSelectionChanged)
-				.OnContextMenuOpening(this, &SRigHierarchy::CreateContextMenuWidget)
-				.OnMouseButtonClick(this, &SRigHierarchy::OnItemClicked)
-				.OnMouseButtonDoubleClick(this, &SRigHierarchy::OnItemDoubleClicked)
-				.OnSetExpansionRecursive(this, &SRigHierarchy::OnSetExpansionRecursive)
-				.HighlightParentNodesForSelection(true)
-				.ItemHeight(24)
+				.RigTreeDelegates(Delegates)
 			]
 		]
 
@@ -436,16 +234,7 @@ void SRigHierarchy::Construct(const FArguments& InArgs, TSharedRef<FControlRigEd
 		*/
 	];
 
-	bFlattenHierarchyOnFilter = false;
-	bHideParentsOnFilter = false;
-	bShowImportedBones = true;
-	bShowBones = true;
-	bShowControls = true;
-	bShowNulls = true;
-	bShowRigidBodies = true;
-	bShowSockets = true;
 	bIsChangingRigHierarchy = false;
-	bShowDynamicHierarchy = false;
 	RefreshTreeView();
 
 	if (ControlRigEditor.IsValid())
@@ -570,62 +359,62 @@ void SRigHierarchy::BindCommands()
 
 	CommandList->MapAction(
 		Commands.FilteringFlattensHierarchy,
-		FExecuteAction::CreateLambda([this]() { bFlattenHierarchyOnFilter = !bFlattenHierarchyOnFilter; RefreshTreeView(); }),
+		FExecuteAction::CreateLambda([this]() { DisplaySettings.bFlattenHierarchyOnFilter = !DisplaySettings.bFlattenHierarchyOnFilter; RefreshTreeView(); }),
 		FCanExecuteAction(),
-		FIsActionChecked::CreateLambda([this]() { return bFlattenHierarchyOnFilter; }));
+		FIsActionChecked::CreateLambda([this]() { return DisplaySettings.bFlattenHierarchyOnFilter; }));
 
 	CommandList->MapAction(
 		Commands.HideParentsWhenFiltering,
-		FExecuteAction::CreateLambda([this]() { bHideParentsOnFilter = !bHideParentsOnFilter; RefreshTreeView(); }),
+		FExecuteAction::CreateLambda([this]() { DisplaySettings.bHideParentsOnFilter = !DisplaySettings.bHideParentsOnFilter; RefreshTreeView(); }),
 		FCanExecuteAction(),
-		FIsActionChecked::CreateLambda([this]() { return bHideParentsOnFilter; }));
+		FIsActionChecked::CreateLambda([this]() { return DisplaySettings.bHideParentsOnFilter; }));
 
 	CommandList->MapAction(
 		Commands.ShowImportedBones,
-		FExecuteAction::CreateLambda([this]() { bShowImportedBones = !bShowImportedBones; RefreshTreeView(); }),
+		FExecuteAction::CreateLambda([this]() { DisplaySettings.bShowImportedBones = !DisplaySettings.bShowImportedBones; RefreshTreeView(); }),
 		FCanExecuteAction(),
-		FIsActionChecked::CreateLambda([this]() { return bShowImportedBones; }));
+		FIsActionChecked::CreateLambda([this]() { return DisplaySettings.bShowImportedBones; }));
 	
 	CommandList->MapAction(
 		Commands.ShowBones,
-		FExecuteAction::CreateLambda([this]() { bShowBones = !bShowBones; RefreshTreeView(); }),
+		FExecuteAction::CreateLambda([this]() { DisplaySettings.bShowBones = !DisplaySettings.bShowBones; RefreshTreeView(); }),
 		FCanExecuteAction(),
-		FIsActionChecked::CreateLambda([this]() { return bShowBones; }));
+		FIsActionChecked::CreateLambda([this]() { return DisplaySettings.bShowBones; }));
 
 	CommandList->MapAction(
 		Commands.ShowControls,
-		FExecuteAction::CreateLambda([this]() { bShowControls = !bShowControls; RefreshTreeView(); }),
+		FExecuteAction::CreateLambda([this]() { DisplaySettings.bShowControls = !DisplaySettings.bShowControls; RefreshTreeView(); }),
 		FCanExecuteAction(),
-		FIsActionChecked::CreateLambda([this]() { return bShowControls; }));
+		FIsActionChecked::CreateLambda([this]() { return DisplaySettings.bShowControls; }));
 	
 	CommandList->MapAction(
 		Commands.ShowNulls,
-		FExecuteAction::CreateLambda([this]() { bShowNulls = !bShowNulls; RefreshTreeView(); }),
+		FExecuteAction::CreateLambda([this]() { DisplaySettings.bShowNulls = !DisplaySettings.bShowNulls; RefreshTreeView(); }),
 		FCanExecuteAction(),
-		FIsActionChecked::CreateLambda([this]() { return bShowNulls; }));
+		FIsActionChecked::CreateLambda([this]() { return DisplaySettings.bShowNulls; }));
 
 	CommandList->MapAction(
 		Commands.ShowRigidBodies,
-		FExecuteAction::CreateLambda([this]() { bShowRigidBodies = !bShowRigidBodies; RefreshTreeView(); }),
+		FExecuteAction::CreateLambda([this]() { DisplaySettings.bShowRigidBodies = !DisplaySettings.bShowRigidBodies; RefreshTreeView(); }),
 		FCanExecuteAction(),
-		FIsActionChecked::CreateLambda([this]() { return bShowRigidBodies; }));
+		FIsActionChecked::CreateLambda([this]() { return DisplaySettings.bShowRigidBodies; }));
 
 	CommandList->MapAction(
 		Commands.ShowSockets,
-		FExecuteAction::CreateLambda([this]() { bShowSockets = !bShowSockets; RefreshTreeView(); }),
+		FExecuteAction::CreateLambda([this]() { DisplaySettings.bShowSockets = !DisplaySettings.bShowSockets; RefreshTreeView(); }),
 		FCanExecuteAction(),
-		FIsActionChecked::CreateLambda([this]() { return bShowSockets; }));
+		FIsActionChecked::CreateLambda([this]() { return DisplaySettings.bShowSockets; }));
 
 	CommandList->MapAction(
 		Commands.ShowDynamicHierarchy,
 		FExecuteAction::CreateLambda([this]()
 		{
-			bShowDynamicHierarchy = !bShowDynamicHierarchy;
+			DisplaySettings.bShowDynamicHierarchy = !DisplaySettings.bShowDynamicHierarchy;
 			HandleSetObjectBeingDebugged(GetControlRigEditor()->GetControlRigBlueprint()->GetObjectBeingDebugged());
 			RefreshTreeView();
 		}),
 		FCanExecuteAction(),
-		FIsActionChecked::CreateLambda([this]() { return bShowDynamicHierarchy; }));
+		FIsActionChecked::CreateLambda([this]() { return DisplaySettings.bShowDynamicHierarchy; }));
 
 	CommandList->MapAction(
 		Commands.ToggleGizmoTransformEdit,
@@ -633,7 +422,11 @@ void SRigHierarchy::BindCommands()
 		{
 			ControlRigEditor.Pin()->GetEditMode()->ToggleGizmoTransformEdit();
 		}));
-	
+
+	CommandList->MapAction(
+		Commands.TestSpaceSwitching,
+		FExecuteAction::CreateSP(this, &SRigHierarchy::HandleTestSpaceSwitching),
+		FCanExecuteAction::CreateSP(this, &SRigHierarchy::IsControlSelected));
 }
 
 FReply SRigHierarchy::OnKeyDown(const FGeometry& MyGeometry, const FKeyEvent& InKeyEvent)
@@ -695,95 +488,13 @@ FReply SRigHierarchy::OnImportSkeletonClicked()
 
 void SRigHierarchy::OnFilterTextChanged(const FText& SearchText)
 {
-	FilterText = SearchText;
-
+	DisplaySettings.FilterText = SearchText;
 	RefreshTreeView();
 }
 
 void SRigHierarchy::RefreshTreeView(bool bRebuildContent)
 {
-	TMap<FRigElementKey, bool> ExpansionState;
-
-	if(bRebuildContent)
-	{
-		for (TPair<FRigElementKey, TSharedPtr<FRigTreeElement>> Pair : ElementMap)
-		{
-			ExpansionState.FindOrAdd(Pair.Key) = TreeView->IsItemExpanded(Pair.Value);
-		}
-
-		// internally save expansion states before rebuilding the tree, so the states can be restored later
-		TreeView->SaveAndClearSparseItemInfos();
-
-		RootElements.Reset();
-		ElementMap.Reset();
-		ParentMap.Reset();
-	}
-
-	if (ControlRigBlueprint.IsValid())
-	{
-		URigHierarchy* Hierarchy = GetHierarchyForTopology();
-		check(Hierarchy);
-
-		if(bRebuildContent)
-		{
-			Hierarchy->Traverse([&](FRigBaseElement* Element, bool& bContinue)
-            {
-                AddElement(Element);
-                bContinue = true;
-            });
-
-			// expand all elements upon the initial construction of the tree
-			if (ExpansionState.Num() == 0)
-			{
-				for (TSharedPtr<FRigTreeElement> RootElement : RootElements)
-				{
-					SetExpansionRecursive(RootElement, false, true);
-				}
-			}
-
-			for (const auto& Pair : ElementMap)
-			{
-				TreeView->RestoreSparseItemInfos(Pair.Value);
-			}
-
-			if (RootElements.Num() > 0)
-			{
-				AddSpacerElement();
-			}
-		}
-		else
-		{
-			if (RootElements.Num()> 0)
-			{
-				// elements may be added at the end of the list after a spacer element
-				// we need to remove the spacer element and re-add it at the end
-				RootElements.RemoveAll([](TSharedPtr<FRigTreeElement> InElement)
-				{
-					return InElement.Get()->Key == FRigElementKey();
-				});
-				AddSpacerElement();
-			}
-		}
-
-		TreeView->RequestTreeRefresh();
-		{
-			TGuardValue<bool> GuardRigHierarchyChanges(bIsChangingRigHierarchy, true);
-			TreeView->ClearSelection();
-
-			TArray<FRigElementKey> Selection = Hierarchy->GetSelectedKeys();
-			for (const FRigElementKey& Key : Selection)
-			{
-				for (int32 RootIndex = 0; RootIndex < RootElements.Num(); ++RootIndex)
-				{
-					TSharedPtr<FRigTreeElement> Found = FindElement(Key, RootElements[RootIndex]);
-					if (Found.IsValid())
-					{
-						TreeView->SetItemSelection(Found, true, ESelectInfo::OnNavigation);
-					}
-				}
-			}
-		}
-	}
+	TreeView->RefreshTreeView(bRebuildContent);
 }
 
 TArray<FRigElementKey> SRigHierarchy::GetSelectedKeys() const
@@ -802,38 +513,6 @@ TArray<FRigElementKey> SRigHierarchy::GetSelectedKeys() const
 	return SelectedKeys;
 }
 
-void SRigHierarchy::SetExpansionRecursive(TSharedPtr<FRigTreeElement> InElement, bool bTowardsParent, bool bShouldBeExpanded)
-{
-	TreeView->SetItemExpansion(InElement, bShouldBeExpanded);
-
-	if (bTowardsParent)
-	{
-		if (const FRigElementKey* ParentKey = ParentMap.Find(InElement->Key))
-		{
-			if (TSharedPtr<FRigTreeElement>* ParentItem = ElementMap.Find(*ParentKey))
-			{
-				SetExpansionRecursive(*ParentItem, bTowardsParent, bShouldBeExpanded);
-			}
-		}
-	}
-	else
-	{
-		for (int32 ChildIndex = 0; ChildIndex < InElement->Children.Num(); ++ChildIndex)
-		{
-			SetExpansionRecursive(InElement->Children[ChildIndex], bTowardsParent, bShouldBeExpanded);
-		}
-	}
-}
-TSharedRef<ITableRow> SRigHierarchy::MakeTableRowWidget(TSharedPtr<FRigTreeElement> InItem, const TSharedRef<STableViewBase>& OwnerTable)
-{
-	return InItem->MakeTreeRowWidget(ControlRigEditor.Pin(), OwnerTable, InItem.ToSharedRef(), CommandList.ToSharedRef(), SharedThis(this));
-}
-
-void SRigHierarchy::HandleGetChildrenForTree(TSharedPtr<FRigTreeElement> InItem, TArray<TSharedPtr<FRigTreeElement>>& OutChildren)
-{
-	OutChildren = InItem.Get()->Children;
-}
-
 void SRigHierarchy::OnSelectionChanged(TSharedPtr<FRigTreeElement> Selection, ESelectInfo::Type SelectInfo)
 {
 	if (bIsChangingRigHierarchy)
@@ -896,242 +575,6 @@ void SRigHierarchy::OnSelectionChanged(TSharedPtr<FRigTreeElement> Selection, ES
 	}
 }
 
-TSharedPtr<FRigTreeElement> SRigHierarchy::FindElement(const FRigElementKey& InElementKey, TSharedPtr<FRigTreeElement> CurrentItem)
-{
-	if (CurrentItem->Key == InElementKey)
-	{
-		return CurrentItem;
-	}
-
-	for (int32 ChildIndex = 0; ChildIndex < CurrentItem->Children.Num(); ++ChildIndex)
-	{
-		TSharedPtr<FRigTreeElement> Found = FindElement(InElementKey, CurrentItem->Children[ChildIndex]);
-		if (Found.IsValid())
-		{
-			return Found;
-		}
-	}
-
-	return TSharedPtr<FRigTreeElement>();
-}
-
-bool SRigHierarchy::AddElement(FRigElementKey InKey, FRigElementKey InParentKey, const bool bIgnoreTextFilter)
-{
-	if(ElementMap.Contains(InKey))
-	{
-		return false;
-	}
-
-	FString FilteredString = FilterText.ToString();
-	if (bIgnoreTextFilter || FilteredString.IsEmpty() || !InKey.IsValid())
-	{
-		TSharedPtr<FRigTreeElement> NewItem = MakeShared<FRigTreeElement>(InKey, SharedThis(this));
-
-		if (InKey.IsValid())
-		{
-			ElementMap.Add(InKey, NewItem);
-			if (InParentKey)
-			{
-				ParentMap.Add(InKey, InParentKey);
-			}
-
-			if (InParentKey)
-			{
-				TSharedPtr<FRigTreeElement>* FoundItem = ElementMap.Find(InParentKey);
-				check(FoundItem);
-				FoundItem->Get()->Children.Add(NewItem);
-			}
-			else
-			{
-				RootElements.Add(NewItem);
-			}
-		}
-		else
-		{
-			RootElements.Add(NewItem);
-		}
-	}
-	else
-	{
-		FString FilteredStringUnderScores = FilteredString.Replace(TEXT(" "), TEXT("_"));
-		if (InKey.Name.ToString().Contains(FilteredString) || InKey.Name.ToString().Contains(FilteredStringUnderScores))	
-		{
-			TSharedPtr<FRigTreeElement> NewItem = MakeShared<FRigTreeElement>(InKey, SharedThis(this));
-			ElementMap.Add(InKey, NewItem);
-			RootElements.Add(NewItem);
-		}
-	}
-
-	return true;
-}
-
-bool SRigHierarchy::AddElement(const FRigBaseElement* InElement, const bool bIgnoreTextFilter)
-{
-	check(InElement);
-	
-	if (ElementMap.Contains(InElement->GetKey()))
-	{
-		return false;
-	}
-
-	switch(InElement->GetType())
-	{
-		case ERigElementType::Bone:
-		{
-			if(!bShowBones)
-			{
-				return false;
-			}
-
-			const FRigBoneElement* BoneElement = CastChecked<FRigBoneElement>(InElement);
-			if (!bShowImportedBones && BoneElement->BoneType == ERigBoneType::Imported)
-			{
-				return false;
-			}
-			break;
-		}
-		case ERigElementType::Null:
-		{
-			if(!bShowNulls)
-			{
-				return false;
-			}
-			break;
-		}
-		case ERigElementType::Control:
-		{
-			if(!bShowControls)
-			{
-				return false;
-			}
-			break;
-		}
-		case ERigElementType::RigidBody:
-		{
-			if(!bShowRigidBodies)
-			{
-				return false;
-			}
-			break;
-		}
-		case ERigElementType::Socket:
-		{
-			if(!bShowSockets)
-			{
-				return false;
-			}
-			break;
-		}
-		case ERigElementType::Curve:
-		{
-			return false;
-		}
-		default:
-		{
-			break;
-		}
-	}
-
-	const URigHierarchy* Hierarchy = GetHierarchyForTopology();
-	check(Hierarchy);
-
-	if(!AddElement(InElement->GetKey(), FRigElementKey(), bIgnoreTextFilter))
-	{
-		return false;
-	}
-
-	if (ElementMap.Contains(InElement->GetKey()))
-	{
-		const FRigElementKey ParentKey = Hierarchy->GetFirstParent(InElement->GetKey());
-		if (ParentKey.IsValid())
-		{
-			if(ElementMap.Contains(ParentKey))
-			{
-				ReparentElement(InElement->GetKey(), ParentKey);
-			}
-		}
-	}
-
-	return true;
-}
-
-void SRigHierarchy::AddSpacerElement()
-{
-	AddElement(FRigElementKey(), FRigElementKey());
-}
-
-bool SRigHierarchy::ReparentElement(FRigElementKey InKey, FRigElementKey InParentKey)
-{
-	if (!InKey.IsValid() || InKey == InParentKey)
-	{
-		return false;
-	}
-
-	TSharedPtr<FRigTreeElement>* FoundItem = ElementMap.Find(InKey);
-	if (FoundItem == nullptr)
-	{
-		return false;
-	}
-
-	if (!FilterText.IsEmpty() && bFlattenHierarchyOnFilter)
-	{
-		return false;
-	}
-
-	if (const FRigElementKey* ExistingParentKey = ParentMap.Find(InKey))
-	{
-		if (*ExistingParentKey == InParentKey)
-		{
-			return false;
-		}
-
-		if (TSharedPtr<FRigTreeElement>* ExistingParent = ElementMap.Find(*ExistingParentKey))
-		{
-			(*ExistingParent)->Children.Remove(*FoundItem);
-		}
-
-		ParentMap.Remove(InKey);
-	}
-	else
-	{
-		if (!InParentKey.IsValid())
-		{
-			return false;
-		}
-
-		RootElements.Remove(*FoundItem);
-	}
-
-	if (InParentKey)
-	{
-		ParentMap.Add(InKey, InParentKey);
-
-		TSharedPtr<FRigTreeElement>* FoundParent = ElementMap.Find(InParentKey);
-		check(FoundParent);
-		FoundParent->Get()->Children.Add(*FoundItem);
-	}
-	else
-	{
-		RootElements.Add(*FoundItem);
-	}
-
-	return true;
-}
-
-bool SRigHierarchy::RemoveElement(FRigElementKey InKey)
-{
-	TSharedPtr<FRigTreeElement>* FoundItem = ElementMap.Find(InKey);
-	if (FoundItem == nullptr)
-	{
-		return false;
-	}
-
-	ReparentElement(InKey, FRigElementKey());
-
-	RootElements.Remove(*FoundItem);
-	return ElementMap.Remove(InKey) > 0;
-}
-
 void SRigHierarchy::OnHierarchyModified(ERigHierarchyNotification InNotif, URigHierarchy* InHierarchy, const FRigBaseElement* InElement)
 {
 	if(!ControlRigBlueprint.IsValid())
@@ -1163,7 +606,7 @@ void SRigHierarchy::OnHierarchyModified(ERigHierarchyNotification InNotif, URigH
 		{
 			if(InElement)
 			{
-				if(AddElement(InElement))
+				if(TreeView->AddElement(InElement))
 				{
 					RefreshTreeView(false);
 				}
@@ -1174,7 +617,7 @@ void SRigHierarchy::OnHierarchyModified(ERigHierarchyNotification InNotif, URigH
 		{
 			if(InElement)
 			{
-				if(RemoveElement(InElement->GetKey()))
+				if(TreeView->RemoveElement(InElement->GetKey()))
 				{
 					RefreshTreeView(false);
 				}
@@ -1187,13 +630,42 @@ void SRigHierarchy::OnHierarchyModified(ERigHierarchyNotification InNotif, URigH
 			if(InElement)
 			{
 				const FRigElementKey ParentKey = InHierarchy->GetFirstParent(InElement->GetKey());
-				if(ReparentElement(InElement->GetKey(), ParentKey))
+				if(TreeView->ReparentElement(InElement->GetKey(), ParentKey))
 				{
 					RefreshTreeView(false);
 				}
 			}
 			break;
 		}
+		case ERigHierarchyNotification::ParentWeightsChanged:
+		{
+			if(DisplaySettings.bShowDynamicHierarchy)
+			{
+				if(URigHierarchy* Hierarchy = GetDebuggedHierarchy())
+				{
+					TArray<FRigElementWeight> ParentWeights = Hierarchy->GetParentWeightArray(InElement->GetKey());
+					if(ParentWeights.Num() > 0)
+					{
+						TArray<FRigElementKey> ParentKeys = Hierarchy->GetParents(InElement->GetKey());
+						check(ParentKeys.Num() == ParentWeights.Num());
+						for(int32 ParentIndex=0;ParentIndex<ParentKeys.Num();ParentIndex++)
+						{
+							if(ParentWeights[ParentIndex].IsAlmostZero())
+							{
+								continue;
+							}
+
+							if(TreeView->ReparentElement(InElement->GetKey(), ParentKeys[ParentIndex]))
+							{
+								RefreshTreeView(false);
+							}
+							break;
+						}
+					}
+				}
+			}
+			break;
+		}
 		case ERigHierarchyNotification::ElementRenamed:
 		case ERigHierarchyNotification::HierarchyReset:
 		{
@@ -1207,9 +679,9 @@ void SRigHierarchy::OnHierarchyModified(ERigHierarchyNotification InNotif, URigH
 			{
 				const bool bSelected = (InNotif == ERigHierarchyNotification::ElementSelected); 
 					
-				for (int32 RootIndex = 0; RootIndex < RootElements.Num(); ++RootIndex)
+				for (int32 RootIndex = 0; RootIndex < TreeView->RootElements.Num(); ++RootIndex)
 				{
-					TSharedPtr<FRigTreeElement> Found = FindElement(InElement->GetKey(), RootElements[RootIndex]);
+					TSharedPtr<FRigTreeElement> Found = TreeView->FindElement(InElement->GetKey(), TreeView->RootElements[RootIndex]);
 					if (Found.IsValid())
 					{
 						TreeView->SetItemSelection(Found, bSelected, ESelectInfo::OnNavigation);
@@ -1237,7 +709,7 @@ void SRigHierarchy::OnHierarchyModified(ERigHierarchyNotification InNotif, URigH
 
 void SRigHierarchy::OnHierarchyModified_AnyThread(ERigHierarchyNotification InNotif, URigHierarchy* InHierarchy, const FRigBaseElement* InElement)
 {
-	if(!bShowDynamicHierarchy)
+	if(!DisplaySettings.bShowDynamicHierarchy)
 	{
 		return;
 	}
@@ -1415,17 +887,17 @@ void SRigHierarchy::OnItemDoubleClicked(TSharedPtr<FRigTreeElement> InItem)
 {
 	if (TreeView->IsItemExpanded(InItem))
 	{
-		SetExpansionRecursive(InItem, false, false);
+		TreeView->SetExpansionRecursive(InItem, false, false);
 	}
 	else
 	{
-		SetExpansionRecursive(InItem, false, true);
+		TreeView->SetExpansionRecursive(InItem, false, true);
 	}
 }
 
 void SRigHierarchy::OnSetExpansionRecursive(TSharedPtr<FRigTreeElement> InItem, bool bShouldBeExpanded)
 {
-	SetExpansionRecursive(InItem, false, bShouldBeExpanded);
+	TreeView->SetExpansionRecursive(InItem, false, bShouldBeExpanded);
 }
 
 UToolMenu* SRigHierarchy::GetOrCreateDragDropMenu(const TArray<FRigElementKey>& DraggedKeys, FRigElementKey TargetKey)
@@ -1489,7 +961,7 @@ UToolMenu* SRigHierarchy::GetOrCreateContextMenu()
 {
 	const FName MenuName = TEXT("ControlRigEditor.RigHierarchy.ContextMenu");
 	const FName InteractionSectionName = TEXT("Interaction");
-	
+
 	UToolMenus* ToolMenus = UToolMenus::Get();
 	const FControlRigHierarchyCommands& Commands = FControlRigHierarchyCommands::Get();
 	
@@ -1533,10 +1005,17 @@ UToolMenu* SRigHierarchy::GetOrCreateContextMenu()
 			{
 				if (const TSharedPtr<SRigHierarchy> RigHierarchyPanel = WeakThis.Pin())
 				{
-					if (RigHierarchyPanel->IsSingleBoneSelected())
+					if (RigHierarchyPanel->IsSingleBoneSelected() || RigHierarchyPanel->IsControlSelected())
 					{
 						FToolMenuSection& InteractionSection = InMenu->AddSection(InteractionSectionName, LOCTEXT("InteractionHeader", "Interaction"));
-						InteractionSection.AddMenuEntry(Commands.ControlBoneTransform);
+						if(RigHierarchyPanel->IsSingleBoneSelected())
+						{
+							InteractionSection.AddMenuEntry(Commands.ControlBoneTransform);
+						}
+						else if(RigHierarchyPanel->IsControlSelected())
+						{
+							InteractionSection.AddMenuEntry(Commands.TestSpaceSwitching);
+						}
 					}
 				}
 			}),
@@ -2251,10 +1730,10 @@ URigHierarchy* SRigHierarchy::GetDebuggedHierarchy() const
 	return GetHierarchy();
 }
 
-URigHierarchy* SRigHierarchy::GetHierarchyForTopology() const
+const URigHierarchy* SRigHierarchy::GetHierarchyForTopology() const
 {
 	URigHierarchy* Hierarchy = GetHierarchy();
-	if(bShowDynamicHierarchy && ControlRigBeingDebuggedPtr.IsValid())
+	if(DisplaySettings.bShowDynamicHierarchy && ControlRigBeingDebuggedPtr.IsValid())
 	{
 		Hierarchy = ControlRigBeingDebuggedPtr->GetHierarchy();
 	}
@@ -2402,7 +1881,7 @@ FReply SRigHierarchy::OnAcceptDrop(const FDragDropEvent& DragDropEvent, EItemDro
 	return FReply::Unhandled();
 }
 
-FName SRigHierarchy::RenameElement(const FRigElementKey& OldKey, const FString& NewName)
+FName SRigHierarchy::HandleRenameElement(const FRigElementKey& OldKey, const FString& NewName)
 {
 	ClearDetailPanel();
 
@@ -2433,7 +1912,7 @@ FName SRigHierarchy::RenameElement(const FRigElementKey& OldKey, const FString&
 	return NAME_None;
 }
 
-bool SRigHierarchy::OnVerifyNameChanged(const FRigElementKey& OldKey, const FString& NewName, FText& OutErrorMessage)
+bool SRigHierarchy::HandleVerifyNameChanged(const FRigElementKey& OldKey, const FString& NewName, FText& OutErrorMessage)
 {
 	if (OldKey.Name.ToString() == NewName)
 	{
@@ -2575,7 +2054,7 @@ void SRigHierarchy::HandleFrameSelection()
 	TArray<TSharedPtr<FRigTreeElement>> SelectedItems = TreeView->GetSelectedItems();
 	for (TSharedPtr<FRigTreeElement> SelectedItem : SelectedItems)
 	{
-		SetExpansionRecursive(SelectedItem, true, true);
+		TreeView->SetExpansionRecursive(SelectedItem, true, true);
 	}
 
 	if (SelectedItems.Num() > 0)
@@ -2724,6 +2203,42 @@ bool SRigHierarchy::FindClosestBone(const FVector& Point, FName& OutRigElementNa
 	return false;
 }
 
+void SRigHierarchy::HandleTestSpaceSwitching()
+{
+	for(const FRigElementKey& ControlKey : GetSelectedKeys())
+	{
+		if(ControlKey.Type == ERigElementType::Control)
+		{
+			TSharedRef<SRigSpacePickerWidget> Dialog =
+				SNew(SRigSpacePickerWidget)
+				.Hierarchy(GetDebuggedHierarchy())
+				.SelectedControl(ControlKey);
+
+			SRigSpacePickerWidget::FResult PickedResult = Dialog->InvokeDialog();
+			if(PickedResult.Reply.IsEventHandled() && PickedResult.Key.IsValid())
+			{
+				if(URigHierarchy* Hierarchy = GetDebuggedHierarchy())
+				{
+					const FRigElementKey SpaceKey = PickedResult.Key;
+					
+					FTransform Transform = Hierarchy->GetGlobalTransform(ControlKey);
+					if(SpaceKey == Hierarchy->GetWorldSpaceSocketKey())
+					{
+						Hierarchy->SwitchToWorldSpace(ControlKey);
+					}
+					else
+					{
+						Hierarchy->SwitchToParent(ControlKey, SpaceKey);
+					}
+					Hierarchy->SetGlobalTransform(ControlKey, Transform);
+				}
+			}
+			
+			break;
+		}
+	}
+}
+
 void SRigHierarchy::HandleParent(const FToolMenuContext& Context)
 {
 	if (UControlRigContextMenuContext* MenuContext = Cast<UControlRigContextMenuContext>(Context.FindByClass(UControlRigContextMenuContext::StaticClass())))
diff --git a/Engine/Plugins/Animation/ControlRig/Source/ControlRigEditor/Private/Editor/SRigHierarchy.h b/Engine/Plugins/Animation/ControlRig/Source/ControlRigEditor/Private/Editor/SRigHierarchy.h
--- a/Engine/Plugins/Animation/ControlRig/Source/ControlRigEditor/Private/Editor/SRigHierarchy.h
+++ b/Engine/Plugins/Animation/ControlRig/Source/ControlRigEditor/Private/Editor/SRigHierarchy.h
@@ -3,11 +3,10 @@
 #pragma once
 
 #include "CoreMinimal.h"
-#include "Widgets/Views/STreeView.h"
-#include "Rigs/RigHierarchy.h"
 #include "EditorUndoClient.h"
 #include "DragAndDrop/GraphNodeDragDropOp.h"
 #include "Engine/SkeletalMesh.h"
+#include "SRigHierarchyTreeView.h"
 #include "SRigHierarchy.generated.h"
 
 class SRigHierarchy;
@@ -18,32 +17,8 @@ class UControlRigBlueprint;
 class UControlRig;
 struct FAssetData;
 class FMenuBuilder;
-class SRigHierarchyItem;
 class UToolMenu;
 
-DECLARE_DELEGATE_RetVal_TwoParams(FName, FOnRenameElement, const FRigElementKey& /*OldKey*/, const FString& /*NewName*/);
-DECLARE_DELEGATE_RetVal_ThreeParams(bool, FOnVerifyElementNameChanged, const FRigElementKey& /*OldKey*/, const FString& /*NewName*/, FText& /*OutErrorMessage*/);
-
-/** An item in the tree */
-class FRigTreeElement : public TSharedFromThis<FRigTreeElement>
-{
-public:
-	FRigTreeElement(const FRigElementKey& InKey, TWeakPtr<SRigHierarchy> InHierarchyHandler);
-public:
-	/** Element Data to display */
-	FRigElementKey Key;
-	bool bIsTransient;
-	TArray<TSharedPtr<FRigTreeElement>> Children;
-
-	TSharedRef<ITableRow> MakeTreeRowWidget(TSharedPtr<FControlRigEditor> InControlRigEditor, const TSharedRef<STableViewBase>& InOwnerTable, TSharedRef<FRigTreeElement> InRigTreeElement, TSharedRef<FUICommandList> InCommandList, TSharedPtr<SRigHierarchy> InHierarchy);
-
-	void RequestRename();
-
-	/** Delegate for when the context menu requests a rename */
-	DECLARE_DELEGATE(FOnRenameRequested);
-	FOnRenameRequested OnRenameRequested;
-};
-
 class FRigElementHierarchyDragDropOp : public FGraphNodeDragDropOp
 {
 public:
@@ -73,30 +48,6 @@ private:
 	TArray<FRigElementKey> Elements;
 };
 
- class SRigHierarchyItem : public STableRow<TSharedPtr<FRigTreeElement>>
-{
-	SLATE_BEGIN_ARGS(SRigHierarchyItem) {}
-		/** Callback when the text is committed. */
-		SLATE_EVENT(FOnRenameElement, OnRenameElement)
-		/** Called whenever the text is changed interactively by the user */
-		SLATE_EVENT(FOnVerifyElementNameChanged, OnVerifyElementNameChanged)
-	SLATE_END_ARGS()
-
-	void Construct(const FArguments& InArgs, TSharedPtr<FControlRigEditor> InControlRigEditor, const TSharedRef<STableViewBase>& OwnerTable, TSharedRef<FRigTreeElement> InRigTreeElement, TSharedRef<FUICommandList> InCommandList, TSharedPtr<SRigHierarchy> InHierarchy);
- 	void OnNameCommitted(const FText& InText, ETextCommit::Type InCommitType) const;
-	bool OnVerifyNameChanged(const FText& InText, FText& OutErrorMessage);
-
-private:
-	TWeakPtr<FRigTreeElement> WeakRigTreeElement;
-	TWeakPtr<FUICommandList> WeakCommandList;
-	TWeakPtr<FControlRigEditor> ControlRigEditor;
-
-	FOnRenameElement OnRenameElement;
-	FOnVerifyElementNameChanged OnVerifyElementNameChanged;
-
-	FText GetName() const;
-};
-
 USTRUCT()
 struct FRigHierarchyImportSettings
 {
@@ -110,53 +61,6 @@ struct FRigHierarchyImportSettings
 	TObjectPtr<USkeletalMesh> Mesh;
 };
 
-class SRigHierarchyTreeView : public STreeView<TSharedPtr<FRigTreeElement>>
-{
-public:
-
-	virtual ~SRigHierarchyTreeView() {}
-
-	virtual FReply OnFocusReceived(const FGeometry& MyGeometry, const FFocusEvent& InFocusEvent) override
-	{
-		FReply Reply = STreeView<TSharedPtr<FRigTreeElement>>::OnFocusReceived(MyGeometry, InFocusEvent);
-
-		LastClickCycles = FPlatformTime::Cycles();
-
-		return Reply;
-	}
-
-	uint32 LastClickCycles = 0;
-
-	/** Save a snapshot of the internal map that tracks item expansion before tree reconstruction */
-	void SaveAndClearSparseItemInfos()
-	{
-		// Only save the info if there is something to save (do not overwrite info with an empty map)
-		if (!SparseItemInfos.IsEmpty())
-		{
-			OldSparseItemInfos = SparseItemInfos;
-		}
-		ClearExpandedItems();
-	}
-
-	/** Restore the expansion infos map from the saved snapshot after tree reconstruction */
-	void RestoreSparseItemInfos(TSharedPtr<FRigTreeElement> ItemPtr)
-	{
-		for (const auto& Pair : OldSparseItemInfos)
-		{
-			if (Pair.Key->Key == ItemPtr->Key)
-			{
-				// the SparseItemInfos now reference the new element, but keep the same expansion state
-				SparseItemInfos.Add(ItemPtr, Pair.Value);
-				break;
-			}
-		}
-	}
-
-private:
-	/** A temporary snapshot of the SparseItemInfos in STreeView, used during SRigHierarchy::RefreshTreeView() */
-	TSparseItemMap OldSparseItemInfos;
-};
-
 /** Widget allowing editing of a control rig's structure */
 class SRigHierarchy : public SCompoundWidget, public FEditorUndoClient
 {
@@ -193,12 +97,6 @@ private:
 	/** Return all selected keys */
 	TArray<FRigElementKey> GetSelectedKeys() const;
 
-	/** Make a row widget for the table */
-	TSharedRef<ITableRow> MakeTableRowWidget(TSharedPtr<FRigTreeElement> InItem, const TSharedRef<STableViewBase>& OwnerTable);
-
-	/** Get children for the tree */
-	void HandleGetChildrenForTree(TSharedPtr<FRigTreeElement> InItem, TArray<TSharedPtr<FRigTreeElement>>& OutChildren);
-
 	/** Check whether we can deleting the selected item(s) */
 	bool CanDeleteItem() const;
 
@@ -257,36 +155,11 @@ private:
 	/** Our owning control rig editor */
 	TWeakPtr<FControlRigEditor> ControlRigEditor;
 
-	/** Flatten when text filtering is active */
-	bool bFlattenHierarchyOnFilter;
-
-	/** Hide parents when text filtering is active */
-	bool bHideParentsOnFilter;
-
-	/** Whether or not to show imported bones in the hierarchy */
-	bool bShowImportedBones;
-
-	/** Whether or not to show bones in the hierarchy */
-	bool bShowBones;
-
-	/** Whether or not to show controls in the hierarchy */
-	bool bShowControls;
-
-	/** Whether or not to show spaces in the hierarchy */
-	bool bShowNulls;
-
-	/** Whether or not to show rigidbodies in the hierarchy */
-	bool bShowRigidBodies;
-
-	/** Whether or not to show sockets in the hierarchy */
-	bool bShowSockets;
-
-	/** Whether or not to show the static or dynamic hierarchy */
-	bool bShowDynamicHierarchy;
+	FRigTreeDisplaySettings DisplaySettings;
+	const FRigTreeDisplaySettings& GetDisplaySettings() const { return DisplaySettings; } 
 
 	/** Search box widget */
 	TSharedPtr<SSearchBox> FilterBox;
-	FText FilterText;
 
 	EVisibility IsToolbarVisible() const;
 	EVisibility IsSearchbarVisible() const;
@@ -296,15 +169,6 @@ private:
 	/** Tree view widget */
 	TSharedPtr<SRigHierarchyTreeView> TreeView;
 
-	/** Backing array for tree view */
-	TArray<TSharedPtr<FRigTreeElement>> RootElements;
-
-	/** A map for looking up items based on their key */
-	TMap<FRigElementKey, TSharedPtr<FRigTreeElement>> ElementMap;
-
-	/** A map for looking up a parent based on their key */
-	TMap<FRigElementKey, FRigElementKey> ParentMap;
-
 	TWeakObjectPtr<UControlRigBlueprint> ControlRigBlueprint;
 	TWeakObjectPtr<UControlRig> ControlRigBeingDebuggedPtr;
 	
@@ -320,7 +184,7 @@ private:
 
 	URigHierarchy* GetHierarchy() const;
 	URigHierarchy* GetDebuggedHierarchy() const;
-	URigHierarchy* GetHierarchyForTopology() const;
+	const URigHierarchy* GetHierarchyForTopology() const;
 
 	void ImportHierarchy(const FAssetData& InAssetData);
 	void CreateImportMenu(FMenuBuilder& MenuBuilder);
@@ -337,6 +201,7 @@ private:
 	void HandleControlBoneOrSpaceTransform();
 	void HandleUnparent();
 	bool FindClosestBone(const FVector& Point, FName& OutRigElementName, FTransform& OutGlobalTransform) const;
+	void HandleTestSpaceSwitching();
 
 	void HandleParent(const FToolMenuContext& Context);
 	void HandleAlign(const FToolMenuContext& Context);
@@ -344,8 +209,6 @@ private:
 
 	FName CreateUniqueName(const FName& InBaseName, ERigElementType InElementType) const;
 
-	void SetExpansionRecursive(TSharedPtr<FRigTreeElement> InElement, bool bTowardsParent, bool bShouldBeExpanded);
-
 	void ClearDetailPanel() const;
 
 	bool bIsChangingRigHierarchy;
@@ -354,16 +217,9 @@ private:
 	void HandleRefreshEditorFromBlueprint(UControlRigBlueprint* InBlueprint);
 	void HandleSetObjectBeingDebugged(UObject* InObject);
 
-	static TSharedPtr<FRigTreeElement> FindElement(const FRigElementKey& InElementKey, TSharedPtr<FRigTreeElement> CurrentItem);
-	bool AddElement(FRigElementKey InKey, FRigElementKey InParentKey = FRigElementKey(), const bool bIgnoreTextFilter = false);
-	bool AddElement(const FRigBaseElement* InElement, const bool bIgnoreTextFilter = false);
-	void AddSpacerElement();
-	bool ReparentElement(FRigElementKey InKey, FRigElementKey InParentKey);
-	bool RemoveElement(FRigElementKey InKey);
-
 public:
-	FName RenameElement(const FRigElementKey& OldKey, const FString& NewName);
-	bool OnVerifyNameChanged(const FRigElementKey& OldKey, const FString& NewName, FText& OutErrorMessage);
+	FName HandleRenameElement(const FRigElementKey& OldKey, const FString& NewName);
+	bool HandleVerifyNameChanged(const FRigElementKey& OldKey, const FString& NewName, FText& OutErrorMessage);
 
 	friend class FRigTreeElement;
 	friend class SRigHierarchyItem;
diff --git a/Engine/Plugins/Animation/ControlRig/Source/ControlRigEditor/Private/Editor/SRigHierarchyTreeView.cpp b/Engine/Plugins/Animation/ControlRig/Source/ControlRigEditor/Private/Editor/SRigHierarchyTreeView.cpp
new file mode 100644
--- /dev/null
+++ b/Engine/Plugins/Animation/ControlRig/Source/ControlRigEditor/Private/Editor/SRigHierarchyTreeView.cpp
@@ -0,0 +1,649 @@
+// Copyright Epic Games, Inc. All Rights Reserved.
+
+#include "SRigHierarchyTreeView.h"
+#include "EditorStyleSet.h"
+#include "Widgets/Layout/SSpacer.h"
+#include "Widgets/Images/SImage.h"
+#include "Widgets/Text/STextBlock.h"
+#include "PropertyCustomizationHelpers.h"
+#include "Framework/Application/SlateApplication.h"
+#include "Editor/EditorEngine.h"
+#include "HelperUtil.h"
+#include "Widgets/Text/SInlineEditableTextBlock.h"
+#include "ControlRig.h"
+#include "ControlRigEditorStyle.h"
+
+#define LOCTEXT_NAMESPACE "SRigHierarchyTreeView"
+
+//////////////////////////////////////////////////////////////
+/// FRigTreeDelegates
+///////////////////////////////////////////////////////////
+FRigTreeDisplaySettings FRigTreeDelegates::DefaultDisplaySettings;
+
+//////////////////////////////////////////////////////////////
+/// FRigTreeElement
+///////////////////////////////////////////////////////////
+FRigTreeElement::FRigTreeElement(const FRigElementKey& InKey, TWeakPtr<SRigHierarchyTreeView> InTreeView, bool InSupportsRename)
+{
+	Key = InKey;
+	bIsTransient = false;
+	bSupportsRename = InSupportsRename;
+
+	if(InTreeView.IsValid())
+	{
+		if(const URigHierarchy* Hierarchy = InTreeView.Pin()->GetRigTreeDelegates().GetHierarchy())
+		{
+			if(const FRigControlElement* ControlElement = Hierarchy->Find<FRigControlElement>(Key))
+			{
+				bIsTransient = ControlElement->Settings.bIsTransientControl;
+			}
+		}
+	}
+}
+
+
+TSharedRef<ITableRow> FRigTreeElement::MakeTreeRowWidget(const TSharedRef<STableViewBase>& InOwnerTable, TSharedRef<FRigTreeElement> InRigTreeElement, TSharedPtr<SRigHierarchyTreeView> InTreeView)
+{
+	return SNew(SRigHierarchyItem, InOwnerTable, InRigTreeElement, InTreeView);
+}
+
+void FRigTreeElement::RequestRename()
+{
+	if(bSupportsRename)
+	{
+		OnRenameRequested.ExecuteIfBound();
+	}
+}
+
+//////////////////////////////////////////////////////////////
+/// SRigHierarchyItem
+///////////////////////////////////////////////////////////
+void SRigHierarchyItem::Construct(const FArguments& InArgs, const TSharedRef<STableViewBase>& OwnerTable, TSharedRef<FRigTreeElement> InRigTreeElement, TSharedPtr<SRigHierarchyTreeView> InTreeView)
+{
+	WeakRigTreeElement = InRigTreeElement;
+	Delegates = InTreeView->GetRigTreeDelegates();
+
+	if (!InRigTreeElement->Key.IsValid())
+	{
+		STableRow<TSharedPtr<FRigTreeElement>>::Construct(
+			STableRow<TSharedPtr<FRigTreeElement>>::FArguments()
+			.ShowSelection(false)
+			.OnCanAcceptDrop(Delegates.OnCanAcceptDrop)
+			.OnAcceptDrop(Delegates.OnAcceptDrop)
+			.Content()
+			[
+				SNew(SVerticalBox)
+				+ SVerticalBox::Slot()
+				.FillHeight(200.f)
+				[
+					SNew(SSpacer)
+				]
+			], OwnerTable);
+		return;
+	}
+
+	TSharedPtr< SInlineEditableTextBlock > InlineWidget;
+
+	const FSlateBrush* Brush = GetBrushForElementType(Delegates.GetHierarchy(), InRigTreeElement->Key);
+
+	STableRow<TSharedPtr<FRigTreeElement>>::Construct(
+		STableRow<TSharedPtr<FRigTreeElement>>::FArguments()
+		.OnDragDetected(Delegates.OnDragDetected)
+		.OnCanAcceptDrop(Delegates.OnCanAcceptDrop)
+		.OnAcceptDrop(Delegates.OnAcceptDrop)
+		.ShowWires(true)
+		.Content()
+		[
+			SNew(SHorizontalBox)
+			+ SHorizontalBox::Slot()
+			.MaxWidth(18)
+			.FillWidth(1.0)
+			.HAlign(HAlign_Left)
+			.VAlign(VAlign_Center)
+			[
+				SNew(SImage)
+				.Image(Brush)
+			]
+			+ SHorizontalBox::Slot()
+			.AutoWidth()
+			.VAlign(VAlign_Center)
+			[
+				SAssignNew(InlineWidget, SInlineEditableTextBlock)
+				.Text(this, &SRigHierarchyItem::GetName)
+				.OnVerifyTextChanged(this, &SRigHierarchyItem::OnVerifyNameChanged)
+				.OnTextCommitted(this, &SRigHierarchyItem::OnNameCommitted)
+				.MultiLine(false)
+			]
+		], OwnerTable);
+
+	InRigTreeElement->OnRenameRequested.BindSP(InlineWidget.Get(), &SInlineEditableTextBlock::EnterEditingMode);
+}
+
+FText SRigHierarchyItem::GetName() const
+{
+	if(WeakRigTreeElement.Pin()->bIsTransient)
+	{
+		static const FText TemporaryControl = FText::FromString(TEXT("Temporary Control"));
+		return TemporaryControl;
+	}
+	return (FText::FromName(WeakRigTreeElement.Pin()->Key.Name));
+}
+
+//////////////////////////////////////////////////////////////
+/// SRigHierarchyTreeView
+///////////////////////////////////////////////////////////
+
+void SRigHierarchyTreeView::Construct(const FArguments& InArgs)
+{
+	Delegates = InArgs._RigTreeDelegates;
+
+	STreeView<TSharedPtr<FRigTreeElement>>::FArguments SuperArgs;
+	SuperArgs.TreeItemsSource(&RootElements);
+	SuperArgs.SelectionMode(ESelectionMode::Multi);
+	SuperArgs.OnGenerateRow(this, &SRigHierarchyTreeView::MakeTableRowWidget);
+	SuperArgs.OnGetChildren(this, &SRigHierarchyTreeView::HandleGetChildrenForTree);
+	SuperArgs.OnSelectionChanged(FOnRigTreeSelectionChanged::CreateRaw(&Delegates, &FRigTreeDelegates::HandleSelectionChanged));
+	SuperArgs.OnContextMenuOpening(Delegates.OnContextMenuOpening);
+	SuperArgs.OnMouseButtonClick(Delegates.OnMouseButtonClick);
+	SuperArgs.OnMouseButtonDoubleClick(Delegates.OnMouseButtonDoubleClick);
+	SuperArgs.OnSetExpansionRecursive(Delegates.OnSetExpansionRecursive);
+	SuperArgs.HighlightParentNodesForSelection(true);
+	SuperArgs.ItemHeight(24);
+	
+	STreeView<TSharedPtr<FRigTreeElement>>::Construct(SuperArgs);
+}
+
+TSharedPtr<FRigTreeElement> SRigHierarchyTreeView::FindElement(const FRigElementKey& InElementKey, TSharedPtr<FRigTreeElement> CurrentItem)
+{
+	if (CurrentItem->Key == InElementKey)
+	{
+		return CurrentItem;
+	}
+
+	for (int32 ChildIndex = 0; ChildIndex < CurrentItem->Children.Num(); ++ChildIndex)
+	{
+		TSharedPtr<FRigTreeElement> Found = FindElement(InElementKey, CurrentItem->Children[ChildIndex]);
+		if (Found.IsValid())
+		{
+			return Found;
+		}
+	}
+
+	return TSharedPtr<FRigTreeElement>();
+}
+
+bool SRigHierarchyTreeView::AddElement(FRigElementKey InKey, FRigElementKey InParentKey)
+{
+	if(ElementMap.Contains(InKey))
+	{
+		return false;
+	}
+
+	const FRigTreeDisplaySettings& Settings = Delegates.GetDisplaySettings();
+	const bool bSupportsRename = Delegates.OnRenameElement.IsBound();
+
+	const FString FilteredString = Settings.FilterText.ToString();
+	if (FilteredString.IsEmpty() || !InKey.IsValid())
+	{
+		TSharedPtr<FRigTreeElement> NewItem = MakeShared<FRigTreeElement>(InKey, SharedThis(this), bSupportsRename);
+
+		if (InKey.IsValid())
+		{
+			ElementMap.Add(InKey, NewItem);
+			if (InParentKey)
+			{
+				ParentMap.Add(InKey, InParentKey);
+			}
+
+			if (InParentKey)
+			{
+				TSharedPtr<FRigTreeElement>* FoundItem = ElementMap.Find(InParentKey);
+				check(FoundItem);
+				FoundItem->Get()->Children.Add(NewItem);
+			}
+			else
+			{
+				RootElements.Add(NewItem);
+			}
+		}
+		else
+		{
+			RootElements.Add(NewItem);
+		}
+	}
+	else
+	{
+		FString FilteredStringUnderScores = FilteredString.Replace(TEXT(" "), TEXT("_"));
+		if (InKey.Name.ToString().Contains(FilteredString) || InKey.Name.ToString().Contains(FilteredStringUnderScores))	
+		{
+			TSharedPtr<FRigTreeElement> NewItem = MakeShared<FRigTreeElement>(InKey, SharedThis(this), bSupportsRename);
+			ElementMap.Add(InKey, NewItem);
+			RootElements.Add(NewItem);
+		}
+	}
+
+	return true;
+}
+
+bool SRigHierarchyTreeView::AddElement(const FRigBaseElement* InElement)
+{
+	check(InElement);
+	
+	if (ElementMap.Contains(InElement->GetKey()))
+	{
+		return false;
+	}
+
+	const FRigTreeDisplaySettings& Settings = Delegates.GetDisplaySettings();
+
+	switch(InElement->GetType())
+	{
+		case ERigElementType::Bone:
+		{
+			if(!Settings.bShowBones)
+			{
+				return false;
+			}
+
+			const FRigBoneElement* BoneElement = CastChecked<FRigBoneElement>(InElement);
+			if (!Settings.bShowImportedBones && BoneElement->BoneType == ERigBoneType::Imported)
+			{
+				return false;
+			}
+			break;
+		}
+		case ERigElementType::Null:
+		{
+			if(!Settings.bShowNulls)
+			{
+				return false;
+			}
+			break;
+		}
+		case ERigElementType::Control:
+		{
+			if(!Settings.bShowControls)
+			{
+				return false;
+			}
+			break;
+		}
+		case ERigElementType::RigidBody:
+		{
+			if(!Settings.bShowRigidBodies)
+			{
+				return false;
+			}
+			break;
+		}
+		case ERigElementType::Socket:
+		{
+			if(!Settings.bShowSockets)
+			{
+				return false;
+			}
+			break;
+		}
+		case ERigElementType::Curve:
+		{
+			return false;
+		}
+		default:
+		{
+			break;
+		}
+	}
+
+	if(!AddElement(InElement->GetKey(), FRigElementKey()))
+	{
+		return false;
+	}
+
+	if (ElementMap.Contains(InElement->GetKey()))
+	{
+		if(const URigHierarchy* Hierarchy = Delegates.GetHierarchy())
+		{
+			FRigElementKey ParentKey = Hierarchy->GetFirstParent(InElement->GetKey());
+			if(Settings.bShowDynamicHierarchy)
+			{
+				TArray<FRigElementWeight> ParentWeights = Hierarchy->GetParentWeightArray(InElement->GetKey());
+				if(ParentWeights.Num() > 0)
+				{
+					TArray<FRigElementKey> ParentKeys = Hierarchy->GetParents(InElement->GetKey());
+					check(ParentKeys.Num() == ParentWeights.Num());
+					for(int32 ParentIndex=0;ParentIndex<ParentKeys.Num();ParentIndex++)
+					{
+						if(ParentWeights[ParentIndex].IsAlmostZero())
+						{
+							continue;
+						}
+						ParentKey = ParentKeys[ParentIndex];
+						break;
+					}
+				}
+			}
+			if (ParentKey.IsValid())
+			{
+				if(const FRigBaseElement* ParentElement = Hierarchy->Find(ParentKey))
+				{
+					AddElement(ParentElement);
+
+					if(ElementMap.Contains(ParentKey))
+					{
+						ReparentElement(InElement->GetKey(), ParentKey);
+					}
+				}
+			}
+		}
+	}
+
+	return true;
+}
+
+void SRigHierarchyTreeView::AddSpacerElement()
+{
+	AddElement(FRigElementKey(), FRigElementKey());
+}
+
+bool SRigHierarchyTreeView::ReparentElement(FRigElementKey InKey, FRigElementKey InParentKey)
+{
+	if (!InKey.IsValid() || InKey == InParentKey)
+	{
+		return false;
+	}
+
+	const FRigTreeDisplaySettings& Settings = Delegates.GetDisplaySettings();
+
+	TSharedPtr<FRigTreeElement>* FoundItem = ElementMap.Find(InKey);
+	if (FoundItem == nullptr)
+	{
+		return false;
+	}
+
+	if (!Settings.FilterText.IsEmpty() && Settings.bFlattenHierarchyOnFilter)
+	{
+		return false;
+	}
+
+	if (const FRigElementKey* ExistingParentKey = ParentMap.Find(InKey))
+	{
+		if (*ExistingParentKey == InParentKey)
+		{
+			return false;
+		}
+
+		if (TSharedPtr<FRigTreeElement>* ExistingParent = ElementMap.Find(*ExistingParentKey))
+		{
+			(*ExistingParent)->Children.Remove(*FoundItem);
+		}
+
+		ParentMap.Remove(InKey);
+	}
+	else
+	{
+		if (!InParentKey.IsValid())
+		{
+			return false;
+		}
+
+		RootElements.Remove(*FoundItem);
+	}
+
+	if (InParentKey)
+	{
+		ParentMap.Add(InKey, InParentKey);
+
+		TSharedPtr<FRigTreeElement>* FoundParent = ElementMap.Find(InParentKey);
+		check(FoundParent);
+		FoundParent->Get()->Children.Add(*FoundItem);
+	}
+	else
+	{
+		RootElements.Add(*FoundItem);
+	}
+
+	return true;
+}
+
+bool SRigHierarchyTreeView::RemoveElement(FRigElementKey InKey)
+{
+	TSharedPtr<FRigTreeElement>* FoundItem = ElementMap.Find(InKey);
+	if (FoundItem == nullptr)
+	{
+		return false;
+	}
+
+	ReparentElement(InKey, FRigElementKey());
+
+	RootElements.Remove(*FoundItem);
+	return ElementMap.Remove(InKey) > 0;
+}
+
+void SRigHierarchyTreeView::RefreshTreeView(bool bRebuildContent)
+{
+		TMap<FRigElementKey, bool> ExpansionState;
+
+	if(bRebuildContent)
+	{
+		for (TPair<FRigElementKey, TSharedPtr<FRigTreeElement>> Pair : ElementMap)
+		{
+			ExpansionState.FindOrAdd(Pair.Key) = IsItemExpanded(Pair.Value);
+		}
+
+		// internally save expansion states before rebuilding the tree, so the states can be restored later
+		SaveAndClearSparseItemInfos();
+
+		RootElements.Reset();
+		ElementMap.Reset();
+		ParentMap.Reset();
+	}
+
+	if(bRebuildContent)
+	{
+		const URigHierarchy* Hierarchy = Delegates.GetHierarchy();
+		if(Hierarchy)
+		{
+			Hierarchy->Traverse([&](FRigBaseElement* Element, bool& bContinue)
+			{
+				AddElement(Element);
+				bContinue = true;
+			});
+
+			// expand all elements upon the initial construction of the tree
+			if (ExpansionState.Num() == 0)
+			{
+				for (TSharedPtr<FRigTreeElement> RootElement : RootElements)
+				{
+					SetExpansionRecursive(RootElement, false, true);
+				}
+			}
+
+			for (const auto& Pair : ElementMap)
+			{
+				RestoreSparseItemInfos(Pair.Value);
+			}
+
+			if (RootElements.Num() > 0)
+			{
+				AddSpacerElement();
+			}
+		}
+	}
+	else
+	{
+		if (RootElements.Num()> 0)
+		{
+			// elements may be added at the end of the list after a spacer element
+			// we need to remove the spacer element and re-add it at the end
+			RootElements.RemoveAll([](TSharedPtr<FRigTreeElement> InElement)
+			{
+				return InElement.Get()->Key == FRigElementKey();
+			});
+			AddSpacerElement();
+		}
+	}
+
+	RequestTreeRefresh();
+	{
+		ClearSelection();
+
+		if(const URigHierarchy* Hierarchy = Delegates.GetHierarchy())
+		{
+			TArray<FRigElementKey> Selection = Hierarchy->GetSelectedKeys();
+			for (const FRigElementKey& Key : Selection)
+			{
+				for (int32 RootIndex = 0; RootIndex < RootElements.Num(); ++RootIndex)
+				{
+					TSharedPtr<FRigTreeElement> Found = FindElement(Key, RootElements[RootIndex]);
+					if (Found.IsValid())
+					{
+						SetItemSelection(Found, true, ESelectInfo::OnNavigation);
+					}
+				}
+			}
+		}
+	}
+}
+
+void SRigHierarchyTreeView::SetExpansionRecursive(TSharedPtr<FRigTreeElement> InElement, bool bTowardsParent,
+	bool bShouldBeExpanded)
+{
+	SetItemExpansion(InElement, bShouldBeExpanded);
+
+	if (bTowardsParent)
+	{
+		if (const FRigElementKey* ParentKey = ParentMap.Find(InElement->Key))
+		{
+			if (TSharedPtr<FRigTreeElement>* ParentItem = ElementMap.Find(*ParentKey))
+			{
+				SetExpansionRecursive(*ParentItem, bTowardsParent, bShouldBeExpanded);
+			}
+		}
+	}
+	else
+	{
+		for (int32 ChildIndex = 0; ChildIndex < InElement->Children.Num(); ++ChildIndex)
+		{
+			SetExpansionRecursive(InElement->Children[ChildIndex], bTowardsParent, bShouldBeExpanded);
+		}
+	}
+}
+
+TSharedRef<ITableRow> SRigHierarchyTreeView::MakeTableRowWidget(TSharedPtr<FRigTreeElement> InItem,
+	const TSharedRef<STableViewBase>& OwnerTable)
+{
+	return InItem->MakeTreeRowWidget(OwnerTable, InItem.ToSharedRef(), SharedThis(this));
+}
+
+void SRigHierarchyTreeView::HandleGetChildrenForTree(TSharedPtr<FRigTreeElement> InItem,
+	TArray<TSharedPtr<FRigTreeElement>>& OutChildren)
+{
+	OutChildren = InItem.Get()->Children;
+}
+
+TArray<FRigElementKey> SRigHierarchyTreeView::GetSelectedKeys() const
+{
+	TArray<FRigElementKey> Keys;
+	TArray<TSharedPtr<FRigTreeElement>> SelectedElements = GetSelectedItems();
+	for(const TSharedPtr<FRigTreeElement>& SelectedElement : SelectedElements)
+	{
+		Keys.Add(SelectedElement->Key);
+	}
+	return Keys;
+}
+
+bool SRigHierarchyItem::OnVerifyNameChanged(const FText& InText, FText& OutErrorMessage)
+{
+	const FString NewName = InText.ToString();
+	const FRigElementKey OldKey = WeakRigTreeElement.Pin()->Key;
+	return Delegates.HandleVerifyElementNameChanged(OldKey, NewName, OutErrorMessage);
+}
+
+const FSlateBrush* SRigHierarchyItem::GetBrushForElementType(const URigHierarchy* InHierarchy, const FRigElementKey& InKey)
+{
+	const FSlateBrush* Brush = nullptr;
+	switch (InKey.Type)
+	{
+		case ERigElementType::Control:
+		{
+			Brush = FControlRigEditorStyle::Get().GetBrush("ControlRig.Tree.Control");
+			break;
+		}
+		case ERigElementType::Null:
+		{
+			Brush = FControlRigEditorStyle::Get().GetBrush("ControlRig.Tree.Null");
+			break;
+		}
+		case ERigElementType::Bone:
+		{
+			ERigBoneType BoneType = ERigBoneType::User;
+
+			if(InHierarchy)
+			{
+				const FRigBoneElement* BoneElement = InHierarchy->Find<FRigBoneElement>(InKey);
+				if(BoneElement)
+				{
+					BoneType = BoneElement->BoneType;
+				}
+			}
+
+			switch (BoneType)
+			{
+				case ERigBoneType::Imported:
+				{
+					Brush = FControlRigEditorStyle::Get().GetBrush("ControlRig.Tree.BoneImported");
+					break;
+				}
+				case ERigBoneType::User:
+				default:
+				{
+					Brush = FControlRigEditorStyle::Get().GetBrush("ControlRig.Tree.BoneUser");
+					break;
+				}
+			}
+
+			break;
+		}
+		case ERigElementType::RigidBody:
+		{
+			Brush = FControlRigEditorStyle::Get().GetBrush("ControlRig.Tree.RigidBody");
+			break;
+		}
+		case ERigElementType::Socket:
+		{
+			Brush = FControlRigEditorStyle::Get().GetBrush("ControlRig.Tree.Socket");
+			break;
+		}
+		default:
+		{
+			break;
+		}
+	}
+
+	return Brush;
+}
+
+void SRigHierarchyItem::OnNameCommitted(const FText& InText, ETextCommit::Type InCommitType) const
+{
+	// for now only allow enter
+	// because it is important to keep the unique names per pose
+	if (InCommitType == ETextCommit::OnEnter)
+	{
+		FString NewName = InText.ToString();
+		const FRigElementKey OldKey = WeakRigTreeElement.Pin()->Key;
+
+		const FName NewSanitizedName = Delegates.HandleRenameElement(OldKey, NewName);
+		if (NewSanitizedName.IsNone())
+		{
+			return;
+		}
+		NewName = NewSanitizedName.ToString();
+
+		if (WeakRigTreeElement.IsValid())
+		{
+			WeakRigTreeElement.Pin()->Key.Name = *NewName;
+		}
+	}
+}
+
+#undef LOCTEXT_NAMESPACE
\ No newline at end of file
diff --git a/Engine/Plugins/Animation/ControlRig/Source/ControlRigEditor/Private/Editor/SRigHierarchyTreeView.h b/Engine/Plugins/Animation/ControlRig/Source/ControlRigEditor/Private/Editor/SRigHierarchyTreeView.h
new file mode 100644
--- /dev/null
+++ b/Engine/Plugins/Animation/ControlRig/Source/ControlRigEditor/Private/Editor/SRigHierarchyTreeView.h
@@ -0,0 +1,259 @@
+// Copyright Epic Games, Inc. All Rights Reserved.
+
+#pragma once
+
+#include "CoreMinimal.h"
+#include "Widgets/Views/STreeView.h"
+#include "Rigs/RigHierarchy.h"
+
+class SRigHierarchyTreeView;
+class SRigHierarchyItem;
+class FRigTreeElement;
+
+struct CONTROLRIGEDITOR_API FRigTreeDisplaySettings
+{
+	FRigTreeDisplaySettings()
+	{
+		FilterText = FText();
+
+		bFlattenHierarchyOnFilter = false;
+		bHideParentsOnFilter = false;
+		bShowImportedBones = true;
+		bShowBones = true;
+		bShowControls = true;
+		bShowNulls = true;
+		bShowRigidBodies = true;
+		bShowSockets = true;
+		bShowDynamicHierarchy = false;
+	}
+	
+	FText FilterText;
+
+	/** Flatten when text filtering is active */
+	bool bFlattenHierarchyOnFilter;
+
+	/** Hide parents when text filtering is active */
+	bool bHideParentsOnFilter;
+
+	/** Whether or not to show imported bones in the hierarchy */
+	bool bShowImportedBones;
+
+	/** Whether or not to show bones in the hierarchy */
+	bool bShowBones;
+
+	/** Whether or not to show controls in the hierarchy */
+	bool bShowControls;
+
+	/** Whether or not to show spaces in the hierarchy */
+	bool bShowNulls;
+
+	/** Whether or not to show rigidbodies in the hierarchy */
+	bool bShowRigidBodies;
+
+	/** Whether or not to show sockets in the hierarchy */
+	bool bShowSockets;
+
+	/** Whether or not to show the static or dynamic hierarchy */
+	bool bShowDynamicHierarchy;
+};
+
+DECLARE_DELEGATE_RetVal(const URigHierarchy*, FOnGetRigTreeHierarchy);
+DECLARE_DELEGATE_RetVal(const FRigTreeDisplaySettings&, FOnGetRigTreeDisplaySettings);
+DECLARE_DELEGATE_RetVal_TwoParams(FName, FOnRigTreeRenameElement, const FRigElementKey& /*OldKey*/, const FString& /*NewName*/);
+DECLARE_DELEGATE_RetVal_ThreeParams(bool, FOnRigTreeVerifyElementNameChanged, const FRigElementKey& /*OldKey*/, const FString& /*NewName*/, FText& /*OutErrorMessage*/);
+
+typedef STableRow<TSharedPtr<FRigTreeElement>>::FOnCanAcceptDrop FOnRigTreeCanAcceptDrop;
+typedef STableRow<TSharedPtr<FRigTreeElement>>::FOnAcceptDrop FOnRigTreeAcceptDrop;
+typedef STreeView<TSharedPtr<FRigTreeElement>>::FOnSelectionChanged FOnRigTreeSelectionChanged;
+typedef STreeView<TSharedPtr<FRigTreeElement>>::FOnMouseButtonClick FOnRigTreeMouseButtonClick;
+typedef STreeView<TSharedPtr<FRigTreeElement>>::FOnMouseButtonDoubleClick FOnRigTreeMouseButtonDoubleClick;
+typedef STreeView<TSharedPtr<FRigTreeElement>>::FOnSetExpansionRecursive FOnRigTreeSetExpansionRecursive;
+
+struct CONTROLRIGEDITOR_API FRigTreeDelegates
+{
+	FOnGetRigTreeHierarchy OnGetHierarchy;
+	FOnGetRigTreeDisplaySettings OnGetDisplaySettings;
+	FOnRigTreeRenameElement OnRenameElement;
+	FOnRigTreeVerifyElementNameChanged OnVerifyElementNameChanged;
+	FOnDragDetected OnDragDetected;
+	FOnRigTreeCanAcceptDrop OnCanAcceptDrop;
+	FOnRigTreeAcceptDrop OnAcceptDrop;
+	FOnRigTreeSelectionChanged OnSelectionChanged;
+	FOnContextMenuOpening OnContextMenuOpening;
+	FOnRigTreeMouseButtonClick OnMouseButtonClick;
+	FOnRigTreeMouseButtonDoubleClick OnMouseButtonDoubleClick;
+	FOnRigTreeSetExpansionRecursive OnSetExpansionRecursive;
+
+	FRigTreeDelegates()
+	{
+		bIsChangingRigHierarchy = false;
+	}
+
+	FORCEINLINE const URigHierarchy* GetHierarchy() const
+	{
+		if(OnGetHierarchy.IsBound())
+		{
+			return OnGetHierarchy.Execute();
+		}
+		return nullptr;
+	}
+
+	FORCEINLINE const FRigTreeDisplaySettings& GetDisplaySettings() const
+	{
+		if(OnGetDisplaySettings.IsBound())
+		{
+			return OnGetDisplaySettings.Execute();
+		}
+		return DefaultDisplaySettings;
+	}
+	
+	FORCEINLINE FName HandleRenameElement(const FRigElementKey& OldKey, const FString& NewName) const
+	{
+		if(OnRenameElement.IsBound())
+		{
+			return OnRenameElement.Execute(OldKey, NewName);
+		}
+		return OldKey.Name;
+	}
+	
+	FORCEINLINE bool HandleVerifyElementNameChanged(const FRigElementKey& OldKey, const FString& NewName, FText& OutErrorMessage) const
+	{
+		if(OnVerifyElementNameChanged.IsBound())
+		{
+			return OnVerifyElementNameChanged.Execute(OldKey, NewName, OutErrorMessage);
+		}
+		return false;
+	}
+
+	FORCEINLINE void HandleSelectionChanged(TSharedPtr<FRigTreeElement> Selection, ESelectInfo::Type SelectInfo)
+	{
+		if(bIsChangingRigHierarchy)
+		{
+			return;
+		}
+		TGuardValue<bool> Guard(bIsChangingRigHierarchy, true);
+		OnSelectionChanged.ExecuteIfBound(Selection, SelectInfo);
+	}
+
+	static FRigTreeDisplaySettings DefaultDisplaySettings;
+	bool bIsChangingRigHierarchy;
+};
+
+
+/** An item in the tree */
+class FRigTreeElement : public TSharedFromThis<FRigTreeElement>
+{
+public:
+	FRigTreeElement(const FRigElementKey& InKey, TWeakPtr<SRigHierarchyTreeView> InTreeView, bool InSupportsRename);
+public:
+	/** Element Data to display */
+	FRigElementKey Key;
+	bool bIsTransient;
+	bool bSupportsRename;
+	TArray<TSharedPtr<FRigTreeElement>> Children;
+
+	TSharedRef<ITableRow> MakeTreeRowWidget(const TSharedRef<STableViewBase>& InOwnerTable, TSharedRef<FRigTreeElement> InRigTreeElement, TSharedPtr<SRigHierarchyTreeView> InTreeView);
+
+	void RequestRename();
+
+	/** Delegate for when the context menu requests a rename */
+	DECLARE_DELEGATE(FOnRenameRequested);
+	FOnRenameRequested OnRenameRequested;
+};
+
+class SRigHierarchyItem : public STableRow<TSharedPtr<FRigTreeElement>>
+{
+public:
+	
+	void Construct(const FArguments& InArgs, const TSharedRef<STableViewBase>& OwnerTable, TSharedRef<FRigTreeElement> InRigTreeElement, TSharedPtr<SRigHierarchyTreeView> InTreeView);
+ 	void OnNameCommitted(const FText& InText, ETextCommit::Type InCommitType) const;
+	bool OnVerifyNameChanged(const FText& InText, FText& OutErrorMessage);
+	static const FSlateBrush* GetBrushForElementType(const URigHierarchy* InHierarchy, const FRigElementKey& InKey);
+
+private:
+	TWeakPtr<FRigTreeElement> WeakRigTreeElement;
+ 	FRigTreeDelegates Delegates;
+
+	FText GetName() const;
+};
+
+class SRigHierarchyTreeView : public STreeView<TSharedPtr<FRigTreeElement>>
+{
+public:
+
+	SLATE_BEGIN_ARGS(SRigHierarchyTreeView) {}
+		SLATE_ARGUMENT(FRigTreeDelegates, RigTreeDelegates)
+	SLATE_END_ARGS()
+
+	void Construct(const FArguments& InArgs);
+	virtual ~SRigHierarchyTreeView() {}
+
+	virtual FReply OnFocusReceived(const FGeometry& MyGeometry, const FFocusEvent& InFocusEvent) override
+	{
+		FReply Reply = STreeView<TSharedPtr<FRigTreeElement>>::OnFocusReceived(MyGeometry, InFocusEvent);
+
+		LastClickCycles = FPlatformTime::Cycles();
+
+		return Reply;
+	}
+
+	uint32 LastClickCycles = 0;
+
+	/** Save a snapshot of the internal map that tracks item expansion before tree reconstruction */
+	void SaveAndClearSparseItemInfos()
+	{
+		// Only save the info if there is something to save (do not overwrite info with an empty map)
+		if (!SparseItemInfos.IsEmpty())
+		{
+			OldSparseItemInfos = SparseItemInfos;
+		}
+		ClearExpandedItems();
+	}
+
+	/** Restore the expansion infos map from the saved snapshot after tree reconstruction */
+	void RestoreSparseItemInfos(TSharedPtr<FRigTreeElement> ItemPtr)
+	{
+		for (const auto& Pair : OldSparseItemInfos)
+		{
+			if (Pair.Key->Key == ItemPtr->Key)
+			{
+				// the SparseItemInfos now reference the new element, but keep the same expansion state
+				SparseItemInfos.Add(ItemPtr, Pair.Value);
+				break;
+			}
+		}
+	}
+
+	static TSharedPtr<FRigTreeElement> FindElement(const FRigElementKey& InElementKey, TSharedPtr<FRigTreeElement> CurrentItem);
+	bool AddElement(FRigElementKey InKey, FRigElementKey InParentKey = FRigElementKey());
+	bool AddElement(const FRigBaseElement* InElement);
+	void AddSpacerElement();
+	bool ReparentElement(FRigElementKey InKey, FRigElementKey InParentKey);
+	bool RemoveElement(FRigElementKey InKey);
+	void RefreshTreeView(bool bRebuildContent = true);
+	void SetExpansionRecursive(TSharedPtr<FRigTreeElement> InElement, bool bTowardsParent, bool bShouldBeExpanded);
+	TSharedRef<ITableRow> MakeTableRowWidget(TSharedPtr<FRigTreeElement> InItem, const TSharedRef<STableViewBase>& OwnerTable);
+	void HandleGetChildrenForTree(TSharedPtr<FRigTreeElement> InItem, TArray<TSharedPtr<FRigTreeElement>>& OutChildren);
+
+	TArray<FRigElementKey> GetSelectedKeys() const;
+	const TArray<TSharedPtr<FRigTreeElement>>& GetRootElements() const { return RootElements; }
+	FRigTreeDelegates& GetRigTreeDelegates() { return Delegates; }
+
+private:
+
+	/** A temporary snapshot of the SparseItemInfos in STreeView, used during RefreshTreeView() */
+	TSparseItemMap OldSparseItemInfos;
+
+	/** Backing array for tree view */
+	TArray<TSharedPtr<FRigTreeElement>> RootElements;
+	
+	/** A map for looking up items based on their key */
+	TMap<FRigElementKey, TSharedPtr<FRigTreeElement>> ElementMap;
+
+	/** A map for looking up a parent based on their key */
+	TMap<FRigElementKey, FRigElementKey> ParentMap;
+
+	FRigTreeDelegates Delegates;
+
+	friend class SRigHierarchy;
+};
diff --git a/Engine/Plugins/Animation/ControlRig/Source/ControlRigEditor/Private/Editor/SRigSpacePickerWidget.cpp b/Engine/Plugins/Animation/ControlRig/Source/ControlRigEditor/Private/Editor/SRigSpacePickerWidget.cpp
new file mode 100644
--- /dev/null
+++ b/Engine/Plugins/Animation/ControlRig/Source/ControlRigEditor/Private/Editor/SRigSpacePickerWidget.cpp
@@ -0,0 +1,218 @@
+// Copyright Epic Games, Inc. All Rights Reserved.
+
+#include "SRigSpacePickerWidget.h"
+#include "DetailLayoutBuilder.h"
+#include "Widgets/Text/STextBlock.h"
+#include "Widgets/Layout/SUniformGridPanel.h"
+#include "Widgets/Input/SButton.h"
+#include "SRigHierarchyTreeView.h"
+
+#define LOCTEXT_NAMESPACE "SRigSpacePickerWidget"
+
+//////////////////////////////////////////////////////////////
+/// SRigSpacePickerWidget
+///////////////////////////////////////////////////////////
+
+void SRigSpacePickerWidget::Construct(const FArguments& InArgs)
+{
+	Hierarchy = InArgs._Hierarchy;
+	ControlKey = InArgs._SelectedControl;
+	
+	ChildSlot
+	[
+		SNew(SBorder)
+		.Visibility(EVisibility::Visible)
+		.BorderImage(FEditorStyle::GetBrush("Menu.Background"))
+		[
+			SAssignNew(ListBox, SVerticalBox)
+			
+			+SVerticalBox::Slot()
+			.AutoHeight()
+			.VAlign(VAlign_Top)
+			.HAlign(HAlign_Left)
+			.Padding(4.0, 0.0, 0.0, 2.0)
+			[
+				SNew( STextBlock )
+				.Text( LOCTEXT("Space", "Space") )
+				.Font( IDetailLayoutBuilder::GetDetailFontBold() )
+			]
+		]
+	];
+
+	AddSpacePickerButton(
+		FEditorStyle::GetBrush("EditorViewport.RelativeCoordinateSystem_Local"),
+		LOCTEXT("Local", "Local"),
+		FOnClicked::CreateSP(this, &SRigSpacePickerWidget::HandleLocalSpaceClicked)
+	);
+	
+	AddSpacePickerButton(
+		FEditorStyle::GetBrush("EditorViewport.RelativeCoordinateSystem_World"),
+		LOCTEXT("World", "World"),
+		FOnClicked::CreateSP(this, &SRigSpacePickerWidget::HandleWorldSpaceClicked)
+	);
+
+	if(const FRigControlElement* Control = Hierarchy->Find<FRigControlElement>(ControlKey))
+	{
+		TArray<FRigElementKey> Favorites = Control->Settings.SpaceFavorites;
+
+		for(const FRigElementKey& Key : Favorites)
+		{
+			if(Key.IsValid() && Hierarchy->Contains(Key))
+			{
+				AddSpacePickerButton(
+					SRigHierarchyItem::GetBrushForElementType(Hierarchy, Key),
+					FText::FromName(Key.Name),
+					FOnClicked::CreateSP(this, &SRigSpacePickerWidget::HandleElementSpaceClicked, Key)
+				);
+			}
+		}
+	}
+
+	// Ok/Cancel buttons
+	ListBox->AddSlot()
+	.AutoHeight()
+	.HAlign(HAlign_Right)
+	.VAlign(VAlign_Bottom)
+	.Padding(4, 12, 0, 0)
+	[
+		SNew(SUniformGridPanel)
+		.SlotPadding(FEditorStyle::GetMargin("StandardDialog.SlotPadding"))
+		.MinDesiredSlotWidth(FEditorStyle::GetFloat("StandardDialog.MinDesiredSlotWidth"))
+		.MinDesiredSlotHeight(FEditorStyle::GetFloat("StandardDialog.MinDesiredSlotHeight"))
+		+SUniformGridPanel::Slot(0,0)
+		[
+			SNew(SButton)
+			.HAlign(HAlign_Center)
+			.ContentPadding( FEditorStyle::GetMargin("StandardDialog.ContentPadding") )
+			.OnClicked(this, &SRigSpacePickerWidget::CancelClicked)
+			.Text(LOCTEXT("RigSpacePickerWidgetCancel", "Cancel"))
+		]
+	];
+}
+
+SRigSpacePickerWidget::~SRigSpacePickerWidget()
+{
+}
+
+SRigSpacePickerWidget::FResult SRigSpacePickerWidget::InvokeDialog()
+{
+	PickedKey = FRigElementKey();
+	
+	const FVector2D CursorPos = FSlateApplication::Get().GetCursorPos();
+
+	TSharedRef<SWindow> Window = SNew(SWindow)
+	.Title( LOCTEXT("SRigSpacePickerWidgetPickSpace", "Pick a new space") )
+	.CreateTitleBar(false)
+	.Type(EWindowType::Menu)
+	.SizingRule( ESizingRule::Autosized )
+	.ScreenPosition(CursorPos)
+	.FocusWhenFirstShown(true)
+	.ActivationPolicy(EWindowActivationPolicy::FirstShown)
+	[
+		AsShared()
+	];
+	
+	Window->SetWidgetToFocusOnActivate(AsShared());
+	
+	PickerWindow = Window;
+
+	Window->MoveWindowTo(CursorPos);
+	GEditor->EditorAddModalWindow(Window);
+
+	FResult Result;
+	Result.Reply = PickedKey.IsValid() ? FReply::Handled() : FReply::Unhandled();
+	Result.Key = PickedKey;
+	return Result;
+}
+
+void SRigSpacePickerWidget::CloseDialog(bool bWasPicked)
+{
+	if(!bWasPicked)
+	{
+		PickedKey = FRigElementKey();
+	}
+
+	if ( PickerWindow.IsValid() )
+	{
+		PickerWindow.Pin()->RequestDestroyWindow();
+	}
+}
+
+FReply SRigSpacePickerWidget::CancelClicked()
+{
+	CloseDialog();
+	return FReply::Handled();
+}
+
+FReply SRigSpacePickerWidget::OnKeyDown( const FGeometry& MyGeometry, const FKeyEvent& InKeyEvent )
+{
+	if (InKeyEvent.GetKey() == EKeys::Escape)
+	{
+		CancelClicked();
+		return FReply::Handled();
+	}
+	return SCompoundWidget::OnKeyDown(MyGeometry, InKeyEvent);
+}
+
+void SRigSpacePickerWidget::AddSpacePickerButton(const FSlateBrush* InBush, const FText& InTitle,
+	FOnClicked OnClickedDelegate)
+{
+	ListBox->AddSlot()
+	.AutoHeight()
+	.VAlign(VAlign_Top)
+	.HAlign(HAlign_Fill)
+	.Padding(24.f, 0.0f, 0.0f, 2.0f)
+	[
+		SNew( SButton )
+		.ButtonStyle(FAppStyle::Get(), "SimpleButton")
+		.OnClicked(OnClickedDelegate)
+		[
+			SNew(SHorizontalBox)
+			+ SHorizontalBox::Slot()
+			.AutoWidth()
+			.VAlign(VAlign_Center)
+			.HAlign(HAlign_Left)
+			.Padding(0)
+			[
+				SNew(SImage)
+				.Image(InBush)
+			]
+
+			+ SHorizontalBox::Slot()
+			.AutoWidth()
+			.VAlign(VAlign_Center)
+			.HAlign(HAlign_Left)
+			.Padding(0)
+			[
+				SNew( STextBlock )
+				.Text( InTitle )
+				.Font( IDetailLayoutBuilder::GetDetailFont() )
+			]
+
+			+ SHorizontalBox::Slot()
+			[
+				SNew(SSpacer)
+			]
+		]
+	];
+}
+
+FReply SRigSpacePickerWidget::HandleLocalSpaceClicked()
+{
+	const FRigElementKey DefaultParent = Hierarchy->GetFirstParent(ControlKey);
+	return HandleElementSpaceClicked(DefaultParent);
+}
+
+FReply SRigSpacePickerWidget::HandleWorldSpaceClicked()
+{
+	return HandleElementSpaceClicked(Hierarchy->GetWorldSpaceSocketKey());
+}
+
+FReply SRigSpacePickerWidget::HandleElementSpaceClicked(FRigElementKey InKey)
+{
+	PickedKey = InKey;
+	CloseDialog(true);
+	return FReply::Handled();
+}
+
+#undef LOCTEXT_NAMESPACE
diff --git a/Engine/Plugins/Animation/ControlRig/Source/ControlRigEditor/Private/Editor/SRigSpacePickerWidget.h b/Engine/Plugins/Animation/ControlRig/Source/ControlRigEditor/Private/Editor/SRigSpacePickerWidget.h
new file mode 100644
--- /dev/null
+++ b/Engine/Plugins/Animation/ControlRig/Source/ControlRigEditor/Private/Editor/SRigSpacePickerWidget.h
@@ -0,0 +1,64 @@
+// Copyright Epic Games, Inc. All Rights Reserved.
+
+#pragma once
+
+#include "CoreMinimal.h"
+#include "Rigs/RigHierarchy.h"
+#include "SRigHierarchy.h"
+#include "Widgets/Layout/SBox.h"
+#include "Rigs/RigSpaceHierarchy.h"
+
+
+/** Widget allowing picking of a space source for space switching */
+class SRigSpacePickerWidget : public SCompoundWidget
+{
+public:
+
+	struct FResult
+	{
+		FResult()
+			: Reply(FReply::Unhandled())
+			// , Actor(nullptr)
+			// , SceneComponent(nullptr)
+			, Key()
+		{
+		}
+		
+		FReply Reply;
+		//AActor* Actor;
+		//USceneComponent* SceneComponent;
+		FRigElementKey Key;
+	};
+
+	SLATE_BEGIN_ARGS(SRigSpacePickerWidget) {}
+		SLATE_ARGUMENT(URigHierarchy*, Hierarchy)
+		SLATE_ARGUMENT(FRigElementKey, SelectedControl)
+	SLATE_END_ARGS()
+
+	void Construct(const FArguments& InArgs);
+	~SRigSpacePickerWidget();
+
+	FResult InvokeDialog();
+	void CloseDialog(bool bWasPicked=false);
+	FReply CancelClicked();
+	virtual FReply OnKeyDown( const FGeometry& MyGeometry, const FKeyEvent& InKeyEvent ) override;
+
+	virtual bool SupportsKeyboardFocus() const override
+	{
+		return true;
+	}
+
+private:
+
+	void AddSpacePickerButton(const FSlateBrush* InBush, const FText& InTitle, FOnClicked OnClickedDelegate);
+
+	FReply HandleLocalSpaceClicked();
+	FReply HandleWorldSpaceClicked();
+	FReply HandleElementSpaceClicked(FRigElementKey InKey);
+
+	URigHierarchy* Hierarchy;
+	FRigElementKey ControlKey;
+	TSharedPtr<SVerticalBox> ListBox;
+	TWeakPtr<SWindow> PickerWindow;
+	FRigElementKey PickedKey;
+};
\ No newline at end of file
