Control Rig: Fix node and link opaqueness

diff --git a/Engine/Plugins/Animation/ControlRig/Source/ControlRigDeveloper/Private/Graph/ControlRigGraph.cpp b/Engine/Plugins/Animation/ControlRig/Source/ControlRigDeveloper/Private/Graph/ControlRigGraph.cpp
--- a/Engine/Plugins/Animation/ControlRig/Source/ControlRigDeveloper/Private/Graph/ControlRigGraph.cpp
+++ b/Engine/Plugins/Animation/ControlRig/Source/ControlRigDeveloper/Private/Graph/ControlRigGraph.cpp
@@ -670,33 +670,78 @@ void UControlRigGraph::HandleModifiedEvent(ERigVMGraphNotifType InNotifType, URi
 	}
 }
 
-int32 UControlRigGraph::GetInstructionIndex(UControlRigGraphNode* InNode)
+int32 UControlRigGraph::GetInstructionIndex(const UControlRigGraphNode* InNode, bool bAsInput)
 {
-	if (const int32* FoundIndex = CachedInstructionIndices.Find(InNode->GetModelNode()))
+	if (const TPair<int32, int32>* FoundIndex = CachedInstructionIndices.Find(InNode->GetModelNode()))
 	{
-		return *FoundIndex;
+		return bAsInput ? FoundIndex->Key : FoundIndex->Value;
 	}
 
 	struct Local
 	{
-		static int32 GetInstructionIndex(URigVMNode* InModelNode, const FRigVMByteCode* InByteCode, TMap<URigVMNode*, int32>& Indices)
+		static int32 GetInstructionIndex(URigVMNode* InModelNode, const FRigVMByteCode* InByteCode, TMap<URigVMNode*, TPair<int32, int32>>& Indices, bool bAsInput)
 		{
 			if (InModelNode == nullptr)
 			{
 				return INDEX_NONE;
 			}
 
-			if (const int32* ExistingIndex = Indices.Find(InModelNode))
+			if (const TPair<int32, int32>* ExistingIndex = Indices.Find(InModelNode))
 			{
-				return *ExistingIndex;
+				const int32 Index = bAsInput ? ExistingIndex->Key : ExistingIndex->Value;
+				if(Index != INDEX_NONE)
+				{
+					return Index;
+				}
 			}
 
-			Indices.Add(InModelNode, INDEX_NONE);
+			if(const URigVMRerouteNode* RerouteNode = Cast<URigVMRerouteNode>(InModelNode))
+			{
+				int32 InstructionIndex = INDEX_NONE;
+				if(bAsInput)
+				{
+					TArray<URigVMNode*> SourceNodes = RerouteNode->GetLinkedSourceNodes();
+					for(URigVMNode* SourceNode : SourceNodes)
+					{
+						InstructionIndex = GetInstructionIndex(SourceNode, InByteCode, Indices, bAsInput);
+						if(InstructionIndex != INDEX_NONE)
+						{
+							break;
+						}
+					}
+					Indices.FindOrAdd(InModelNode).Key = InstructionIndex;
+				}
+				else
+				{
+					TArray<URigVMNode*> TargetNodes = RerouteNode->GetLinkedTargetNodes();
+					for(URigVMNode* TargetNode : TargetNodes)
+					{
+						InstructionIndex = GetInstructionIndex(TargetNode, InByteCode, Indices, bAsInput);
+						if(InstructionIndex != INDEX_NONE)
+						{
+							break;
+						}
+					}
+					Indices.FindOrAdd(InModelNode).Value = InstructionIndex;
+				}
+
+				return InstructionIndex;
+			}
+
+			Indices.FindOrAdd(InModelNode, TPair<int32, int32>(INDEX_NONE, INDEX_NONE));
 
 			int32 InstructionIndex = InByteCode->GetFirstInstructionIndexForSubject(InModelNode);
 			if (InstructionIndex != INDEX_NONE)
 			{
-				Indices.FindOrAdd(InModelNode) = InstructionIndex;
+				TPair<int32, int32> Pair(InstructionIndex, InstructionIndex);
+				if(bAsInput)
+				{
+					Indices.FindOrAdd(InModelNode).Key = InstructionIndex;
+				}
+				else
+				{
+					Indices.FindOrAdd(InModelNode).Value = InstructionIndex;
+				}
 				return InstructionIndex;
 			}
 
@@ -706,7 +751,14 @@ int32 UControlRigGraph::GetInstructionIndex(UControlRigGraphNode* InNode)
 				const FRigVMASTProxy Proxy = FRigVMASTProxy::MakeFromCallPath(InByteCode->GetCallPathForInstruction(i), InModelNode->GetRootGraph());
 				if (Proxy.GetCallstack().Contains(InModelNode))
 				{
-					Indices.FindOrAdd(InModelNode) = i;
+					if(bAsInput)
+					{
+						Indices.FindOrAdd(InModelNode).Key = i;
+					}
+					else
+					{
+						Indices.FindOrAdd(InModelNode).Value = i;
+					}
 					return i;
 				}
 			}
@@ -716,7 +768,9 @@ int32 UControlRigGraph::GetInstructionIndex(UControlRigGraphNode* InNode)
 
 	if (const FRigVMByteCode* ByteCode = GetController()->GetCurrentByteCode())
 	{
-		return Local::GetInstructionIndex(InNode->GetModelNode(), ByteCode, CachedInstructionIndices);
+		const int32 SourceInstructionIndex = Local::GetInstructionIndex(InNode->GetModelNode(), ByteCode, CachedInstructionIndices, true);
+		const int32 TargetInstructionIndex = Local::GetInstructionIndex(InNode->GetModelNode(), ByteCode, CachedInstructionIndices, false);
+		return bAsInput ? SourceInstructionIndex : TargetInstructionIndex;
 	}
 
 	return INDEX_NONE;
diff --git a/Engine/Plugins/Animation/ControlRig/Source/ControlRigDeveloper/Private/Graph/ControlRigGraphNode.cpp b/Engine/Plugins/Animation/ControlRig/Source/ControlRigDeveloper/Private/Graph/ControlRigGraphNode.cpp
--- a/Engine/Plugins/Animation/ControlRig/Source/ControlRigDeveloper/Private/Graph/ControlRigGraphNode.cpp
+++ b/Engine/Plugins/Animation/ControlRig/Source/ControlRigDeveloper/Private/Graph/ControlRigGraphNode.cpp
@@ -375,11 +375,11 @@ void UControlRigGraphNode::HandleInsertArrayElement(FString InPinPath)
 	}
 }
 
-int32 UControlRigGraphNode::GetInstructionIndex()
+int32 UControlRigGraphNode::GetInstructionIndex(bool bAsInput) const
 {
 	if (UControlRigGraph* RigGraph = Cast<UControlRigGraph>(GetGraph()))
 	{
-		return RigGraph->GetInstructionIndex(this);
+		return RigGraph->GetInstructionIndex(this, bAsInput);
 	}
 	return INDEX_NONE;
 }
@@ -680,15 +680,9 @@ FLinearColor UControlRigGraphNode::GetNodeOpacityColor() const
 			return FLinearColor::White;
 		}
 
-		if (URigVMController* Controller = GetController())
+		if(GetInstructionIndex(true) == INDEX_NONE)
 		{
-			if (const FRigVMByteCode* ByteCode = Controller->GetCurrentByteCode())
-			{
-				if (ByteCode->GetFirstInstructionIndexForSubject(ModelNode) == INDEX_NONE)
-				{
-					return FLinearColor(0.35f, 0.35f, 0.35f, 0.35f);
-				}
-			}
+			return FLinearColor(0.35f, 0.35f, 0.35f, 0.35f);
 		}
 	}
 	return FLinearColor::White;
diff --git a/Engine/Plugins/Animation/ControlRig/Source/ControlRigDeveloper/Public/Graph/ControlRigGraph.h b/Engine/Plugins/Animation/ControlRig/Source/ControlRigDeveloper/Public/Graph/ControlRigGraph.h
--- a/Engine/Plugins/Animation/ControlRig/Source/ControlRigDeveloper/Public/Graph/ControlRigGraph.h
+++ b/Engine/Plugins/Animation/ControlRig/Source/ControlRigDeveloper/Public/Graph/ControlRigGraph.h
@@ -74,7 +74,7 @@ public:
 
 	void HandleModifiedEvent(ERigVMGraphNotifType InNotifType, URigVMGraph* InGraph, UObject* InSubject);
 
-	int32 GetInstructionIndex(UControlRigGraphNode* InNode);
+	int32 GetInstructionIndex(const UControlRigGraphNode* InNode, bool bAsInput);
 
 	UPROPERTY()
 	FString ModelNodePath;
@@ -133,7 +133,7 @@ private:
 
 	FControlRigGraphNodeClicked OnGraphNodeClicked;
 
-	TMap<URigVMNode*, int32> CachedInstructionIndices;
+	TMap<URigVMNode*, TPair<int32, int32>> CachedInstructionIndices;
 
 #endif
 #if WITH_EDITORONLY_DATA
diff --git a/Engine/Plugins/Animation/ControlRig/Source/ControlRigDeveloper/Public/Graph/ControlRigGraphNode.h b/Engine/Plugins/Animation/ControlRig/Source/ControlRigDeveloper/Public/Graph/ControlRigGraphNode.h
--- a/Engine/Plugins/Animation/ControlRig/Source/ControlRigDeveloper/Public/Graph/ControlRigGraphNode.h
+++ b/Engine/Plugins/Animation/ControlRig/Source/ControlRigDeveloper/Public/Graph/ControlRigGraphNode.h
@@ -149,7 +149,7 @@ public:
 
 	FNodeTitleDirtied& GetNodeTitleDirtied() { return NodeTitleDirtied; }
 
-	int32 GetInstructionIndex();
+	int32 GetInstructionIndex(bool bAsInput) const;
 
 protected:
 
diff --git a/Engine/Plugins/Animation/ControlRig/Source/ControlRigEditor/Private/Graph/ControlRigConnectionDrawingPolicy.cpp b/Engine/Plugins/Animation/ControlRig/Source/ControlRigEditor/Private/Graph/ControlRigConnectionDrawingPolicy.cpp
--- a/Engine/Plugins/Animation/ControlRig/Source/ControlRigEditor/Private/Graph/ControlRigConnectionDrawingPolicy.cpp
+++ b/Engine/Plugins/Animation/ControlRig/Source/ControlRigEditor/Private/Graph/ControlRigConnectionDrawingPolicy.cpp
@@ -305,8 +305,8 @@ void FControlRigConnectionDrawingPolicy::DetermineWiringStyle(UEdGraphPin* Outpu
 		if (OutputModelPin && InputModelPin)
 		{
 			bool bVisited = false;
-			int32 OutputInstructionIndex = OutputNode->GetInstructionIndex();
-			int32 InputInstructionIndex = InputNode->GetInstructionIndex();
+			int32 OutputInstructionIndex = OutputNode->GetInstructionIndex(false);
+			int32 InputInstructionIndex = InputNode->GetInstructionIndex(true);
 
 			if (OutputInstructionIndex != INDEX_NONE && InputInstructionIndex != INDEX_NONE)
 			{
