Control Rig: Adapt RigHierarchy to work for new Space Switching Nodes

diff --git a/Engine/Plugins/Animation/ControlRig/Source/ControlRig/Private/Rigs/RigHierarchy.cpp b/Engine/Plugins/Animation/ControlRig/Source/ControlRig/Private/Rigs/RigHierarchy.cpp
--- a/Engine/Plugins/Animation/ControlRig/Source/ControlRig/Private/Rigs/RigHierarchy.cpp
+++ b/Engine/Plugins/Animation/ControlRig/Source/ControlRig/Private/Rigs/RigHierarchy.cpp
@@ -1054,6 +1054,18 @@ bool URigHierarchy::SetParentWeightArray(FRigElementKey InChild, TArray<FRigElem
 
 bool URigHierarchy::SetParentWeightArray(FRigBaseElement* InChild, const TArray<FRigElementWeight>& InWeights,
 	bool bInitial, bool bAffectChildren)
+{
+	if(InWeights.Num() == 0)
+	{
+		return false;
+	}
+	
+	TArrayView<const FRigElementWeight> View(InWeights.GetData(), InWeights.Num());
+	return SetParentWeightArray(InChild, View, bInitial, bAffectChildren);
+}
+
+bool URigHierarchy::SetParentWeightArray(FRigBaseElement* InChild,  const TArrayView<const FRigElementWeight>& InWeights,
+	bool bInitial, bool bAffectChildren)
 {
 	using namespace ERigTransformType;
 
@@ -2702,6 +2714,11 @@ bool URigHierarchy::IsParentedTo(FRigBaseElement* InChild, FRigBaseElement* InPa
 		return false;
 	}
 
+	if(InChild == InParent)
+	{
+		return true;
+	}
+
 	if(FRigSingleParentElement* SingleParentElement = Cast<FRigSingleParentElement>(InChild))
 	{
 		if(SingleParentElement->ParentElement == InParent)
diff --git a/Engine/Plugins/Animation/ControlRig/Source/ControlRig/Public/Rigs/RigHierarchy.h b/Engine/Plugins/Animation/ControlRig/Source/ControlRig/Public/Rigs/RigHierarchy.h
--- a/Engine/Plugins/Animation/ControlRig/Source/ControlRig/Public/Rigs/RigHierarchy.h
+++ b/Engine/Plugins/Animation/ControlRig/Source/ControlRig/Public/Rigs/RigHierarchy.h
@@ -1699,6 +1699,16 @@ public:
 	 */
 	bool SetParentWeightArray(FRigBaseElement* InChild,  const TArray<FRigElementWeight>& InWeights, bool bInitial = false, bool bAffectChildren = true);
 
+	/**
+	* Sets the all of the weights of the parents of a multi parent element
+	* @param InChild The multi parented element
+	* @param InWeights The new weights to set for the parents
+	* @param bInitial If true the initial weights will be used
+	* @param bAffectChildren If set to false children will not move (maintain global).
+	* @return Returns true if changing the weight was successful
+	*/
+	bool SetParentWeightArray(FRigBaseElement* InChild,  const TArrayView<const FRigElementWeight>& InWeights, bool bInitial = false, bool bAffectChildren = true);
+
 	/**
 		* Switches a multi parent element to a single parent.
 		* This sets the new parent's weight to 1.0 and disables
