Control Rig: Introduce dynamic hierarchy nodes for allows parent / space switching

diff --git a/Engine/Plugins/Animation/ControlRig/Source/ControlRig/Private/Units/Execution/RigUnit_DynamicHierarchy.cpp b/Engine/Plugins/Animation/ControlRig/Source/ControlRig/Private/Units/Execution/RigUnit_DynamicHierarchy.cpp
new file mode 100644
--- /dev/null
+++ b/Engine/Plugins/Animation/ControlRig/Source/ControlRig/Private/Units/Execution/RigUnit_DynamicHierarchy.cpp
@@ -0,0 +1,136 @@
+// Copyright Epic Games, Inc. All Rights Reserved.
+
+#include "RigUnit_DynamicHierarchy.h"
+#include "Rigs/RigHierarchyController.h"
+#include "Units/RigUnitContext.h"
+
+FRigUnit_AddParent_Execute()
+{
+	if((Context.State != EControlRigState::Update) || (ExecuteContext.Hierarchy == nullptr))
+	{
+		return;
+	}
+
+	FRigTransformElement* ChildElement = ExecuteContext.Hierarchy->Find<FRigTransformElement>(Child);
+	if(ChildElement == nullptr)
+	{
+		UE_CONTROLRIG_RIGUNIT_REPORT_ERROR(TEXT("Child item %s does not exist."), *Child.ToString())
+		return;
+	}
+
+	FRigTransformElement* ParentElement = ExecuteContext.Hierarchy->Find<FRigTransformElement>(Parent);
+	if(ParentElement == nullptr)
+	{
+		UE_CONTROLRIG_RIGUNIT_REPORT_ERROR(TEXT("Parent item %s does not exist."), *Parent.ToString())
+		return;
+	}
+
+	if(URigHierarchyController* Controller = ExecuteContext.Hierarchy->GetController(true))
+	{
+		Controller->AddParent(ChildElement, ParentElement, 0.f, true, false);
+	}
+}
+
+FRigUnit_SwitchParent_Execute()
+{
+	if((Context.State != EControlRigState::Update) || (ExecuteContext.Hierarchy == nullptr))
+	{
+		return;
+	}
+
+	FRigTransformElement* ChildElement = ExecuteContext.Hierarchy->Find<FRigTransformElement>(Child);
+	if(ChildElement == nullptr)
+	{
+		UE_CONTROLRIG_RIGUNIT_REPORT_ERROR(TEXT("Child item %s does not exist."), *Child.ToString())
+		return;
+	}
+	if(!ChildElement->IsA<FRigMultiParentElement>())
+	{
+		UE_CONTROLRIG_RIGUNIT_REPORT_ERROR(TEXT("Child item %s cannot be space switched (only Nulls and Controls can)."), *Child.ToString())
+		return;
+	}
+
+	FRigTransformElement* ParentElement = nullptr;
+
+	if(Mode == ERigSwitchParentMode::ParentItem)
+	{
+		ParentElement = ExecuteContext.Hierarchy->Find<FRigTransformElement>(Parent);
+		if(ParentElement == nullptr)
+		{
+			UE_CONTROLRIG_RIGUNIT_REPORT_ERROR(TEXT("Parent item %s does not exist."), *Parent.ToString())
+			return;
+		}
+
+		if(!ParentElement->IsA<FRigTransformElement>())
+		{
+			UE_CONTROLRIG_RIGUNIT_REPORT_ERROR(TEXT("Parent item %s does not have a transform."), *Parent.ToString())
+			return;
+		}
+	}
+
+	const ERigTransformType::Type TransformTypeToMaintain =
+		bMaintainGlobal ? ERigTransformType::CurrentGlobal : ERigTransformType::CurrentLocal;
+	
+	const FTransform Transform = ExecuteContext.Hierarchy->GetTransform(ChildElement, TransformTypeToMaintain);
+	
+	switch(Mode)
+	{
+		case ERigSwitchParentMode::World:
+		{
+			ExecuteContext.Hierarchy->SwitchToWorldSpace(ChildElement, false, true);
+			break;
+		}
+		case ERigSwitchParentMode::DefaultParent:
+		{
+			ExecuteContext.Hierarchy->SwitchToDefaultParent(ChildElement, false, true);
+			break;
+		}
+		case ERigSwitchParentMode::ParentItem:
+		default:
+		{
+			ExecuteContext.Hierarchy->SwitchToParent(ChildElement, ParentElement, false, true);
+			break;
+		}
+	}
+	
+	ExecuteContext.Hierarchy->SetTransform(ChildElement, Transform, TransformTypeToMaintain, true);
+}
+
+FRigUnit_HierarchyGetParentWeights_Execute()
+{
+	if((Context.State != EControlRigState::Update) || (Context.Hierarchy == nullptr))
+	{
+		return;
+	}
+
+	FRigBaseElement* ChildElement = Context.Hierarchy->Find(Child);
+	if(ChildElement == nullptr)
+	{
+		UE_CONTROLRIG_RIGUNIT_REPORT_ERROR(TEXT("Item %s does not exist."), *Child.ToString())
+		return;
+	}
+	Weights = Context.Hierarchy->GetParentWeightArray(ChildElement, false);
+}
+
+FRigUnit_HierarchySetParentWeights_Execute()
+{
+	if((Context.State != EControlRigState::Update) || (ExecuteContext.Hierarchy == nullptr))
+	{
+		return;
+	}
+
+	FRigBaseElement* ChildElement = Context.Hierarchy->Find(Child);
+	if(ChildElement == nullptr)
+	{
+		UE_CONTROLRIG_RIGUNIT_REPORT_ERROR(TEXT("Item %s does not exist."), *Child.ToString())
+		return;
+	}
+
+	if(Weights.Num() != ExecuteContext.Hierarchy->GetNumberOfParents(ChildElement))
+	{
+		UE_CONTROLRIG_RIGUNIT_REPORT_ERROR(TEXT("Provided incorrect number of weights(%d) for %s - expected %d."), Weights.Num(), *Child.ToString(), ExecuteContext.Hierarchy->GetNumberOfParents(Child))
+		return;
+	}
+
+	ExecuteContext.Hierarchy->SetParentWeightArray(ChildElement, Weights, false, true);
+}
diff --git a/Engine/Plugins/Animation/ControlRig/Source/ControlRig/Private/Units/Execution/RigUnit_DynamicHierarchy.h b/Engine/Plugins/Animation/ControlRig/Source/ControlRig/Private/Units/Execution/RigUnit_DynamicHierarchy.h
new file mode 100644
--- /dev/null
+++ b/Engine/Plugins/Animation/ControlRig/Source/ControlRig/Private/Units/Execution/RigUnit_DynamicHierarchy.h
@@ -0,0 +1,144 @@
+// Copyright Epic Games, Inc. All Rights Reserved.
+
+#pragma once
+
+#include "Units/RigUnit.h"
+#include "ControlRigDefines.h" 
+#include "RigUnit_DynamicHierarchy.generated.h"
+
+USTRUCT(meta = (Abstract, NodeColor="0.262745, 0.8, 0, 0.229412", Category = "DynamicHierarchy"))
+struct CONTROLRIG_API FRigUnit_DynamicHierarchyBase : public FRigUnit
+{
+	GENERATED_BODY()
+};
+
+USTRUCT(meta = (Abstract, NodeColor="0.262745, 0.8, 0, 0.229412", Category = "DynamicHierarchy"))
+struct CONTROLRIG_API FRigUnit_DynamicHierarchyBaseMutable : public FRigUnitMutable
+{
+	GENERATED_BODY()
+};
+
+/**
+* Adds a new parent to an element. The weight for the new parent will be 0.0.
+* You can use the SetParentWeights node to change the parent weights later.
+*/
+USTRUCT(meta=(DisplayName="Add Parent", Keywords="Children,Parent,Constraint,Space", Varying))
+struct CONTROLRIG_API FRigUnit_AddParent : public FRigUnit_DynamicHierarchyBaseMutable
+{
+	GENERATED_BODY()
+
+	FRigUnit_AddParent()
+	{
+		Child = Parent = FRigElementKey(NAME_None, ERigElementType::Control);
+	}
+
+	RIGVM_METHOD()
+	virtual void Execute(const FRigUnitContext& Context) override;
+
+	UPROPERTY(meta = (Input, ExpandByDefault))
+	FRigElementKey Child;
+
+	UPROPERTY(meta = (Input, ExpandByDefault))
+	FRigElementKey Parent;
+};
+
+UENUM()
+enum class ERigSwitchParentMode : uint8
+{
+	/** Switches the element to be parented to the world */
+	World,
+
+	/** Switches back to the original / default parent */
+	DefaultParent,
+
+	/** Switches the child to the provided parent item */
+	ParentItem
+};
+
+/**
+* Switches an element to a new parent.
+*/
+USTRUCT(meta=(DisplayName="Switch Parent", Keywords="Children,Parent,Constraint,Space,Switch", Varying))
+struct CONTROLRIG_API FRigUnit_SwitchParent : public FRigUnit_DynamicHierarchyBaseMutable
+{
+	GENERATED_BODY()
+
+	FRigUnit_SwitchParent()
+	{
+		Mode = ERigSwitchParentMode::ParentItem;
+		Child = Parent = FRigElementKey(NAME_None, ERigElementType::Control);
+		bMaintainGlobal = true;
+	}
+
+	RIGVM_METHOD()
+	virtual void Execute(const FRigUnitContext& Context) override;
+
+	/* Depending on this the child will switch to the world,
+	 * back to its default or to the item provided by the Parent pin
+	 */
+	UPROPERTY(meta = (Input))
+	ERigSwitchParentMode Mode;
+
+	/* The child to switch to a new parent */
+	UPROPERTY(meta = (Input, ExpandByDefault))
+	FRigElementKey Child;
+
+	/* The optional parent to switch to. This is only used if the mode is set to 'Parent Item' */
+	UPROPERTY(meta = (Input, ExpandByDefault, EditCondition="Mode==ParentItem"))
+	FRigElementKey Parent;
+
+	/* If set to true the item will maintain its global transform,
+	 * otherwise it will maintain local
+	 */
+	UPROPERTY(meta = (Input))
+	bool bMaintainGlobal;
+};
+
+/**
+* Returns the item's parent's weights
+*/
+USTRUCT(meta=(DisplayName="Get Parent Weights", Keywords="Chain,Parents,Hierarchy", Varying))
+struct CONTROLRIG_API FRigUnit_HierarchyGetParentWeights : public FRigUnit_DynamicHierarchyBase
+{
+	GENERATED_BODY()
+
+	FRigUnit_HierarchyGetParentWeights()
+	{
+		Child = FRigElementKey(NAME_None, ERigElementType::Control);
+	}
+
+	RIGVM_METHOD()
+	virtual void Execute(const FRigUnitContext& Context) override;
+
+	UPROPERTY(meta = (Input, ExpandByDefault))
+	FRigElementKey Child;
+
+	UPROPERTY(meta = (Output))
+	TArray<FRigElementWeight> Weights;
+
+	UPROPERTY(meta = (Output))
+	FRigElementKeyCollection Parents;
+};
+
+/**
+* Sets the item's parent's weights
+*/
+USTRUCT(meta=(DisplayName="Set Parent Weights", Keywords="Chain,Parents,Hierarchy", Varying))
+struct CONTROLRIG_API FRigUnit_HierarchySetParentWeights : public FRigUnit_DynamicHierarchyBaseMutable
+{
+	GENERATED_BODY()
+
+	FRigUnit_HierarchySetParentWeights()
+	{
+		Child = FRigElementKey(NAME_None, ERigElementType::Control);
+	}
+
+	RIGVM_METHOD()
+	virtual void Execute(const FRigUnitContext& Context) override;
+
+	UPROPERTY(meta = (Input, ExpandByDefault))
+	FRigElementKey Child;
+
+	UPROPERTY(meta = (Input))
+	TArray<FRigElementWeight> Weights;
+};
\ No newline at end of file
