Control Rig: Updated parent constraint solver to support component wise separation

diff --git a/Engine/Plugins/Animation/ControlRig/Source/ControlRig/Private/ControlRig.cpp b/Engine/Plugins/Animation/ControlRig/Source/ControlRig/Private/ControlRig.cpp
--- a/Engine/Plugins/Animation/ControlRig/Source/ControlRig/Private/ControlRig.cpp
+++ b/Engine/Plugins/Animation/ControlRig/Source/ControlRig/Private/ControlRig.cpp
@@ -1479,8 +1479,7 @@ FRigControlValue UControlRig::GetControlValueFromGlobalTransform(const FName& In
 	{
 		if(DynamicHierarchy)
 		{
-			FTransform ParentTransform = DynamicHierarchy->GetControlOffsetTransform(ControlElement, ERigTransformType::CurrentGlobal);
-			FTransform Transform = InGlobalTransform.GetRelativeTransform(ParentTransform);
+			FTransform Transform = DynamicHierarchy->ComputeLocalControlValue(ControlElement, InGlobalTransform, ERigTransformType::CurrentGlobal);
 			Value.SetFromTransform(Transform, ControlElement->Settings.ControlType, ControlElement->Settings.PrimaryAxis);
 
 			if (ShouldApplyLimits())
diff --git a/Engine/Plugins/Animation/ControlRig/Source/ControlRig/Private/Rigs/RigHierarchy.cpp b/Engine/Plugins/Animation/ControlRig/Source/ControlRig/Private/Rigs/RigHierarchy.cpp
--- a/Engine/Plugins/Animation/ControlRig/Source/ControlRig/Private/Rigs/RigHierarchy.cpp
+++ b/Engine/Plugins/Animation/ControlRig/Source/ControlRig/Private/Rigs/RigHierarchy.cpp
@@ -496,9 +496,9 @@ void URigHierarchy::ResetPoseToInitial(ERigElementType InTypeFilter)
 			if(bHasFilteredChildren)
 			{
 				const FTransform OffsetTransform = GetControlOffsetTransform(ControlElement, ERigTransformType::InitialLocal);
-				SetControlOffsetTransform(ControlElement, OffsetTransform, ERigTransformType::CurrentLocal, true);
+				SetControlOffsetTransform(ControlElement, OffsetTransform, ERigTransformType::CurrentLocal, true, false, true);
 				const FTransform GizmoTransform = GetControlGizmoTransform(ControlElement, ERigTransformType::InitialLocal);
-				SetControlGizmoTransform(ControlElement, GizmoTransform, ERigTransformType::CurrentLocal, true);
+				SetControlGizmoTransform(ControlElement, GizmoTransform, ERigTransformType::CurrentLocal, false, true);
 			}
 			else
 			{
@@ -996,11 +996,21 @@ bool URigHierarchy::SetParentWeight(FRigBaseElement* InChild, int32 InParentInde
 
 			if(bAffectChildren)
 			{
+				GetParentTransform(MultiParentElement, LocalType);
+				if(FRigControlElement* ControlElement = Cast<FRigControlElement>(MultiParentElement))
+				{
+					GetControlOffsetTransform(ControlElement, LocalType);
+				}
 				GetTransform(MultiParentElement, LocalType);
 				MultiParentElement->Pose.MarkDirty(GlobalType);
 			}
 			else
 			{
+				GetParentTransform(MultiParentElement, GlobalType);
+				if(FRigControlElement* ControlElement = Cast<FRigControlElement>(MultiParentElement))
+				{
+					GetControlOffsetTransform(ControlElement, GlobalType);
+				}
 				GetTransform(MultiParentElement, GlobalType);
 				MultiParentElement->Pose.MarkDirty(LocalType);
 			}
@@ -1655,26 +1665,40 @@ FTransform URigHierarchy::GetTransform(FRigTransformElement* InTransformElement,
 		ensure(!InTransformElement->Pose.IsDirty(OpposedType));
 
 		FTransform ParentTransform;
-		if(FRigControlElement* ControlElement = Cast<FRigControlElement>(InTransformElement))
-		{
-			ParentTransform = GetControlOffsetTransform(ControlElement, GlobalType);
-		}
-		else
-		{
-			ParentTransform = GetParentTransform(InTransformElement, GlobalType);
-		}
-
 		if(IsLocal(InTransformType))
 		{
-			FTransform NewTransform = InTransformElement->Pose.Get(OpposedType).GetRelativeTransform(ParentTransform);
-			NewTransform.NormalizeRotation();
-			InTransformElement->Pose.Set(InTransformType, NewTransform);
+			if(FRigControlElement* ControlElement = Cast<FRigControlElement>(InTransformElement))
+			{
+				const FTransform NewTransform = ComputeLocalControlValue(ControlElement, ControlElement->Pose.Get(OpposedType), InTransformType);
+				InTransformElement->Pose.Set(InTransformType, NewTransform);
+			}
+			else
+			{
+				ParentTransform = GetParentTransform(InTransformElement, GlobalType);
+
+				FTransform NewTransform = InTransformElement->Pose.Get(OpposedType).GetRelativeTransform(ParentTransform);
+				NewTransform.NormalizeRotation();
+				InTransformElement->Pose.Set(InTransformType, NewTransform);
+			}
 		}
 		else
 		{
-			FTransform NewTransform = InTransformElement->Pose.Get(OpposedType) * ParentTransform;
-			NewTransform.NormalizeRotation();
-			InTransformElement->Pose.Set(InTransformType, NewTransform);
+			if(FRigControlElement* ControlElement = Cast<FRigControlElement>(InTransformElement))
+			{
+				const FTransform NewTransform = SolveParentConstraints(
+					ControlElement->ParentConstraints, InTransformType,
+					ControlElement->Offset.Get(OpposedType), true,
+					ControlElement->Pose.Get(OpposedType), true);
+				ControlElement->Pose.Set(InTransformType, NewTransform);
+			}
+			else
+			{
+				ParentTransform = GetParentTransform(InTransformElement, GlobalType);
+
+				FTransform NewTransform = InTransformElement->Pose.Get(OpposedType) * ParentTransform;
+				NewTransform.NormalizeRotation();
+				InTransformElement->Pose.Set(InTransformType, NewTransform);
+			}
 		}
 	}
 	return InTransformElement->Pose.Get(InTransformType);
@@ -1691,9 +1715,7 @@ void URigHierarchy::SetTransform(FRigTransformElement* InTransformElement, const
 	{
 		if(FRigControlElement* ControlElement = Cast<FRigControlElement>(InTransformElement))
 		{
-			const FTransform OffsetTransform = GetControlOffsetTransform(ControlElement, InTransformType);
-			FTransform LocalTransform = InTransform.GetRelativeTransform(OffsetTransform);
-			
+			FTransform LocalTransform = ComputeLocalControlValue(ControlElement, InTransform, InTransformType);
 			ControlElement->Settings.ApplyLimits(LocalTransform);
 			SetTransform(ControlElement, LocalTransform, MakeLocal(InTransformType), bAffectChildren);
 			return;
@@ -1801,14 +1823,18 @@ FTransform URigHierarchy::GetControlOffsetTransform(FRigControlElement* InContro
 		const ERigTransformType::Type GlobalType = MakeGlobal(InTransformType);
 		ensure(!InControlElement->Offset.IsDirty(OpposedType));
 
-		const FTransform ParentTransform =  GetParentTransform(InControlElement, GlobalType);
 		if(IsLocal(InTransformType))
 		{
+			const FTransform ParentTransform = GetParentTransform(InControlElement, GlobalType);
 			InControlElement->Offset.Set(InTransformType, InControlElement->Offset.Get(OpposedType).GetRelativeTransform(ParentTransform));
 		}
 		else
 		{
-			InControlElement->Offset.Set(InTransformType, InControlElement->Offset.Get(OpposedType) * ParentTransform);
+			const FTransform GlobalTransform = SolveParentConstraints(
+				InControlElement->ParentConstraints, InTransformType,
+				InControlElement->Offset.Get(OpposedType), true,
+				FTransform::Identity, false);
+			InControlElement->Offset.Set(InTransformType, GlobalTransform);
 		}
 	}
 	return InControlElement->Offset.Get(InTransformType);
@@ -2041,433 +2067,14 @@ FTransform URigHierarchy::GetParentTransform(FRigBaseElement* InElement, const E
 
 		if(Output.bDirty)
 		{
-			FTransform OutputTransform = FTransform::Identity;
-
-			const bool bInitial = IsInitial(InTransformType);
-			const FRigElementParentConstraintArray& ParentConstraints = MultiParentElement->ParentConstraints;
-
-			struct Local
-			{
-				static float GetWeightForLerp(const float WeightA, const float WeightB)
-				{
-					float Weight = 0.f;
-					const float ClampedWeightA = FMath::Max(WeightA, 0.f);
-					const float ClampedWeightB = FMath::Max(WeightB, 0.f);
-					const float OverallWeight = ClampedWeightA + ClampedWeightB;
-					if(OverallWeight > SMALL_NUMBER)
-					{
-						Weight = ClampedWeightB / OverallWeight;
-					}
-					return Weight;
-				}
-			};
-
-			struct FTransformIndex
-			{
-				int32 Location;
-				int32 Rotation;
-				int32 Scale;
-
-				FTransformIndex()
-					: Location(INDEX_NONE)
-					, Rotation(INDEX_NONE)
-					, Scale(INDEX_NONE)
-				{}
-
-				FTransformIndex(int32 InIndex)
-					: Location(InIndex)
-					, Rotation(InIndex)
-					, Scale(InIndex)
-				{}
-			};
-
-			// collect all of the weights
-			FTransformIndex FirstConstraint;
-			FTransformIndex SecondConstraint;
-			FTransformIndex NumConstraintsAffecting(0);
-			FRigElementWeight TotalWeight(0.f);
-
-			// find all of the weights affecting this output
-			for(int32 ConstraintIndex = 0; ConstraintIndex < ParentConstraints.Num(); ConstraintIndex++)
-			{
-				const FRigElementWeight& Weight = ParentConstraints[ConstraintIndex].GetWeight(bInitial);
-				if(Weight.AffectsLocation())
-				{
-					NumConstraintsAffecting.Location++;
-					TotalWeight.Location += Weight.Location;
-
-					if(FirstConstraint.Location == INDEX_NONE)
-					{
-						FirstConstraint.Location = ConstraintIndex;
-					}
-					else if(SecondConstraint.Location == INDEX_NONE)
-					{
-						SecondConstraint.Location = ConstraintIndex;
-					}
-				}
-				if(Weight.AffectsRotation())
-				{
-					NumConstraintsAffecting.Rotation++;
-					TotalWeight.Rotation += Weight.Rotation;
-
-					if(FirstConstraint.Rotation == INDEX_NONE)
-					{
-						FirstConstraint.Rotation = ConstraintIndex;
-					}
-					else if(SecondConstraint.Rotation == INDEX_NONE)
-					{
-						SecondConstraint.Rotation = ConstraintIndex;
-					}
-				}
-				if(Weight.AffectsScale())
-				{
-					NumConstraintsAffecting.Scale++;
-					TotalWeight.Scale += Weight.Scale;
-
-					if(FirstConstraint.Scale == INDEX_NONE)
-					{
-						FirstConstraint.Scale = ConstraintIndex;
-					}
-					else if(SecondConstraint.Scale == INDEX_NONE)
-					{
-						SecondConstraint.Scale = ConstraintIndex;
-					}
-				}
-			}
-
-			if(NumConstraintsAffecting.Location == 1)
-			{
-				check(FirstConstraint.Location != INDEX_NONE);
-
-				const FRigElementParentConstraint& ParentConstraint = ParentConstraints[FirstConstraint.Location];
-				const FRigElementWeight& Weight = ParentConstraint.GetWeight(bInitial);
-				const FTransform ParentTransform = GetTransform(ParentConstraint.ParentElement, InTransformType);
-				check(Weight.AffectsLocation());
-				OutputTransform.SetLocation(ParentTransform.GetLocation());
-			}
-			else if(NumConstraintsAffecting.Location == 2)
-			{
-				check(FirstConstraint.Location != INDEX_NONE);
-				check(SecondConstraint.Location != INDEX_NONE);
-
-				const FRigElementParentConstraint& ParentConstraintA = ParentConstraints[FirstConstraint.Location];
-				const FRigElementParentConstraint& ParentConstraintB = ParentConstraints[SecondConstraint.Location];
-
-				const FRigElementWeight& WeightA = ParentConstraintA.GetWeight(bInitial); 
-				const FRigElementWeight& WeightB = ParentConstraintB.GetWeight(bInitial);
-				check(WeightA.AffectsLocation());
-				check(WeightB.AffectsLocation());
-				const float Weight = Local::GetWeightForLerp(WeightA.Location, WeightB.Location);
-
-				const FVector ParentLocationA = GetTransform(ParentConstraintA.ParentElement, InTransformType).GetLocation();
-				const FVector ParentLocationB = GetTransform(ParentConstraintB.ParentElement, InTransformType).GetLocation();
-				OutputTransform.SetLocation(FMath::Lerp<FVector>(ParentLocationA, ParentLocationB, Weight));
-			}
-			else if(NumConstraintsAffecting.Location > 2)
-			{
-				check(TotalWeight.Location > SMALL_NUMBER);
-				
-				FVector Location = FVector::ZeroVector;
-				
-				for(int32 ConstraintIndex = 0; ConstraintIndex < ParentConstraints.Num(); ConstraintIndex++)
-				{
-					const FRigElementParentConstraint& ParentConstraint = ParentConstraints[ConstraintIndex];
-					const FRigElementWeight& Weight = ParentConstraint.GetWeight(bInitial);
-					if(!Weight.AffectsLocation())
-					{
-						continue;
-					}
-
-					const FVector ParentLocation = GetTransform(ParentConstraint.ParentElement, InTransformType).GetLocation();
-					Location += ParentLocation * Weight.Location / TotalWeight.Location; 
-				}
-
-				OutputTransform.SetLocation(Location);
-			}
-
-			if(NumConstraintsAffecting.Rotation == 1)
-			{
-				check(FirstConstraint.Rotation != INDEX_NONE);
-
-				const FRigElementParentConstraint& ParentConstraint = ParentConstraints[FirstConstraint.Rotation];
-				const FRigElementWeight& Weight = ParentConstraint.GetWeight(bInitial);
-				const FTransform ParentTransform = GetTransform(ParentConstraint.ParentElement, InTransformType);
-				check(Weight.AffectsRotation());
-				OutputTransform.SetRotation(ParentTransform.GetRotation());
-			}
-			else if(NumConstraintsAffecting.Rotation == 2)
-			{
-				check(FirstConstraint.Rotation != INDEX_NONE);
-				check(SecondConstraint.Rotation != INDEX_NONE);
-
-				const FRigElementParentConstraint& ParentConstraintA = ParentConstraints[FirstConstraint.Rotation];
-				const FRigElementParentConstraint& ParentConstraintB = ParentConstraints[SecondConstraint.Rotation];
-
-				const FRigElementWeight& WeightA = ParentConstraintA.GetWeight(bInitial); 
-				const FRigElementWeight& WeightB = ParentConstraintB.GetWeight(bInitial);
-				check(WeightA.AffectsRotation());
-				check(WeightB.AffectsRotation());
-				const float Weight = Local::GetWeightForLerp(WeightA.Rotation, WeightB.Rotation);
-
-				const FQuat ParentRotationA = GetTransform(ParentConstraintA.ParentElement, InTransformType).GetRotation();
-				const FQuat ParentRotationB = GetTransform(ParentConstraintB.ParentElement, InTransformType).GetRotation();
-				OutputTransform.SetRotation(FQuat::Slerp(ParentRotationA, ParentRotationB, Weight));
-			}
-			else if(NumConstraintsAffecting.Rotation > 2)
-			{
-				check(TotalWeight.Rotation > SMALL_NUMBER);
-				
-				int NumMixedRotations = 0;
-				FQuat FirstRotation = FQuat::Identity;
-				FQuat MixedRotation = FQuat(0.f, 0.f, 0.f, 0.f);
-
-				for(int32 ConstraintIndex = 0; ConstraintIndex < ParentConstraints.Num(); ConstraintIndex++)
-				{
-					const FRigElementParentConstraint& ParentConstraint = ParentConstraints[ConstraintIndex];
-					const FRigElementWeight& Weight = ParentConstraint.GetWeight(bInitial);
-					if(!Weight.AffectsRotation())
-					{
-						continue;
-					}
-
-					FQuat ParentRotation = GetTransform(ParentConstraint.ParentElement, InTransformType).GetRotation().GetNormalized();
-
-					float NormalizedWeight = Weight.Rotation / TotalWeight.Rotation;
-
-					if(NumMixedRotations == 0)
-					{
-						FirstRotation = ParentRotation; 
-					}
-					else if ((ParentRotation | FirstRotation) <= 0.f)
-					{
-						NormalizedWeight = -NormalizedWeight;
-					}
-
-					MixedRotation.X += NormalizedWeight * ParentRotation.X;
-					MixedRotation.Y += NormalizedWeight * ParentRotation.Y;
-					MixedRotation.Z += NormalizedWeight * ParentRotation.Z;
-					MixedRotation.W += NormalizedWeight * ParentRotation.W;
-					NumMixedRotations++;
-				}
-
-				OutputTransform.SetRotation(MixedRotation.GetNormalized());
-			}
-
-			if(NumConstraintsAffecting.Scale == 1)
-			{
-				check(FirstConstraint.Scale != INDEX_NONE);
-
-				const FRigElementParentConstraint& ParentConstraint = ParentConstraints[FirstConstraint.Scale];
-				const FRigElementWeight& Weight = ParentConstraint.GetWeight(bInitial);
-				const FTransform ParentTransform = GetTransform(ParentConstraint.ParentElement, InTransformType);
-				check(Weight.AffectsScale());
-				OutputTransform.SetScale3D(ParentTransform.GetScale3D());
-			}
-			else if(NumConstraintsAffecting.Scale == 2)
-			{
-				check(FirstConstraint.Scale != INDEX_NONE);
-				check(SecondConstraint.Scale != INDEX_NONE);
-
-				const FRigElementParentConstraint& ParentConstraintA = ParentConstraints[FirstConstraint.Scale];
-				const FRigElementParentConstraint& ParentConstraintB = ParentConstraints[SecondConstraint.Scale];
-
-				const FRigElementWeight& WeightA = ParentConstraintA.GetWeight(bInitial); 
-				const FRigElementWeight& WeightB = ParentConstraintB.GetWeight(bInitial);
-				check(WeightA.AffectsScale());
-				check(WeightB.AffectsScale());
-				const float Weight = Local::GetWeightForLerp(WeightA.Scale, WeightB.Scale);
-
-				const FVector ParentScaleA = GetTransform(ParentConstraintA.ParentElement, InTransformType).GetScale3D();
-				const FVector ParentScaleB = GetTransform(ParentConstraintB.ParentElement, InTransformType).GetScale3D();
-				OutputTransform.SetScale3D(FMath::Lerp<FVector>(ParentScaleA, ParentScaleB, Weight));
-			}
-			else if(NumConstraintsAffecting.Scale > 2)
-			{
-				check(TotalWeight.Scale > SMALL_NUMBER);
-				
-				FVector Scale = FVector::ZeroVector;
-				
-				for(int32 ConstraintIndex = 0; ConstraintIndex < ParentConstraints.Num(); ConstraintIndex++)
-				{
-					const FRigElementParentConstraint& ParentConstraint = ParentConstraints[ConstraintIndex];
-					const FRigElementWeight& Weight = ParentConstraint.GetWeight(bInitial);
-					if(!Weight.AffectsScale())
-					{
-						continue;
-					}
-
-					const FVector ParentScale = GetTransform(ParentConstraint.ParentElement, InTransformType).GetScale3D();
-					Scale += ParentScale * Weight.Scale / TotalWeight.Scale; 
-				}
-
-				OutputTransform.SetScale3D(Scale);
-			}
-
-			/*
-			// if we have only one parent we behave like a single parent element
-			if(ParentConstraints.Num() == 1)
-			{
-				Output.Set(GetTransform(ParentConstraints[0], InTransformType));
-			}
-			// for exactly two parents we'll lerp
-			else if(ParentConstraints.Num() == 2)
-			{
-				float Weight = 0.f;
-				const float WeightA = ParentWeights[0];
-				const float WeightB = ParentWeights[1];
-				const float ClampedWeightA = FMath::Max(WeightA, 0.f);
-				const float ClampedWeightB = FMath::Max(WeightB, 0.f);
-				const float OverallWeight = ClampedWeightA + ClampedWeightB;
-				if(OverallWeight > SMALL_NUMBER)
-				{
-					Weight = ClampedWeightB / OverallWeight;
-				}
-
-				if(Weight <= SMALL_NUMBER)
-				{
-					Output.Set(GetTransform(ParentConstraints[0], InTransformType));
-				}
-				else if(Weight >= (1.0f - SMALL_NUMBER))
-				{
-					Output.Set(GetTransform(ParentConstraints[1], InTransformType));
-				}
-				else
-				{
-					const FTransform ParentTransformA = GetTransform(ParentConstraints[0], InTransformType);
-					const FTransform ParentTransformB = GetTransform(ParentConstraints[1], InTransformType);
-					Output.Set(FControlRigMathLibrary::LerpTransform(ParentTransformA, ParentTransformB, Weight));
-				}
-			}
-			else if(ParentConstraints.Num() > 2)
-			{
-				ensure(ParentConstraints.Num() == ParentWeights.Num());
-
-				// determine if there are more than one parent weighted
-				float OverallWeight = 0.f;
-				int32 FirstWeightedParent = INDEX_NONE;
-				int32 SecondWeightedParent = INDEX_NONE;
-				int32 NumWeightedParents = 0;
-				for(int32 ParentIndex = 0; ParentIndex < ParentWeights.Num(); ParentIndex++)
-				{
-					const float Weight = ParentWeights[ParentIndex];
-					const float ClampedWeight = FMath::Max(Weight, 0.f);
-					OverallWeight += ClampedWeight;
-
-					if(ClampedWeight > SMALL_NUMBER)
-					{
-						NumWeightedParents++;
-						if(FirstWeightedParent == INDEX_NONE)
-						{
-							FirstWeightedParent = ParentIndex;
-						}
-						else if(SecondWeightedParent == INDEX_NONE)
-						{
-							SecondWeightedParent = ParentIndex;
-						}
-					}
-				}
-
-				// if there is only weight on one parent, behave like a single parent element
-				if(NumWeightedParents == 1)
-				{
-					Output.Set(GetTransform(ParentConstraints[FirstWeightedParent], InTransformType));
-				}
-				else if(NumWeightedParents == 2) // for two weighted parents, special case once more
-				{
-					float Weight = 0.f;
-					const float WeightA = ParentWeights[FirstWeightedParent];
-					const float WeightB = ParentWeights[SecondWeightedParent];
-					const float ClampedWeightA = FMath::Max(WeightA, 0.f);
-					const float ClampedWeightB = FMath::Max(WeightB, 0.f);
-					if(OverallWeight > SMALL_NUMBER)
-					{
-						Weight = ClampedWeightB / OverallWeight;
-					}
-
-					const FTransform ParentTransformA = GetTransform(ParentConstraints[FirstWeightedParent], InTransformType);
-					const FTransform ParentTransformB = GetTransform(ParentConstraints[SecondWeightedParent], InTransformType);
-					Output.Set(FControlRigMathLibrary::LerpTransform(ParentTransformA, ParentTransformB, Weight));
-				}
-				else if(OverallWeight > SMALL_NUMBER)
-				{
-					FVector MixedTranslation = FVector::ZeroVector;
-					FQuat MixedRotation = FQuat(0.f, 0.f, 0.f, 0.f);;
-					FVector MixedScale3D = FVector::ZeroVector;
-
-					FQuat FirstRotation = FQuat::Identity;
-					int NumMixedRotations = 0;
-
-					for(int32 ParentIndex = 0; ParentIndex < ParentWeights.Num(); ParentIndex++)
-					{
-						const float Weight = ParentWeights[ParentIndex];
-						const float ClampedWeight = FMath::Max(Weight, 0.f);
-						if(ClampedWeight <= SMALL_NUMBER)
-						{
-							continue;
-						}
-						const float NormalizedWeight = ClampedWeight / OverallWeight;
-
-						const FTransform ParentTransform = GetTransform(ParentConstraints[ParentIndex], InTransformType);
-
-						// mix translation
-						MixedTranslation += ParentTransform.GetTranslation() * NormalizedWeight; 
-
-						// mix rotation
-						FQuat CurrentRotation = ParentTransform.GetRotation();
-						if(NumMixedRotations == 0)
-						{
-							FirstRotation = CurrentRotation; 
-						}
-						else if ((CurrentRotation | FirstRotation) <= 0.f)
-						{
-							// invert sign of rotation (NOT the same as .Inverse() )
-							CurrentRotation = FQuat(-CurrentRotation.X, -CurrentRotation.Y, -CurrentRotation.Z, -CurrentRotation.W);
-						}
-
-						// scale down the rotation if needed
-						if(NormalizedWeight < (1.f - SMALL_NUMBER))
-						{
-							CurrentRotation = FQuat::Slerp(FQuat::Identity, CurrentRotation, NormalizedWeight);
-						}
-
-						// accumulate part of current rotation
-						MixedRotation.W += CurrentRotation.W;
-						MixedRotation.X += CurrentRotation.X;
-						MixedRotation.Y += CurrentRotation.Y;
-						MixedRotation.Z += CurrentRotation.Z;
-
-						NumMixedRotations++;
-
-						// mix scale
-						MixedScale3D += ParentTransform.GetScale3D() * NormalizedWeight; 
-					}
-
-					float W = MixedRotation.W;
-					float X = MixedRotation.X;
-					float Y = MixedRotation.Y;
-					float Z = MixedRotation.Z;
-						
-					//Normalize. Note: experiment to see whether you can skip this step.
-					float D = 1.0f / (W * W + X * X + Y * Y + Z * Z);
-					W *= D;
-					X *= D;
-					Y *= D;
-					Z *= D;
-						
-					MixedRotation.X = X;
-					MixedRotation.Y = Y;
-					MixedRotation.Z = Z;
-					MixedRotation.W = W;
-
-					FTransform MixedTransform = FTransform::Identity;
-					MixedTransform.SetTranslation(MixedTranslation);
-					MixedTransform.SetRotation(MixedRotation.GetNormalized());
-					MixedTransform.SetScale3D(MixedScale3D);
-					Output.Set(MixedTransform);
-				}
-			}
-			*/
-
+			const FTransform OutputTransform = SolveParentConstraints(
+				MultiParentElement->ParentConstraints,
+				InTransformType,
+				FTransform::Identity,
+				false,
+				FTransform::Identity,
+				false
+			);
 			Output.Set(OutputTransform);
 			Output.bDirty = false;
 		}
@@ -3181,11 +2788,41 @@ void URigHierarchy::PropagateDirtyFlags(FRigTransformElement* InTransformElement
 		for(const FRigTransformElement::FElementToDirty& ElementToDirty : InTransformElement->ElementsToDirty)
 		{
 #if URIGHIERARCHY_RECURSIVE_DIRTY_PROPAGATION
-			if(FRigMultiParentElement* MultiParentElement = Cast<FRigMultiParentElement>(ElementToDirty.Element))
+			if(FRigControlElement* ControlElement = Cast<FRigControlElement>(ElementToDirty.Element))
 			{
-				if(MultiParentElement->Parent.IsDirty(TypeToDirty))
+				if(ERigTransformType::IsGlobal(TypeToDirty))
 				{
-					continue;
+					if(ControlElement->Parent.IsDirty(TypeToDirty) &&
+						ControlElement->Offset.IsDirty(TypeToDirty) &&
+						ControlElement->Pose.IsDirty(TypeToDirty))
+					{
+						continue;
+					}
+				}
+				else
+				{
+					if(ControlElement->Parent.IsDirty(TypeToDirty))
+					{
+						continue;
+					}
+				}
+			}
+			else if(FRigMultiParentElement* MultiParentElement = Cast<FRigMultiParentElement>(ElementToDirty.Element))
+			{
+				if(ERigTransformType::IsGlobal(TypeToDirty))
+				{
+					if(MultiParentElement->Parent.IsDirty(TypeToDirty) &&
+						MultiParentElement->Pose.IsDirty(TypeToDirty))
+					{
+						continue;
+					}
+				}
+				else
+				{
+					if(MultiParentElement->Parent.IsDirty(TypeToDirty))
+					{
+						continue;
+					}
 				}
 			}
 			else
@@ -3224,11 +2861,41 @@ void URigHierarchy::PropagateDirtyFlags(FRigTransformElement* InTransformElement
 		{
 #if URIGHIERARCHY_RECURSIVE_DIRTY_PROPAGATION
 
-			if(FRigMultiParentElement* MultiParentElement = Cast<FRigMultiParentElement>(ElementToDirty.Element))
+			if(FRigControlElement* ControlElement = Cast<FRigControlElement>(ElementToDirty.Element))
 			{
-				if(MultiParentElement->Parent.IsDirty(TypeToDirty))
+				if(ERigTransformType::IsGlobal(TypeToDirty))
 				{
-					continue;
+					if(ControlElement->Parent.IsDirty(TypeToDirty) &&
+						ControlElement->Offset.IsDirty(TypeToDirty) &&
+						ControlElement->Pose.IsDirty(TypeToDirty))
+					{
+						continue;
+					}
+				}
+				else
+				{
+					if(ControlElement->Parent.IsDirty(TypeToDirty))
+					{
+						continue;
+					}
+				}
+			}
+			else if(FRigMultiParentElement* MultiParentElement = Cast<FRigMultiParentElement>(ElementToDirty.Element))
+			{
+				if(ERigTransformType::IsGlobal(TypeToDirty))
+				{
+					if(MultiParentElement->Parent.IsDirty(TypeToDirty) &&
+						MultiParentElement->Pose.IsDirty(TypeToDirty))
+					{
+						continue;
+					}
+				}
+				else
+				{
+					if(MultiParentElement->Parent.IsDirty(TypeToDirty))
+					{
+						continue;
+					}
 				}
 			}
 			else
@@ -3494,7 +3161,614 @@ FTransform URigHierarchy::GetWorldTransformForSocket(const FRigUnitContext* InCo
 {
 	if(const USceneComponent* OuterSceneComponent = GetTypedOuter<USceneComponent>())
 	{
-		return OuterSceneComponent->GetComponentToWorld();
+		return OuterSceneComponent->GetComponentToWorld().Inverse();
 	}
 	return FTransform::Identity;
 }
+
+FTransform URigHierarchy::ComputeLocalControlValue(FRigControlElement* ControlElement,
+	const FTransform& InGlobalTransform, ERigTransformType::Type InTransformType) const
+{
+	check(ERigTransformType::IsGlobal(InTransformType));
+
+	const FTransform OffsetTransform =
+		GetControlOffsetTransform(ControlElement, ERigTransformType::MakeLocal(InTransformType));
+
+	return InverseSolveParentConstraints(
+		InGlobalTransform,
+		ControlElement->ParentConstraints,
+		InTransformType,
+		OffsetTransform);
+}
+
+FTransform URigHierarchy::SolveParentConstraints(
+	const FRigElementParentConstraintArray& InConstraints,
+	const ERigTransformType::Type InTransformType,
+	const FTransform& InLocalOffsetTransform,
+	bool bApplyLocalOffsetTransform,
+	const FTransform& InLocalPoseTransform,
+	bool bApplyLocalPoseTransform) const
+{
+	FTransform Result = FTransform::Identity;
+	const bool bInitial = IsInitial(InTransformType);
+
+	// collect all of the weights
+	FConstraintIndex FirstConstraint;
+	FConstraintIndex SecondConstraint;
+	FConstraintIndex NumConstraintsAffecting(0);
+	FRigElementWeight TotalWeight(0.f);
+	ComputeParentConstraintIndices(InConstraints, InTransformType, FirstConstraint, SecondConstraint, NumConstraintsAffecting, TotalWeight);
+
+	if(NumConstraintsAffecting.Location == 0 ||
+		NumConstraintsAffecting.Rotation == 0 ||
+		NumConstraintsAffecting.Scale == 0)
+	{
+		if(bApplyLocalOffsetTransform)
+		{
+			Result = InLocalOffsetTransform;
+		}
+		
+		if(bApplyLocalPoseTransform)
+		{
+			Result = InLocalPoseTransform * Result;
+		}
+
+		if(NumConstraintsAffecting.Location == 0 &&
+			NumConstraintsAffecting.Rotation == 0 &&
+			NumConstraintsAffecting.Scale == 0)
+		{
+			Result.NormalizeRotation();
+			return Result;
+		}
+	}
+
+	if(NumConstraintsAffecting.Location == 1)
+	{
+		check(FirstConstraint.Location != INDEX_NONE);
+
+		const FRigElementParentConstraint& ParentConstraint = InConstraints[FirstConstraint.Location];
+		const FRigElementWeight& Weight = ParentConstraint.GetWeight(bInitial);
+		const FTransform Transform = LazilyComputeParentConstraint(InConstraints, FirstConstraint.Location, InTransformType,
+			InLocalOffsetTransform, bApplyLocalOffsetTransform, InLocalPoseTransform, bApplyLocalPoseTransform);
+
+		check(Weight.AffectsLocation());
+		Result.SetLocation(Transform.GetLocation());
+	}
+	else if(NumConstraintsAffecting.Location == 2)
+	{
+		check(FirstConstraint.Location != INDEX_NONE);
+		check(SecondConstraint.Location != INDEX_NONE);
+
+		const FRigElementParentConstraint& ParentConstraintA = InConstraints[FirstConstraint.Location];
+		const FRigElementParentConstraint& ParentConstraintB = InConstraints[SecondConstraint.Location];
+
+		const FRigElementWeight& WeightA = ParentConstraintA.GetWeight(bInitial); 
+		const FRigElementWeight& WeightB = ParentConstraintB.GetWeight(bInitial);
+		check(WeightA.AffectsLocation());
+		check(WeightB.AffectsLocation());
+		const float Weight = GetWeightForLerp(WeightA.Location, WeightB.Location);
+
+		const FTransform TransformA = LazilyComputeParentConstraint(InConstraints, FirstConstraint.Location, InTransformType,
+			InLocalOffsetTransform, bApplyLocalOffsetTransform, InLocalPoseTransform, bApplyLocalPoseTransform);
+		const FTransform TransformB = LazilyComputeParentConstraint(InConstraints, SecondConstraint.Location, InTransformType,
+			InLocalOffsetTransform, bApplyLocalOffsetTransform, InLocalPoseTransform, bApplyLocalPoseTransform);
+
+		const FVector ParentLocationA = TransformA.GetLocation();
+		const FVector ParentLocationB = TransformB.GetLocation();
+		Result.SetLocation(FMath::Lerp<FVector>(ParentLocationA, ParentLocationB, Weight));
+	}
+	else if(NumConstraintsAffecting.Location > 2)
+	{
+		check(TotalWeight.Location > SMALL_NUMBER);
+		
+		FVector Location = FVector::ZeroVector;
+		
+		for(int32 ConstraintIndex = 0; ConstraintIndex < InConstraints.Num(); ConstraintIndex++)
+		{
+			const FRigElementParentConstraint& ParentConstraint = InConstraints[ConstraintIndex];
+			const FRigElementWeight& Weight = ParentConstraint.GetWeight(bInitial);
+			if(!Weight.AffectsLocation())
+			{
+				continue;
+			}
+
+			const FTransform Transform = LazilyComputeParentConstraint(InConstraints, ConstraintIndex, InTransformType,
+				InLocalOffsetTransform, bApplyLocalOffsetTransform, InLocalPoseTransform, bApplyLocalPoseTransform);
+
+			IntegrateParentConstraintVector(Location, Transform, Weight.Location / TotalWeight.Location, true);
+		}
+
+		Result.SetLocation(Location);
+	}
+
+	if(NumConstraintsAffecting.Rotation == 1)
+	{
+		check(FirstConstraint.Rotation != INDEX_NONE);
+
+		const FRigElementParentConstraint& ParentConstraint = InConstraints[FirstConstraint.Rotation];
+		const FRigElementWeight& Weight = ParentConstraint.GetWeight(bInitial);
+		const FTransform Transform = LazilyComputeParentConstraint(InConstraints, FirstConstraint.Rotation, InTransformType,
+			InLocalOffsetTransform, bApplyLocalOffsetTransform, InLocalPoseTransform, bApplyLocalPoseTransform);
+		check(Weight.AffectsRotation());
+		Result.SetRotation(Transform.GetRotation());
+	}
+	else if(NumConstraintsAffecting.Rotation == 2)
+	{
+		check(FirstConstraint.Rotation != INDEX_NONE);
+		check(SecondConstraint.Rotation != INDEX_NONE);
+
+		const FRigElementParentConstraint& ParentConstraintA = InConstraints[FirstConstraint.Rotation];
+		const FRigElementParentConstraint& ParentConstraintB = InConstraints[SecondConstraint.Rotation];
+
+		const FRigElementWeight& WeightA = ParentConstraintA.GetWeight(bInitial); 
+		const FRigElementWeight& WeightB = ParentConstraintB.GetWeight(bInitial);
+		check(WeightA.AffectsRotation());
+		check(WeightB.AffectsRotation());
+		const float Weight = GetWeightForLerp(WeightA.Rotation, WeightB.Rotation);
+
+		const FTransform TransformA = LazilyComputeParentConstraint(InConstraints, FirstConstraint.Rotation, InTransformType,
+			InLocalOffsetTransform, bApplyLocalOffsetTransform, InLocalPoseTransform, bApplyLocalPoseTransform);
+		const FTransform TransformB = LazilyComputeParentConstraint(InConstraints, SecondConstraint.Rotation, InTransformType,
+			InLocalOffsetTransform, bApplyLocalOffsetTransform, InLocalPoseTransform, bApplyLocalPoseTransform);
+
+		const FQuat ParentRotationA = TransformA.GetRotation();
+		const FQuat ParentRotationB = TransformB.GetRotation();
+		Result.SetRotation(FQuat::Slerp(ParentRotationA, ParentRotationB, Weight));
+	}
+	else if(NumConstraintsAffecting.Rotation > 2)
+	{
+		check(TotalWeight.Rotation > SMALL_NUMBER);
+		
+		int NumMixedRotations = 0;
+		FQuat FirstRotation = FQuat::Identity;
+		FQuat MixedRotation = FQuat(0.f, 0.f, 0.f, 0.f);
+
+		for(int32 ConstraintIndex = 0; ConstraintIndex < InConstraints.Num(); ConstraintIndex++)
+		{
+			const FRigElementParentConstraint& ParentConstraint = InConstraints[ConstraintIndex];
+			const FRigElementWeight& Weight = ParentConstraint.GetWeight(bInitial);
+			if(!Weight.AffectsRotation())
+			{
+				continue;
+			}
+
+			const FTransform Transform = LazilyComputeParentConstraint(InConstraints, ConstraintIndex, InTransformType,
+				InLocalOffsetTransform, bApplyLocalOffsetTransform, InLocalPoseTransform, bApplyLocalPoseTransform);
+
+			IntegrateParentConstraintQuat(
+				NumMixedRotations,
+				FirstRotation,
+				MixedRotation,
+				Transform,
+				Weight.Rotation / TotalWeight.Rotation);
+		}
+
+		Result.SetRotation(MixedRotation.GetNormalized());
+	}
+
+	if(NumConstraintsAffecting.Scale == 1)
+	{
+		check(FirstConstraint.Scale != INDEX_NONE);
+
+		const FRigElementParentConstraint& ParentConstraint = InConstraints[FirstConstraint.Scale];
+		const FRigElementWeight& Weight = ParentConstraint.GetWeight(bInitial);
+		
+		const FTransform Transform = LazilyComputeParentConstraint(InConstraints, FirstConstraint.Scale, InTransformType,
+			InLocalOffsetTransform, bApplyLocalOffsetTransform, InLocalPoseTransform, bApplyLocalPoseTransform);
+
+		check(Weight.AffectsScale());
+		Result.SetScale3D(Transform.GetScale3D());
+	}
+	else if(NumConstraintsAffecting.Scale == 2)
+	{
+		check(FirstConstraint.Scale != INDEX_NONE);
+		check(SecondConstraint.Scale != INDEX_NONE);
+
+		const FRigElementParentConstraint& ParentConstraintA = InConstraints[FirstConstraint.Scale];
+		const FRigElementParentConstraint& ParentConstraintB = InConstraints[SecondConstraint.Scale];
+
+		const FRigElementWeight& WeightA = ParentConstraintA.GetWeight(bInitial); 
+		const FRigElementWeight& WeightB = ParentConstraintB.GetWeight(bInitial);
+		check(WeightA.AffectsScale());
+		check(WeightB.AffectsScale());
+		const float Weight = GetWeightForLerp(WeightA.Scale, WeightB.Scale);
+
+		const FTransform TransformA = LazilyComputeParentConstraint(InConstraints, FirstConstraint.Scale, InTransformType,
+			InLocalOffsetTransform, bApplyLocalOffsetTransform, InLocalPoseTransform, bApplyLocalPoseTransform);
+		const FTransform TransformB = LazilyComputeParentConstraint(InConstraints, SecondConstraint.Scale, InTransformType,
+			InLocalOffsetTransform, bApplyLocalOffsetTransform, InLocalPoseTransform, bApplyLocalPoseTransform);
+
+		const FVector ParentScaleA = TransformA.GetScale3D();
+		const FVector ParentScaleB = TransformB.GetScale3D();
+		Result.SetScale3D(FMath::Lerp<FVector>(ParentScaleA, ParentScaleB, Weight));
+	}
+	else if(NumConstraintsAffecting.Scale > 2)
+	{
+		check(TotalWeight.Scale > SMALL_NUMBER);
+		
+		FVector Scale = FVector::ZeroVector;
+		
+		for(int32 ConstraintIndex = 0; ConstraintIndex < InConstraints.Num(); ConstraintIndex++)
+		{
+			const FRigElementParentConstraint& ParentConstraint = InConstraints[ConstraintIndex];
+			const FRigElementWeight& Weight = ParentConstraint.GetWeight(bInitial);
+			if(!Weight.AffectsScale())
+			{
+				continue;
+			}
+
+			const FTransform Transform = LazilyComputeParentConstraint(InConstraints, ConstraintIndex, InTransformType,
+				InLocalOffsetTransform, bApplyLocalOffsetTransform, InLocalPoseTransform, bApplyLocalPoseTransform);
+
+			IntegrateParentConstraintVector(Scale, Transform, Weight.Scale / TotalWeight.Scale, false);
+		}
+
+		Result.SetScale3D(Scale);
+	}
+
+	Result.NormalizeRotation();
+	return Result;
+}
+
+FTransform URigHierarchy::InverseSolveParentConstraints(
+	const FTransform& InGlobalTransform,
+	const FRigElementParentConstraintArray& InConstraints,
+	const ERigTransformType::Type InTransformType,
+	const FTransform& InLocalOffsetTransform) const
+{
+	FTransform Result = FTransform::Identity;
+	const bool bInitial = IsInitial(InTransformType);
+
+	// collect all of the weights
+	FConstraintIndex FirstConstraint;
+	FConstraintIndex SecondConstraint;
+	FConstraintIndex NumConstraintsAffecting(0);
+	FRigElementWeight TotalWeight(0.f);
+	ComputeParentConstraintIndices(InConstraints, InTransformType, FirstConstraint, SecondConstraint, NumConstraintsAffecting, TotalWeight);
+
+	if(NumConstraintsAffecting.Location == 0 ||
+		NumConstraintsAffecting.Rotation == 0 ||
+		NumConstraintsAffecting.Scale == 0)
+	{
+		Result = InGlobalTransform.GetRelativeTransform(InLocalOffsetTransform);
+		
+		if(NumConstraintsAffecting.Location == 0 &&
+			NumConstraintsAffecting.Rotation == 0 &&
+			NumConstraintsAffecting.Scale == 0)
+		{
+			Result.NormalizeRotation();
+			return Result;
+		}
+	}
+
+	if(NumConstraintsAffecting.Location == 1)
+	{
+		check(FirstConstraint.Location != INDEX_NONE);
+
+		const FRigElementParentConstraint& ParentConstraint = InConstraints[FirstConstraint.Location];
+		const FRigElementWeight& Weight = ParentConstraint.GetWeight(bInitial);
+		const FTransform Transform = LazilyComputeParentConstraint(InConstraints, FirstConstraint.Location, InTransformType,
+			InLocalOffsetTransform, true, FTransform::Identity, false);
+
+		check(Weight.AffectsLocation());
+		Result.SetLocation(InGlobalTransform.GetRelativeTransform(Transform).GetLocation());
+	}
+	else if(NumConstraintsAffecting.Location == 2)
+	{
+		check(FirstConstraint.Location != INDEX_NONE);
+		check(SecondConstraint.Location != INDEX_NONE);
+
+		const FRigElementParentConstraint& ParentConstraintA = InConstraints[FirstConstraint.Location];
+		const FRigElementParentConstraint& ParentConstraintB = InConstraints[SecondConstraint.Location];
+
+		const FRigElementWeight& WeightA = ParentConstraintA.GetWeight(bInitial); 
+		const FRigElementWeight& WeightB = ParentConstraintB.GetWeight(bInitial);
+		check(WeightA.AffectsLocation());
+		check(WeightB.AffectsLocation());
+		const float Weight = GetWeightForLerp(WeightA.Location, WeightB.Location);
+
+		const FTransform TransformA = LazilyComputeParentConstraint(InConstraints, FirstConstraint.Location, InTransformType,
+			InLocalOffsetTransform, true, FTransform::Identity, false);
+		const FTransform TransformB = LazilyComputeParentConstraint(InConstraints, SecondConstraint.Location, InTransformType,
+			InLocalOffsetTransform, true, FTransform::Identity, false);
+
+		const FTransform MixedTransform = FControlRigMathLibrary::LerpTransform(TransformA, TransformB, Weight);
+		Result.SetLocation(InGlobalTransform.GetRelativeTransform(MixedTransform).GetLocation());
+	}
+	else if(NumConstraintsAffecting.Location > 2)
+	{
+		check(TotalWeight.Location > SMALL_NUMBER);
+		
+		FVector Location = FVector::ZeroVector;
+		int NumMixedRotations = 0;
+		FQuat FirstRotation = FQuat::Identity;
+		FQuat MixedRotation = FQuat(0.f, 0.f, 0.f, 0.f);
+		FVector Scale = FVector::ZeroVector;
+
+		for(int32 ConstraintIndex = 0; ConstraintIndex < InConstraints.Num(); ConstraintIndex++)
+		{
+			const FRigElementParentConstraint& ParentConstraint = InConstraints[ConstraintIndex];
+			const FRigElementWeight& Weight = ParentConstraint.GetWeight(bInitial);
+			if(!Weight.AffectsLocation())
+			{
+				continue;
+			}
+
+			const FTransform Transform = LazilyComputeParentConstraint(InConstraints, ConstraintIndex, InTransformType,
+				InLocalOffsetTransform, true, FTransform::Identity, false);
+
+			const float NormalizedWeight = Weight.Location / TotalWeight.Location;
+			IntegrateParentConstraintVector(Location, Transform, NormalizedWeight, true);
+			IntegrateParentConstraintQuat(NumMixedRotations, FirstRotation, MixedRotation, Transform, NormalizedWeight);
+			IntegrateParentConstraintVector(Scale, Transform, NormalizedWeight, false);
+		}
+
+		FTransform ParentTransform(MixedRotation.GetNormalized(), Location, Scale);
+		Result.SetLocation(InGlobalTransform.GetRelativeTransform(ParentTransform).GetLocation());
+	}
+
+	if(NumConstraintsAffecting.Rotation == 1)
+	{
+		check(FirstConstraint.Rotation != INDEX_NONE);
+
+		const FRigElementParentConstraint& ParentConstraint = InConstraints[FirstConstraint.Rotation];
+		const FRigElementWeight& Weight = ParentConstraint.GetWeight(bInitial);
+		const FTransform Transform = LazilyComputeParentConstraint(InConstraints, FirstConstraint.Rotation, InTransformType,
+			InLocalOffsetTransform, true, FTransform::Identity, false);
+		check(Weight.AffectsRotation());
+		Result.SetRotation(InGlobalTransform.GetRelativeTransform(Transform).GetRotation());
+	}
+	else if(NumConstraintsAffecting.Rotation == 2)
+	{
+		check(FirstConstraint.Rotation != INDEX_NONE);
+		check(SecondConstraint.Rotation != INDEX_NONE);
+
+		const FRigElementParentConstraint& ParentConstraintA = InConstraints[FirstConstraint.Rotation];
+		const FRigElementParentConstraint& ParentConstraintB = InConstraints[SecondConstraint.Rotation];
+
+		const FRigElementWeight& WeightA = ParentConstraintA.GetWeight(bInitial); 
+		const FRigElementWeight& WeightB = ParentConstraintB.GetWeight(bInitial);
+		check(WeightA.AffectsRotation());
+		check(WeightB.AffectsRotation());
+		const float Weight = GetWeightForLerp(WeightA.Rotation, WeightB.Rotation);
+
+		const FTransform TransformA = LazilyComputeParentConstraint(InConstraints, FirstConstraint.Rotation, InTransformType,
+			InLocalOffsetTransform, true, FTransform::Identity, false);
+		const FTransform TransformB = LazilyComputeParentConstraint(InConstraints, SecondConstraint.Rotation, InTransformType,
+			InLocalOffsetTransform, true, FTransform::Identity, false);
+
+		const FTransform MixedTransform = FControlRigMathLibrary::LerpTransform(TransformA, TransformB, Weight);
+		Result.SetRotation(InGlobalTransform.GetRelativeTransform(MixedTransform).GetRotation());
+	}
+	else if(NumConstraintsAffecting.Rotation > 2)
+	{
+		check(TotalWeight.Rotation > SMALL_NUMBER);
+		
+		FVector Location = FVector::ZeroVector;
+		int NumMixedRotations = 0;
+		FQuat FirstRotation = FQuat::Identity;
+		FQuat MixedRotation = FQuat(0.f, 0.f, 0.f, 0.f);
+		FVector Scale = FVector::ZeroVector;
+
+		for(int32 ConstraintIndex = 0; ConstraintIndex < InConstraints.Num(); ConstraintIndex++)
+		{
+			const FRigElementParentConstraint& ParentConstraint = InConstraints[ConstraintIndex];
+			const FRigElementWeight& Weight = ParentConstraint.GetWeight(bInitial);
+			if(!Weight.AffectsRotation())
+			{
+				continue;
+			}
+
+			const FTransform Transform = LazilyComputeParentConstraint(InConstraints, ConstraintIndex, InTransformType,
+				InLocalOffsetTransform, true, FTransform::Identity, false);
+
+			const float NormalizedWeight = Weight.Rotation / TotalWeight.Rotation;
+			IntegrateParentConstraintVector(Location, Transform, NormalizedWeight, true);
+			IntegrateParentConstraintQuat(NumMixedRotations, FirstRotation, MixedRotation, Transform, NormalizedWeight);
+			IntegrateParentConstraintVector(Scale, Transform, NormalizedWeight, false);
+		}
+
+		FTransform ParentTransform(MixedRotation.GetNormalized(), Location, Scale);
+		Result.SetRotation(InGlobalTransform.GetRelativeTransform(ParentTransform).GetRotation());
+	}
+
+	if(NumConstraintsAffecting.Scale == 1)
+	{
+		check(FirstConstraint.Scale != INDEX_NONE);
+
+		const FRigElementParentConstraint& ParentConstraint = InConstraints[FirstConstraint.Scale];
+		const FRigElementWeight& Weight = ParentConstraint.GetWeight(bInitial);
+		
+		const FTransform Transform = LazilyComputeParentConstraint(InConstraints, FirstConstraint.Scale, InTransformType,
+			InLocalOffsetTransform, true, FTransform::Identity, false);
+
+		check(Weight.AffectsScale());
+		Result.SetScale3D(InGlobalTransform.GetRelativeTransform(Transform).GetScale3D());
+	}
+	else if(NumConstraintsAffecting.Scale == 2)
+	{
+		check(FirstConstraint.Scale != INDEX_NONE);
+		check(SecondConstraint.Scale != INDEX_NONE);
+
+		const FRigElementParentConstraint& ParentConstraintA = InConstraints[FirstConstraint.Scale];
+		const FRigElementParentConstraint& ParentConstraintB = InConstraints[SecondConstraint.Scale];
+
+		const FRigElementWeight& WeightA = ParentConstraintA.GetWeight(bInitial); 
+		const FRigElementWeight& WeightB = ParentConstraintB.GetWeight(bInitial);
+		check(WeightA.AffectsScale());
+		check(WeightB.AffectsScale());
+		const float Weight = GetWeightForLerp(WeightA.Scale, WeightB.Scale);
+
+		const FTransform TransformA = LazilyComputeParentConstraint(InConstraints, FirstConstraint.Scale, InTransformType,
+			InLocalOffsetTransform, true, FTransform::Identity, false);
+		const FTransform TransformB = LazilyComputeParentConstraint(InConstraints, SecondConstraint.Scale, InTransformType,
+			InLocalOffsetTransform, true, FTransform::Identity, false);
+
+		const FTransform MixedTransform = FControlRigMathLibrary::LerpTransform(TransformA, TransformB, Weight);
+		Result.SetScale3D(InGlobalTransform.GetRelativeTransform(MixedTransform).GetScale3D());
+	}
+	else if(NumConstraintsAffecting.Scale > 2)
+	{
+		check(TotalWeight.Scale > SMALL_NUMBER);
+		
+		FVector Location = FVector::ZeroVector;
+		int NumMixedRotations = 0;
+		FQuat FirstRotation = FQuat::Identity;
+		FQuat MixedRotation = FQuat(0.f, 0.f, 0.f, 0.f);
+		FVector Scale = FVector::ZeroVector;
+		
+		for(int32 ConstraintIndex = 0; ConstraintIndex < InConstraints.Num(); ConstraintIndex++)
+		{
+			const FRigElementParentConstraint& ParentConstraint = InConstraints[ConstraintIndex];
+			const FRigElementWeight& Weight = ParentConstraint.GetWeight(bInitial);
+			if(!Weight.AffectsScale())
+			{
+				continue;
+			}
+
+			const FTransform Transform = LazilyComputeParentConstraint(InConstraints, ConstraintIndex, InTransformType,
+				InLocalOffsetTransform, true, FTransform::Identity, false);
+
+			const float NormalizedWeight = Weight.Scale / TotalWeight.Scale;
+			IntegrateParentConstraintVector(Location, Transform, NormalizedWeight, true);
+			IntegrateParentConstraintQuat(NumMixedRotations, FirstRotation, MixedRotation, Transform, NormalizedWeight);
+			IntegrateParentConstraintVector(Scale, Transform, NormalizedWeight, false);
+		}
+
+		FTransform ParentTransform(MixedRotation.GetNormalized(), Location, Scale);
+		Result.SetScale3D(InGlobalTransform.GetRelativeTransform(ParentTransform).GetScale3D());
+	}
+
+	Result.NormalizeRotation();
+	return Result;
+}
+
+FTransform URigHierarchy::LazilyComputeParentConstraint(
+	const FRigElementParentConstraintArray& InConstraints,
+	int32 InIndex,
+	const ERigTransformType::Type InTransformType,
+	const FTransform& InLocalOffsetTransform,
+	bool bApplyLocalOffsetTransform,
+	const FTransform& InLocalPoseTransform,
+	bool bApplyLocalPoseTransform) const
+{
+	const FRigElementParentConstraint& Constraint = InConstraints[InIndex];
+	if(Constraint.Cache.bDirty)
+	{
+		FTransform Transform = GetTransform(Constraint.ParentElement, InTransformType);
+		if(bApplyLocalOffsetTransform)
+		{
+			Transform = InLocalOffsetTransform * Transform;
+		}
+		if(bApplyLocalPoseTransform)
+		{
+			Transform = InLocalPoseTransform * Transform;
+		}
+
+		Constraint.Cache.Transform = Transform;
+		Constraint.Cache.bDirty = false;
+	}
+	return Constraint.Cache.Transform;
+}
+
+void URigHierarchy::ComputeParentConstraintIndices(
+	const FRigElementParentConstraintArray& InConstraints,
+	ERigTransformType::Type InTransformType,
+	FConstraintIndex& OutFirstConstraint,
+	FConstraintIndex& OutSecondConstraint,
+	FConstraintIndex& OutNumConstraintsAffecting,
+	FRigElementWeight& OutTotalWeight)
+{
+	const bool bInitial = IsInitial(InTransformType);
+	
+	// find all of the weights affecting this output
+	for(int32 ConstraintIndex = 0; ConstraintIndex < InConstraints.Num(); ConstraintIndex++)
+	{
+		InConstraints[ConstraintIndex].Cache.bDirty = true;
+		
+		const FRigElementWeight& Weight = InConstraints[ConstraintIndex].GetWeight(bInitial);
+		if(Weight.AffectsLocation())
+		{
+			OutNumConstraintsAffecting.Location++;
+			OutTotalWeight.Location += Weight.Location;
+
+			if(OutFirstConstraint.Location == INDEX_NONE)
+			{
+				OutFirstConstraint.Location = ConstraintIndex;
+			}
+			else if(OutSecondConstraint.Location == INDEX_NONE)
+			{
+				OutSecondConstraint.Location = ConstraintIndex;
+			}
+		}
+		if(Weight.AffectsRotation())
+		{
+			OutNumConstraintsAffecting.Rotation++;
+			OutTotalWeight.Rotation += Weight.Rotation;
+
+			if(OutFirstConstraint.Rotation == INDEX_NONE)
+			{
+				OutFirstConstraint.Rotation = ConstraintIndex;
+			}
+			else if(OutSecondConstraint.Rotation == INDEX_NONE)
+			{
+				OutSecondConstraint.Rotation = ConstraintIndex;
+			}
+		}
+		if(Weight.AffectsScale())
+		{
+			OutNumConstraintsAffecting.Scale++;
+			OutTotalWeight.Scale += Weight.Scale;
+
+			if(OutFirstConstraint.Scale == INDEX_NONE)
+			{
+				OutFirstConstraint.Scale = ConstraintIndex;
+			}
+			else if(OutSecondConstraint.Scale == INDEX_NONE)
+			{
+				OutSecondConstraint.Scale = ConstraintIndex;
+			}
+		}
+	}
+}
+
+void URigHierarchy::IntegrateParentConstraintVector(
+	FVector& OutVector,
+	const FTransform& InTransform,
+	float InWeight,
+	bool bIsLocation)
+{
+	if(bIsLocation)
+	{
+		OutVector += InTransform.GetLocation() * InWeight;
+	}
+	else
+	{
+		OutVector += InTransform.GetScale3D() * InWeight;
+	}
+}
+
+void URigHierarchy::IntegrateParentConstraintQuat(
+	int32& OutNumMixedRotations,
+	FQuat& OutFirstRotation,
+	FQuat& OutMixedRotation,
+	const FTransform& InTransform,
+	float InWeight)
+{
+	FQuat ParentRotation = InTransform.GetRotation().GetNormalized();
+
+	if(OutNumMixedRotations == 0)
+	{
+		OutFirstRotation = ParentRotation; 
+	}
+	else if ((ParentRotation | OutFirstRotation) <= 0.f)
+	{
+		InWeight = -InWeight;
+	}
+
+	OutMixedRotation.X += InWeight * ParentRotation.X;
+	OutMixedRotation.Y += InWeight * ParentRotation.Y;
+	OutMixedRotation.Z += InWeight * ParentRotation.Z;
+	OutMixedRotation.W += InWeight * ParentRotation.W;
+	OutNumMixedRotations++;
+}
diff --git a/Engine/Plugins/Animation/ControlRig/Source/ControlRig/Private/Rigs/RigHierarchyElements.cpp b/Engine/Plugins/Animation/ControlRig/Source/ControlRig/Private/Rigs/RigHierarchyElements.cpp
--- a/Engine/Plugins/Animation/ControlRig/Source/ControlRig/Private/Rigs/RigHierarchyElements.cpp
+++ b/Engine/Plugins/Animation/ControlRig/Source/ControlRig/Private/Rigs/RigHierarchyElements.cpp
@@ -300,6 +300,7 @@ void FRigMultiParentElement::Load(FArchive& Ar, URigHierarchy* Hierarchy, ESeria
 			ensure(ParentKey.IsValid());
 
 			ParentConstraints[ParentIndex].ParentElement = Hierarchy->FindChecked<FRigTransformElement>(ParentKey);
+			ParentConstraints[ParentIndex].Cache.bDirty = true;
 
 			if (Ar.CustomVer(FControlRigObjectVersion::GUID) >= FControlRigObjectVersion::RigHierarchyMultiParentConstraints)
 			{
diff --git a/Engine/Plugins/Animation/ControlRig/Source/ControlRig/Public/Rigs/RigHierarchy.h b/Engine/Plugins/Animation/ControlRig/Source/ControlRig/Public/Rigs/RigHierarchy.h
--- a/Engine/Plugins/Animation/ControlRig/Source/ControlRig/Public/Rigs/RigHierarchy.h
+++ b/Engine/Plugins/Animation/ControlRig/Source/ControlRig/Public/Rigs/RigHierarchy.h
@@ -2683,6 +2683,88 @@ protected:
 
 	FTransform GetWorldTransformForSocket(const FRigUnitContext* InContext, const FRigElementKey& InKey, bool bInitial);
 	
+	FORCEINLINE static float GetWeightForLerp(const float WeightA, const float WeightB)
+	{
+		float Weight = 0.f;
+		const float ClampedWeightA = FMath::Max(WeightA, 0.f);
+		const float ClampedWeightB = FMath::Max(WeightB, 0.f);
+		const float OverallWeight = ClampedWeightA + ClampedWeightB;
+		if(OverallWeight > SMALL_NUMBER)
+		{
+			Weight = ClampedWeightB / OverallWeight;
+		}
+		return Weight;
+	}
+
+	struct FConstraintIndex
+	{
+		int32 Location;
+		int32 Rotation;
+		int32 Scale;
+
+		FConstraintIndex()
+			: Location(INDEX_NONE)
+			, Rotation(INDEX_NONE)
+			, Scale(INDEX_NONE)
+		{}
+
+		FConstraintIndex(int32 InIndex)
+			: Location(InIndex)
+			, Rotation(InIndex)
+			, Scale(InIndex)
+		{}
+	};
+
+	FTransform ComputeLocalControlValue(
+		FRigControlElement* ControlElement,
+		const FTransform& InGlobalTransform,
+		ERigTransformType::Type InTransformType) const;
+	
+	FTransform SolveParentConstraints(
+		const FRigElementParentConstraintArray& InConstraints,
+		const ERigTransformType::Type InTransformType,
+		const FTransform& InLocalOffsetTransform,
+		bool bApplyLocalOffsetTransform,
+		const FTransform& InLocalPoseTransform,
+		bool bApplyLocalPoseTransform) const;
+
+	FTransform InverseSolveParentConstraints(
+		const FTransform& InGlobalTransform,
+		const FRigElementParentConstraintArray& InConstraints,
+		const ERigTransformType::Type InTransformType,
+		const FTransform& InLocalOffsetTransform) const;
+
+	FTransform LazilyComputeParentConstraint(
+		const FRigElementParentConstraintArray& Constraints,
+		int32 InIndex,
+		const ERigTransformType::Type InTransformType,
+		const FTransform& InLocalOffsetTransform,
+		bool bApplyLocalOffsetTransform,
+		const FTransform& InLocalPoseTransform,
+		bool bApplyLocalPoseTransform) const;
+
+	static void ComputeParentConstraintIndices(
+		const FRigElementParentConstraintArray& InConstraints,
+		ERigTransformType::Type InTransformType,
+		FConstraintIndex& OutFirstConstraint,
+		FConstraintIndex& OutSecondConstraint,
+		FConstraintIndex& OutNumConstraintsAffecting,
+		FRigElementWeight& OutTotalWeight
+	);
+
+	static void IntegrateParentConstraintVector(
+		FVector& OutVector,
+		const FTransform& InTransform,
+		float InWeight,
+		bool bIsLocation);
+
+	static void IntegrateParentConstraintQuat(
+		int32& OutNumMixedRotations,
+		FQuat& OutFirstRotation,
+		FQuat& OutMixedRotation,
+		const FTransform& InTransform,
+		float InWeight);
+
 	friend class URigHierarchyController;
 	friend class UControlRig;
 	friend class FControlRigEditor;
diff --git a/Engine/Plugins/Animation/ControlRig/Source/ControlRig/Public/Rigs/RigHierarchyElements.h b/Engine/Plugins/Animation/ControlRig/Source/ControlRig/Public/Rigs/RigHierarchyElements.h
--- a/Engine/Plugins/Animation/ControlRig/Source/ControlRig/Public/Rigs/RigHierarchyElements.h
+++ b/Engine/Plugins/Animation/ControlRig/Source/ControlRig/Public/Rigs/RigHierarchyElements.h
@@ -714,10 +714,12 @@ public:
 	FRigTransformElement* ParentElement;
 	FRigElementWeight Weight;
 	FRigElementWeight InitialWeight;
+	mutable FRigComputedTransform Cache;
 		
 	FORCEINLINE FRigElementParentConstraint()
 		: ParentElement(nullptr)
 	{
+		Cache.bDirty = true;
 	}
 
 	FORCEINLINE const FRigElementWeight& GetWeight(bool bInitial = false) const
