Control Rig: Sync floor actor with floor mesh component in preview scene

diff --git a/Engine/Plugins/Animation/ControlRig/Source/ControlRigEditor/ControlRigEditor.Build.cs b/Engine/Plugins/Animation/ControlRig/Source/ControlRigEditor/ControlRigEditor.Build.cs
--- a/Engine/Plugins/Animation/ControlRig/Source/ControlRigEditor/ControlRigEditor.Build.cs
+++ b/Engine/Plugins/Animation/ControlRig/Source/ControlRigEditor/ControlRigEditor.Build.cs
@@ -82,6 +82,7 @@ namespace UnrealBuildTool.Rules
 					"PropertyAccessEditor",
 					"KismetWidgets",
 					"PythonScriptPlugin",
+					"AdvancedPreviewScene",
 				}
             );
 
diff --git a/Engine/Plugins/Animation/ControlRig/Source/ControlRigEditor/Private/Editor/ControlRigEditor.cpp b/Engine/Plugins/Animation/ControlRig/Source/ControlRigEditor/Private/Editor/ControlRigEditor.cpp
--- a/Engine/Plugins/Animation/ControlRig/Source/ControlRigEditor/Private/Editor/ControlRigEditor.cpp
+++ b/Engine/Plugins/Animation/ControlRig/Source/ControlRigEditor/Private/Editor/ControlRigEditor.cpp
@@ -2964,6 +2964,16 @@ void FControlRigEditor::Tick(float DeltaTime)
 			}
 		}
 	}
+
+	if(WeakGroundActorPtr.IsValid())
+	{
+		const TSharedRef<IPersonaPreviewScene> CurrentPreviewScene = GetPersonaToolkit()->GetPreviewScene();
+		WeakGroundActorPtr->GetStaticMeshComponent()->SetVisibility(CurrentPreviewScene->GetFloorVisibility());
+
+		const float FloorOffset = CurrentPreviewScene->GetFloorOffset();
+		const FTransform FloorTransform(FRotator(0, 0, 0), FVector(0, 0, -(FloorOffset)), FVector(4.0f, 4.0f, 1.0f));
+		WeakGroundActorPtr->GetStaticMeshComponent()->SetRelativeTransform(FloorTransform);
+	}
 }
 
 bool FControlRigEditor::IsEditable(UEdGraph* InGraph) const
@@ -3693,6 +3703,8 @@ void FControlRigEditor::HandlePreviewSceneCreated(const TSharedRef<IPersonaPrevi
 	GroundActor->GetStaticMeshComponent()->SetCollisionProfileName(UCollisionProfile::BlockAll_ProfileName);
 	GroundActor->GetStaticMeshComponent()->bSelectable = false;
 
+	WeakGroundActorPtr = GroundActor;
+
 	AAnimationEditorPreviewActor* Actor = InPersonaPreviewScene->GetWorld()->SpawnActor<AAnimationEditorPreviewActor>(AAnimationEditorPreviewActor::StaticClass(), FTransform::Identity);
 	Actor->SetFlags(RF_Transient);
 	InPersonaPreviewScene->SetActor(Actor);
diff --git a/Engine/Plugins/Animation/ControlRig/Source/ControlRigEditor/Private/Editor/ControlRigEditor.h b/Engine/Plugins/Animation/ControlRig/Source/ControlRigEditor/Private/Editor/ControlRigEditor.h
--- a/Engine/Plugins/Animation/ControlRig/Source/ControlRigEditor/Private/Editor/ControlRigEditor.h
+++ b/Engine/Plugins/Animation/ControlRig/Source/ControlRigEditor/Private/Editor/ControlRigEditor.h
@@ -477,6 +477,7 @@ protected:
 
 	TMap<int32, int32> CapsuleToHierarchyIndex;
 	TArray<TStrongObjectPtr<UDetailsViewWrapperObject>> WrapperObjects;
+	TWeakObjectPtr<AStaticMeshActor> WeakGroundActorPtr;
 	
 	friend class FControlRigEditorMode;
 	friend class SControlRigStackView;
diff --git a/Engine/Source/Editor/AdvancedPreviewScene/Private/AdvancedPreviewScene.cpp b/Engine/Source/Editor/AdvancedPreviewScene/Private/AdvancedPreviewScene.cpp
--- a/Engine/Source/Editor/AdvancedPreviewScene/Private/AdvancedPreviewScene.cpp
+++ b/Engine/Source/Editor/AdvancedPreviewScene/Private/AdvancedPreviewScene.cpp
@@ -193,6 +193,11 @@ FLinearColor FAdvancedPreviewScene::GetBackgroundColor() const
 	return Color * DefaultSettings->Profiles[CurrentProfileIndex].EnvironmentIntensity;
 }
 
+float FAdvancedPreviewScene::GetFloorOffset() const
+{
+	return -(float)FloorMeshComponent->GetRelativeTransform().GetLocation().Z;
+}
+
 void FAdvancedPreviewScene::SetFloorOffset(const float InFloorOffset)
 {
 	FTransform FloorTransform(FRotator(0, 0, 0), FVector(0, 0, -(InFloorOffset)), FVector(4.0f, 4.0f, 1.0f));
@@ -306,6 +311,15 @@ const bool FAdvancedPreviewScene::HandleInputKey(FViewport* InViewport, int32 Co
 	return false;
 }
 
+bool FAdvancedPreviewScene::GetFloorVisibility() const
+{
+	if(FloorMeshComponent)
+	{
+		return FloorMeshComponent->GetVisibleFlag();
+	}
+	return false;
+}
+
 void FAdvancedPreviewScene::SetFloorVisibility(const bool bVisible, const bool bDirect)
 {
 	// If not direct set visibility in profile and refresh the scene
diff --git a/Engine/Source/Editor/AdvancedPreviewScene/Public/AdvancedPreviewScene.h b/Engine/Source/Editor/AdvancedPreviewScene/Public/AdvancedPreviewScene.h
--- a/Engine/Source/Editor/AdvancedPreviewScene/Public/AdvancedPreviewScene.h
+++ b/Engine/Source/Editor/AdvancedPreviewScene/Public/AdvancedPreviewScene.h
@@ -45,8 +45,10 @@ public:
 
 	void SetSkyRotation(const float SkyRotation);
 	/* Sets the visiblity state for the floor/environment by storing it in the scene profile and refreshing the scene, in case bDirect is true it sets the visibility directly and leaves the profile untouched. */
+	bool GetFloorVisibility() const;
 	void SetFloorVisibility(const bool bVisible, const bool bDirect = false);
 	void SetEnvironmentVisibility(const bool bVisible, const bool bDirect = false);
+	float GetFloorOffset() const;
 	void SetFloorOffset(const float InFloorOffset);
 	void SetProfileIndex(const int32 InProfileIndex);
 
