Control Rig: Space switching dialog to close based on focus loss

diff --git a/Engine/Plugins/Animation/ControlRig/Source/ControlRigEditor/Private/Editor/SRigSpacePickerWidget.cpp b/Engine/Plugins/Animation/ControlRig/Source/ControlRigEditor/Private/Editor/SRigSpacePickerWidget.cpp
--- a/Engine/Plugins/Animation/ControlRig/Source/ControlRigEditor/Private/Editor/SRigSpacePickerWidget.cpp
+++ b/Engine/Plugins/Animation/ControlRig/Source/ControlRigEditor/Private/Editor/SRigSpacePickerWidget.cpp
@@ -67,40 +67,39 @@ void SRigSpacePickerWidget::Construct(const FArguments& InArgs)
 			}
 		}
 	}
-
-	// Ok/Cancel buttons
-	ListBox->AddSlot()
-	.AutoHeight()
-	.HAlign(HAlign_Right)
-	.VAlign(VAlign_Bottom)
-	.Padding(4, 12, 0, 0)
-	[
-		SNew(SUniformGridPanel)
-		.SlotPadding(FEditorStyle::GetMargin("StandardDialog.SlotPadding"))
-		.MinDesiredSlotWidth(FEditorStyle::GetFloat("StandardDialog.MinDesiredSlotWidth"))
-		.MinDesiredSlotHeight(FEditorStyle::GetFloat("StandardDialog.MinDesiredSlotHeight"))
-		+SUniformGridPanel::Slot(0,0)
-		[
-			SNew(SButton)
-			.HAlign(HAlign_Center)
-			.ContentPadding( FEditorStyle::GetMargin("StandardDialog.ContentPadding") )
-			.OnClicked(this, &SRigSpacePickerWidget::CancelClicked)
-			.Text(LOCTEXT("RigSpacePickerWidgetCancel", "Cancel"))
-		]
-	];
 }
 
 SRigSpacePickerWidget::~SRigSpacePickerWidget()
 {
 }
 
+class SRigSpacePickerWindow : public SWindow
+{
+public:
+
+	virtual bool OnIsActiveChanged(const FWindowActivateEvent& ActivateEvent) override
+	{
+		if(ActivateEvent.GetActivationType() == FWindowActivateEvent::EA_Deactivate)
+		{
+			DeactivatedDelegate.ExecuteIfBound();
+		}
+		return SWindow::OnIsActiveChanged(ActivateEvent);
+	}
+
+	FSimpleDelegate& OnDeactivated() { return DeactivatedDelegate; }
+
+private:
+
+	FSimpleDelegate DeactivatedDelegate;
+}; 
+
 SRigSpacePickerWidget::FResult SRigSpacePickerWidget::InvokeDialog()
 {
 	PickedKey = FRigElementKey();
 	
 	const FVector2D CursorPos = FSlateApplication::Get().GetCursorPos();
 
-	TSharedRef<SWindow> Window = SNew(SWindow)
+	TSharedRef<SRigSpacePickerWindow> Window = SNew(SRigSpacePickerWindow)
 	.Title( LOCTEXT("SRigSpacePickerWidgetPickSpace", "Pick a new space") )
 	.CreateTitleBar(false)
 	.Type(EWindowType::Menu)
@@ -113,6 +112,7 @@ SRigSpacePickerWidget::FResult SRigSpacePickerWidget::InvokeDialog()
 	];
 	
 	Window->SetWidgetToFocusOnActivate(AsShared());
+	Window->OnDeactivated().BindSP(this, &SRigSpacePickerWidget::CloseDialog);
 	
 	PickerWindow = Window;
 
@@ -125,30 +125,19 @@ SRigSpacePickerWidget::FResult SRigSpacePickerWidget::InvokeDialog()
 	return Result;
 }
 
-void SRigSpacePickerWidget::CloseDialog(bool bWasPicked)
+void SRigSpacePickerWidget::CloseDialog()
 {
-	if(!bWasPicked)
-	{
-		PickedKey = FRigElementKey();
-	}
-
 	if ( PickerWindow.IsValid() )
 	{
 		PickerWindow.Pin()->RequestDestroyWindow();
 	}
 }
 
-FReply SRigSpacePickerWidget::CancelClicked()
-{
-	CloseDialog();
-	return FReply::Handled();
-}
-
 FReply SRigSpacePickerWidget::OnKeyDown( const FGeometry& MyGeometry, const FKeyEvent& InKeyEvent )
 {
 	if (InKeyEvent.GetKey() == EKeys::Escape)
 	{
-		CancelClicked();
+		PickedKey = FRigElementKey();
 		return FReply::Handled();
 	}
 	return SCompoundWidget::OnKeyDown(MyGeometry, InKeyEvent);
@@ -211,7 +200,7 @@ FReply SRigSpacePickerWidget::HandleWorldSpaceClicked()
 FReply SRigSpacePickerWidget::HandleElementSpaceClicked(FRigElementKey InKey)
 {
 	PickedKey = InKey;
-	CloseDialog(true);
+	CloseDialog();
 	return FReply::Handled();
 }
 
diff --git a/Engine/Plugins/Animation/ControlRig/Source/ControlRigEditor/Private/Editor/SRigSpacePickerWidget.h b/Engine/Plugins/Animation/ControlRig/Source/ControlRigEditor/Private/Editor/SRigSpacePickerWidget.h
--- a/Engine/Plugins/Animation/ControlRig/Source/ControlRigEditor/Private/Editor/SRigSpacePickerWidget.h
+++ b/Engine/Plugins/Animation/ControlRig/Source/ControlRigEditor/Private/Editor/SRigSpacePickerWidget.h
@@ -39,8 +39,7 @@ public:
 	~SRigSpacePickerWidget();
 
 	FResult InvokeDialog();
-	void CloseDialog(bool bWasPicked=false);
-	FReply CancelClicked();
+	void CloseDialog();
 	virtual FReply OnKeyDown( const FGeometry& MyGeometry, const FKeyEvent& InKeyEvent ) override;
 
 	virtual bool SupportsKeyboardFocus() const override
