Control Rig: Fix serialization of min and max values for controls

diff --git a/Engine/Plugins/Animation/ControlRig/Source/ControlRig/Private/Rigs/RigHierarchyDefines.cpp b/Engine/Plugins/Animation/ControlRig/Source/ControlRig/Private/Rigs/RigHierarchyDefines.cpp
--- a/Engine/Plugins/Animation/ControlRig/Source/ControlRig/Private/Rigs/RigHierarchyDefines.cpp
+++ b/Engine/Plugins/Animation/ControlRig/Source/ControlRig/Private/Rigs/RigHierarchyDefines.cpp
@@ -309,3 +309,41 @@ FVector FRigMirrorSettings::MirrorVector(const FVector& InVector) const
 	Axis.SetComponentForAxis(MirrorAxis, 1.f);
 	return InVector.MirrorByVector(Axis);
 }
+
+FArchive& operator<<(FArchive& Ar, FRigControlValue& Value)
+{
+	Ar <<  Value.FloatStorage.Float00;
+	Ar <<  Value.FloatStorage.Float01;
+	Ar <<  Value.FloatStorage.Float02;
+	Ar <<  Value.FloatStorage.Float03;
+	Ar <<  Value.FloatStorage.Float10;
+	Ar <<  Value.FloatStorage.Float11;
+	Ar <<  Value.FloatStorage.Float12;
+	Ar <<  Value.FloatStorage.Float13;
+	Ar <<  Value.FloatStorage.Float20;
+	Ar <<  Value.FloatStorage.Float21;
+	Ar <<  Value.FloatStorage.Float22;
+	Ar <<  Value.FloatStorage.Float23;
+	Ar <<  Value.FloatStorage.Float30;
+	Ar <<  Value.FloatStorage.Float31;
+	Ar <<  Value.FloatStorage.Float32;
+	Ar <<  Value.FloatStorage.Float33;
+	Ar <<  Value.FloatStorage.Float00_2;
+	Ar <<  Value.FloatStorage.Float01_2;
+	Ar <<  Value.FloatStorage.Float02_2;
+	Ar <<  Value.FloatStorage.Float03_2;
+	Ar <<  Value.FloatStorage.Float10_2;
+	Ar <<  Value.FloatStorage.Float11_2;
+	Ar <<  Value.FloatStorage.Float12_2;
+	Ar <<  Value.FloatStorage.Float13_2;
+	Ar <<  Value.FloatStorage.Float20_2;
+	Ar <<  Value.FloatStorage.Float21_2;
+	Ar <<  Value.FloatStorage.Float22_2;
+	Ar <<  Value.FloatStorage.Float23_2;
+	Ar <<  Value.FloatStorage.Float30_2;
+	Ar <<  Value.FloatStorage.Float31_2;
+	Ar <<  Value.FloatStorage.Float32_2;
+	Ar <<  Value.FloatStorage.Float33_2;
+
+	return Ar;
+}
diff --git a/Engine/Plugins/Animation/ControlRig/Source/ControlRig/Private/Rigs/RigHierarchyElements.cpp b/Engine/Plugins/Animation/ControlRig/Source/ControlRig/Private/Rigs/RigHierarchyElements.cpp
--- a/Engine/Plugins/Animation/ControlRig/Source/ControlRig/Private/Rigs/RigHierarchyElements.cpp
+++ b/Engine/Plugins/Animation/ControlRig/Source/ControlRig/Private/Rigs/RigHierarchyElements.cpp
@@ -435,8 +435,7 @@ void FRigControlSettings::Save(FArchive& Ar)
 
 	FName ControlTypeName = ControlTypeEnum->GetNameByValue((int64)ControlType);
 	FName PrimaryAxisName = ControlAxisEnum->GetNameByValue((int64)PrimaryAxis);
-	FTransform MinimumTransform = MinimumValue.GetAsTransform(ControlType, PrimaryAxis);
-	FTransform MaximumTransform = MaximumValue.GetAsTransform(ControlType, PrimaryAxis);
+
 	FString ControlEnumPathName;
 	if(ControlEnum)
 	{
@@ -452,8 +451,8 @@ void FRigControlSettings::Save(FArchive& Ar)
 	Ar << bLimitRotation;
 	Ar << bLimitScale;
 	Ar << bDrawLimits;
-	Ar << MinimumTransform;
-	Ar << MaximumTransform;
+	Ar << MinimumValue;
+	Ar << MaximumValue;
 	Ar << bGizmoEnabled;
 	Ar << bGizmoVisible;
 	Ar << GizmoName;
@@ -471,7 +470,6 @@ void FRigControlSettings::Load(FArchive& Ar)
 	static const UEnum* ControlAxisEnum = StaticEnum<ERigControlAxis>();
 
 	FName ControlTypeName, PrimaryAxisName;
-	FTransform MinimumTransform, MaximumTransform;
 	FString ControlEnumPathName;
 
 	Ar << ControlTypeName;
@@ -483,8 +481,22 @@ void FRigControlSettings::Load(FArchive& Ar)
 	Ar << bLimitRotation;
 	Ar << bLimitScale;
 	Ar << bDrawLimits;
-	Ar << MinimumTransform;
-	Ar << MaximumTransform;
+
+	if (Ar.CustomVer(FControlRigObjectVersion::GUID) >= FControlRigObjectVersion::StorageMinMaxValuesAsFloatStorage)
+	{
+		Ar << MinimumValue;
+		Ar << MaximumValue;
+	}
+	else
+	{
+		FTransform MinimumTransform, MaximumTransform;
+		Ar << MinimumTransform;
+		Ar << MaximumTransform;
+		MinimumValue.SetFromTransform(MinimumTransform, ControlType, PrimaryAxis);
+		MaximumValue.SetFromTransform(MaximumTransform, ControlType, PrimaryAxis);
+	}
+	
+
 	Ar << bGizmoEnabled;
 	Ar << bGizmoVisible;
 	Ar << GizmoName;
@@ -494,8 +506,6 @@ void FRigControlSettings::Load(FArchive& Ar)
 
 	ControlType = (ERigControlType)ControlTypeEnum->GetValueByName(ControlTypeName);
 	PrimaryAxis = (ERigControlAxis)ControlAxisEnum->GetValueByName(PrimaryAxisName);
-	MinimumValue.SetFromTransform(MinimumTransform, ControlType, PrimaryAxis);
-	MaximumValue.SetFromTransform(MaximumTransform, ControlType, PrimaryAxis);
 
 	ControlEnum = nullptr;
 	if(!ControlEnumPathName.IsEmpty())
diff --git a/Engine/Plugins/Animation/ControlRig/Source/ControlRig/Public/ControlRigObjectVersion.h b/Engine/Plugins/Animation/ControlRig/Source/ControlRig/Public/ControlRigObjectVersion.h
--- a/Engine/Plugins/Animation/ControlRig/Source/ControlRig/Public/ControlRigObjectVersion.h
+++ b/Engine/Plugins/Animation/ControlRig/Source/ControlRig/Public/ControlRigObjectVersion.h
@@ -39,6 +39,9 @@ struct CONTROLRIG_API FControlRigObjectVersion
 		// RigHierarchy now supports space favorites per control
 		RigHierarchyControlSpaceFavorites,
 
+		// RigHierarchy now stores min and max values as float storages 
+		StorageMinMaxValuesAsFloatStorage,
+
 		// -----<new versions can be added above this line>-------------------------------------------------
 		VersionPlusOne,
 		LatestVersion = VersionPlusOne - 1
diff --git a/Engine/Plugins/Animation/ControlRig/Source/ControlRig/Public/Rigs/RigHierarchyDefines.h b/Engine/Plugins/Animation/ControlRig/Source/ControlRig/Public/Rigs/RigHierarchyDefines.h
--- a/Engine/Plugins/Animation/ControlRig/Source/ControlRig/Public/Rigs/RigHierarchyDefines.h
+++ b/Engine/Plugins/Animation/ControlRig/Source/ControlRig/Public/Rigs/RigHierarchyDefines.h
@@ -280,6 +280,8 @@ public:
 		return Value;
 	}
 
+	friend FArchive& operator <<(FArchive& Ar, FRigControlValue& Value);
+
 	FORCEINLINE_DEBUGGABLE FString ToPythonString(ERigControlType InControlType) const
 	{
 		FString ValueStr;
@@ -951,6 +953,7 @@ private:
 	friend class UControlRigBlueprint;
 	friend class URigHierarchyController;
 	friend class UControlRig;
+	friend struct FRigControlSettings;
 };
 
 template<>
