Control Rig: Avoid ensure during import

diff --git a/Engine/Plugins/Animation/ControlRig/Source/ControlRig/Private/Rigs/RigHierarchyController.cpp b/Engine/Plugins/Animation/ControlRig/Source/ControlRig/Private/Rigs/RigHierarchyController.cpp
--- a/Engine/Plugins/Animation/ControlRig/Source/ControlRig/Private/Rigs/RigHierarchyController.cpp
+++ b/Engine/Plugins/Animation/ControlRig/Source/ControlRig/Private/Rigs/RigHierarchyController.cpp
@@ -1003,7 +1003,7 @@ TArray<FRigElementKey> URigHierarchyController::ImportFromText(FString InContent
 		}
 
 		check(NewElement);
-		ensure(NewElement->GetKey() == PerElementData.Key);
+		NewElement->Key = PerElementData.Key;
 
 		if(bReplaceExistingElements)
 		{
