Control Rig: Renamed space picker local space to parent space

diff --git a/Engine/Plugins/Animation/ControlRig/Source/ControlRigEditor/Private/Editor/SRigSpacePickerWidget.cpp b/Engine/Plugins/Animation/ControlRig/Source/ControlRigEditor/Private/Editor/SRigSpacePickerWidget.cpp
--- a/Engine/Plugins/Animation/ControlRig/Source/ControlRigEditor/Private/Editor/SRigSpacePickerWidget.cpp
+++ b/Engine/Plugins/Animation/ControlRig/Source/ControlRigEditor/Private/Editor/SRigSpacePickerWidget.cpp
@@ -41,7 +41,7 @@ void SRigSpacePickerWidget::Construct(const FArguments& InArgs)
 
 	AddSpacePickerButton(
 		FEditorStyle::GetBrush("EditorViewport.RelativeCoordinateSystem_Local"),
-		LOCTEXT("Local", "Local"),
+		LOCTEXT("Parent", "Parent"),
 		FOnClicked::CreateSP(this, &SRigSpacePickerWidget::HandleLocalSpaceClicked)
 	);
 	
