Control Rig: Shorten overly long node subtitles

diff --git a/Engine/Plugins/Animation/ControlRig/Source/ControlRigDeveloper/Private/Graph/ControlRigGraphNode.cpp b/Engine/Plugins/Animation/ControlRig/Source/ControlRigDeveloper/Private/Graph/ControlRigGraphNode.cpp
--- a/Engine/Plugins/Animation/ControlRig/Source/ControlRigDeveloper/Private/Graph/ControlRigGraphNode.cpp
+++ b/Engine/Plugins/Animation/ControlRig/Source/ControlRigDeveloper/Private/Graph/ControlRigGraphNode.cpp
@@ -130,6 +130,11 @@ FText UControlRigGraphNode::GetNodeTitle(ENodeTitleType::Type TitleType) const
 						}
 					}
 				}
+
+				if(SubTitle.Len() > 40)
+				{
+					SubTitle = SubTitle.Left(36) + TEXT(" ...");
+				}
 			}
 
 			if (NodeTitle.IsEmpty())
