Control Rig: Crash when browsing to a bool control from a node

diff --git a/Engine/Plugins/Animation/ControlRig/Source/ControlRigEditor/Private/Graph/SControlRigGraphPinNameList.cpp b/Engine/Plugins/Animation/ControlRig/Source/ControlRigEditor/Private/Graph/SControlRigGraphPinNameList.cpp
--- a/Engine/Plugins/Animation/ControlRig/Source/ControlRigEditor/Private/Graph/SControlRigGraphPinNameList.cpp
+++ b/Engine/Plugins/Animation/ControlRig/Source/ControlRigEditor/Private/Graph/SControlRigGraphPinNameList.cpp
@@ -213,7 +213,7 @@ FSlateColor SControlRigGraphPinNameList::OnGetWidgetBackground() const
 
 FReply SControlRigGraphPinNameList::OnGetSelectedClicked()
 {
-	if (ModelPin->GetCustomWidgetName() == TEXT("ElementName"))
+	if (UControlRigGraph* Graph = Cast<UControlRigGraph>(GetPinObj()->GetOwningNode()->GetGraph()))
 	{
 		if (OnGetNameFromSelection.IsBound())
 		{
@@ -222,19 +222,49 @@ FReply SControlRigGraphPinNameList::OnGetSelectedClicked()
 			{
 				if (Result[0].IsValid() && Result[0] != nullptr)
 				{
-					if (URigVMPin* ParentPin = ModelPin->GetParentPin())
+					const FString DefaultValue = *Result[0].Get();
+					
+					if(ModelPin->GetCustomWidgetName() == TEXT("ElementName"))
 					{
-						if (UControlRigGraph* Graph = Cast<UControlRigGraph>(GetPinObj()->GetOwningNode()->GetGraph()))
+						if (URigVMPin* ParentPin = ModelPin->GetParentPin())
 						{
-							Graph->GetController()->SetPinDefaultValue(ParentPin->GetPinPath(), *Result[0].Get(), true, true, false, true);
+							Graph->GetController()->SetPinDefaultValue(ParentPin->GetPinPath(), DefaultValue, true, true, false, true);
 							CurrentList = GetNameList();
 						}
 					}
+					
+					// if we don't have a key pin - this is just a plain name.
+					// let's derive the type of element this node deals with from its name.
+					// there's nothing better in place for now.
+					else if(const URigVMUnitNode* UnitNode = Cast<URigVMUnitNode>(ModelPin->GetNode()))
+					{
+						const int32 LastIndex = StaticEnum<ERigElementType>()->GetIndexByName(TEXT("Last")); 
+						const FString UnitName = UnitNode->GetScriptStruct()->GetStructCPPName();
+						for(int32 EnumIndex = 0; EnumIndex < LastIndex; EnumIndex++)
+						{
+							const FString EnumDisplayName = StaticEnum<ERigElementType>()->GetDisplayNameTextByIndex(EnumIndex).ToString();
+							if(UnitName.Contains(EnumDisplayName))
+							{
+								const ERigElementType ElementType = (ERigElementType)StaticEnum<ERigElementType>()->GetValueByIndex(EnumIndex);
+
+								FRigElementKey Key;
+								FRigElementKey::StaticStruct()->ImportText(*DefaultValue, &Key, nullptr, EPropertyPortFlags::PPF_None, nullptr, FRigElementKey::StaticStruct()->GetName(), true);
+								if (Key.IsValid())
+								{
+									if(Key.Type == ElementType)
+									{
+										Graph->GetController()->SetPinDefaultValue(ModelPin->GetPinPath(), Key.Name.ToString(), true, true, false, true);
+										CurrentList = GetNameList();
+									}
+								}
+								break;
+							}
+						}
+					}
 				}
 			}
 		}
 	}
-
 	return FReply::Handled();
 }
 
@@ -243,16 +273,45 @@ FReply SControlRigGraphPinNameList::OnBrowseClicked()
 	if (UControlRigGraph* Graph = Cast<UControlRigGraph>(GetPinObj()->GetOwningNode()->GetGraph()))
 	{
 		TSharedPtr<FString> Selected = NameListComboBox->GetSelectedItem();
-		if (Selected.IsValid())
+		if (Selected.IsValid() && ModelPin)
 		{
-			FString DefaultValue = ModelPin->GetParentPin()->GetDefaultValue();
-			if (!DefaultValue.IsEmpty())
+			if(URigVMPin* KeyPin = ModelPin->GetParentPin())
 			{
-				FRigElementKey Key;
-				FRigElementKey::StaticStruct()->ImportText(*DefaultValue, &Key, nullptr, EPropertyPortFlags::PPF_None, nullptr, FRigElementKey::StaticStruct()->GetName(), true);
-				if (Key.IsValid())
+				if(KeyPin->GetCPPTypeObject() == FRigElementKey::StaticStruct())
 				{
-					Graph->GetBlueprint()->GetHierarchyController()->SetSelection({Key});
+					FString DefaultValue = ModelPin->GetParentPin()->GetDefaultValue();
+					if (!DefaultValue.IsEmpty())
+					{
+						FRigElementKey Key;
+						FRigElementKey::StaticStruct()->ImportText(*DefaultValue, &Key, nullptr, EPropertyPortFlags::PPF_None, nullptr, FRigElementKey::StaticStruct()->GetName(), true);
+						if (Key.IsValid())
+						{
+							Graph->GetBlueprint()->GetHierarchyController()->SetSelection({Key});
+						}
+					}
+				}
+			}
+			else
+			{
+				// if we don't have a key pin - this is just a plain name.
+				// let's derive the type of element this node deals with from its name.
+				// there's nothing better in place for now.
+				if(const URigVMUnitNode* UnitNode = Cast<URigVMUnitNode>(ModelPin->GetNode()))
+				{
+					const int32 LastIndex = StaticEnum<ERigElementType>()->GetIndexByName(TEXT("Last")); 
+					const FString UnitName = UnitNode->GetScriptStruct()->GetStructCPPName();
+					for(int32 EnumIndex = 0; EnumIndex < LastIndex; EnumIndex++)
+					{
+						const FString EnumDisplayName = StaticEnum<ERigElementType>()->GetDisplayNameTextByIndex(EnumIndex).ToString();
+						if(UnitName.Contains(EnumDisplayName))
+						{
+							const FString DefaultValue = ModelPin->GetDefaultValue();
+							const ERigElementType ElementType = (ERigElementType)StaticEnum<ERigElementType>()->GetValueByIndex(EnumIndex);
+							FRigElementKey Key(*DefaultValue, ElementType);
+							Graph->GetBlueprint()->GetHierarchyController()->SetSelection({Key});
+							break;
+						}
+					}
 				}
 			}
 		}
