Control Rig; Check transform cache integrity on load

diff --git a/Engine/Plugins/Animation/ControlRig/Source/ControlRig/Private/ControlRig.cpp b/Engine/Plugins/Animation/ControlRig/Source/ControlRig/Private/ControlRig.cpp
--- a/Engine/Plugins/Animation/ControlRig/Source/ControlRig/Private/ControlRig.cpp
+++ b/Engine/Plugins/Animation/ControlRig/Source/ControlRig/Private/ControlRig.cpp
@@ -723,7 +723,7 @@ void UControlRig::Execute(const EControlRigState InState, const FName& InEventNa
 				{
 					// only do this in non-setup mode because 
 					// when setup mode is enabled, the control values are cleared before reaching here (too late to save them)
-					PoseScope = MakeUnique<UControlRig::FPoseScope>(this);
+					PoseScope = MakeUnique<UControlRig::FPoseScope>(this, ERigElementType::ToResetAfterSetupEvent);
 				}
 
 				// reset the pose to initial such that setup event can run from a deterministic initial state
@@ -2771,13 +2771,12 @@ void UControlRig::OnHierarchyTransformUndoRedo(URigHierarchy* InHierarchy, const
 	}
 }
 
-UControlRig::FPoseScope::FPoseScope(UControlRig* InControlRig)
+UControlRig::FPoseScope::FPoseScope(UControlRig* InControlRig, ERigElementType InFilter)
+: ControlRig(InControlRig)
+, Filter(InFilter)
 {
 	check(InControlRig);
-
-	CachedPose = InControlRig->GetHierarchy()->GetPose();
-
-	ControlRig = InControlRig;
+	CachedPose = InControlRig->GetHierarchy()->GetPose(false, InFilter, FRigElementKeyCollection());
 }
 
 UControlRig::FPoseScope::~FPoseScope()
diff --git a/Engine/Plugins/Animation/ControlRig/Source/ControlRig/Private/Rigs/RigHierarchy.cpp b/Engine/Plugins/Animation/ControlRig/Source/ControlRig/Private/Rigs/RigHierarchy.cpp
--- a/Engine/Plugins/Animation/ControlRig/Source/ControlRig/Private/Rigs/RigHierarchy.cpp
+++ b/Engine/Plugins/Animation/ControlRig/Source/ControlRig/Private/Rigs/RigHierarchy.cpp
@@ -235,6 +235,95 @@ void URigHierarchy::Load(FArchive& Ar)
 	Notify(ERigHierarchyNotification::HierarchyReset, nullptr);
 }
 
+void URigHierarchy::PostLoad()
+{
+	UObject::PostLoad();
+
+	struct Local
+	{
+		static bool NeedsCheck(const FRigLocalAndGlobalTransform& InTransform)
+		{
+			return !InTransform.Local.bDirty && !InTransform.Global.bDirty;
+		}
+	};
+
+	// we need to check the elements for integrity (global vs local) to be correct.
+	for(int32 ElementIndex = 0; ElementIndex < Elements.Num(); ElementIndex++)
+	{
+		FRigBaseElement* BaseElement = Elements[ElementIndex];
+
+		if(FRigControlElement* ControlElement = Cast<FRigControlElement>(BaseElement))
+		{
+			if(Local::NeedsCheck(ControlElement->Offset.Initial))
+			{
+				const FTransform ComputedGlobalTransform = SolveParentConstraints(
+					ControlElement->ParentConstraints, ERigTransformType::InitialGlobal,
+					ControlElement->Offset.Get(ERigTransformType::InitialLocal), true,
+					FTransform::Identity, false);
+
+				const FTransform CachedGlobalTransform = ControlElement->Offset.Get(ERigTransformType::InitialGlobal);
+				if(!FRigComputedTransform::Equals(ComputedGlobalTransform, CachedGlobalTransform, 0.01))
+				{
+					ControlElement->Offset.MarkDirty(ERigTransformType::InitialGlobal);
+				}
+			}
+
+			if(Local::NeedsCheck(ControlElement->Pose.Initial))
+			{
+				const FTransform ComputedGlobalTransform = SolveParentConstraints(
+					ControlElement->ParentConstraints, ERigTransformType::InitialGlobal,
+					GetControlOffsetTransform(ControlElement, ERigTransformType::InitialGlobal), true,
+					ControlElement->Pose.Get(ERigTransformType::InitialLocal), true);
+				
+				const FTransform CachedGlobalTransform = ControlElement->Pose.Get(ERigTransformType::InitialGlobal);
+				
+				if(!FRigComputedTransform::Equals(ComputedGlobalTransform, CachedGlobalTransform, 0.01))
+				{
+					// for nulls we perceive the local transform as less relevant
+					ControlElement->Pose.MarkDirty(ERigTransformType::InitialLocal);
+				}
+			}
+
+			// we also need to check the pose here - for controls it is a bit different than for other
+			// types.
+			continue;
+		}
+
+		if(FRigMultiParentElement* MultiParentElement = Cast<FRigMultiParentElement>(BaseElement))
+		{
+			if(Local::NeedsCheck(MultiParentElement->Pose.Initial))
+			{
+				const FTransform ComputedGlobalTransform = SolveParentConstraints(
+					MultiParentElement->ParentConstraints, ERigTransformType::InitialGlobal,
+					FTransform::Identity, false,
+					MultiParentElement->Pose.Get(ERigTransformType::InitialLocal), true);
+				
+				const FTransform CachedGlobalTransform = MultiParentElement->Pose.Get(ERigTransformType::InitialGlobal);
+				
+				if(!FRigComputedTransform::Equals(ComputedGlobalTransform, CachedGlobalTransform, 0.01))
+				{
+					// for nulls we perceive the local transform as less relevant
+					MultiParentElement->Pose.MarkDirty(ERigTransformType::InitialLocal);
+				}
+			}
+		}
+
+		if(FRigTransformElement* TransformElement = Cast<FRigTransformElement>(BaseElement))
+		{
+			if(Local::NeedsCheck(TransformElement->Pose.Initial))
+			{
+				const FTransform ParentTransform = GetParentTransform(TransformElement, ERigTransformType::InitialGlobal);
+				const FTransform ComputedGlobalTransform = TransformElement->Pose.Get(ERigTransformType::InitialLocal) * ParentTransform;
+				const FTransform CachedGlobalTransform = TransformElement->Pose.Get(ERigTransformType::InitialGlobal);
+				if(!FRigComputedTransform::Equals(ComputedGlobalTransform, CachedGlobalTransform, 0.01))
+				{
+					TransformElement->Pose.MarkDirty(ERigTransformType::InitialGlobal);
+				}
+			}
+		}
+	}
+}
+
 void URigHierarchy::Reset()
 {
 	TopologyVersion = 0;
@@ -1669,7 +1758,7 @@ FTransform URigHierarchy::GetTransform(FRigTransformElement* InTransformElement,
 		{
 			if(FRigControlElement* ControlElement = Cast<FRigControlElement>(InTransformElement))
 			{
-				const FTransform NewTransform = ComputeLocalControlValue(ControlElement, ControlElement->Pose.Get(OpposedType), InTransformType);
+				const FTransform NewTransform = ComputeLocalControlValue(ControlElement, ControlElement->Pose.Get(OpposedType), GlobalType);
 				InTransformElement->Pose.Set(InTransformType, NewTransform);
 			}
 			else
@@ -1825,8 +1914,10 @@ FTransform URigHierarchy::GetControlOffsetTransform(FRigControlElement* InContro
 
 		if(IsLocal(InTransformType))
 		{
-			const FTransform ParentTransform = GetParentTransform(InControlElement, GlobalType);
-			InControlElement->Offset.Set(InTransformType, InControlElement->Offset.Get(OpposedType).GetRelativeTransform(ParentTransform));
+			const FTransform LocalTransform = InverseSolveParentConstraints(
+				InControlElement->Offset.Get(OpposedType), 
+				InControlElement->ParentConstraints, InTransformType, FTransform::Identity);
+			InControlElement->Offset.Set(InTransformType, LocalTransform);
 		}
 		else
 		{
@@ -2075,8 +2166,7 @@ FTransform URigHierarchy::GetParentTransform(FRigBaseElement* InElement, const E
 				FTransform::Identity,
 				false
 			);
-			Output.Set(OutputTransform);
-			Output.bDirty = false;
+			MultiParentElement->Parent.Set(InTransformType, OutputTransform);
 		}
 		return Output.Transform;
 	}
diff --git a/Engine/Plugins/Animation/ControlRig/Source/ControlRig/Public/ControlRig.h b/Engine/Plugins/Animation/ControlRig/Source/ControlRig/Public/ControlRig.h
--- a/Engine/Plugins/Animation/ControlRig/Source/ControlRig/Public/ControlRig.h
+++ b/Engine/Plugins/Animation/ControlRig/Source/ControlRig/Public/ControlRig.h
@@ -726,12 +726,13 @@ private:
 	class FPoseScope
 	{
 	public:
-		FPoseScope(UControlRig* InControlRig);
+		FPoseScope(UControlRig* InControlRig, ERigElementType InFilter = ERigElementType::All);
 		~FPoseScope();
 
 	private:
 
 		UControlRig* ControlRig;
+		ERigElementType Filter;
 		FRigPose CachedPose;
 	};
 
diff --git a/Engine/Plugins/Animation/ControlRig/Source/ControlRig/Public/Rigs/RigHierarchy.h b/Engine/Plugins/Animation/ControlRig/Source/ControlRig/Public/Rigs/RigHierarchy.h
--- a/Engine/Plugins/Animation/ControlRig/Source/ControlRig/Public/Rigs/RigHierarchy.h
+++ b/Engine/Plugins/Animation/ControlRig/Source/ControlRig/Public/Rigs/RigHierarchy.h
@@ -93,6 +93,7 @@ public:
 	virtual void Serialize(FArchive& Ar) override;
 	void Save(FArchive& Ar);
 	void Load(FArchive& Ar);
+	virtual void PostLoad() override;
 #if WITH_EDITOR
 	virtual void PostTransacted(const FTransactionObjectEvent& TransactionEvent) override;
 #endif
diff --git a/Engine/Plugins/Animation/ControlRig/Source/ControlRig/Public/Rigs/RigHierarchyDefines.h b/Engine/Plugins/Animation/ControlRig/Source/ControlRig/Public/Rigs/RigHierarchyDefines.h
--- a/Engine/Plugins/Animation/ControlRig/Source/ControlRig/Public/Rigs/RigHierarchyDefines.h
+++ b/Engine/Plugins/Animation/ControlRig/Source/ControlRig/Public/Rigs/RigHierarchyDefines.h
@@ -26,6 +26,7 @@ enum class ERigElementType : uint8
 	Socket = 0x020,
 	Last = 0x040 UMETA(Hidden),
 	All = Bone | Null | Control | Curve | RigidBody | Socket,
+	ToResetAfterSetupEvent = Bone | Control | Curve UMETA(Hidden),
 };
 
 UENUM(BlueprintType)
