Control Rig: Fix Control Pin Value for Modify Transform Node

diff --git a/Engine/Plugins/Animation/ControlRig/Source/ControlRig/Private/ControlRig.cpp b/Engine/Plugins/Animation/ControlRig/Source/ControlRig/Private/ControlRig.cpp
--- a/Engine/Plugins/Animation/ControlRig/Source/ControlRig/Private/ControlRig.cpp
+++ b/Engine/Plugins/Animation/ControlRig/Source/ControlRig/Private/ControlRig.cpp
@@ -1616,7 +1616,7 @@ void UControlRig::HandleHierarchyModified(ERigHierarchyNotification InNotificati
 
 #if WITH_EDITOR
 
-FName UControlRig::AddTransientControl(URigVMPin* InPin, FRigElementKey SpaceKey)
+FName UControlRig::AddTransientControl(URigVMPin* InPin, FRigElementKey SpaceKey, FTransform OffsetTransform)
 {
 	if ((InPin == nullptr) || (DynamicHierarchy == nullptr))
 	{
@@ -1669,7 +1669,7 @@ FName UControlRig::AddTransientControl(URigVMPin* InPin, FRigElementKey SpaceKey
     	Parent,
     	Settings,
     	FRigControlValue::Make(FTransform::Identity),
-    	FTransform::Identity,
+    	OffsetTransform,
     	GizmoTransform, false);
 
 	SetTransientControlValue(InPin);
diff --git a/Engine/Plugins/Animation/ControlRig/Source/ControlRig/Private/Units/Highlevel/Hierarchy/RigUnit_ModifyTransforms.h b/Engine/Plugins/Animation/ControlRig/Source/ControlRig/Private/Units/Highlevel/Hierarchy/RigUnit_ModifyTransforms.h
--- a/Engine/Plugins/Animation/ControlRig/Source/ControlRig/Private/Units/Highlevel/Hierarchy/RigUnit_ModifyTransforms.h
+++ b/Engine/Plugins/Animation/ControlRig/Source/ControlRig/Private/Units/Highlevel/Hierarchy/RigUnit_ModifyTransforms.h
@@ -97,21 +97,57 @@ struct CONTROLRIG_API FRigUnit_ModifyTransforms : public FRigUnit_HighlevelBaseM
 
 			if (ItemToModify.IsValidIndex(Index))
 			{
-				if (Mode == EControlRigModifyBoneMode::AdditiveLocal)
+				if (Mode == EControlRigModifyBoneMode::AdditiveLocal || Mode == EControlRigModifyBoneMode::OverrideLocal)
+				{
+					if (const URigHierarchy* Hierarchy = (const URigHierarchy*)InUserContext)
+					{
+						// in UControlRig::AddTransientControl, the transient control will be added to the parent of this item
+						return ItemToModify[Index].Item;
+					}
+				}
+			}
+		}
+		return FRigElementKey();
+	}
+
+	virtual FTransform DetermineOffsetTransformForPin(const FString& InPinPath, void* InUserContext) const override
+	{
+		if (InPinPath.StartsWith(TEXT("ItemToModify")))
+		{
+			int32 Index = INDEX_NONE;
+			FString Left, Middle, Right;
+			if (InPinPath.Replace(TEXT("["), TEXT(".")).Split(TEXT("."), &Left, &Middle))
+			{
+				if (Middle.Replace(TEXT("]"), TEXT(".")).Split(TEXT("."), &Left, &Right))
 				{
-					return ItemToModify[Index].Item;
+					Index = FCString::Atoi(*Left);
 				}
+			}
 
-				if (Mode == EControlRigModifyBoneMode::OverrideLocal)
+			if (ItemToModify.IsValidIndex(Index))
+			{
+				if (Mode == EControlRigModifyBoneMode::AdditiveLocal)
 				{
 					if (const URigHierarchy* Hierarchy = (const URigHierarchy*)InUserContext)
 					{
-						return Hierarchy->GetFirstParent(ItemToModify[Index].Item);
+						if (ItemToModify[Index].Item.IsValid())
+						{
+							// in the case of AdditiveLocal, we want the transient control's global transform to match the item's global transform
+							// we know: transient control's global = transform value * offset transfom * parent transform
+							// we know: item's global transform = item's local transform from rig graph evaluation * its parent transform
+							// we ensure: the item and the transient control share the same parent
+							// so from our goal, transient control's global == item's global transform
+							// we get: transform value * offset transfom * parent transform == item's local transform from rig graph evaluation * its parent transform
+							// we get: transform value * offset transfom == item's local transform from rig graph evaluation
+							// we get: offset transfom == transform value.inverse * item's local transform from rig graph evaluation
+							return ItemToModify[Index].Transform.Inverse() * Hierarchy->GetLocalTransform(ItemToModify[Index].Item);
+						}
 					}
 				}
 			}
 		}
-		return FRigElementKey();
+
+		return FTransform::Identity;
 	}
 
 	RIGVM_METHOD()
diff --git a/Engine/Plugins/Animation/ControlRig/Source/ControlRig/Public/ControlRig.h b/Engine/Plugins/Animation/ControlRig/Source/ControlRig/Public/ControlRig.h
--- a/Engine/Plugins/Animation/ControlRig/Source/ControlRig/Public/ControlRig.h
+++ b/Engine/Plugins/Animation/ControlRig/Source/ControlRig/Public/ControlRig.h
@@ -507,7 +507,7 @@ private:
 
 #if WITH_EDITOR
 	/** Remove a transient / temporary control used to interact with a pin */
-	FName AddTransientControl(URigVMPin* InPin, FRigElementKey SpaceKey = FRigElementKey());
+	FName AddTransientControl(URigVMPin* InPin, FRigElementKey SpaceKey = FRigElementKey(), FTransform OffsetTransform = FTransform::Identity);
 
 	/** Sets the value of a transient control based on a pin */
 	bool SetTransientControlValue(URigVMPin* InPin);
diff --git a/Engine/Plugins/Animation/ControlRig/Source/ControlRig/Public/Units/RigUnit.h b/Engine/Plugins/Animation/ControlRig/Source/ControlRig/Public/Units/RigUnit.h
--- a/Engine/Plugins/Animation/ControlRig/Source/ControlRig/Public/Units/RigUnit.h
+++ b/Engine/Plugins/Animation/ControlRig/Source/ControlRig/Public/Units/RigUnit.h
@@ -28,7 +28,9 @@ struct CONTROLRIG_API FRigUnit : public FRigVMStruct
 	virtual void Execute(const FRigUnitContext& Context) {}
 
 	virtual FRigElementKey DetermineSpaceForPin(const FString& InPinPath, void* InUserContext) const { return FRigElementKey(); }
-
+	
+	virtual FTransform DetermineOffsetTransformForPin(const FString& InPinPath, void* InUserContext) const { return FTransform::Identity; }
+	
 	/** The name of the method used within each rig unit */
 	FORCEINLINE_DEBUGGABLE static FName GetMethodName()
 	{
diff --git a/Engine/Plugins/Animation/ControlRig/Source/ControlRigDeveloper/Private/ControlRigBlueprint.cpp b/Engine/Plugins/Animation/ControlRig/Source/ControlRigDeveloper/Private/ControlRigBlueprint.cpp
--- a/Engine/Plugins/Animation/ControlRig/Source/ControlRigDeveloper/Private/ControlRigBlueprint.cpp
+++ b/Engine/Plugins/Animation/ControlRig/Source/ControlRigDeveloper/Private/ControlRigBlueprint.cpp
@@ -2068,6 +2068,7 @@ FName UControlRigBlueprint::AddTransientControl(URigVMPin* InPin)
 	UControlRig* CDO = Cast<UControlRig>(RigClass->GetDefaultObject(true /* create if needed */));
 
 	FRigElementKey SpaceKey;
+	FTransform OffsetTransform = FTransform::Identity;
 	if (URigVMUnitNode* UnitNode = Cast<URigVMUnitNode>(InPin->GetPinForLink()->GetNode()))
 	{
 		if (TSharedPtr<FStructOnScope> DefaultStructScope = UnitNode->ConstructStructInstance())
@@ -2078,8 +2079,19 @@ FName UControlRigBlueprint::AddTransientControl(URigVMPin* InPin)
 			FString Left, Right;
 
 			if (URigVMPin::SplitPinPathAtStart(PinPath, Left, Right))
-		{
+			{
 				SpaceKey = DefaultStruct->DetermineSpaceForPin(Right, Hierarchy);
+				
+				URigHierarchy* RigHierarchy = Hierarchy;
+
+				// use the active rig instead of the CDO rig because we want to access the evaluation result of the rig graph
+				// to calculate the offset transform, for example take a look at RigUnit_ModifyTransform
+				if (UControlRig* RigBeingDebugged = Cast<UControlRig>(GetObjectBeingDebugged()))
+				{
+					RigHierarchy = RigBeingDebugged->GetHierarchy();
+				}
+				
+				OffsetTransform = DefaultStruct->DetermineOffsetTransformForPin(Right, RigHierarchy);
 			}
 		}
 	}
@@ -2092,7 +2104,7 @@ FName UControlRigBlueprint::AddTransientControl(URigVMPin* InPin)
 		UControlRig* InstancedControlRig = Cast<UControlRig>(ArchetypeInstance);
 		if (InstancedControlRig)
 		{
-			FName ControlName = InstancedControlRig->AddTransientControl(InPin, SpaceKey);
+			FName ControlName = InstancedControlRig->AddTransientControl(InPin, SpaceKey, OffsetTransform);
 			if (ReturnName == NAME_None)
 			{
 				ReturnName = ControlName;
