Control Rig: Mark FControlRigExecuteContext as blueprint type.

diff --git a/Engine/Plugins/Animation/ControlRig/Source/ControlRig/Public/ControlRigDefines.h b/Engine/Plugins/Animation/ControlRig/Source/ControlRig/Public/ControlRigDefines.h
--- a/Engine/Plugins/Animation/ControlRig/Source/ControlRig/Public/ControlRigDefines.h
+++ b/Engine/Plugins/Animation/ControlRig/Source/ControlRig/Public/ControlRigDefines.h
@@ -13,7 +13,7 @@
 //#define DECLARE_SCOPE_HIERARCHICAL_COUNTER_RIGUNIT() \
 //	DECLARE_SCOPE_HIERARCHICAL_COUNTER_FUNC()
 
-USTRUCT()
+USTRUCT(BlueprintType)
 struct FControlRigExecuteContext : public FRigVMExecuteContext
 {
 	GENERATED_BODY()
