Control Rig: Support controls for input on the AnimNode

diff --git a/Engine/Plugins/Animation/ControlRig/Source/ControlRig/Private/AnimNode_ControlRig.cpp b/Engine/Plugins/Animation/ControlRig/Source/ControlRig/Private/AnimNode_ControlRig.cpp
--- a/Engine/Plugins/Animation/ControlRig/Source/ControlRig/Private/AnimNode_ControlRig.cpp
+++ b/Engine/Plugins/Animation/ControlRig/Source/ControlRig/Private/AnimNode_ControlRig.cpp
@@ -127,6 +127,7 @@ void FAnimNode_ControlRig::CacheBones_AnyThread(const FAnimationCacheBonesContex
 
 	FBoneContainer& RequiredBones = Context.AnimInstanceProxy->GetRequiredBones();
 	InputToCurveMappingUIDs.Reset();
+	InputToControlIndex.Reset();
 
 	if(RequiredBones.IsValid())
 	{
@@ -134,36 +135,50 @@ void FAnimNode_ControlRig::CacheBones_AnyThread(const FAnimationCacheBonesContex
 		
 		TArray<FName> const& UIDToNameLookUpTable = RequiredBones.GetUIDToNameLookupTable();
 
-		auto CacheCurveMappingUIDs = [&](const TMap<FName, FName>& Mapping, TArray<FName> const& InUIDToNameLookUpTable, 
-			const FAnimationCacheBonesContext& InContext)
+		auto CacheMapping = [&](const TMap<FName, FName>& Mapping, TArray<FName> const& InUIDToNameLookUpTable, 
+			const FAnimationCacheBonesContext& InContext, URigHierarchy* InHierarchy)
 		{
 			for (auto Iter = Mapping.CreateConstIterator(); Iter; ++Iter)
 			{
 				// we need to have list of variables using pin
 				const FName SourcePath = Iter.Key();
-				const FName CurveName = Iter.Value();
+				const FName TargetPath = Iter.Value();
 
-				if (SourcePath != NAME_None && CurveName != NAME_None)
+				if (SourcePath != NAME_None && TargetPath != NAME_None)
 				{
-					int32 Found = InUIDToNameLookUpTable.Find(CurveName);
+					int32 Found = InUIDToNameLookUpTable.Find(TargetPath);
 					if (Found != INDEX_NONE)
 					{
 						// set value - sound should be UID
 						InputToCurveMappingUIDs.Add(Iter.Value()) = Found;
+						continue;
 					}
-					else
+					else if(InHierarchy)
 					{
-						UE_LOG(LogAnimation, Warning, TEXT("Curve %s Not Found from the Skeleton %s"), 
-							*CurveName.ToString(), *GetNameSafe(InContext.AnimInstanceProxy->GetSkeleton()));
+						const FRigElementKey Key(TargetPath, ERigElementType::Control);
+						if(const FRigControlElement* ControlElement = InHierarchy->Find<FRigControlElement>(Key))
+						{
+							InputToControlIndex.Add(TargetPath, ControlElement->GetIndex());
+							continue;
+						}
 					}
+
+					UE_LOG(LogAnimation, Warning, TEXT("Curve %s Not Found from the Skeleton %s"), 
+						*TargetPath.ToString(), *GetNameSafe(InContext.AnimInstanceProxy->GetSkeleton()));
 				}
 
 				// @todo: should we clear the item if not found?
 			}
 		};
 
-		CacheCurveMappingUIDs(InputMapping, UIDToNameLookUpTable, Context);
-		CacheCurveMappingUIDs(OutputMapping, UIDToNameLookUpTable, Context);
+		URigHierarchy* Hierarchy = nullptr;
+		if(UControlRig* CurrentControlRig = GetControlRig())
+		{
+			Hierarchy = CurrentControlRig->GetHierarchy();
+		}
+
+		CacheMapping(InputMapping, UIDToNameLookUpTable, Context, Hierarchy);
+		CacheMapping(OutputMapping, UIDToNameLookUpTable, Context, Hierarchy);
 	}
 }
 
@@ -388,44 +403,180 @@ void FAnimNode_ControlRig::PropagateInputProperties(const UObject* InSourceInsta
 		{
 			FProperty* CallerProperty = SourceProperties[PropIdx];
 
-			FRigVMExternalVariable Variable = TargetControlRig->GetPublicVariableByName(DestPropertyNames[PropIdx]);
-			if (Variable.bIsReadOnly)
+			if(FRigControlElement* ControlElement = TargetControlRig->FindControl(DestPropertyNames[PropIdx]))
 			{
-				continue;
-			}
+				const uint8* SrcPtr = CallerProperty->ContainerPtrToValuePtr<uint8>(InSourceInstance);
 
-			const uint8* SrcPtr = CallerProperty->ContainerPtrToValuePtr<uint8>(InSourceInstance);
+				bool bIsValid = false;
+				FRigControlValue Value;
+				switch(ControlElement->Settings.ControlType)
+				{
+					case ERigControlType::Bool:
+					{
+						if(ensure(CastField<FBoolProperty>(CallerProperty)))
+						{
+							Value = FRigControlValue::Make<bool>(*(bool*)SrcPtr);
+							bIsValid = true;
+						}
+						break;
+					}
+					case ERigControlType::Float:
+					{
+						if(ensure(CastField<FFloatProperty>(CallerProperty)))
+						{
+							Value = FRigControlValue::Make<float>(*(float*)SrcPtr);
+							bIsValid = true;
+						}
+						break;
+					}
+					case ERigControlType::Integer:
+					{
+						if(ensure(CastField<FIntProperty>(CallerProperty)))
+						{
+							Value = FRigControlValue::Make<int32>(*(int32*)SrcPtr);
+							bIsValid = true;
+						}
+						break;
+					}
+					case ERigControlType::Vector2D:
+					{
+						FStructProperty* StructProperty = CastField<FStructProperty>(CallerProperty);
+						if(ensure(StructProperty))
+						{
+							if(ensure(StructProperty->Struct == TBaseStructure<FVector2D>::Get()))
+							{
+								const FVector2D& SrcVector = *(FVector2D*)SrcPtr;
+								Value = FRigControlValue::Make<FVector2D>(SrcVector);
+								bIsValid = true;
+							}
+						}
+						break;
+					}
+					case ERigControlType::Position:
+					case ERigControlType::Scale:
+					{
+						FStructProperty* StructProperty = CastField<FStructProperty>(CallerProperty);
+						if(ensure(StructProperty))
+						{
+							if(ensure(StructProperty->Struct == TBaseStructure<FVector>::Get()))
+							{
+								const FVector& SrcVector = *(FVector*)SrcPtr;  
+								Value = FRigControlValue::Make<FVector>(SrcVector);
+								bIsValid = true;
+							}
+						}
+						break;
+					}
+					case ERigControlType::Rotator:
+					{
+						FStructProperty* StructProperty = CastField<FStructProperty>(CallerProperty);
+						if(ensure(StructProperty))
+						{
+							if(ensure(StructProperty->Struct == TBaseStructure<FRotator>::Get()))
+							{
+								const FRotator& SrcRotator = *(FRotator*)SrcPtr;
+								Value = FRigControlValue::Make<FRotator>(SrcRotator);
+								bIsValid = true;
+							}
+						}
+						break;
+					}
+					case ERigControlType::Transform:
+					{
+						FStructProperty* StructProperty = CastField<FStructProperty>(CallerProperty);
+						if(ensure(StructProperty))
+						{
+							if(ensure(StructProperty->Struct == TBaseStructure<FTransform>::Get()))
+							{
+								const FTransform& SrcTransform = *(FTransform*)SrcPtr;  
+								Value = FRigControlValue::Make<FTransform>(SrcTransform);
+								bIsValid = true;
+							}
+						}
+						break;
+					}
+					case ERigControlType::TransformNoScale:
+					{
+						FStructProperty* StructProperty = CastField<FStructProperty>(CallerProperty);
+						if(ensure(StructProperty))
+						{
+							if(ensure(StructProperty->Struct == TBaseStructure<FTransform>::Get()))
+							{
+								const FTransform& SrcTransform = *(FTransform*)SrcPtr;  
+								Value = FRigControlValue::Make<FTransformNoScale>(SrcTransform);
+								bIsValid = true;
+							}
+						}
+						break;
+					}
+					case ERigControlType::EulerTransform:
+					{
+						FStructProperty* StructProperty = CastField<FStructProperty>(CallerProperty);
+						if(ensure(StructProperty))
+						{
+							if(ensure(StructProperty->Struct == TBaseStructure<FTransform>::Get()))
+							{
+								const FTransform& SrcTransform = *(FTransform*)SrcPtr;  
+								Value = FRigControlValue::Make<FEulerTransform>(SrcTransform);
+								bIsValid = true;
+							}
+						}
+						break;
+					}
+					default:
+					{
+						checkNoEntry();
+					}
+				}
 
-			if (CastField<FBoolProperty>(CallerProperty) != nullptr && Variable.TypeName == TEXT("bool"))
-			{
-				const bool Value = *(const bool*)SrcPtr;
-				Variable.SetValue<bool>(Value);
-			}
-			else if (CastField<FFloatProperty>(CallerProperty) != nullptr && Variable.TypeName == TEXT("float"))
-			{
-				const float Value = *(const float*)SrcPtr;
-				Variable.SetValue<float>(Value);
-			}
-			else if (CastField<FIntProperty>(CallerProperty) != nullptr && Variable.TypeName == TEXT("int32"))
-			{
-				const int32 Value = *(const int32*)SrcPtr;
-				Variable.SetValue<int32>(Value);
-			}
-			else if (CastField<FNameProperty>(CallerProperty) != nullptr && Variable.TypeName == TEXT("FName"))
-			{
-				const FName Value = *(const FName*)SrcPtr;
-				Variable.SetValue<FName>(Value);
-			}
-			else if (CastField<FNameProperty>(CallerProperty) != nullptr && Variable.TypeName == TEXT("FString"))
-			{
-				const FString Value = *(const FString*)SrcPtr;
-				Variable.SetValue<FString>(Value);
+				if(bIsValid)
+				{
+					TargetControlRig->GetHierarchy()->SetControlValue(ControlElement, Value, ERigControlValueType::Current);
+				}
+				continue;
 			}
-			else if (FStructProperty* StructProperty = CastField<FStructProperty>(CallerProperty))
+
+			FRigVMExternalVariable Variable = TargetControlRig->GetPublicVariableByName(DestPropertyNames[PropIdx]);
+			if(Variable.IsValid())
 			{
-				if (StructProperty->Struct == Variable.TypeObject)
+				if (Variable.bIsReadOnly)
+				{
+					continue;
+				}
+
+				const uint8* SrcPtr = CallerProperty->ContainerPtrToValuePtr<uint8>(InSourceInstance);
+
+				if (CastField<FBoolProperty>(CallerProperty) != nullptr && Variable.TypeName == TEXT("bool"))
+				{
+					const bool Value = *(const bool*)SrcPtr;
+					Variable.SetValue<bool>(Value);
+				}
+				else if (CastField<FFloatProperty>(CallerProperty) != nullptr && Variable.TypeName == TEXT("float"))
+				{
+					const float Value = *(const float*)SrcPtr;
+					Variable.SetValue<float>(Value);
+				}
+				else if (CastField<FIntProperty>(CallerProperty) != nullptr && Variable.TypeName == TEXT("int32"))
 				{
-					StructProperty->Struct->CopyScriptStruct(Variable.Memory, SrcPtr, 1);
+					const int32 Value = *(const int32*)SrcPtr;
+					Variable.SetValue<int32>(Value);
+				}
+				else if (CastField<FNameProperty>(CallerProperty) != nullptr && Variable.TypeName == TEXT("FName"))
+				{
+					const FName Value = *(const FName*)SrcPtr;
+					Variable.SetValue<FName>(Value);
+				}
+				else if (CastField<FNameProperty>(CallerProperty) != nullptr && Variable.TypeName == TEXT("FString"))
+				{
+					const FString Value = *(const FString*)SrcPtr;
+					Variable.SetValue<FString>(Value);
+				}
+				else if (FStructProperty* StructProperty = CastField<FStructProperty>(CallerProperty))
+				{
+					if (StructProperty->Struct == Variable.TypeObject)
+					{
+						StructProperty->Struct->CopyScriptStruct(Variable.Memory, SrcPtr, 1);
+					}
 				}
 			}
 		}
diff --git a/Engine/Plugins/Animation/ControlRig/Source/ControlRig/Private/Rigs/RigHierarchy.cpp b/Engine/Plugins/Animation/ControlRig/Source/ControlRig/Private/Rigs/RigHierarchy.cpp
--- a/Engine/Plugins/Animation/ControlRig/Source/ControlRig/Private/Rigs/RigHierarchy.cpp
+++ b/Engine/Plugins/Animation/ControlRig/Source/ControlRig/Private/Rigs/RigHierarchy.cpp
@@ -783,6 +783,134 @@ FName URigHierarchy::GetSafeNewName(const FString& InPotentialNewName, ERigEleme
 	return *Name;
 }
 
+FEdGraphPinType URigHierarchy::GetControlPinType(FRigControlElement* InControlElement) const
+{
+	check(InControlElement);
+
+	// local copy of UEdGraphSchema_K2::PC_ ... static members
+	static const FName PC_Boolean(TEXT("bool"));
+	static const FName PC_Float(TEXT("float"));
+	static const FName PC_Int(TEXT("int"));
+	static const FName PC_Struct(TEXT("struct"));
+
+	FEdGraphPinType PinType;
+
+	switch(InControlElement->Settings.ControlType)
+	{
+		case ERigControlType::Bool:
+		{
+			PinType.PinCategory = PC_Boolean;
+			break;
+		}
+		case ERigControlType::Float:
+		{
+			PinType.PinCategory = PC_Float;
+			break;
+		}
+		case ERigControlType::Integer:
+		{
+			PinType.PinCategory = PC_Int;
+			break;
+		}
+		case ERigControlType::Vector2D:
+		{
+			PinType.PinCategory = PC_Struct;
+			PinType.PinSubCategoryObject = TBaseStructure<FVector2D>::Get();
+			break;
+		}
+		case ERigControlType::Position:
+		case ERigControlType::Scale:
+		{
+			PinType.PinCategory = PC_Struct;
+			PinType.PinSubCategoryObject = TBaseStructure<FVector>::Get();
+			break;
+		}
+		case ERigControlType::Rotator:
+		{
+			PinType.PinCategory = PC_Struct;
+			PinType.PinSubCategoryObject = TBaseStructure<FRotator>::Get();
+			break;
+		}
+		case ERigControlType::Transform:
+		case ERigControlType::TransformNoScale:
+		case ERigControlType::EulerTransform:
+		{
+			PinType.PinCategory = PC_Struct;
+			PinType.PinSubCategoryObject = TBaseStructure<FTransform>::Get();
+			break;
+		}
+	}
+
+	return PinType;
+}
+
+FString URigHierarchy::GetControlPinDefaultValue(FRigControlElement* InControlElement, bool bForEdGraph, ERigControlValueType InValueType) const
+{
+	check(InControlElement);
+
+	FRigControlValue Value = GetControlValue(InControlElement, InValueType);
+	switch(InControlElement->Settings.ControlType)
+	{
+		case ERigControlType::Bool:
+		{
+			return Value.ToString<bool>();
+		}
+		case ERigControlType::Float:
+		{
+			return Value.ToString<float>();
+		}
+		case ERigControlType::Integer:
+		{
+			return Value.ToString<int32>();
+		}
+		case ERigControlType::Vector2D:
+		{
+			if(bForEdGraph)
+			{
+				const FVector3f Vector = Value.Get<FVector3f>();
+				return FVector2D(Vector.X, Vector.Y).ToString();
+			}
+			return Value.ToString<FVector2D>();
+		}
+		case ERigControlType::Position:
+		case ERigControlType::Scale:
+		{
+			if(bForEdGraph)
+			{
+				return FVector(Value.Get<FVector3f>()).ToString();
+			}
+			return Value.ToString<FVector>();
+		}
+		case ERigControlType::Rotator:
+		{
+				if(bForEdGraph)
+				{
+					const FRotator Rotator = FRotator::MakeFromEuler(Value.GetRef<FVector3f>());
+					return Rotator.ToString();
+				}
+				return Value.ToString<FRotator>();
+		}
+		case ERigControlType::Transform:
+		case ERigControlType::TransformNoScale:
+		case ERigControlType::EulerTransform:
+		{
+			const FTransform Transform = Value.GetAsTransform(
+				InControlElement->Settings.ControlType,
+				InControlElement->Settings.PrimaryAxis);
+				
+			if(bForEdGraph)
+			{
+				return Transform.ToString();
+			}
+
+			FString Result;
+			TBaseStructure<FTransform>::Get()->ExportText(Result, &Transform, nullptr, nullptr, PPF_None, nullptr);
+			return Result;
+		}
+	}
+	return FString();
+}
+
 TArray<FRigElementKey> URigHierarchy::GetChildren(FRigElementKey InKey, bool bRecursive) const
 {
 	FRigBaseElementChildrenArray LocalChildren;
diff --git a/Engine/Plugins/Animation/ControlRig/Source/ControlRig/Public/AnimNode_ControlRigBase.h b/Engine/Plugins/Animation/ControlRig/Source/ControlRig/Public/AnimNode_ControlRigBase.h
--- a/Engine/Plugins/Animation/ControlRig/Source/ControlRig/Public/AnimNode_ControlRigBase.h
+++ b/Engine/Plugins/Animation/ControlRig/Source/ControlRig/Public/AnimNode_ControlRigBase.h
@@ -119,6 +119,7 @@ protected:
 	TMap<FName, uint16> ControlRigCurveMappingByName;
 
 	TMap<FName, uint16> InputToCurveMappingUIDs;
+	TMap<FName, int32> InputToControlIndex;
 
 	/** Node Mapping Container */
 	UPROPERTY(transient)
diff --git a/Engine/Plugins/Animation/ControlRig/Source/ControlRig/Public/Rigs/RigHierarchy.h b/Engine/Plugins/Animation/ControlRig/Source/ControlRig/Public/Rigs/RigHierarchy.h
--- a/Engine/Plugins/Animation/ControlRig/Source/ControlRig/Public/Rigs/RigHierarchy.h
+++ b/Engine/Plugins/Animation/ControlRig/Source/ControlRig/Public/Rigs/RigHierarchy.h
@@ -6,6 +6,7 @@
 #include "RigHierarchyElements.h"
 #include "RigHierarchyPose.h"
 #include "UObject/WeakObjectPtrTemplates.h"
+#include "EdGraph/EdGraphPin.h"
 #include "RigHierarchy.generated.h"
 
 class URigHierarchy;
@@ -1178,6 +1179,22 @@ public:
 		return GetInitialControlValue(InElementIndex).Get<T>();
 	}
 
+	/**
+	 * Returns the pin type to use for a control
+	 * @param InControlElement The control to return the pin type for
+	 * @return The pin type
+	 */
+	FEdGraphPinType GetControlPinType(FRigControlElement* InControlElement) const;
+
+	/**
+	 * Returns the default value to use for a pin for a control
+	 * @param InControlElement The control to return the pin default value for
+	 * @param bForEdGraph If this is true to the math types' ::ToString will be used rather than text export
+	 * @param InValueType The type of value to return
+	 * @return The pin default value
+	 */
+	FString GetControlPinDefaultValue(FRigControlElement* InControlElement, bool bForEdGraph, ERigControlValueType InValueType = ERigControlValueType::Initial) const;
+
 	/**
 	 * Sets a control's current value given its key
 	 * @param InKey The key of the element to set the current value for
diff --git a/Engine/Plugins/Animation/ControlRig/Source/ControlRig/Public/Rigs/RigHierarchyDefines.h b/Engine/Plugins/Animation/ControlRig/Source/ControlRig/Public/Rigs/RigHierarchyDefines.h
--- a/Engine/Plugins/Animation/ControlRig/Source/ControlRig/Public/Rigs/RigHierarchyDefines.h
+++ b/Engine/Plugins/Animation/ControlRig/Source/ControlRig/Public/Rigs/RigHierarchyDefines.h
@@ -997,6 +997,29 @@ inline FRigControlValue FRigControlValue::Make(FEulerTransform InValue)
 	return Make<FEulerTransform_Float>(InValue);
 }
 
+template<>
+inline FString FRigControlValue::ToString<bool>() const
+{
+	const bool Value = Get<bool>();
+	static const FString True = TEXT("True");
+	static const FString False = TEXT("False");
+	return Value ? True : False;
+}
+
+template<>
+inline FString FRigControlValue::ToString<int32>() const
+{
+	const int32 Value = Get<int32>();
+	return FString::FromInt(Value);
+}
+
+template<>
+inline FString FRigControlValue::ToString<float>() const
+{
+	const float Value = Get<float>();
+	return FString::SanitizeFloat(Value);
+}
+
 #if !UE_LARGE_WORLD_COORDINATES_DISABLED
 
 template<>
@@ -1008,6 +1031,16 @@ inline FString FRigControlValue::ToString<FVector>() const
 	return Result;
 }
 
+template<>
+inline FString FRigControlValue::ToString<FVector2D>() const
+{
+	const FVector3f Value = GetRef<FVector3f>();
+	FVector2D Value2D(Value.X, Value.Y);
+	FString Result;
+	TBaseStructure<FVector2D>::Get()->ExportText(Result, &Value2D, nullptr, nullptr, PPF_None, nullptr);
+	return Result;
+}
+
 #endif
 
 template<>
diff --git a/Engine/Plugins/Animation/ControlRig/Source/ControlRigDeveloper/Private/AnimGraphNode_ControlRig.cpp b/Engine/Plugins/Animation/ControlRig/Source/ControlRigDeveloper/Private/AnimGraphNode_ControlRig.cpp
--- a/Engine/Plugins/Animation/ControlRig/Source/ControlRigDeveloper/Private/AnimGraphNode_ControlRig.cpp
+++ b/Engine/Plugins/Animation/ControlRig/Source/ControlRigDeveloper/Private/AnimGraphNode_ControlRig.cpp
@@ -179,6 +179,67 @@ void UAnimGraphNode_ControlRig::ReallocatePinsDuringReconstruction(TArray<UEdGra
 		}
 	}
 
+	if(const UClass* ControlRigClass = GetTargetClass())
+	{
+		if(UControlRig* CDO = ControlRigClass->GetDefaultObject<UControlRig>())
+		{
+			if(const URigHierarchy* Hierarchy = CDO->GetHierarchy())
+			{
+				Hierarchy->ForEach<FRigControlElement>([&](FRigControlElement* ControlElement) -> bool
+				{
+					if(ControlElement->Settings.bAnimatable)
+					{
+						const FName ControlName = ControlElement->GetName();
+						BeginExposableNames.Remove(ControlName);
+
+						if (!KnownExposableProperties.Contains(ControlName))
+						{
+							// New property added to the target class
+							KnownExposableProperties.Add(ControlName);
+						}
+
+						if (ExposedPropertyNames.Contains(ControlName))
+						{
+							const FEdGraphPinType PinType = Hierarchy->GetControlPinType(ControlElement);
+							if (!PinType.PinCategory.IsValid())
+							{
+								return true;
+							}
+
+							UEdGraphPin* NewPin = CreatePin(EEdGraphPinDirection::EGPD_Input, PinType, ControlName);
+							NewPin->PinFriendlyName = FText::FromName(ControlElement->GetName());
+
+							// Newly created pin does not have an auto generated default value, so we need to generate one here
+							// Missing the auto-gen default would cause UEdGraphPin::MovePersistentDataFromOldPin to override
+							// the actual default value with the empty auto-gen default, causing BP compiler to complain
+							// This is similar to how the following two functions create and initialize new pins, create first 
+							// and then set an auto-gen default
+							// FOptionalPinManager::CreateVisiblePins()
+							// FAnimBlueprintNodeOptionalPinManager::PostInitNewPin()
+							AnimGraphDefaultSchema->SetPinAutogeneratedDefaultValueBasedOnType(NewPin);
+
+							// We cant interrogate CDO here as we may be mid-compile, so we can only really
+							// reset to the autogenerated default.
+							AnimGraphDefaultSchema->ResetPinToAutogeneratedDefaultValue(NewPin, false);
+
+							// Extract default values from the Target Control Rig if possible
+							const FString DefaultValue = Hierarchy->GetControlPinDefaultValue(ControlElement, true);
+							if(!DefaultValue.IsEmpty())
+							{
+								AnimGraphDefaultSchema->TrySetDefaultValue(*NewPin, DefaultValue);
+							}
+							
+							// sustain the current set of custom pins - we'll refrain from changing the node until post load is complete 
+							CustomizePinData(NewPin, ControlName, INDEX_NONE);
+						}
+					}
+					
+					return true;
+				});
+			}
+		}
+	}
+
 	// Remove any properties that no longer exist on the target class
 	for (FName& RemovedPropertyName : BeginExposableNames)
 	{
@@ -269,6 +330,26 @@ void UAnimGraphNode_ControlRig::RebuildExposedProperties()
 		KnownExposableProperties.Add(Iter.Key());
 	}
 
+	// also add all of the controls
+	if(const UClass* TargetClass = GetTargetClass())
+	{
+		if(UControlRig* CDO = TargetClass->GetDefaultObject<UControlRig>())
+		{
+			if(const URigHierarchy* Hierarchy = CDO->GetHierarchy())
+			{
+				Hierarchy->ForEach<FRigControlElement>([&](const FRigControlElement* ControlElement) -> bool
+				{
+					if(ControlElement->Settings.bAnimatable)
+					{
+						const FName ControlName = ControlElement->GetName();
+						KnownExposableProperties.Add(ControlName);
+					}
+					return true;
+				});
+			}
+		}
+	}
+
 	for (int32 Index = 0; Index < ExposedPropertyNames.Num(); ++Index)
 	{
 		// remove if it's not there anymore
@@ -284,7 +365,23 @@ bool UAnimGraphNode_ControlRig::IsInputProperty(const FName& PropertyName) const
 {
 	DECLARE_SCOPE_HIERARCHICAL_COUNTER_FUNC()
 
-	return InputVariables.Contains(PropertyName);
+	// this is true for both input variables and controls
+	return !OutputVariables.Contains(PropertyName);
+}
+
+FRigControlElement* UAnimGraphNode_ControlRig::FindControlElement(const FName& InControlName) const
+{
+	if(const UClass* ControlRigClass = GetTargetClass())
+	{
+		if(UControlRig* CDO = ControlRigClass->GetDefaultObject<UControlRig>())
+		{
+			if(const URigHierarchy* Hierarchy = CDO->GetHierarchy())
+			{
+				return (FRigControlElement*)Hierarchy->Find<FRigControlElement>(FRigElementKey(InControlName, ERigElementType::Control));
+			}
+		}
+	}
+	return nullptr;
 }
 
 bool UAnimGraphNode_ControlRig::IsAvailableToMapToCurve(const FName& PropertyName, bool bInput) const
@@ -294,12 +391,17 @@ bool UAnimGraphNode_ControlRig::IsAvailableToMapToCurve(const FName& PropertyNam
 	// find if input or output
 	// ensure it could convert to float
 	const FRigVMExternalVariable* Variable = (bInput) ? InputVariables.Find(PropertyName) : OutputVariables.Find(PropertyName);
-	if (ensure(Variable))
+	if (Variable)
 	{
 		return Variable->TypeName == TEXT("float");
 	}
 
-	return false;
+	if(const FRigControlElement* ControlElement = FindControlElement(PropertyName))
+	{
+		return ControlElement->Settings.ControlType == ERigControlType::Float;
+	}
+
+	return ensure(false);
 }
 
 bool UAnimGraphNode_ControlRig::IsPropertyExposeEnabled(FName PropertyName) const
@@ -483,23 +585,42 @@ void UAnimGraphNode_ControlRig::GetAvailableMapping(const FName& PathName, TArra
 	{
 		const FSmartNameMapping* CurveMapping = TargetSkeleton->GetSmartNameContainer(USkeleton::AnimCurveMappingName);
 		CurveMapping->FillNameArray(OutArray);
+		
+		// also add all controls
+		if(const UClass* ControlRigClass = GetTargetClass())
+		{
+			if(UControlRig* CDO = ControlRigClass->GetDefaultObject<UControlRig>())
+			{
+				if(const URigHierarchy* Hierarchy = CDO->GetHierarchy())
+				{
+					Hierarchy->ForEach<FRigControlElement>([&](FRigControlElement* ControlElement) -> bool
+					{
+						if(ControlElement->Settings.bAnimatable)
+						{
+							OutArray.Add(ControlElement->GetName());
+						}
+						return true;
+					});
+				}
+			}
+		}
 
-		TArray<FName> InputCurves, OutputCurves;
 		// we want to exclude anything that has been mapped already
-		Node.InputMapping.GenerateValueArray(InputCurves);
-		Node.InputMapping.GenerateValueArray(OutputCurves);
+		TArray<FName> InputMapping, OutputMapping;
+		Node.InputMapping.GenerateValueArray(InputMapping);
+		Node.OutputMapping.GenerateValueArray(OutputMapping);
 
 		// I have to remove Input/Output Curves from OutArray
 		for (int32 Index = 0; Index < OutArray.Num(); ++Index)
 		{
 			const FName& Item = OutArray[Index];
 
-			if (InputCurves.Contains(Item))
+			if (InputMapping.Contains(Item))
 			{
 				OutArray.RemoveAt(Index);
 				--Index;
 			}
-			else if (OutputCurves.Contains(Item))
+			else if (OutputMapping.Contains(Item))
 			{
 				OutArray.RemoveAt(Index);
 				--Index;
@@ -536,6 +657,35 @@ void UAnimGraphNode_ControlRig::CreateVariableMapping(const FString& FilteredTex
 			OutArray.Add(Item);
 		}
 	}
+
+	if(bInput)
+	{
+		// add all controls
+		if(const UClass* ControlRigClass = GetTargetClass())
+		{
+			if(UControlRig* CDO = ControlRigClass->GetDefaultObject<UControlRig>())
+			{
+				if(const URigHierarchy* Hierarchy = CDO->GetHierarchy())
+				{
+					Hierarchy->ForEach<FRigControlElement>([&](FRigControlElement* ControlElement) -> bool
+					{
+						if(ControlElement->Settings.bAnimatable)
+						{
+							const FName ControlName = ControlElement->GetName();
+							const FString& DisplayName = ControlName.ToString();
+
+							if (!bDoFiltering || DisplayName.Contains(FilteredText))
+							{
+								TSharedRef<FVariableMappingInfo> Item = FVariableMappingInfo::Make(ControlName);
+								OutArray.Add(Item);
+							}
+						}
+						return true;
+					});
+				}
+			}
+		}
+	}
 }
 
 void UAnimGraphNode_ControlRig::PostEditChangeProperty(FPropertyChangedEvent& PropertyChangedEvent)
diff --git a/Engine/Plugins/Animation/ControlRig/Source/ControlRigDeveloper/Public/AnimGraphNode_ControlRig.h b/Engine/Plugins/Animation/ControlRig/Source/ControlRigDeveloper/Public/AnimGraphNode_ControlRig.h
--- a/Engine/Plugins/Animation/ControlRig/Source/ControlRigDeveloper/Public/AnimGraphNode_ControlRig.h
+++ b/Engine/Plugins/Animation/ControlRig/Source/ControlRigDeveloper/Public/AnimGraphNode_ControlRig.h
@@ -56,5 +56,6 @@ private:
 
 	bool IsAvailableToMapToCurve(const FName& PropertyName, bool bInput) const;
 	bool IsInputProperty(const FName& PropertyName) const;
+	FRigControlElement* FindControlElement(const FName& InControlName) const;
 };
 
