Control Rig: Follow up fix for CL 17368377 for builds without Large World Coordinates

diff --git a/Engine/Plugins/Animation/ControlRig/Source/ControlRig/Private/Rigs/RigHierarchy.cpp b/Engine/Plugins/Animation/ControlRig/Source/ControlRig/Private/Rigs/RigHierarchy.cpp
--- a/Engine/Plugins/Animation/ControlRig/Source/ControlRig/Private/Rigs/RigHierarchy.cpp
+++ b/Engine/Plugins/Animation/ControlRig/Source/ControlRig/Private/Rigs/RigHierarchy.cpp
@@ -865,12 +865,17 @@ FString URigHierarchy::GetControlPinDefaultValue(FRigControlElement* InControlEl
 		}
 		case ERigControlType::Vector2D:
 		{
+			const FVector3f Vector = Value.Get<FVector3f>();
+			const FVector2D Vector2D(Vector.X, Vector.Y);
+
 			if(bForEdGraph)
 			{
-				const FVector3f Vector = Value.Get<FVector3f>();
-				return FVector2D(Vector.X, Vector.Y).ToString();
+				return Vector2D.ToString();
 			}
-			return Value.ToString<FVector2D>();
+
+			FString Result;
+			TBaseStructure<FVector2D>::Get()->ExportText(Result, &Vector2D, nullptr, nullptr, PPF_None, nullptr);
+			return Result;
 		}
 		case ERigControlType::Position:
 		case ERigControlType::Scale:
