Control Rig: Add PromoteToVariable actions

diff --git a/Engine/Plugins/Animation/ControlRig/Source/ControlRigDeveloper/Private/Graph/ControlRigGraphSchema.cpp b/Engine/Plugins/Animation/ControlRig/Source/ControlRigDeveloper/Private/Graph/ControlRigGraphSchema.cpp
--- a/Engine/Plugins/Animation/ControlRig/Source/ControlRigDeveloper/Private/Graph/ControlRigGraphSchema.cpp
+++ b/Engine/Plugins/Animation/ControlRig/Source/ControlRigDeveloper/Private/Graph/ControlRigGraphSchema.cpp
@@ -127,6 +127,141 @@ bool FControlRigGraphSchemaAction_LocalVar::IsVariableUsed()
 	return false;
 }
 
+FControlRigGraphSchemaAction_PromoteToVariable::FControlRigGraphSchemaAction_PromoteToVariable(UEdGraphPin* InEdGraphPin, bool InLocalVariable)
+: FEdGraphSchemaAction(	FText(), 
+						InLocalVariable ? LOCTEXT("PromoteToLocalVariable", "Promote to local variable") : LOCTEXT("PromoteToVariable", "Promote to variable"),
+						InLocalVariable ? LOCTEXT("PromoteToLocalVariable", "Promote to local variable") : LOCTEXT("PromoteToVariable", "Promote to variable"),
+						1)
+, EdGraphPin(InEdGraphPin)
+, bLocalVariable(InLocalVariable)
+{
+}
+
+UEdGraphNode* FControlRigGraphSchemaAction_PromoteToVariable::PerformAction(UEdGraph* ParentGraph, UEdGraphPin* FromPin,
+	const FVector2D Location, bool bSelectNewNode)
+{
+	UControlRigGraph* RigGraph = Cast<UControlRigGraph>(ParentGraph);
+	if(RigGraph == nullptr)
+	{
+		return nullptr;
+	}
+
+	UControlRigBlueprint* Blueprint = RigGraph->GetBlueprint();
+	URigVMGraph* Model = RigGraph->GetModel();
+	URigVMPin* ModelPin = Model->FindPin(FromPin->GetName());
+	URigVMController* Controller = RigGraph->GetController();
+	if((Blueprint == nullptr) ||
+		(Model == nullptr) ||
+		(Controller == nullptr))
+	{
+		return nullptr;
+	}
+
+	FName VariableName(NAME_None);
+
+	const FScopedTransaction Transaction(
+		bLocalVariable ?
+		LOCTEXT("GraphEd_PromoteToLocalVariable", "Promote Pin To Local Variable") :
+		LOCTEXT("GraphEd_PromoteToVariable", "Promote Pin To Variable"));
+
+	if(bLocalVariable)
+	{
+		const FRigVMGraphVariableDescription VariableDescription = Controller->AddLocalVariable(
+			*ModelPin->GetPinPath(),
+			ModelPin->GetCPPType(),
+			ModelPin->GetCPPTypeObject(),
+			ModelPin->GetDefaultValue(),
+			true,
+			true
+		);
+
+		VariableName = VariableDescription.Name;
+	}
+	else
+	{
+		Blueprint->Modify();
+
+		FString DefaultValue = ModelPin->GetDefaultValue();
+		if(!DefaultValue.IsEmpty())
+		{
+			if(UScriptStruct* ScriptStruct = Cast<UScriptStruct>(ModelPin->GetCPPTypeObject()))
+			{
+				if(ScriptStruct == TBaseStructure<FVector2D>::Get())
+				{
+					FVector2D Value = FVector2D::ZeroVector;
+					ScriptStruct->ImportText(*DefaultValue, &Value, nullptr, PPF_None, nullptr, FString());
+					DefaultValue = Value.ToString();
+				}
+				if(ScriptStruct == TBaseStructure<FVector>::Get())
+				{
+					FVector Value = FVector::ZeroVector;
+					ScriptStruct->ImportText(*DefaultValue, &Value, nullptr, PPF_None, nullptr, FString());
+					DefaultValue = Value.ToString();
+				}
+				if(ScriptStruct == TBaseStructure<FQuat>::Get())
+				{
+					FQuat Value = FQuat::Identity;
+					ScriptStruct->ImportText(*DefaultValue, &Value, nullptr, PPF_None, nullptr, FString());
+					DefaultValue = Value.ToString();
+				}
+				if(ScriptStruct == TBaseStructure<FRotator>::Get())
+				{
+					FRotator Value = FRotator::ZeroRotator;
+					ScriptStruct->ImportText(*DefaultValue, &Value, nullptr, PPF_None, nullptr, FString());
+					DefaultValue = Value.ToString();
+				}
+				if(ScriptStruct == TBaseStructure<FTransform>::Get())
+				{
+					FTransform Value = FTransform::Identity;
+					ScriptStruct->ImportText(*DefaultValue, &Value, nullptr, PPF_None, nullptr, FString());
+					DefaultValue = Value.ToString();
+				}
+			}
+		}
+
+		FRigVMExternalVariable ExternalVariable;
+		ExternalVariable.Name = FromPin->GetFName();
+		ExternalVariable.bIsArray = ModelPin->IsArray();
+		ExternalVariable.TypeName = ModelPin->IsArray() ? *ModelPin->GetArrayElementCppType() : *ModelPin->GetCPPType();
+		ExternalVariable.TypeObject = ModelPin->GetCPPTypeObject();
+		
+		VariableName = Blueprint->AddCRMemberVariableFromExternal(
+			ExternalVariable,
+			DefaultValue
+		);
+	}
+
+	if(!VariableName.IsNone())
+	{
+		URigVMNode* ModelNode = Controller->AddVariableNode(
+			VariableName,
+			ModelPin->GetCPPType(),
+			ModelPin->GetCPPTypeObject(),
+			FromPin->Direction == EGPD_Input,
+			ModelPin->GetDefaultValue(),
+			Location,
+			FString(),
+			true,
+			true
+		);
+
+		if(ModelNode)
+		{
+			if(FromPin->Direction == EGPD_Input)
+			{
+				Controller->AddLink(ModelNode->FindPin(TEXT("Value")), ModelPin, true);
+			}
+			else
+			{
+				Controller->AddLink(ModelPin, ModelNode->FindPin(TEXT("Value")), true);
+			}
+			return RigGraph->FindNodeForModelNodeName(ModelNode->GetFName());
+		}
+	}
+	
+	return nullptr;
+}
+
 FReply FControlRigFunctionDragDropAction::DroppedOnPanel(const TSharedRef< class SWidget >& Panel, FVector2D ScreenPosition, FVector2D GraphPosition, UEdGraph& Graph)
 {
 	// For local variables
@@ -526,6 +661,38 @@ FLinearColor UControlRigGraphSchema::GetPinTypeColor(const FEdGraphPinType& PinT
 	return GetDefault<UEdGraphSchema_K2>()->GetPinTypeColor(PinType);
 }
 
+void UControlRigGraphSchema::InsertAdditionalActions(TArray<UBlueprint*> InBlueprints, TArray<UEdGraph*> EdGraphs,
+	TArray<UEdGraphPin*> EdGraphPins, FGraphActionListBuilderBase& OutAllActions) const
+{
+	Super::InsertAdditionalActions(InBlueprints, EdGraphs, EdGraphPins, OutAllActions);
+
+	if(EdGraphPins.Num() > 0)
+	{
+		if(UControlRigGraphNode* RigNode = Cast<UControlRigGraphNode>(EdGraphPins[0]->GetOwningNode()))
+		{
+			if(URigVMPin* ModelPin = RigNode->GetModelPinFromPinPath(EdGraphPins[0]->GetName()))
+			{
+				if(!ModelPin->IsExecuteContext())
+				{
+					if(!ModelPin->GetNode()->IsA<URigVMVariableNode>())
+					{
+						OutAllActions.AddAction(TSharedPtr<FControlRigGraphSchemaAction_PromoteToVariable>(
+							new FControlRigGraphSchemaAction_PromoteToVariable(EdGraphPins[0], false)
+						));
+
+						if(!ModelPin->GetGraph()->IsRootGraph())
+						{
+							OutAllActions.AddAction(TSharedPtr<FControlRigGraphSchemaAction_PromoteToVariable>(
+								new FControlRigGraphSchemaAction_PromoteToVariable(EdGraphPins[0], true)
+							));
+						}
+					}
+				}
+			}
+		}
+	}
+}
+
 bool UControlRigGraphSchema::SupportsPinType(UScriptStruct* ScriptStruct) const
 {
 	if(!ScriptStruct)
diff --git a/Engine/Plugins/Animation/ControlRig/Source/ControlRigDeveloper/Public/ControlRigBlueprint.h b/Engine/Plugins/Animation/ControlRig/Source/ControlRigDeveloper/Public/ControlRigBlueprint.h
--- a/Engine/Plugins/Animation/ControlRig/Source/ControlRigDeveloper/Public/ControlRigBlueprint.h
+++ b/Engine/Plugins/Animation/ControlRig/Source/ControlRigDeveloper/Public/ControlRigBlueprint.h
@@ -686,4 +686,5 @@ private:
 	friend class FControlRigGraphDetails;
 	friend class FControlRigEditorModule;
 	friend class UControlRigComponent;
+	friend struct FControlRigGraphSchemaAction_PromoteToVariable;
 };
diff --git a/Engine/Plugins/Animation/ControlRig/Source/ControlRigDeveloper/Public/Graph/ControlRigGraphSchema.h b/Engine/Plugins/Animation/ControlRig/Source/ControlRigDeveloper/Public/Graph/ControlRigGraphSchema.h
--- a/Engine/Plugins/Animation/ControlRig/Source/ControlRigDeveloper/Public/Graph/ControlRigGraphSchema.h
+++ b/Engine/Plugins/Animation/ControlRig/Source/ControlRigDeveloper/Public/Graph/ControlRigGraphSchema.h
@@ -81,6 +81,36 @@ public:
 	virtual bool IsVariableUsed() override;
 };
 
+USTRUCT()
+struct CONTROLRIGDEVELOPER_API FControlRigGraphSchemaAction_PromoteToVariable : public FEdGraphSchemaAction
+{
+	GENERATED_BODY()
+
+public:
+
+	// Simple type info
+	static FName StaticGetTypeId() {static FName Type("FControlRigGraphSchemaAction_PromoteToVariable"); return Type;}
+	virtual FName GetTypeId() const override { return StaticGetTypeId(); } 
+
+	FControlRigGraphSchemaAction_PromoteToVariable()
+		: FEdGraphSchemaAction()
+	{}
+
+	FControlRigGraphSchemaAction_PromoteToVariable(UEdGraphPin* InEdGraphPin, bool InLocalVariable);
+
+	virtual bool IsA(const FName& InType) const override
+	{
+		return InType == GetTypeId();
+	}
+
+	virtual UEdGraphNode* PerformAction(UEdGraph* ParentGraph, UEdGraphPin* FromPin, const FVector2D Location, bool bSelectNewNode) override;
+
+private:
+
+	UEdGraphPin* EdGraphPin;
+	bool bLocalVariable;
+};
+
 /** DragDropAction class for drag and dropping an item from the My Blueprints tree (e.g., variable or function) */
 class CONTROLRIGDEVELOPER_API FControlRigFunctionDragDropAction : public FGraphSchemaActionDragDropAction
 {
@@ -173,6 +203,7 @@ public:
 	virtual UEdGraphPin* DropPinOnNode(UEdGraphNode* InTargetNode, const FName& InSourcePinName, const FEdGraphPinType& InSourcePinType, EEdGraphPinDirection InSourcePinDirection) const override;
 	virtual bool SupportsDropPinOnNode(UEdGraphNode* InTargetNode, const FEdGraphPinType& InSourcePinType, EEdGraphPinDirection InSourcePinDirection, FText& OutErrorMessage) const override;
 	virtual void SetPinBeingDroppedOnNode(UEdGraphPin* InSourcePin) const override { PinBeingDropped = InSourcePin; }
+	virtual void InsertAdditionalActions(TArray<UBlueprint*> InBlueprints, TArray<UEdGraph*> InGraphs, TArray<UEdGraphPin*> InPins, FGraphActionListBuilderBase& OutAllActions) const override;
 
 	/** Returns true if the schema supports the script type */
 	bool SupportsPinType(UScriptStruct* ScriptStruct) const;
diff --git a/Engine/Source/Editor/Kismet/Private/SBlueprintActionMenu.cpp b/Engine/Source/Editor/Kismet/Private/SBlueprintActionMenu.cpp
--- a/Engine/Source/Editor/Kismet/Private/SBlueprintActionMenu.cpp
+++ b/Engine/Source/Editor/Kismet/Private/SBlueprintActionMenu.cpp
@@ -459,6 +459,12 @@ void SBlueprintActionMenu::CollectAllActions(FGraphActionListBuilderBase& OutAll
 	OutAllActions.Append(MenuBuilder); // @TODO: Avoid this copy
 	// also try adding promote to variable if we can do so.
 	TryInsertPromoteToVariable(FilterContext, OutAllActions);
+
+	// give the schema the opportunity to add another action
+	if(const UEdGraphSchema* Schema = Cast<const UEdGraphSchema>(GraphObj->GetSchema()))
+	{
+		Schema->InsertAdditionalActions(FilterContext.Blueprints, FilterContext.Graphs, FilterContext.Pins, OutAllActions);
+	}
 }
 
 void SBlueprintActionMenu::ConstructActionContext(FBlueprintActionContext& ContextDescOut)
diff --git a/Engine/Source/Runtime/Engine/Classes/EdGraph/EdGraphSchema.h b/Engine/Source/Runtime/Engine/Classes/EdGraph/EdGraphSchema.h
--- a/Engine/Source/Runtime/Engine/Classes/EdGraph/EdGraphSchema.h
+++ b/Engine/Source/Runtime/Engine/Classes/EdGraph/EdGraphSchema.h
@@ -1212,4 +1212,12 @@ class ENGINE_API UEdGraphSchema : public UObject
 	*/
 	virtual TSharedPtr<FEdGraphSchemaAction> MakeActionFromVariableDescription(const UEdGraph* InEdGraph, const FBPVariableDescription& VariableDescription) const { return nullptr; }
 
+	/**
+	 * Insert additional actions into the blueprint action menu
+	 * @param InBlueprints List of all blueprints you want actions for.
+	 * @param InGraphs A list of graphs you want compatible actions for.
+	 * @param InPins A list of pins you want compatible actions for.
+	 * @param OutAllActions Resulting compatible actions
+	 */
+	virtual void InsertAdditionalActions(TArray<UBlueprint*> InBlueprints, TArray<UEdGraph*> InGraphs, TArray<UEdGraphPin*> InPins, FGraphActionListBuilderBase& OutAllActions) const {}
 };
