Control Rig: before removing/renaming an element, deselect it first to avoid invalid item in selection array

diff --git a/Engine/Plugins/Animation/ControlRig/Source/ControlRig/Private/Rigs/RigHierarchyController.cpp b/Engine/Plugins/Animation/ControlRig/Source/ControlRig/Private/Rigs/RigHierarchyController.cpp
--- a/Engine/Plugins/Animation/ControlRig/Source/ControlRig/Private/Rigs/RigHierarchyController.cpp
+++ b/Engine/Plugins/Animation/ControlRig/Source/ControlRig/Private/Rigs/RigHierarchyController.cpp
@@ -1758,6 +1758,14 @@ bool URigHierarchyController::RenameElement(FRigBaseElement* InElement, const FN
 	}
 
 	const FRigElementKey OldKey = InElement->GetKey();
+
+	// deselect the key that no longer exists
+	// no need to trigger a reselect since we always clear selection after rename
+	if (Hierarchy->IsSelected(InElement))
+	{
+		DeselectElement(OldKey);
+	}
+	
 	InElement->Key.Name = Hierarchy->GetSafeNewName(InName.ToString(), InElement->GetType());
 	const FRigElementKey NewKey = InElement->GetKey();
 
diff --git a/Engine/Plugins/Animation/ControlRig/Source/ControlRigEditor/Private/Editor/SRigHierarchy.cpp b/Engine/Plugins/Animation/ControlRig/Source/ControlRigEditor/Private/Editor/SRigHierarchy.cpp
--- a/Engine/Plugins/Animation/ControlRig/Source/ControlRigEditor/Private/Editor/SRigHierarchy.cpp
+++ b/Engine/Plugins/Animation/ControlRig/Source/ControlRigEditor/Private/Editor/SRigHierarchy.cpp
@@ -1312,6 +1312,11 @@ void SRigHierarchy::HandleDeleteItem()
  		check(Controller);
 
  		TArray<FRigElementKey> SelectedKeys = GetSelectedKeys();
+
+ 		// clear selection early here to make sure ControlRigEditMode can react to this deletion
+ 		// it cannot react to it during Controller->RemoveElement() later because bSuspendAllNotifications is true
+ 		Controller->ClearSelection();
+ 		
 		for (const FRigElementKey& SelectedKey : SelectedKeys)
 		{
 			TGuardValue<bool> GuardRigHierarchyChanges(bIsChangingRigHierarchy, true);
