Control Rig: Fix for GetControlOffsetTransform returning the wrong transform.

diff --git a/Engine/Plugins/Animation/ControlRig/Source/ControlRig/Private/Rigs/RigHierarchy.cpp b/Engine/Plugins/Animation/ControlRig/Source/ControlRig/Private/Rigs/RigHierarchy.cpp
--- a/Engine/Plugins/Animation/ControlRig/Source/ControlRig/Private/Rigs/RigHierarchy.cpp
+++ b/Engine/Plugins/Animation/ControlRig/Source/ControlRig/Private/Rigs/RigHierarchy.cpp
@@ -2048,8 +2048,8 @@ FTransform URigHierarchy::GetControlOffsetTransform(FRigControlElement* InContro
 		if(IsLocal(InTransformType))
 		{
 			const FTransform LocalTransform = InverseSolveParentConstraints(
-				InControlElement->Offset.Get(OpposedType), 
-				InControlElement->ParentConstraints, InTransformType, FTransform::Identity);
+				InControlElement->Offset.Get(GlobalType), 
+				InControlElement->ParentConstraints, GlobalType, FTransform::Identity);
 			InControlElement->Offset.Set(InTransformType, LocalTransform);
 		}
 		else
@@ -3639,8 +3639,15 @@ FTransform URigHierarchy::InverseSolveParentConstraints(
 	const ERigTransformType::Type InTransformType,
 	const FTransform& InLocalOffsetTransform) const
 {
+
 	FTransform Result = FTransform::Identity;
+
+	// this function is doing roughly the following 
+	// ResultLocalTransform = InGlobalTransform.GetRelative(ParentGlobal)
+	// InTransformType is only used to determine Initial vs Current
 	const bool bInitial = IsInitial(InTransformType);
+	check(ERigTransformType::IsGlobal(InTransformType));
+
 
 	// collect all of the weights
 	FConstraintIndex FirstConstraint;
