Control Rig: Fix frustum select. Swap X/Y as in FDragTool_ActorFrustumSelect::EndDrag

diff --git a/Engine/Plugins/Animation/ControlRig/Source/ControlRigEditor/Private/EditMode/ControlRigEditMode.cpp b/Engine/Plugins/Animation/ControlRig/Source/ControlRigEditor/Private/EditMode/ControlRigEditMode.cpp
--- a/Engine/Plugins/Animation/ControlRig/Source/ControlRigEditor/Private/EditMode/ControlRigEditMode.cpp
+++ b/Engine/Plugins/Animation/ControlRig/Source/ControlRigEditor/Private/EditMode/ControlRigEditMode.cpp
@@ -1389,6 +1389,17 @@ bool FControlRigEditMode::FrustumSelect(const FConvexVolume& InFrustum, FEditorV
 	}
 	const int32 ViewportSizeX = InViewportClient->Viewport->GetSizeXY().X;
 	const int32 ViewportSizeY = InViewportClient->Viewport->GetSizeXY().Y;
+
+	if( StartX > EndX )
+	{
+		Swap( StartX, EndX );
+	}
+
+	if( StartY > EndY )
+	{
+		Swap( StartY, EndY );
+	}
+
 	FIntRect BoxRect(FIntPoint(FMath::Max(0.0f, StartX), FMath::Max(0.0f, StartY)), FIntPoint(FMath::Min(ViewportSizeX, FMath::TruncToInt(EndX + 1)), FMath::Min(ViewportSizeY, FMath::TruncToInt(EndY + 1))));
 
 	TSet<AActor*> HitActors;
