Control Rig: Fix Set Control Global Initial Transform not working. We were accessing tranform without checking the dirty flag.

diff --git a/Engine/Plugins/Animation/ControlRig/Source/ControlRig/Private/Rigs/RigHierarchy.cpp b/Engine/Plugins/Animation/ControlRig/Source/ControlRig/Private/Rigs/RigHierarchy.cpp
--- a/Engine/Plugins/Animation/ControlRig/Source/ControlRig/Private/Rigs/RigHierarchy.cpp
+++ b/Engine/Plugins/Animation/ControlRig/Source/ControlRig/Private/Rigs/RigHierarchy.cpp
@@ -1907,9 +1907,13 @@ FTransform URigHierarchy::GetTransform(FRigTransformElement* InTransformElement,
 		{
 			if(FRigControlElement* ControlElement = Cast<FRigControlElement>(InTransformElement))
 			{
+				
+				// using GetControlOffsetTransform to check dirty flag before accessing the transform
+				// note: no need to do the same for Pose.Local because there is already an ensure:
+				// "ensure(!InTransformElement->Pose.IsDirty(OpposedType));" above
 				const FTransform NewTransform = SolveParentConstraints(
 					ControlElement->ParentConstraints, InTransformType,
-					ControlElement->Offset.Get(OpposedType), true,
+					GetControlOffsetTransform(ControlElement, OpposedType), true,
 					ControlElement->Pose.Get(OpposedType), true);
 				ControlElement->Pose.Set(InTransformType, NewTransform);
 			}
