Control Rig: Space Picker Widget WIP

diff --git a/Engine/Plugins/Animation/ControlRig/Source/ControlRig/Private/ControlRig.cpp b/Engine/Plugins/Animation/ControlRig/Source/ControlRig/Private/ControlRig.cpp
--- a/Engine/Plugins/Animation/ControlRig/Source/ControlRig/Private/ControlRig.cpp
+++ b/Engine/Plugins/Animation/ControlRig/Source/ControlRig/Private/ControlRig.cpp
@@ -2760,6 +2760,33 @@ void UControlRig::SetBoneInitialTransformsFromCompactPose(FCompactPose* InCompac
 	RequestSetup();
 }
 
+const FRigControlElementCustomization* UControlRig::GetControlCustomization(const FRigElementKey& InControl) const
+{
+	check(InControl.Type == ERigElementType::Control);
+
+	if(const FRigControlElementCustomization* Customization = ControlCustomizations.Find(InControl))
+	{
+		return Customization;
+	}
+
+	if(DynamicHierarchy)
+	{
+		if(const FRigControlElement* ControlElement = DynamicHierarchy->Find<FRigControlElement>(InControl))
+		{
+			return &ControlElement->Settings.Customization;
+		}
+	}
+
+	return nullptr;
+}
+
+void UControlRig::SetControlCustomization(const FRigElementKey& InControl, const FRigControlElementCustomization& InCustomization)
+{
+	check(InControl.Type == ERigElementType::Control);
+	
+	ControlCustomizations.FindOrAdd(InControl) = InCustomization;
+}
+
 void UControlRig::OnHierarchyTransformUndoRedo(URigHierarchy* InHierarchy, const FRigElementKey& InKey, ERigTransformType::Type InTransformType, const FTransform& InTransform, bool bIsUndo)
 {
 	if(InKey.Type == ERigElementType::Control)
diff --git a/Engine/Plugins/Animation/ControlRig/Source/ControlRig/Private/Rigs/RigHierarchy.cpp b/Engine/Plugins/Animation/ControlRig/Source/ControlRig/Private/Rigs/RigHierarchy.cpp
--- a/Engine/Plugins/Animation/ControlRig/Source/ControlRig/Private/Rigs/RigHierarchy.cpp
+++ b/Engine/Plugins/Animation/ControlRig/Source/ControlRig/Private/Rigs/RigHierarchy.cpp
@@ -1412,21 +1412,23 @@ bool URigHierarchy::SwitchToParent(FRigBaseElement* InChild, FRigBaseElement* In
 	if(const FRigMultiParentElement* MultiParentElement = Cast<FRigMultiParentElement>(InChild))
 	{
 		int32 ParentIndex = INDEX_NONE;
-		if(const int32* ParentIndexPtr = MultiParentElement->IndexLookup.Find(InParent->GetKey()))
-		{
-			ParentIndex = *ParentIndexPtr;
-		}
-		else
+		if(InParent)
 		{
-			if(URigHierarchyController* Controller = GetController(true))
+			if(const int32* ParentIndexPtr = MultiParentElement->IndexLookup.Find(InParent->GetKey()))
 			{
-				if(Controller->AddParent(InChild, InParent, 0.f, true, false))
+				ParentIndex = *ParentIndexPtr;
+			}
+			else
+			{
+				if(URigHierarchyController* Controller = GetController(true))
 				{
-					ParentIndex = MultiParentElement->IndexLookup.FindChecked(InParent->GetKey());
+					if(Controller->AddParent(InChild, InParent, 0.f, true, false))
+					{
+						ParentIndex = MultiParentElement->IndexLookup.FindChecked(InParent->GetKey());
+					}
 				}
 			}
 		}
-
 		return SwitchToParent(InChild, ParentIndex, bInitial, bAffectChildren);
 	}
 	return false;
@@ -1435,13 +1437,12 @@ bool URigHierarchy::SwitchToParent(FRigBaseElement* InChild, FRigBaseElement* In
 bool URigHierarchy::SwitchToParent(FRigBaseElement* InChild, int32 InParentIndex, bool bInitial, bool bAffectChildren)
 {
 	TArray<FRigElementWeight> Weights = GetParentWeightArray(InChild, bInitial);
+	FMemory::Memzero(Weights.GetData(), Weights.GetAllocatedSize());
 	if(Weights.IsValidIndex(InParentIndex))
 	{
-		FMemory::Memzero(Weights.GetData(), Weights.GetAllocatedSize());
 		Weights[InParentIndex] = 1.f;
-		return SetParentWeightArray(InChild, Weights, bInitial, bAffectChildren);
 	}
-	return false;
+	return SetParentWeightArray(InChild, Weights, bInitial, bAffectChildren);
 }
 
 bool URigHierarchy::SwitchToDefaultParent(FRigElementKey InChild, bool bInitial, bool bAffectChildren)
diff --git a/Engine/Plugins/Animation/ControlRig/Source/ControlRig/Private/Rigs/RigHierarchyController.cpp b/Engine/Plugins/Animation/ControlRig/Source/ControlRig/Private/Rigs/RigHierarchyController.cpp
--- a/Engine/Plugins/Animation/ControlRig/Source/ControlRig/Private/Rigs/RigHierarchyController.cpp
+++ b/Engine/Plugins/Animation/ControlRig/Source/ControlRig/Private/Rigs/RigHierarchyController.cpp
@@ -1679,7 +1679,8 @@ bool URigHierarchyController::RemoveElement(FRigBaseElement* InElement)
 	{
 		if(FRigControlElement* ControlElement = Cast<FRigControlElement>(Element))
 		{
-			ControlElement->Settings.SpaceFavorites.Remove(InElement->GetKey());
+			ControlElement->Settings.Customization.AvailableSpaces.Remove(InElement->GetKey());
+			ControlElement->Settings.Customization.RemovedSpaces.Remove(InElement->GetKey());
 		}
 	}
 	
@@ -1788,7 +1789,7 @@ bool URigHierarchyController::RenameElement(FRigBaseElement* InElement, const FN
 
 		if(FRigControlElement* ControlElement = Cast<FRigControlElement>(Element))
 		{
-			for(FRigElementKey& Favorite : ControlElement->Settings.SpaceFavorites)
+			for(FRigElementKey& Favorite : ControlElement->Settings.Customization.AvailableSpaces)
 			{
 				if(Favorite == OldKey)
 				{
diff --git a/Engine/Plugins/Animation/ControlRig/Source/ControlRig/Private/Rigs/RigHierarchyElements.cpp b/Engine/Plugins/Animation/ControlRig/Source/ControlRig/Private/Rigs/RigHierarchyElements.cpp
--- a/Engine/Plugins/Animation/ControlRig/Source/ControlRig/Private/Rigs/RigHierarchyElements.cpp
+++ b/Engine/Plugins/Animation/ControlRig/Source/ControlRig/Private/Rigs/RigHierarchyElements.cpp
@@ -422,7 +422,7 @@ FRigControlSettings::FRigControlSettings()
 , GizmoColor(FLinearColor::Red)
 , bIsTransientControl(false)
 , ControlEnum(nullptr)
-, SpaceFavorites()
+, Customization()
 {
 }
 
@@ -459,7 +459,7 @@ void FRigControlSettings::Save(FArchive& Ar)
 	Ar << GizmoColor;
 	Ar << bIsTransientControl;
 	Ar << ControlEnumPathName;
-	Ar << SpaceFavorites;
+	Ar << Customization.AvailableSpaces;
 }
 
 void FRigControlSettings::Load(FArchive& Ar)
@@ -515,11 +515,11 @@ void FRigControlSettings::Load(FArchive& Ar)
 
 	if (Ar.CustomVer(FControlRigObjectVersion::GUID) >= FControlRigObjectVersion::RigHierarchyControlSpaceFavorites)
 	{
-		Ar << SpaceFavorites;
+		Ar << Customization.AvailableSpaces;
 	}
 	else
 	{
-		SpaceFavorites.Reset();
+		Customization.AvailableSpaces.Reset();
 	}
 }
 
@@ -589,7 +589,7 @@ bool FRigControlSettings::operator==(const FRigControlSettings& InOther) const
 	{
 		return false;
 	}
-	if(SpaceFavorites != InOther.SpaceFavorites)
+	if(Customization.AvailableSpaces != InOther.Customization.AvailableSpaces)
 	{
 		return false;
 	}
diff --git a/Engine/Plugins/Animation/ControlRig/Source/ControlRig/Public/ControlRig.h b/Engine/Plugins/Animation/ControlRig/Source/ControlRig/Public/ControlRig.h
--- a/Engine/Plugins/Animation/ControlRig/Source/ControlRig/Public/ControlRig.h
+++ b/Engine/Plugins/Animation/ControlRig/Source/ControlRig/Public/ControlRig.h
@@ -414,8 +414,14 @@ public:
 	const FControlRigDrawContainer& GetDrawContainer() const { return DrawContainer; };
 	FControlRigDrawContainer& GetDrawContainer() { return DrawContainer; };
 
+	const FRigControlElementCustomization* GetControlCustomization(const FRigElementKey& InControl) const;
+	void SetControlCustomization(const FRigElementKey& InControl, const FRigControlElementCustomization& InCustomization);
+	
 private:
 
+	UPROPERTY()
+	TMap<FRigElementKey, FRigControlElementCustomization> ControlCustomizations;
+
 	UPROPERTY()
 	TObjectPtr<URigVM> VM;
 
diff --git a/Engine/Plugins/Animation/ControlRig/Source/ControlRig/Public/Rigs/RigHierarchyElements.h b/Engine/Plugins/Animation/ControlRig/Source/ControlRig/Public/Rigs/RigHierarchyElements.h
--- a/Engine/Plugins/Animation/ControlRig/Source/ControlRig/Public/Rigs/RigHierarchyElements.h
+++ b/Engine/Plugins/Animation/ControlRig/Source/ControlRig/Public/Rigs/RigHierarchyElements.h
@@ -835,6 +835,18 @@ protected:
 	friend struct FRigBaseElement;
 };
 
+USTRUCT(BlueprintType)
+struct FRigControlElementCustomization
+{
+	GENERATED_BODY()
+
+	UPROPERTY(BlueprintReadOnly, EditAnywhere, Category = Customization)
+	TArray<FRigElementKey> AvailableSpaces;
+
+	UPROPERTY(BlueprintReadOnly, EditAnywhere, Category = Customization)
+	TArray<FRigElementKey> RemovedSpaces;
+};
+
 USTRUCT(BlueprintType)
 struct CONTROLRIG_API FRigControlSettings
 {
@@ -910,11 +922,12 @@ struct CONTROLRIG_API FRigControlSettings
 	UPROPERTY(BlueprintReadOnly, EditAnywhere, Category = Control)
 	TObjectPtr<UEnum> ControlEnum;
 
-	/** The favorite spaces for this control.
-	 * This will be used as the default content for the space picker
+	/**
+	 * The User interface customization used for a control
+	 * This will be used as the default content for the space picker and other widgets
 	 */
-	UPROPERTY(BlueprintReadOnly, EditAnywhere, Category = Animation)
-	TArray<FRigElementKey> SpaceFavorites;
+	UPROPERTY(BlueprintReadOnly, EditAnywhere, Category = Animation, meta = (DisplayName = "Customization"))
+	FRigControlElementCustomization Customization;
 
 	/** Applies the limits expressed by these settings to a value */
 	FORCEINLINE void ApplyLimits(FRigControlValue& InOutValue) const
diff --git a/Engine/Plugins/Animation/ControlRig/Source/ControlRigEditor/Private/ControlRigEditorStyle.h b/Engine/Plugins/Animation/ControlRig/Source/ControlRigEditor/Private/ControlRigEditorStyle.h
--- a/Engine/Plugins/Animation/ControlRig/Source/ControlRigEditor/Private/ControlRigEditorStyle.h
+++ b/Engine/Plugins/Animation/ControlRig/Source/ControlRigEditor/Private/ControlRigEditorStyle.h
@@ -6,6 +6,7 @@
 #include "Styling/SlateStyleRegistry.h"
 #include "Styling/SlateStyleMacros.h"
 #include "Styling/StyleColors.h"
+#include "Brushes/SlateRoundedBoxBrush.h"
 
 #define TTF_FONT(RelativePath, ...) FSlateFontInfo(RootToContentDir(RelativePath, TEXT(".ttf")), __VA_ARGS__)
 #define OTF_FONT(RelativePath, ...) FSlateFontInfo(RootToContentDir(RelativePath, TEXT(".otf")), __VA_ARGS__)
@@ -143,6 +144,11 @@ public:
 			Set("ControlRig.Hierarchy.Menu", TTF_FONT("Fonts/Roboto-Regular", 12));
 		}
 
+		// Space picker
+		SpacePickerSelectColor = FStyleColors::Select;
+		{
+			Set("ControlRig.SpacePicker.RoundedRect", new FSlateRoundedBoxBrush(FStyleColors::White, 4.0f, FStyleColors::Transparent, 0.0f));
+		}
 
 		FSlateStyleRegistry::RegisterSlateStyle(*this);
 	}
@@ -157,6 +163,8 @@ public:
 	{
 		FSlateStyleRegistry::UnRegisterSlateStyle(*this);
 	}
+
+	FSlateColor SpacePickerSelectColor;
 };
 
 #undef IMAGE_BRUSH
diff --git a/Engine/Plugins/Animation/ControlRig/Source/ControlRigEditor/Private/ControlRigElementDetails.cpp b/Engine/Plugins/Animation/ControlRig/Source/ControlRigEditor/Private/ControlRigElementDetails.cpp
--- a/Engine/Plugins/Animation/ControlRig/Source/ControlRigEditor/Private/ControlRigElementDetails.cpp
+++ b/Engine/Plugins/Animation/ControlRig/Source/ControlRigEditor/Private/ControlRigElementDetails.cpp
@@ -1909,7 +1909,7 @@ void FRigControlElementDetails::CustomizeChildren(TSharedRef<class IPropertyHand
 	}
 
 	IDetailGroup* AnimationGroup = &StructBuilder.AddGroup(TEXT("Animation"), LOCTEXT("Animation", "Animation"));
-	AnimationGroup->AddPropertyRow(SettingsHandle->GetChildHandle(TEXT("SpaceFavorites")).ToSharedRef());
+	AnimationGroup->AddPropertyRow(SettingsHandle->GetChildHandle(TEXT("Customization"))->GetChildHandle(TEXT("AvailableSpaces")).ToSharedRef());
 }
 
 bool FRigControlElementDetails::IsGizmoEnabled() const
diff --git a/Engine/Plugins/Animation/ControlRig/Source/ControlRigEditor/Private/EditMode/SControlRigEditModeTools.cpp b/Engine/Plugins/Animation/ControlRig/Source/ControlRigEditor/Private/EditMode/SControlRigEditModeTools.cpp
--- a/Engine/Plugins/Animation/ControlRig/Source/ControlRigEditor/Private/EditMode/SControlRigEditModeTools.cpp
+++ b/Engine/Plugins/Animation/ControlRig/Source/ControlRigEditor/Private/EditMode/SControlRigEditModeTools.cpp
@@ -306,6 +306,26 @@ void SControlRigEditModeTools::Construct(const FArguments& InArgs, FControlRigEd
 				ControlDetailsView.ToSharedRef()
 			]
 
+			+SVerticalBox::Slot()
+			.AutoHeight()
+			[
+				SAssignNew(PickerExpander, SExpandableArea)
+				.InitiallyCollapsed(true)
+				.AreaTitle(LOCTEXT("Picker_SpaceWidget", "Spaces"))
+				.AreaTitleFont(FEditorStyle::GetFontStyle("DetailsView.CategoryFontStyle"))
+				.BorderBackgroundColor(FLinearColor(.6f, .6f, .6f))
+				.Padding(FMargin(8.f))
+				.BodyContent()
+				[
+					SAssignNew(SpacePickerWidget, SRigSpacePickerWidget)
+					.AllowDelete(true)
+					.AllowReorder(true)
+					.AllowAdd(true)
+					.ShowBakeButton(true)
+					// todo: implement GetAdditionalSpacesDelegate to pull spaces from sequencer
+				]
+			]
+
 			+ SVerticalBox::Slot()
 			.AutoHeight()
 			[
@@ -324,6 +344,8 @@ void SControlRigEditModeTools::Construct(const FArguments& InArgs, FControlRigEd
 	];
 
 	HierarchyTreeView->RefreshTreeView(true);
+	SpacePickerWidget->OnActiveSpaceChanged().AddSP(this, &SControlRigEditModeTools::HandleActiveSpaceChanged);
+	SpacePickerWidget->OnSpaceListChanged().AddSP(this, &SControlRigEditModeTools::HandleSpaceListChanged);
 }
 
 void SControlRigEditModeTools::SetDetailsObjects(const TArray<TWeakObjectPtr<>>& InObjects)
@@ -576,6 +598,59 @@ void SControlRigEditModeTools::OnRigElementSelected(UControlRig* Subject, FRigCo
 			}
 		}
 	}
+
+	if (UControlRig* ControlRig = SequencerRig.Get())
+	{
+		// get the customization for this control
+		const FRigControlElementCustomization* ExistingCustomization = ControlRig->GetControlCustomization(Key);
+		check(ExistingCustomization);
+
+		ControlCustomization = *ExistingCustomization; 
+		SpacePickerWidget->SetControl(ControlRig->GetHierarchy(), Key, &ControlCustomization);
+	}
+}
+
+void SControlRigEditModeTools::HandleActiveSpaceChanged(URigHierarchy* InHierarchy, const FRigElementKey& InControlKey,
+	const FRigElementKey& InSpaceKey)
+{
+	// todo: when the user picks a new space
+	// for now we'll just do it directly in the hierarchy
+	if(InSpaceKey == InHierarchy->GetWorldSpaceSocketKey())
+	{
+		InHierarchy->SwitchToWorldSpace(InControlKey);
+	}
+	else
+	{
+		InHierarchy->SwitchToParent(InControlKey, InSpaceKey);
+	}
+}
+
+void SControlRigEditModeTools::HandleSpaceListChanged(URigHierarchy* InHierarchy, const FRigElementKey& InControlKey,
+	const TArray<FRigElementKey>& InSpaceList)
+{
+	if (UControlRig* ControlRig = SequencerRig.Get())
+	{
+		if(const FRigControlElement* ControlElement = InHierarchy->Find<FRigControlElement>(InControlKey))
+		{
+			const FRigControlElementCustomization* ExistingCustomization = ControlRig->GetControlCustomization(InControlKey);
+			ControlCustomization = *ExistingCustomization;	
+
+			ControlCustomization.AvailableSpaces = InSpaceList;
+			ControlCustomization.RemovedSpaces.Reset();
+
+			// remember  the elements which are in the asset's available list but removed by the user
+			for(const FRigElementKey& AvailableSpace : ControlElement->Settings.Customization.AvailableSpaces)
+			{
+				if(!ControlCustomization.AvailableSpaces.Contains(AvailableSpace))
+				{
+					ControlCustomization.RemovedSpaces.Add(AvailableSpace);
+				}
+			}
+
+			ControlRig->SetControlCustomization(InControlKey, ControlCustomization);
+			InHierarchy->Notify(ERigHierarchyNotification::ControlSettingChanged, ControlElement);
+		}
+	}
 }
 
 EVisibility SControlRigEditModeTools::GetRigOptionExpanderVisibility() const
diff --git a/Engine/Plugins/Animation/ControlRig/Source/ControlRigEditor/Private/EditMode/SControlRigEditModeTools.h b/Engine/Plugins/Animation/ControlRig/Source/ControlRigEditor/Private/EditMode/SControlRigEditModeTools.h
--- a/Engine/Plugins/Animation/ControlRig/Source/ControlRigEditor/Private/EditMode/SControlRigEditModeTools.h
+++ b/Engine/Plugins/Animation/ControlRig/Source/ControlRigEditor/Private/EditMode/SControlRigEditModeTools.h
@@ -7,6 +7,7 @@
 #include "IDetailKeyframeHandler.h"
 #include "RigVMModel/RigVMGraph.h"
 #include "SRigHierarchyTreeView.h"
+#include "SRigSpacePickerWidget.h"
 #include "Widgets/DeclarativeSyntaxSupport.h"
 
 class FControlRigEditMode;
@@ -59,6 +60,10 @@ private:
 	/** Hierarchy picker for controls*/
 	TSharedPtr<SRigHierarchyTreeView> HierarchyTreeView;
 
+	/** Hierarchy picker for controls*/
+	TSharedPtr<SRigSpacePickerWidget> SpacePickerWidget;
+	FRigControlElementCustomization ControlCustomization;
+
 	/** Special picker for controls, no longer used */
 	TSharedPtr<SControlPicker> ControlPicker;
 	TSharedPtr<SExpandableArea> PickerExpander;
@@ -78,6 +83,9 @@ private:
 	void HandleSelectionChanged(TSharedPtr<FRigTreeElement> Selection, ESelectInfo::Type SelectInfo);
 	void OnRigElementSelected(UControlRig* Subject, FRigControlElement* ControlElement, bool bSelected);
 
+	void HandleActiveSpaceChanged(URigHierarchy* InHierarchy, const FRigElementKey& InControlKey, const FRigElementKey& InSpaceKey);
+	void HandleSpaceListChanged(URigHierarchy* InHierarchy, const FRigElementKey& InControlKey, const TArray<FRigElementKey>& InSpaceList);
+
 	EVisibility GetRigOptionExpanderVisibility() const;
 
 	void OnRigOptionFinishedChange(const FPropertyChangedEvent& PropertyChangedEvent);
diff --git a/Engine/Plugins/Animation/ControlRig/Source/ControlRigEditor/Private/Editor/SRigHierarchy.cpp b/Engine/Plugins/Animation/ControlRig/Source/ControlRigEditor/Private/Editor/SRigHierarchy.cpp
--- a/Engine/Plugins/Animation/ControlRig/Source/ControlRigEditor/Private/Editor/SRigHierarchy.cpp
+++ b/Engine/Plugins/Animation/ControlRig/Source/ControlRigEditor/Private/Editor/SRigHierarchy.cpp
@@ -2214,31 +2214,66 @@ void SRigHierarchy::HandleTestSpaceSwitching()
 	{
 		if(ControlKey.Type == ERigElementType::Control)
 		{
-			TSharedRef<SRigSpacePickerWidget> Dialog =
+			TSharedRef<SRigSpacePickerWidget> PickerWidget =
 				SNew(SRigSpacePickerWidget)
 				.Hierarchy(GetDebuggedHierarchy())
-				.SelectedControl(ControlKey);
+				.Control(ControlKey)
+				.AllowDelete(true)
+				.AllowReorder(true)
+				.AllowAdd(true)
+				.Title(LOCTEXT("PickSpace", "Pick Space"));
 
-			SRigSpacePickerWidget::FResult PickedResult = Dialog->InvokeDialog();
-			if(PickedResult.Reply.IsEventHandled() && PickedResult.Key.IsValid())
+			PickerWidget->OnActiveSpaceChanged().AddLambda([this, ControlKey](URigHierarchy*, const FRigElementKey& InControlKey, const FRigElementKey& InSpaceKey)
 			{
+				check(ControlKey == InControlKey);
+				
 				if(URigHierarchy* Hierarchy = GetDebuggedHierarchy())
 				{
-					const FRigElementKey SpaceKey = PickedResult.Key;
-					
-					FTransform Transform = Hierarchy->GetGlobalTransform(ControlKey);
-					if(SpaceKey == Hierarchy->GetWorldSpaceSocketKey())
+					const FTransform Transform = Hierarchy->GetGlobalTransform(ControlKey);
+					if(InSpaceKey == Hierarchy->GetWorldSpaceSocketKey())
 					{
 						Hierarchy->SwitchToWorldSpace(ControlKey);
 					}
 					else
 					{
-						Hierarchy->SwitchToParent(ControlKey, SpaceKey);
+						Hierarchy->SwitchToParent(ControlKey, InSpaceKey);
 					}
 					Hierarchy->SetGlobalTransform(ControlKey, Transform);
 				}
-			}
-			
+			});
+
+			PickerWidget->OnSpaceListChanged().AddLambda([this, ControlKey](URigHierarchy* InHierarchy, const FRigElementKey& InControlKey, const TArray<FRigElementKey>& InSpaceList)
+			{
+				check(ControlKey == InControlKey);
+
+				if(UControlRigBlueprint* RigBlueprint = ControlRigBlueprint.Get())
+				{
+					if(URigHierarchy* Hierarchy = RigBlueprint->Hierarchy)
+					{
+						// update the settings in the control element
+						if(FRigControlElement* ControlElement = Hierarchy->Find<FRigControlElement>(ControlKey))
+						{
+							RigBlueprint->Modify();
+							FScopedTransaction Transaction(LOCTEXT("ControlChangeAvailableSpaces", "Edit Available Spaces"));
+
+							ControlElement->Settings.Customization.AvailableSpaces = InSpaceList;
+							Hierarchy->Notify(ERigHierarchyNotification::ControlSettingChanged, ControlElement);
+						}
+
+						// also update the debugged instance
+						if(Hierarchy != InHierarchy)
+						{
+							if(FRigControlElement* ControlElement = InHierarchy->Find<FRigControlElement>(ControlKey))
+							{
+								ControlElement->Settings.Customization.AvailableSpaces = InSpaceList;
+								InHierarchy->Notify(ERigHierarchyNotification::ControlSettingChanged, ControlElement);
+							}
+						}
+					}
+				}
+			});
+
+			PickerWidget->OpenDialog(false);
 			break;
 		}
 	}
diff --git a/Engine/Plugins/Animation/ControlRig/Source/ControlRigEditor/Private/Editor/SRigHierarchyTreeView.cpp b/Engine/Plugins/Animation/ControlRig/Source/ControlRigEditor/Private/Editor/SRigHierarchyTreeView.cpp
--- a/Engine/Plugins/Animation/ControlRig/Source/ControlRigEditor/Private/Editor/SRigHierarchyTreeView.cpp
+++ b/Engine/Plugins/Animation/ControlRig/Source/ControlRigEditor/Private/Editor/SRigHierarchyTreeView.cpp
@@ -3,7 +3,9 @@
 #include "SRigHierarchyTreeView.h"
 #include "EditorStyleSet.h"
 #include "Widgets/Layout/SSpacer.h"
+#include "Widgets/Layout/SScrollBox.h"
 #include "Widgets/Images/SImage.h"
+#include "Widgets/Input/SSearchBox.h"
 #include "Widgets/Text/STextBlock.h"
 #include "PropertyCustomizationHelpers.h"
 #include "Framework/Application/SlateApplication.h"
@@ -646,4 +648,67 @@ void SRigHierarchyItem::OnNameCommitted(const FText& InText, ETextCommit::Type I
 	}
 }
 
-#undef LOCTEXT_NAMESPACE
\ No newline at end of file
+//////////////////////////////////////////////////////////////
+/// SSearchableRigHierarchyTreeView
+///////////////////////////////////////////////////////////
+
+void SSearchableRigHierarchyTreeView::Construct(const FArguments& InArgs)
+{
+	FRigTreeDelegates TreeDelegates = InArgs._RigTreeDelegates;
+	SuperGetRigTreeDisplaySettings = TreeDelegates.OnGetDisplaySettings;
+
+	TreeDelegates.OnGetDisplaySettings.BindSP(this, &SSearchableRigHierarchyTreeView::GetDisplaySettings);
+	
+	ChildSlot
+	[
+		SNew(SVerticalBox)
+		+SVerticalBox::Slot()
+		.AutoHeight()
+		.VAlign(VAlign_Top)
+		.HAlign(HAlign_Fill)
+		.Padding(0.0f)
+		[
+			SNew(SSearchBox)
+			.InitialText(InArgs._InitialFilterText)
+			.OnTextChanged(this, &SSearchableRigHierarchyTreeView::OnFilterTextChanged)
+		]
+
+		+SVerticalBox::Slot()
+		.AutoHeight()
+		.VAlign(VAlign_Top)
+		.HAlign(HAlign_Fill)
+		.Padding(0.0f, 0.0f)
+		[
+			SNew(SScrollBox)
+			+ SScrollBox::Slot()
+			[
+				SNew(SBorder)
+				.Padding(2.0f)
+				.BorderImage(FEditorStyle::GetBrush("SCSEditor.TreePanel"))
+				[
+					SAssignNew(TreeView, SRigHierarchyTreeView)
+					.RigTreeDelegates(TreeDelegates)
+				]
+			]
+		]
+	];
+}
+
+const FRigTreeDisplaySettings& SSearchableRigHierarchyTreeView::GetDisplaySettings()
+{
+	if(SuperGetRigTreeDisplaySettings.IsBound())
+	{
+		Settings = SuperGetRigTreeDisplaySettings.Execute();
+	}
+	Settings.FilterText = FilterText;
+	return Settings;
+}
+
+void SSearchableRigHierarchyTreeView::OnFilterTextChanged(const FText& SearchText)
+{
+	FilterText = SearchText;
+	GetTreeView()->RefreshTreeView();
+}
+
+
+#undef LOCTEXT_NAMESPACE
diff --git a/Engine/Plugins/Animation/ControlRig/Source/ControlRigEditor/Private/Editor/SRigHierarchyTreeView.h b/Engine/Plugins/Animation/ControlRig/Source/ControlRigEditor/Private/Editor/SRigHierarchyTreeView.h
--- a/Engine/Plugins/Animation/ControlRig/Source/ControlRigEditor/Private/Editor/SRigHierarchyTreeView.h
+++ b/Engine/Plugins/Animation/ControlRig/Source/ControlRigEditor/Private/Editor/SRigHierarchyTreeView.h
@@ -257,3 +257,27 @@ private:
 
 	friend class SRigHierarchy;
 };
+
+class SSearchableRigHierarchyTreeView : public SCompoundWidget
+{
+public:
+
+	SLATE_BEGIN_ARGS(SSearchableRigHierarchyTreeView) {}
+		SLATE_ARGUMENT(FRigTreeDelegates, RigTreeDelegates)
+		SLATE_ARGUMENT(FText, InitialFilterText)
+	SLATE_END_ARGS()
+
+	void Construct(const FArguments& InArgs);
+	virtual ~SSearchableRigHierarchyTreeView() {}
+	TSharedRef<SRigHierarchyTreeView> GetTreeView() const { return TreeView.ToSharedRef(); }
+	const FRigTreeDisplaySettings& GetDisplaySettings();
+
+private:
+
+	void OnFilterTextChanged(const FText& SearchText);
+
+	FOnGetRigTreeDisplaySettings SuperGetRigTreeDisplaySettings;
+	FText FilterText;
+	FRigTreeDisplaySettings Settings;
+	TSharedPtr<SRigHierarchyTreeView> TreeView;
+};
diff --git a/Engine/Plugins/Animation/ControlRig/Source/ControlRigEditor/Private/Editor/SRigSpacePickerWidget.cpp b/Engine/Plugins/Animation/ControlRig/Source/ControlRigEditor/Private/Editor/SRigSpacePickerWidget.cpp
--- a/Engine/Plugins/Animation/ControlRig/Source/ControlRigEditor/Private/Editor/SRigSpacePickerWidget.cpp
+++ b/Engine/Plugins/Animation/ControlRig/Source/ControlRigEditor/Private/Editor/SRigSpacePickerWidget.cpp
@@ -6,6 +6,8 @@
 #include "Widgets/Layout/SUniformGridPanel.h"
 #include "Widgets/Input/SButton.h"
 #include "SRigHierarchyTreeView.h"
+#include "ControlRigEditorStyle.h"
+#include "PropertyCustomizationHelpers.h"
 
 #define LOCTEXT_NAMESPACE "SRigSpacePickerWidget"
 
@@ -13,67 +15,206 @@
 /// SRigSpacePickerWidget
 ///////////////////////////////////////////////////////////
 
+FRigElementKey SRigSpacePickerWidget::InValidKey;
+
 void SRigSpacePickerWidget::Construct(const FArguments& InArgs)
 {
-	Hierarchy = InArgs._Hierarchy;
-	ControlKey = InArgs._SelectedControl;
-	
+	bShowDefaultSpaces = InArgs._ShowDefaultSpaces;
+	bShowFavoriteSpaces = InArgs._ShowFavoriteSpaces;
+	bShowAdditionalSpaces = InArgs._ShowAdditionalSpaces;
+	bAllowReorder = InArgs._AllowReorder;
+	bAllowDelete = InArgs._AllowDelete;
+	bAllowAdd = InArgs._AllowAdd;
+	bShowBakeButton = InArgs._ShowBakeButton;
+	GetAdditionalSpacesDelegate = InArgs._GetAdditionalSpacesDelegate;
+	bRepopulateRequired = false;
+	bLaunchingContextMenu = false;
+
+	if(!GetAdditionalSpacesDelegate.IsBound())
+	{
+		GetAdditionalSpacesDelegate = FRigSpacePickerGetAdditionalSpaces::CreateRaw(this, &SRigSpacePickerWidget::GetCurrentParents);
+	}
+
+	Hierarchy = nullptr;
+	ControlKey = DefaultParentKey = WorldSocketKey = FRigElementKey();
+	Customization = nullptr;
+
 	ChildSlot
 	[
 		SNew(SBorder)
 		.Visibility(EVisibility::Visible)
-		.BorderImage(FEditorStyle::GetBrush("Menu.Background"))
+		.BorderImage(InArgs._BackgroundBrush)
 		[
-			SAssignNew(ListBox, SVerticalBox)
-			
-			+SVerticalBox::Slot()
-			.AutoHeight()
-			.VAlign(VAlign_Top)
-			.HAlign(HAlign_Left)
-			.Padding(4.0, 0.0, 0.0, 2.0)
-			[
-				SNew( STextBlock )
-				.Text( LOCTEXT("Space", "Space") )
-				.Font( IDetailLayoutBuilder::GetDetailFontBold() )
-			]
+			SAssignNew(TopLevelListBox, SVerticalBox)
 		]
 	];
 
-	AddSpacePickerButton(
-		FEditorStyle::GetBrush("EditorViewport.RelativeCoordinateSystem_Local"),
-		LOCTEXT("Parent", "Parent"),
-		FOnClicked::CreateSP(this, &SRigSpacePickerWidget::HandleLocalSpaceClicked)
-	);
-	
-	AddSpacePickerButton(
-		FEditorStyle::GetBrush("EditorViewport.RelativeCoordinateSystem_World"),
-		LOCTEXT("World", "World"),
-		FOnClicked::CreateSP(this, &SRigSpacePickerWidget::HandleWorldSpaceClicked)
-	);
+	if(!InArgs._Title.IsEmpty())
+	{
+		TopLevelListBox->AddSlot()
+		.AutoHeight()
+		.VAlign(VAlign_Top)
+		.HAlign(HAlign_Left)
+		.Padding(4.0, 0.0, 4.0, 12.0)
+		[
+			SNew( STextBlock )
+			.Text( InArgs._Title )
+			.Font( IDetailLayoutBuilder::GetDetailFontBold() )
+		];
+	}
 
-	if(const FRigControlElement* Control = Hierarchy->Find<FRigControlElement>(ControlKey))
+	if(InArgs._ShowDefaultSpaces)
 	{
-		TArray<FRigElementKey> Favorites = Control->Settings.SpaceFavorites;
+		AddSpacePickerRow(
+			TopLevelListBox,
+			ESpacePickerType_Parent,
+			DefaultParentKey,
+			FEditorStyle::GetBrush("EditorViewport.RelativeCoordinateSystem_Local"),
+			LOCTEXT("Parent", "Parent"),
+			FOnClicked::CreateSP(this, &SRigSpacePickerWidget::HandleParentSpaceClicked)
+		);
+		
+		AddSpacePickerRow(
+			TopLevelListBox,
+			ESpacePickerType_World,
+			WorldSocketKey,
+			FEditorStyle::GetBrush("EditorViewport.RelativeCoordinateSystem_World"),
+			LOCTEXT("World", "World"),
+			FOnClicked::CreateSP(this, &SRigSpacePickerWidget::HandleWorldSpaceClicked)
+		);
+	}
 
-		for(const FRigElementKey& Key : Favorites)
+	TopLevelListBox->AddSlot()
+	.AutoHeight()
+	.VAlign(VAlign_Top)
+	.HAlign(HAlign_Fill)
+	.Padding(0.0)
+	[
+		SNew(SHorizontalBox)
+		+ SHorizontalBox::Slot()
+		.FillWidth(1.f)
+		.VAlign(VAlign_Center)
+		.HAlign(HAlign_Fill)
+		.Padding(0)
+		[
+			SAssignNew(ItemSpacesListBox, SVerticalBox)
+		]
+	];
+
+	if(bAllowAdd || bShowBakeButton)
+	{
+		TopLevelListBox->AddSlot()
+		.AutoHeight()
+		.VAlign(VAlign_Top)
+		.HAlign(HAlign_Fill)
+		.Padding(11.f, 8.f, 4.f, 4.f)
+		[
+			SAssignNew(BottomButtonsListBox, SHorizontalBox)
+		];
+
+		if(bAllowAdd)
 		{
-			if(Key.IsValid() && Hierarchy->Contains(Key))
-			{
-				AddSpacePickerButton(
-					SRigHierarchyItem::GetBrushForElementType(Hierarchy, Key),
-					FText::FromName(Key.Name),
-					FOnClicked::CreateSP(this, &SRigSpacePickerWidget::HandleElementSpaceClicked, Key)
-				);
-			}
+			BottomButtonsListBox->AddSlot()
+			.AutoWidth()
+			.VAlign(VAlign_Center)
+			.HAlign(HAlign_Left)
+			.Padding(0.f)
+			[
+				SNew(SButton)
+				.ContentPadding(0.0f)
+				.ButtonStyle(FEditorStyle::Get(), "NoBorder")
+				.OnClicked(this, &SRigSpacePickerWidget::HandleAddElementClicked)
+				.Cursor(EMouseCursor::Default)
+				.ToolTipText(LOCTEXT("AddSpace", "Add Space"))
+				[
+					SNew(SImage)
+					.Image(FEditorStyle::GetBrush(TEXT("Icons.PlusCircle")))
+				]
+			];
+		}
+
+		BottomButtonsListBox->AddSlot()
+		.FillWidth(1.f)
+		.HAlign(HAlign_Fill)
+		[
+			SNew(SSpacer)
+		];
+
+		if(bShowBakeButton)
+		{
+			BottomButtonsListBox->AddSlot()
+			.AutoWidth()
+			.VAlign(VAlign_Center)
+			.HAlign(HAlign_Right)
+			.Padding(0.f)
+			[
+				SNew(SButton)
+				.ButtonStyle(FEditorStyle::Get(), "FlatButton.Default")
+				.Text(LOCTEXT("BakeButton", "Bake"))
+				.OnClicked(InArgs._OnBakeButtonClicked)
+				.ToolTipText(LOCTEXT("BakeButtonToolTip", "Allows to bake the animation of one or more controls to a single space."))
+			];
 		}
 	}
+
+	SetControl(InArgs._Hierarchy, InArgs._Control, InArgs._Customization);
+	SetCanTick(true);
+	
 }
 
 SRigSpacePickerWidget::~SRigSpacePickerWidget()
 {
+	if(HierarchyModifiedHandle.IsValid())
+	{
+		if(Hierarchy)
+		{
+			Hierarchy->OnModified().Remove(HierarchyModifiedHandle);
+			HierarchyModifiedHandle.Reset();
+		}
+	}
+}
+
+void SRigSpacePickerWidget::SetControl(URigHierarchy* InHierarchy, const FRigElementKey& InControl,
+	FRigControlElementCustomization* InCustomization)
+{
+	if(!InControl.IsValid() || InControl.Type != ERigElementType::Control)
+	{
+		return;
+	}
+
+	check(InHierarchy);
+
+	if(Hierarchy && Hierarchy != InHierarchy)
+	{
+		if(HierarchyModifiedHandle.IsValid())
+		{
+			Hierarchy->OnModified().Remove(HierarchyModifiedHandle);
+			HierarchyModifiedHandle.Reset();
+		}
+	}
+	
+	Hierarchy = InHierarchy;
+	ControlKey = InControl;
+	DefaultParentKey = Hierarchy->GetFirstParent(ControlKey);
+	WorldSocketKey = Hierarchy->GetWorldSpaceSocketKey();
+	Customization = InCustomization;
+
+	HierarchyModifiedHandle = Hierarchy->OnModified().AddSP(this, &SRigSpacePickerWidget::OnHierarchyModified);
+
+	// if the customization is not provided
+	if(Customization == nullptr)
+	{
+		if(FRigControlElement* ControlElement = Hierarchy->Find<FRigControlElement>(ControlKey))
+		{
+			Customization = &ControlElement->Settings.Customization;
+		}
+	}
+
+	UpdateActiveSpace();
+	RepopulateItemSpaces();
 }
 
-class SRigSpacePickerWindow : public SWindow
+class SRigSpaceDialogWindow : public SWindow
 {
 public:
 
@@ -93,13 +234,13 @@ private:
 	FSimpleDelegate DeactivatedDelegate;
 }; 
 
-SRigSpacePickerWidget::FResult SRigSpacePickerWidget::InvokeDialog()
+FReply SRigSpacePickerWidget::OpenDialog(bool bModal)
 {
-	PickedKey = FRigElementKey();
-	
+	check(!DialogWindow.IsValid());
+		
 	const FVector2D CursorPos = FSlateApplication::Get().GetCursorPos();
 
-	TSharedRef<SRigSpacePickerWindow> Window = SNew(SRigSpacePickerWindow)
+	TSharedRef<SRigSpaceDialogWindow> Window = SNew(SRigSpaceDialogWindow)
 	.Title( LOCTEXT("SRigSpacePickerWidgetPickSpace", "Pick a new space") )
 	.CreateTitleBar(false)
 	.Type(EWindowType::Menu)
@@ -114,22 +255,38 @@ SRigSpacePickerWidget::FResult SRigSpacePickerWidget::InvokeDialog()
 	Window->SetWidgetToFocusOnActivate(AsShared());
 	Window->OnDeactivated().BindSP(this, &SRigSpacePickerWidget::CloseDialog);
 	
-	PickerWindow = Window;
+	DialogWindow = Window;
 
 	Window->MoveWindowTo(CursorPos);
-	GEditor->EditorAddModalWindow(Window);
 
-	FResult Result;
-	Result.Reply = PickedKey.IsValid() ? FReply::Handled() : FReply::Unhandled();
-	Result.Key = PickedKey;
-	return Result;
+	if(bModal)
+	{
+		GEditor->EditorAddModalWindow(Window);
+	}
+	else
+	{
+		FSlateApplication::Get().AddWindow( Window );
+	}
+
+	return FReply::Handled();
 }
 
 void SRigSpacePickerWidget::CloseDialog()
 {
-	if ( PickerWindow.IsValid() )
+	if(bLaunchingContextMenu)
+	{
+		return;
+	}
+	
+	if(ContextMenu.IsValid())
 	{
-		PickerWindow.Pin()->RequestDestroyWindow();
+		return;
+	}
+	
+	if ( DialogWindow.IsValid() )
+	{
+		DialogWindow.Pin()->RequestDestroyWindow();
+		DialogWindow.Reset();
 	}
 }
 
@@ -137,71 +294,577 @@ FReply SRigSpacePickerWidget::OnKeyDown( const FGeometry& MyGeometry, const FKey
 {
 	if (InKeyEvent.GetKey() == EKeys::Escape)
 	{
-		PickedKey = FRigElementKey();
+		if(DialogWindow.IsValid())
+		{
+			CloseDialog();
+		}
 		return FReply::Handled();
 	}
 	return SCompoundWidget::OnKeyDown(MyGeometry, InKeyEvent);
 }
 
-void SRigSpacePickerWidget::AddSpacePickerButton(const FSlateBrush* InBush, const FText& InTitle,
-	FOnClicked OnClickedDelegate)
+bool SRigSpacePickerWidget::SupportsKeyboardFocus() const
+{
+	return true;
+}
+
+void SRigSpacePickerWidget::Tick(const FGeometry& AllottedGeometry, const double InCurrentTime, const float InDeltaTime)
+{
+	SCompoundWidget::Tick(AllottedGeometry, InCurrentTime, InDeltaTime);
+
+	if(bRepopulateRequired)
+	{
+		UpdateActiveSpace();
+		RepopulateItemSpaces();
+		bRepopulateRequired = false;
+	}
+	else if(GetAdditionalSpacesDelegate.IsBound())
+	{
+		const TArray<FRigElementKey> CurrentAdditionalSpaces = GetAdditionalSpacesDelegate.Execute(Hierarchy, ControlKey);
+		if(CurrentAdditionalSpaces != AdditionalSpaces)
+		{
+			RepopulateItemSpaces();
+		}
+	}
+}
+
+const FRigElementKey& SRigSpacePickerWidget::GetActiveSpace() const
 {
-	ListBox->AddSlot()
+	return ActiveSpace;
+}
+
+TArray<FRigElementKey> SRigSpacePickerWidget::GetDefaultSpaces() const
+{
+	TArray<FRigElementKey> DefaultSpaces;
+	DefaultSpaces.Add(DefaultParentKey);
+	DefaultSpaces.Add(WorldSocketKey);
+	return DefaultSpaces;
+}
+
+TArray<FRigElementKey> SRigSpacePickerWidget::GetSpaceList(bool bIncludeDefaultSpaces) const
+{
+	if(bIncludeDefaultSpaces && bShowDefaultSpaces)
+	{
+		TArray<FRigElementKey> Spaces;
+		Spaces.Append(GetDefaultSpaces());
+		Spaces.Append(CurrentSpaceKeys);
+		return Spaces;
+	}
+	return CurrentSpaceKeys;
+}
+
+void SRigSpacePickerWidget::AddSpacePickerRow(
+	TSharedPtr<SVerticalBox> InListBox,
+	ESpacePickerType InType,
+	const FRigElementKey& InKey,
+	const FSlateBrush* InBush,
+	const FText& InTitle,
+    FOnClicked OnClickedDelegate)
+{
+	static const FSlateBrush* RoundedBoxBrush = FControlRigEditorStyle::Get().GetBrush(TEXT("ControlRig.SpacePicker.RoundedRect"));
+
+	TSharedPtr<SHorizontalBox> RowBox, ButtonBox;
+	InListBox->AddSlot()
 	.AutoHeight()
 	.VAlign(VAlign_Top)
 	.HAlign(HAlign_Fill)
-	.Padding(24.f, 0.0f, 0.0f, 2.0f)
+	.Padding(4.0, 0.0, 4.0, 0.0)
 	[
 		SNew( SButton )
 		.ButtonStyle(FAppStyle::Get(), "SimpleButton")
+		.ContentPadding(FMargin(0.0))
 		.OnClicked(OnClickedDelegate)
 		[
-			SNew(SHorizontalBox)
+			SAssignNew(RowBox, SHorizontalBox)
 			+ SHorizontalBox::Slot()
 			.AutoWidth()
 			.VAlign(VAlign_Center)
-			.HAlign(HAlign_Left)
+			.HAlign(HAlign_Fill)
 			.Padding(0)
 			[
-				SNew(SImage)
-				.Image(InBush)
+				SNew(SBorder)
+				.Padding(FMargin(5.0, 2.0, 5.0, 2.0))
+				.BorderImage(RoundedBoxBrush)
+				.BorderBackgroundColor(this, &SRigSpacePickerWidget::GetButtonColor, InType, InKey)
+				.Content()
+				[
+					SAssignNew(ButtonBox, SHorizontalBox)
+					+ SHorizontalBox::Slot()
+					.AutoWidth()
+					.VAlign(VAlign_Center)
+					.HAlign(HAlign_Left)
+					.Padding(0)
+					[
+						SNew(SImage)
+						.Image(InBush)
+					]
+
+					+ SHorizontalBox::Slot()
+					.AutoWidth()
+					.VAlign(VAlign_Center)
+					.HAlign(HAlign_Left)
+					.Padding(0)
+					[
+						SNew( STextBlock )
+						.Text( InTitle )
+						.Font( IDetailLayoutBuilder::GetDetailFont() )
+					]
+
+					+ SHorizontalBox::Slot()
+					.FillWidth(1.f)
+					[
+						SNew(SSpacer)
+					]
+				]
 			]
+		]
+	];
 
-			+ SHorizontalBox::Slot()
+	if(!IsDefaultSpace(InKey))
+	{
+		if(bAllowDelete || bAllowReorder)
+		{
+			RowBox->AddSlot()
+			.FillWidth(1.f)
+			[
+				SNew(SSpacer)
+			];
+		}
+		
+		if(bAllowReorder)
+		{
+			RowBox->AddSlot()
 			.AutoWidth()
 			.VAlign(VAlign_Center)
 			.HAlign(HAlign_Left)
 			.Padding(0)
 			[
-				SNew( STextBlock )
-				.Text( InTitle )
-				.Font( IDetailLayoutBuilder::GetDetailFont() )
-			]
+				SNew(SButton)
+				.ButtonStyle(FAppStyle::Get(), TEXT("SimpleButton"))
+				.ContentPadding(0)
+				.OnClicked(this, &SRigSpacePickerWidget::HandleSpaceMoveUp, InKey)
+				.IsEnabled(this, &SRigSpacePickerWidget::IsSpaceMoveUpEnabled, InKey)
+				.ToolTipText(LOCTEXT("MoveSpaceDown", "Move this space down in the list."))
+				[
+					SNew(SImage)
+					.Image(FEditorStyle::GetBrush("Icons.ChevronUp"))
+					.ColorAndOpacity(FSlateColor::UseForeground())
+				]
+			];
 
-			+ SHorizontalBox::Slot()
+			RowBox->AddSlot()
+			.AutoWidth()
+			.VAlign(VAlign_Center)
+			.HAlign(HAlign_Left)
+			.Padding(0)
 			[
-				SNew(SSpacer)
-			]
-		]
-	];
+				SNew(SButton)
+				.ButtonStyle(FAppStyle::Get(), TEXT("SimpleButton"))
+				.ContentPadding(0)
+				.OnClicked(this, &SRigSpacePickerWidget::HandleSpaceMoveDown, InKey)
+				.IsEnabled(this, &SRigSpacePickerWidget::IsSpaceMoveDownEnabled, InKey)
+				.ToolTipText(LOCTEXT("MoveSpaceUp", "Move this space up in the list."))
+				[
+					SNew(SImage)
+					.Image(FEditorStyle::GetBrush("Icons.ChevronDown"))
+					.ColorAndOpacity(FSlateColor::UseForeground())
+				]
+			];
+		}
+
+		if(bAllowDelete)
+		{
+			RowBox->AddSlot()
+			.AutoWidth()
+			.VAlign(VAlign_Center)
+			.HAlign(HAlign_Left)
+			.Padding(0)
+			[
+				PropertyCustomizationHelpers::MakeClearButton(FSimpleDelegate::CreateSP(this, &SRigSpacePickerWidget::HandleSpaceDelete, InKey), LOCTEXT("DeleteSpace", "Remove this space."), true)
+			];
+		}
+	}
 }
 
-FReply SRigSpacePickerWidget::HandleLocalSpaceClicked()
+FReply SRigSpacePickerWidget::HandleParentSpaceClicked()
 {
-	const FRigElementKey DefaultParent = Hierarchy->GetFirstParent(ControlKey);
-	return HandleElementSpaceClicked(DefaultParent);
+	return HandleElementSpaceClicked(DefaultParentKey);
 }
 
 FReply SRigSpacePickerWidget::HandleWorldSpaceClicked()
 {
-	return HandleElementSpaceClicked(Hierarchy->GetWorldSpaceSocketKey());
+	return HandleElementSpaceClicked(WorldSocketKey);
 }
 
 FReply SRigSpacePickerWidget::HandleElementSpaceClicked(FRigElementKey InKey)
 {
-	PickedKey = InKey;
-	CloseDialog();
+	ActiveSpaceChangedEvent.Broadcast(Hierarchy, ControlKey, InKey);
+
+	if(DialogWindow.IsValid())
+	{
+		CloseDialog();
+	}
+	
+	return FReply::Handled();
+}
+
+FReply SRigSpacePickerWidget::HandleSpaceMoveUp(FRigElementKey InKey)
+{
+	if(CurrentSpaceKeys.Num() > 1)
+	{
+		const int32 Index = CurrentSpaceKeys.Find(InKey);
+		if(CurrentSpaceKeys.IsValidIndex(Index))
+		{
+			if(Index > 0)
+			{
+				TArray<FRigElementKey> ChangedSpaceKeys = CurrentSpaceKeys;
+				ChangedSpaceKeys.Swap(Index, Index - 1);
+				SpaceListChangedEvent.Broadcast(Hierarchy, ControlKey, ChangedSpaceKeys);
+				return FReply::Handled();
+			}
+		}
+	}
+	return FReply::Unhandled();
+}
+
+FReply SRigSpacePickerWidget::HandleSpaceMoveDown(FRigElementKey InKey)
+{
+	if(CurrentSpaceKeys.Num() > 1)
+	{
+		const int32 Index = CurrentSpaceKeys.Find(InKey);
+		if(CurrentSpaceKeys.IsValidIndex(Index))
+		{
+			if(Index < CurrentSpaceKeys.Num() - 1)
+			{
+				TArray<FRigElementKey> ChangedSpaceKeys = CurrentSpaceKeys;
+				ChangedSpaceKeys.Swap(Index, Index + 1);
+				SpaceListChangedEvent.Broadcast(Hierarchy, ControlKey, ChangedSpaceKeys);
+				return FReply::Handled();
+			}
+		}
+	}
+	return FReply::Unhandled();
+}
+
+void SRigSpacePickerWidget::HandleSpaceDelete(FRigElementKey InKey)
+{
+	TArray<FRigElementKey> ChangedSpaceKeys = CurrentSpaceKeys;
+	if(ChangedSpaceKeys.Remove(InKey) > 0)
+	{
+		SpaceListChangedEvent.Broadcast(Hierarchy, ControlKey, ChangedSpaceKeys);
+	}
+}
+
+FReply SRigSpacePickerWidget::HandleAddElementClicked()
+{
+	FRigTreeDelegates TreeDelegates;
+	TreeDelegates.OnGetHierarchy = FOnGetRigTreeHierarchy::CreateSP(this, &SRigSpacePickerWidget::GetHierarchy);
+	TreeDelegates.OnMouseButtonClick = FOnRigTreeMouseButtonClick::CreateLambda([this](TSharedPtr<FRigTreeElement> InItem)
+	{
+		if(InItem.IsValid())
+		{
+			const FRigElementKey Key = InItem->Key;
+			if(!IsDefaultSpace(Key) && IsValidKey(Key))
+			{
+				TArray<FRigElementKey> ChangedSpaceKeys = CurrentSpaceKeys;
+				ChangedSpaceKeys.AddUnique(Key);
+				SpaceListChangedEvent.Broadcast(Hierarchy, ControlKey, ChangedSpaceKeys);
+			}
+		}
+
+		if(ContextMenu.IsValid())
+		{
+			ContextMenu.Pin()->Dismiss();
+			ContextMenu.Reset();
+		}
+	});
+
+	TSharedPtr<SSearchableRigHierarchyTreeView> SearchableTreeView = SNew(SSearchableRigHierarchyTreeView)
+	.RigTreeDelegates(TreeDelegates);
+	SearchableTreeView->GetTreeView()->RefreshTreeView(true);
+
+	// Create as context menu
+	TGuardValue<bool> AboutToShowMenu(bLaunchingContextMenu, true);
+	ContextMenu = FSlateApplication::Get().PushMenu(
+		AsShared(),
+		FWidgetPath(),
+		SearchableTreeView.ToSharedRef(),
+		FSlateApplication::Get().GetCursorPos(),
+		FPopupTransitionEffect(FPopupTransitionEffect::ContextMenu)
+	);
+
+	ContextMenu.Pin()->GetOnMenuDismissed().AddLambda([this](TSharedRef<IMenu> InMenu)
+	{
+		ContextMenu.Reset();
+
+		if(DialogWindow.IsValid())
+		{
+			DialogWindow.Pin()->BringToFront(true);
+
+			TSharedRef<SWidget> ThisRef = AsShared();
+			FSlateApplication::Get().ForEachUser([&ThisRef](FSlateUser& User) {
+				User.SetFocus(ThisRef, EFocusCause::SetDirectly);
+			});
+		}
+	});
+
 	return FReply::Handled();
 }
 
+bool SRigSpacePickerWidget::IsSpaceMoveUpEnabled(FRigElementKey InKey) const
+{
+	if(CurrentSpaceKeys.IsEmpty())
+	{
+		return false;
+	}
+	return CurrentSpaceKeys[0] != InKey;
+}
+
+bool SRigSpacePickerWidget::IsSpaceMoveDownEnabled(FRigElementKey InKey) const
+{
+	if(CurrentSpaceKeys.IsEmpty())
+	{
+		return false;
+	}
+	return CurrentSpaceKeys.Last( )!= InKey;
+}
+
+void SRigSpacePickerWidget::OnHierarchyModified(ERigHierarchyNotification InNotif, URigHierarchy* InHierarchy,
+                                                const FRigBaseElement* InElement)
+{
+	if(InElement == nullptr)
+	{
+		return;
+	}
+
+	if(InElement->GetKey() != ControlKey)
+	{
+		return;
+	}
+	
+	switch(InNotif)
+	{
+		case ERigHierarchyNotification::ParentChanged:
+		case ERigHierarchyNotification::ParentWeightsChanged:
+		case ERigHierarchyNotification::ControlSettingChanged:
+		{
+			bRepopulateRequired = true;
+			break;
+		}
+		default:
+		{
+			break;
+		}
+	}
+}
+
+FSlateColor SRigSpacePickerWidget::GetButtonColor(ESpacePickerType InType, FRigElementKey InKey) const
+{
+	static const FSlateColor ActiveColor = FControlRigEditorStyle::Get().SpacePickerSelectColor;
+	
+	switch(InType)
+	{
+		case ESpacePickerType_Parent:
+		{
+			// this is also true if the object has no parent
+			if(GetActiveSpace() == DefaultParentKey)
+			{
+				return ActiveColor;
+			}
+			break;
+		}
+		case ESpacePickerType_World:
+		{
+			if(GetActiveSpace() == WorldSocketKey && WorldSocketKey.IsValid())
+			{
+				return ActiveColor;
+			}
+			break;
+		}
+		case ESpacePickerType_Item:
+		default:
+		{
+			if(GetActiveSpace() == InKey && InKey.IsValid())
+			{
+				return ActiveColor;
+			}
+			break;
+		}
+	}
+	
+	return FStyleColors::Transparent;
+}
+
+TArray<FRigElementKey> SRigSpacePickerWidget::GetCurrentParents(URigHierarchy* InHierarchy,
+	const FRigElementKey& InControlKey) const
+{
+	if(!InControlKey.IsValid() || InHierarchy == nullptr)
+	{
+		return TArray<FRigElementKey>();
+	}
+
+	check(InControlKey == ControlKey);
+	return InHierarchy->GetParents(ControlKey);
+}
+
+void SRigSpacePickerWidget::RepopulateItemSpaces()
+{
+	if(!ItemSpacesListBox.IsValid())
+	{
+		return;
+	}
+	
+	TArray<FRigElementKey> FavoriteKeys, SpacesFromDelegate;
+
+	// gather the keys coming from the input customization
+	if(Customization && bShowFavoriteSpaces)
+	{
+		for(const FRigElementKey& Key : Customization->AvailableSpaces)
+		{
+			if(IsDefaultSpace(Key) || !IsValidKey(Key))
+			{
+				continue;
+			}
+			FavoriteKeys.AddUnique(Key);
+		}
+	}
+
+	// check if the customization is different from the base one in the asset
+	if(const FRigControlElement* ControlElement = Hierarchy->Find<FRigControlElement>(ControlKey))
+	{
+		if(Customization != &ControlElement->Settings.Customization)
+		{
+			for(const FRigElementKey& Key : ControlElement->Settings.Customization.AvailableSpaces)
+			{
+				if(IsDefaultSpace(Key) || !IsValidKey(Key))
+				{
+					continue;
+				}
+
+				if(Customization)
+				{
+					if(Customization->AvailableSpaces.Contains(Key))
+					{
+						continue;
+					}
+					if(Customization->RemovedSpaces.Contains(Key))
+					{
+						continue;
+					}
+				}
+				FavoriteKeys.AddUnique(Key);
+			}
+		}
+	}
+
+	// now gather all of the spaces using the get additional spaces delegate
+	if(GetAdditionalSpacesDelegate.IsBound() && bShowAdditionalSpaces)
+	{
+		AdditionalSpaces = GetAdditionalSpacesDelegate.Execute(Hierarchy, ControlKey);
+		for(const FRigElementKey& Key : AdditionalSpaces)
+		{
+			if(IsDefaultSpace(Key)  || !IsValidKey(Key))
+			{
+				continue;
+			}
+			SpacesFromDelegate.AddUnique(Key);
+		}
+	}
+
+	/*
+	struct FKeySortPredicate
+	{
+		bool operator()(const FRigElementKey& A, const FRigElementKey& B) const
+		{
+			static TMap<ERigElementType, int32> TypeOrder;
+			if(TypeOrder.IsEmpty())
+			{
+				TypeOrder.Add(ERigElementType::Control, 0);
+				TypeOrder.Add(ERigElementType::Socket, 1);
+				TypeOrder.Add(ERigElementType::Null, 2);
+				TypeOrder.Add(ERigElementType::Bone, 3);
+				TypeOrder.Add(ERigElementType::RigidBody, 4);
+			}
+
+			const int32 TypeIndexA = TypeOrder.FindChecked(A.Type);
+			const int32 TypeIndexB = TypeOrder.FindChecked(B.Type);
+			if(TypeIndexA != TypeIndexB)
+			{
+				return TypeIndexA < TypeIndexB;
+			}
+
+			return A.Name.Compare(B.Name) < 0; 
+		}
+	};
+	SpacesFromDelegate.Sort(FKeySortPredicate());
+	*/
+
+	TArray<FRigElementKey> Keys = FavoriteKeys;
+	for(const FRigElementKey& Key : SpacesFromDelegate)
+	{
+		Keys.AddUnique(Key);
+	}
+
+	if(Keys == CurrentSpaceKeys)
+	{
+		return;
+	}
+
+	ClearListBox(ItemSpacesListBox);
+
+	for(const FRigElementKey& Key : Keys)
+	{
+		AddSpacePickerRow(
+			ItemSpacesListBox,
+			ESpacePickerType_Item,
+			Key,
+			SRigHierarchyItem::GetBrushForElementType(Hierarchy, Key),
+			FText::FromName(Key.Name),
+			FOnClicked::CreateSP(this, &SRigSpacePickerWidget::HandleElementSpaceClicked, Key)
+		);
+	}
+
+	CurrentSpaceKeys = Keys;
+}
+
+void SRigSpacePickerWidget::ClearListBox(TSharedPtr<SVerticalBox> InListBox)
+{
+	InListBox->ClearChildren();
+}
+
+void SRigSpacePickerWidget::UpdateActiveSpace()
+{
+	ActiveSpace = FRigElementKey();
+	
+	const TArray<FRigElementWeight> ParentWeights = Hierarchy->GetParentWeightArray(ControlKey);
+	if(ParentWeights.Num() > 0)
+	{
+		const TArray<FRigElementKey> ParentKeys = Hierarchy->GetParents(ControlKey);
+		check(ParentKeys.Num() == ParentWeights.Num());
+		for(int32 ParentIndex=0;ParentIndex<ParentKeys.Num();ParentIndex++)
+		{
+			if(ParentWeights[ParentIndex].IsAlmostZero())
+			{
+				continue;
+			}
+			ActiveSpace = ParentKeys[ParentIndex];
+			return;
+		}
+	}
+}
+
+bool SRigSpacePickerWidget::IsValidKey(const FRigElementKey& InKey) const
+{
+	return InKey.IsValid() && Hierarchy->Contains(InKey);
+}
+
+bool SRigSpacePickerWidget::IsDefaultSpace(const FRigElementKey& InKey) const
+{
+	if(bShowDefaultSpaces)
+	{
+		return InKey == WorldSocketKey || InKey == DefaultParentKey;
+	}
+	return false;
+}
+
 #undef LOCTEXT_NAMESPACE
diff --git a/Engine/Plugins/Animation/ControlRig/Source/ControlRigEditor/Private/Editor/SRigSpacePickerWidget.h b/Engine/Plugins/Animation/ControlRig/Source/ControlRigEditor/Private/Editor/SRigSpacePickerWidget.h
--- a/Engine/Plugins/Animation/ControlRig/Source/ControlRigEditor/Private/Editor/SRigSpacePickerWidget.h
+++ b/Engine/Plugins/Animation/ControlRig/Source/ControlRigEditor/Private/Editor/SRigSpacePickerWidget.h
@@ -3,61 +3,139 @@
 #pragma once
 
 #include "CoreMinimal.h"
+#include "ControlRig.h"
 #include "Rigs/RigHierarchy.h"
 #include "SRigHierarchy.h"
 #include "Widgets/Layout/SBox.h"
 #include "Rigs/RigSpaceHierarchy.h"
 
+DECLARE_EVENT_ThreeParams(SRigSpacePickerWidget, FRigSpacePickerActiveSpaceChanged, URigHierarchy*, const FRigElementKey&, const FRigElementKey&);
+DECLARE_EVENT_ThreeParams(SRigSpacePickerWidget, FRigSpacePickerSpaceListChanged, URigHierarchy*, const FRigElementKey&, const TArray<FRigElementKey>&);
+DECLARE_DELEGATE_RetVal_TwoParams(TArray<FRigElementKey>, FRigSpacePickerGetAdditionalSpaces, URigHierarchy*, const FRigElementKey&);
 
 /** Widget allowing picking of a space source for space switching */
 class SRigSpacePickerWidget : public SCompoundWidget
 {
 public:
 
-	struct FResult
-	{
-		FResult()
-			: Reply(FReply::Unhandled())
-			// , Actor(nullptr)
-			// , SceneComponent(nullptr)
-			, Key()
-		{
-		}
-		
-		FReply Reply;
-		//AActor* Actor;
-		//USceneComponent* SceneComponent;
-		FRigElementKey Key;
-	};
-
-	SLATE_BEGIN_ARGS(SRigSpacePickerWidget) {}
+	SLATE_BEGIN_ARGS(SRigSpacePickerWidget)
+		: _Hierarchy(nullptr)
+		, _Control()
+		, _Customization(nullptr)
+		, _ShowDefaultSpaces(true)
+		, _ShowFavoriteSpaces(true)
+		, _ShowAdditionalSpaces(true)
+		, _AllowReorder(false)
+		, _AllowDelete(false)
+		, _AllowAdd(false)
+		, _ShowBakeButton(false)
+		, _Title()
+		, _BackgroundBrush(FEditorStyle::GetBrush("Menu.Background"))
+		{}
 		SLATE_ARGUMENT(URigHierarchy*, Hierarchy)
-		SLATE_ARGUMENT(FRigElementKey, SelectedControl)
+		SLATE_ARGUMENT(FRigElementKey, Control)
+		SLATE_ARGUMENT(FRigControlElementCustomization*, Customization)
+		SLATE_ARGUMENT(bool, ShowDefaultSpaces)
+		SLATE_ARGUMENT(bool, ShowFavoriteSpaces)
+		SLATE_ARGUMENT(bool, ShowAdditionalSpaces)
+		SLATE_ARGUMENT(bool, AllowReorder)
+		SLATE_ARGUMENT(bool, AllowDelete)
+		SLATE_ARGUMENT(bool, AllowAdd)
+		SLATE_ARGUMENT(bool, ShowBakeButton)
+		SLATE_ARGUMENT(FText, Title)
+		SLATE_ARGUMENT(const FSlateBrush*, BackgroundBrush)
+		SLATE_ARGUMENT(FRigSpacePickerGetAdditionalSpaces, GetAdditionalSpacesDelegate)
+		SLATE_EVENT( FOnClicked, OnBakeButtonClicked )
 	SLATE_END_ARGS()
 
 	void Construct(const FArguments& InArgs);
-	~SRigSpacePickerWidget();
+	virtual ~SRigSpacePickerWidget() override;
+
+	void SetControl(URigHierarchy* InHierarchy, const FRigElementKey& InControl, FRigControlElementCustomization* InCustomization = nullptr);
 
-	FResult InvokeDialog();
+	FReply OpenDialog(bool bModal = true);
 	void CloseDialog();
+	
 	virtual FReply OnKeyDown( const FGeometry& MyGeometry, const FKeyEvent& InKeyEvent ) override;
+	virtual bool SupportsKeyboardFocus() const override;
+	virtual void Tick(const FGeometry& AllottedGeometry, const double InCurrentTime, const float InDeltaTime) override;
 
-	virtual bool SupportsKeyboardFocus() const override
-	{
-		return true;
-	}
+	const FRigElementKey& GetActiveSpace() const;
+	TArray<FRigElementKey> GetDefaultSpaces() const;
+	TArray<FRigElementKey> GetSpaceList(bool bIncludeDefaultSpaces = false) const;
+	FRigSpacePickerActiveSpaceChanged& OnActiveSpaceChanged() { return ActiveSpaceChangedEvent; }
+	FRigSpacePickerSpaceListChanged& OnSpaceListChanged() { return SpaceListChangedEvent; }
 
 private:
 
-	void AddSpacePickerButton(const FSlateBrush* InBush, const FText& InTitle, FOnClicked OnClickedDelegate);
+	enum ESpacePickerType
+	{
+		ESpacePickerType_Parent,
+		ESpacePickerType_World,
+		ESpacePickerType_Item
+	};
+
+	void AddSpacePickerRow(
+		TSharedPtr<SVerticalBox> InListBox,
+		ESpacePickerType InType,
+		const FRigElementKey& InKey,
+		const FSlateBrush* InBush,
+		const FText& InTitle,
+		FOnClicked OnClickedDelegate);
 
-	FReply HandleLocalSpaceClicked();
+	void RepopulateItemSpaces();
+	void ClearListBox(TSharedPtr<SVerticalBox> InListBox);
+	void UpdateActiveSpace();
+	bool IsValidKey(const FRigElementKey& InKey) const;
+	bool IsDefaultSpace(const FRigElementKey& InKey) const;
+
+	FReply HandleParentSpaceClicked();
 	FReply HandleWorldSpaceClicked();
 	FReply HandleElementSpaceClicked(FRigElementKey InKey);
+	FReply HandleSpaceMoveUp(FRigElementKey InKey);
+	FReply HandleSpaceMoveDown(FRigElementKey InKey);
+	void HandleSpaceDelete(FRigElementKey InKey);
+	FReply HandleAddElementClicked();
+	bool IsSpaceMoveUpEnabled(FRigElementKey InKey) const;
+	bool IsSpaceMoveDownEnabled(FRigElementKey InKey) const;
+	const URigHierarchy* GetHierarchy() const { return Hierarchy; }
+
+	void OnHierarchyModified(ERigHierarchyNotification InNotif, URigHierarchy* InHierarchy, const FRigBaseElement* InElement);
+
+	FSlateColor GetButtonColor(ESpacePickerType InType, FRigElementKey InKey) const;
+	TArray<FRigElementKey> GetCurrentParents(URigHierarchy* InHierarchy, const FRigElementKey& InControlKey) const;
+	
+	FRigSpacePickerActiveSpaceChanged ActiveSpaceChangedEvent;
+	FRigSpacePickerSpaceListChanged SpaceListChangedEvent;
 
 	URigHierarchy* Hierarchy;
 	FRigElementKey ControlKey;
-	TSharedPtr<SVerticalBox> ListBox;
-	TWeakPtr<SWindow> PickerWindow;
-	FRigElementKey PickedKey;
-};
\ No newline at end of file
+	FRigElementKey DefaultParentKey;
+	FRigElementKey WorldSocketKey;
+	TArray<FRigElementKey> CurrentSpaceKeys;
+	bool bRepopulateRequired;
+
+	FRigControlElementCustomization* Customization;
+	bool bShowDefaultSpaces;
+	bool bShowFavoriteSpaces;
+	bool bShowAdditionalSpaces;
+	bool bAllowReorder;
+	bool bAllowDelete;
+	bool bAllowAdd;
+	bool bShowBakeButton;
+	bool bLaunchingContextMenu;
+
+	FRigSpacePickerGetAdditionalSpaces GetAdditionalSpacesDelegate; 
+	TArray<FRigElementKey> AdditionalSpaces;
+
+	TSharedPtr<SVerticalBox> TopLevelListBox;
+	TSharedPtr<SVerticalBox> ItemSpacesListBox;
+	TSharedPtr<SHorizontalBox> BottomButtonsListBox;
+	TWeakPtr<SWindow> DialogWindow;
+	TWeakPtr<IMenu> ContextMenu;
+	FRigElementKey ActiveSpace;
+	FDelegateHandle HierarchyModifiedHandle;
+	FDelegateHandle ActiveSpaceChangedWindowHandle;
+
+	static FRigElementKey InValidKey;
+};
