Control Rig: Follow up to Space Picker Widget

* support for multi selection
* shortcut to Ctrl + Q
* active spaces chagnes / space list changes are now slate events
* support within the edit mode / level editor viewport

diff --git a/Engine/Plugins/Animation/ControlRig/Source/ControlRig/Private/Rigs/RigHierarchy.cpp b/Engine/Plugins/Animation/ControlRig/Source/ControlRig/Private/Rigs/RigHierarchy.cpp
--- a/Engine/Plugins/Animation/ControlRig/Source/ControlRig/Private/Rigs/RigHierarchy.cpp
+++ b/Engine/Plugins/Animation/ControlRig/Source/ControlRig/Private/Rigs/RigHierarchy.cpp
@@ -752,7 +752,14 @@ bool URigHierarchy::IsNameAvailable(const FString& InPotentialNewName, ERigEleme
 		return false;
 	}
 
-	if (GetIndex(FRigElementKey(*InPotentialNewName, InType)) != INDEX_NONE)
+	// check for fixed keywords
+	const FRigElementKey PotentialKey(*InPotentialNewName, InType);
+	if(PotentialKey == URigHierarchy::GetDefaultParentSocketKey())
+	{
+		return false;
+	}
+
+	if (GetIndex(PotentialKey) != INDEX_NONE)
 	{
 		if (OutErrorMessage)
 		{
@@ -1403,6 +1410,22 @@ bool URigHierarchy::SetParentWeightArray(FRigBaseElement* InChild,  const TArray
 
 bool URigHierarchy::SwitchToParent(FRigElementKey InChild, FRigElementKey InParent, bool bInitial, bool bAffectChildren)
 {
+	if(InParent == GetWorldSpaceSocketKey())
+	{
+		InParent = GetOrAddWorldSpaceSocket();
+	}
+	else if(InParent == GetDefaultParentSocketKey())
+	{
+		const FRigElementKey FirstParent = GetFirstParent(InChild);
+		if(FirstParent == GetWorldSpaceSocketKey())
+		{
+			InParent = FRigElementKey();
+		}
+		else
+		{
+			InParent = FirstParent;
+		}
+	}
 	return SwitchToParent(Find(InChild), Find(InParent), bInitial, bAffectChildren);
 }
 
@@ -1447,28 +1470,25 @@ bool URigHierarchy::SwitchToParent(FRigBaseElement* InChild, int32 InParentIndex
 
 bool URigHierarchy::SwitchToDefaultParent(FRigElementKey InChild, bool bInitial, bool bAffectChildren)
 {
-	return SwitchToDefaultParent(Find(InChild), bInitial, bAffectChildren);
+	return SwitchToParent(InChild, GetDefaultParentSocketKey(), bInitial, bAffectChildren);
 }
 
 bool URigHierarchy::SwitchToDefaultParent(FRigBaseElement* InChild, bool bInitial, bool bAffectChildren)
 {
 	// we assume that the first stored parent is the default parent
-	return SwitchToParent(InChild, 0, bInitial, bAffectChildren);
+	check(InChild);
+	return SwitchToParent(InChild->GetKey(), GetDefaultParentSocketKey(), bInitial, bAffectChildren);
 }
 
 bool URigHierarchy::SwitchToWorldSpace(FRigElementKey InChild, bool bInitial, bool bAffectChildren)
 {
-	return SwitchToWorldSpace(Find(InChild), bInitial, bAffectChildren);
+	return SwitchToParent(InChild, GetWorldSpaceSocketKey(), bInitial, bAffectChildren);
 }
 
 bool URigHierarchy::SwitchToWorldSpace(FRigBaseElement* InChild, bool bInitial, bool bAffectChildren)
 {
-	FRigElementKey WorldSocket = GetOrAddWorldSpaceSocket();
-	if(FRigBaseElement* Parent = Find(WorldSocket))
-	{
-		return SwitchToParent(InChild, Parent, bInitial, bAffectChildren);
-	}
-	return false;
+	check(InChild);
+	return SwitchToParent(InChild->GetKey(), GetWorldSpaceSocketKey(), bInitial, bAffectChildren);
 }
 
 FRigElementKey URigHierarchy::GetOrAddWorldSpaceSocket()
@@ -1493,7 +1513,13 @@ FRigElementKey URigHierarchy::GetOrAddWorldSpaceSocket()
 	return FRigElementKey();
 }
 
-FRigElementKey URigHierarchy::GetWorldSpaceSocketKey() const
+FRigElementKey URigHierarchy::GetDefaultParentSocketKey()
+{
+	static const FName DefaultParentSocketName = TEXT("DefaultParent");
+	return FRigElementKey(DefaultParentSocketName, ERigElementType::Socket); 
+}
+
+FRigElementKey URigHierarchy::GetWorldSpaceSocketKey()
 {
 	static const FName WorldSpaceSocketName = TEXT("WorldSpace");
 	return FRigElementKey(WorldSpaceSocketName, ERigElementType::Socket); 
diff --git a/Engine/Plugins/Animation/ControlRig/Source/ControlRig/Public/Rigs/RigHierarchy.h b/Engine/Plugins/Animation/ControlRig/Source/ControlRig/Public/Rigs/RigHierarchy.h
--- a/Engine/Plugins/Animation/ControlRig/Source/ControlRig/Public/Rigs/RigHierarchy.h
+++ b/Engine/Plugins/Animation/ControlRig/Source/ControlRig/Public/Rigs/RigHierarchy.h
@@ -1810,7 +1810,9 @@ public:
 	 * Adds the world space socket or returns it
 	 */
 	FRigElementKey GetOrAddWorldSpaceSocket();
-	FRigElementKey GetWorldSpaceSocketKey() const;
+	
+	static FRigElementKey GetDefaultParentSocketKey();
+	static FRigElementKey GetWorldSpaceSocketKey();
 
 	/**
 	 * Returns true if an element is parented to another element
diff --git a/Engine/Plugins/Animation/ControlRig/Source/ControlRigEditor/Private/ControlRigHierarchyCommands.cpp b/Engine/Plugins/Animation/ControlRig/Source/ControlRigEditor/Private/ControlRigHierarchyCommands.cpp
--- a/Engine/Plugins/Animation/ControlRig/Source/ControlRigEditor/Private/ControlRigHierarchyCommands.cpp
+++ b/Engine/Plugins/Animation/ControlRig/Source/ControlRigEditor/Private/ControlRigHierarchyCommands.cpp
@@ -37,7 +37,7 @@ void FControlRigHierarchyCommands::RegisterCommands()
 	UI_COMMAND(ShowSockets, "Show Sockets", "Whether to show or hide sockets", EUserInterfaceActionType::ToggleButton, FInputChord());
 	UI_COMMAND(ShowDynamicHierarchy, "Dynamic Hierarchy", "Whether to show the static or dynamic hierarchy", EUserInterfaceActionType::ToggleButton, FInputChord());
 	UI_COMMAND(ToggleGizmoTransformEdit, "Toggle Gizmo Transform Edit", "Toggle Editing Selected Control's Gizmo Transform", EUserInterfaceActionType::ToggleButton, FInputChord(EKeys::Period, EModifierKey::Control)); 
-	UI_COMMAND(TestSpaceSwitching, "Test Space Switching", "Try out space switching on the control", EUserInterfaceActionType::Button, FInputChord(EKeys::H, EModifierKey::Control)); 
+	UI_COMMAND(TestSpaceSwitching, "Test Space Switching", "Try out space switching on the control", EUserInterfaceActionType::Button, FInputChord(EKeys::Q, EModifierKey::Control)); 
 }
 
 #undef LOCTEXT_NAMESPACE
diff --git a/Engine/Plugins/Animation/ControlRig/Source/ControlRigEditor/Private/EditMode/ControlRigEditMode.cpp b/Engine/Plugins/Animation/ControlRig/Source/ControlRigEditor/Private/EditMode/ControlRigEditMode.cpp
--- a/Engine/Plugins/Animation/ControlRig/Source/ControlRigEditor/Private/EditMode/ControlRigEditMode.cpp
+++ b/Engine/Plugins/Animation/ControlRig/Source/ControlRigEditor/Private/EditMode/ControlRigEditMode.cpp
@@ -55,7 +55,7 @@
 #include "ToolMenus.h"
 #include "Sequencer/MovieSceneControlRigParameterTrack.h"
 #include "Sequencer/MovieSceneControlRigParameterSection.h"
-
+#include "SRigSpacePickerWidget.h"
 
 void UControlRigEditModeDelegateHelper::OnPoseInitialized()
 {
@@ -1927,6 +1927,10 @@ void FControlRigEditMode::BindCommands()
 	CommandBindings->MapAction(
 		Commands.ToggleGizmoTransformEdit,
 		FExecuteAction::CreateRaw(this, &FControlRigEditMode::ToggleGizmoTransformEdit));
+
+	CommandBindings->MapAction(
+		Commands.OpenSpacePickerWidget,
+		FExecuteAction::CreateRaw(this, &FControlRigEditMode::OpenSpacePickerWidget));
 }
 
 bool FControlRigEditMode::IsControlSelected() const
@@ -2089,6 +2093,104 @@ void FControlRigEditMode::ToggleGizmoTransformEdit()
 	}
 }
 
+void FControlRigEditMode::OpenSpacePickerWidget()
+{
+	UControlRig* RuntimeRig = GetControlRig(false);
+	if (RuntimeRig == nullptr)
+	{
+		return;
+	}
+
+	URigHierarchy* Hierarchy = RuntimeRig->GetHierarchy();
+	TArray<FRigElementKey> SelectedControls = Hierarchy->GetSelectedKeys(ERigElementType::Control);
+
+	TSharedRef<SRigSpacePickerWidget> PickerWidget =
+	SNew(SRigSpacePickerWidget)
+	.Hierarchy(Hierarchy)
+	.Controls(SelectedControls)
+	.Title(LOCTEXT("PickSpace", "Pick Space"))
+	.AllowDelete(!IsInLevelEditor())
+	.AllowReorder(!IsInLevelEditor())
+	.AllowAdd(!IsInLevelEditor())
+	.GetControlCustomization_Lambda([this, RuntimeRig](URigHierarchy*, const FRigElementKey& InControlKey)
+	{
+		return RuntimeRig->GetControlCustomization(InControlKey);
+	})
+	.OnActiveSpaceChanged_Lambda([this, SelectedControls](URigHierarchy* InHierarchy, const FRigElementKey& InControlKey, const FRigElementKey& InSpaceKey)
+	{
+		check(SelectedControls.Contains(InControlKey));
+		
+		const FTransform Transform = InHierarchy->GetGlobalTransform(InControlKey);
+		InHierarchy->SwitchToParent(InControlKey, InSpaceKey);
+		InHierarchy->SetGlobalTransform(InControlKey, Transform);
+		
+	})
+	.OnSpaceListChanged_Lambda([this, SelectedControls, RuntimeRig](URigHierarchy* InHierarchy, const FRigElementKey& InControlKey, const TArray<FRigElementKey>& InSpaceList)
+	{
+		check(SelectedControls.Contains(InControlKey));
+
+		// check if we are in the control rig editor or in the level
+		if(!IsInLevelEditor())
+		{
+			if (UControlRigBlueprint* Blueprint = Cast<UControlRigBlueprint>(RuntimeRig->GetClass()->ClassGeneratedBy))
+			{
+				if(URigHierarchy* Hierarchy = Blueprint->Hierarchy)
+				{
+					// update the settings in the control element
+					if(FRigControlElement* ControlElement = Hierarchy->Find<FRigControlElement>(InControlKey))
+					{
+						Blueprint->Modify();
+						FScopedTransaction Transaction(LOCTEXT("ControlChangeAvailableSpaces", "Edit Available Spaces"));
+
+						ControlElement->Settings.Customization.AvailableSpaces = InSpaceList;
+						Hierarchy->Notify(ERigHierarchyNotification::ControlSettingChanged, ControlElement);
+					}
+
+					// also update the debugged instance
+					if(Hierarchy != InHierarchy)
+					{
+						if(FRigControlElement* ControlElement = InHierarchy->Find<FRigControlElement>(InControlKey))
+						{
+							ControlElement->Settings.Customization.AvailableSpaces = InSpaceList;
+							InHierarchy->Notify(ERigHierarchyNotification::ControlSettingChanged, ControlElement);
+						}
+					}
+				}
+			}
+		}
+		else
+		{
+			// update the settings in the control element
+			if(FRigControlElement* ControlElement = InHierarchy->Find<FRigControlElement>(InControlKey))
+			{
+				FScopedTransaction Transaction(LOCTEXT("ControlChangeAvailableSpaces", "Edit Available Spaces"));
+
+				InHierarchy->Modify();
+
+				FRigControlElementCustomization ControlCustomization = *RuntimeRig->GetControlCustomization(InControlKey);	
+				ControlCustomization.AvailableSpaces = InSpaceList;
+				ControlCustomization.RemovedSpaces.Reset();
+
+				// remember  the elements which are in the asset's available list but removed by the user
+				for(const FRigElementKey& AvailableSpace : ControlElement->Settings.Customization.AvailableSpaces)
+				{
+					if(!ControlCustomization.AvailableSpaces.Contains(AvailableSpace))
+					{
+						ControlCustomization.RemovedSpaces.Add(AvailableSpace);
+					}
+				}
+
+				RuntimeRig->SetControlCustomization(InControlKey, ControlCustomization);
+				InHierarchy->Notify(ERigHierarchyNotification::ControlSettingChanged, ControlElement);
+			}
+		}
+
+	});
+	// todo: implement GetAdditionalSpacesDelegate to pull spaces from sequencer
+
+	PickerWidget->OpenDialog(false);
+}
+
 FText FControlRigEditMode::GetToggleGizmoTransformEditHotKey() const
 {
 	const FControlRigEditModeCommands& Commands = FControlRigEditModeCommands::Get();
diff --git a/Engine/Plugins/Animation/ControlRig/Source/ControlRigEditor/Private/EditMode/ControlRigEditMode.h b/Engine/Plugins/Animation/ControlRig/Source/ControlRigEditor/Private/EditMode/ControlRigEditMode.h
--- a/Engine/Plugins/Animation/ControlRig/Source/ControlRigEditor/Private/EditMode/ControlRigEditMode.h
+++ b/Engine/Plugins/Animation/ControlRig/Source/ControlRigEditor/Private/EditMode/ControlRigEditMode.h
@@ -188,6 +188,9 @@ public:
 	/** Frame a list of provided items*/
    	void FrameItems(const TArray<FRigElementKey>& InItems);
 
+	/** Opens up the space picker widget */
+	void OpenSpacePickerWidget();
+
 private:
 	
 	/** Whether or not we should Frame Selection or not*/
diff --git a/Engine/Plugins/Animation/ControlRig/Source/ControlRigEditor/Private/EditMode/ControlRigEditModeCommands.cpp b/Engine/Plugins/Animation/ControlRig/Source/ControlRigEditor/Private/EditMode/ControlRigEditModeCommands.cpp
--- a/Engine/Plugins/Animation/ControlRig/Source/ControlRigEditor/Private/EditMode/ControlRigEditModeCommands.cpp
+++ b/Engine/Plugins/Animation/ControlRig/Source/ControlRigEditor/Private/EditMode/ControlRigEditModeCommands.cpp
@@ -17,6 +17,7 @@ void FControlRigEditModeCommands::RegisterCommands()
 	UI_COMMAND(ResetGizmoSize, "Reset Gizmo Size", "Resize Gizmo Size To Default", EUserInterfaceActionType::Button, FInputChord(EKeys::Equals));
 
 	UI_COMMAND(ToggleGizmoTransformEdit, "Toggle Gizmo Transform Edit", "Toggle Editing Selected Control's Gizmo Transform", EUserInterfaceActionType::ToggleButton, FInputChord(EKeys::Period, EModifierKey::Control)); 
+	UI_COMMAND(OpenSpacePickerWidget, "Open the Space Picker", "Allows space switching on the control", EUserInterfaceActionType::Button, FInputChord(EKeys::Q, EModifierKey::Control)); 
 }
 
 #undef LOCTEXT_NAMESPACE
diff --git a/Engine/Plugins/Animation/ControlRig/Source/ControlRigEditor/Private/EditMode/ControlRigEditModeCommands.h b/Engine/Plugins/Animation/ControlRig/Source/ControlRigEditor/Private/EditMode/ControlRigEditModeCommands.h
--- a/Engine/Plugins/Animation/ControlRig/Source/ControlRigEditor/Private/EditMode/ControlRigEditModeCommands.h
+++ b/Engine/Plugins/Animation/ControlRig/Source/ControlRigEditor/Private/EditMode/ControlRigEditModeCommands.h
@@ -47,6 +47,9 @@ public:
 	/** Toggle Gizmo Transform Edit*/
 	TSharedPtr< FUICommandInfo > ToggleGizmoTransformEdit;
 
+	/** Opens up the space picker widget */
+	TSharedPtr< FUICommandInfo > OpenSpacePickerWidget;
+
 	/**
 	 * Initialize commands
 	 */
diff --git a/Engine/Plugins/Animation/ControlRig/Source/ControlRigEditor/Private/EditMode/SControlRigEditModeTools.cpp b/Engine/Plugins/Animation/ControlRig/Source/ControlRigEditor/Private/EditMode/SControlRigEditModeTools.cpp
--- a/Engine/Plugins/Animation/ControlRig/Source/ControlRigEditor/Private/EditMode/SControlRigEditModeTools.cpp
+++ b/Engine/Plugins/Animation/ControlRig/Source/ControlRigEditor/Private/EditMode/SControlRigEditModeTools.cpp
@@ -322,6 +322,9 @@ void SControlRigEditModeTools::Construct(const FArguments& InArgs, FControlRigEd
 					.AllowReorder(true)
 					.AllowAdd(true)
 					.ShowBakeButton(true)
+					.GetControlCustomization(this, &SControlRigEditModeTools::HandleGetControlElementCustomization)
+					.OnActiveSpaceChanged(this, &SControlRigEditModeTools::HandleActiveSpaceChanged)
+					.OnSpaceListChanged(this, &SControlRigEditModeTools::HandleSpaceListChanged)
 					// todo: implement GetAdditionalSpacesDelegate to pull spaces from sequencer
 				]
 			]
@@ -344,8 +347,6 @@ void SControlRigEditModeTools::Construct(const FArguments& InArgs, FControlRigEd
 	];
 
 	HierarchyTreeView->RefreshTreeView(true);
-	SpacePickerWidget->OnActiveSpaceChanged().AddSP(this, &SControlRigEditModeTools::HandleActiveSpaceChanged);
-	SpacePickerWidget->OnSpaceListChanged().AddSP(this, &SControlRigEditModeTools::HandleSpaceListChanged);
 }
 
 void SControlRigEditModeTools::SetDetailsObjects(const TArray<TWeakObjectPtr<>>& InObjects)
@@ -601,13 +602,19 @@ void SControlRigEditModeTools::OnRigElementSelected(UControlRig* Subject, FRigCo
 
 	if (UControlRig* ControlRig = SequencerRig.Get())
 	{
-		// get the customization for this control
-		const FRigControlElementCustomization* ExistingCustomization = ControlRig->GetControlCustomization(Key);
-		check(ExistingCustomization);
+		// get the selected controls
+		TArray<FRigElementKey> SelectedControls = ControlRig->GetHierarchy()->GetSelectedKeys(ERigElementType::Control);
+		SpacePickerWidget->SetControls(ControlRig->GetHierarchy(), SelectedControls);
+	}
+}
 
-		ControlCustomization = *ExistingCustomization; 
-		SpacePickerWidget->SetControl(ControlRig->GetHierarchy(), Key, &ControlCustomization);
+const FRigControlElementCustomization* SControlRigEditModeTools::HandleGetControlElementCustomization(URigHierarchy* InHierarchy, const FRigElementKey& InControlKey)
+{
+	if (UControlRig* ControlRig = SequencerRig.Get())
+	{
+		return ControlRig->GetControlCustomization(InControlKey);
 	}
+	return nullptr;
 }
 
 void SControlRigEditModeTools::HandleActiveSpaceChanged(URigHierarchy* InHierarchy, const FRigElementKey& InControlKey,
@@ -615,14 +622,7 @@ void SControlRigEditModeTools::HandleActiveSpaceChanged(URigHierarchy* InHierarc
 {
 	// todo: when the user picks a new space
 	// for now we'll just do it directly in the hierarchy
-	if(InSpaceKey == InHierarchy->GetWorldSpaceSocketKey())
-	{
-		InHierarchy->SwitchToWorldSpace(InControlKey);
-	}
-	else
-	{
-		InHierarchy->SwitchToParent(InControlKey, InSpaceKey);
-	}
+	InHierarchy->SwitchToParent(InControlKey, InSpaceKey);
 }
 
 void SControlRigEditModeTools::HandleSpaceListChanged(URigHierarchy* InHierarchy, const FRigElementKey& InControlKey,
@@ -632,9 +632,7 @@ void SControlRigEditModeTools::HandleSpaceListChanged(URigHierarchy* InHierarchy
 	{
 		if(const FRigControlElement* ControlElement = InHierarchy->Find<FRigControlElement>(InControlKey))
 		{
-			const FRigControlElementCustomization* ExistingCustomization = ControlRig->GetControlCustomization(InControlKey);
-			ControlCustomization = *ExistingCustomization;	
-
+			FRigControlElementCustomization ControlCustomization = *ControlRig->GetControlCustomization(InControlKey);	
 			ControlCustomization.AvailableSpaces = InSpaceList;
 			ControlCustomization.RemovedSpaces.Reset();
 
diff --git a/Engine/Plugins/Animation/ControlRig/Source/ControlRigEditor/Private/EditMode/SControlRigEditModeTools.h b/Engine/Plugins/Animation/ControlRig/Source/ControlRigEditor/Private/EditMode/SControlRigEditModeTools.h
--- a/Engine/Plugins/Animation/ControlRig/Source/ControlRigEditor/Private/EditMode/SControlRigEditModeTools.h
+++ b/Engine/Plugins/Animation/ControlRig/Source/ControlRigEditor/Private/EditMode/SControlRigEditModeTools.h
@@ -62,7 +62,6 @@ private:
 
 	/** Hierarchy picker for controls*/
 	TSharedPtr<SRigSpacePickerWidget> SpacePickerWidget;
-	FRigControlElementCustomization ControlCustomization;
 
 	/** Special picker for controls, no longer used */
 	TSharedPtr<SControlPicker> ControlPicker;
@@ -83,6 +82,7 @@ private:
 	void HandleSelectionChanged(TSharedPtr<FRigTreeElement> Selection, ESelectInfo::Type SelectInfo);
 	void OnRigElementSelected(UControlRig* Subject, FRigControlElement* ControlElement, bool bSelected);
 
+	const FRigControlElementCustomization* HandleGetControlElementCustomization(URigHierarchy* InHierarchy, const FRigElementKey& InControlKey);
 	void HandleActiveSpaceChanged(URigHierarchy* InHierarchy, const FRigElementKey& InControlKey, const FRigElementKey& InSpaceKey);
 	void HandleSpaceListChanged(URigHierarchy* InHierarchy, const FRigElementKey& InControlKey, const TArray<FRigElementKey>& InSpaceList);
 
diff --git a/Engine/Plugins/Animation/ControlRig/Source/ControlRigEditor/Private/Editor/SRigHierarchy.cpp b/Engine/Plugins/Animation/ControlRig/Source/ControlRigEditor/Private/Editor/SRigHierarchy.cpp
--- a/Engine/Plugins/Animation/ControlRig/Source/ControlRigEditor/Private/Editor/SRigHierarchy.cpp
+++ b/Engine/Plugins/Animation/ControlRig/Source/ControlRigEditor/Private/Editor/SRigHierarchy.cpp
@@ -2210,72 +2210,10 @@ bool SRigHierarchy::FindClosestBone(const FVector& Point, FName& OutRigElementNa
 
 void SRigHierarchy::HandleTestSpaceSwitching()
 {
-	for(const FRigElementKey& ControlKey : GetSelectedKeys())
+	if (FControlRigEditorEditMode* EditMode = ControlRigEditor.Pin()->GetEditMode())
 	{
-		if(ControlKey.Type == ERigElementType::Control)
-		{
-			TSharedRef<SRigSpacePickerWidget> PickerWidget =
-				SNew(SRigSpacePickerWidget)
-				.Hierarchy(GetDebuggedHierarchy())
-				.Control(ControlKey)
-				.AllowDelete(true)
-				.AllowReorder(true)
-				.AllowAdd(true)
-				.Title(LOCTEXT("PickSpace", "Pick Space"));
-
-			PickerWidget->OnActiveSpaceChanged().AddLambda([this, ControlKey](URigHierarchy*, const FRigElementKey& InControlKey, const FRigElementKey& InSpaceKey)
-			{
-				check(ControlKey == InControlKey);
-				
-				if(URigHierarchy* Hierarchy = GetDebuggedHierarchy())
-				{
-					const FTransform Transform = Hierarchy->GetGlobalTransform(ControlKey);
-					if(InSpaceKey == Hierarchy->GetWorldSpaceSocketKey())
-					{
-						Hierarchy->SwitchToWorldSpace(ControlKey);
-					}
-					else
-					{
-						Hierarchy->SwitchToParent(ControlKey, InSpaceKey);
-					}
-					Hierarchy->SetGlobalTransform(ControlKey, Transform);
-				}
-			});
-
-			PickerWidget->OnSpaceListChanged().AddLambda([this, ControlKey](URigHierarchy* InHierarchy, const FRigElementKey& InControlKey, const TArray<FRigElementKey>& InSpaceList)
-			{
-				check(ControlKey == InControlKey);
-
-				if(UControlRigBlueprint* RigBlueprint = ControlRigBlueprint.Get())
-				{
-					if(URigHierarchy* Hierarchy = RigBlueprint->Hierarchy)
-					{
-						// update the settings in the control element
-						if(FRigControlElement* ControlElement = Hierarchy->Find<FRigControlElement>(ControlKey))
-						{
-							RigBlueprint->Modify();
-							FScopedTransaction Transaction(LOCTEXT("ControlChangeAvailableSpaces", "Edit Available Spaces"));
-
-							ControlElement->Settings.Customization.AvailableSpaces = InSpaceList;
-							Hierarchy->Notify(ERigHierarchyNotification::ControlSettingChanged, ControlElement);
-						}
-
-						// also update the debugged instance
-						if(Hierarchy != InHierarchy)
-						{
-							if(FRigControlElement* ControlElement = InHierarchy->Find<FRigControlElement>(ControlKey))
-							{
-								ControlElement->Settings.Customization.AvailableSpaces = InSpaceList;
-								InHierarchy->Notify(ERigHierarchyNotification::ControlSettingChanged, ControlElement);
-							}
-						}
-					}
-				}
-			});
-
-			PickerWidget->OpenDialog(false);
-			break;
-		}
+		/// toooo centralize code here
+		EditMode->OpenSpacePickerWidget();
 	}
 }
 
diff --git a/Engine/Plugins/Animation/ControlRig/Source/ControlRigEditor/Private/Editor/SRigSpacePickerWidget.cpp b/Engine/Plugins/Animation/ControlRig/Source/ControlRigEditor/Private/Editor/SRigSpacePickerWidget.cpp
--- a/Engine/Plugins/Animation/ControlRig/Source/ControlRigEditor/Private/Editor/SRigSpacePickerWidget.cpp
+++ b/Engine/Plugins/Animation/ControlRig/Source/ControlRigEditor/Private/Editor/SRigSpacePickerWidget.cpp
@@ -26,18 +26,32 @@ void SRigSpacePickerWidget::Construct(const FArguments& InArgs)
 	bAllowDelete = InArgs._AllowDelete;
 	bAllowAdd = InArgs._AllowAdd;
 	bShowBakeButton = InArgs._ShowBakeButton;
-	GetAdditionalSpacesDelegate = InArgs._GetAdditionalSpacesDelegate;
+	GetActiveSpaceDelegate = InArgs._GetActiveSpace;
+	GetControlCustomizationDelegate = InArgs._GetControlCustomization;
+	GetAdditionalSpacesDelegate = InArgs._GetAdditionalSpaces;
 	bRepopulateRequired = false;
 	bLaunchingContextMenu = false;
 
+	if(!GetActiveSpaceDelegate.IsBound())
+	{
+		GetActiveSpaceDelegate = FRigSpacePickerGetActiveSpace::CreateRaw(this, &SRigSpacePickerWidget::GetActiveSpace_Private);
+	}
 	if(!GetAdditionalSpacesDelegate.IsBound())
 	{
-		GetAdditionalSpacesDelegate = FRigSpacePickerGetAdditionalSpaces::CreateRaw(this, &SRigSpacePickerWidget::GetCurrentParents);
+		GetAdditionalSpacesDelegate = FRigSpacePickerGetAdditionalSpaces::CreateRaw(this, &SRigSpacePickerWidget::GetCurrentParents_Private);
+	}
+
+	if(InArgs._OnActiveSpaceChanged.IsBound())
+	{
+		OnActiveSpaceChanged().Add(InArgs._OnActiveSpaceChanged);
+	}
+	if(InArgs._OnSpaceListChanged.IsBound())
+	{
+		OnSpaceListChanged().Add(InArgs._OnSpaceListChanged);
 	}
 
 	Hierarchy = nullptr;
-	ControlKey = DefaultParentKey = WorldSocketKey = FRigElementKey();
-	Customization = nullptr;
+	ControlKeys.Reset();
 
 	ChildSlot
 	[
@@ -68,7 +82,7 @@ void SRigSpacePickerWidget::Construct(const FArguments& InArgs)
 		AddSpacePickerRow(
 			TopLevelListBox,
 			ESpacePickerType_Parent,
-			DefaultParentKey,
+			URigHierarchy::GetDefaultParentSocketKey(),
 			FEditorStyle::GetBrush("EditorViewport.RelativeCoordinateSystem_Local"),
 			LOCTEXT("Parent", "Parent"),
 			FOnClicked::CreateSP(this, &SRigSpacePickerWidget::HandleParentSpaceClicked)
@@ -77,7 +91,7 @@ void SRigSpacePickerWidget::Construct(const FArguments& InArgs)
 		AddSpacePickerRow(
 			TopLevelListBox,
 			ESpacePickerType_World,
-			WorldSocketKey,
+			URigHierarchy::GetWorldSpaceSocketKey(),
 			FEditorStyle::GetBrush("EditorViewport.RelativeCoordinateSystem_World"),
 			LOCTEXT("World", "World"),
 			FOnClicked::CreateSP(this, &SRigSpacePickerWidget::HandleWorldSpaceClicked)
@@ -157,7 +171,7 @@ void SRigSpacePickerWidget::Construct(const FArguments& InArgs)
 		}
 	}
 
-	SetControl(InArgs._Hierarchy, InArgs._Control, InArgs._Customization);
+	SetControls(InArgs._Hierarchy, InArgs._Controls);
 	SetCanTick(true);
 	
 }
@@ -174,16 +188,10 @@ SRigSpacePickerWidget::~SRigSpacePickerWidget()
 	}
 }
 
-void SRigSpacePickerWidget::SetControl(URigHierarchy* InHierarchy, const FRigElementKey& InControl,
-	FRigControlElementCustomization* InCustomization)
+void SRigSpacePickerWidget::SetControls(
+	URigHierarchy* InHierarchy,
+	const TArray<FRigElementKey>& InControls)
 {
-	if(!InControl.IsValid() || InControl.Type != ERigElementType::Control)
-	{
-		return;
-	}
-
-	check(InHierarchy);
-
 	if(Hierarchy && Hierarchy != InHierarchy)
 	{
 		if(HierarchyModifiedHandle.IsValid())
@@ -194,23 +202,14 @@ void SRigSpacePickerWidget::SetControl(URigHierarchy* InHierarchy, const FRigEle
 	}
 	
 	Hierarchy = InHierarchy;
-	ControlKey = InControl;
-	DefaultParentKey = Hierarchy->GetFirstParent(ControlKey);
-	WorldSocketKey = Hierarchy->GetWorldSpaceSocketKey();
-	Customization = InCustomization;
-
-	HierarchyModifiedHandle = Hierarchy->OnModified().AddSP(this, &SRigSpacePickerWidget::OnHierarchyModified);
+	ControlKeys = InControls;
 
-	// if the customization is not provided
-	if(Customization == nullptr)
+	if(Hierarchy)
 	{
-		if(FRigControlElement* ControlElement = Hierarchy->Find<FRigControlElement>(ControlKey))
-		{
-			Customization = &ControlElement->Settings.Customization;
-		}
+		HierarchyModifiedHandle = Hierarchy->OnModified().AddSP(this, &SRigSpacePickerWidget::OnHierarchyModified);
 	}
 
-	UpdateActiveSpace();
+	UpdateActiveSpaces();
 	RepopulateItemSpaces();
 }
 
@@ -314,13 +313,18 @@ void SRigSpacePickerWidget::Tick(const FGeometry& AllottedGeometry, const double
 
 	if(bRepopulateRequired)
 	{
-		UpdateActiveSpace();
+		UpdateActiveSpaces();
 		RepopulateItemSpaces();
 		bRepopulateRequired = false;
 	}
 	else if(GetAdditionalSpacesDelegate.IsBound())
 	{
-		const TArray<FRigElementKey> CurrentAdditionalSpaces = GetAdditionalSpacesDelegate.Execute(Hierarchy, ControlKey);
+		TArray<FRigElementKey> CurrentAdditionalSpaces;
+		for(const FRigElementKey& ControlKey: ControlKeys)
+		{
+			CurrentAdditionalSpaces.Append(GetAdditionalSpacesDelegate.Execute(Hierarchy, ControlKey));
+		}
+		
 		if(CurrentAdditionalSpaces != AdditionalSpaces)
 		{
 			RepopulateItemSpaces();
@@ -328,16 +332,16 @@ void SRigSpacePickerWidget::Tick(const FGeometry& AllottedGeometry, const double
 	}
 }
 
-const FRigElementKey& SRigSpacePickerWidget::GetActiveSpace() const
+const TArray<FRigElementKey>& SRigSpacePickerWidget::GetActiveSpaces() const
 {
-	return ActiveSpace;
+	return ActiveSpaceKeys;
 }
 
 TArray<FRigElementKey> SRigSpacePickerWidget::GetDefaultSpaces() const
 {
 	TArray<FRigElementKey> DefaultSpaces;
-	DefaultSpaces.Add(DefaultParentKey);
-	DefaultSpaces.Add(WorldSocketKey);
+	DefaultSpaces.Add(URigHierarchy::GetDefaultParentSocketKey());
+	DefaultSpaces.Add(URigHierarchy::GetWorldSpaceSocketKey());
 	return DefaultSpaces;
 }
 
@@ -488,17 +492,20 @@ void SRigSpacePickerWidget::AddSpacePickerRow(
 
 FReply SRigSpacePickerWidget::HandleParentSpaceClicked()
 {
-	return HandleElementSpaceClicked(DefaultParentKey);
+	return HandleElementSpaceClicked(URigHierarchy::GetDefaultParentSocketKey());
 }
 
 FReply SRigSpacePickerWidget::HandleWorldSpaceClicked()
 {
-	return HandleElementSpaceClicked(WorldSocketKey);
+	return HandleElementSpaceClicked(URigHierarchy::GetWorldSpaceSocketKey());
 }
 
 FReply SRigSpacePickerWidget::HandleElementSpaceClicked(FRigElementKey InKey)
 {
-	ActiveSpaceChangedEvent.Broadcast(Hierarchy, ControlKey, InKey);
+	for(const FRigElementKey& ControlKey : ControlKeys)
+	{
+		ActiveSpaceChangedEvent.Broadcast(Hierarchy, ControlKey, InKey);
+	}
 
 	if(DialogWindow.IsValid())
 	{
@@ -519,7 +526,12 @@ FReply SRigSpacePickerWidget::HandleSpaceMoveUp(FRigElementKey InKey)
 			{
 				TArray<FRigElementKey> ChangedSpaceKeys = CurrentSpaceKeys;
 				ChangedSpaceKeys.Swap(Index, Index - 1);
-				SpaceListChangedEvent.Broadcast(Hierarchy, ControlKey, ChangedSpaceKeys);
+
+				for(const FRigElementKey& ControlKey : ControlKeys)
+				{
+					SpaceListChangedEvent.Broadcast(Hierarchy, ControlKey, ChangedSpaceKeys);
+				}
+
 				return FReply::Handled();
 			}
 		}
@@ -538,7 +550,12 @@ FReply SRigSpacePickerWidget::HandleSpaceMoveDown(FRigElementKey InKey)
 			{
 				TArray<FRigElementKey> ChangedSpaceKeys = CurrentSpaceKeys;
 				ChangedSpaceKeys.Swap(Index, Index + 1);
-				SpaceListChangedEvent.Broadcast(Hierarchy, ControlKey, ChangedSpaceKeys);
+
+				for(const FRigElementKey& ControlKey : ControlKeys)
+				{
+					SpaceListChangedEvent.Broadcast(Hierarchy, ControlKey, ChangedSpaceKeys);
+				}
+				
 				return FReply::Handled();
 			}
 		}
@@ -551,7 +568,10 @@ void SRigSpacePickerWidget::HandleSpaceDelete(FRigElementKey InKey)
 	TArray<FRigElementKey> ChangedSpaceKeys = CurrentSpaceKeys;
 	if(ChangedSpaceKeys.Remove(InKey) > 0)
 	{
-		SpaceListChangedEvent.Broadcast(Hierarchy, ControlKey, ChangedSpaceKeys);
+		for(const FRigElementKey& ControlKey : ControlKeys)
+		{
+			SpaceListChangedEvent.Broadcast(Hierarchy, ControlKey, ChangedSpaceKeys);
+		}
 	}
 }
 
@@ -568,7 +588,11 @@ FReply SRigSpacePickerWidget::HandleAddElementClicked()
 			{
 				TArray<FRigElementKey> ChangedSpaceKeys = CurrentSpaceKeys;
 				ChangedSpaceKeys.AddUnique(Key);
-				SpaceListChangedEvent.Broadcast(Hierarchy, ControlKey, ChangedSpaceKeys);
+
+				for(const FRigElementKey& ControlKey : ControlKeys)
+				{
+					SpaceListChangedEvent.Broadcast(Hierarchy, ControlKey, ChangedSpaceKeys);
+				}
 			}
 		}
 
@@ -637,7 +661,7 @@ void SRigSpacePickerWidget::OnHierarchyModified(ERigHierarchyNotification InNoti
 		return;
 	}
 
-	if(InElement->GetKey() != ControlKey)
+	if(!ControlKeys.Contains(InElement->GetKey()))
 	{
 		return;
 	}
@@ -661,13 +685,13 @@ void SRigSpacePickerWidget::OnHierarchyModified(ERigHierarchyNotification InNoti
 FSlateColor SRigSpacePickerWidget::GetButtonColor(ESpacePickerType InType, FRigElementKey InKey) const
 {
 	static const FSlateColor ActiveColor = FControlRigEditorStyle::Get().SpacePickerSelectColor;
-	
+
 	switch(InType)
 	{
 		case ESpacePickerType_Parent:
 		{
 			// this is also true if the object has no parent
-			if(GetActiveSpace() == DefaultParentKey)
+			if(ActiveSpaceKeys.Contains(URigHierarchy::GetDefaultParentSocketKey()))
 			{
 				return ActiveColor;
 			}
@@ -675,7 +699,7 @@ FSlateColor SRigSpacePickerWidget::GetButtonColor(ESpacePickerType InType, FRigE
 		}
 		case ESpacePickerType_World:
 		{
-			if(GetActiveSpace() == WorldSocketKey && WorldSocketKey.IsValid())
+			if(ActiveSpaceKeys.Contains(URigHierarchy::GetWorldSpaceSocketKey()))
 			{
 				return ActiveColor;
 			}
@@ -684,27 +708,66 @@ FSlateColor SRigSpacePickerWidget::GetButtonColor(ESpacePickerType InType, FRigE
 		case ESpacePickerType_Item:
 		default:
 		{
-			if(GetActiveSpace() == InKey && InKey.IsValid())
+			if(ActiveSpaceKeys.Contains(InKey) && InKey.IsValid())
 			{
 				return ActiveColor;
 			}
 			break;
 		}
 	}
-	
 	return FStyleColors::Transparent;
 }
 
-TArray<FRigElementKey> SRigSpacePickerWidget::GetCurrentParents(URigHierarchy* InHierarchy,
+FRigElementKey SRigSpacePickerWidget::GetActiveSpace_Private(URigHierarchy* InHierarchy,
 	const FRigElementKey& InControlKey) const
+{
+	if(InHierarchy)
+	{
+		const TArray<FRigElementWeight> ParentWeights = InHierarchy->GetParentWeightArray(InControlKey);
+		if(ParentWeights.Num() > 0)
+		{
+			const TArray<FRigElementKey> ParentKeys = InHierarchy->GetParents(InControlKey);
+			check(ParentKeys.Num() == ParentWeights.Num());
+			for(int32 ParentIndex=0;ParentIndex<ParentKeys.Num();ParentIndex++)
+			{
+				if(ParentWeights[ParentIndex].IsAlmostZero())
+				{
+					continue;
+				}
+
+				if(ParentIndex == 0)
+				{
+					if(!IsDefaultSpace(ParentKeys[ParentIndex]))
+					{
+						return URigHierarchy::GetDefaultParentSocketKey();
+					}
+				}
+
+				return ParentKeys[ParentIndex];
+			}
+		}
+	}
+	return URigHierarchy::GetDefaultParentSocketKey();
+}
+
+TArray<FRigElementKey> SRigSpacePickerWidget::GetCurrentParents_Private(URigHierarchy* InHierarchy,
+                                                                const FRigElementKey& InControlKey) const
 {
 	if(!InControlKey.IsValid() || InHierarchy == nullptr)
 	{
 		return TArray<FRigElementKey>();
 	}
 
-	check(InControlKey == ControlKey);
-	return InHierarchy->GetParents(ControlKey);
+	check(ControlKeys.Contains(InControlKey));
+	TArray<FRigElementKey> Parents = InHierarchy->GetParents(InControlKey);
+	if(Parents.Num() > 0)
+	{
+		if(!IsDefaultSpace(Parents[0]))
+		{
+			Parents[0] = URigHierarchy::GetDefaultParentSocketKey();
+		}
+	}
+	return Parents;
 }
 
 void SRigSpacePickerWidget::RepopulateItemSpaces()
@@ -713,54 +776,74 @@ void SRigSpacePickerWidget::RepopulateItemSpaces()
 	{
 		return;
 	}
+	if(Hierarchy == nullptr)
+	{
+		return;
+	}
 	
 	TArray<FRigElementKey> FavoriteKeys, SpacesFromDelegate;
 
-	// gather the keys coming from the input customization
-	if(Customization && bShowFavoriteSpaces)
+	if(bShowFavoriteSpaces)
 	{
-		for(const FRigElementKey& Key : Customization->AvailableSpaces)
+		for(const FRigElementKey& ControlKey : ControlKeys)
 		{
-			if(IsDefaultSpace(Key) || !IsValidKey(Key))
+			const FRigControlElementCustomization* Customization = nullptr;
+			if(GetControlCustomizationDelegate.IsBound())
 			{
-				continue;
+				Customization = GetControlCustomizationDelegate.Execute(Hierarchy, ControlKey);
 			}
-			FavoriteKeys.AddUnique(Key);
-		}
-	}
 
-	// check if the customization is different from the base one in the asset
-	if(const FRigControlElement* ControlElement = Hierarchy->Find<FRigControlElement>(ControlKey))
-	{
-		if(Customization != &ControlElement->Settings.Customization)
-		{
-			for(const FRigElementKey& Key : ControlElement->Settings.Customization.AvailableSpaces)
+			if(Customization)
 			{
-				if(IsDefaultSpace(Key) || !IsValidKey(Key))
-				{
-					continue;
-				}
-
-				if(Customization)
+				for(const FRigElementKey& Key : Customization->AvailableSpaces)
 				{
-					if(Customization->AvailableSpaces.Contains(Key))
+					if(IsDefaultSpace(Key) || !IsValidKey(Key))
 					{
 						continue;
 					}
-					if(Customization->RemovedSpaces.Contains(Key))
+					FavoriteKeys.AddUnique(Key);
+				}
+			}
+			
+			// check if the customization is different from the base one in the asset
+			if(const FRigControlElement* ControlElement = Hierarchy->Find<FRigControlElement>(ControlKey))
+			{
+				if(Customization != &ControlElement->Settings.Customization)
+				{
+					for(const FRigElementKey& Key : ControlElement->Settings.Customization.AvailableSpaces)
 					{
-						continue;
+						if(IsDefaultSpace(Key) || !IsValidKey(Key))
+						{
+							continue;
+						}
+
+						if(Customization)
+						{
+							if(Customization->AvailableSpaces.Contains(Key))
+							{
+								continue;
+							}
+							if(Customization->RemovedSpaces.Contains(Key))
+							{
+								continue;
+							}
+						}
+						FavoriteKeys.AddUnique(Key);
 					}
 				}
-				FavoriteKeys.AddUnique(Key);
 			}
 		}
 	}
-
+	
 	// now gather all of the spaces using the get additional spaces delegate
 	if(GetAdditionalSpacesDelegate.IsBound() && bShowAdditionalSpaces)
 	{
-		AdditionalSpaces = GetAdditionalSpacesDelegate.Execute(Hierarchy, ControlKey);
+		AdditionalSpaces.Reset();
+		for(const FRigElementKey& ControlKey: ControlKeys)
+		{
+			AdditionalSpaces.Append(GetAdditionalSpacesDelegate.Execute(Hierarchy, ControlKey));
+		}
+		
 		for(const FRigElementKey& Key : AdditionalSpaces)
 		{
 			if(IsDefaultSpace(Key)  || !IsValidKey(Key))
@@ -832,37 +915,60 @@ void SRigSpacePickerWidget::ClearListBox(TSharedPtr<SVerticalBox> InListBox)
 	InListBox->ClearChildren();
 }
 
-void SRigSpacePickerWidget::UpdateActiveSpace()
+void SRigSpacePickerWidget::UpdateActiveSpaces()
 {
-	ActiveSpace = FRigElementKey();
+	ActiveSpaceKeys.Reset();
+
+	if(Hierarchy == nullptr)
+	{
+		return;
+	}
 	
-	const TArray<FRigElementWeight> ParentWeights = Hierarchy->GetParentWeightArray(ControlKey);
-	if(ParentWeights.Num() > 0)
+	for(int32 ControlIndex=0;ControlIndex<ControlKeys.Num();ControlIndex++)
 	{
-		const TArray<FRigElementKey> ParentKeys = Hierarchy->GetParents(ControlKey);
-		check(ParentKeys.Num() == ParentWeights.Num());
-		for(int32 ParentIndex=0;ParentIndex<ParentKeys.Num();ParentIndex++)
+		ActiveSpaceKeys.Add(URigHierarchy::GetDefaultParentSocketKey());
+		
+		const FRigElementKey& ControlKey = ControlKeys[ControlIndex];
+		const TArray<FRigElementWeight> ParentWeights = Hierarchy->GetParentWeightArray(ControlKey);
+		if(ParentWeights.Num() > 0)
 		{
-			if(ParentWeights[ParentIndex].IsAlmostZero())
+			const TArray<FRigElementKey> ParentKeys = Hierarchy->GetParents(ControlKey);
+			check(ParentKeys.Num() == ParentWeights.Num());
+			for(int32 ParentIndex=0;ParentIndex<ParentKeys.Num();ParentIndex++)
 			{
-				continue;
+				if(ParentWeights[ParentIndex].IsAlmostZero())
+				{
+					continue;
+				}
+
+				if(ParentIndex > 0 || IsDefaultSpace(ParentKeys[ParentIndex]))
+				{
+					ActiveSpaceKeys[ControlIndex] = ParentKeys[ParentIndex];
+				}
+				break;
 			}
-			ActiveSpace = ParentKeys[ParentIndex];
-			return;
 		}
 	}
 }
 
 bool SRigSpacePickerWidget::IsValidKey(const FRigElementKey& InKey) const
 {
-	return InKey.IsValid() && Hierarchy->Contains(InKey);
+	if(!InKey.IsValid())
+	{
+		return false;
+	}
+	if(Hierarchy == nullptr)
+	{
+		return false;
+	}
+	return Hierarchy->Contains(InKey);
 }
 
 bool SRigSpacePickerWidget::IsDefaultSpace(const FRigElementKey& InKey) const
 {
 	if(bShowDefaultSpaces)
 	{
-		return InKey == WorldSocketKey || InKey == DefaultParentKey;
+		return InKey == URigHierarchy::GetDefaultParentSocketKey() || InKey == URigHierarchy::GetWorldSpaceSocketKey();
 	}
 	return false;
 }
diff --git a/Engine/Plugins/Animation/ControlRig/Source/ControlRigEditor/Private/Editor/SRigSpacePickerWidget.h b/Engine/Plugins/Animation/ControlRig/Source/ControlRigEditor/Private/Editor/SRigSpacePickerWidget.h
--- a/Engine/Plugins/Animation/ControlRig/Source/ControlRigEditor/Private/Editor/SRigSpacePickerWidget.h
+++ b/Engine/Plugins/Animation/ControlRig/Source/ControlRigEditor/Private/Editor/SRigSpacePickerWidget.h
@@ -9,6 +9,8 @@
 #include "Widgets/Layout/SBox.h"
 #include "Rigs/RigSpaceHierarchy.h"
 
+DECLARE_DELEGATE_RetVal_TwoParams(FRigElementKey, FRigSpacePickerGetActiveSpace, URigHierarchy*, const FRigElementKey&);
+DECLARE_DELEGATE_RetVal_TwoParams(const FRigControlElementCustomization*, FRigSpacePickerGetControlCustomization, URigHierarchy*, const FRigElementKey&);
 DECLARE_EVENT_ThreeParams(SRigSpacePickerWidget, FRigSpacePickerActiveSpaceChanged, URigHierarchy*, const FRigElementKey&, const FRigElementKey&);
 DECLARE_EVENT_ThreeParams(SRigSpacePickerWidget, FRigSpacePickerSpaceListChanged, URigHierarchy*, const FRigElementKey&, const TArray<FRigElementKey>&);
 DECLARE_DELEGATE_RetVal_TwoParams(TArray<FRigElementKey>, FRigSpacePickerGetAdditionalSpaces, URigHierarchy*, const FRigElementKey&);
@@ -20,8 +22,7 @@ public:
 
 	SLATE_BEGIN_ARGS(SRigSpacePickerWidget)
 		: _Hierarchy(nullptr)
-		, _Control()
-		, _Customization(nullptr)
+		, _Controls()
 		, _ShowDefaultSpaces(true)
 		, _ShowFavoriteSpaces(true)
 		, _ShowAdditionalSpaces(true)
@@ -33,8 +34,7 @@ public:
 		, _BackgroundBrush(FEditorStyle::GetBrush("Menu.Background"))
 		{}
 		SLATE_ARGUMENT(URigHierarchy*, Hierarchy)
-		SLATE_ARGUMENT(FRigElementKey, Control)
-		SLATE_ARGUMENT(FRigControlElementCustomization*, Customization)
+		SLATE_ARGUMENT(TArray<FRigElementKey>, Controls)
 		SLATE_ARGUMENT(bool, ShowDefaultSpaces)
 		SLATE_ARGUMENT(bool, ShowFavoriteSpaces)
 		SLATE_ARGUMENT(bool, ShowAdditionalSpaces)
@@ -44,14 +44,19 @@ public:
 		SLATE_ARGUMENT(bool, ShowBakeButton)
 		SLATE_ARGUMENT(FText, Title)
 		SLATE_ARGUMENT(const FSlateBrush*, BackgroundBrush)
-		SLATE_ARGUMENT(FRigSpacePickerGetAdditionalSpaces, GetAdditionalSpacesDelegate)
+	
+		SLATE_EVENT(FRigSpacePickerGetActiveSpace, GetActiveSpace)
+		SLATE_EVENT(FRigSpacePickerGetControlCustomization, GetControlCustomization)
+		SLATE_EVENT(FRigSpacePickerActiveSpaceChanged::FDelegate, OnActiveSpaceChanged)
+		SLATE_EVENT(FRigSpacePickerSpaceListChanged::FDelegate, OnSpaceListChanged)
+		SLATE_ARGUMENT(FRigSpacePickerGetAdditionalSpaces, GetAdditionalSpaces)
 		SLATE_EVENT( FOnClicked, OnBakeButtonClicked )
 	SLATE_END_ARGS()
 
 	void Construct(const FArguments& InArgs);
 	virtual ~SRigSpacePickerWidget() override;
 
-	void SetControl(URigHierarchy* InHierarchy, const FRigElementKey& InControl, FRigControlElementCustomization* InCustomization = nullptr);
+	void SetControls(URigHierarchy* InHierarchy, const TArray<FRigElementKey>& InControls);
 
 	FReply OpenDialog(bool bModal = true);
 	void CloseDialog();
@@ -60,7 +65,7 @@ public:
 	virtual bool SupportsKeyboardFocus() const override;
 	virtual void Tick(const FGeometry& AllottedGeometry, const double InCurrentTime, const float InDeltaTime) override;
 
-	const FRigElementKey& GetActiveSpace() const;
+	const TArray<FRigElementKey>& GetActiveSpaces() const;
 	TArray<FRigElementKey> GetDefaultSpaces() const;
 	TArray<FRigElementKey> GetSpaceList(bool bIncludeDefaultSpaces = false) const;
 	FRigSpacePickerActiveSpaceChanged& OnActiveSpaceChanged() { return ActiveSpaceChangedEvent; }
@@ -85,7 +90,7 @@ private:
 
 	void RepopulateItemSpaces();
 	void ClearListBox(TSharedPtr<SVerticalBox> InListBox);
-	void UpdateActiveSpace();
+	void UpdateActiveSpaces();
 	bool IsValidKey(const FRigElementKey& InKey) const;
 	bool IsDefaultSpace(const FRigElementKey& InKey) const;
 
@@ -103,19 +108,18 @@ private:
 	void OnHierarchyModified(ERigHierarchyNotification InNotif, URigHierarchy* InHierarchy, const FRigBaseElement* InElement);
 
 	FSlateColor GetButtonColor(ESpacePickerType InType, FRigElementKey InKey) const;
-	TArray<FRigElementKey> GetCurrentParents(URigHierarchy* InHierarchy, const FRigElementKey& InControlKey) const;
+	FRigElementKey GetActiveSpace_Private(URigHierarchy* InHierarchy, const FRigElementKey& InControlKey) const;
+	TArray<FRigElementKey> GetCurrentParents_Private(URigHierarchy* InHierarchy, const FRigElementKey& InControlKey) const;
 	
 	FRigSpacePickerActiveSpaceChanged ActiveSpaceChangedEvent;
 	FRigSpacePickerSpaceListChanged SpaceListChangedEvent;
 
 	URigHierarchy* Hierarchy;
-	FRigElementKey ControlKey;
-	FRigElementKey DefaultParentKey;
-	FRigElementKey WorldSocketKey;
+	TArray<FRigElementKey> ControlKeys;
 	TArray<FRigElementKey> CurrentSpaceKeys;
+	TArray<FRigElementKey> ActiveSpaceKeys;
 	bool bRepopulateRequired;
 
-	FRigControlElementCustomization* Customization;
 	bool bShowDefaultSpaces;
 	bool bShowFavoriteSpaces;
 	bool bShowAdditionalSpaces;
@@ -125,6 +129,8 @@ private:
 	bool bShowBakeButton;
 	bool bLaunchingContextMenu;
 
+	FRigSpacePickerGetControlCustomization GetControlCustomizationDelegate;
+	FRigSpacePickerGetActiveSpace GetActiveSpaceDelegate;
 	FRigSpacePickerGetAdditionalSpaces GetAdditionalSpacesDelegate; 
 	TArray<FRigElementKey> AdditionalSpaces;
 
@@ -133,7 +139,6 @@ private:
 	TSharedPtr<SHorizontalBox> BottomButtonsListBox;
 	TWeakPtr<SWindow> DialogWindow;
 	TWeakPtr<IMenu> ContextMenu;
-	FRigElementKey ActiveSpace;
 	FDelegateHandle HierarchyModifiedHandle;
 	FDelegateHandle ActiveSpaceChangedWindowHandle;
 
