Control Rig: Selection issues with clearing incorrectly and Frustom. I got rid of the edit mode storing it's own state but instead it just gets it from the Control Rig. Also frustum select is now just 3d based, the 3d to 2d selection was borked and no one else was doing that anymore, everything is now 3D.

diff --git a/Engine/Plugins/Animation/ControlRig/Source/ControlRigEditor/Private/EditMode/ControlRigEditMode.cpp b/Engine/Plugins/Animation/ControlRig/Source/ControlRigEditor/Private/EditMode/ControlRigEditMode.cpp
--- a/Engine/Plugins/Animation/ControlRig/Source/ControlRigEditor/Private/EditMode/ControlRigEditMode.cpp
+++ b/Engine/Plugins/Animation/ControlRig/Source/ControlRigEditor/Private/EditMode/ControlRigEditMode.cpp
@@ -434,7 +434,7 @@ void FControlRigEditMode::Tick(FEditorViewportClient* ViewportClient, float Delt
 	if (bRecreateGizmosRequired && !(FSlateApplication::Get().HasAnyMouseCaptor() || GUnrealEd->IsUserInteracting()))
 	{
 		RecreateGizmoActors();
-
+		TArray<FRigElementKey> SelectedRigElements = GetSelectedRigElements();
 		for (const FRigElementKey& SelectedKey : SelectedRigElements)
 		{
 			if (SelectedKey.Type == ERigElementType::Control)
@@ -473,6 +473,7 @@ void FControlRigEditMode::Tick(FEditorViewportClient* ViewportClient, float Delt
 
 		if (UControlRig* ControlRig = GetControlRig(true))
 		{
+			TArray<FRigElementKey> SelectedRigElements = GetSelectedRigElements();
 			const UE::Widget::EWidgetMode CurrentWidgetMode = ViewportClient->GetWidgetMode();
 			for (FRigElementKey SelectedRigElement : SelectedRigElements)
 			{
@@ -738,6 +739,7 @@ void FControlRigEditMode::Render(const FSceneView* View, FViewport* Viewport, FP
 		{
 			if (ControlRig->GetWorld() && ControlRig->GetWorld()->IsPreviewWorld())
 			{
+				TArray<FRigElementKey> SelectedRigElements = GetSelectedRigElements();
 				const float Scale = Settings->AxisScale;
 				PDI->AddReserveLines(SDPG_Foreground, SelectedRigElements.Num() * 3);
 
@@ -862,6 +864,7 @@ bool FControlRigEditMode::StartTracking(FEditorViewportClient* InViewportClient,
 			bool bShouldModify = IsInLevelEditor();
 			if (!bShouldModify)
 			{
+				TArray<FRigElementKey> SelectedRigElements = GetSelectedRigElements();
 				for (const FRigElementKey& Key : SelectedRigElements)
 				{
 					if (Key.Type != ERigElementType::Control)
@@ -1354,57 +1357,6 @@ bool FControlRigEditMode::BoxSelect(FBox& InBox, bool InSelect)
 
 bool FControlRigEditMode::FrustumSelect(const FConvexVolume& InFrustum, FEditorViewportClient* InViewportClient, bool InSelect)
 {
-	float StartX = TNumericLimits<float>::Max();
-	float StartY = TNumericLimits<float>::Max();
-	float EndX = TNumericLimits<float>::Lowest();
-	float EndY = TNumericLimits<float>::Lowest();
-
-	// Frustum sides are in the first four indices
-	// Find intersection points and project to screen space to determine the bounding rectangle of the selection box
-	for (int32 PlaneIndex = 0; PlaneIndex < 4; ++PlaneIndex)
-	{
-		if (!InFrustum.Planes.IsValidIndex(PlaneIndex) ||
-			!InFrustum.Planes.IsValidIndex((PlaneIndex + 1) % 4))
-		{
-			break;
-		}
-
-		const FPlane& Plane1 = InFrustum.Planes[PlaneIndex];
-		const FPlane& Plane2 = InFrustum.Planes[(PlaneIndex + 1) % 4];
-		FVector I, D;
-		if (FMath::IntersectPlanes2(I, D, Plane1, Plane2))
-		{
-			FSceneViewFamilyContext ViewFamily(FSceneViewFamily::ConstructionValues(InViewportClient->Viewport, InViewportClient->GetScene(), InViewportClient->EngineShowFlags));
-			FSceneView* SceneView = InViewportClient->CalcSceneView(&ViewFamily);
-
-			FVector2D V;
-			if (SceneView->WorldToPixel(I, V))
-			{
-				StartX = FMath::Min(StartX, V.X);
-				StartY = FMath::Min(StartY, V.Y);
-				EndX = FMath::Max(EndX, V.X);
-				EndY = FMath::Max(EndY, V.Y);
-			}
-		}
-	}
-	const int32 ViewportSizeX = InViewportClient->Viewport->GetSizeXY().X;
-	const int32 ViewportSizeY = InViewportClient->Viewport->GetSizeXY().Y;
-
-	if( StartX > EndX )
-	{
-		Swap( StartX, EndX );
-	}
-
-	if( StartY > EndY )
-	{
-		Swap( StartY, EndY );
-	}
-
-	FIntRect BoxRect(FIntPoint(FMath::Max(0.0f, StartX), FMath::Max(0.0f, StartY)), FIntPoint(FMath::Min(ViewportSizeX, FMath::TruncToInt(EndX + 1)), FMath::Min(ViewportSizeY, FMath::TruncToInt(EndY + 1))));
-
-	TSet<AActor*> HitActors;
-	TSet<UModel*> HitModels;
-	InViewportClient->Viewport->GetActorsAndModelsInHitProxy(BoxRect, HitActors, HitModels);
 	FScopedTransaction ScopedTransaction(LOCTEXT("SelectControlTransaction", "Select Control"), IsInLevelEditor() && !GIsTransacting);
 	bool bSomethingSelected(false);
 	const bool bShiftDown = InViewportClient->Viewport->KeyState(EKeys::LeftShift) || InViewportClient->Viewport->KeyState(EKeys::RightShift);
@@ -1412,23 +1364,22 @@ bool FControlRigEditMode::FrustumSelect(const FConvexVolume& InFrustum, FEditorV
 	{
 		ClearRigElementSelection(FRigElementTypeHelper::ToMask(ERigElementType::Control));
 	}
-	for (AActor* Actor : HitActors)
+
+	for (AControlRigGizmoActor* GizmoActor : GizmoActors)
 	{
-		if (Actor->IsA<AControlRigGizmoActor>())
+		for (UActorComponent* Component : GizmoActor->GetComponents())
 		{
-			AControlRigGizmoActor* GizmoActor = CastChecked<AControlRigGizmoActor>(Actor);
-			if (GizmoActor->IsSelectable())
+			UPrimitiveComponent* PrimitiveComponent = Cast<UPrimitiveComponent>(Component);
+			if (PrimitiveComponent && PrimitiveComponent->IsRegistered() && PrimitiveComponent->IsVisibleInEditor())
 			{
-				bSomethingSelected = true;
-				const FName& ControlName = GizmoActor->ControlName;
-				SetRigElementSelection(ERigElementType::Control, ControlName, true);
-
-				if (bShiftDown)
-				{
-				}
-				else
+				if (PrimitiveComponent->ComponentIsTouchingSelectionFrustum(InFrustum, InViewportClient->EngineShowFlags, false /*only bsp*/, false/*encompass entire*/))
 				{
-					SetRigElementSelection(ERigElementType::Control, ControlName, true);
+					if (GizmoActor->IsSelectable())
+					{
+						bSomethingSelected = true;
+						const FName& ControlName = GizmoActor->ControlName;
+						SetRigElementSelection(ERigElementType::Control, ControlName, true);
+					}
 				}
 			}
 		}
@@ -1437,7 +1388,6 @@ bool FControlRigEditMode::FrustumSelect(const FConvexVolume& InFrustum, FEditorV
 	{
 		return true;
 	}
-	
 	ScopedTransaction.Cancel();
 	return FEdMode::FrustumSelect(InFrustum, InViewportClient, InSelect);
 }
@@ -1544,6 +1494,8 @@ bool FControlRigEditMode::InputDelta(FEditorViewportClient* InViewportClient, FV
 
 			// set Bone transform
 			// that will set initial Bone transform
+			TArray<FRigElementKey> SelectedRigElements = GetSelectedRigElements();
+
 			for (int32 Index = 0; Index < SelectedRigElements.Num(); ++Index)
 			{
 				const ERigElementType SelectedRigElementType = SelectedRigElements[Index].Type;
@@ -1691,19 +1643,34 @@ void FControlRigEditMode::SetRigElementSelection(ERigElementType Type, const TAr
 {
 	if (!bSelecting)
 	{
-TGuardValue<bool> ReentrantGuard(bSelecting, true);
+		TGuardValue<bool> ReentrantGuard(bSelecting, true);
+		for (const FName& ElementName : InRigElementNames)
+		{
+			SetRigElementSelectionInternal(Type, ElementName, bSelected);
+		}
 
-for (const FName& ElementName : InRigElementNames)
-{
-	SetRigElementSelectionInternal(Type, ElementName, bSelected);
+		HandleSelectionChanged();
+	}
 }
 
-HandleSelectionChanged();
+TArray<FRigElementKey> FControlRigEditMode::GetSelectedRigElements() const
+{
+	if (UControlRig* ControlRig = GetControlRig(true))
+	{
+		if (ControlRig->GetHierarchy())
+		{
+			return ControlRig->GetHierarchy()->GetSelectedKeys();
+		}
 	}
+
+	TArray<FRigElementKey> Empty;
+	return Empty;
 }
 
 bool FControlRigEditMode::AreRigElementsSelected(uint32 InTypes) const
 {
+	TArray<FRigElementKey> SelectedRigElements = GetSelectedRigElements();
+
 	for (const FRigElementKey& Ele : SelectedRigElements)
 	{
 		if (FRigElementTypeHelper::DoesHave(InTypes, Ele.Type))
@@ -1717,6 +1684,7 @@ bool FControlRigEditMode::AreRigElementsSelected(uint32 InTypes) const
 
 int32 FControlRigEditMode::GetNumSelectedRigElements(uint32 InTypes) const
 {
+	TArray<FRigElementKey> SelectedRigElements = GetSelectedRigElements();
 	if (FRigElementTypeHelper::DoesHave(InTypes, ERigElementType::All))
 	{
 		return SelectedRigElements.Num();
@@ -1764,6 +1732,7 @@ void FControlRigEditMode::RecalcPivotTransform()
 	PivotTransform = FTransform::Identity;
 
 	// @todo: support bones also
+	TArray<FRigElementKey> SelectedRigElements = GetSelectedRigElements();
 	if (AreRigElementsSelected(FRigElementTypeHelper::ToMask(ERigElementType::Control)))
 	{
 		FTransform LastTransform = FTransform::Identity;
@@ -1968,17 +1937,7 @@ bool FControlRigEditMode::GetRigElementGlobalTransform(const FRigElementKey& InE
 
 bool FControlRigEditMode::CanFrameSelection()
 {
-	return SelectedRigElements.Num() > 0;
-	/*
-	for (const FRigElementKey& SelectedKey : SelectedRigElements)
-	{
-		if (SelectedKey.Type == ERigElementType::Control)
-		{
-			return true;
-		}
-	}
-	return false;
-	*/
+	return  GetSelectedRigElements().Num() > 0;
 }
 
 void FControlRigEditMode::ClearSelection()
@@ -2006,6 +1965,7 @@ void FControlRigEditMode::FrameSelection()
     }
 
 	TArray<AActor*> Actors;
+	TArray<FRigElementKey> SelectedRigElements = GetSelectedRigElements();
 	for (const FRigElementKey& SelectedKey : SelectedRigElements)
 	{
 		if (SelectedKey.Type == ERigElementType::Control)
@@ -2207,6 +2167,7 @@ void FControlRigEditMode::ResetTransforms(bool bSelectionOnly)
 {
 	if (UControlRig* ControlRig = GetControlRig(true))
 	{
+		TArray<FRigElementKey> SelectedRigElements = GetSelectedRigElements();
 		TArray<FRigElementKey> ControlsToReset = SelectedRigElements;
 		if (!bSelectionOnly)
 		{
@@ -2436,15 +2397,7 @@ void FControlRigEditMode::OnHierarchyModified(ERigHierarchyNotification InNotif,
             	case ERigElementType::Socket:
 				{
 					const bool bSelected = InNotif == ERigHierarchyNotification::ElementSelected;
-					if (bSelected)
-					{
-						SelectedRigElements.AddUnique(Key);
-					}
-					else
-					{
-						SelectedRigElements.Remove(Key);
-					}
-
+					
 					// if it's control
 					if (Key.Type == ERigElementType::Control)
 					{
@@ -2550,6 +2503,7 @@ bool FControlRigEditMode::CanChangeControlGizmoTransform()
 {
 	if (!IsInLevelEditor())
 	{
+		TArray<FRigElementKey> SelectedRigElements = GetSelectedRigElements();
 		// do not allow multi-select
 		if (SelectedRigElements.Num() == 1)
 		{
diff --git a/Engine/Plugins/Animation/ControlRig/Source/ControlRigEditor/Private/EditMode/ControlRigEditMode.h b/Engine/Plugins/Animation/ControlRig/Source/ControlRigEditor/Private/EditMode/ControlRigEditMode.h
--- a/Engine/Plugins/Animation/ControlRig/Source/ControlRigEditor/Private/EditMode/ControlRigEditMode.h
+++ b/Engine/Plugins/Animation/ControlRig/Source/ControlRigEditor/Private/EditMode/ControlRigEditMode.h
@@ -270,7 +270,8 @@ protected:
 	FSimpleMulticastDelegate OnAnimSystemInitializedDelegate;
 	FControlRigAddedOrRemoved OnControlRigAddedOrRemovedDelegate;
 
-	TArray<FRigElementKey> SelectedRigElements;
+	/** GetSelectedRigElements */
+	TArray<FRigElementKey> GetSelectedRigElements() const;
 
 	/* Flag to recreate gizmos during tick */
 	bool bRecreateGizmosRequired;
diff --git a/Engine/Plugins/Animation/ControlRig/Source/ControlRigEditor/Private/Editor/ControlRigEditorEditMode.cpp b/Engine/Plugins/Animation/ControlRig/Source/ControlRigEditor/Private/Editor/ControlRigEditorEditMode.cpp
--- a/Engine/Plugins/Animation/ControlRig/Source/ControlRigEditor/Private/Editor/ControlRigEditorEditMode.cpp
+++ b/Engine/Plugins/Animation/ControlRig/Source/ControlRigEditor/Private/Editor/ControlRigEditorEditMode.cpp
@@ -109,7 +109,7 @@ void FControlRigEditorEditMode::Render(const FSceneView* View, FViewport* Viewpo
 bool FControlRigEditorEditMode::GetCameraTarget(FSphere& OutTarget) const
 {
 	FBox Box(ForceInit);
-
+	TArray<FRigElementKey> SelectedRigElements = GetSelectedRigElements();
 	for (int32 Index = 0; Index < SelectedRigElements.Num(); ++Index)
 	{
 		static const float Radius = 20.f;
