Control Rig: Space Switching Bake Dialog Mockup

diff --git a/Engine/Plugins/Animation/ControlRig/Source/ControlRigEditor/Private/EditMode/SControlRigEditModeTools.cpp b/Engine/Plugins/Animation/ControlRig/Source/ControlRigEditor/Private/EditMode/SControlRigEditModeTools.cpp
--- a/Engine/Plugins/Animation/ControlRig/Source/ControlRigEditor/Private/EditMode/SControlRigEditModeTools.cpp
+++ b/Engine/Plugins/Animation/ControlRig/Source/ControlRigEditor/Private/EditMode/SControlRigEditModeTools.cpp
@@ -325,6 +325,7 @@ void SControlRigEditModeTools::Construct(const FArguments& InArgs, FControlRigEd
 					.GetControlCustomization(this, &SControlRigEditModeTools::HandleGetControlElementCustomization)
 					.OnActiveSpaceChanged(this, &SControlRigEditModeTools::HandleActiveSpaceChanged)
 					.OnSpaceListChanged(this, &SControlRigEditModeTools::HandleSpaceListChanged)
+					.OnBakeButtonClicked(this, &SControlRigEditModeTools::OnBakeControlsToNewSpaceButtonClicked)
 					// todo: implement GetAdditionalSpacesDelegate to pull spaces from sequencer
 				]
 			]
@@ -333,7 +334,7 @@ void SControlRigEditModeTools::Construct(const FArguments& InArgs, FControlRigEd
 			.AutoHeight()
 			[
 				SAssignNew(RigOptionExpander, SExpandableArea)
-				.InitiallyCollapsed(true)
+				.InitiallyCollapsed(false)
 				.Visibility(this, &SControlRigEditModeTools::GetRigOptionExpanderVisibility)
 				.AreaTitle(LOCTEXT("RigOption_Header", "Rig Options"))
 				.AreaTitleFont(FEditorStyle::GetFontStyle("DetailsView.CategoryFontStyle"))
@@ -651,6 +652,37 @@ void SControlRigEditModeTools::HandleSpaceListChanged(URigHierarchy* InHierarchy
 	}
 }
 
+FReply SControlRigEditModeTools::OnBakeControlsToNewSpaceButtonClicked()
+{
+	if(SpacePickerWidget->GetHierarchy() == nullptr)
+	{
+		return FReply::Unhandled();
+	}
+	if(SpacePickerWidget->GetControls().Num() == 0)
+	{
+		return FReply::Unhandled();
+	}
+
+	FRigSpacePickerBakeSettings Settings;
+	// todoo pick the default target space... parent?
+	Settings.TargetSpace = URigHierarchy::GetDefaultParentSocketKey();
+	// todoooo determine the range
+
+	TSharedRef<SRigSpacePickerBakeWidget> BakeWidget =
+		SNew(SRigSpacePickerBakeWidget)
+		.Settings(Settings)
+		.Hierarchy(SpacePickerWidget->GetHierarchy())
+		.Controls(SpacePickerWidget->GetControls())
+		.GetControlCustomization(this, &SControlRigEditModeTools::HandleGetControlElementCustomization)
+		.OnBake_Lambda([](URigHierarchy* InHierarchy, TArray<FRigElementKey> InControls, FRigSpacePickerBakeSettings InSettings)
+		{
+			// todoo react to the bake
+			return FReply::Handled();
+		});
+
+	return BakeWidget->OpenDialog(true);
+}
+
 EVisibility SControlRigEditModeTools::GetRigOptionExpanderVisibility() const
 {
 	if (UControlRig* ControlRig = SequencerRig.Get())
diff --git a/Engine/Plugins/Animation/ControlRig/Source/ControlRigEditor/Private/EditMode/SControlRigEditModeTools.h b/Engine/Plugins/Animation/ControlRig/Source/ControlRigEditor/Private/EditMode/SControlRigEditModeTools.h
--- a/Engine/Plugins/Animation/ControlRig/Source/ControlRigEditor/Private/EditMode/SControlRigEditModeTools.h
+++ b/Engine/Plugins/Animation/ControlRig/Source/ControlRigEditor/Private/EditMode/SControlRigEditModeTools.h
@@ -85,6 +85,7 @@ private:
 	const FRigControlElementCustomization* HandleGetControlElementCustomization(URigHierarchy* InHierarchy, const FRigElementKey& InControlKey);
 	void HandleActiveSpaceChanged(URigHierarchy* InHierarchy, const FRigElementKey& InControlKey, const FRigElementKey& InSpaceKey);
 	void HandleSpaceListChanged(URigHierarchy* InHierarchy, const FRigElementKey& InControlKey, const TArray<FRigElementKey>& InSpaceList);
+	FReply OnBakeControlsToNewSpaceButtonClicked();
 
 	EVisibility GetRigOptionExpanderVisibility() const;
 
diff --git a/Engine/Plugins/Animation/ControlRig/Source/ControlRigEditor/Private/Editor/SRigSpacePickerWidget.cpp b/Engine/Plugins/Animation/ControlRig/Source/ControlRigEditor/Private/Editor/SRigSpacePickerWidget.cpp
--- a/Engine/Plugins/Animation/ControlRig/Source/ControlRigEditor/Private/Editor/SRigSpacePickerWidget.cpp
+++ b/Engine/Plugins/Animation/ControlRig/Source/ControlRigEditor/Private/Editor/SRigSpacePickerWidget.cpp
@@ -357,6 +357,12 @@ TArray<FRigElementKey> SRigSpacePickerWidget::GetSpaceList(bool bIncludeDefaultS
 	return CurrentSpaceKeys;
 }
 
+void SRigSpacePickerWidget::RefreshContents()
+{
+	UpdateActiveSpaces();
+	RepopulateItemSpaces();
+}
+
 void SRigSpacePickerWidget::AddSpacePickerRow(
 	TSharedPtr<SVerticalBox> InListBox,
 	ESpacePickerType InType,
@@ -578,7 +584,7 @@ void SRigSpacePickerWidget::HandleSpaceDelete(FRigElementKey InKey)
 FReply SRigSpacePickerWidget::HandleAddElementClicked()
 {
 	FRigTreeDelegates TreeDelegates;
-	TreeDelegates.OnGetHierarchy = FOnGetRigTreeHierarchy::CreateSP(this, &SRigSpacePickerWidget::GetHierarchy);
+	TreeDelegates.OnGetHierarchy = FOnGetRigTreeHierarchy::CreateSP(this, &SRigSpacePickerWidget::GetHierarchyConst);
 	TreeDelegates.OnMouseButtonClick = FOnRigTreeMouseButtonClick::CreateLambda([this](TSharedPtr<FRigTreeElement> InItem)
 	{
 		if(InItem.IsValid())
@@ -927,26 +933,10 @@ void SRigSpacePickerWidget::UpdateActiveSpaces()
 	for(int32 ControlIndex=0;ControlIndex<ControlKeys.Num();ControlIndex++)
 	{
 		ActiveSpaceKeys.Add(URigHierarchy::GetDefaultParentSocketKey());
-		
-		const FRigElementKey& ControlKey = ControlKeys[ControlIndex];
-		const TArray<FRigElementWeight> ParentWeights = Hierarchy->GetParentWeightArray(ControlKey);
-		if(ParentWeights.Num() > 0)
-		{
-			const TArray<FRigElementKey> ParentKeys = Hierarchy->GetParents(ControlKey);
-			check(ParentKeys.Num() == ParentWeights.Num());
-			for(int32 ParentIndex=0;ParentIndex<ParentKeys.Num();ParentIndex++)
-			{
-				if(ParentWeights[ParentIndex].IsAlmostZero())
-				{
-					continue;
-				}
 
-				if(ParentIndex > 0 || IsDefaultSpace(ParentKeys[ParentIndex]))
-				{
-					ActiveSpaceKeys[ControlIndex] = ParentKeys[ParentIndex];
-				}
-				break;
-			}
+		if(GetActiveSpaceDelegate.IsBound())
+		{
+			ActiveSpaceKeys[ControlIndex] = GetActiveSpaceDelegate.Execute(Hierarchy, ControlKeys[ControlIndex]);
 		}
 	}
 }
@@ -973,4 +963,179 @@ bool SRigSpacePickerWidget::IsDefaultSpace(const FRigElementKey& InKey) const
 	return false;
 }
 
+//////////////////////////////////////////////////////////////
+/// SRigSpacePickerBakeWidget
+///////////////////////////////////////////////////////////
+
+void SRigSpacePickerBakeWidget::Construct(const FArguments& InArgs)
+{
+	check(InArgs._Hierarchy);
+	check(InArgs._Controls.Num() > 0);
+	check(InArgs._OnBake.IsBound());
+	
+	Settings = InArgs._Settings;
+
+	FStructureDetailsViewArgs StructureViewArgs;
+	StructureViewArgs.bShowObjects = true;
+	StructureViewArgs.bShowAssets = true;
+	StructureViewArgs.bShowClasses = true;
+	StructureViewArgs.bShowInterfaces = true;
+
+	FDetailsViewArgs ViewArgs;
+	ViewArgs.bAllowSearch = false;
+	ViewArgs.bHideSelectionTip = false;
+	ViewArgs.bShowObjectLabel = false;
+
+	TSharedPtr<FStructOnScope> StructToDisplay = MakeShareable(new FStructOnScope(FRigSpacePickerBakeSettings::StaticStruct(), (uint8*)&Settings));
+	FPropertyEditorModule& EditModule = FModuleManager::Get().GetModuleChecked<FPropertyEditorModule>("PropertyEditor");
+	DetailsView = EditModule.CreateStructureDetailView(ViewArgs, StructureViewArgs, StructToDisplay, LOCTEXT("Struct", "Struct View"));
+
+	ChildSlot
+	[
+		SNew(SBorder)
+		.Visibility(EVisibility::Visible)
+		[
+			SNew(SVerticalBox)
+
+			+SVerticalBox::Slot()
+			.AutoHeight()
+			[
+				SAssignNew(SpacePickerWidget, SRigSpacePickerWidget)
+				.Hierarchy(InArgs._Hierarchy)
+				.Controls(InArgs._Controls)
+				.AllowDelete(false)
+				.AllowReorder(false)
+				.AllowAdd(true)
+				.ShowBakeButton(false)
+				.GetControlCustomization_Lambda([this] (URigHierarchy*, const FRigElementKey)
+				{
+					return &Customization;
+				})
+				.OnSpaceListChanged_Lambda([this](URigHierarchy*, const FRigElementKey&, const TArray<FRigElementKey>& InSpaceList)
+				{
+					if(Customization.AvailableSpaces != InSpaceList)
+					{
+						Customization.AvailableSpaces = InSpaceList;
+						SpacePickerWidget->RefreshContents();
+					}
+				})
+				.GetActiveSpace_Lambda([this](URigHierarchy*, const FRigElementKey&)
+				{
+					return Settings.TargetSpace;
+				})
+				.OnActiveSpaceChanged_Lambda([this] (URigHierarchy*, const FRigElementKey&, const FRigElementKey InSpaceKey)
+				{
+					if(Settings.TargetSpace != InSpaceKey)
+					{
+						Settings.TargetSpace = InSpaceKey;	
+						SpacePickerWidget->RefreshContents();
+					}
+				})
+			]
+
+			+SVerticalBox::Slot()
+			.AutoHeight()
+			.Padding(0.f, 8.f, 0.f, 0.f)
+			[
+				DetailsView->GetWidget().ToSharedRef()
+			]
+
+			+SVerticalBox::Slot()
+			.AutoHeight()
+			.Padding(0.f, 16.f, 0.f, 0.f)
+			[
+				SNew(SHorizontalBox)
+
+				+SHorizontalBox::Slot()
+				.FillWidth(1.f)
+				[
+					SNew(SSpacer)
+				]
+
+				+SHorizontalBox::Slot()
+				.AutoWidth()
+				.HAlign(HAlign_Left)
+				.VAlign(VAlign_Center)
+				.Padding(8.f, 0.f, 0.f, 0.f)
+				[
+					SNew(SButton)
+					.HAlign(HAlign_Center)
+					.ContentPadding(FEditorStyle::GetMargin("StandardDialog.ContentPadding"))
+					.Text(LOCTEXT("OK", "OK"))
+					.OnClicked_Lambda([this, InArgs]()
+					{
+						return InArgs._OnBake.Execute(SpacePickerWidget->GetHierarchy(), SpacePickerWidget->GetControls(), Settings);
+					})
+					.IsEnabled_Lambda([this]()
+					{
+						return Settings.TargetSpace.IsValid();
+					})
+				]
+
+				+SHorizontalBox::Slot()
+				.AutoWidth()
+				.HAlign(HAlign_Left)
+				.VAlign(VAlign_Center)
+				.Padding(8.f, 0.f, 16.f, 0.f)
+				[
+					SNew(SButton)
+					.HAlign(HAlign_Center)
+					.ContentPadding(FEditorStyle::GetMargin("StandardDialog.ContentPadding"))
+					.Text(LOCTEXT("Cancel", "Cancel"))
+					.OnClicked_Lambda([this]()
+					{
+						CloseDialog();
+						return FReply::Handled();
+					})
+				]
+			]
+		]
+	];
+}
+
+FReply SRigSpacePickerBakeWidget::OpenDialog(bool bModal)
+{
+	check(!DialogWindow.IsValid());
+		
+	const FVector2D CursorPos = FSlateApplication::Get().GetCursorPos();
+
+	TSharedRef<SRigSpaceDialogWindow> Window = SNew(SRigSpaceDialogWindow)
+	.Title( LOCTEXT("SRigSpacePickerBakeWidgetTitle", "Bake Controls to new space") )
+	.CreateTitleBar(true)
+	.Type(EWindowType::Normal)
+	.SizingRule( ESizingRule::Autosized )
+	.ScreenPosition(CursorPos)
+	.FocusWhenFirstShown(true)
+	.ActivationPolicy(EWindowActivationPolicy::FirstShown)
+	[
+		AsShared()
+	];
+	
+	Window->SetWidgetToFocusOnActivate(AsShared());
+	
+	DialogWindow = Window;
+
+	Window->MoveWindowTo(CursorPos);
+
+	if(bModal)
+	{
+		GEditor->EditorAddModalWindow(Window);
+	}
+	else
+	{
+		FSlateApplication::Get().AddWindow( Window );
+	}
+
+	return FReply::Handled();
+}
+
+void SRigSpacePickerBakeWidget::CloseDialog()
+{
+	if ( DialogWindow.IsValid() )
+	{
+		DialogWindow.Pin()->RequestDestroyWindow();
+		DialogWindow.Reset();
+	}
+}
+
 #undef LOCTEXT_NAMESPACE
diff --git a/Engine/Plugins/Animation/ControlRig/Source/ControlRigEditor/Private/Editor/SRigSpacePickerWidget.h b/Engine/Plugins/Animation/ControlRig/Source/ControlRigEditor/Private/Editor/SRigSpacePickerWidget.h
--- a/Engine/Plugins/Animation/ControlRig/Source/ControlRigEditor/Private/Editor/SRigSpacePickerWidget.h
+++ b/Engine/Plugins/Animation/ControlRig/Source/ControlRigEditor/Private/Editor/SRigSpacePickerWidget.h
@@ -8,15 +8,44 @@
 #include "SRigHierarchy.h"
 #include "Widgets/Layout/SBox.h"
 #include "Rigs/RigSpaceHierarchy.h"
+#include "IStructureDetailsView.h"
+#include "SRigSpacePickerWidget.generated.h"
+
+USTRUCT()
+struct CONTROLRIGEDITOR_API FRigSpacePickerBakeSettings
+{
+	GENERATED_BODY();
+
+	FRigSpacePickerBakeSettings()
+	{
+		TargetSpace = FRigElementKey();
+		StartFrame = 1.f;
+		EndFrame = 100.f;
+		bKeyEveryFrame = true;
+	}
+
+	UPROPERTY()
+	FRigElementKey TargetSpace;
+
+	UPROPERTY(EditAnywhere, Category = "Settings")
+	float StartFrame;
+
+	UPROPERTY(EditAnywhere, Category = "Settings")
+	float EndFrame;
+
+	UPROPERTY(EditAnywhere, Category = "Settings")
+	bool bKeyEveryFrame;
+};
 
 DECLARE_DELEGATE_RetVal_TwoParams(FRigElementKey, FRigSpacePickerGetActiveSpace, URigHierarchy*, const FRigElementKey&);
 DECLARE_DELEGATE_RetVal_TwoParams(const FRigControlElementCustomization*, FRigSpacePickerGetControlCustomization, URigHierarchy*, const FRigElementKey&);
 DECLARE_EVENT_ThreeParams(SRigSpacePickerWidget, FRigSpacePickerActiveSpaceChanged, URigHierarchy*, const FRigElementKey&, const FRigElementKey&);
 DECLARE_EVENT_ThreeParams(SRigSpacePickerWidget, FRigSpacePickerSpaceListChanged, URigHierarchy*, const FRigElementKey&, const TArray<FRigElementKey>&);
 DECLARE_DELEGATE_RetVal_TwoParams(TArray<FRigElementKey>, FRigSpacePickerGetAdditionalSpaces, URigHierarchy*, const FRigElementKey&);
+DECLARE_DELEGATE_RetVal_ThreeParams(FReply, SRigSpacePickerOnBake, URigHierarchy*, TArray<FRigElementKey> /* Controls */, FRigSpacePickerBakeSettings);
 
 /** Widget allowing picking of a space source for space switching */
-class SRigSpacePickerWidget : public SCompoundWidget
+class CONTROLRIGEDITOR_API SRigSpacePickerWidget : public SCompoundWidget
 {
 public:
 
@@ -65,12 +94,16 @@ public:
 	virtual bool SupportsKeyboardFocus() const override;
 	virtual void Tick(const FGeometry& AllottedGeometry, const double InCurrentTime, const float InDeltaTime) override;
 
+	URigHierarchy* GetHierarchy() const { return Hierarchy; }
+	const URigHierarchy* GetHierarchyConst() const { return Hierarchy; }
+	const TArray<FRigElementKey>& GetControls() const { return ControlKeys; }
 	const TArray<FRigElementKey>& GetActiveSpaces() const;
 	TArray<FRigElementKey> GetDefaultSpaces() const;
 	TArray<FRigElementKey> GetSpaceList(bool bIncludeDefaultSpaces = false) const;
 	FRigSpacePickerActiveSpaceChanged& OnActiveSpaceChanged() { return ActiveSpaceChangedEvent; }
 	FRigSpacePickerSpaceListChanged& OnSpaceListChanged() { return SpaceListChangedEvent; }
-
+	void RefreshContents();
+	
 private:
 
 	enum ESpacePickerType
@@ -103,7 +136,6 @@ private:
 	FReply HandleAddElementClicked();
 	bool IsSpaceMoveUpEnabled(FRigElementKey InKey) const;
 	bool IsSpaceMoveDownEnabled(FRigElementKey InKey) const;
-	const URigHierarchy* GetHierarchy() const { return Hierarchy; }
 
 	void OnHierarchyModified(ERigHierarchyNotification InNotif, URigHierarchy* InHierarchy, const FRigBaseElement* InElement);
 
@@ -144,3 +176,36 @@ private:
 
 	static FRigElementKey InValidKey;
 };
+
+/** Widget allowing baking controls from one space to another */
+class CONTROLRIGEDITOR_API SRigSpacePickerBakeWidget : public SCompoundWidget
+{
+public:
+
+	SLATE_BEGIN_ARGS(SRigSpacePickerBakeWidget)
+	: _Hierarchy(nullptr)
+	, _Controls()
+	{}
+	SLATE_ARGUMENT(URigHierarchy*, Hierarchy)
+	SLATE_ARGUMENT(TArray<FRigElementKey>, Controls)
+	SLATE_EVENT(FRigSpacePickerGetControlCustomization, GetControlCustomization)
+	SLATE_ARGUMENT(FRigSpacePickerBakeSettings, Settings)
+	SLATE_EVENT(SRigSpacePickerOnBake, OnBake)
+
+	SLATE_END_ARGS()
+
+	void Construct(const FArguments& InArgs);
+	virtual ~SRigSpacePickerBakeWidget() override {}
+
+	FReply OpenDialog(bool bModal = true);
+	void CloseDialog();
+
+private:
+
+	FRigSpacePickerBakeSettings Settings;
+	FRigControlElementCustomization Customization;
+	
+	TWeakPtr<SWindow> DialogWindow;
+	TSharedPtr<SRigSpacePickerWidget> SpacePickerWidget;
+	TSharedPtr<IStructureDetailsView> DetailsView;
+};
