Control Rig: Controls named as 'None' disappear from the Rig Hierarchy

diff --git a/Engine/Plugins/Animation/ControlRig/Source/ControlRig/Private/Rigs/RigHierarchy.cpp b/Engine/Plugins/Animation/ControlRig/Source/ControlRig/Private/Rigs/RigHierarchy.cpp
--- a/Engine/Plugins/Animation/ControlRig/Source/ControlRig/Private/Rigs/RigHierarchy.cpp
+++ b/Engine/Plugins/Animation/ControlRig/Source/ControlRig/Private/Rigs/RigHierarchy.cpp
@@ -740,6 +740,15 @@ bool URigHierarchy::IsNameAvailable(const FString& InPotentialNewName, ERigEleme
 		return false;
 	}
 
+	if (UnsanitizedName == TEXT("None"))
+	{
+		if (OutErrorMessage)
+		{
+			*OutErrorMessage = TEXT("None is not a valid name.");
+		}
+		return false;
+	}
+
 	FString SanitizedName = UnsanitizedName;
 	SanitizeName(SanitizedName);
 
diff --git a/Engine/Plugins/Animation/ControlRig/Source/ControlRigEditor/Private/ControlRigElementDetails.cpp b/Engine/Plugins/Animation/ControlRig/Source/ControlRigEditor/Private/ControlRigElementDetails.cpp
--- a/Engine/Plugins/Animation/ControlRig/Source/ControlRigEditor/Private/ControlRigElementDetails.cpp
+++ b/Engine/Plugins/Animation/ControlRig/Source/ControlRigEditor/Private/ControlRigElementDetails.cpp
@@ -775,6 +775,7 @@ void FRigBaseElementDetails::CustomizeChildren(TSharedRef<class IPropertyHandle>
 		.Font(IDetailLayoutBuilder::GetDetailFont())
 		.Text(this, &FRigBaseElementDetails::GetName)
 		.OnTextCommitted(this, &FRigBaseElementDetails::SetName)
+		.OnVerifyTextChanged(this, &FRigBaseElementDetails::OnVerifyNameChanged)
 		.IsEnabled(ObjectsBeingCustomized.Num() == 1)
 	];
 }
@@ -823,6 +824,43 @@ void FRigBaseElementDetails::SetName(const FText& InNewText, ETextCommit::Type I
 	}
 }
 
+bool FRigBaseElementDetails::OnVerifyNameChanged(const FText& InText, FText& OutErrorMessage)
+{
+	if(ObjectsBeingCustomized.Num() > 1)
+	{
+		return false;
+	}
+
+	URigHierarchy* Hierarchy = nullptr;
+	if (BlueprintBeingCustomized)
+	{
+		Hierarchy = BlueprintBeingCustomized->Hierarchy;
+	}
+	else
+	{
+		Hierarchy = GetHierarchy();
+	}
+
+	if (!Hierarchy)
+	{
+		return false;
+	}
+
+	if (*GetElementKey().Name.ToString() == InText.ToString())
+	{
+		return true;
+	}
+
+	FString OutErrorMessageStr;
+	if (!Hierarchy->IsNameAvailable(InText.ToString(), GetElementKey().Type, &OutErrorMessageStr))
+	{
+		OutErrorMessage = FText::FromString(OutErrorMessageStr);
+		return false;
+	}
+
+	return true;
+}
+
 void FRigBaseElementDetails::OnStructContentsChanged(FProperty* InProperty, const TSharedRef<IPropertyUtilities> PropertyUtilities)
 {
 	const FPropertyChangedEvent ChangeEvent(InProperty, EPropertyChangeType::ValueSet);
diff --git a/Engine/Plugins/Animation/ControlRig/Source/ControlRigEditor/Private/ControlRigElementDetails.h b/Engine/Plugins/Animation/ControlRig/Source/ControlRigEditor/Private/ControlRigElementDetails.h
--- a/Engine/Plugins/Animation/ControlRig/Source/ControlRigEditor/Private/ControlRigElementDetails.h
+++ b/Engine/Plugins/Animation/ControlRig/Source/ControlRigEditor/Private/ControlRigElementDetails.h
@@ -132,6 +132,7 @@ public:
 	FRigElementKey GetElementKey() const;
 	FText GetName() const;
 	void SetName(const FText& InNewText, ETextCommit::Type InCommitType);
+	bool OnVerifyNameChanged(const FText& InText, FText& OutErrorMessage);
 
 	void OnStructContentsChanged(FProperty* InProperty, const TSharedRef<IPropertyUtilities> PropertyUtilities);
 	bool IsSetupModeEnabled() const;
