Control Rig: Added Warnings to AddMappedSkeletalMesh APIs, people should not use them during construction

diff --git a/Engine/Plugins/Animation/ControlRig/Source/ControlRig/Public/ControlRigComponent.h b/Engine/Plugins/Animation/ControlRig/Source/ControlRig/Public/ControlRigComponent.h
--- a/Engine/Plugins/Animation/ControlRig/Source/ControlRig/Public/ControlRigComponent.h
+++ b/Engine/Plugins/Animation/ControlRig/Source/ControlRig/Public/ControlRigComponent.h
@@ -271,24 +271,24 @@ public:
 	UFUNCTION(BlueprintCallable, Category = "ControlRig")
 	void ClearMappedElements();
 
-	/** Replaces the mapped elements on the component with the provided array */
-	UFUNCTION(BlueprintCallable, Category = "ControlRig")
+	/** Replaces the mapped elements on the component with the provided array, should not be used before OnPreInitialize Event */
+	UFUNCTION(BlueprintCallable, Category = "ControlRig", meta = (UnsafeDuringActorConstruction = "true"))
 	void SetMappedElements(TArray<FControlRigComponentMappedElement> NewMappedElements);
 
-	/** Adds the provided mapped elements to the component */
-	UFUNCTION(BlueprintCallable, Category = "ControlRig")
+	/** Adds the provided mapped elements to the component, should not be used before OnPreInitialize Event*/
+	UFUNCTION(BlueprintCallable, Category = "ControlRig", meta = (UnsafeDuringActorConstruction = "true"))
 	void AddMappedElements(TArray<FControlRigComponentMappedElement> NewMappedElements);
 
-	/** Adds a series of mapped bones to the rig */
-	UFUNCTION(BlueprintCallable, Category = "ControlRig")
+	/** Adds a series of mapped bones to the rig, should not be used before OnPreInitialize Event */
+	UFUNCTION(BlueprintCallable, Category = "ControlRig", meta = (UnsafeDuringActorConstruction = "true"))
 	void AddMappedComponents(TArray<FControlRigComponentMappedComponent> Components);
 
-	/** Adds a series of mapped bones to the rig */
-	UFUNCTION(BlueprintCallable, Category = "ControlRig", meta = (DisplayName = "Add Mapped Skeletal Mesh Bone Array"))
+	/** Adds a series of mapped bones to the rig, should not be used before OnPreInitialize Event */
+	UFUNCTION(BlueprintCallable, Category = "ControlRig", meta = (DisplayName = "Add Mapped Skeletal Mesh Bone Array", UnsafeDuringActorConstruction = "true"))
 	void AddMappedSkeletalMesh(USkeletalMeshComponent* SkeletalMeshComponent, TArray<FControlRigComponentMappedBone> Bones, TArray<FControlRigComponentMappedCurve> Curves);
 
-	/** Adds all matching bones to the rig */
-	UFUNCTION(BlueprintCallable, Category = "ControlRig", meta = (DisplayName = "Add Mapped Skeletal Mesh"))
+	/** Adds all matching bones to the rig, should not be used before OnPreInitialize Event */
+	UFUNCTION(BlueprintCallable, Category = "ControlRig", meta = (DisplayName = "Add Mapped Skeletal Mesh", UnsafeDuringActorConstruction = "true"))
 	void AddMappedCompleteSkeletalMesh(USkeletalMeshComponent* SkeletalMeshComponent);
 
 	/** Setup the initial transforms / ref pose of the bones based on a skeletal mesh */
