Control Rig: Display local variables in details panel

diff --git a/Engine/Plugins/Animation/ControlRig/Source/ControlRigEditor/Private/ControlRigLocalVariableDetails.cpp b/Engine/Plugins/Animation/ControlRig/Source/ControlRigEditor/Private/ControlRigLocalVariableDetails.cpp
--- a/Engine/Plugins/Animation/ControlRig/Source/ControlRigEditor/Private/ControlRigLocalVariableDetails.cpp
+++ b/Engine/Plugins/Animation/ControlRig/Source/ControlRigEditor/Private/ControlRigLocalVariableDetails.cpp
@@ -9,62 +9,11 @@
 #include "Graph/ControlRigGraphSchema.h"
 #include "Kismet2/BlueprintEditorUtils.h"
 #include "Widgets/Input/STextComboBox.h"
+#include "ControlRigBlueprintGeneratedClass.h"
+#include "RigVMCore/RigVM.h"
 
 #define LOCTEXT_NAMESPACE "LocalVariableDetails"
 
-void LocalVariableDetails_GetCustomizedInfo(TSharedRef<IPropertyHandle> InStructPropertyHandle, UControlRigBlueprint*& OutBlueprint)
-{
-	TArray<UObject*> Objects;
-	InStructPropertyHandle->GetOuterObjects(Objects);
-	for (UObject* Object : Objects)
-	{
-		if (Object->IsA<UControlRigBlueprint>())
-		{
-			OutBlueprint = Cast<UControlRigBlueprint>(Object);
-			if (OutBlueprint)
-			{
-				break;
-			}
-		}
-	}
-
-	if (OutBlueprint == nullptr)
-	{
-		TArray<UPackage*> Packages;
-		InStructPropertyHandle->GetOuterPackages(Packages);
-		for (UPackage* Package : Packages)
-		{
-			if (Package == nullptr)
-			{
-				continue;
-			}
-
-			TArray<UObject*> SubObjects;
-			Package->GetDefaultSubobjects(SubObjects);
-			for (UObject* SubObject : SubObjects)
-			{
-				if (UControlRig* Rig = Cast<UControlRig>(SubObject))
-				{
-					UControlRigBlueprint* Blueprint = Cast<UControlRigBlueprint>(Rig->GetClass()->ClassGeneratedBy);
-					if (Blueprint)
-					{
-						if(Blueprint->GetOutermost() == Package)
-						{
-							OutBlueprint = Blueprint;
-							break;
-						}
-					}
-				}
-			}
-
-			if (OutBlueprint)
-			{
-				break;
-			}
-		}
-	}
-}
-
 void FRigVMLocalVariableDetails::CustomizeHeader(TSharedRef<IPropertyHandle> InStructPropertyHandle, FDetailWidgetRow& HeaderRow, IPropertyTypeCustomizationUtils& StructCustomizationUtils)
 {
 	ObjectsBeingCustomized.Reset();
@@ -81,6 +30,7 @@ void FRigVMLocalVariableDetails::CustomizeHeader(TSharedRef<IPropertyHandle> InS
 	{
 		VariableDescription = *ObjectsBeingCustomized[0]->GetContent<FRigVMGraphVariableDescription>();
 		GraphBeingCustomized = ObjectsBeingCustomized[0]->GetTypedOuter<URigVMGraph>();
+		BlueprintBeingCustomized = GraphBeingCustomized->GetTypedOuter<UControlRigBlueprint>();
 	}
 }
 
@@ -131,98 +81,39 @@ void FRigVMLocalVariableDetails::CustomizeChildren(TSharedRef<IPropertyHandle> I
 		];
 
 
-	IDetailCategoryBuilder& DefaultValueCategory = StructBuilder.GetParentCategory().GetParentLayout().EditCategory(TEXT("DefaultValueCategory"), LOCTEXT("DefaultValueCategoryHeading", "Default Value"));
-
-	if (!VariableDescription.CPPTypeObject)
+	if (BlueprintBeingCustomized)
 	{
-		if (VariableDescription.CPPType == "bool")
-		{
-			DefaultValueCategory.AddCustomRow(LOCTEXT("DefaultValue", "Default Value"))
-			.NameContent()
-			[
-				SNew(STextBlock)
-				.Text(LOCTEXT("DefaultValue", "Default Value"))
-				.Font(DetailFontInfo)
-			]
-			.ValueContent()
-			.MaxDesiredWidth(980.f)
-			[
-				SNew(SCheckBox)
-				.IsChecked(this, &FRigVMLocalVariableDetails::HandleBoolDefaultValueIsChecked)
-				.OnCheckStateChanged(this, &FRigVMLocalVariableDetails::OnBoolDefaultValueChanged)
-			];
-		}
-		else
+		UControlRigBlueprintGeneratedClass* RigClass = BlueprintBeingCustomized->GetControlRigBlueprintGeneratedClass();
+		UControlRig* CDO = Cast<UControlRig>(RigClass->GetDefaultObject(true /* create if needed */));
+		if (CDO->GetVM() != nullptr)
 		{
-			DefaultValueCategory.AddCustomRow( LOCTEXT("LocalVariableDefaultValue", "Default Value") )
-			.NameContent()
-			[
-				SNew(STextBlock)
-				.Text(LOCTEXT("LocalVariableDefaultValue", "Default Value"))
-				.Font(IDetailLayoutBuilder::GetDetailFont())
-			]
-			.ValueContent()
-			.MaxDesiredWidth(250.0f)
-			[
-				DefaultValueHandle->CreatePropertyValueWidget()
-			];
-		}
-	}
-	else
-	{
-		if(UEnum* Enum = Cast<UEnum>(VariableDescription.CPPTypeObject))
-		{
-			int32 CurrentValueIndex = 0;
-			for (int32 i = 0; i < Enum->NumEnums()-1; i++)
-			{
-				if (!Enum->HasMetaData(TEXT("Hidden"), i))
-				{
-					FString DisplayName = Enum->GetDisplayNameTextByIndex(i).ToString();
-					EnumOptions.Add(MakeShareable(new FString(DisplayName)));
-					if (DisplayName == VariableDescription.DefaultValue)
-					{
-						CurrentValueIndex = i;
-					}					
-				}
-			}
+#if UE_RIGVM_UCLASS_BASED_STORAGE_DISABLED
 			
-			DefaultValueCategory.AddCustomRow(LOCTEXT("VariableReplicationConditionsLabel", "Replication Condition"))
-			.NameContent()
-			[
-				SNew(STextBlock)
-				.Text(LOCTEXT("LocalVariableDefaultValue", "Default Value"))
-				.Font(DetailFontInfo)
-			]
-			.ValueContent()
-			[
-				SNew(STextComboBox)
-				.OptionsSource(&EnumOptions)
-				.InitiallySelectedItem(EnumOptions[CurrentValueIndex])
-				.OnSelectionChanged_Lambda([this](TSharedPtr<FString> InValue, ESelectInfo::Type)
-				{
-					VariableDescription.DefaultValue = *InValue;
-					DefaultValueHandle->SetValue(VariableDescription.DefaultValue);					
-				})
-			];			
-		}
-		else if(UScriptStruct* ScriptStruct = Cast<UScriptStruct>(VariableDescription.CPPTypeObject))
-		{
-			TSharedPtr<FStructOnScope> StructOnScope = MakeShareable(new FStructOnScope(ScriptStruct));
-			IDetailPropertyRow* Row = DefaultValueCategory.AddExternalStructureProperty(StructOnScope, NAME_None);
-			Row->GetPropertyHandle()->SetValueFromFormattedString(VariableDescription.DefaultValue);
+#else
+			FString SourcePath = FString::Printf(TEXT("LocalVariableDefault::%s|%s::Const"), *GraphBeingCustomized->GetGraphName(), *VariableDescription.Name.ToString());
+			URigVMMemoryStorage* LiteralMemory = CDO->GetVM()->GetLiteralMemory();
+			FProperty* Property = LiteralMemory->FindPropertyByName(*SourcePath);
+			if (Property)
+			{
+				IDetailCategoryBuilder& DefaultValueCategory = StructBuilder.GetParentCategory().GetParentLayout().EditCategory(TEXT("DefaultValueCategory"), LOCTEXT("DefaultValueCategoryHeading", "Default Value"));
+				Property->ClearPropertyFlags(CPF_EditConst);
 			
-			Row->GetPropertyHandle()->SetOnChildPropertyValueChanged(FSimpleDelegate::CreateLambda(
-				[this, Row]()
+				const FName SanitizedName = FRigVMPropertyDescription::SanitizeName(*SourcePath);
+				TArray<UObject*> Objects = {LiteralMemory};
+				IDetailPropertyRow* Row = DefaultValueCategory.AddExternalObjectProperty(Objects, SanitizedName);
+				Row->DisplayName(FText::FromName(VariableDescription.Name));
+
+				const FSimpleDelegate OnDefaultValueChanged = FSimpleDelegate::CreateLambda([this, Property, LiteralMemory]()
 				{
-					VariableDescription.DefaultValue.Reset();
-					Row->GetPropertyHandle()->GetValueAsFormattedString(VariableDescription.DefaultValue);
+					VariableDescription.DefaultValue = LiteralMemory->GetDataAsString(LiteralMemory->GetPropertyIndex(Property));
 					DefaultValueHandle->SetValue(VariableDescription.DefaultValue);
-				}
-			));
-		}
-		else
-		{
-			checkNoEntry();
+				});
+
+				TSharedPtr<IPropertyHandle> Handle = Row->GetPropertyHandle();
+				Handle->SetOnPropertyValueChanged(OnDefaultValueChanged);
+				Handle->SetOnChildPropertyValueChanged(OnDefaultValueChanged);
+			}
+#endif
 		}
 	}
 }
@@ -239,8 +130,10 @@ FEdGraphPinType FRigVMLocalVariableDetails::OnGetPinInfo() const
 void FRigVMLocalVariableDetails::HandlePinInfoChanged(const FEdGraphPinType& PinType)
 {
 	VariableDescription.ChangeType(PinType);
+	BlueprintBeingCustomized->IncrementVMRecompileBracket();
 	TypeHandle->SetValue(VariableDescription.CPPType);
-	TypeObjectHandle->SetValue(VariableDescription.CPPTypeObject);
+	TypeObjectHandle->SetValue(VariableDescription.CPPTypeObject);	
+	BlueprintBeingCustomized->DecrementVMRecompileBracket();
 }
 
 ECheckBoxState FRigVMLocalVariableDetails::HandleBoolDefaultValueIsChecked() const
diff --git a/Engine/Plugins/Animation/ControlRig/Source/ControlRigEditor/Private/ControlRigLocalVariableDetails.h b/Engine/Plugins/Animation/ControlRig/Source/ControlRigEditor/Private/ControlRigLocalVariableDetails.h
--- a/Engine/Plugins/Animation/ControlRig/Source/ControlRigEditor/Private/ControlRigLocalVariableDetails.h
+++ b/Engine/Plugins/Animation/ControlRig/Source/ControlRigEditor/Private/ControlRigLocalVariableDetails.h
@@ -18,6 +18,12 @@ class IPropertyHandle;
 
 class FRigVMLocalVariableDetails : public IPropertyTypeCustomization
 {
+	FRigVMLocalVariableDetails()
+	: GraphBeingCustomized(nullptr)
+	, BlueprintBeingCustomized(nullptr)
+	{}
+
+	
 public:
 
 	// Makes a new instance of this detail layout class for a specific detail view requesting it
@@ -34,6 +40,7 @@ public:
 private:
 
 	URigVMGraph* GraphBeingCustomized;
+	UControlRigBlueprint* BlueprintBeingCustomized;
 	FRigVMGraphVariableDescription VariableDescription;
 	TArray<TWeakObjectPtr<UDetailsViewWrapperObject>> ObjectsBeingCustomized;
 
diff --git a/Engine/Plugins/Animation/ControlRig/Source/ControlRigEditor/Private/Editor/ControlRigEditor.cpp b/Engine/Plugins/Animation/ControlRig/Source/ControlRigEditor/Private/Editor/ControlRigEditor.cpp
--- a/Engine/Plugins/Animation/ControlRig/Source/ControlRigEditor/Private/Editor/ControlRigEditor.cpp
+++ b/Engine/Plugins/Animation/ControlRig/Source/ControlRigEditor/Private/Editor/ControlRigEditor.cpp
@@ -1576,6 +1576,25 @@ void FControlRigEditor::SetDetailViewForFocusedGraph()
 	SetDetailViewForGraph(Model);
 }
 
+void FControlRigEditor::SetDetailViewForLocalVariable()
+{
+	FName VariableName;
+	TArray< TWeakObjectPtr<UObject> > SelectedObjects = Inspector->GetSelectedObjects();
+	for (TWeakObjectPtr<UObject> SelectedObject : SelectedObjects)
+	{
+		if (SelectedObject.IsValid())
+		{
+			if(UDetailsViewWrapperObject* WrapperObject = Cast<UDetailsViewWrapperObject>(SelectedObject.Get()))
+			{
+				VariableName = WrapperObject->GetContent<FRigVMGraphVariableDescription>()->Name;
+				break;
+			}
+		}
+	}
+		
+	SelectLocalVariable(GetFocusedGraph(), VariableName);
+}
+
 void FControlRigEditor::RefreshDetailView()
 {
 	if(DetailViewShowsAnyRigElement())
@@ -1586,6 +1605,10 @@ void FControlRigEditor::RefreshDetailView()
 	{
 		SetDetailViewForFocusedGraph();
 	}
+	else if(DetailViewShowsLocalVariable())
+	{
+		SetDetailViewForLocalVariable();	
+	}
 }
 
 bool FControlRigEditor::DetailViewShowsAnyRigElement() const
@@ -1598,6 +1621,11 @@ bool FControlRigEditor::DetailViewShowsAnyRigUnit() const
 	return DetailViewShowsStruct(FRigUnit::StaticStruct());
 }
 
+bool FControlRigEditor::DetailViewShowsLocalVariable() const
+{
+	return DetailViewShowsStruct(FRigVMGraphVariableDescription::StaticStruct());
+}
+
 bool FControlRigEditor::DetailViewShowsStruct(UScriptStruct* InStruct) const
 {
 	TArray< TWeakObjectPtr<UObject> > SelectedObjects = Inspector->GetSelectedObjects();
@@ -1677,6 +1705,7 @@ void FControlRigEditor::ClearDetailObject(bool bChangeUISelectionState)
 	}
 	WrapperObjects.Reset();
 	
+	Inspector->GetPropertyView()->SetObjects(TArray<UObject*>(), true); // clear property view synchronously
 	Inspector->ShowDetailsForObjects(TArray<UObject*>());
 	Inspector->ShowSingleStruct(TSharedPtr<FStructOnScope>());
 
@@ -4475,11 +4504,13 @@ void FControlRigEditor::OnWrappedPropertyChangedChainEvent(UDetailsViewWrapperOb
 						if (Variable.Name == OldVariableName)
 						{
 							Controller->RenameLocalVariable(OldVariableName, VariableDescription->Name);
+							break;
 						}
 					}
 				}
 			}
-			RefreshMyBlueprint();	
+			RefreshMyBlueprint();
+			GetControlRigBlueprint()->RequestAutoVMRecompilation();
 		}
 		else if (PropertyPath == TEXT("CPPType") || PropertyPath == TEXT("CPPTypeObject"))
 		{			
@@ -4488,8 +4519,10 @@ void FControlRigEditor::OnWrappedPropertyChangedChainEvent(UDetailsViewWrapperOb
 				if (Variable.Name == VariableDescription->Name)
 				{
 					Controller->SetLocalVariableType(Variable.Name, VariableDescription->CPPType, VariableDescription->CPPTypeObject);
+					break;
 				}
 			}
+			GetControlRigBlueprint()->RequestAutoVMRecompilation();
 		}
 		else if (PropertyPath == TEXT("DefaultValue"))
 		{			
@@ -4497,12 +4530,14 @@ void FControlRigEditor::OnWrappedPropertyChangedChainEvent(UDetailsViewWrapperOb
 			{
 				if (Variable.Name == VariableDescription->Name)
 				{
-					Controller->SetLocalVariableDefaultValue(Variable.Name, VariableDescription->DefaultValue);
+					Controller->SetLocalVariableDefaultValue(Variable.Name, VariableDescription->DefaultValue, true, true, false);
+					break;
 				}
 			}
-		}
-		
-		
+
+			// Do not recompile now! That destroys the object that is currently being displayed (the literal memory storage), and can cause a crash.
+			// The user has to manually trigger the recompilation.
+		}		
 	}
 }
 
diff --git a/Engine/Plugins/Animation/ControlRig/Source/ControlRigEditor/Private/Editor/ControlRigEditor.h b/Engine/Plugins/Animation/ControlRig/Source/ControlRigEditor/Private/Editor/ControlRigEditor.h
--- a/Engine/Plugins/Animation/ControlRig/Source/ControlRigEditor/Private/Editor/ControlRigEditor.h
+++ b/Engine/Plugins/Animation/ControlRig/Source/ControlRigEditor/Private/Editor/ControlRigEditor.h
@@ -161,9 +161,11 @@ public:
 	void SetDetailViewForRigElements();
 	void SetDetailViewForGraph(URigVMGraph* InGraph);
 	void SetDetailViewForFocusedGraph();
+	void SetDetailViewForLocalVariable();
 	void RefreshDetailView();
 	bool DetailViewShowsAnyRigElement() const;
 	bool DetailViewShowsAnyRigUnit() const;
+	bool DetailViewShowsLocalVariable() const;
 	bool DetailViewShowsStruct(UScriptStruct* InStruct) const;
 	bool DetailViewShowsRigElement(FRigElementKey InKey) const;
 	bool DetailViewShowsRigUnit(URigVMNode* InNode) const;
diff --git a/Engine/Source/Developer/RigVMDeveloper/Private/RigVMCompiler/RigVMCompiler.cpp b/Engine/Source/Developer/RigVMDeveloper/Private/RigVMCompiler/RigVMCompiler.cpp
--- a/Engine/Source/Developer/RigVMDeveloper/Private/RigVMCompiler/RigVMCompiler.cpp
+++ b/Engine/Source/Developer/RigVMDeveloper/Private/RigVMCompiler/RigVMCompiler.cpp
@@ -260,6 +260,20 @@ bool URigVMCompiler::Compile(URigVMGraph* InGraph, URigVMController* InControlle
 
 	UE_LOG_RIGVMMEMORY(TEXT("RigVMCompiler: Begin '%s'..."), *InGraph->GetPathName());
 
+#if WITH_EDITOR
+	// If in editor, make sure we visit all the graphs to initialize local variables
+	// in case the user wants to edit default values
+	URigVMFunctionLibrary* FunctionLibrary = InGraph->GetDefaultFunctionLibrary();
+	for (URigVMLibraryNode* LibraryNode : FunctionLibrary->GetFunctions())
+	{
+		for (FRigVMGraphVariableDescription& Variable : LibraryNode->GetContainedGraph()->LocalVariables)
+		{
+			FString Path = FString::Printf(TEXT("LocalVariableDefault::%s|%s::Const"), *LibraryNode->GetFName().ToString(), *Variable.Name.ToString());
+			FRigVMOperand Operand = WorkData.AddProperty(ERigVMMemoryType::Literal, *Path, Variable.CPPType, Variable.CPPTypeObject, Variable.DefaultValue);
+			WorkData.PinPathToOperand->Add(Path, Operand);
+		}
+	}
+#endif
 
 	// Look for all local variables to create the register with the default value in the literal memory
 	int32 IndexLocalVariable = 0;
diff --git a/Engine/Source/Developer/RigVMDeveloper/Private/RigVMModel/RigVMController.cpp b/Engine/Source/Developer/RigVMDeveloper/Private/RigVMModel/RigVMController.cpp
--- a/Engine/Source/Developer/RigVMDeveloper/Private/RigVMModel/RigVMController.cpp
+++ b/Engine/Source/Developer/RigVMDeveloper/Private/RigVMModel/RigVMController.cpp
@@ -5571,11 +5571,13 @@ bool URigVMController::SetPinDefaultValue(const FString& InPinPath, const FStrin
 	return true;
 }
 
-bool URigVMController::SetPinDefaultValue(URigVMPin* InPin, const FString& InDefaultValue, bool bResizeArrays, bool bSetupUndoRedo, bool bMergeUndoAction)
+bool URigVMController::SetPinDefaultValue(URigVMPin* InPin, const FString& InDefaultValue, bool bResizeArrays, bool bSetupUndoRedo, bool bMergeUndoAction, bool bNotify)
 {
 	check(InPin);
 	ensure(!InDefaultValue.IsEmpty());
 
+	TGuardValue<bool> Guard(bSuspendNotifications, !bNotify);
+
 	if (InPin->GetDirection() == ERigVMPinDirection::Hidden)
 	{
 		return false;
@@ -8077,7 +8079,7 @@ bool URigVMController::SetLocalVariableTypeFromObjectPath(const FName& InVariabl
 	return SetLocalVariableType(InVariableName, InCPPType, CPPTypeObject, bSetupUndoRedo, bPrintPythonCommand);
 }
 
-bool URigVMController::SetLocalVariableDefaultValue(const FName& InVariableName, const FString& InDefaultValue, bool bSetupUndoRedo, bool bPrintPythonCommand)
+bool URigVMController::SetLocalVariableDefaultValue(const FName& InVariableName, const FString& InDefaultValue, bool bSetupUndoRedo, bool bPrintPythonCommand, bool bNotify)
 {
 	if (!IsValidGraph())
 	{
@@ -8112,8 +8114,9 @@ bool URigVMController::SetLocalVariableDefaultValue(const FName& InVariableName,
 		ActionStack->AddAction(FRigVMChangeLocalVariableDefaultValueAction(LocalVariables[FoundIndex], InDefaultValue));
 		ActionStack->EndAction(InverseAction);
 	}	
-	
-	LocalVariables[FoundIndex].DefaultValue = InDefaultValue;
+
+	FRigVMGraphVariableDescription& VariableDescription = LocalVariables[FoundIndex];
+	VariableDescription.DefaultValue = InDefaultValue;
 	
 	// Refresh variable nodes to reflect change in default value
 	TArray<URigVMNode*> Nodes = Graph->GetNodes();
@@ -8125,9 +8128,7 @@ bool URigVMController::SetLocalVariableDefaultValue(const FName& InVariableName,
 			{
 				if (VariablePin->GetDefaultValue() == InVariableName.ToString())
 				{
-					VariableNode->FindPin(URigVMVariableNode::ValueName)->DefaultValue = InDefaultValue;
-					Notify(ERigVMGraphNotifType::VariableRenamed, VariableNode);
-					continue;
+					SetPinDefaultValue(VariableNode->FindPin(URigVMVariableNode::ValueName), InDefaultValue, true, true, true, bNotify);
 				}
 			}
 		}
@@ -11128,6 +11129,11 @@ bool URigVMController::ChangePinType(URigVMPin* InPin, const FString& InCPPType,
 		return false;
 	}
 
+	if (FRigVMPropertyDescription::RequiresCPPTypeObject(InCPPType) && !CPPTypeObject)
+	{
+		return false;
+	}
+
 	if (CPPTypeObject)
 	{
 		// for now we don't allow UObjects
diff --git a/Engine/Source/Developer/RigVMDeveloper/Public/RigVMModel/RigVMController.h b/Engine/Source/Developer/RigVMDeveloper/Public/RigVMModel/RigVMController.h
--- a/Engine/Source/Developer/RigVMDeveloper/Public/RigVMModel/RigVMController.h
+++ b/Engine/Source/Developer/RigVMDeveloper/Public/RigVMModel/RigVMController.h
@@ -655,7 +655,7 @@ public:
 	bool SetLocalVariableTypeFromObjectPath(const FName& InVariableName, const FString& InCPPType, const FString& InCPPTypeObjectPath, bool bSetupUndoRedo = true, bool bPrintPythonCommand = false);
 
 	UFUNCTION(BlueprintCallable, Category = RigVMController)
-	bool SetLocalVariableDefaultValue(const FName& InVariableName, const FString& InDefaultValue, bool bSetupUndoRedo = true, bool bPrintPythonCommand = false);
+	bool SetLocalVariableDefaultValue(const FName& InVariableName, const FString& InDefaultValue, bool bSetupUndoRedo = true, bool bPrintPythonCommand = false, bool bNotify = true);
 
 	// Determine affected function references for a potential bulk edit on a library node
 	TArray<TSoftObjectPtr<URigVMFunctionReferenceNode>> GetAffectedReferences(ERigVMControllerBulkEditType InEditType, bool bForceLoad = false, bool bNotify = true);
@@ -777,7 +777,7 @@ private:
 	void ConfigurePinFromPin(URigVMPin* InOutPin, URigVMPin* InPin);
 	virtual bool ShouldStructBeUnfolded(const UStruct* InStruct);
 	virtual bool ShouldPinBeUnfolded(URigVMPin* InPin);
-	bool SetPinDefaultValue(URigVMPin* InPin, const FString& InDefaultValue, bool bResizeArrays, bool bSetupUndoRedo, bool bMergeUndoAction);
+	bool SetPinDefaultValue(URigVMPin* InPin, const FString& InDefaultValue, bool bResizeArrays, bool bSetupUndoRedo, bool bMergeUndoAction, bool bNotify = true);
 	bool ResetPinDefaultValue(URigVMPin* InPin, bool bSetupUndoRedo);
 	URigVMPin* InsertArrayPin(URigVMPin* ArrayPin, int32 InIndex, const FString& InDefaultValue, bool bSetupUndoRedo);
 	bool RemovePin(URigVMPin* InPinToRemove, bool bSetupUndoRedo, bool bNotify);
diff --git a/Engine/Source/Runtime/RigVM/Private/RigVMCore/RigVMMemoryStorage.cpp b/Engine/Source/Runtime/RigVM/Private/RigVMCore/RigVMMemoryStorage.cpp
--- a/Engine/Source/Runtime/RigVM/Private/RigVMCore/RigVMMemoryStorage.cpp
+++ b/Engine/Source/Runtime/RigVM/Private/RigVMCore/RigVMMemoryStorage.cpp
@@ -907,7 +907,7 @@ int32 URigVMMemoryStorage::GetPropertyIndexByName(const FName& InName) const
 	return GetPropertyIndex(Property);
 }
 
-const FProperty* URigVMMemoryStorage::FindPropertyByName(const FName& InName) const
+FProperty* URigVMMemoryStorage::FindPropertyByName(const FName& InName) const
 {
 	const FName SanitizedName = FRigVMPropertyDescription::SanitizeName(InName);
 	return GetClass()->FindPropertyByName(SanitizedName);
diff --git a/Engine/Source/Runtime/RigVM/Public/RigVMCore/RigVMMemoryStorage.h b/Engine/Source/Runtime/RigVM/Public/RigVMCore/RigVMMemoryStorage.h
--- a/Engine/Source/Runtime/RigVM/Public/RigVMCore/RigVMMemoryStorage.h
+++ b/Engine/Source/Runtime/RigVM/Public/RigVMCore/RigVMMemoryStorage.h
@@ -353,7 +353,7 @@ public:
 	}
 
 	// Returns a property given its name (or nullptr if the name wasn't found)
-	const FProperty* FindPropertyByName(const FName& InName) const;
+	FProperty* FindPropertyByName(const FName& InName) const;
 
 	// Creates and returns a new operand for a property (and optionally a property path)
 	FRigVMOperand GetOperand(int32 InPropertyIndex, int32 InPropertyPathIndex = INDEX_NONE) const;
