Control Rig: Consolidate SRigHierarchy / SRigCurveContainer on close events / handling

diff --git a/Engine/Plugins/Animation/ControlRig/Source/ControlRigEditor/Private/Editor/SRigCurveContainer.cpp b/Engine/Plugins/Animation/ControlRig/Source/ControlRigEditor/Private/Editor/SRigCurveContainer.cpp
--- a/Engine/Plugins/Animation/ControlRig/Source/ControlRigEditor/Private/Editor/SRigCurveContainer.cpp
+++ b/Engine/Plugins/Animation/ControlRig/Source/ControlRigEditor/Private/Editor/SRigCurveContainer.cpp
@@ -145,14 +145,6 @@ void SRigCurveContainer::Construct(const FArguments& InArgs, TSharedRef<FControl
 	ControlRigBlueprint = InControlRigEditor.Get().GetControlRigBlueprint();
 	bIsChangingRigHierarchy = false;
 
-#if WITH_EDITOR
-	GEditor->OnEditorClose().AddLambda([&]()
-	{
-		const FControlRigEditor* Editor = ControlRigEditor.IsValid() ? ControlRigEditor.Pin().Get() : nullptr;
-		OnEditorClose(Editor, ControlRigBlueprint.Get());
-	});
-#endif
-
 	ControlRigBlueprint->Hierarchy->OnModified().AddRaw(this, &SRigCurveContainer::OnHierarchyModified);
 	ControlRigBlueprint->OnRefreshEditor().AddRaw(this, &SRigCurveContainer::HandleRefreshEditorFromBlueprint);
 
@@ -214,9 +206,6 @@ void SRigCurveContainer::Construct(const FArguments& InArgs, TSharedRef<FControl
 
 SRigCurveContainer::~SRigCurveContainer()
 {
-	// Make sure we don't get called on editor exit if we've already been shut down.
-	GEditor->OnEditorClose().RemoveAll(this);
-
 	const FControlRigEditor* Editor = ControlRigEditor.IsValid() ? ControlRigEditor.Pin().Get() : nullptr;
 	OnEditorClose(Editor, ControlRigBlueprint.Get());
 }
