Control Rig: Fix loading order for limits

diff --git a/Engine/Plugins/Animation/ControlRig/Source/ControlRig/Private/Rigs/RigHierarchyElements.cpp b/Engine/Plugins/Animation/ControlRig/Source/ControlRig/Private/Rigs/RigHierarchyElements.cpp
--- a/Engine/Plugins/Animation/ControlRig/Source/ControlRig/Private/Rigs/RigHierarchyElements.cpp
+++ b/Engine/Plugins/Animation/ControlRig/Source/ControlRig/Private/Rigs/RigHierarchyElements.cpp
@@ -482,6 +482,7 @@ void FRigControlSettings::Load(FArchive& Ar)
 	Ar << bLimitScale;
 	Ar << bDrawLimits;
 
+	FTransform MinimumTransform, MaximumTransform;
 	if (Ar.CustomVer(FControlRigObjectVersion::GUID) >= FControlRigObjectVersion::StorageMinMaxValuesAsFloatStorage)
 	{
 		Ar << MinimumValue;
@@ -489,13 +490,9 @@ void FRigControlSettings::Load(FArchive& Ar)
 	}
 	else
 	{
-		FTransform MinimumTransform, MaximumTransform;
 		Ar << MinimumTransform;
 		Ar << MaximumTransform;
-		MinimumValue.SetFromTransform(MinimumTransform, ControlType, PrimaryAxis);
-		MaximumValue.SetFromTransform(MaximumTransform, ControlType, PrimaryAxis);
 	}
-	
 
 	Ar << bGizmoEnabled;
 	Ar << bGizmoVisible;
@@ -507,6 +504,12 @@ void FRigControlSettings::Load(FArchive& Ar)
 	ControlType = (ERigControlType)ControlTypeEnum->GetValueByName(ControlTypeName);
 	PrimaryAxis = (ERigControlAxis)ControlAxisEnum->GetValueByName(PrimaryAxisName);
 
+	if (Ar.CustomVer(FControlRigObjectVersion::GUID) < FControlRigObjectVersion::StorageMinMaxValuesAsFloatStorage)
+	{
+		MinimumValue.SetFromTransform(MinimumTransform, ControlType, PrimaryAxis);
+		MaximumValue.SetFromTransform(MaximumTransform, ControlType, PrimaryAxis);
+	}
+
 	ControlEnum = nullptr;
 	if(!ControlEnumPathName.IsEmpty())
 	{
