Control Rig: Create local variable node through tab menu

diff --git a/Engine/Plugins/Animation/ControlRig/Source/ControlRigEditor/Private/ControlRigEditorModule.cpp b/Engine/Plugins/Animation/ControlRig/Source/ControlRigEditor/Private/ControlRigEditorModule.cpp
--- a/Engine/Plugins/Animation/ControlRig/Source/ControlRigEditor/Private/ControlRigEditorModule.cpp
+++ b/Engine/Plugins/Animation/ControlRig/Source/ControlRigEditor/Private/ControlRigEditorModule.cpp
@@ -45,6 +45,7 @@
 #include "ControlRigGizmoLibraryActions.h"
 #include "Graph/ControlRigGraphSchema.h"
 #include "Graph/ControlRigGraph.h"
+#include "EdGraph/EdGraph.h"
 #include "Graph/NodeSpawners/ControlRigUnitNodeSpawner.h"
 #include "Graph/NodeSpawners/ControlRigVariableNodeSpawner.h"
 #include "Graph/NodeSpawners/ControlRigRerouteNodeSpawner.h"
@@ -83,6 +84,7 @@
 #include "ControlRigThumbnailRenderer.h"
 #include "RigVMModel/RigVMController.h"
 #include "RigVMModel/Nodes/RigVMLibraryNode.h"
+#include "RigVMModel/RigVMVariableDescription.h"
 #include "ControlRigBlueprint.h"
 #include "ControlRig/Private/Units/Simulation/RigUnit_AlphaInterp.h"
 #include "ControlRig/Private/Units/Debug/RigUnit_VisualDebug.h"
@@ -1253,6 +1255,21 @@ void FControlRigEditorModule::GetInstanceActions(UControlRigBlueprint* CRB, FBlu
 				check(NodeSpawner != nullptr);
 				ActionRegistrar.AddBlueprintAction(GeneratedClass, NodeSpawner);
 			}
+
+			static const FString CategoryDelimiter(TEXT("|"));
+			FText NodeCategory = LOCTEXT("LocalVariables", "Local Variables");
+			for (URigVMLibraryNode* Function : Functions)
+			{
+				for (const FRigVMGraphVariableDescription& LocalVariable : Function->GetContainedGraph()->GetLocalVariables())
+				{
+					FText MenuDesc = FText::FromName(LocalVariable.Name);
+					FText ToolTip = FText::FromString(FString::Printf(TEXT("Get the value of variable %s"), *LocalVariable.Name.ToString()));
+					ActionRegistrar.AddBlueprintAction(GeneratedClass, UControlRigVariableNodeSpawner::CreateFromLocalVariable(CRB, Function->GetContainedGraph(), LocalVariable, true, MenuDesc, NodeCategory, ToolTip));
+
+					ToolTip = FText::FromString(FString::Printf(TEXT("Set the value of variable %s"), *LocalVariable.Name.ToString()));
+					ActionRegistrar.AddBlueprintAction(GeneratedClass, UControlRigVariableNodeSpawner::CreateFromLocalVariable(CRB, Function->GetContainedGraph(), LocalVariable, false, MenuDesc, NodeCategory, ToolTip));
+				}
+			}
 		}
 	}
 }
diff --git a/Engine/Plugins/Animation/ControlRig/Source/ControlRigEditor/Private/Graph/NodeSpawners/ControlRigVariableNodeSpawner.cpp b/Engine/Plugins/Animation/ControlRig/Source/ControlRigEditor/Private/Graph/NodeSpawners/ControlRigVariableNodeSpawner.cpp
--- a/Engine/Plugins/Animation/ControlRig/Source/ControlRigEditor/Private/Graph/NodeSpawners/ControlRigVariableNodeSpawner.cpp
+++ b/Engine/Plugins/Animation/ControlRig/Source/ControlRigEditor/Private/Graph/NodeSpawners/ControlRigVariableNodeSpawner.cpp
@@ -43,6 +43,16 @@ UControlRigVariableNodeSpawner* UControlRigVariableNodeSpawner::CreateFromExtern
 	return NodeSpawner;
 }
 
+UControlRigVariableNodeSpawner* UControlRigVariableNodeSpawner::CreateFromLocalVariable(
+	UControlRigBlueprint* InBlueprint, URigVMGraph* InGraphOwner, const FRigVMGraphVariableDescription& InLocalVariable,
+	bool bInIsGetter, const FText& InMenuDesc, const FText& InCategory, const FText& InTooltip)
+{
+	UControlRigVariableNodeSpawner* Spawner = CreateFromExternalVariable(InBlueprint, InLocalVariable.ToExternalVariable(), bInIsGetter, InMenuDesc, InCategory, InTooltip);
+	Spawner->bIsLocalVariable = true;
+	Spawner->GraphOwner = InGraphOwner;
+	return Spawner;
+}
+
 void UControlRigVariableNodeSpawner::Prime()
 {
 	// we expect that you don't need a node template to construct menu entries
@@ -57,6 +67,26 @@ bool UControlRigVariableNodeSpawner::IsTemplateNodeFilteredOut(FBlueprintActionF
 		{
 			return true;
 		}
+
+		if (bIsLocalVariable)
+		{
+			bool bIsFiltered = true;
+			if (Filter.Context.Graphs.Num() == 1 && GraphOwner.IsValid())
+			{
+				if (UControlRigGraph* Graph = Cast<UControlRigGraph>(Filter.Context.Graphs[0]))
+				{
+					if (GraphOwner.Get() == Graph->GetModel())
+					{
+						bIsFiltered = false;
+					}
+				}
+			}
+
+			if (bIsFiltered)
+			{
+				return true;
+			}
+		}
 	}
 	return Super::IsTemplateNodeFilteredOut(Filter);
 }
diff --git a/Engine/Plugins/Animation/ControlRig/Source/ControlRigEditor/Private/Graph/NodeSpawners/ControlRigVariableNodeSpawner.h b/Engine/Plugins/Animation/ControlRig/Source/ControlRigEditor/Private/Graph/NodeSpawners/ControlRigVariableNodeSpawner.h
--- a/Engine/Plugins/Animation/ControlRig/Source/ControlRigEditor/Private/Graph/NodeSpawners/ControlRigVariableNodeSpawner.h
+++ b/Engine/Plugins/Animation/ControlRig/Source/ControlRigEditor/Private/Graph/NodeSpawners/ControlRigVariableNodeSpawner.h
@@ -13,6 +13,7 @@
 #include "BlueprintNodeSignature.h"
 #include "BlueprintFieldNodeSpawner.h"
 #include "RigVMCore/RigVMExternalVariable.h"
+#include "RigVMModel/RigVMVariableDescription.h"
 #include "ControlRigVariableNodeSpawner.generated.h"
 
 class UControlRigGraphNode;
@@ -32,6 +33,7 @@ public:
 	 * @return A newly allocated instance of this class.
 	 */
 	static UControlRigVariableNodeSpawner* CreateFromExternalVariable(UControlRigBlueprint* InBlueprint, const FRigVMExternalVariable& InExternalVariable, bool bInIsGetter, const FText& InMenuDesc, const FText& InCategory, const FText& InTooltip);
+	static UControlRigVariableNodeSpawner* CreateFromLocalVariable(UControlRigBlueprint* InBlueprint, URigVMGraph* InGraphOwner, const FRigVMGraphVariableDescription& InLocalVariable, bool bInIsGetter, const FText& InMenuDesc, const FText& InCategory, const FText& InTooltip);
 
 	// UBlueprintNodeSpawner interface
 	virtual void Prime() override;
@@ -45,8 +47,10 @@ private:
 
 	/** The pin type we will spawn */
 	TWeakObjectPtr<UControlRigBlueprint> Blueprint;
+	TWeakObjectPtr<URigVMGraph> GraphOwner;
 	FRigVMExternalVariable ExternalVariable;
 	bool bIsGetter;
+	bool bIsLocalVariable;	
 
 	friend class UEngineTestControlRig;
 };
