Control Rig: Support infinite duration for random nodes

diff --git a/Engine/Plugins/Animation/ControlRig/Source/ControlRig/Private/Units/Math/RigUnit_Random.cpp b/Engine/Plugins/Animation/ControlRig/Source/ControlRig/Private/Units/Math/RigUnit_Random.cpp
--- a/Engine/Plugins/Animation/ControlRig/Source/ControlRig/Private/Units/Math/RigUnit_Random.cpp
+++ b/Engine/Plugins/Animation/ControlRig/Source/ControlRig/Private/Units/Math/RigUnit_Random.cpp
@@ -39,7 +39,7 @@ FRigUnit_RandomFloat_Execute()
 
 	Result = FRigUnit_Random_Helper(LastSeed);
 	Result = FMath::Lerp<float>(Minimum, Maximum, Result);
-	TimeLeft = Duration;
+	TimeLeft = Duration < -SMALL_NUMBER ? FLT_MAX : Duration;
 	LastResult = Result;
 }
 
@@ -68,6 +68,6 @@ FRigUnit_RandomVector_Execute()
 	Result.X = FMath::Lerp<float>(Minimum, Maximum, FRigUnit_Random_Helper(LastSeed));
 	Result.Y = FMath::Lerp<float>(Minimum, Maximum, FRigUnit_Random_Helper(LastSeed));
 	Result.Z = FMath::Lerp<float>(Minimum, Maximum, FRigUnit_Random_Helper(LastSeed));
-	TimeLeft = Duration;
+	TimeLeft = Duration < -SMALL_NUMBER ? FLT_MAX : Duration;
 	LastResult = Result;
 }
\ No newline at end of file
diff --git a/Engine/Plugins/Animation/ControlRig/Source/ControlRig/Private/Units/Math/RigUnit_Random.h b/Engine/Plugins/Animation/ControlRig/Source/ControlRig/Private/Units/Math/RigUnit_Random.h
--- a/Engine/Plugins/Animation/ControlRig/Source/ControlRig/Private/Units/Math/RigUnit_Random.h
+++ b/Engine/Plugins/Animation/ControlRig/Source/ControlRig/Private/Units/Math/RigUnit_Random.h
@@ -34,7 +34,11 @@ struct CONTROLRIG_API FRigUnit_RandomFloat : public FRigUnit_MathBase
 	UPROPERTY(meta = (Input))
 	float Maximum;
 
-	/** The duration at which the number won't change. Use 0 for a different number every time. */
+	/**
+	 * The duration at which the number won't change.
+	 * Use 0 for a different number every time.
+	 * A negative number (for ex: -1) results in an infinite duration.
+	 */
 	UPROPERTY(meta = (Input))
 	float Duration;
 
@@ -84,7 +88,12 @@ struct CONTROLRIG_API FRigUnit_RandomVector: public FRigUnit_MathBase
 	UPROPERTY(meta = (Input))
 	float Maximum;
 
-	/** The duration at which the number won't change. Use 0 for a different number every time. */
+	
+	/**
+	 * The duration at which the number won't change.
+	 * Use 0 for a different number every time.
+	 * A negative number (for ex: -1) results in an infinite duration.
+	 */
 	UPROPERTY(meta = (Input))
 	float Duration;
 
