Control Rig: Add outline to Rig hierarchy tree when displaying dynamic hierarchy

diff --git a/Engine/Plugins/Animation/ControlRig/Source/ControlRig/Private/Settings/ControlRigSettings.cpp b/Engine/Plugins/Animation/ControlRig/Source/ControlRig/Private/Settings/ControlRigSettings.cpp
--- a/Engine/Plugins/Animation/ControlRig/Source/ControlRig/Private/Settings/ControlRigSettings.cpp
+++ b/Engine/Plugins/Animation/ControlRig/Source/ControlRig/Private/Settings/ControlRigSettings.cpp
@@ -30,5 +30,6 @@ UControlRigEditorSettings::UControlRigEditorSettings(const FObjectInitializer& O
 	SetupEventBorderColor = FLinearColor::Red;
 	BackwardsSolveBorderColor = FLinearColor::Yellow;
 	BackwardsAndForwardsBorderColor = FLinearColor::Blue;
+	DynamicHierarchyBorderColor = FLinearColor::Blue;
 #endif
 }
diff --git a/Engine/Plugins/Animation/ControlRig/Source/ControlRig/Public/Settings/ControlRigSettings.h b/Engine/Plugins/Animation/ControlRig/Source/ControlRig/Public/Settings/ControlRigSettings.h
--- a/Engine/Plugins/Animation/ControlRig/Source/ControlRig/Public/Settings/ControlRigSettings.h
+++ b/Engine/Plugins/Animation/ControlRig/Source/ControlRig/Public/Settings/ControlRigSettings.h
@@ -110,6 +110,12 @@ public:
 	UPROPERTY(EditAnywhere, config, Category = Viewport)
 	FLinearColor BackwardsAndForwardsBorderColor;
 
+	/**
+	* The border color of the hierarchy tree in dynamic hierarchy mode
+	*/
+	UPROPERTY(EditAnywhere, config, Category = Viewport)
+	FLinearColor DynamicHierarchyBorderColor;
+
 	/**
 	* The default node snippet to create when pressing 1 + Left Mouse Button
 	*/
diff --git a/Engine/Plugins/Animation/ControlRig/Source/ControlRigEditor/Private/Editor/SRigHierarchy.cpp b/Engine/Plugins/Animation/ControlRig/Source/ControlRigEditor/Private/Editor/SRigHierarchy.cpp
--- a/Engine/Plugins/Animation/ControlRig/Source/ControlRigEditor/Private/Editor/SRigHierarchy.cpp
+++ b/Engine/Plugins/Animation/ControlRig/Source/ControlRigEditor/Private/Editor/SRigHierarchy.cpp
@@ -41,6 +41,7 @@
 #include "ToolMenus.h"
 #include "ControlRigContextMenuContext.h"
 #include "SRigSpacePickerWidget.h"
+#include "Settings/ControlRigSettings.h"
 
 #define LOCTEXT_NAMESPACE "SRigHierarchy"
 
@@ -211,11 +212,27 @@ void SRigHierarchy::Construct(const FArguments& InArgs, TSharedRef<FControlRigEd
 		.Padding(0.0f, 0.0f)
 		[
 			SNew(SBorder)
-			.Padding(2.0f)
-			.BorderImage(FEditorStyle::GetBrush("SCSEditor.TreePanel"))
+			.Padding(0.0f)
+			.BorderImage(FControlRigEditorStyle::Get().GetBrush( "ControlRig.Viewport.Border"))
+			.BorderBackgroundColor_Lambda([this]()
+			{
+				FLinearColor Color = FLinearColor::Transparent;
+				if(DisplaySettings.bShowDynamicHierarchy)
+				{
+					Color = UControlRigEditorSettings::Get()->DynamicHierarchyBorderColor;
+				}
+				return Color;
+			})
+			.Padding(0.0f)
+			.ShowEffectWhenDisabled(false)
 			[
-				SAssignNew(TreeView, SRigHierarchyTreeView)
-				.RigTreeDelegates(Delegates)
+				SNew(SBorder)
+				.Padding(2.0f)
+				.BorderImage(FEditorStyle::GetBrush("SCSEditor.TreePanel"))
+				[
+					SAssignNew(TreeView, SRigHierarchyTreeView)
+					.RigTreeDelegates(Delegates)
+				]
 			]
 		]
 
@@ -228,7 +245,7 @@ void SRigHierarchy::Construct(const FArguments& InArgs, TSharedRef<FControlRigEd
 			.Padding(2.0f)
 			.BorderImage(FEditorStyle::GetBrush("SCSEditor.TreePanel"))
 			[
-				SNew(SSpacer)
+			SNew(SSpacer)
 			]
 		]
 		*/
