Control Rig: Remove scoped transactoin when creating a control rig asset (since it is not undoable)

diff --git a/Engine/Plugins/Animation/ControlRig/Source/ControlRigEditor/Private/ControlRigBlueprintActions.cpp b/Engine/Plugins/Animation/ControlRig/Source/ControlRigEditor/Private/ControlRigBlueprintActions.cpp
--- a/Engine/Plugins/Animation/ControlRig/Source/ControlRigEditor/Private/ControlRigBlueprintActions.cpp
+++ b/Engine/Plugins/Animation/ControlRig/Source/ControlRigEditor/Private/ControlRigBlueprintActions.cpp
@@ -143,8 +143,6 @@ UControlRigBlueprint* FControlRigBlueprintActions::CreateNewControlRigAsset(cons
 
 UControlRigBlueprint* FControlRigBlueprintActions::CreateControlRigFromSkeletalMeshOrSkeleton(UObject* InSelectedObject)
 {
-	FScopedTransaction Transaction(LOCTEXT("CreateControlRigFromSkeletalMesh", "Create Control Rig"));
-
 	USkeletalMesh* SkeletalMesh = Cast<USkeletalMesh>(InSelectedObject);
 	USkeleton* Skeleton = Cast<USkeleton>(InSelectedObject);
 	const FReferenceSkeleton* RefSkeleton = nullptr;
