Control Rig: Support report / warning delegate for hierarchy controller and integrate with Control Rig runtime.

diff --git a/Engine/Plugins/Animation/ControlRig/Source/ControlRig/Private/ControlRig.cpp b/Engine/Plugins/Animation/ControlRig/Source/ControlRig/Private/ControlRig.cpp
--- a/Engine/Plugins/Animation/ControlRig/Source/ControlRig/Private/ControlRig.cpp
+++ b/Engine/Plugins/Animation/ControlRig/Source/ControlRig/Private/ControlRig.cpp
@@ -218,7 +218,7 @@ void UControlRig::InitializeFromCDO()
 		DrawContainer = CDO->DrawContainer;
 
 		// copy vm settings
-		VMRuntimeSettings = CDO->VMRuntimeSettings; 
+		VMRuntimeSettings = CDO->VMRuntimeSettings;
 	}
 }
 
@@ -535,6 +535,25 @@ void UControlRig::Execute(const EControlRigState InState, const FName& InEventNa
 			{
 				VM->AddExternalVariable(ExternalVariable);
 			}
+			
+#if WITH_EDITOR
+			// setup the hierarchy's controller log function
+			if(URigHierarchyController* HierarchyController = GetHierarchy()->GetController(true))
+			{
+				HierarchyController->LogFunction = [this](EMessageSeverity::Type InSeverity, const FString& Message)
+				{
+					const FRigVMExecuteContext& Context = GetVM()->GetContext();
+					if(ControlRigLog)
+					{
+						ControlRigLog->Report(InSeverity, Context.FunctionName, Context.InstructionIndex, Message);
+					}
+					else
+					{
+						LogOnce(InSeverity, Context.InstructionIndex, Message);
+					}
+				};
+			}
+#endif
 		}
 #if WITH_EDITOR
 		// default to always clear data after each execution
diff --git a/Engine/Plugins/Animation/ControlRig/Source/ControlRig/Private/Rigs/RigHierarchyController.cpp b/Engine/Plugins/Animation/ControlRig/Source/ControlRig/Private/Rigs/RigHierarchyController.cpp
--- a/Engine/Plugins/Animation/ControlRig/Source/ControlRig/Private/Rigs/RigHierarchyController.cpp
+++ b/Engine/Plugins/Animation/ControlRig/Source/ControlRig/Private/Rigs/RigHierarchyController.cpp
@@ -1879,7 +1879,7 @@ bool URigHierarchyController::AddParent(FRigBaseElement* InChild, FRigBaseElemen
 
 	if(Hierarchy->IsParentedTo(InParent, InChild))
 	{
-		ReportAndNotifyErrorf(TEXT("Cannot parent '%s' to '%s' - would cause a cycle."), *InChild->Key.ToString(), *InParent->Key.ToString());
+		ReportErrorf(TEXT("Cannot parent '%s' to '%s' - would cause a cycle."), *InChild->Key.ToString(), *InParent->Key.ToString());
 		return false;
 	}
 
@@ -2516,6 +2516,12 @@ void URigHierarchyController::ReportWarning(const FString& InMessage) const
 		return;
 	}
 
+	if(LogFunction)
+	{
+		LogFunction(EMessageSeverity::Warning, InMessage);
+		return;
+	}
+
 	FString Message = InMessage;
 	if (Hierarchy)
 	{
@@ -2535,6 +2541,12 @@ void URigHierarchyController::ReportError(const FString& InMessage) const
 		return;
 	}
 
+	if(LogFunction)
+	{
+		LogFunction(EMessageSeverity::Error, InMessage);
+		return;
+	}
+
 	FString Message = InMessage;
 	if (Hierarchy)
 	{
diff --git a/Engine/Plugins/Animation/ControlRig/Source/ControlRig/Public/Rigs/RigHierarchyController.h b/Engine/Plugins/Animation/ControlRig/Source/ControlRig/Public/Rigs/RigHierarchyController.h
--- a/Engine/Plugins/Animation/ControlRig/Source/ControlRig/Public/Rigs/RigHierarchyController.h
+++ b/Engine/Plugins/Animation/ControlRig/Source/ControlRig/Public/Rigs/RigHierarchyController.h
@@ -605,6 +605,11 @@ private:
 	 */
 	bool bSuspendNotifications;
 
+	/**
+	 * This function can be used to override the controller's logging mechanism
+	 */
+	TFunction<void(EMessageSeverity::Type,const FString&)> LogFunction = nullptr;
+
 	friend class UControlRig;
 	friend class URigHierarchy;
 };
