Control Rig: Multi-Select and adding/removing/toggling breakpoint does not work

diff --git a/Engine/Plugins/Animation/ControlRig/Source/ControlRigEditor/Private/ControlRigEditorModule.cpp b/Engine/Plugins/Animation/ControlRig/Source/ControlRigEditor/Private/ControlRigEditorModule.cpp
--- a/Engine/Plugins/Animation/ControlRig/Source/ControlRigEditor/Private/ControlRigEditorModule.cpp
+++ b/Engine/Plugins/Animation/ControlRig/Source/ControlRigEditor/Private/ControlRigEditorModule.cpp
@@ -2003,39 +2003,61 @@ void FControlRigEditorModule::GetContextMenuActions(const UControlRigGraphSchema
 
 					
 					FToolMenuSection& DebugSection = Menu->AddSection("EdGraphSchemaDebug", LOCTEXT("DebugHeader", "Debug"));
-					URigVMNode* Node = RigNode->GetModelNode();
-					if (Node->HasBreakpoint())
+					bool bNoneHasBreakpoint = true;
+					TArray<URigVMNode*> SelectedNodes;
+					for (FName SelectedNodeName : SelectedNodeNames)
+					{
+						if (URigVMNode* ModelNode = Model->FindNodeByName(SelectedNodeName))
+						{
+							SelectedNodes.Add(ModelNode);
+							if (ModelNode->HasBreakpoint())
+							{
+								bNoneHasBreakpoint = false;
+							}
+						}
+					}
+					
+					if (bNoneHasBreakpoint)
 					{
-						
 						DebugSection.AddMenuEntry(
-                        "Remove Breakpoint",
-                        LOCTEXT("RemoveBreakpoint", "Remove Breakpoint"),
-                        LOCTEXT("RemoveBreakpoint_Tooltip", "Removes a breakpoint to the graph at this node"),
-                        FSlateIcon(),
-                        FUIAction(FExecuteAction::CreateLambda([Controller, Node, RigBlueprint]()
-                        {
-                            if (RigBlueprint->RemoveBreakpoint(Node))
-                            {
-                                Node->SetHasBreakpoint(false);
-                            }
-                        })));
+						"Add Breakpoint",
+						LOCTEXT("AddBreakpoint", "Add Breakpoint"),
+						LOCTEXT("AddBreakpoint_Tooltip", "Adds a breakpoint to the graph at this node"),
+						FSlateIcon(),
+						FUIAction(FExecuteAction::CreateLambda([Controller, SelectedNodes, RigBlueprint]()
+						{
+							for (URigVMNode* SelectedNode : SelectedNodes)
+							{
+								if (RigBlueprint->AddBreakpoint(SelectedNode))
+								{
+									SelectedNode->SetHasBreakpoint(true);
+								}
+							}
+						})));
 					}
 					else
-					{
+					{						
 						DebugSection.AddMenuEntry(
-                        "Add Breakpoint",
-                        LOCTEXT("AddBreakpoint", "Add Breakpoint"),
-                        LOCTEXT("AddBreakpoint_Tooltip", "Adds a breakpoint to the graph at this node"),
+                        "Remove Breakpoint",
+                        LOCTEXT("RemoveBreakpoint", "Remove Breakpoint"),
+                        LOCTEXT("RemoveBreakpoint_Tooltip", "Removes a breakpoint to the graph at this node"),
                         FSlateIcon(),
-                        FUIAction(FExecuteAction::CreateLambda([Controller, Node, RigBlueprint]()
+                        FUIAction(FExecuteAction::CreateLambda([Controller, SelectedNodes, RigBlueprint]()
                         {
-                            if (RigBlueprint->AddBreakpoint(Node))
-                            {
-                                Node->SetHasBreakpoint(true);
-                            }
+                        	for (URigVMNode* SelectedNode : SelectedNodes)
+                        	{
+                        		if (SelectedNode->HasBreakpoint())
+                        		{
+                        			if (RigBlueprint->RemoveBreakpoint(SelectedNode))
+                        			{
+										SelectedNode->SetHasBreakpoint(false);
+									}
+                        		}
+							}                            
                         })));
 					}
 					
+					
 
 					FToolMenuSection& OrganizationSection = Menu->AddSection("EdGraphSchemaOrganization", LOCTEXT("OrganizationHeader", "Organization"));
 					OrganizationSection.AddMenuEntry(
diff --git a/Engine/Plugins/Animation/ControlRig/Source/ControlRigEditor/Private/Editor/ControlRigEditor.cpp b/Engine/Plugins/Animation/ControlRig/Source/ControlRigEditor/Private/Editor/ControlRigEditor.cpp
--- a/Engine/Plugins/Animation/ControlRig/Source/ControlRigEditor/Private/Editor/ControlRigEditor.cpp
+++ b/Engine/Plugins/Animation/ControlRig/Source/ControlRigEditor/Private/Editor/ControlRigEditor.cpp
@@ -245,11 +245,7 @@ FControlRigEditor::~FControlRigEditor()
 #if WITH_EDITOR
 		RigBlueprint->SetDebugMode(false);
 		RigBlueprint->ClearBreakpoints();
-		if (HaltedAtNode)
-		{
-			HaltedAtNode->SetExecutionIsHaltedAtThisNode(false);
-			HaltedAtNode = nullptr;
-		}
+		SetHaltedNode(nullptr);
 #endif
 	}
 
@@ -981,12 +977,8 @@ void FControlRigEditor::SetExecutionMode(const EControlRigExecutionModeType InEx
 		ControlRig->bIsInDebugMode = InExecutionMode == EControlRigExecutionModeType_Debug;
 	}
 
-	if (HaltedAtNode)
-	{
-		HaltedAtNode->SetExecutionIsHaltedAtThisNode(false);
-		HaltedAtNode = nullptr;
-	}
-
+	SetHaltedNode(nullptr);
+	
 	RefreshDetailView();
 }
 
@@ -2911,18 +2903,8 @@ void FControlRigEditor::HandleControlRigExecutionHalted(const int32 InstructionI
 		
 	if (URigVMNode* InNode = Cast<URigVMNode>(InNodeObject))
 	{
-		if (HaltedAtNode != nullptr)
-		{
-			HaltedAtNode->SetExecutionIsHaltedAtThisNode(false);
-		}
-		HaltedAtNode = InNode;
-
-		if (InNode == nullptr)
-		{
-			return;
-		}	
-		InNode->SetExecutionIsHaltedAtThisNode(true);
-	
+		SetHaltedNode(InNode);
+		
 		if (UControlRigBlueprint* Blueprint = GetControlRigBlueprint())
 		{
 			if (Blueprint->GetAllModels().Contains(InNode->GetGraph()))
@@ -2941,11 +2923,24 @@ void FControlRigEditor::HandleControlRigExecutionHalted(const int32 InstructionI
 	{
 		if (InEntryName == ControlRig->GetEventQueue().Last())
 		{
-			HaltedAtNode = nullptr;
+			SetHaltedNode(nullptr);
 		}
 	}
 }
 
+void FControlRigEditor::SetHaltedNode(URigVMNode* Node)
+{
+	if (HaltedAtNode)
+	{
+		HaltedAtNode->SetExecutionIsHaltedAtThisNode(false);
+	}
+	HaltedAtNode = Node;
+	if (HaltedAtNode)
+	{
+		HaltedAtNode->SetExecutionIsHaltedAtThisNode(true);
+	}
+}
+
 void FControlRigEditor::CreateEditorModeManager()
 {
 	EditorModeManager = MakeShareable(FModuleManager::LoadModuleChecked<FPersonaModule>("Persona").CreatePersonaEditorModeManager());
diff --git a/Engine/Plugins/Animation/ControlRig/Source/ControlRigEditor/Private/Editor/ControlRigEditor.h b/Engine/Plugins/Animation/ControlRig/Source/ControlRigEditor/Private/Editor/ControlRigEditor.h
--- a/Engine/Plugins/Animation/ControlRig/Source/ControlRigEditor/Private/Editor/ControlRigEditor.h
+++ b/Engine/Plugins/Animation/ControlRig/Source/ControlRigEditor/Private/Editor/ControlRigEditor.h
@@ -246,6 +246,7 @@ protected:
 	void HandleVMCompiledEvent(UBlueprint* InBlueprint, URigVM* InVM);
 	void HandleControlRigExecutedEvent(UControlRig* InControlRig, const EControlRigState InState, const FName& InEventName);
 	void HandleControlRigExecutionHalted(const int32 InstructionIndex, UObject* InNode, const FName& InEntryName);
+	void SetHaltedNode(URigVMNode* Node);
 
 	// FBaseToolKit overrides
 	void CreateEditorModeManager() override;
