Control Rig: Refresh hierarchy to set preview mesh

diff --git a/Engine/Plugins/Animation/ControlRig/Source/ControlRigEditor/Private/Editor/SRigHierarchy.cpp b/Engine/Plugins/Animation/ControlRig/Source/ControlRigEditor/Private/Editor/SRigHierarchy.cpp
--- a/Engine/Plugins/Animation/ControlRig/Source/ControlRigEditor/Private/Editor/SRigHierarchy.cpp
+++ b/Engine/Plugins/Animation/ControlRig/Source/ControlRigEditor/Private/Editor/SRigHierarchy.cpp
@@ -1153,6 +1153,17 @@ void SRigHierarchy::RefreshHierarchy(const FAssetData& InAssetData)
 	ControlRigBlueprint->BroadcastRefreshEditor();
 	RefreshTreeView();
 	FSlateApplication::Get().DismissAllMenus();
+
+	if (ControlRigEditor.IsValid() && 
+		Mesh != nullptr)
+	{
+		ControlRigEditor.Pin()->GetPersonaToolkit()->SetPreviewMesh(Mesh, true);
+	}
+
+	if (ControlRigEditor.IsValid())
+	{
+		ControlRigEditor.Pin()->Compile();
+	}
 }
 
 void SRigHierarchy::CreateImportMenu(FMenuBuilder& MenuBuilder)
