Control Rig: Changing to Setup Mode from Forwards Solve while in Debug mode does not disable UI if no breakpoints in Setup

diff --git a/Engine/Plugins/Animation/ControlRig/Source/ControlRigEditor/Private/Editor/ControlRigEditor.cpp b/Engine/Plugins/Animation/ControlRig/Source/ControlRigEditor/Private/Editor/ControlRigEditor.cpp
--- a/Engine/Plugins/Animation/ControlRig/Source/ControlRigEditor/Private/Editor/ControlRigEditor.cpp
+++ b/Engine/Plugins/Animation/ControlRig/Source/ControlRigEditor/Private/Editor/ControlRigEditor.cpp
@@ -844,6 +844,8 @@ void FControlRigEditor::SetEventQueue(EControlRigEditorEventQueue InEventQueue)
 
 	LastEventQueue = GetEventQueue();
 
+	SetHaltedNode(nullptr);
+
 	if (ControlRig)
 	{
 		TArray<FName> EventNames;
@@ -5914,11 +5916,7 @@ void FControlRigEditor::HandleBreakpointActionRequested(const ERigVMBreakpointAc
 
 bool FControlRigEditor::IsHaltedAtBreakpoint() const
 {
-	if (ControlRig && ControlRig->GetVM())
-	{
-		return ControlRig->GetVM()->GetHaltedAtBreakpoint() != nullptr;
-	}
-	return false;
+	return HaltedAtNode != nullptr;
 }
 
 void FControlRigEditor::UpdateCapsules()
