Control Rig: Perform initialization after preview mesh change

diff --git a/Engine/Plugins/Animation/ControlRig/Source/ControlRigEditor/Private/Editor/ControlRigEditor.cpp b/Engine/Plugins/Animation/ControlRig/Source/ControlRigEditor/Private/Editor/ControlRigEditor.cpp
--- a/Engine/Plugins/Animation/ControlRig/Source/ControlRigEditor/Private/Editor/ControlRigEditor.cpp
+++ b/Engine/Plugins/Animation/ControlRig/Source/ControlRigEditor/Private/Editor/ControlRigEditor.cpp
@@ -4019,9 +4019,11 @@ void FControlRigEditor::HandlePreviewMeshChanged(USkeletalMesh* InOldSkeletalMes
 			if(UControlRig* DebuggedControlRig = Cast<UControlRig>(GetBlueprintObj()->GetObjectBeingDebugged()))
 			{
 				DebuggedControlRig->GetHierarchy()->Notify(ERigHierarchyNotification::HierarchyReset, nullptr);
-				DebuggedControlRig->RequestSetup();
+				DebuggedControlRig->Initialize(true);
 			}
 		}
+
+		Compile();
 	}
 }
 
