Control Rig: Filter on OnCleared commits from TextFields

diff --git a/Engine/Plugins/Animation/ControlRig/Source/ControlRigEditor/Private/ControlRigElementDetails.cpp b/Engine/Plugins/Animation/ControlRig/Source/ControlRigEditor/Private/ControlRigElementDetails.cpp
--- a/Engine/Plugins/Animation/ControlRig/Source/ControlRigEditor/Private/ControlRigElementDetails.cpp
+++ b/Engine/Plugins/Animation/ControlRig/Source/ControlRigEditor/Private/ControlRigElementDetails.cpp
@@ -2037,6 +2037,11 @@ FText FRigControlElementDetails::GetDisplayName() const
 
 void FRigControlElementDetails::SetDisplayName(const FText& InNewText, ETextCommit::Type InCommitType)
 {
+	if(InCommitType == ETextCommit::OnCleared)
+	{
+		return;
+	}
+	
 	const FName DisplayName = InNewText.IsEmpty() ? FName(NAME_None) : FName(*InNewText.ToString());
 	
 	for(int32 ObjectIndex = 0; ObjectIndex < ObjectsBeingCustomized.Num(); ObjectIndex++)
diff --git a/Engine/Plugins/Animation/ControlRig/Source/ControlRigEditor/Private/ControlRigGraphDetails.cpp b/Engine/Plugins/Animation/ControlRig/Source/ControlRigEditor/Private/ControlRigGraphDetails.cpp
--- a/Engine/Plugins/Animation/ControlRig/Source/ControlRigEditor/Private/ControlRigGraphDetails.cpp
+++ b/Engine/Plugins/Animation/ControlRig/Source/ControlRigEditor/Private/ControlRigGraphDetails.cpp
@@ -855,6 +855,11 @@ FText FControlRigGraphDetails::GetNodeCategory() const
 
 void FControlRigGraphDetails::SetNodeCategory(const FText& InNewText, ETextCommit::Type InCommitType)
 {
+	if(InCommitType == ETextCommit::OnCleared)
+	{
+		return;
+	}
+
 	if (GraphPtr.IsValid() && ControlRigBlueprintPtr.IsValid())
 	{
 		UControlRigBlueprint* Blueprint = ControlRigBlueprintPtr.Get();
@@ -890,6 +895,11 @@ FText FControlRigGraphDetails::GetNodeKeywords() const
 
 void FControlRigGraphDetails::SetNodeKeywords(const FText& InNewText, ETextCommit::Type InCommitType)
 {
+	if(InCommitType == ETextCommit::OnCleared)
+	{
+		return;
+	}
+
 	if (GraphPtr.IsValid() && ControlRigBlueprintPtr.IsValid())
 	{
 		UControlRigBlueprint* Blueprint = ControlRigBlueprintPtr.Get();
@@ -925,6 +935,11 @@ FText FControlRigGraphDetails::GetNodeDescription() const
 
 void FControlRigGraphDetails::SetNodeDescription(const FText& InNewText, ETextCommit::Type InCommitType)
 {
+	if(InCommitType == ETextCommit::OnCleared)
+	{
+		return;
+	}
+
 	if (GraphPtr.IsValid() && ControlRigBlueprintPtr.IsValid())
 	{
 		UControlRigBlueprint* Blueprint = ControlRigBlueprintPtr.Get();
diff --git a/Engine/Plugins/Animation/ControlRig/Source/ControlRigEditor/Private/Editor/SControlRigProfilingView.cpp b/Engine/Plugins/Animation/ControlRig/Source/ControlRigEditor/Private/Editor/SControlRigProfilingView.cpp
--- a/Engine/Plugins/Animation/ControlRig/Source/ControlRigEditor/Private/Editor/SControlRigProfilingView.cpp
+++ b/Engine/Plugins/Animation/ControlRig/Source/ControlRigEditor/Private/Editor/SControlRigProfilingView.cpp
@@ -1033,6 +1033,11 @@ void SControlRigProfilingView::ToggleItemExpansion(const TSharedPtr<FStatsTreeEl
 
 void SControlRigProfilingView::OnSearchTextCommitted(const FText& InNewText, ETextCommit::Type InCommitType)
 {
+	if(InCommitType == ETextCommit::OnCleared)
+	{
+		return;
+	}
+
 	SearchText = InNewText.ToString();
 	TreeView->RebuildList();
 	//TreeView->RequestTreeRefresh();
