Control Rig: Fix control's gizmo disappears during control multiple editing and enabling limits

diff --git a/Engine/Plugins/Animation/ControlRig/Source/ControlRigEditor/Private/Editor/SRigHierarchy.cpp b/Engine/Plugins/Animation/ControlRig/Source/ControlRigEditor/Private/Editor/SRigHierarchy.cpp
--- a/Engine/Plugins/Animation/ControlRig/Source/ControlRigEditor/Private/Editor/SRigHierarchy.cpp
+++ b/Engine/Plugins/Animation/ControlRig/Source/ControlRigEditor/Private/Editor/SRigHierarchy.cpp
@@ -1438,6 +1438,11 @@ void SRigHierarchy::HandleNewItem(ERigElementType InElementType)
 				{
 					FRigControlSettings Settings;
 					Settings.ControlType = ERigControlType::EulerTransform;
+
+					FEulerTransform Identity = FEulerTransform::Identity;
+					FRigControlValue ValueToSet = FRigControlValue::Make<FEulerTransform>(Identity);
+					Settings.MinimumValue = ValueToSet;
+					Settings.MaximumValue = ValueToSet;
 						
 					NewItemKey = Controller->AddControl(NewElementName, ParentKey, Settings, Settings.GetIdentityValue(), FTransform::Identity, FTransform::Identity, true, true);
 					break;
