Adding FKControlRig track collapses SkelMesh
#fix Set control value as EulerTransform (which is the correct type) rather than regular Transform

diff --git a/Engine/Plugins/Animation/ControlRig/Source/ControlRig/Private/Rigs/FKControlRig.cpp b/Engine/Plugins/Animation/ControlRig/Source/ControlRig/Private/Rigs/FKControlRig.cpp
--- a/Engine/Plugins/Animation/ControlRig/Source/ControlRig/Private/Rigs/FKControlRig.cpp
+++ b/Engine/Plugins/Animation/ControlRig/Source/ControlRig/Private/Rigs/FKControlRig.cpp
@@ -143,7 +143,7 @@ void UFKControlRig::ExecuteUnits(FRigUnitContext& InOutContext, const FName& InE
 			Transform.NormalizeRotation();
 			Transform.SetScale3D(FVector::OneVector);
 
-			SetControlValue(ControlName, FRigControlValue::Make(Transform), bNotify, Context, bSetupUndo);
+			SetControlValue(ControlName, FRigControlValue::Make(FEulerTransform(Transform)), bNotify, Context, bSetupUndo);
 
 		}, true);
 
