Control Rig: Dragging rig hierarchy elements to create free reroute node

diff --git a/Engine/Plugins/Animation/ControlRig/Source/ControlRigEditor/Private/Editor/ControlRigEditor.cpp b/Engine/Plugins/Animation/ControlRig/Source/ControlRigEditor/Private/Editor/ControlRigEditor.cpp
--- a/Engine/Plugins/Animation/ControlRig/Source/ControlRigEditor/Private/Editor/ControlRigEditor.cpp
+++ b/Engine/Plugins/Animation/ControlRig/Source/ControlRigEditor/Private/Editor/ControlRigEditor.cpp
@@ -5272,19 +5272,22 @@ void FControlRigEditor::OnGraphNodeDropToPerform(TSharedPtr<FGraphNodeDragDropOp
 				MenuBuilder.AddMenuSeparator();
 
 				MenuBuilder.AddMenuEntry(
-					LOCTEXT("CreateCollectionFromKeys", "Create Collection"),
-					LOCTEXT("CreateCollectionFromKeysTooltip", "Creates a collection from the selected elements in the hierarchy"),
+					LOCTEXT("CreateItemArrayFromKeys", "Create Item Array"),
+					LOCTEXT("CreateItemArrayFromKeysTooltip", "Creates an item array from the selected elements in the hierarchy"),
 					FSlateIcon(),
 					FUIAction(
 						FExecuteAction::CreateLambda([this, Blueprint, DraggedKeys, NodePosition]()
 							{
 								if (URigVMController* Controller = GetFocusedController())
 								{
-									Controller->OpenUndoBracket(TEXT("Create Collection from Items"));
+									Controller->OpenUndoBracket(TEXT("Create Item Array From Selection"));
 
-									if (URigVMNode* ItemsNode = Controller->AddUnitNode(FRigUnit_CollectionItems::StaticStruct(), FRigUnit::GetMethodName(), NodePosition, FString(), true, true))
+									static const FString ItemArrayCPPType = FString::Printf(TEXT("TArray<%s>"), *FRigElementKey::StaticStruct()->GetStructCPPName());
+									static const FString ItemArrayObjectPath = FRigElementKey::StaticStruct()->GetPathName();
+									
+									if (URigVMNode* ItemsNode = Controller->AddFreeRerouteNode(true, ItemArrayCPPType, *ItemArrayObjectPath, false, NAME_None, FString(), NodePosition))
 									{
-										if (URigVMPin* ItemsPin = ItemsNode->FindPin(TEXT("Items")))
+										if (URigVMPin* ItemsPin = ItemsNode->FindPin(TEXT("Value")))
 										{
 											Controller->SetArrayPinSize(ItemsPin->GetPinPath(), DraggedKeys.Num());
 
@@ -5298,6 +5301,8 @@ void FControlRigEditor::OnGraphNodeDropToPerform(TSharedPtr<FGraphNodeDragDropOp
 												Controller->SetPinDefaultValue(ItemPins[ItemIndex]->GetPinPath(), DefaultValue, true, true, false, true);
 												Controller->SetPinExpansion(ItemPins[ItemIndex]->GetPinPath(), true, true, true);
 											}
+
+											Controller->SetPinExpansion(ItemsPin->GetPinPath(), true, true, true);
 										}
 									}
 
