Control Rig: Clear Last Pin used for compatibility checks when graph is changing

diff --git a/Engine/Plugins/Animation/ControlRig/Source/ControlRigDeveloper/Private/Graph/ControlRigGraph.cpp b/Engine/Plugins/Animation/ControlRig/Source/ControlRigDeveloper/Private/Graph/ControlRigGraph.cpp
--- a/Engine/Plugins/Animation/ControlRig/Source/ControlRigDeveloper/Private/Graph/ControlRigGraph.cpp
+++ b/Engine/Plugins/Animation/ControlRig/Source/ControlRigDeveloper/Private/Graph/ControlRigGraph.cpp
@@ -226,6 +226,11 @@ void UControlRigGraph::HandleModifiedEvent(ERigVMGraphNotifType InNotifType, URi
 		return;
 	}
 
+	if(UControlRigGraphSchema* ControlRigSchema = (UControlRigGraphSchema*)GetControlRigGraphSchema())
+	{
+		ControlRigSchema->HandleModifiedEvent(InNotifType, InGraph, InSubject);
+	}
+
 	// increment the node topology version for any interaction
 	// with a node.
 	{
diff --git a/Engine/Plugins/Animation/ControlRig/Source/ControlRigDeveloper/Private/Graph/ControlRigGraphSchema.cpp b/Engine/Plugins/Animation/ControlRig/Source/ControlRigDeveloper/Private/Graph/ControlRigGraphSchema.cpp
--- a/Engine/Plugins/Animation/ControlRig/Source/ControlRigDeveloper/Private/Graph/ControlRigGraphSchema.cpp
+++ b/Engine/Plugins/Animation/ControlRig/Source/ControlRigDeveloper/Private/Graph/ControlRigGraphSchema.cpp
@@ -454,7 +454,6 @@ bool UControlRigGraphSchema::TryCreateConnection(UEdGraphPin* PinA, UEdGraphPin*
 	}
 
 	UControlRigGraphSchema* MutableThis = (UControlRigGraphSchema*)this;
-	MutableThis->LastPinForCompatibleCheck = nullptr;
 
 	UBlueprint* Blueprint = FBlueprintEditorUtils::FindBlueprintForNodeChecked(PinA->GetOwningNode());
 	UControlRigBlueprint* RigBlueprint = Cast<UControlRigBlueprint>(Blueprint);
@@ -1745,5 +1744,29 @@ FVector2D UControlRigGraphSchema::GetNodePositionAtStartOfInteraction(UEdGraphNo
 	return FVector2D::ZeroVector;
 }
 
+void UControlRigGraphSchema::HandleModifiedEvent(ERigVMGraphNotifType InNotifType, URigVMGraph* InGraph,
+	UObject* InSubject)
+{
+	switch(InNotifType)
+	{
+		case ERigVMGraphNotifType::NodeAdded:
+		case ERigVMGraphNotifType::NodeRemoved:
+		case ERigVMGraphNotifType::PinAdded:
+		case ERigVMGraphNotifType::PinRemoved:
+		case ERigVMGraphNotifType::PinRenamed:
+		case ERigVMGraphNotifType::PinArraySizeChanged:
+		case ERigVMGraphNotifType::PinTypeChanged:
+		case ERigVMGraphNotifType::LinkAdded:
+		case ERigVMGraphNotifType::LinkRemoved:
+		{
+			LastPinForCompatibleCheck = nullptr;
+			break;
+		}
+		default:
+		{
+			break;
+		}
+	}
+}
 
 #undef LOCTEXT_NAMESPACE
diff --git a/Engine/Plugins/Animation/ControlRig/Source/ControlRigDeveloper/Public/Graph/ControlRigGraphSchema.h b/Engine/Plugins/Animation/ControlRig/Source/ControlRigDeveloper/Public/Graph/ControlRigGraphSchema.h
--- a/Engine/Plugins/Animation/ControlRig/Source/ControlRigDeveloper/Public/Graph/ControlRigGraphSchema.h
+++ b/Engine/Plugins/Animation/ControlRig/Source/ControlRigDeveloper/Public/Graph/ControlRigGraphSchema.h
@@ -228,6 +228,8 @@ public:
 	static TArray<UEdGraphNode*> GetNodesToMoveForNode(UEdGraphNode* InNode);
 	FVector2D GetNodePositionAtStartOfInteraction(UEdGraphNode* InNode) const;
 
+	void HandleModifiedEvent(ERigVMGraphNotifType InNotifType, URigVMGraph* InGraph, UObject* InSubject);
+
 private:
 
 	const UEdGraphPin* LastPinForCompatibleCheck = nullptr;
