Control Rig: Improve array pin editing

diff --git a/Engine/Plugins/Animation/ControlRig/Source/ControlRigDeveloper/Private/Graph/ControlRigGraphNode.cpp b/Engine/Plugins/Animation/ControlRig/Source/ControlRigDeveloper/Private/Graph/ControlRigGraphNode.cpp
--- a/Engine/Plugins/Animation/ControlRig/Source/ControlRigDeveloper/Private/Graph/ControlRigGraphNode.cpp
+++ b/Engine/Plugins/Animation/ControlRig/Source/ControlRigDeveloper/Private/Graph/ControlRigGraphNode.cpp
@@ -562,6 +562,12 @@ void UControlRigGraphNode::CreateInputOutputPins(URigVMPin* InParentPin, bool bH
 {
 	DECLARE_SCOPE_HIERARCHICAL_COUNTER_FUNC()
 
+	bool bIsContainer = false;
+	if(InParentPin)
+	{
+		bIsContainer = InParentPin->IsArray();
+	}
+
 	const TArray<URigVMPin*> ModelPins = InParentPin == nullptr ? InputOutputPins : InParentPin->GetSubPins();
 	for (URigVMPin* ModelPin : ModelPins)
 	{
@@ -586,7 +592,7 @@ void UControlRigGraphNode::CreateInputOutputPins(URigVMPin* InParentPin, bool bH
 		}
 	}
 	}
-		if (Pair.OutputPin == nullptr)
+		if (Pair.OutputPin == nullptr && !bIsContainer)
 	{
 			Pair.OutputPin = CreatePin(EGPD_Output, GetPinTypeForModelPin(ModelPin), FName(*ModelPin->GetPinPath()));
 			if (Pair.OutputPin != nullptr)
@@ -614,7 +620,14 @@ void UControlRigGraphNode::CreateInputOutputPins(URigVMPin* InParentPin, bool bH
 		}
 	}
 
-		CreateInputOutputPins(ModelPin, bHidden);
+		if(bIsContainer)
+		{
+			CreateInputPins(ModelPin);
+		}
+		else
+		{
+			CreateInputOutputPins(ModelPin, bHidden);
+		}
 	}
 }
 
diff --git a/Engine/Plugins/Animation/ControlRig/Source/ControlRigEditor/Private/Graph/SControlRigGraphNode.cpp b/Engine/Plugins/Animation/ControlRig/Source/ControlRigEditor/Private/Graph/SControlRigGraphNode.cpp
--- a/Engine/Plugins/Animation/ControlRig/Source/ControlRigEditor/Private/Graph/SControlRigGraphNode.cpp
+++ b/Engine/Plugins/Animation/ControlRig/Source/ControlRigEditor/Private/Graph/SControlRigGraphNode.cpp
@@ -294,6 +294,35 @@ void SControlRigGraphNode::Construct( const FArguments& InArgs )
 		}
 	}
 
+	auto AddArrayPlusButtonLambda = [this](URigVMPin* InModelPin, TSharedPtr<SHorizontalBox> InSlotLayout, const float InEmptySidePadding)
+	{
+		// add array plus button
+		InSlotLayout->AddSlot()
+		.AutoWidth()
+		.HAlign(HAlign_Left)
+		.Padding(PinWidgetSidePadding, TopPadding, InEmptySidePadding, 0.f)
+		[
+			SNew(SButton)
+			.ContentPadding(0.0f)
+			.ButtonStyle(FEditorStyle::Get(), "NoBorder")
+			.OnClicked(this, &SControlRigGraphNode::HandleAddArrayElement, InModelPin->GetPinPath())
+			.IsEnabled(this, &SGraphNode::IsNodeEditable)
+			.Cursor(EMouseCursor::Default)
+			.Visibility(this, &SControlRigGraphNode::GetArrayPlusButtonVisibility, InModelPin)
+			.ToolTipText(LOCTEXT("AddArrayElement", "Add Array Element"))
+			[
+				SNew(SHorizontalBox)
+				+SHorizontalBox::Slot()
+				.AutoWidth()
+				.VAlign(VAlign_Center)
+				[
+					SNew(SImage)
+					.Image(FEditorStyle::GetBrush(TEXT("Icons.PlusCircle")))
+				]
+			]
+		];
+	};
+
 	for(const FPinInfo& PinInfo : PinInfos)
 	{
 		if(PinInfo.InputPinWidget.IsValid())
@@ -336,41 +365,20 @@ void SControlRigGraphNode::Construct( const FArguments& InArgs )
 				if(PinInfo.bIsContainer)
 				{
 					URigVMPin* ModelPin = ModelNode->GetGraph()->FindPin(PinInfo.ModelPinPath);
-					if(ModelPin && ModelPin->GetSourceLinks().Num() == 0)
+					if(ModelPin)
 					{
 						// make sure to minimize the width of the label
 						FirstSlot->SetAutoWidth();
-						
-						// add array plus button
-						SlotLayout->AddSlot()
-						.AutoWidth()
-						.HAlign(HAlign_Left)
-						.Padding(PinWidgetSidePadding, TopPadding, MyEmptySidePadding, 0.f)
-						[
-							SNew(SButton)
-							.ContentPadding(0.0f)
-							.ButtonStyle(FEditorStyle::Get(), "NoBorder")
-							.OnClicked(this, &SControlRigGraphNode::HandleAddArrayElement, PinInfo.ModelPinPath)
-							.IsEnabled(this, &SGraphNode::IsNodeEditable)
-							.Cursor(EMouseCursor::Default)
-							.ToolTipText(LOCTEXT("AddArrayElement", "Add Array Element"))
-							[
-								SNew(SHorizontalBox)
-								+SHorizontalBox::Slot()
-								.AutoWidth()
-								.VAlign(VAlign_Center)
-								[
-									SNew(SImage)
-									.Image(FEditorStyle::GetBrush(TEXT("Icons.PlusCircle")))
-								]
-							]
-						];
+						AddArrayPlusButtonLambda(ModelPin, SlotLayout, MyEmptySidePadding);
 					}
 				}
 			}
 			// io pins
 			else
 			{
+				TSharedPtr<SHorizontalBox> SlotLayout;
+				SHorizontalBox::FSlot* FirstSlot = nullptr;
+
 				PinInfo.OutputPinWidget->SetShowLabel(false);
 			
 				LeftNodeBox->AddSlot()
@@ -379,10 +387,11 @@ void SControlRigGraphNode::Construct( const FArguments& InArgs )
                 .AutoHeight()
 				.MaxHeight(MaxHeight)
                 [
-                    SNew(SHorizontalBox)
+                    SAssignNew(SlotLayout, SHorizontalBox)
                     .Visibility(this, &SControlRigGraphNode::GetPinVisibility, PinInfo.Index)
 
                     +SHorizontalBox::Slot()
+					.Expose(FirstSlot)
                     .FillWidth(1.f)
                     .HAlign(HAlign_Left)
                     .VAlign(VAlign_Center)
@@ -390,16 +399,27 @@ void SControlRigGraphNode::Construct( const FArguments& InArgs )
                     [
                         PinInfo.InputPinWidget.ToSharedRef()
                     ]
-
-                    +SHorizontalBox::Slot()
-                    .FillWidth(1.f)
-                    .HAlign(HAlign_Right)
-					.VAlign(VAlign_Center)
-                    .Padding(0.f, TopPadding, PinWidgetSidePadding, 0.f)
-                    [
-                        PinInfo.OutputPinWidget.ToSharedRef()
-                    ]
                 ];
+
+				if(PinInfo.bIsContainer)
+				{
+					URigVMPin* ModelPin = ModelNode->GetGraph()->FindPin(PinInfo.ModelPinPath);
+					if(ModelPin)
+					{
+						// make sure to minimize the width of the label
+						FirstSlot->SetAutoWidth();
+						AddArrayPlusButtonLambda(ModelPin, SlotLayout, EmptySidePadding);
+					}
+				}
+
+				SlotLayout->AddSlot()
+				.FillWidth(1.f)
+				.HAlign(HAlign_Right)
+				.VAlign(VAlign_Center)
+				.Padding(0.f, TopPadding, PinWidgetSidePadding, 0.f)
+				[
+					PinInfo.OutputPinWidget.ToSharedRef()
+				];
 			}
 		}
 		// output pins
@@ -709,6 +729,18 @@ EVisibility SControlRigGraphNode::GetTitleVisibility() const
 	return UseLowDetailNodeTitles() ? EVisibility::Hidden : EVisibility::Visible;
 }
 
+EVisibility SControlRigGraphNode::GetArrayPlusButtonVisibility(URigVMPin* InModelPin) const
+{
+	if(InModelPin)
+	{
+		if(InModelPin->GetSourceLinks().Num() == 0)
+		{
+			return EVisibility::Visible;
+		}
+	}
+	return EVisibility::Hidden;
+}
+
 TSharedRef<SWidget> SControlRigGraphNode::CreateTitleWidget(TSharedPtr<SNodeTitle> InNodeTitle)
 {
 	NodeTitle = InNodeTitle;
diff --git a/Engine/Plugins/Animation/ControlRig/Source/ControlRigEditor/Private/Graph/SControlRigGraphNode.h b/Engine/Plugins/Animation/ControlRig/Source/ControlRigEditor/Private/Graph/SControlRigGraphNode.h
--- a/Engine/Plugins/Animation/ControlRig/Source/ControlRigEditor/Private/Graph/SControlRigGraphNode.h
+++ b/Engine/Plugins/Animation/ControlRig/Source/ControlRigEditor/Private/Graph/SControlRigGraphNode.h
@@ -60,6 +60,7 @@ private:
 	FVector2D GetLowDetailDesiredSize() const;
 
 	EVisibility GetTitleVisibility() const;
+	EVisibility GetArrayPlusButtonVisibility(URigVMPin* InModelPin) const;
 
 	FText GetPinLabel(TWeakPtr<SGraphPin> GraphPin) const;
 
