Control Rig: Fix crash when promoting pin to variable twice

diff --git a/Engine/Plugins/Animation/ControlRig/Source/ControlRigDeveloper/Private/ControlRigBlueprint.cpp b/Engine/Plugins/Animation/ControlRig/Source/ControlRigDeveloper/Private/ControlRigBlueprint.cpp
--- a/Engine/Plugins/Animation/ControlRig/Source/ControlRigDeveloper/Private/ControlRigBlueprint.cpp
+++ b/Engine/Plugins/Animation/ControlRig/Source/ControlRigDeveloper/Private/ControlRigBlueprint.cpp
@@ -3146,10 +3146,10 @@ void UControlRigBlueprint::CreateMemberVariablesOnLoad()
 
 FName UControlRigBlueprint::FindCRMemberVariableUniqueName(TSharedPtr<FKismetNameValidator> InNameValidator, const FString& InBaseName)
 {
-	FString KismetName = InBaseName;
-	if (InNameValidator->IsValid(KismetName) == EValidatorResult::ContainsInvalidCharacters)
+	FString BaseName = InBaseName;
+	if (InNameValidator->IsValid(BaseName) == EValidatorResult::ContainsInvalidCharacters)
 	{
-		for (TCHAR& TestChar : KismetName)
+		for (TCHAR& TestChar : BaseName)
 		{
 			for (TCHAR BadChar : UE_BLUEPRINT_INVALID_NAME_CHARACTERS)
 			{
@@ -3162,10 +3162,12 @@ FName UControlRigBlueprint::FindCRMemberVariableUniqueName(TSharedPtr<FKismetNam
 		}
 	}
 
+	FString KismetName = BaseName;
+
 	int32 Suffix = 0;
 	while (InNameValidator->IsValid(KismetName) != EValidatorResult::Ok)
 	{
-		KismetName = FString::Printf(TEXT("%s_%d"), *InBaseName, Suffix);
+		KismetName = FString::Printf(TEXT("%s_%d"), *BaseName, Suffix);
 		Suffix++;
 	}
 
