Control Rig: Deprecate all units using a collection and move them over to use an item array

diff --git a/Engine/Plugins/Animation/ControlRig/Source/ControlRig/Private/Rigs/RigHierarchy.cpp b/Engine/Plugins/Animation/ControlRig/Source/ControlRig/Private/Rigs/RigHierarchy.cpp
--- a/Engine/Plugins/Animation/ControlRig/Source/ControlRig/Private/Rigs/RigHierarchy.cpp
+++ b/Engine/Plugins/Animation/ControlRig/Source/ControlRig/Private/Rigs/RigHierarchy.cpp
@@ -1794,6 +1794,12 @@ FRigPose URigHierarchy::GetPose(
 	ERigElementType InElementType,
 	const FRigElementKeyCollection& InItems 
 ) const
+{
+	return GetPose(bInitial, InElementType, TArrayView<const FRigElementKey>(InItems.Keys.GetData(), InItems.Num()));
+}
+
+FRigPose URigHierarchy::GetPose(bool bInitial, ERigElementType InElementType,
+	const TArrayView<const FRigElementKey>& InItems) const
 {
 	FRigPose Pose;
 	Pose.HierarchyTopologyVersion = GetTopologyVersion();
@@ -1847,6 +1853,12 @@ void URigHierarchy::SetPose(
 	const FRigElementKeyCollection& InItems,
 	float InWeight
 )
+{
+	SetPose(InPose, InTransformType, InElementType, TArrayView<const FRigElementKey>(InItems.Keys.GetData(), InItems.Num()), InWeight);
+}
+
+void URigHierarchy::SetPose(const FRigPose& InPose, ERigTransformType::Type InTransformType,
+	ERigElementType InElementType, const TArrayView<const FRigElementKey>& InItems, float InWeight)
 {
 	const float U = FMath::Clamp(InWeight, 0.f, 1.f);
 	if(U < SMALL_NUMBER)
diff --git a/Engine/Plugins/Animation/ControlRig/Source/ControlRig/Private/Units/Execution/RigUnit_Collection.cpp b/Engine/Plugins/Animation/ControlRig/Source/ControlRig/Private/Units/Execution/RigUnit_Collection.cpp
--- a/Engine/Plugins/Animation/ControlRig/Source/ControlRig/Private/Units/Execution/RigUnit_Collection.cpp
+++ b/Engine/Plugins/Animation/ControlRig/Source/ControlRig/Private/Units/Execution/RigUnit_Collection.cpp
@@ -10,8 +10,11 @@
 
 FRigUnit_CollectionChain_Execute()
 {
-    DECLARE_SCOPE_HIERARCHICAL_COUNTER_RIGUNIT()
+	FRigUnit_CollectionChainArray::StaticExecute(RigVMExecuteContext, FirstItem, LastItem, Reverse, Collection.Keys, CachedCollection, CachedHierarchyHash, Context);
+}
 
+FRigUnit_CollectionChainArray_Execute()
+{
 	if (Context.State == EControlRigState::Init)
 	{
 		CachedHierarchyHash = INDEX_NONE;
@@ -46,10 +49,15 @@ FRigUnit_CollectionChain_Execute()
 		}
 	}
 
-	Collection = CachedCollection;
+	Items = CachedCollection.Keys;
 }
 
 FRigUnit_CollectionNameSearch_Execute()
+{
+	FRigUnit_CollectionNameSearchArray::StaticExecute(RigVMExecuteContext, PartialName, TypeToSearch, Collection.Keys, CachedCollection, CachedHierarchyHash, Context);
+}
+
+FRigUnit_CollectionNameSearchArray_Execute()
 {
 	DECLARE_SCOPE_HIERARCHICAL_COUNTER_RIGUNIT()
 
@@ -68,10 +76,15 @@ FRigUnit_CollectionNameSearch_Execute()
 		CachedCollection = FRigElementKeyCollection::MakeFromName(Context.Hierarchy, PartialName, (uint8)TypeToSearch);
 	}
 
-	Collection = CachedCollection;
+	Items = CachedCollection.Keys;
 }
 
 FRigUnit_CollectionChildren_Execute()
+{
+	FRigUnit_CollectionChildrenArray::StaticExecute(RigVMExecuteContext, Parent, bIncludeParent, bRecursive, TypeToSearch, Collection.Keys, CachedCollection, CachedHierarchyHash, Context);
+}
+
+FRigUnit_CollectionChildrenArray_Execute()
 {
 	DECLARE_SCOPE_HIERARCHICAL_COUNTER_RIGUNIT()
 
@@ -99,7 +112,7 @@ FRigUnit_CollectionChildren_Execute()
 		}
 	}
 
-	Collection = CachedCollection;
+	Items = CachedCollection.Keys;
 }
 
 #if WITH_EDITOR
@@ -142,6 +155,11 @@ IMPLEMENT_RIGUNIT_AUTOMATION_TEST(FRigUnit_CollectionChildren)
 #endif
 
 FRigUnit_CollectionReplaceItems_Execute()
+{
+	FRigUnit_CollectionReplaceItemsArray::StaticExecute(RigVMExecuteContext, Items.Keys, Old, New, RemoveInvalidItems, bAllowDuplicates, Collection.Keys, CachedCollection, CachedHierarchyHash, Context);
+}
+
+FRigUnit_CollectionReplaceItemsArray_Execute()
 {
 	DECLARE_SCOPE_HIERARCHICAL_COUNTER_RIGUNIT()
 
@@ -184,7 +202,7 @@ FRigUnit_CollectionReplaceItems_Execute()
 		}
 	}
 
-	Collection = CachedCollection;
+	Result = CachedCollection.Keys;
 }
 
 FRigUnit_CollectionItems_Execute()
@@ -213,6 +231,11 @@ FRigUnit_CollectionGetItems_Execute()
 }
 
 FRigUnit_CollectionGetParentIndices_Execute()
+{
+	FRigUnit_CollectionGetParentIndicesItemArray::StaticExecute(RigVMExecuteContext, Collection.Keys, ParentIndices, Context);
+}
+
+FRigUnit_CollectionGetParentIndicesItemArray_Execute()
 {
 	DECLARE_SCOPE_HIERARCHICAL_COUNTER_RIGUNIT()
 
@@ -222,19 +245,19 @@ FRigUnit_CollectionGetParentIndices_Execute()
 		return;
 	}
 	
-	ParentIndices.SetNumUninitialized(Collection.Num());
+	ParentIndices.SetNumUninitialized(Items.Num());
 
-	for(int32 Index=0;Index<Collection.Num();Index++)
+	for(int32 Index=0;Index<Items.Num();Index++)
 	{
 		ParentIndices[Index] = INDEX_NONE;
 
-		const int32 ItemIndex = Context.Hierarchy->GetIndex(Collection[Index]);
+		const int32 ItemIndex = Context.Hierarchy->GetIndex(Items[Index]);
 		if(ItemIndex == INDEX_NONE)
 		{
 			continue;;
 		}
 
-		switch(Collection[Index].Type)
+		switch(Items[Index].Type)
 		{
 			case ERigElementType::Curve:
 			{
@@ -270,7 +293,7 @@ FRigUnit_CollectionGetParentIndices_Execute()
 
 			while(ParentIndices[Index] == INDEX_NONE && ParentIndex != INDEX_NONE)
 			{
-				ParentIndices[Index] = Collection.GetKeys().Find(Context.Hierarchy->GetKey(ParentIndex));
+				ParentIndices[Index] = Items.Find(Context.Hierarchy->GetKey(ParentIndex));
 				ParentIndex = Context.Hierarchy->GetFirstParent(ParentIndex);
 			}
 		}
diff --git a/Engine/Plugins/Animation/ControlRig/Source/ControlRig/Private/Units/Execution/RigUnit_Collection.h b/Engine/Plugins/Animation/ControlRig/Source/ControlRig/Private/Units/Execution/RigUnit_Collection.h
--- a/Engine/Plugins/Animation/ControlRig/Source/ControlRig/Private/Units/Execution/RigUnit_Collection.h
+++ b/Engine/Plugins/Animation/ControlRig/Source/ControlRig/Private/Units/Execution/RigUnit_Collection.h
@@ -6,13 +6,13 @@
 #include "ControlRigDefines.h" 
 #include "RigUnit_Collection.generated.h"
 
-USTRUCT(meta = (Abstract, NodeColor = "0.7 0.05 0.5", Category = "Collections"))
+USTRUCT(meta = (Abstract, NodeColor = "0.7 0.05 0.5", Category = "Items"))
 struct CONTROLRIG_API FRigUnit_CollectionBase : public FRigUnit
 {
 	GENERATED_BODY()
 };
 
-USTRUCT(meta = (Abstract, NodeColor = "0.7 0.05 0.5", Category = "Collections"))
+USTRUCT(meta = (Abstract, NodeColor = "0.7 0.05 0.5", Category = "Items"))
 struct CONTROLRIG_API FRigUnit_CollectionBaseMutable : public FRigUnitMutable
 {
 	GENERATED_BODY()
@@ -22,7 +22,7 @@ struct CONTROLRIG_API FRigUnit_CollectionBaseMutable : public FRigUnitMutable
  * Creates a collection based on a first and last item within a chain.
  * Chains can refer to bone chains or chains within a control hierarchy.
  */
-USTRUCT(meta=(DisplayName="Item Chain", Keywords="Bone,Joint,Collection", Varying))
+USTRUCT(meta=(DisplayName="Item Chain", Keywords="Bone,Joint,Collection", Varying, Deprecated = "5.0"))
 struct CONTROLRIG_API FRigUnit_CollectionChain : public FRigUnit_CollectionBase
 {
 	GENERATED_BODY()
@@ -56,11 +56,49 @@ struct CONTROLRIG_API FRigUnit_CollectionChain : public FRigUnit_CollectionBase
 	int32 CachedHierarchyHash;
 };
 
+/**
+* Creates an item array based on a first and last item within a chain.
+* Chains can refer to bone chains or chains within a control hierarchy.
+*/
+USTRUCT(meta=(DisplayName="Item Chain", Keywords="Bone,Joint,Collection", Varying))
+struct CONTROLRIG_API FRigUnit_CollectionChainArray : public FRigUnit_CollectionBase
+{
+	GENERATED_BODY()
+
+	FRigUnit_CollectionChainArray()
+	{
+		FirstItem = LastItem = FRigElementKey(NAME_None, ERigElementType::Bone);
+		Reverse = false;
+		CachedHierarchyHash = INDEX_NONE;
+	}
+
+	RIGVM_METHOD()
+	virtual void Execute(const FRigUnitContext& Context) override;
+
+	UPROPERTY(meta = (Input, ExpandByDefault))
+	FRigElementKey FirstItem;
+
+	UPROPERTY(meta = (Input, ExpandByDefault))
+	FRigElementKey LastItem;
+
+	UPROPERTY(meta = (Input))
+	bool Reverse;
+
+	UPROPERTY(meta = (Output))
+	TArray<FRigElementKey> Items;
+
+	UPROPERTY()
+	FRigElementKeyCollection CachedCollection;
+
+	UPROPERTY()
+	int32 CachedHierarchyHash;
+};
+
 /**
  * Creates a collection based on a name search.
  * The name search is case sensitive.
  */
-USTRUCT(meta = (DisplayName = "Item Name Search", Keywords = "Bone,Joint,Collection,Filter", Varying))
+USTRUCT(meta = (DisplayName = "Item Name Search", Keywords = "Bone,Joint,Collection,Filter", Varying, Deprecated = "5.0"))
 struct CONTROLRIG_API FRigUnit_CollectionNameSearch : public FRigUnit_CollectionBase
 {
 	GENERATED_BODY()
@@ -91,11 +129,46 @@ struct CONTROLRIG_API FRigUnit_CollectionNameSearch : public FRigUnit_Collection
 	int32 CachedHierarchyHash;
 };
 
+/**
+* Creates an item array based on a name search.
+* The name search is case sensitive.
+*/
+USTRUCT(meta = (DisplayName = "Item Name Search", Keywords = "Bone,Joint,Collection,Filter", Varying))
+struct CONTROLRIG_API FRigUnit_CollectionNameSearchArray : public FRigUnit_CollectionBase
+{
+	GENERATED_BODY()
+
+	FRigUnit_CollectionNameSearchArray()
+	{
+		PartialName = NAME_None;
+		TypeToSearch = ERigElementType::All;
+		CachedHierarchyHash = INDEX_NONE;
+	}
+
+	RIGVM_METHOD()
+	virtual void Execute(const FRigUnitContext& Context) override;
+
+	UPROPERTY(meta = (Input))
+	FName PartialName;
+
+	UPROPERTY(meta = (Input))
+	ERigElementType TypeToSearch;
+
+	UPROPERTY(meta = (Output))
+	TArray<FRigElementKey> Items;
+
+	UPROPERTY()
+	FRigElementKeyCollection CachedCollection;
+
+	UPROPERTY()
+	int32 CachedHierarchyHash;
+};
+
 /**
  * Creates a collection based on the direct or recursive children
  * of a provided parent item. Returns an empty collection for an invalid parent item.
  */
-USTRUCT(meta = (DisplayName = "Children", Keywords = "Bone,Joint,Collection,Filter,Parent", Varying))
+USTRUCT(meta = (DisplayName = "Children", Keywords = "Bone,Joint,Collection,Filter,Parent", Varying, Deprecated = "5.0"))
 struct CONTROLRIG_API FRigUnit_CollectionChildren : public FRigUnit_CollectionBase
 {
 	GENERATED_BODY()
@@ -134,10 +207,53 @@ struct CONTROLRIG_API FRigUnit_CollectionChildren : public FRigUnit_CollectionBa
 	int32 CachedHierarchyHash;
 };
 
+/**
+* Creates an item array based on the direct or recursive children
+* of a provided parent item. Returns an empty collection for an invalid parent item.
+*/
+USTRUCT(meta = (DisplayName = "Children", Keywords = "Bone,Joint,Collection,Filter,Parent", Varying))
+struct CONTROLRIG_API FRigUnit_CollectionChildrenArray : public FRigUnit_CollectionBase
+{
+	GENERATED_BODY()
+
+	FRigUnit_CollectionChildrenArray()
+	{
+		Parent = FRigElementKey(NAME_None, ERigElementType::Bone);
+		bIncludeParent = false;
+		bRecursive = false;
+		TypeToSearch = ERigElementType::All;
+		CachedHierarchyHash = INDEX_NONE;
+	}
+
+	RIGVM_METHOD()
+	virtual void Execute(const FRigUnitContext& Context) override;
+
+	UPROPERTY(meta = (Input, ExpandByDefault))
+	FRigElementKey Parent;
+
+	UPROPERTY(meta = (Input))
+	bool bIncludeParent;
+
+	UPROPERTY(meta = (Input))
+	bool bRecursive;
+
+	UPROPERTY(meta = (Input))
+	ERigElementType TypeToSearch;
+
+	UPROPERTY(meta = (Output))
+	TArray<FRigElementKey> Items;
+
+	UPROPERTY()
+	FRigElementKeyCollection CachedCollection;
+
+	UPROPERTY()
+	int32 CachedHierarchyHash;
+};
+
 /**
  * Replaces all names within the collection
  */
-USTRUCT(meta = (DisplayName = "Replace Items", Keywords = "Replace,Find", Varying))
+USTRUCT(meta = (DisplayName = "Replace Items", Keywords = "Replace,Find,Collection", Varying, Deprecated = "5.0"))
 struct CONTROLRIG_API FRigUnit_CollectionReplaceItems : public FRigUnit_CollectionBase
 {
 	GENERATED_BODY()
@@ -178,10 +294,54 @@ struct CONTROLRIG_API FRigUnit_CollectionReplaceItems : public FRigUnit_Collecti
 	int32 CachedHierarchyHash;
 };
 
+/**
+* Replaces all names within the item array
+*/
+USTRUCT(meta = (DisplayName = "Replace Items", Keywords = "Replace,Find,Collection", Varying))
+struct CONTROLRIG_API FRigUnit_CollectionReplaceItemsArray : public FRigUnit_CollectionBase
+{
+	GENERATED_BODY()
+
+	FRigUnit_CollectionReplaceItemsArray()
+	{
+		Old = New = NAME_None;
+		RemoveInvalidItems = false;
+		bAllowDuplicates = false;
+		CachedHierarchyHash = INDEX_NONE;
+	}
+
+	RIGVM_METHOD()
+	virtual void Execute(const FRigUnitContext& Context) override;
+
+	UPROPERTY(meta = (Input))
+	TArray<FRigElementKey> Items;
+
+	UPROPERTY(meta = (Input))
+	FName Old;
+
+	UPROPERTY(meta = (Input))
+	FName New;
+
+	UPROPERTY(meta = (Input))
+	bool RemoveInvalidItems;
+
+	UPROPERTY(meta = (Input))
+	bool bAllowDuplicates;
+
+	UPROPERTY(meta = (Output))
+	TArray<FRigElementKey> Result;
+
+	UPROPERTY()
+	FRigElementKeyCollection CachedCollection;
+
+	UPROPERTY()
+	int32 CachedHierarchyHash;
+};
+
 /**
  * Returns a collection provided a specific list of items.
  */
-USTRUCT(meta = (DisplayName = "Items", Keywords = "Collection,Array", Varying))
+USTRUCT(meta = (DisplayName = "Collection from Items", Keywords = "Collection,Array", Varying))
 struct CONTROLRIG_API FRigUnit_CollectionItems : public FRigUnit_CollectionBase
 {
 	GENERATED_BODY()
@@ -208,7 +368,7 @@ struct CONTROLRIG_API FRigUnit_CollectionItems : public FRigUnit_CollectionBase
 /**
 * Returns a collection provided a specific list of items.
 */
-USTRUCT(meta = (DisplayName = "Get Items", Keywords = "Collection,Array", Varying))
+USTRUCT(meta = (DisplayName = "Get Items from Collection", Keywords = "Collection,Array", Varying))
 struct CONTROLRIG_API FRigUnit_CollectionGetItems : public FRigUnit_CollectionBase
 {
 	GENERATED_BODY()
@@ -233,7 +393,7 @@ struct CONTROLRIG_API FRigUnit_CollectionGetItems : public FRigUnit_CollectionBa
  * b) If an item has a parent that's not part of the collection INDEX_NONE will be returned.
  * c) If an item has a parent that's not part of the collection, but a grand parent is we'll use that index instead.
  */
-USTRUCT(meta = (DisplayName = "Get Parent Indices", Keywords = "Collection,Array", Varying))
+USTRUCT(meta = (DisplayName = "Get Parent Indices", Keywords = "Collection,Array", Varying, Deprecated = "5.0"))
 struct CONTROLRIG_API FRigUnit_CollectionGetParentIndices : public FRigUnit_CollectionBase
 {
 	GENERATED_BODY()
@@ -252,12 +412,36 @@ struct CONTROLRIG_API FRigUnit_CollectionGetParentIndices : public FRigUnit_Coll
 	TArray<int32> ParentIndices;
 };
 
+/**
+* Returns an array of relative parent indices for each item. Several options here
+* a) If an item has multiple parents the major parent (based on the weights) will be returned.
+* b) If an item has a parent that's not part of the collection INDEX_NONE will be returned.
+* c) If an item has a parent that's not part of the collection, but a grand parent is we'll use that index instead.
+*/
+USTRUCT(meta = (DisplayName = "Get Parent Indices", Keywords = "Collection,Array", Varying))
+struct CONTROLRIG_API FRigUnit_CollectionGetParentIndicesItemArray : public FRigUnit_CollectionBase
+{
+	GENERATED_BODY()
+
+	FRigUnit_CollectionGetParentIndicesItemArray()
+	{
+	}
+
+	RIGVM_METHOD()
+	virtual void Execute(const FRigUnitContext& Context) override;
+
+	UPROPERTY(meta = (Input))
+	TArray<FRigElementKey> Items;
+
+	UPROPERTY(meta = (Output))
+	TArray<int32> ParentIndices;
+};
 
 /**
  * Returns the union of two provided collections
  * (the combination of all items from both A and B).
  */
-USTRUCT(meta = (DisplayName = "Union", Keywords = "Combine,Add,Merge,Collection,Hierarchy"))
+USTRUCT(meta = (DisplayName = "Union", Keywords = "Combine,Add,Merge,Collection,Hierarchy", Deprecated = "5.0"))
 struct CONTROLRIG_API FRigUnit_CollectionUnion : public FRigUnit_CollectionBase
 {
 	GENERATED_BODY()
@@ -287,7 +471,7 @@ struct CONTROLRIG_API FRigUnit_CollectionUnion : public FRigUnit_CollectionBase
  * Returns the intersection of two provided collections
  * (the items present in both A and B).
  */
-USTRUCT(meta = (DisplayName = "Intersection", Keywords = "Combine,Merge,Collection,Hierarchy"))
+USTRUCT(meta = (DisplayName = "Intersection", Keywords = "Combine,Merge,Collection,Hierarchy", Deprecated = "5.0"))
 struct CONTROLRIG_API FRigUnit_CollectionIntersection : public FRigUnit_CollectionBase
 {
 	GENERATED_BODY()
@@ -313,7 +497,7 @@ struct CONTROLRIG_API FRigUnit_CollectionIntersection : public FRigUnit_Collecti
  * Returns the difference between two collections
  * (the items present in A but not in B).
  */
-USTRUCT(meta = (DisplayName = "Difference", Keywords = "Collection,Exclude,Subtract"))
+USTRUCT(meta = (DisplayName = "Difference", Keywords = "Collection,Exclude,Subtract", Deprecated = "5.0"))
 struct CONTROLRIG_API FRigUnit_CollectionDifference : public FRigUnit_CollectionBase
 {
 	GENERATED_BODY()
@@ -338,7 +522,7 @@ struct CONTROLRIG_API FRigUnit_CollectionDifference : public FRigUnit_Collection
 /**
  * Returns the collection in reverse order
  */
-USTRUCT(meta = (DisplayName = "Reverse", Keywords = "Direction,Order,Reverse"))
+USTRUCT(meta = (DisplayName = "Reverse", Keywords = "Direction,Order,Reverse", Deprecated = "5.0"))
 struct CONTROLRIG_API FRigUnit_CollectionReverse : public FRigUnit_CollectionBase
 {
 	GENERATED_BODY()
@@ -360,7 +544,7 @@ struct CONTROLRIG_API FRigUnit_CollectionReverse : public FRigUnit_CollectionBas
 /**
  * Returns the number of elements in a collection
  */
-USTRUCT(meta = (DisplayName = "Count", Keywords = "Collection,Array,Count,Num,Length,Size"))
+USTRUCT(meta = (DisplayName = "Count", Keywords = "Collection,Array,Count,Num,Length,Size", Deprecated = "5.0"))
 struct CONTROLRIG_API FRigUnit_CollectionCount : public FRigUnit_CollectionBase
 {
 	GENERATED_BODY()
@@ -384,7 +568,7 @@ struct CONTROLRIG_API FRigUnit_CollectionCount : public FRigUnit_CollectionBase
 /**
  * Returns a single item within a collection by index
  */
-USTRUCT(meta = (DisplayName = "Item At Index", Keywords = "Item,GetIndex,AtIndex,At,ForIndex,[]"))
+USTRUCT(meta = (DisplayName = "Item At Index", Keywords = "Item,GetIndex,AtIndex,At,ForIndex,[]", Deprecated = "5.0"))
 struct CONTROLRIG_API FRigUnit_CollectionItemAtIndex : public FRigUnit_CollectionBase
 {
 	GENERATED_BODY()
@@ -412,7 +596,7 @@ struct CONTROLRIG_API FRigUnit_CollectionItemAtIndex : public FRigUnit_Collectio
 /**
  * Given a collection of items, execute iteratively across all items in a given collection
  */
-USTRUCT(meta=(DisplayName="For Each Item", Keywords="Collection,Loop,Iterate", Icon="EditorStyle|GraphEditor.Macro.ForEach_16x"))
+USTRUCT(meta=(DisplayName="For Each Item", Keywords="Collection,Loop,Iterate", Icon="EditorStyle|GraphEditor.Macro.ForEach_16x", Deprecated = "5.0"))
 struct CONTROLRIG_API FRigUnit_CollectionLoop : public FRigUnit_CollectionBaseMutable
 {
 	GENERATED_BODY()
@@ -462,7 +646,7 @@ struct CONTROLRIG_API FRigUnit_CollectionLoop : public FRigUnit_CollectionBaseMu
 /**
 * Adds an element to an existing collection
 */
-USTRUCT(meta = (DisplayName = "Add Item", Keywords = "Item,Add,Push,Insert"))
+USTRUCT(meta = (DisplayName = "Add Item", Keywords = "Item,Add,Push,Insert", Deprecated = "5.0"))
 struct CONTROLRIG_API FRigUnit_CollectionAddItem : public FRigUnit_CollectionBase
 {
 	GENERATED_BODY()
diff --git a/Engine/Plugins/Animation/ControlRig/Source/ControlRig/Private/Units/Execution/RigUnit_DynamicHierarchy.cpp b/Engine/Plugins/Animation/ControlRig/Source/ControlRig/Private/Units/Execution/RigUnit_DynamicHierarchy.cpp
--- a/Engine/Plugins/Animation/ControlRig/Source/ControlRig/Private/Units/Execution/RigUnit_DynamicHierarchy.cpp
+++ b/Engine/Plugins/Animation/ControlRig/Source/ControlRig/Private/Units/Execution/RigUnit_DynamicHierarchy.cpp
@@ -97,6 +97,11 @@ FRigUnit_SwitchParent_Execute()
 }
 
 FRigUnit_HierarchyGetParentWeights_Execute()
+{
+	FRigUnit_HierarchyGetParentWeightsArray::StaticExecute(RigVMExecuteContext, Child, Weights, Parents.Keys, Context);
+}
+
+FRigUnit_HierarchyGetParentWeightsArray_Execute()
 {
 	if((Context.State != EControlRigState::Update) || (Context.Hierarchy == nullptr))
 	{
@@ -109,7 +114,9 @@ FRigUnit_HierarchyGetParentWeights_Execute()
 		UE_CONTROLRIG_RIGUNIT_REPORT_ERROR(TEXT("Item %s does not exist."), *Child.ToString())
 		return;
 	}
+	
 	Weights = Context.Hierarchy->GetParentWeightArray(ChildElement, false);
+	Parents = Context.Hierarchy->GetParents(ChildElement->GetKey(), false);
 }
 
 FRigUnit_HierarchySetParentWeights_Execute()
diff --git a/Engine/Plugins/Animation/ControlRig/Source/ControlRig/Private/Units/Execution/RigUnit_DynamicHierarchy.h b/Engine/Plugins/Animation/ControlRig/Source/ControlRig/Private/Units/Execution/RigUnit_DynamicHierarchy.h
--- a/Engine/Plugins/Animation/ControlRig/Source/ControlRig/Private/Units/Execution/RigUnit_DynamicHierarchy.h
+++ b/Engine/Plugins/Animation/ControlRig/Source/ControlRig/Private/Units/Execution/RigUnit_DynamicHierarchy.h
@@ -97,7 +97,7 @@ struct CONTROLRIG_API FRigUnit_SwitchParent : public FRigUnit_DynamicHierarchyBa
 /**
 * Returns the item's parent's weights
 */
-USTRUCT(meta=(DisplayName="Get Parent Weights", Keywords="Chain,Parents,Hierarchy", Varying))
+USTRUCT(meta=(DisplayName="Get Parent Weights", Keywords="Chain,Parents,Hierarchy", Varying, Deprecated = "5.0"))
 struct CONTROLRIG_API FRigUnit_HierarchyGetParentWeights : public FRigUnit_DynamicHierarchyBase
 {
 	GENERATED_BODY()
@@ -120,6 +120,32 @@ struct CONTROLRIG_API FRigUnit_HierarchyGetParentWeights : public FRigUnit_Dynam
 	FRigElementKeyCollection Parents;
 };
 
+/**
+* Returns the item's parent's weights
+*/
+USTRUCT(meta=(DisplayName="Get Parent Weights", Keywords="Chain,Parents,Hierarchy", Varying))
+struct CONTROLRIG_API FRigUnit_HierarchyGetParentWeightsArray : public FRigUnit_DynamicHierarchyBase
+{
+	GENERATED_BODY()
+
+	FRigUnit_HierarchyGetParentWeightsArray()
+	{
+		Child = FRigElementKey(NAME_None, ERigElementType::Control);
+	}
+
+	RIGVM_METHOD()
+	virtual void Execute(const FRigUnitContext& Context) override;
+
+	UPROPERTY(meta = (Input, ExpandByDefault))
+	FRigElementKey Child;
+
+	UPROPERTY(meta = (Output))
+	TArray<FRigElementWeight> Weights;
+
+	UPROPERTY(meta = (Output))
+	TArray<FRigElementKey> Parents;
+};
+
 /**
 * Sets the item's parent's weights
 */
diff --git a/Engine/Plugins/Animation/ControlRig/Source/ControlRig/Private/Units/Execution/RigUnit_Hierarchy.cpp b/Engine/Plugins/Animation/ControlRig/Source/ControlRig/Private/Units/Execution/RigUnit_Hierarchy.cpp
--- a/Engine/Plugins/Animation/ControlRig/Source/ControlRig/Private/Units/Execution/RigUnit_Hierarchy.cpp
+++ b/Engine/Plugins/Animation/ControlRig/Source/ControlRig/Private/Units/Execution/RigUnit_Hierarchy.cpp
@@ -35,6 +35,11 @@ FRigUnit_HierarchyGetParent_Execute()
 }
 
 FRigUnit_HierarchyGetParents_Execute()
+{
+	FRigUnit_HierarchyGetParentsItemArray::StaticExecute(RigVMExecuteContext, Child, bIncludeChild, bReverse, Parents.Keys, CachedChild, CachedParents, Context);
+}
+
+FRigUnit_HierarchyGetParentsItemArray_Execute()
 {
     DECLARE_SCOPE_HIERARCHICAL_COUNTER_RIGUNIT()
 
@@ -70,7 +75,7 @@ FRigUnit_HierarchyGetParents_Execute()
 		}
 	}
 
-	Parents = CachedParents;
+	Parents = CachedParents.Keys;
 }
 
 FRigUnit_HierarchyGetChildren_Execute()
@@ -107,6 +112,11 @@ FRigUnit_HierarchyGetChildren_Execute()
 }
 
 FRigUnit_HierarchyGetSiblings_Execute()
+{
+	FRigUnit_HierarchyGetSiblingsItemArray::StaticExecute(RigVMExecuteContext, Item, bIncludeItem, Siblings.Keys, CachedItem, CachedSiblings, Context);
+}
+
+FRigUnit_HierarchyGetSiblingsItemArray_Execute()
 {
     DECLARE_SCOPE_HIERARCHICAL_COUNTER_RIGUNIT()
 
@@ -146,15 +156,25 @@ FRigUnit_HierarchyGetSiblings_Execute()
 		}
 	}
 
-	Siblings = CachedSiblings;
+	Siblings = CachedSiblings.Keys;
 }
 
 FRigUnit_HierarchyGetPose_Execute()
+{
+	FRigUnit_HierarchyGetPoseItemArray::StaticExecute(RigVMExecuteContext, Initial, ElementType, ItemsToGet.Keys, Pose, Context);
+}
+
+FRigUnit_HierarchyGetPoseItemArray_Execute()
 {
 	Pose = Context.Hierarchy->GetPose(Initial, ElementType, ItemsToGet);
 }
 
 FRigUnit_HierarchySetPose_Execute()
+{
+	FRigUnit_HierarchySetPoseItemArray::StaticExecute(RigVMExecuteContext, Pose, ElementType, Space, ItemsToSet.Keys, Weight, ExecuteContext, Context);
+}
+
+FRigUnit_HierarchySetPoseItemArray_Execute()
 {
 	ExecuteContext.Hierarchy->SetPose(
 		Pose,
@@ -171,6 +191,11 @@ FRigUnit_PoseIsEmpty_Execute()
 }
 
 FRigUnit_PoseGetItems_Execute()
+{
+	FRigUnit_PoseGetItemsItemArray::StaticExecute(RigVMExecuteContext, Pose, ElementType, Items.Keys, Context);
+}
+
+FRigUnit_PoseGetItemsItemArray_Execute()
 {
 	Items.Reset();
 
diff --git a/Engine/Plugins/Animation/ControlRig/Source/ControlRig/Private/Units/Execution/RigUnit_Hierarchy.h b/Engine/Plugins/Animation/ControlRig/Source/ControlRig/Private/Units/Execution/RigUnit_Hierarchy.h
--- a/Engine/Plugins/Animation/ControlRig/Source/ControlRig/Private/Units/Execution/RigUnit_Hierarchy.h
+++ b/Engine/Plugins/Animation/ControlRig/Source/ControlRig/Private/Units/Execution/RigUnit_Hierarchy.h
@@ -53,7 +53,7 @@ struct CONTROLRIG_API FRigUnit_HierarchyGetParent : public FRigUnit_HierarchyBas
 /**
  * Returns the item's parents
  */
-USTRUCT(meta=(DisplayName="Get Parents", Keywords="Chain,Parents,Hierarchy", Varying))
+USTRUCT(meta=(DisplayName="Get Parents", Keywords="Chain,Parents,Hierarchy", Varying, Deprecated = "5.0"))
 struct CONTROLRIG_API FRigUnit_HierarchyGetParents : public FRigUnit_HierarchyBase
 {
 	GENERATED_BODY()
@@ -91,6 +91,48 @@ struct CONTROLRIG_API FRigUnit_HierarchyGetParents : public FRigUnit_HierarchyBa
 	FRigElementKeyCollection CachedParents;
 };
 
+/**
+ * Returns the item's parents
+ */
+USTRUCT(meta=(DisplayName="Get Parents", Keywords="Chain,Parents,Hierarchy", Varying))
+struct CONTROLRIG_API FRigUnit_HierarchyGetParentsItemArray : public FRigUnit_HierarchyBase
+{
+	GENERATED_BODY()
+
+	FRigUnit_HierarchyGetParentsItemArray()
+	{
+		Child = FRigElementKey(NAME_None, ERigElementType::Bone);
+		CachedChild = FCachedRigElement();
+		CachedParents = FRigElementKeyCollection();
+		bIncludeChild = false;
+		bReverse = false;
+	}
+
+	RIGVM_METHOD()
+	virtual void Execute(const FRigUnitContext& Context) override;
+
+	UPROPERTY(meta = (Input, ExpandByDefault))
+	FRigElementKey Child;
+
+	UPROPERTY(meta = (Input))
+	bool bIncludeChild;
+
+	UPROPERTY(meta = (Input))
+	bool bReverse;
+
+	UPROPERTY(meta = (Output))
+	TArray<FRigElementKey> Parents;
+
+	// Used to cache the internally
+	UPROPERTY()
+	FCachedRigElement CachedChild;
+
+	// Used to cache the internally
+	UPROPERTY()
+	FRigElementKeyCollection CachedParents;
+};
+
+
 /**
  * Returns the item's children
  */
@@ -135,7 +177,7 @@ struct CONTROLRIG_API FRigUnit_HierarchyGetChildren : public FRigUnit_HierarchyB
 /**
  * Returns the item's siblings
  */
-USTRUCT(meta=(DisplayName="Get Siblings", Keywords="Chain,Siblings,Hierarchy", Varying))
+USTRUCT(meta=(DisplayName="Get Siblings", Keywords="Chain,Siblings,Hierarchy", Varying, Deprecated = "5.0"))
 struct CONTROLRIG_API FRigUnit_HierarchyGetSiblings : public FRigUnit_HierarchyBase
 {
 	GENERATED_BODY()
@@ -169,10 +211,47 @@ struct CONTROLRIG_API FRigUnit_HierarchyGetSiblings : public FRigUnit_HierarchyB
 	FRigElementKeyCollection CachedSiblings;
 };
 
+/**
+ * Returns the item's siblings
+ */
+USTRUCT(meta=(DisplayName="Get Siblings", Keywords="Chain,Siblings,Hierarchy", Varying))
+struct CONTROLRIG_API FRigUnit_HierarchyGetSiblingsItemArray : public FRigUnit_HierarchyBase
+{
+	GENERATED_BODY()
+
+	FRigUnit_HierarchyGetSiblingsItemArray()
+	{
+		Item = FRigElementKey(NAME_None, ERigElementType::Bone);
+		CachedItem = FCachedRigElement();
+		CachedSiblings = FRigElementKeyCollection();
+		bIncludeItem = false;
+	}
+
+	RIGVM_METHOD()
+	virtual void Execute(const FRigUnitContext& Context) override;
+
+	UPROPERTY(meta = (Input, ExpandByDefault))
+	FRigElementKey Item;
+
+	UPROPERTY(meta = (Input))
+	bool bIncludeItem;
+
+	UPROPERTY(meta = (Output))
+	TArray<FRigElementKey> Siblings;
+
+	// Used to cache the internally
+	UPROPERTY()
+	FCachedRigElement CachedItem;
+
+	// Used to cache the internally
+	UPROPERTY()
+	FRigElementKeyCollection CachedSiblings;
+};
+
 /**
  * Returns the hierarchy's pose
  */
-USTRUCT(meta=(DisplayName="Get Pose Cache", Keywords="Hierarchy,Pose,State", Varying))
+USTRUCT(meta=(DisplayName="Get Pose Cache", Keywords="Hierarchy,Pose,State", Varying, Deprecated = "5.0"))
 struct CONTROLRIG_API FRigUnit_HierarchyGetPose : public FRigUnit_HierarchyBase
 {
 	GENERATED_BODY()
@@ -202,10 +281,43 @@ struct CONTROLRIG_API FRigUnit_HierarchyGetPose : public FRigUnit_HierarchyBase
 	FRigPose Pose;
 };
 
+/**
+ * Returns the hierarchy's pose
+ */
+USTRUCT(meta=(DisplayName="Get Pose Cache", Keywords="Hierarchy,Pose,State", Varying))
+struct CONTROLRIG_API FRigUnit_HierarchyGetPoseItemArray : public FRigUnit_HierarchyBase
+{
+	GENERATED_BODY()
+
+	FRigUnit_HierarchyGetPoseItemArray()
+	{
+		Initial = false;
+		ElementType = ERigElementType::All;
+		Pose = FRigPose();
+	}
+
+	RIGVM_METHOD()
+	virtual void Execute(const FRigUnitContext& Context) override;
+
+	UPROPERTY(meta = (Input))
+	bool Initial;
+
+	UPROPERTY(meta = (Input))
+	ERigElementType ElementType;
+
+	// An optional collection to filter against
+	UPROPERTY(meta = (Input))
+	TArray<FRigElementKey> ItemsToGet;
+
+	UPROPERTY(meta = (Output))
+	FRigPose Pose;
+};
+
+
 /**
  * Sets the hierarchy's pose
  */
-USTRUCT(meta=(DisplayName="Apply Pose Cache", Keywords="Hierarchy,Pose,State", Varying))
+USTRUCT(meta=(DisplayName="Apply Pose Cache", Keywords="Hierarchy,Pose,State", Varying, Deprecated = "5.0"))
 struct CONTROLRIG_API FRigUnit_HierarchySetPose : public FRigUnit_HierarchyBaseMutable
 {
 	GENERATED_BODY()
@@ -239,6 +351,42 @@ struct CONTROLRIG_API FRigUnit_HierarchySetPose : public FRigUnit_HierarchyBaseM
 	float Weight;
 };
 
+/**
+ * Sets the hierarchy's pose
+ */
+USTRUCT(meta=(DisplayName="Apply Pose Cache", Keywords="Hierarchy,Pose,State", Varying))
+struct CONTROLRIG_API FRigUnit_HierarchySetPoseItemArray : public FRigUnit_HierarchyBaseMutable
+{
+	GENERATED_BODY()
+
+	FRigUnit_HierarchySetPoseItemArray()
+	{
+		Pose = FRigPose();
+		ElementType = ERigElementType::All;
+		Space = EBoneGetterSetterMode::LocalSpace;
+		Weight = 1.f;
+	}
+
+	RIGVM_METHOD()
+	virtual void Execute(const FRigUnitContext& Context) override;
+
+	UPROPERTY(meta = (Input))
+	FRigPose Pose;
+
+	UPROPERTY(meta = (Input))
+	ERigElementType ElementType;
+
+	UPROPERTY(meta = (Input))
+	EBoneGetterSetterMode Space;
+
+	// An optional collection to filter against
+	UPROPERTY(meta = (Input))
+	TArray<FRigElementKey> ItemsToSet;
+
+	UPROPERTY(meta = (Input))
+	float Weight;
+};
+
 /**
 * Returns true if the hierarchy pose is empty (has no items)
 */
@@ -266,7 +414,7 @@ struct CONTROLRIG_API FRigUnit_PoseIsEmpty : public FRigUnit_HierarchyBase
 /**
 * Returns the items in the hierarchy pose
 */
-USTRUCT(meta=(DisplayName="Get Pose Cache Items", Keywords="Hierarchy,Pose,State", Varying))
+USTRUCT(meta=(DisplayName="Get Pose Cache Items", Keywords="Hierarchy,Pose,State", Varying, Deprecated = "5.0"))
 struct CONTROLRIG_API FRigUnit_PoseGetItems : public FRigUnit_HierarchyBase
 {
 	GENERATED_BODY()
@@ -291,6 +439,33 @@ struct CONTROLRIG_API FRigUnit_PoseGetItems : public FRigUnit_HierarchyBase
 	FRigElementKeyCollection Items;
 };
 
+/**
+* Returns the items in the hierarchy pose
+*/
+USTRUCT(meta=(DisplayName="Get Pose Cache Items", Keywords="Hierarchy,Pose,State", Varying))
+struct CONTROLRIG_API FRigUnit_PoseGetItemsItemArray : public FRigUnit_HierarchyBase
+{
+	GENERATED_BODY()
+
+	FRigUnit_PoseGetItemsItemArray()
+	{
+		Pose = FRigPose();
+		ElementType = ERigElementType::All;
+	}
+
+	RIGVM_METHOD()
+	virtual void Execute(const FRigUnitContext& Context) override;
+
+	UPROPERTY(meta = (Input))
+	FRigPose Pose;
+
+	UPROPERTY(meta = (Input))
+	ERigElementType ElementType;
+
+	UPROPERTY(meta = (Output))
+	TArray<FRigElementKey> Items;
+};
+
 /**
 * Compares two pose caches and compares their values.
 */
diff --git a/Engine/Plugins/Animation/ControlRig/Source/ControlRig/Private/Units/Execution/RigUnit_Item.h b/Engine/Plugins/Animation/ControlRig/Source/ControlRig/Private/Units/Execution/RigUnit_Item.h
--- a/Engine/Plugins/Animation/ControlRig/Source/ControlRig/Private/Units/Execution/RigUnit_Item.h
+++ b/Engine/Plugins/Animation/ControlRig/Source/ControlRig/Private/Units/Execution/RigUnit_Item.h
@@ -6,13 +6,13 @@
 #include "ControlRigDefines.h" 
 #include "RigUnit_Item.generated.h"
 
-USTRUCT(meta = (Abstract, NodeColor = "0.7 0.05 0.5", Category = "Collections"))
+USTRUCT(meta = (Abstract, NodeColor = "0.7 0.05 0.5", Category = "Items"))
 struct CONTROLRIG_API FRigUnit_ItemBase : public FRigUnit
 {
 	GENERATED_BODY()
 };
 
-USTRUCT(meta = (Abstract, NodeColor = "0.7 0.05 0.5", Category = "Collections"))
+USTRUCT(meta = (Abstract, NodeColor = "0.7 0.05 0.5", Category = "Items"))
 struct CONTROLRIG_API FRigUnit_ItemBaseMutable : public FRigUnitMutable
 {
 	GENERATED_BODY()
diff --git a/Engine/Plugins/Animation/ControlRig/Source/ControlRig/Private/Units/Hierarchy/RigUnit_GetTransform.cpp b/Engine/Plugins/Animation/ControlRig/Source/ControlRig/Private/Units/Hierarchy/RigUnit_GetTransform.cpp
--- a/Engine/Plugins/Animation/ControlRig/Source/ControlRig/Private/Units/Hierarchy/RigUnit_GetTransform.cpp
+++ b/Engine/Plugins/Animation/ControlRig/Source/ControlRig/Private/Units/Hierarchy/RigUnit_GetTransform.cpp
@@ -81,6 +81,11 @@ FRigUnit_GetTransform_Execute()
 }
 
 FRigUnit_GetTransformArray_Execute()
+{
+	FRigUnit_GetTransformItemArray::StaticExecute(RigVMExecuteContext, Items.Keys, Space, bInitial, Transforms, CachedIndex, Context);
+}
+
+FRigUnit_GetTransformItemArray_Execute()
 {
     DECLARE_SCOPE_HIERARCHICAL_COUNTER_RIGUNIT()
 
diff --git a/Engine/Plugins/Animation/ControlRig/Source/ControlRig/Private/Units/Hierarchy/RigUnit_GetTransform.h b/Engine/Plugins/Animation/ControlRig/Source/ControlRig/Private/Units/Hierarchy/RigUnit_GetTransform.h
--- a/Engine/Plugins/Animation/ControlRig/Source/ControlRig/Private/Units/Hierarchy/RigUnit_GetTransform.h
+++ b/Engine/Plugins/Animation/ControlRig/Source/ControlRig/Private/Units/Hierarchy/RigUnit_GetTransform.h
@@ -69,7 +69,7 @@ struct CONTROLRIG_API FRigUnit_GetTransform : public FRigUnit
 /**
 * GetTransformArray is used to retrieve an array of transforms from the hierarchy.
 */
-USTRUCT(meta=(DisplayName="Get Transform Array", Category="Hierarchy", DocumentationPolicy = "Strict", Keywords="GetBoneTransform,GetControlTransform,GetInitialTransform,GetSpaceTransform,GetTransform", NodeColor="0.462745, 1,0, 0.329412",Varying))
+USTRUCT(meta=(DisplayName="Get Transform Array", Category="Hierarchy", DocumentationPolicy = "Strict", Keywords="GetBoneTransform,GetControlTransform,GetInitialTransform,GetSpaceTransform,GetTransform", NodeColor="0.462745, 1,0, 0.329412",Varying, Deprecated = "5.0"))
 struct CONTROLRIG_API FRigUnit_GetTransformArray : public FRigUnit
 {
 	GENERATED_BODY()
@@ -112,3 +112,50 @@ struct CONTROLRIG_API FRigUnit_GetTransformArray : public FRigUnit
 	UPROPERTY()
 	TArray<FCachedRigElement> CachedIndex;
 };
+
+/**
+* GetTransformArray is used to retrieve an array of transforms from the hierarchy.
+*/
+USTRUCT(meta=(DisplayName="Get Transform Array", Category="Hierarchy", DocumentationPolicy = "Strict", Keywords="GetBoneTransform,GetControlTransform,GetInitialTransform,GetSpaceTransform,GetTransform", NodeColor="0.462745, 1,0, 0.329412",Varying))
+struct CONTROLRIG_API FRigUnit_GetTransformItemArray : public FRigUnit
+{
+	GENERATED_BODY()
+
+	FRigUnit_GetTransformItemArray()
+		: Items()
+		, Space(EBoneGetterSetterMode::GlobalSpace)
+		, bInitial(false)
+		, Transforms()
+		, CachedIndex()
+	{}
+
+	RIGVM_METHOD()
+	virtual void Execute(const FRigUnitContext& Context) override;
+
+	/**
+	* The items to retrieve the transforms for
+	*/
+	UPROPERTY(meta = (Input))
+	TArray<FRigElementKey> Items;
+
+	/**
+	* Defines if the transforms should be retrieved in local or global space
+	*/ 
+	UPROPERTY(meta = (Input))
+	EBoneGetterSetterMode Space;
+
+	/**
+	* Defines if the transforms should be retrieved as current (false) or initial (true).
+	* Initial transforms for bones and other elements in the hierarchy represent the reference pose's value.
+	*/ 
+	UPROPERTY(meta = (Input))
+	bool bInitial;
+
+	// The current transform of the given item - or identity in case it wasn't found.
+	UPROPERTY(meta=(Output))
+	TArray<FTransform> Transforms;
+
+	// Used to cache the internally
+	UPROPERTY()
+	TArray<FCachedRigElement> CachedIndex;
+};
diff --git a/Engine/Plugins/Animation/ControlRig/Source/ControlRig/Private/Units/Hierarchy/RigUnit_ParentSwitchConstraint.cpp b/Engine/Plugins/Animation/ControlRig/Source/ControlRig/Private/Units/Hierarchy/RigUnit_ParentSwitchConstraint.cpp
--- a/Engine/Plugins/Animation/ControlRig/Source/ControlRig/Private/Units/Hierarchy/RigUnit_ParentSwitchConstraint.cpp
+++ b/Engine/Plugins/Animation/ControlRig/Source/ControlRig/Private/Units/Hierarchy/RigUnit_ParentSwitchConstraint.cpp
@@ -4,6 +4,11 @@
 #include "Units/RigUnitContext.h"
 
 FRigUnit_ParentSwitchConstraint_Execute()
+{
+	FRigUnit_ParentSwitchConstraintArray::StaticExecute(RigVMExecuteContext, Subject, ParentIndex, Parents.Keys, InitialGlobalTransform, Weight, Transform, Switched, CachedSubject, CachedParent, RelativeOffset, ExecuteContext, Context);	
+}
+
+FRigUnit_ParentSwitchConstraintArray_Execute()
 {
     DECLARE_SCOPE_HIERARCHICAL_COUNTER_RIGUNIT()
 
diff --git a/Engine/Plugins/Animation/ControlRig/Source/ControlRig/Private/Units/Hierarchy/RigUnit_ParentSwitchConstraint.h b/Engine/Plugins/Animation/ControlRig/Source/ControlRig/Private/Units/Hierarchy/RigUnit_ParentSwitchConstraint.h
--- a/Engine/Plugins/Animation/ControlRig/Source/ControlRig/Private/Units/Hierarchy/RigUnit_ParentSwitchConstraint.h
+++ b/Engine/Plugins/Animation/ControlRig/Source/ControlRig/Private/Units/Hierarchy/RigUnit_ParentSwitchConstraint.h
@@ -10,7 +10,7 @@
  * The Parent Switch Constraint is used to have an item follow one of multiple parents,
  * and allowing to switch between the parent in question.
  */
-USTRUCT(meta=(DisplayName="Parent Switch Constraint", Category="Hierarchy", DocumentationPolicy="Strict", Keywords = "SpaceSwitch"))
+USTRUCT(meta=(DisplayName="Parent Switch Constraint", Category="Hierarchy", DocumentationPolicy="Strict", Keywords = "SpaceSwitch", Deprecated = "5.0"))
 struct CONTROLRIG_API FRigUnit_ParentSwitchConstraint : public FRigUnitMutable
 {
 	GENERATED_BODY()
@@ -81,3 +81,79 @@ struct CONTROLRIG_API FRigUnit_ParentSwitchConstraint : public FRigUnitMutable
 	UPROPERTY()
 	FTransform RelativeOffset;
 };
+
+/**
+ * The Parent Switch Constraint is used to have an item follow one of multiple parents,
+ * and allowing to switch between the parent in question.
+ */
+USTRUCT(meta=(DisplayName="Parent Switch Constraint", Category="Hierarchy", DocumentationPolicy="Strict", Keywords = "SpaceSwitch"))
+struct CONTROLRIG_API FRigUnit_ParentSwitchConstraintArray : public FRigUnitMutable
+{
+	GENERATED_BODY()
+
+	FRigUnit_ParentSwitchConstraintArray()
+		: Subject(NAME_None, ERigElementType::Null)
+		, ParentIndex(0)
+		, Weight(1.0f)
+		, Switched(false)
+		, CachedSubject(FCachedRigElement())
+		, CachedParent(FCachedRigElement())
+	{}
+
+	RIGVM_METHOD()
+	virtual void Execute(const FRigUnitContext& Context) override;
+
+	/**
+	 * The subject to constrain
+	 */
+	UPROPERTY(meta = (Input, ExpandByDefault))
+	FRigElementKey Subject;
+
+	/**
+	 * The parent index to use for constraining the subject
+	 */
+	UPROPERTY(meta = (Input))
+	int32 ParentIndex;
+
+	/**
+	 * The list of parents to constrain to
+	 */
+	UPROPERTY(meta = (Input))
+	TArray<FRigElementKey> Parents;
+
+	/**
+	 * The initial global transform for the subject
+	 */
+	UPROPERTY(meta = (Input))
+	FTransform InitialGlobalTransform;
+
+	/**
+	 * The weight of the change - how much the change should be applied
+	 */
+	UPROPERTY(meta = (Input, UIMin = "0.0", UIMax = "1.0"))
+	float Weight;
+
+	/**
+	 * The transform result (full without weighting)
+	 */
+	UPROPERTY(meta = (Output))
+	FTransform Transform;
+
+	/**
+	 * Returns true if the parent has changed
+	 */
+	UPROPERTY(meta = (Output))
+	bool Switched;
+
+	// Used to cache the internally used subject
+	UPROPERTY()
+	FCachedRigElement CachedSubject;
+
+	// Used to cache the internally used parent
+	UPROPERTY()
+	FCachedRigElement CachedParent;
+
+	// The cached relative offset between subject and parent
+	UPROPERTY()
+	FTransform RelativeOffset;
+};
diff --git a/Engine/Plugins/Animation/ControlRig/Source/ControlRig/Private/Units/Hierarchy/RigUnit_SetTransform.cpp b/Engine/Plugins/Animation/ControlRig/Source/ControlRig/Private/Units/Hierarchy/RigUnit_SetTransform.cpp
--- a/Engine/Plugins/Animation/ControlRig/Source/ControlRig/Private/Units/Hierarchy/RigUnit_SetTransform.cpp
+++ b/Engine/Plugins/Animation/ControlRig/Source/ControlRig/Private/Units/Hierarchy/RigUnit_SetTransform.cpp
@@ -160,6 +160,11 @@ FRigUnit_SetScale_Execute()
 
 
 FRigUnit_SetTransformArray_Execute()
+{
+	FRigUnit_SetTransformItemArray::StaticExecute(RigVMExecuteContext, Items.Keys, Space, bInitial, Transforms, Weight, bPropagateToChildren, CachedIndex, ExecuteContext, Context);
+}
+
+FRigUnit_SetTransformItemArray_Execute()
 {
 	DECLARE_SCOPE_HIERARCHICAL_COUNTER_RIGUNIT()
 
diff --git a/Engine/Plugins/Animation/ControlRig/Source/ControlRig/Private/Units/Highlevel/Hierarchy/RigUnit_CCDIK.cpp b/Engine/Plugins/Animation/ControlRig/Source/ControlRig/Private/Units/Highlevel/Hierarchy/RigUnit_CCDIK.cpp
--- a/Engine/Plugins/Animation/ControlRig/Source/ControlRig/Private/Units/Highlevel/Hierarchy/RigUnit_CCDIK.cpp
+++ b/Engine/Plugins/Animation/ControlRig/Source/ControlRig/Private/Units/Highlevel/Hierarchy/RigUnit_CCDIK.cpp
@@ -49,6 +49,11 @@ FRigUnit_CCDIK_Execute()
 }
 
 FRigUnit_CCDIKPerItem_Execute()
+{
+	FRigUnit_CCDIKItemArray::StaticExecute(RigVMExecuteContext, Items.Keys, EffectorTransform, Precision, Weight, MaxIterations, bStartFromTail, BaseRotationLimit, RotationLimits, bPropagateToChildren, WorkData, ExecuteContext, Context);
+}
+
+FRigUnit_CCDIKItemArray_Execute()
 {
     DECLARE_SCOPE_HIERARCHICAL_COUNTER_RIGUNIT()
 	URigHierarchy* Hierarchy = ExecuteContext.Hierarchy;
diff --git a/Engine/Plugins/Animation/ControlRig/Source/ControlRig/Private/Units/Highlevel/Hierarchy/RigUnit_CCDIK.h b/Engine/Plugins/Animation/ControlRig/Source/ControlRig/Private/Units/Highlevel/Hierarchy/RigUnit_CCDIK.h
--- a/Engine/Plugins/Animation/ControlRig/Source/ControlRig/Private/Units/Highlevel/Hierarchy/RigUnit_CCDIK.h
+++ b/Engine/Plugins/Animation/ControlRig/Source/ControlRig/Private/Units/Highlevel/Hierarchy/RigUnit_CCDIK.h
@@ -175,7 +175,7 @@ struct CONTROLRIG_API FRigUnit_CCDIK : public FRigUnit_HighlevelBaseMutable
  * the Cyclic Coordinate Descent Inverse Kinematics algorithm.
  * For now this node supports single effector chains only.
  */
-USTRUCT(meta=(DisplayName="CCDIK", Category="Hierarchy", Keywords="N-Bone,IK"))
+USTRUCT(meta=(DisplayName="CCDIK", Category="Hierarchy", Keywords="N-Bone,IK", Deprecated = "5.0"))
 struct CONTROLRIG_API FRigUnit_CCDIKPerItem : public FRigUnit_HighlevelBaseMutable
 {
 	GENERATED_BODY()
@@ -253,3 +253,87 @@ struct CONTROLRIG_API FRigUnit_CCDIKPerItem : public FRigUnit_HighlevelBaseMutab
 	UPROPERTY()
 	FRigUnit_CCDIK_WorkData WorkData;
 };
+
+/**
+ * The CCID solver can solve N-Bone chains using 
+ * the Cyclic Coordinate Descent Inverse Kinematics algorithm.
+ * For now this node supports single effector chains only.
+ */
+USTRUCT(meta=(DisplayName="CCDIK", Category="Hierarchy", Keywords="N-Bone,IK"))
+struct CONTROLRIG_API FRigUnit_CCDIKItemArray : public FRigUnit_HighlevelBaseMutable
+{
+	GENERATED_BODY()
+
+	RIGVM_METHOD()
+	virtual void Execute(const FRigUnitContext& Context) override;
+
+	FRigUnit_CCDIKItemArray()
+	{
+		EffectorTransform = FTransform::Identity;
+		Precision = 1.f;
+		Weight = 1.f;
+		MaxIterations = 10;
+		bStartFromTail = true;
+		bPropagateToChildren = true;
+		BaseRotationLimit = 30.f;
+	}
+
+	/**
+	 * The chain to use
+	 */
+	UPROPERTY(meta = (Input, ExpandByDefault))
+	TArray<FRigElementKey> Items;
+
+	/**
+	 * The transform of the effector in global space
+	 */
+	UPROPERTY(meta = (Input))
+	FTransform EffectorTransform;
+
+	/**
+	 * The precision to use for the fabrik solver
+	 */
+	UPROPERTY(meta = (Input, Constant))
+	float Precision;
+
+	/**
+	 * The weight of the solver - how much the IK should be applied.
+	 */
+	UPROPERTY(meta = (Input))
+	float Weight;
+
+	/**
+	 * The maximum number of iterations. Values between 4 and 16 are common.
+	 */
+	UPROPERTY(meta = (Input))
+	int32 MaxIterations;
+
+	/**
+	 * If set to true the direction of the solvers is flipped.
+	 */
+	UPROPERTY(meta = (Input, Constant))
+	bool bStartFromTail;
+
+	/**
+	 * The general rotation limit to be applied to bones
+	 */
+	UPROPERTY(meta = (Input, Constant))
+	float BaseRotationLimit;
+
+	/**
+	 * Defines the limits of rotation per bone.
+	 */
+	UPROPERTY(meta = (Input, Constant))
+	TArray<FRigUnit_CCDIK_RotationLimitPerItem> RotationLimits;
+
+	/**
+	 * If set to true all of the global transforms of the children
+	 * of this bone will be recalculated based on their local transforms.
+	 * Note: This is computationally more expensive than turning it off.
+	 */
+	UPROPERTY(meta = (Input, Constant))
+	bool bPropagateToChildren;
+
+	UPROPERTY()
+	FRigUnit_CCDIK_WorkData WorkData;
+};
diff --git a/Engine/Plugins/Animation/ControlRig/Source/ControlRig/Private/Units/Highlevel/Hierarchy/RigUnit_DistributeRotation.cpp b/Engine/Plugins/Animation/ControlRig/Source/ControlRig/Private/Units/Highlevel/Hierarchy/RigUnit_DistributeRotation.cpp
--- a/Engine/Plugins/Animation/ControlRig/Source/ControlRig/Private/Units/Highlevel/Hierarchy/RigUnit_DistributeRotation.cpp
+++ b/Engine/Plugins/Animation/ControlRig/Source/ControlRig/Private/Units/Highlevel/Hierarchy/RigUnit_DistributeRotation.cpp
@@ -57,6 +57,11 @@ FRigUnit_DistributeRotation_Execute()
 }
 
 FRigUnit_DistributeRotationForCollection_Execute()
+{
+	FRigUnit_DistributeRotationForItemArray::StaticExecute(RigVMExecuteContext, Items.Keys, Rotations, RotationEaseType, Weight, WorkData, ExecuteContext, Context);
+}
+
+FRigUnit_DistributeRotationForItemArray_Execute()
 {
     DECLARE_SCOPE_HIERARCHICAL_COUNTER_RIGUNIT()
 	URigHierarchy* Hierarchy = ExecuteContext.Hierarchy;
diff --git a/Engine/Plugins/Animation/ControlRig/Source/ControlRig/Private/Units/Highlevel/Hierarchy/RigUnit_DistributeRotation.h b/Engine/Plugins/Animation/ControlRig/Source/ControlRig/Private/Units/Highlevel/Hierarchy/RigUnit_DistributeRotation.h
--- a/Engine/Plugins/Animation/ControlRig/Source/ControlRig/Private/Units/Highlevel/Hierarchy/RigUnit_DistributeRotation.h
+++ b/Engine/Plugins/Animation/ControlRig/Source/ControlRig/Private/Units/Highlevel/Hierarchy/RigUnit_DistributeRotation.h
@@ -128,7 +128,7 @@ struct CONTROLRIG_API FRigUnit_DistributeRotation : public FRigUnit_HighlevelBas
  * Each rotation is expressed by a quaternion and a ratio, where the ratio is between 0.0 and 1.0
  * Note: This node adds rotation in local space of each item!
  */
-USTRUCT(meta=(DisplayName="Distribute Rotation", Category="Hierarchy", Keywords="TwistBones"))
+USTRUCT(meta=(DisplayName="Distribute Rotation", Category="Hierarchy", Keywords="TwistBones", Deprecated = "5.0"))
 struct CONTROLRIG_API FRigUnit_DistributeRotationForCollection : public FRigUnit_HighlevelBaseMutable
 {
 	GENERATED_BODY()
@@ -169,3 +169,50 @@ struct CONTROLRIG_API FRigUnit_DistributeRotationForCollection : public FRigUnit
 	UPROPERTY(transient)
 	FRigUnit_DistributeRotation_WorkData WorkData;
 };
+
+/**
+ * Distributes rotations provided across a collection of items.
+ * Each rotation is expressed by a quaternion and a ratio, where the ratio is between 0.0 and 1.0
+ * Note: This node adds rotation in local space of each item!
+ */
+USTRUCT(meta=(DisplayName="Distribute Rotation", Category="Hierarchy", Keywords="TwistBones"))
+struct CONTROLRIG_API FRigUnit_DistributeRotationForItemArray : public FRigUnit_HighlevelBaseMutable
+{
+	GENERATED_BODY()
+
+	FRigUnit_DistributeRotationForItemArray()
+	{
+		RotationEaseType = EControlRigAnimEasingType::Linear;
+		Weight = 1.f;
+	}
+
+	RIGVM_METHOD()
+	virtual void Execute(const FRigUnitContext& Context) override;
+
+	/** 
+	 * The items to use to distribute the rotation
+	 */
+	UPROPERTY(meta = (Input))
+	TArray<FRigElementKey> Items;
+
+	/** 
+	 * The list of rotations to be applied
+	 */
+	UPROPERTY(meta = (Input))
+	TArray<FRigUnit_DistributeRotation_Rotation> Rotations;
+
+	/**
+	 * The easing to use between to rotations.
+	 */
+	UPROPERTY(meta = (Input, Constant))
+	EControlRigAnimEasingType RotationEaseType;
+
+	/**
+	 * The weight of the solver - how much the rotation should be applied
+	 */	
+	UPROPERTY(meta = (Input))
+	float Weight;
+
+	UPROPERTY(transient)
+	FRigUnit_DistributeRotation_WorkData WorkData;
+};
diff --git a/Engine/Plugins/Animation/ControlRig/Source/ControlRig/Private/Units/Highlevel/Hierarchy/RigUnit_FABRIK.cpp b/Engine/Plugins/Animation/ControlRig/Source/ControlRig/Private/Units/Highlevel/Hierarchy/RigUnit_FABRIK.cpp
--- a/Engine/Plugins/Animation/ControlRig/Source/ControlRig/Private/Units/Highlevel/Hierarchy/RigUnit_FABRIK.cpp
+++ b/Engine/Plugins/Animation/ControlRig/Source/ControlRig/Private/Units/Highlevel/Hierarchy/RigUnit_FABRIK.cpp
@@ -36,6 +36,11 @@ FRigUnit_FABRIK_Execute()
 }
 
 FRigUnit_FABRIKPerItem_Execute()
+{
+	FRigUnit_FABRIKItemArray::StaticExecute(RigVMExecuteContext, Items.Keys, EffectorTransform, Precision, Weight, bPropagateToChildren, MaxIterations, WorkData, ExecuteContext, Context);
+}
+
+FRigUnit_FABRIKItemArray_Execute()
 {
     DECLARE_SCOPE_HIERARCHICAL_COUNTER_RIGUNIT()
 	URigHierarchy* Hierarchy = ExecuteContext.Hierarchy;
diff --git a/Engine/Plugins/Animation/ControlRig/Source/ControlRig/Private/Units/Highlevel/Hierarchy/RigUnit_FABRIK.h b/Engine/Plugins/Animation/ControlRig/Source/ControlRig/Private/Units/Highlevel/Hierarchy/RigUnit_FABRIK.h
--- a/Engine/Plugins/Animation/ControlRig/Source/ControlRig/Private/Units/Highlevel/Hierarchy/RigUnit_FABRIK.h
+++ b/Engine/Plugins/Animation/ControlRig/Source/ControlRig/Private/Units/Highlevel/Hierarchy/RigUnit_FABRIK.h
@@ -101,7 +101,7 @@ struct CONTROLRIG_API FRigUnit_FABRIK : public FRigUnit_HighlevelBaseMutable
  * the Forward and Backward Reaching Inverse Kinematics algorithm.
  * For now this node supports single effector chains only.
  */
-USTRUCT(meta=(DisplayName="Basic FABRIK", Category="Hierarchy", Keywords="N-Bone,IK"))
+USTRUCT(meta=(DisplayName="Basic FABRIK", Category="Hierarchy", Keywords="N-Bone,IK", Deprecated = "5.0"))
 struct CONTROLRIG_API FRigUnit_FABRIKPerItem : public FRigUnit_HighlevelBaseMutable
 {
 	GENERATED_BODY()
@@ -159,3 +159,67 @@ struct CONTROLRIG_API FRigUnit_FABRIKPerItem : public FRigUnit_HighlevelBaseMuta
 	UPROPERTY(transient)
 	FRigUnit_FABRIK_WorkData WorkData;
 };
+
+/**
+ * The FABRIK solver can solve N-Bone chains using 
+ * the Forward and Backward Reaching Inverse Kinematics algorithm.
+ * For now this node supports single effector chains only.
+ */
+USTRUCT(meta=(DisplayName="Basic FABRIK", Category="Hierarchy", Keywords="N-Bone,IK"))
+struct CONTROLRIG_API FRigUnit_FABRIKItemArray : public FRigUnit_HighlevelBaseMutable
+{
+	GENERATED_BODY()
+
+	RIGVM_METHOD()
+	virtual void Execute(const FRigUnitContext& Context) override;
+
+	FRigUnit_FABRIKItemArray()
+	{
+		Precision = 1.f;
+		Weight = 1.f;
+		bPropagateToChildren = true;
+		MaxIterations = 10;
+		EffectorTransform = FTransform::Identity;
+	}
+
+	/**
+	 * The chain to use
+	 */
+	UPROPERTY(meta = (Input, ExpandByDefault))
+	TArray<FRigElementKey> Items;
+
+	/**
+	 * The transform of the effector in global space
+	 */
+	UPROPERTY(meta = (Input))
+	FTransform EffectorTransform;
+
+	/**
+	 * The precision to use for the fabrik solver
+	 */
+	UPROPERTY(meta = (Input, Constant))
+	float Precision;
+
+	/**
+	 * The weight of the solver - how much the IK should be applied.
+	 */
+	UPROPERTY(meta = (Input))
+	float Weight;
+
+	/**
+	 * If set to true all of the global transforms of the children
+	 * of this bone will be recalculated based on their local transforms.
+	 * Note: This is computationally more expensive than turning it off.
+	 */
+	UPROPERTY(meta = (Input, Constant))
+	bool bPropagateToChildren;
+
+	/**
+	 * The maximum number of iterations. Values between 4 and 16 are common.
+	 */
+	UPROPERTY(meta = (Input))
+	int32 MaxIterations;
+
+	UPROPERTY(transient)
+	FRigUnit_FABRIK_WorkData WorkData;
+};
diff --git a/Engine/Plugins/Animation/ControlRig/Source/ControlRig/Private/Units/Highlevel/Hierarchy/RigUnit_FitChainToCurve.cpp b/Engine/Plugins/Animation/ControlRig/Source/ControlRig/Private/Units/Highlevel/Hierarchy/RigUnit_FitChainToCurve.cpp
--- a/Engine/Plugins/Animation/ControlRig/Source/ControlRig/Private/Units/Highlevel/Hierarchy/RigUnit_FitChainToCurve.cpp
+++ b/Engine/Plugins/Animation/ControlRig/Source/ControlRig/Private/Units/Highlevel/Hierarchy/RigUnit_FitChainToCurve.cpp
@@ -44,6 +44,29 @@ FRigUnit_FitChainToCurve_Execute()
 }
 
 FRigUnit_FitChainToCurvePerItem_Execute()
+{
+	FRigUnit_FitChainToCurveItemArray::StaticExecute(
+		RigVMExecuteContext,
+		Items.Keys,
+		Bezier,
+		Alignment,
+		Minimum,
+		Maximum,
+		SamplingPrecision,
+		PrimaryAxis,
+		SecondaryAxis,
+		PoleVectorPosition,
+		Rotations,
+		RotationEaseType,
+		Weight,
+		bPropagateToChildren,
+		DebugSettings,
+		WorkData,
+		ExecuteContext,
+		Context);
+}
+
+FRigUnit_FitChainToCurveItemArray_Execute()
 {
     DECLARE_SCOPE_HIERARCHICAL_COUNTER_RIGUNIT()
 
diff --git a/Engine/Plugins/Animation/ControlRig/Source/ControlRig/Private/Units/Highlevel/Hierarchy/RigUnit_FitChainToCurve.h b/Engine/Plugins/Animation/ControlRig/Source/ControlRig/Private/Units/Highlevel/Hierarchy/RigUnit_FitChainToCurve.h
--- a/Engine/Plugins/Animation/ControlRig/Source/ControlRig/Private/Units/Highlevel/Hierarchy/RigUnit_FitChainToCurve.h
+++ b/Engine/Plugins/Animation/ControlRig/Source/ControlRig/Private/Units/Highlevel/Hierarchy/RigUnit_FitChainToCurve.h
@@ -251,7 +251,7 @@ struct CONTROLRIG_API FRigUnit_FitChainToCurve : public FRigUnit_HighlevelBaseMu
  * Fits a given chain to a four point bezier curve.
  * Additionally provides rotational control matching the features of the Distribute Rotation node.
  */
-USTRUCT(meta=(DisplayName="Fit Chain on Curve", Category="Hierarchy", Keywords="Fit,Resample,Bezier"))
+USTRUCT(meta=(DisplayName="Fit Chain on Curve", Category="Hierarchy", Keywords="Fit,Resample,Bezier", Deprecated = "5.0"))
 struct CONTROLRIG_API FRigUnit_FitChainToCurvePerItem : public FRigUnit_HighlevelBaseMutable
 {
 	GENERATED_BODY()
@@ -363,3 +363,120 @@ struct CONTROLRIG_API FRigUnit_FitChainToCurvePerItem : public FRigUnit_Highleve
 	UPROPERTY(transient)
 	FRigUnit_FitChainToCurve_WorkData WorkData;
 };
+
+/**
+ * Fits a given chain to a four point bezier curve.
+ * Additionally provides rotational control matching the features of the Distribute Rotation node.
+ */
+USTRUCT(meta=(DisplayName="Fit Chain on Curve", Category="Hierarchy", Keywords="Fit,Resample,Bezier"))
+struct CONTROLRIG_API FRigUnit_FitChainToCurveItemArray : public FRigUnit_HighlevelBaseMutable
+{
+	GENERATED_BODY()
+
+	FRigUnit_FitChainToCurveItemArray()
+	{
+		Bezier = FCRFourPointBezier();
+		Alignment = EControlRigCurveAlignment::Stretched;
+		Minimum = 0.f;
+		Maximum = 1.f;
+		SamplingPrecision = 12;
+		PrimaryAxis = FVector(1.f, 0.f, 0.f);
+		SecondaryAxis = FVector(0.f, 0.f, 0.f);
+		PoleVectorPosition = FVector::ZeroVector;
+		RotationEaseType = EControlRigAnimEasingType::Linear;
+		Weight = 1.f;
+		bPropagateToChildren = true;
+		DebugSettings = FRigUnit_FitChainToCurve_DebugSettings();
+	}
+
+	RIGVM_METHOD()
+	virtual void Execute(const FRigUnitContext& Context) override;
+
+	/** 
+	 * The items to align
+	 */
+	UPROPERTY(meta = (Input))
+	TArray<FRigElementKey> Items;
+
+	/** 
+	 * The curve to align to
+	 */
+	UPROPERTY(meta = (Input))
+	FCRFourPointBezier Bezier;
+
+	/** 
+	 * Specifies how to align the chain on the curve
+	 */
+	UPROPERTY(meta = (Input, Constant))
+	EControlRigCurveAlignment Alignment;
+
+	/** 
+	 * The minimum U value to use on the curve
+	 */
+	UPROPERTY(meta = (Input, Constant))
+	float Minimum;
+
+	/** 
+	 * The maximum U value to use on the curve
+	 */
+	UPROPERTY(meta = (Input, Constant))
+	float Maximum;
+
+	/**
+	 * The number of samples to use on the curve. Clamped at 64.
+	 */
+	UPROPERTY(meta = (Input, Constant))
+	int32 SamplingPrecision;
+
+	/**
+	 * The major axis being aligned - along the bone
+	 */
+	UPROPERTY(meta = (Input))
+	FVector PrimaryAxis;
+
+	/**
+	 * The minor axis being aligned - towards the pole vector.
+	 * You can use (0.0, 0.0, 0.0) to disable it.
+	 */
+	UPROPERTY(meta = (Input))
+	FVector SecondaryAxis;
+
+	/**
+	 * The the position of the pole vector used for aligning the secondary axis.
+	 * Only has an effect if the secondary axis is set.
+	 */
+	UPROPERTY(meta = (Input))
+	FVector PoleVectorPosition;
+
+	/** 
+	 * The list of rotations to be applied along the curve
+	 */
+	UPROPERTY(meta = (Input))
+	TArray<FRigUnit_FitChainToCurve_Rotation> Rotations;
+
+	/**
+	 * The easing to use between to rotations.
+	 */
+	UPROPERTY(meta = (Input, Constant))
+	EControlRigAnimEasingType RotationEaseType;
+
+	/**
+	 * The weight of the solver - how much the rotation should be applied
+	 */
+	UPROPERTY(meta = (Input))
+	float Weight;
+
+	/**
+	 * If set to true all of the global transforms of the children
+	 * of this bone will be recalculated based on their local transforms.
+	 * Note: This is computationally more expensive than turning it off.
+	 */
+	UPROPERTY(meta = (Input, Constant))
+	bool bPropagateToChildren;
+
+	UPROPERTY(meta = (Input, DetailsOnly))
+	FRigUnit_FitChainToCurve_DebugSettings DebugSettings;
+
+	UPROPERTY(transient)
+	FRigUnit_FitChainToCurve_WorkData WorkData;
+};
diff --git a/Engine/Plugins/Animation/ControlRig/Source/ControlRig/Private/Units/Highlevel/Hierarchy/RigUnit_SlideChain.cpp b/Engine/Plugins/Animation/ControlRig/Source/ControlRig/Private/Units/Highlevel/Hierarchy/RigUnit_SlideChain.cpp
--- a/Engine/Plugins/Animation/ControlRig/Source/ControlRig/Private/Units/Highlevel/Hierarchy/RigUnit_SlideChain.cpp
+++ b/Engine/Plugins/Animation/ControlRig/Source/ControlRig/Private/Units/Highlevel/Hierarchy/RigUnit_SlideChain.cpp
@@ -34,6 +34,11 @@ FRigUnit_SlideChain_Execute()
 }
 
 FRigUnit_SlideChainPerItem_Execute()
+{
+	FRigUnit_SlideChainItemArray::StaticExecute(RigVMExecuteContext, Items.Keys, SlideAmount, bPropagateToChildren, WorkData, ExecuteContext, Context);
+}
+
+FRigUnit_SlideChainItemArray_Execute()
 {
 	URigHierarchy* Hierarchy = ExecuteContext.Hierarchy;
 	if (Hierarchy == nullptr)
diff --git a/Engine/Plugins/Animation/ControlRig/Source/ControlRig/Private/Units/Highlevel/Hierarchy/RigUnit_SlideChain.h b/Engine/Plugins/Animation/ControlRig/Source/ControlRig/Private/Units/Highlevel/Hierarchy/RigUnit_SlideChain.h
--- a/Engine/Plugins/Animation/ControlRig/Source/ControlRig/Private/Units/Highlevel/Hierarchy/RigUnit_SlideChain.h
+++ b/Engine/Plugins/Animation/ControlRig/Source/ControlRig/Private/Units/Highlevel/Hierarchy/RigUnit_SlideChain.h
@@ -83,7 +83,7 @@ struct CONTROLRIG_API FRigUnit_SlideChain: public FRigUnit_HighlevelBaseMutable
 /**
  * Slides an existing chain along itself with control over extrapolation.
  */
-USTRUCT(meta=(DisplayName="Slide Chain", Category="Hierarchy", Keywords="Fit,Refit"))
+USTRUCT(meta=(DisplayName="Slide Chain", Category="Hierarchy", Keywords="Fit,Refit", Deprecated = "5.0"))
 struct CONTROLRIG_API FRigUnit_SlideChainPerItem: public FRigUnit_HighlevelBaseMutable
 {
 	GENERATED_BODY()
@@ -120,3 +120,44 @@ struct CONTROLRIG_API FRigUnit_SlideChainPerItem: public FRigUnit_HighlevelBaseM
 	UPROPERTY(transient)
 	FRigUnit_SlideChain_WorkData WorkData;
 };
+
+/**
+ * Slides an existing chain along itself with control over extrapolation.
+ */
+USTRUCT(meta=(DisplayName="Slide Chain", Category="Hierarchy", Keywords="Fit,Refit"))
+struct CONTROLRIG_API FRigUnit_SlideChainItemArray: public FRigUnit_HighlevelBaseMutable
+{
+	GENERATED_BODY()
+
+	FRigUnit_SlideChainItemArray()
+	{
+		SlideAmount = 0.f;
+		bPropagateToChildren = true;
+	}
+
+	RIGVM_METHOD()
+	virtual void Execute(const FRigUnitContext& Context) override;
+
+	/** 
+	 * The items to slide
+	 */
+	UPROPERTY(meta = (Input))
+	TArray<FRigElementKey> Items;
+
+	/** 
+	 * The amount of sliding. This unit is multiple of the chain length.
+	 */
+	UPROPERTY(meta = (Input))
+	float SlideAmount;
+
+	/**
+	 * If set to true all of the global transforms of the children
+	 * of this bone will be recalculated based on their local transforms.
+	 * Note: This is computationally more expensive than turning it off.
+	 */
+	UPROPERTY(meta = (Input, Constant))
+	bool bPropagateToChildren;
+
+	UPROPERTY(transient)
+	FRigUnit_SlideChain_WorkData WorkData;
+};
diff --git a/Engine/Plugins/Animation/ControlRig/Source/ControlRig/Public/Rigs/RigHierarchy.h b/Engine/Plugins/Animation/ControlRig/Source/ControlRig/Public/Rigs/RigHierarchy.h
--- a/Engine/Plugins/Animation/ControlRig/Source/ControlRig/Public/Rigs/RigHierarchy.h
+++ b/Engine/Plugins/Animation/ControlRig/Source/ControlRig/Public/Rigs/RigHierarchy.h
@@ -1972,6 +1972,19 @@ public:
 		const FRigElementKeyCollection& InItems 
 	) const;
 
+	/**
+	 * Returns the current / initial pose of the hierarchy
+	 * @param bInitial If set to true the initial pose will be returned
+	 * @param InElementType The types of elements to get
+	 * @param InItems An optional list of items to get
+	 * @return The pose of the hierarchy
+	 */
+	FRigPose GetPose(
+		bool bInitial,
+		ERigElementType InElementType,
+		const TArrayView<const FRigElementKey>& InItems 
+	) const;
+
 	/**
 	 * Sets the current / initial pose of the hierarchy
 	 * @param InPose The pose to set on the hierarchy
@@ -2001,6 +2014,22 @@ public:
 		float InWeight
 	);
 
+	/**
+	 * Sets the current / initial pose of the hierarchy
+	 * @param InPose The pose to set on the hierarchy
+	 * @param InTransformType The transform type to set
+	 * @param InElementType The types of elements to set
+	 * @param InItems An optional list of items to set
+	 * @param InWeight A weight to define how much the pose needs to be mixed in
+	 */
+	void SetPose(
+		const FRigPose& InPose,
+		ERigTransformType::Type InTransformType,
+		ERigElementType InElementType,
+		const TArrayView<const FRigElementKey>& InItems, 
+		float InWeight
+	);
+
 	/**
 	 * Sets the current / initial pose of the hierarchy
 	 * @param InPose The pose to set on the hierarchy
diff --git a/Engine/Plugins/Animation/ControlRig/Source/ControlRig/Public/Rigs/RigHierarchyDefines.h b/Engine/Plugins/Animation/ControlRig/Source/ControlRig/Public/Rigs/RigHierarchyDefines.h
--- a/Engine/Plugins/Animation/ControlRig/Source/ControlRig/Public/Rigs/RigHierarchyDefines.h
+++ b/Engine/Plugins/Animation/ControlRig/Source/ControlRig/Public/Rigs/RigHierarchyDefines.h
@@ -1298,6 +1298,21 @@ public:
 		return GetTypeHash(Key.Name) * 10 + (uint32)Key.Type;
 	}
 
+	friend FORCEINLINE uint32 GetTypeHash(const TArrayView<const FRigElementKey>& Keys)
+	{
+		uint32 Hash = (uint32)(Keys.Num() * 17 + 3);
+		for (const FRigElementKey& Key : Keys)
+		{
+			Hash += GetTypeHash(Key);
+		}
+		return Hash;
+	}
+
+	friend FORCEINLINE uint32 GetTypeHash(const TArray<FRigElementKey>& Keys)
+	{
+		return GetTypeHash(TArrayView<const FRigElementKey>(Keys.GetData(), Keys.Num()));
+	}
+
 	FORCEINLINE bool operator ==(const FRigElementKey& Other) const
 	{
 		return Name == Other.Name && Type == Other.Type;
@@ -1465,12 +1480,7 @@ struct CONTROLRIG_API FRigElementKeyCollection
 
 	friend FORCEINLINE uint32 GetTypeHash(const FRigElementKeyCollection& Collection)
 	{
-		uint32 Hash = (uint32)(Collection.Num() * 17 + 3);
-		for (const FRigElementKey& Key : Collection)
-		{
-			Hash += GetTypeHash(Key);
-		}
-		return Hash;
+		return GetTypeHash(Collection.GetKeys());
 	}
 
 	// creates a collection containing all of the children of a given 
@@ -1517,9 +1527,7 @@ struct CONTROLRIG_API FRigElementKeyCollection
 	// filters a collection by name
 	FRigElementKeyCollection FilterByName(const FName& InPartialName) const;
 
-protected:
-
-	UPROPERTY(BlueprintReadWrite, EditAnywhere, Category = Collection, meta = (AllowPrivateAccess = "true"))
+	UPROPERTY(BlueprintReadWrite, EditAnywhere, Category = Collection)
 	TArray<FRigElementKey> Keys;
 };
 
diff --git a/Engine/Plugins/Animation/ControlRig/Source/ControlRig/Public/Units/Hierarchy/RigUnit_SetTransform.h b/Engine/Plugins/Animation/ControlRig/Source/ControlRig/Public/Units/Hierarchy/RigUnit_SetTransform.h
--- a/Engine/Plugins/Animation/ControlRig/Source/ControlRig/Public/Units/Hierarchy/RigUnit_SetTransform.h
+++ b/Engine/Plugins/Animation/ControlRig/Source/ControlRig/Public/Units/Hierarchy/RigUnit_SetTransform.h
@@ -269,7 +269,7 @@ struct FRigUnit_SetScale : public FRigUnitMutable
  * actually sets the Control's offset transform and resets the local
  * values to (0, 0, 0).
  */
-USTRUCT(meta=(DisplayName="Set Transform Array", Category="Hierarchy", DocumentationPolicy = "Strict", Keywords="SetBoneTransform,SetControlTransform,SetInitialTransform,SetSpaceTransform", NodeColor="0, 0.364706, 1.0", Varying))
+USTRUCT(meta=(DisplayName="Set Transform Array", Category="Hierarchy", DocumentationPolicy = "Strict", Keywords="SetBoneTransform,SetControlTransform,SetInitialTransform,SetSpaceTransform", NodeColor="0, 0.364706, 1.0", Varying, Deprecated = "5.0"))
 struct CONTROLRIG_API FRigUnit_SetTransformArray : public FRigUnitMutable
 {
 	GENERATED_BODY()
@@ -323,3 +323,65 @@ struct CONTROLRIG_API FRigUnit_SetTransformArray : public FRigUnitMutable
 	UPROPERTY()
 	TArray<FCachedRigElement> CachedIndex;
 };
+
+/**
+ * SetTransformArray is used to set an array of transforms on the hierarchy.
+ * 
+ * Note: For Controls when setting the initial transform this node
+ * actually sets the Control's offset transform and resets the local
+ * values to (0, 0, 0).
+ */
+USTRUCT(meta=(DisplayName="Set Transform Array", Category="Hierarchy", DocumentationPolicy = "Strict", Keywords="SetBoneTransform,SetControlTransform,SetInitialTransform,SetSpaceTransform", NodeColor="0, 0.364706, 1.0", Varying))
+struct CONTROLRIG_API FRigUnit_SetTransformItemArray : public FRigUnitMutable
+{
+	GENERATED_BODY()
+
+	FRigUnit_SetTransformItemArray()
+		: Items()
+		, Space(EBoneGetterSetterMode::GlobalSpace)
+		, bInitial(false)
+		, Transforms()
+		, Weight(1.f)
+		, bPropagateToChildren(true)
+		, CachedIndex()
+	{}
+
+	RIGVM_METHOD()
+	virtual void Execute(const FRigUnitContext& Context) override;
+
+	/**
+	 * The item to set the transform for
+	 */
+	UPROPERTY(meta = (Input))
+	TArray<FRigElementKey> Items;
+
+	/**
+	 * Defines if the transform should be set in local or global space
+	 */ 
+	UPROPERTY(meta = (Input))
+	EBoneGetterSetterMode Space;
+
+	/**
+	 * Defines if the transform should be set as current (false) or initial (true).
+	 * Initial transforms for bones and other elements in the hierarchy represent the reference pose's value.
+	 */ 
+	UPROPERTY(meta = (Input))
+	bool bInitial;
+
+	// The new transform of the given item
+	UPROPERTY(meta=(Input))
+	TArray<FTransform> Transforms;
+
+	// Defines how much the change will be applied
+	UPROPERTY(meta = (Input, UIMin = "0.0", UIMax = "1.0"))
+	float Weight;
+
+	// If set to true children of affected items in the hierarchy
+	// will follow the transform change - otherwise only the parent will move.
+	UPROPERTY(meta=(Input))
+	bool bPropagateToChildren;
+
+	// Used to cache the internally
+	UPROPERTY()
+	TArray<FCachedRigElement> CachedIndex;
+};
diff --git a/Engine/Plugins/Animation/ControlRigSpline/Source/ControlRigSpline/Private/ControlRigSplineUnits.cpp b/Engine/Plugins/Animation/ControlRigSpline/Source/ControlRigSpline/Private/ControlRigSplineUnits.cpp
--- a/Engine/Plugins/Animation/ControlRigSpline/Source/ControlRigSpline/Private/ControlRigSplineUnits.cpp
+++ b/Engine/Plugins/Animation/ControlRigSpline/Source/ControlRigSpline/Private/ControlRigSplineUnits.cpp
@@ -224,6 +224,29 @@ FRigUnit_GetLengthControlRigSpline_Execute()
 }
 
 FRigUnit_FitChainToSplineCurve_Execute()
+{
+	FRigUnit_FitChainToSplineCurveItemArray::StaticExecute(
+		RigVMExecuteContext,
+		Items.Keys,
+		Spline,
+		Alignment,
+		Minimum,
+		Maximum,
+		SamplingPrecision,
+		PrimaryAxis,
+		SecondaryAxis,
+		PoleVectorPosition,
+		Rotations,
+		RotationEaseType,
+		Weight,
+		bPropagateToChildren,
+		DebugSettings,
+		WorkData,
+		ExecuteContext,
+		Context);
+}
+
+FRigUnit_FitChainToSplineCurveItemArray_Execute()
 {
     DECLARE_SCOPE_HIERARCHICAL_COUNTER_RIGUNIT()
 
@@ -631,6 +654,11 @@ FRigUnit_FitChainToSplineCurve_Execute()
 }
 
 FRigUnit_FitSplineCurveToChain_Execute()
+{
+	FRigUnit_FitSplineCurveToChainItemArray::StaticExecute(RigVMExecuteContext, Items.Keys, Spline, ExecuteContext, Context);
+}
+
+FRigUnit_FitSplineCurveToChainItemArray_Execute()
 {
     DECLARE_SCOPE_HIERARCHICAL_COUNTER_RIGUNIT()
 
diff --git a/Engine/Plugins/Animation/ControlRigSpline/Source/ControlRigSpline/Public/ControlRigSplineUnits.h b/Engine/Plugins/Animation/ControlRigSpline/Source/ControlRigSpline/Public/ControlRigSplineUnits.h
--- a/Engine/Plugins/Animation/ControlRigSpline/Source/ControlRigSpline/Public/ControlRigSplineUnits.h
+++ b/Engine/Plugins/Animation/ControlRigSpline/Source/ControlRigSpline/Public/ControlRigSplineUnits.h
@@ -208,7 +208,7 @@ struct CONTROLRIGSPLINE_API FRigUnit_GetLengthControlRigSpline : public FRigUnit
  * Fits a given chain to a spline curve.
  * Additionally provides rotational control matching the features of the Distribute Rotation node.
  */
-USTRUCT(meta=(DisplayName="Fit Chain on Spline Curve", Category="Hierarchy", Keywords="Fit,Resample,Spline"))
+USTRUCT(meta=(DisplayName="Fit Chain on Spline Curve", Category="Hierarchy", Keywords="Fit,Resample,Spline", Deprecated = "5.0"))
 struct CONTROLRIGSPLINE_API FRigUnit_FitChainToSplineCurve : public FRigUnit_HighlevelBaseMutable
 {
 	GENERATED_BODY()
@@ -320,10 +320,126 @@ struct CONTROLRIGSPLINE_API FRigUnit_FitChainToSplineCurve : public FRigUnit_Hig
 	FRigUnit_FitChainToCurve_WorkData WorkData;
 };
 
+/**
+ * Fits a given chain to a spline curve.
+ * Additionally provides rotational control matching the features of the Distribute Rotation node.
+ */
+USTRUCT(meta=(DisplayName="Fit Chain on Spline Curve", Category="Hierarchy", Keywords="Fit,Resample,Spline"))
+struct CONTROLRIGSPLINE_API FRigUnit_FitChainToSplineCurveItemArray : public FRigUnit_HighlevelBaseMutable
+{
+	GENERATED_BODY()
+
+	FRigUnit_FitChainToSplineCurveItemArray()
+	{
+		Alignment = EControlRigCurveAlignment::Stretched;
+		Minimum = 0.f;
+		Maximum = 1.f;
+		SamplingPrecision = 12;
+		PrimaryAxis = FVector(1.f, 0.f, 0.f);
+		SecondaryAxis = FVector(0.f, 0.f, 0.f);
+		PoleVectorPosition = FVector::ZeroVector;
+		RotationEaseType = EControlRigAnimEasingType::Linear;
+		Weight = 1.f;
+		bPropagateToChildren = true;
+		DebugSettings = FRigUnit_FitChainToCurve_DebugSettings();
+	}
+
+	RIGVM_METHOD()
+	virtual void Execute(const FRigUnitContext& Context) override;
+
+	/** 
+	 * The items to align
+	 */
+	UPROPERTY(meta = (Input))
+	TArray<FRigElementKey> Items;
+
+	/** 
+	 * The curve to align to
+	 */
+	UPROPERTY(meta = (Input))
+	FControlRigSpline Spline;
+
+	/** 
+	 * Specifies how to align the chain on the curve
+	 */
+	UPROPERTY(meta = (Input, Constant))
+	EControlRigCurveAlignment Alignment;
+
+	/** 
+	 * The minimum U value to use on the curve
+	 */
+	UPROPERTY(meta = (Input, Constant))
+	float Minimum;
+
+	/** 
+	 * The maximum U value to use on the curve
+	 */
+	UPROPERTY(meta = (Input, Constant))
+	float Maximum;
+
+	/**
+	 * The number of samples to use on the curve. Clamped at 64.
+	 */
+	UPROPERTY(meta = (Input, Constant))
+	int32 SamplingPrecision;
+
+	/**
+	 * The major axis being aligned - along the bone
+	 */
+	UPROPERTY(meta = (Input))
+	FVector PrimaryAxis;
+
+	/**
+	 * The minor axis being aligned - towards the pole vector.
+	 * You can use (0.0, 0.0, 0.0) to disable it.
+	 */
+	UPROPERTY(meta = (Input))
+	FVector SecondaryAxis;
+
+	/**
+	 * The the position of the pole vector used for aligning the secondary axis.
+	 * Only has an effect if the secondary axis is set.
+	 */
+	UPROPERTY(meta = (Input))
+	FVector PoleVectorPosition;
+
+	/** 
+	 * The list of rotations to be applied along the curve
+	 */
+	UPROPERTY(meta = (Input))
+	TArray<FRigUnit_FitChainToCurve_Rotation> Rotations;
+
+	/**
+	 * The easing to use between to rotations.
+	 */
+	UPROPERTY(meta = (Input, Constant))
+	EControlRigAnimEasingType RotationEaseType;
+
+	/**
+	 * The weight of the solver - how much the rotation should be applied
+	 */
+	UPROPERTY(meta = (Input))
+	float Weight;
+
+	/**
+	 * If set to true all of the global transforms of the children
+	 * of this bone will be recalculated based on their local transforms.
+	 * Note: This is computationally more expensive than turning it off.
+	 */
+	UPROPERTY(meta = (Input, Constant))
+	bool bPropagateToChildren;
+
+	UPROPERTY(meta = (Input, DetailsOnly))
+	FRigUnit_FitChainToCurve_DebugSettings DebugSettings;
+
+	UPROPERTY(transient)
+	FRigUnit_FitChainToCurve_WorkData WorkData;
+};
+
 /**
  * Fits a given spline curve to a chain.
  */
-USTRUCT(meta=(DisplayName="Fit Spline Curve on Chain", Category="Hierarchy", Keywords="Fit,Resample,Spline"))
+USTRUCT(meta=(DisplayName="Fit Spline Curve on Chain", Category="Hierarchy", Keywords="Fit,Resample,Spline", Deprecated = "5.0"))
 struct CONTROLRIGSPLINE_API FRigUnit_FitSplineCurveToChain : public FRigUnit_HighlevelBaseMutable
 {
 	GENERATED_BODY()
@@ -349,6 +465,35 @@ struct CONTROLRIGSPLINE_API FRigUnit_FitSplineCurveToChain : public FRigUnit_Hig
 	FControlRigSpline Spline;	
 };
 
+/**
+ * Fits a given spline curve to a chain.
+ */
+USTRUCT(meta=(DisplayName="Fit Spline Curve on Chain", Category="Hierarchy", Keywords="Fit,Resample,Spline"))
+struct CONTROLRIGSPLINE_API FRigUnit_FitSplineCurveToChainItemArray : public FRigUnit_HighlevelBaseMutable
+{
+	GENERATED_BODY()
+
+	FRigUnit_FitSplineCurveToChainItemArray()
+	{
+		
+	}
+
+	RIGVM_METHOD()
+	virtual void Execute(const FRigUnitContext& Context) override;
+
+	/** 
+	 * The items to align to
+	 */
+	UPROPERTY(meta = (Input))
+	TArray<FRigElementKey> Items;
+
+	/** 
+	 * The curve to align
+	 */
+	UPROPERTY(meta = (Input, Output))
+	FControlRigSpline Spline;	
+};
+
 USTRUCT(meta = (DisplayName = "Closest Parameter From Spline", Category = "Control Rig"))
 struct CONTROLRIGSPLINE_API FRigUnit_ClosestParameterFromControlRigSpline : public FRigUnit_ControlRigSplineBase
 {
