Control Rig: Fix display / refresh of Control Values in details panel - for current, minimum and maximum

diff --git a/Engine/Plugins/Animation/ControlRig/Source/ControlRigEditor/Private/ControlRigElementDetails.cpp b/Engine/Plugins/Animation/ControlRig/Source/ControlRigEditor/Private/ControlRigElementDetails.cpp
--- a/Engine/Plugins/Animation/ControlRig/Source/ControlRigEditor/Private/ControlRigElementDetails.cpp
+++ b/Engine/Plugins/Animation/ControlRig/Source/ControlRigEditor/Private/ControlRigElementDetails.cpp
@@ -801,6 +801,11 @@ FText FRigBaseElementDetails::GetName() const
 
 void FRigBaseElementDetails::SetName(const FText& InNewText, ETextCommit::Type InCommitType)
 {
+	if(InCommitType == ETextCommit::OnCleared)
+	{
+		return;
+	}
+	
 	if(ObjectsBeingCustomized.Num() > 1)
 	{
 		return;
@@ -1430,38 +1435,26 @@ void FRigControlElementDetails_SetupStructValueWidget(IDetailGroup& InGroup, IDe
 
 	const FString ValueTypeName = ValueTypeEnum->GetDisplayNameTextByValue((int64)InValueType).ToString();
 	const FText PropertyLabel = FText::FromString(FString::Printf(TEXT("%s Value"), *ValueTypeName));
-	const UStruct* ValueStruct = TBaseStructure<T>::Get();
+	const UScriptStruct* ValueStruct = TBaseStructure<T>::Get();
 
 	const TSharedPtr<FStructOnScope> StructToDisplay = MakeShareable(new FStructOnScope(ValueStruct));
 
 	TWeakObjectPtr<URigHierarchy> HierarchyPtr = InHierarchy;
 	const FRigElementKey Key = InControlElement->GetKey();
 
-	const TAttribute<bool> EnabledAttribute = TAttribute<bool>::Create(TAttribute<bool>::FGetter::CreateLambda([HierarchyPtr, Key, InValueType, StructToDisplay, ValueStruct]()->bool
-    {
-        if(HierarchyPtr.IsValid())
-        {
-            if(FRigControlElement* ControlElement = HierarchyPtr->Find<FRigControlElement>(Key))
-            {
-            	// update the struct with the current control value
-            	uint8* StructMemory = StructToDisplay->GetStructMemory();
-            	const FRigControlValue& ControlValue = HierarchyPtr->GetControlValue(Key, InValueType);
-            	T ExtractedValue = FRigControlElementDetails_ExtractValue<T>(ControlValue);
-            	FMemory::Memcpy(StructToDisplay->GetStructMemory(), &ExtractedValue, sizeof(T));
-
-            	return ControlElement->Settings.IsValueTypeEnabled(InValueType);
-            }
-        }
-        return false;
-    }));
-
-	const TAttribute<EVisibility> VisibilityAttribute = TAttribute<EVisibility>::Create(TAttribute<EVisibility>::FGetter::CreateLambda([HierarchyPtr, Key, InValueType]()->EVisibility
+	const TAttribute<EVisibility> VisibilityAttribute = TAttribute<EVisibility>::Create(TAttribute<EVisibility>::FGetter::CreateLambda([HierarchyPtr, Key, InValueType, StructToDisplay, ValueStruct]()->EVisibility
     {
 		if(HierarchyPtr.IsValid())
 		{
 			if(FRigControlElement* ControlElement = HierarchyPtr->Find<FRigControlElement>(Key))
 			{
-                if(ControlElement->Settings.IsValueTypeEnabled(InValueType))
+				// update the struct with the current control value
+				uint8* StructMemory = StructToDisplay->GetStructMemory();
+				const FRigControlValue& ControlValue = HierarchyPtr->GetControlValue(ControlElement, InValueType);
+				T ExtractedValue = FRigControlElementDetails_ExtractValue<T>(ControlValue);
+				ValueStruct->CopyScriptStruct(StructToDisplay->GetStructMemory(), &ExtractedValue, 1);
+				
+				if(ControlElement->Settings.IsValueTypeEnabled(InValueType))
                 {
                     return EVisibility::Visible;
                 }
@@ -1473,7 +1466,6 @@ void FRigControlElementDetails_SetupStructValueWidget(IDetailGroup& InGroup, IDe
 	IDetailPropertyRow* Row = InStructBuilder.AddExternalStructure(StructToDisplay.ToSharedRef());
 	Row->DisplayName(PropertyLabel);
 	Row->ShouldAutoExpand(true);
-	Row->IsEnabled(EnabledAttribute);
 	Row->Visibility(VisibilityAttribute);
 
 	TSharedPtr<SWidget> NameWidget, ValueWidget;
@@ -1849,6 +1841,7 @@ void FRigControlElementDetails::CustomizeChildren(TSharedRef<class IPropertyHand
 			}
 		}
 		
+		
 		FRigControlElementDetails_SetupValueWidget(*ControlGroup, StructBuilder, ERigControlValueType::Minimum, ControlElements[0], HierarchyBeingCustomized);
 		FRigControlElementDetails_SetupValueWidget(*ControlGroup, StructBuilder, ERigControlValueType::Maximum, ControlElements[0], HierarchyBeingCustomized);
 	}
@@ -2062,4 +2055,4 @@ void FRigControlElementDetails::SetDisplayName(const FText& InNewText, ETextComm
 	}
 }
 
-#undef LOCTEXT_NAMESPACE
+#undef LOCTEXT_NAMESPACE
\ No newline at end of file
