Control Rig: Promote to collapse node from function without invalid local variables

diff --git a/Engine/Plugins/Animation/ControlRig/Source/ControlRigDeveloper/Private/ControlRigBlueprint.cpp b/Engine/Plugins/Animation/ControlRig/Source/ControlRigDeveloper/Private/ControlRigBlueprint.cpp
--- a/Engine/Plugins/Animation/ControlRig/Source/ControlRigDeveloper/Private/ControlRigBlueprint.cpp
+++ b/Engine/Plugins/Animation/ControlRig/Source/ControlRigDeveloper/Private/ControlRigBlueprint.cpp
@@ -344,170 +344,170 @@ void UControlRigBlueprint::PostLoad()
 {
 	Super::PostLoad();
 
-	IncrementVMRecompileBracket();
-
-	TArray<UControlRigBlueprint*> ReferencedBlueprints = GetReferencedControlRigBlueprints();
-
-	// PostLoad all referenced BPs so that their function graphs are fully loaded 
-	// and ready to be inlined into this BP during compilation
-	for (UControlRigBlueprint* BP : ReferencedBlueprints)
 	{
-		if (BP->HasAllFlags(RF_NeedPostLoad))
-		{
-			BP->ConditionalPostLoad();
-		}
-	}
-	
-	// temporarily disable default value validation during load time, serialized values should always be accepted
-	TGuardValue<bool> DisablePinDefaultValueValidation(GetOrCreateController()->bValidatePinDefaults, false);
+		FControlRigBlueprintVMCompileScope CompileScope(this);
+		TArray<UControlRigBlueprint*> ReferencedBlueprints = GetReferencedControlRigBlueprints();
 
-	// correct the offset transforms
-	if (GetLinkerCustomVersion(FControlRigObjectVersion::GUID) < FControlRigObjectVersion::ControlOffsetTransform)
-	{
-		HierarchyContainer_DEPRECATED.ControlHierarchy.PostLoad();
-		if (HierarchyContainer_DEPRECATED.ControlHierarchy.Num() > 0)
+		// PostLoad all referenced BPs so that their function graphs are fully loaded 
+		// and ready to be inlined into this BP during compilation
+		for (UControlRigBlueprint* BP : ReferencedBlueprints)
 		{
-			bDirtyDuringLoad = true;
+			if (BP->HasAllFlags(RF_NeedPostLoad))
+			{
+				BP->ConditionalPostLoad();
+			}
 		}
+		
+		// temporarily disable default value validation during load time, serialized values should always be accepted
+		TGuardValue<bool> DisablePinDefaultValueValidation(GetOrCreateController()->bValidatePinDefaults, false);
 
-		for (FRigControl& Control : HierarchyContainer_DEPRECATED.ControlHierarchy)
+		// correct the offset transforms
+		if (GetLinkerCustomVersion(FControlRigObjectVersion::GUID) < FControlRigObjectVersion::ControlOffsetTransform)
 		{
-			const FTransform PreviousOffsetTransform = Control.GetTransformFromValue(ERigControlValueType::Initial);
-			Control.OffsetTransform = PreviousOffsetTransform;
-			Control.InitialValue = Control.Value;
-
-			if (Control.ControlType == ERigControlType::Transform)
-			{
-				Control.InitialValue = FRigControlValue::Make<FTransform>(FTransform::Identity);
-			}
-			else if (Control.ControlType == ERigControlType::TransformNoScale)
+			HierarchyContainer_DEPRECATED.ControlHierarchy.PostLoad();
+			if (HierarchyContainer_DEPRECATED.ControlHierarchy.Num() > 0)
 			{
-				Control.InitialValue = FRigControlValue::Make<FTransformNoScale>(FTransformNoScale::Identity);
+				bDirtyDuringLoad = true;
 			}
-			else if (Control.ControlType == ERigControlType::EulerTransform)
+
+			for (FRigControl& Control : HierarchyContainer_DEPRECATED.ControlHierarchy)
 			{
-				Control.InitialValue = FRigControlValue::Make<FEulerTransform>(FEulerTransform::Identity);
+				const FTransform PreviousOffsetTransform = Control.GetTransformFromValue(ERigControlValueType::Initial);
+				Control.OffsetTransform = PreviousOffsetTransform;
+				Control.InitialValue = Control.Value;
+
+				if (Control.ControlType == ERigControlType::Transform)
+				{
+					Control.InitialValue = FRigControlValue::Make<FTransform>(FTransform::Identity);
+				}
+				else if (Control.ControlType == ERigControlType::TransformNoScale)
+				{
+					Control.InitialValue = FRigControlValue::Make<FTransformNoScale>(FTransformNoScale::Identity);
+				}
+				else if (Control.ControlType == ERigControlType::EulerTransform)
+				{
+					Control.InitialValue = FRigControlValue::Make<FEulerTransform>(FEulerTransform::Identity);
+				}
 			}
 		}
-	}
 
-	// convert the hierarchy from V1 to V2
-	if (GetLinkerCustomVersion(FControlRigObjectVersion::GUID) < FControlRigObjectVersion::RigHierarchyV2)
-	{
-		Modify();
-		
-		TGuardValue<bool> SuspendNotifGuard(Hierarchy->GetSuspendNotificationsFlag(), true);
-		
-		Hierarchy->Reset();
-		GetHierarchyController()->ImportFromHierarchyContainer(HierarchyContainer_DEPRECATED, false);
-	}
-
-	PropagateHierarchyFromBPToInstances();
-	
-	// remove all non-controlrig-graphs
-	TArray<UEdGraph*> NewUberGraphPages;
-	for (UEdGraph* Graph : UbergraphPages)
-	{
-		UControlRigGraph* RigGraph = Cast<UControlRigGraph>(Graph);
-		if (RigGraph)
+		// convert the hierarchy from V1 to V2
+		if (GetLinkerCustomVersion(FControlRigObjectVersion::GUID) < FControlRigObjectVersion::RigHierarchyV2)
 		{
-			NewUberGraphPages.Add(RigGraph);
+			Modify();
+			
+			TGuardValue<bool> SuspendNotifGuard(Hierarchy->GetSuspendNotificationsFlag(), true);
+			
+			Hierarchy->Reset();
+			GetHierarchyController()->ImportFromHierarchyContainer(HierarchyContainer_DEPRECATED, false);
 		}
-		else
+
+		PropagateHierarchyFromBPToInstances();
+		
+		// remove all non-controlrig-graphs
+		TArray<UEdGraph*> NewUberGraphPages;
+		for (UEdGraph* Graph : UbergraphPages)
 		{
-			Graph->MarkPendingKill();
-			Graph->Rename(nullptr, GetTransientPackage(), REN_ForceNoResetLoaders);
+			UControlRigGraph* RigGraph = Cast<UControlRigGraph>(Graph);
+			if (RigGraph)
+			{
+				NewUberGraphPages.Add(RigGraph);
+			}
+			else
+			{
+				Graph->MarkPendingKill();
+				Graph->Rename(nullptr, GetTransientPackage(), REN_ForceNoResetLoaders);
+			}
 		}
-	}
-	UbergraphPages = NewUberGraphPages;
+		UbergraphPages = NewUberGraphPages;
 
-	InitializeModelIfRequired(false /* recompile vm */);
+		InitializeModelIfRequired(false /* recompile vm */);
 
-	PatchVariableNodesOnLoad();
+		PatchVariableNodesOnLoad();
 
 #if WITH_EDITOR
 
-	TArray<URigVMGraph*> GraphsToDetach;
-	GraphsToDetach.Add(GetModel());
-	GraphsToDetach.Add(GetLocalFunctionLibrary());
+		TArray<URigVMGraph*> GraphsToDetach;
+		GraphsToDetach.Add(GetModel());
+		GraphsToDetach.Add(GetLocalFunctionLibrary());
 
-	if (ensure(IsInGameThread()))
-	{
-		for (URigVMGraph* GraphToDetach : GraphsToDetach)
+		if (ensure(IsInGameThread()))
 		{
-			URigVMController* Controller = GetOrCreateController(GraphToDetach);
-			Controller->DetachLinksFromPinObjects();
-			TArray<URigVMNode*> Nodes = GraphToDetach->GetNodes();
-			for (URigVMNode* Node : Nodes)
+			for (URigVMGraph* GraphToDetach : GraphsToDetach)
 			{
-				Controller->RepopulatePinsOnNode(Node, true, false, true);
+				URigVMController* Controller = GetOrCreateController(GraphToDetach);
+				Controller->DetachLinksFromPinObjects();
+				TArray<URigVMNode*> Nodes = GraphToDetach->GetNodes();
+				for (URigVMNode* Node : Nodes)
+				{
+					Controller->RepopulatePinsOnNode(Node, true, false, true);
+				}
 			}
+			SetupPinRedirectorsForBackwardsCompatibility();
 		}
-		SetupPinRedirectorsForBackwardsCompatibility();
-	}
-
-	for (URigVMGraph* GraphToDetach : GraphsToDetach)
-	{
-		URigVMController* Controller = GetOrCreateController(GraphToDetach);
-		Controller->ReattachLinksToPinObjects(true /* follow redirectors */, nullptr, false, true);
-	}
 
-	// perform backwards compat value upgrades
-	TArray<URigVMGraph*> GraphsToValidate = GetAllModels();
-	for (int32 GraphIndex = 0; GraphIndex < GraphsToValidate.Num(); GraphIndex++)
-	{
-		URigVMGraph* GraphToValidate = GraphsToValidate[GraphIndex];
-		if(GraphToValidate == nullptr)
+		for (URigVMGraph* GraphToDetach : GraphsToDetach)
 		{
-			continue;
+			URigVMController* Controller = GetOrCreateController(GraphToDetach);
+			Controller->ReattachLinksToPinObjects(true /* follow redirectors */, nullptr, false, true);
 		}
 
-		for(URigVMNode* Node : GraphToValidate->GetNodes())
-		{
-			URigVMController* Controller = GetOrCreateController(GraphToValidate);
-			Controller->RemoveUnusedOrphanedPins(Node, false);
-		}
-			
-		for(URigVMNode* Node : GraphToValidate->GetNodes())
+		// perform backwards compat value upgrades
+		TArray<URigVMGraph*> GraphsToValidate = GetAllModels();
+		for (int32 GraphIndex = 0; GraphIndex < GraphsToValidate.Num(); GraphIndex++)
 		{
-			TArray<URigVMPin*> Pins = Node->GetAllPinsRecursively();
-			for(URigVMPin* Pin : Pins)
+			URigVMGraph* GraphToValidate = GraphsToValidate[GraphIndex];
+			if(GraphToValidate == nullptr)
 			{
-				if(Pin->GetCPPTypeObject() == StaticEnum<ERigElementType>())
+				continue;
+			}
+
+			for(URigVMNode* Node : GraphToValidate->GetNodes())
+			{
+				URigVMController* Controller = GetOrCreateController(GraphToValidate);
+				Controller->RemoveUnusedOrphanedPins(Node, false);
+			}
+				
+			for(URigVMNode* Node : GraphToValidate->GetNodes())
+			{
+				TArray<URigVMPin*> Pins = Node->GetAllPinsRecursively();
+				for(URigVMPin* Pin : Pins)
 				{
-					if(Pin->GetDefaultValue() == TEXT("Space"))
+					if(Pin->GetCPPTypeObject() == StaticEnum<ERigElementType>())
 					{
-						if(URigVMController* Controller = GetController(GraphToValidate))
+						if(Pin->GetDefaultValue() == TEXT("Space"))
 						{
-							Controller->SuspendNotifications(true);
-							Controller->SetPinDefaultValue(Pin->GetPinPath(), TEXT("Null"), false, false, false);
-							Controller->SuspendNotifications(false);
+							if(URigVMController* Controller = GetController(GraphToValidate))
+							{
+								Controller->SuspendNotifications(true);
+								Controller->SetPinDefaultValue(Pin->GetPinPath(), TEXT("Null"), false, false, false);
+								Controller->SuspendNotifications(false);
+							}
 						}
 					}
 				}
-			}
 
-			if(URigVMFunctionReferenceNode* FunctionReferenceNode = Cast<URigVMFunctionReferenceNode>(Node))
-			{
-				if(URigVMLibraryNode* DependencyNode = FunctionReferenceNode->GetReferencedNode())
+				if(URigVMFunctionReferenceNode* FunctionReferenceNode = Cast<URigVMFunctionReferenceNode>(Node))
 				{
-					if(UControlRigBlueprint* DependencyBlueprint = DependencyNode->GetTypedOuter<UControlRigBlueprint>())
+					if(URigVMLibraryNode* DependencyNode = FunctionReferenceNode->GetReferencedNode())
 					{
-						if(DependencyBlueprint != this)
+						if(UControlRigBlueprint* DependencyBlueprint = DependencyNode->GetTypedOuter<UControlRigBlueprint>())
 						{
-							DependencyBlueprint->GetLocalFunctionLibrary()->UpdateReferencesForReferenceNode(FunctionReferenceNode);
+							if(DependencyBlueprint != this)
+							{
+								DependencyBlueprint->GetLocalFunctionLibrary()->UpdateReferencesForReferenceNode(FunctionReferenceNode);
+							}
 						}
 					}
 				}
 			}
-		}
-	}	
+		}	
 
-	CompileLog.Messages.Reset();
-	CompileLog.NumErrors = CompileLog.NumWarnings = 0;
+		CompileLog.Messages.Reset();
+		CompileLog.NumErrors = CompileLog.NumWarnings = 0;
+#endif
+	}
 
-	
-	DecrementVMRecompileBracket();
 	RecompileVM();
 	RequestControlRigInit();
 
@@ -516,8 +516,6 @@ void UControlRigBlueprint::PostLoad()
 	FCoreUObjectDelegates::OnObjectModified.AddUObject(this, &UControlRigBlueprint::OnPreVariableChange);
 	OnChanged().AddUObject(this, &UControlRigBlueprint::OnPostVariableChange);
 
-#endif
-
 	if (UPackage* Package = GetOutermost())
 	{
 		Package->SetDirtyFlag(bDirtyDuringLoad);
@@ -604,7 +602,6 @@ void UControlRigBlueprint::RecompileVM()
 		}
 
 		bVMRecompilationRequired = false;
-		VMRecompilationBracket = 0;
 		VMCompiledEvent.Broadcast(this, CDO->VM);
 
 #if WITH_EDITOR
@@ -1290,6 +1287,29 @@ URigVMController* UControlRigBlueprint::GetOrCreateController(URigVMGraph* InGra
 		}
 		return FRigVMController_BulkEditResult();
 	});
+
+	Controller->RequestNewExternalVariableDelegate.BindLambda([WeakThis](FRigVMGraphVariableDescription InVariable, bool bInIsPublic, bool bInIsReadOnly) -> FName
+	{
+		if (WeakThis.IsValid())
+		{
+			for (FBPVariableDescription& ExistingVariable : WeakThis->NewVariables)
+			{
+				if (ExistingVariable.VarName == InVariable.Name)
+				{
+					return FName();
+				}
+			}
+
+			FRigVMExternalVariable ExternalVariable = InVariable.ToExternalVariable();
+			return WeakThis->AddMemberVariable(InVariable.Name,
+				ExternalVariable.TypeObject ? ExternalVariable.TypeObject->GetPathName() : ExternalVariable.TypeName.ToString(),
+				bInIsPublic,
+				bInIsReadOnly,
+				InVariable.DefaultValue);
+		}
+		
+		return FName();
+	});
 	
 #endif
 
diff --git a/Engine/Plugins/Animation/ControlRig/Source/ControlRigDeveloper/Public/ControlRigBlueprint.h b/Engine/Plugins/Animation/ControlRig/Source/ControlRigDeveloper/Public/ControlRigBlueprint.h
--- a/Engine/Plugins/Animation/ControlRig/Source/ControlRigDeveloper/Public/ControlRigBlueprint.h
+++ b/Engine/Plugins/Animation/ControlRig/Source/ControlRigDeveloper/Public/ControlRigBlueprint.h
@@ -688,3 +688,24 @@ private:
 	friend class UControlRigComponent;
 	friend struct FControlRigGraphSchemaAction_PromoteToVariable;
 };
+
+class CONTROLRIGDEVELOPER_API FControlRigBlueprintVMCompileScope
+{
+public:
+   
+	FControlRigBlueprintVMCompileScope(UControlRigBlueprint *InBlueprint)
+	: Blueprint(InBlueprint)
+	{
+		check(Blueprint);
+		Blueprint->IncrementVMRecompileBracket();
+	}
+
+	~FControlRigBlueprintVMCompileScope()
+	{
+		Blueprint->DecrementVMRecompileBracket();
+	}
+   
+private:
+
+	UControlRigBlueprint* Blueprint;
+};
\ No newline at end of file
diff --git a/Engine/Plugins/Animation/ControlRig/Source/ControlRigEditor/Private/ControlRigLocalVariableDetails.cpp b/Engine/Plugins/Animation/ControlRig/Source/ControlRigEditor/Private/ControlRigLocalVariableDetails.cpp
--- a/Engine/Plugins/Animation/ControlRig/Source/ControlRigEditor/Private/ControlRigLocalVariableDetails.cpp
+++ b/Engine/Plugins/Animation/ControlRig/Source/ControlRigEditor/Private/ControlRigLocalVariableDetails.cpp
@@ -10,6 +10,7 @@
 #include "Kismet2/BlueprintEditorUtils.h"
 #include "Widgets/Input/STextComboBox.h"
 #include "ControlRigBlueprintGeneratedClass.h"
+#include "ControlRigBlueprint.h"
 #include "RigVMCore/RigVM.h"
 
 #define LOCTEXT_NAMESPACE "LocalVariableDetails"
@@ -130,10 +131,9 @@ FEdGraphPinType FRigVMLocalVariableDetails::OnGetPinInfo() const
 void FRigVMLocalVariableDetails::HandlePinInfoChanged(const FEdGraphPinType& PinType)
 {
 	VariableDescription.ChangeType(PinType);
-	BlueprintBeingCustomized->IncrementVMRecompileBracket();
+	FControlRigBlueprintVMCompileScope CompileScope(BlueprintBeingCustomized);
 	TypeHandle->SetValue(VariableDescription.CPPType);
 	TypeObjectHandle->SetValue(VariableDescription.CPPTypeObject);	
-	BlueprintBeingCustomized->DecrementVMRecompileBracket();
 }
 
 ECheckBoxState FRigVMLocalVariableDetails::HandleBoolDefaultValueIsChecked() const
diff --git a/Engine/Plugins/Animation/ControlRig/Source/ControlRigEditor/Private/Editor/ControlRigEditor.cpp b/Engine/Plugins/Animation/ControlRig/Source/ControlRigEditor/Private/Editor/ControlRigEditor.cpp
--- a/Engine/Plugins/Animation/ControlRig/Source/ControlRigEditor/Private/Editor/ControlRigEditor.cpp
+++ b/Engine/Plugins/Animation/ControlRig/Source/ControlRigEditor/Private/Editor/ControlRigEditor.cpp
@@ -4842,40 +4842,42 @@ void FControlRigEditor::OnHierarchyModified(ERigHierarchyNotification InNotif, U
 					continue;
 				}
 
-				RigBlueprint->IncrementVMRecompileBracket();
-				for (UEdGraphNode* Node : RigGraph->Nodes)
 				{
-					if (UControlRigGraphNode* RigNode = Cast<UControlRigGraphNode>(Node))
+					FControlRigBlueprintVMCompileScope CompileScope(RigBlueprint);
+					for (UEdGraphNode* Node : RigGraph->Nodes)
 					{
-						if (URigVMNode* ModelNode = RigNode->GetModelNode())
+						if (UControlRigGraphNode* RigNode = Cast<UControlRigGraphNode>(Node))
 						{
-							TArray<URigVMPin*> ModelPins = ModelNode->GetAllPinsRecursively();
-							for (URigVMPin * ModelPin : ModelPins)
+							if (URigVMNode* ModelNode = RigNode->GetModelNode())
 							{
-								if ((ModelPin->GetCPPType() == TEXT("FName") && ModelPin->GetCustomWidgetName() == TEXT("BoneName") && ElementType == ERigElementType::Bone) ||
-                                    (ModelPin->GetCPPType() == TEXT("FName") && ModelPin->GetCustomWidgetName() == TEXT("ControlName") && ElementType == ERigElementType::Control) ||
-                                    (ModelPin->GetCPPType() == TEXT("FName") && ModelPin->GetCustomWidgetName() == TEXT("SpaceName") && ElementType == ERigElementType::Null) ||
-                                    (ModelPin->GetCPPType() == TEXT("FName") && ModelPin->GetCustomWidgetName() == TEXT("CurveName") && ElementType == ERigElementType::Curve))
+								TArray<URigVMPin*> ModelPins = ModelNode->GetAllPinsRecursively();
+								for (URigVMPin * ModelPin : ModelPins)
 								{
-									if (ModelPin->GetDefaultValue() == OldNameStr)
+									if ((ModelPin->GetCPPType() == TEXT("FName") && ModelPin->GetCustomWidgetName() == TEXT("BoneName") && ElementType == ERigElementType::Bone) ||
+										(ModelPin->GetCPPType() == TEXT("FName") && ModelPin->GetCustomWidgetName() == TEXT("ControlName") && ElementType == ERigElementType::Control) ||
+										(ModelPin->GetCPPType() == TEXT("FName") && ModelPin->GetCustomWidgetName() == TEXT("SpaceName") && ElementType == ERigElementType::Null) ||
+										(ModelPin->GetCPPType() == TEXT("FName") && ModelPin->GetCustomWidgetName() == TEXT("CurveName") && ElementType == ERigElementType::Curve))
 									{
-										Controller->SetPinDefaultValue(ModelPin->GetPinPath(), NewNameStr, false);
+										if (ModelPin->GetDefaultValue() == OldNameStr)
+										{
+											Controller->SetPinDefaultValue(ModelPin->GetPinPath(), NewNameStr, false);
+										}
 									}
-								}
-								else if (ModelPin->GetCPPTypeObject() == FRigElementKey::StaticStruct())
-								{
-									if (URigVMPin* TypePin = ModelPin->FindSubPin(TEXT("Type")))
+									else if (ModelPin->GetCPPTypeObject() == FRigElementKey::StaticStruct())
 									{
-										const FString TypeStr = TypePin->GetDefaultValue();
-										const int64 TypeValue = RigElementTypeEnum->GetValueByNameString(TypeStr);
-										if (TypeValue == (int64)ElementType)
+										if (URigVMPin* TypePin = ModelPin->FindSubPin(TEXT("Type")))
 										{
-											if (URigVMPin* NamePin = ModelPin->FindSubPin(TEXT("Name")))
+											const FString TypeStr = TypePin->GetDefaultValue();
+											const int64 TypeValue = RigElementTypeEnum->GetValueByNameString(TypeStr);
+											if (TypeValue == (int64)ElementType)
 											{
-												FString NameStr = NamePin->GetDefaultValue();
-												if (NameStr == OldNameStr)
+												if (URigVMPin* NamePin = ModelPin->FindSubPin(TEXT("Name")))
 												{
-													Controller->SetPinDefaultValue(NamePin->GetPinPath(), NewNameStr);
+													FString NameStr = NamePin->GetDefaultValue();
+													if (NameStr == OldNameStr)
+													{
+														Controller->SetPinDefaultValue(NamePin->GetPinPath(), NewNameStr);
+													}
 												}
 											}
 										}
@@ -4885,7 +4887,6 @@ void FControlRigEditor::OnHierarchyModified(ERigHierarchyNotification InNotif, U
 						}
 					}
 				}
-				RigBlueprint->DecrementVMRecompileBracket();
 			}
 				
 			OnHierarchyChanged();
diff --git a/Engine/Source/Developer/RigVMDeveloper/Private/RigVMModel/RigVMController.cpp b/Engine/Source/Developer/RigVMDeveloper/Private/RigVMModel/RigVMController.cpp
--- a/Engine/Source/Developer/RigVMDeveloper/Private/RigVMModel/RigVMController.cpp
+++ b/Engine/Source/Developer/RigVMDeveloper/Private/RigVMModel/RigVMController.cpp
@@ -30,6 +30,30 @@
 
 TMap<URigVMController::FControlRigStructPinRedirectorKey, FString> URigVMController::PinPathCoreRedirectors;
 
+FRigVMControllerCompileBracketScope::FRigVMControllerCompileBracketScope(URigVMController* InController)
+: Graph(nullptr), bSuspendNotifications(InController->bSuspendNotifications)
+{
+	check(InController);
+	Graph = InController->GetGraph();
+	check(Graph);
+	
+	if (bSuspendNotifications)
+	{
+		return;
+	}
+	Graph->Notify(ERigVMGraphNotifType::InteractionBracketOpened, nullptr);
+}
+
+FRigVMControllerCompileBracketScope::~FRigVMControllerCompileBracketScope()
+{
+	check(Graph);
+	if (bSuspendNotifications)
+	{
+		return;
+	}
+	Graph->Notify(ERigVMGraphNotifType::InteractionBracketClosed, nullptr);
+}
+
 URigVMController::URigVMController()
 	: bValidatePinDefaults(true)
 	, bSuspendNotifications(false)
@@ -3658,6 +3682,67 @@ TArray<URigVMNode*> URigVMController::ExpandLibraryNode(URigVMLibraryNode* InNod
 		return TArray<URigVMNode*>();
 	}
 
+	// Find local variables that need to be added as member variables. If member variables of same name and type already
+	// exist, they will be reused. If a local variable is not used, it will not be created.
+	if (URigVMFunctionReferenceNode* FunctionReferenceNode = Cast<URigVMFunctionReferenceNode>(InNode))
+	{
+		TArray<FRigVMGraphVariableDescription> LocalVariables = FunctionReferenceNode->GetContainedGraph()->LocalVariables;
+		TArray<FRigVMExternalVariable> CurrentVariables = GetAllVariables();
+		TArray<FRigVMGraphVariableDescription> VariablesToAdd;
+		for (const URigVMNode* Node : FunctionReferenceNode->GetContainedGraph()->GetNodes())
+		{
+			if (const URigVMVariableNode* VariableNode = Cast<URigVMVariableNode>(Node))
+			{
+				for (FRigVMGraphVariableDescription& LocalVariable : LocalVariables)
+				{
+					if (LocalVariable.Name == VariableNode->GetVariableName())
+					{
+						bool bVariableExists = false;
+						bool bVariableIncompatible = false;
+						FRigVMExternalVariable LocalVariableExternalType = LocalVariable.ToExternalVariable();
+						for (FRigVMExternalVariable& CurrentVariable : CurrentVariables)
+						{						
+							if (CurrentVariable.Name == LocalVariable.Name)
+							{
+								if (CurrentVariable.TypeName != LocalVariableExternalType.TypeName ||
+									CurrentVariable.TypeObject != LocalVariableExternalType.TypeObject ||
+									CurrentVariable.bIsArray != LocalVariableExternalType.bIsArray)
+								{
+									bVariableIncompatible = true;	
+								}
+								bVariableExists = true;
+								break;
+							}
+						}
+
+						if (!bVariableExists)
+						{
+							VariablesToAdd.Add(LocalVariable);	
+						}
+						else if(bVariableIncompatible)
+						{
+							ReportErrorf(TEXT("Found variable %s of incompatible type with a local variable inside function %s"), *LocalVariable.Name.ToString(), *FunctionReferenceNode->GetReferencedNode()->GetName());
+							if (bSetupUndoRedo)
+							{
+								ActionStack->CancelAction(ExpandAction);
+							}
+							return TArray<URigVMNode*>();
+						}
+						break;
+					}
+				}
+			}
+		}
+
+		if (RequestNewExternalVariableDelegate.IsBound())
+		{
+			for (const FRigVMGraphVariableDescription& OldVariable : VariablesToAdd)
+			{
+				RequestNewExternalVariableDelegate.Execute(OldVariable, false, false);
+			}
+		}
+	}
+
 	FVector2D Diagonal = Bounds.Max - Bounds.Min;
 	FVector2D Center = (Bounds.Min + Bounds.Max) * 0.5f;
 
@@ -3748,8 +3833,16 @@ TArray<URigVMNode*> URigVMController::ExpandLibraryNode(URigVMLibraryNode* InNod
 			TArray<FString>& LinkedPins = FromEntryNode.FindOrAdd(PinPath);
 
 			URigVMPin::SplitPinPathAtStart(Link->GetTargetPin()->GetPinPath(), NodeName, PinPath);
-			NodeName = NodeNameMap.FindChecked(*NodeName).ToString();
-			LinkedPins.Add(URigVMPin::JoinPinPath(NodeName, PinPath));
+			
+			if (NodeNameMap.Contains(*NodeName))
+			{
+				NodeName = NodeNameMap.FindChecked(*NodeName).ToString();
+				LinkedPins.Add(URigVMPin::JoinPinPath(NodeName, PinPath));	
+			}
+			else if (NodeName == TEXT("Return"))
+			{
+				LinkedPins.Add(URigVMPin::JoinPinPath(NodeName, PinPath));
+			}
 		}
 	}
 
@@ -3776,8 +3869,16 @@ TArray<URigVMNode*> URigVMController::ExpandLibraryNode(URigVMLibraryNode* InNod
 			TArray<FString>& LinkedPins = ToReturnNode.FindOrAdd(PinPath);
 
 			URigVMPin::SplitPinPathAtStart(Link->GetSourcePin()->GetPinPath(), NodeName, PinPath);
-			NodeName = NodeNameMap.FindChecked(*NodeName).ToString();
-			LinkedPins.Add(URigVMPin::JoinPinPath(NodeName, PinPath));
+			
+			if (NodeNameMap.Contains(*NodeName))
+			{
+				NodeName = NodeNameMap.FindChecked(*NodeName).ToString();
+				LinkedPins.Add(URigVMPin::JoinPinPath(NodeName, PinPath));	
+			}
+			else if (NodeName == TEXT("Entry"))
+			{
+				LinkedPins.Add(URigVMPin::JoinPinPath(NodeName, PinPath));
+			}
 		}
 	}
 
@@ -4312,6 +4413,52 @@ URigVMCollapseNode* URigVMController::PromoteFunctionReferenceNodeToCollapseNode
 		return nullptr;
 	}
 
+	// Find local variables that need to be added as member variables. If member variables of same name and type already
+	// exist, they will be reused. If a local variable is not used, it will not be created.
+	TArray<FRigVMGraphVariableDescription> LocalVariables = FunctionDefinition->GetContainedGraph()->LocalVariables;
+	TArray<FRigVMExternalVariable> CurrentVariables = GetAllVariables();
+	TArray<FRigVMGraphVariableDescription> VariablesToAdd;
+	for (const URigVMNode* Node : FunctionDefinition->GetContainedGraph()->GetNodes())
+	{
+		if (const URigVMVariableNode* VariableNode = Cast<URigVMVariableNode>(Node))
+		{
+			for (FRigVMGraphVariableDescription& LocalVariable : LocalVariables)
+			{
+				if (LocalVariable.Name == VariableNode->GetVariableName())
+				{
+					bool bVariableExists = false;
+					bool bVariableIncompatible = false;
+					FRigVMExternalVariable LocalVariableExternalType = LocalVariable.ToExternalVariable();
+					for (FRigVMExternalVariable& CurrentVariable : CurrentVariables)
+					{						
+						if (CurrentVariable.Name == LocalVariable.Name)
+						{
+							if (CurrentVariable.TypeName != LocalVariableExternalType.TypeName ||
+								CurrentVariable.TypeObject != LocalVariableExternalType.TypeObject ||
+								CurrentVariable.bIsArray != LocalVariableExternalType.bIsArray)
+							{
+								bVariableIncompatible = true;	
+							}
+							bVariableExists = true;
+							break;
+						}
+					}
+
+					if (!bVariableExists)
+					{
+						VariablesToAdd.Add(LocalVariable);	
+					}
+					else if(bVariableIncompatible)
+					{
+						ReportErrorf(TEXT("Found variable %s of incompatible type with a local variable inside function %s"), *LocalVariable.Name.ToString(), *FunctionDefinition->GetName());
+						return nullptr;
+					}
+					break;
+				}
+			}
+		}
+	}
+
 	if (bSetupUndoRedo)
 	{
 		OpenUndoBracket(TEXT("Promote to Collapse Node"));
@@ -4330,9 +4477,35 @@ URigVMCollapseNode* URigVMController::PromoteFunctionReferenceNodeToCollapseNode
 
 	RemoveNode(InFunctionRefNode, bSetupUndoRedo);
 
+	if (RequestNewExternalVariableDelegate.IsBound())
+	{
+		for (const FRigVMGraphVariableDescription& OldVariable : VariablesToAdd)
+		{
+			RequestNewExternalVariableDelegate.Execute(OldVariable, false, false);
+		}
+	}
+
 	URigVMCollapseNode* CollapseNode = DuplicateObject<URigVMCollapseNode>(FunctionDefinition, Graph, *NodeName);
 	if(CollapseNode)
 	{
+		{
+			FRigVMControllerGraphGuard Guard(this, CollapseNode->GetContainedGraph(), bSetupUndoRedo);
+			ReattachLinksToPinObjects();
+
+			for (URigVMNode* Node : CollapseNode->GetContainedGraph()->GetNodes())
+			{
+				if (URigVMVariableNode* VariableNode = Cast<URigVMVariableNode>(Node))
+				{
+					TArray<URigVMLink*> VariableLinks = VariableNode->GetLinks();
+					DetachLinksFromPinObjects(&VariableLinks);
+					RepopulatePinsOnNode(VariableNode);
+					ReattachLinksToPinObjects(false, &VariableLinks);
+				}
+			}
+
+			CollapseNode->GetContainedGraph()->LocalVariables.Empty();
+		}		
+				
 		CollapseNode->NodeColor = FLinearColor::White;
 		CollapseNode->Position = NodePosition;
 		Graph->Nodes.Add(CollapseNode);
@@ -7581,7 +7754,7 @@ bool URigVMController::SetExposedPinIndex(const FName& InPinName, int32 InNewInd
 		return false;
 	}
 
-	Notify(ERigVMGraphNotifType::InteractionBracketOpened, nullptr);
+	FRigVMControllerCompileBracketScope CompileBracketScope(this);
 
 	FRigVMSetPinIndexAction PinIndexAction(Pin, InNewIndex);
 	{
@@ -7601,8 +7774,6 @@ bool URigVMController::SetExposedPinIndex(const FName& InPinName, int32 InNewInd
 		ActionStack->AddAction(PinIndexAction);
 	}
 
-	Notify(ERigVMGraphNotifType::InteractionBracketClosed, nullptr);
-
 	if (bPrintPythonCommand)
 	{
 		const FString GraphName = GetSanitizedGraphName(GetGraph()->GetGraphName());
@@ -10595,7 +10766,7 @@ void URigVMController::RepopulatePinsOnNode(URigVMNode* InNode, bool bFollowCore
 		return;
 	}
 
-	Notify(ERigVMGraphNotifType::InteractionBracketOpened, nullptr);
+	FRigVMControllerCompileBracketScope CompileBracketScope(this);
 
 	URigVMUnitNode* UnitNode = Cast<URigVMUnitNode>(InNode);
 	URigVMRerouteNode* RerouteNode = Cast<URigVMRerouteNode>(InNode);
@@ -10822,8 +10993,6 @@ void URigVMController::RepopulatePinsOnNode(URigVMNode* InNode, bool bFollowCore
 		InjectionInfo->InputPin = InNode->FindPin(InjectionInputPinName.ToString());
 		InjectionInfo->OutputPin = InNode->FindPin(InjectionOutputPinName.ToString());
 	}
-
-	Notify(ERigVMGraphNotifType::InteractionBracketClosed, nullptr);
 }
 
 void URigVMController::RemovePinsDuringRepopulate(URigVMNode* InNode, TArray<URigVMPin*>& InPins, bool bNotify, bool bSetupOrphanedPins)
@@ -11095,7 +11264,7 @@ void URigVMController::ApplyPinState(URigVMPin* InPin, const FPinState& InPinSta
 
 void URigVMController::ApplyPinStates(URigVMNode* InNode, const TMap<FString, URigVMController::FPinState>& InPinStates, const TMap<FString, FString>& InRedirectedPinPaths)
 {
-	Notify(ERigVMGraphNotifType::InteractionBracketOpened, nullptr);
+	FRigVMControllerCompileBracketScope CompileBracketScope(this);
 	for (const TPair<FString, FPinState>& PinStatePair : InPinStates)
 	{
 		FString PinPath = PinStatePair.Key;
@@ -11119,7 +11288,6 @@ void URigVMController::ApplyPinStates(URigVMNode* InNode, const TMap<FString, UR
 			}
 		}
 	}
-	Notify(ERigVMGraphNotifType::InteractionBracketClosed, nullptr);
 }
 
 void URigVMController::ReportWarning(const FString& InMessage)
diff --git a/Engine/Source/Developer/RigVMDeveloper/Private/RigVMModel/RigVMControllerActions.cpp b/Engine/Source/Developer/RigVMDeveloper/Private/RigVMModel/RigVMControllerActions.cpp
--- a/Engine/Source/Developer/RigVMDeveloper/Private/RigVMModel/RigVMControllerActions.cpp
+++ b/Engine/Source/Developer/RigVMDeveloper/Private/RigVMModel/RigVMControllerActions.cpp
@@ -138,6 +138,7 @@ void URigVMActionStack::PostTransacted(const FTransactionObjectEvent& Transactio
 			}
 			if (!Undo(Controller))
 			{
+				ModifiedEvent.Broadcast(ERigVMGraphNotifType::InteractionBracketCanceled, nullptr, nullptr);
 				return;
 			}
 		}
@@ -149,6 +150,7 @@ void URigVMActionStack::PostTransacted(const FTransactionObjectEvent& Transactio
 			}
 			if (!Redo(Controller))
 			{
+				ModifiedEvent.Broadcast(ERigVMGraphNotifType::InteractionBracketCanceled, nullptr, nullptr);
 				return;
 			}
 		}
diff --git a/Engine/Source/Developer/RigVMDeveloper/Public/RigVMModel/RigVMController.h b/Engine/Source/Developer/RigVMDeveloper/Public/RigVMModel/RigVMController.h
--- a/Engine/Source/Developer/RigVMDeveloper/Public/RigVMModel/RigVMController.h
+++ b/Engine/Source/Developer/RigVMDeveloper/Public/RigVMModel/RigVMController.h
@@ -57,11 +57,26 @@ struct FRigVMController_BulkEditResult
 	{}
 };
 
+class RIGVMDEVELOPER_API FRigVMControllerCompileBracketScope
+{
+public:
+   
+	FRigVMControllerCompileBracketScope(URigVMController *InController);
+
+	~FRigVMControllerCompileBracketScope();
+
+private:
+
+	URigVMGraph* Graph;
+	bool bSuspendNotifications;
+};
+
 DECLARE_DELEGATE_RetVal_OneParam(bool, FRigVMController_ShouldStructUnfoldDelegate, const UStruct*)
 DECLARE_DELEGATE_RetVal_OneParam(TArray<FRigVMExternalVariable>, FRigVMController_GetExternalVariablesDelegate, URigVMGraph*)
 DECLARE_DELEGATE_RetVal(const FRigVMByteCode*, FRigVMController_GetByteCodeDelegate)
 DECLARE_DELEGATE_RetVal_OneParam(bool, FRigVMController_IsFunctionAvailableDelegate, URigVMLibraryNode*)
 DECLARE_DELEGATE_RetVal_OneParam(bool, FRigVMController_RequestLocalizeFunctionDelegate, URigVMLibraryNode*)
+DECLARE_DELEGATE_RetVal_ThreeParams(FName, FRigVMController_RequestNewExternalVariableDelegate, FRigVMGraphVariableDescription, bool, bool);
 DECLARE_DELEGATE_RetVal_TwoParams(bool, FRigVMController_IsDependencyCyclicDelegate, UObject*, UObject*)
 DECLARE_DELEGATE_RetVal_TwoParams(FRigVMController_BulkEditResult, FRigVMController_RequestBulkEditDialogDelegate, URigVMLibraryNode*, ERigVMControllerBulkEditType)
 DECLARE_DELEGATE_FiveParams(FRigVMController_OnBulkEditProgressDelegate, TSoftObjectPtr<URigVMFunctionReferenceNode>, ERigVMControllerBulkEditType, ERigVMControllerBulkEditProgress, int32, int32)
@@ -680,6 +695,9 @@ public:
 
 	// A delegate to localize a function on demand
 	FRigVMController_RequestLocalizeFunctionDelegate RequestLocalizeFunctionDelegate;
+
+	// A delegate to create a new blueprint member variable
+	FRigVMController_RequestNewExternalVariableDelegate RequestNewExternalVariableDelegate;
 	
 	// A delegate to validate if we are allowed to introduce a dependency between two objects
 	FRigVMController_IsDependencyCyclicDelegate IsDependencyCyclicDelegate;
@@ -928,6 +946,7 @@ private:
 	friend struct FRigVMAddRerouteNodeAction;
 	friend struct FRigVMChangePinTypeAction;
 	friend class FRigVMParserAST;
+	friend class FRigVMControllerCompileBracketScope;
 };
 
 class FRigVMControllerGraphGuard
diff --git a/Engine/Source/Developer/RigVMDeveloper/Public/RigVMModel/RigVMControllerActions.h b/Engine/Source/Developer/RigVMDeveloper/Public/RigVMModel/RigVMControllerActions.h
--- a/Engine/Source/Developer/RigVMDeveloper/Public/RigVMModel/RigVMControllerActions.h
+++ b/Engine/Source/Developer/RigVMDeveloper/Public/RigVMModel/RigVMControllerActions.h
@@ -135,6 +135,11 @@ class URigVMActionStack : public UObject
 {
 	GENERATED_BODY()
 
+	virtual ~URigVMActionStack()
+	{
+		ensure(CurrentActions.IsEmpty());
+	}
+
 public:
 
 	// Begins an action and opens a bracket / scope.
diff --git a/Engine/Source/Developer/RigVMDeveloper/Public/RigVMModel/RigVMGraph.h b/Engine/Source/Developer/RigVMDeveloper/Public/RigVMModel/RigVMGraph.h
--- a/Engine/Source/Developer/RigVMDeveloper/Public/RigVMModel/RigVMGraph.h
+++ b/Engine/Source/Developer/RigVMDeveloper/Public/RigVMModel/RigVMGraph.h
@@ -165,6 +165,7 @@ private:
 	bool IsNameAvailable(const FString& InName);
 
 	friend class URigVMController;
+	friend class FRigVMControllerCompileBracketScope;
 	friend class URigVMCompiler;
 	friend struct FRigVMControllerObjectFactory;
 };
