Sequencer: Control Rig: Space Switching in Sequencer

Created New Space Channel Type (todo make it more modular for constraints).
Added to Control Rig Sections.
Added their evaluation to the template.
Three places where we can add a space, in Sequencer, in pop up dialog (in ControlRigEditor), and in the tool window (in SControlRigEditModeTools).

When space added we do automatic compensation to make sure they match up.

ControlRigChannelEditors contain the template specializations for the channel editor overrides and helper class to perform some common actions.

RigSpaceBakeWidget now supports baking and key reduction.
Curve Editor Views still WIP but started preflight with them added.

diff --git a/Engine/Plugins/Animation/ControlRig/Source/ControlRig/Private/Sequencer/MovieSceneControlRigParameterSection.cpp b/Engine/Plugins/Animation/ControlRig/Source/ControlRig/Private/Sequencer/MovieSceneControlRigParameterSection.cpp
--- a/Engine/Plugins/Animation/ControlRig/Source/ControlRig/Private/Sequencer/MovieSceneControlRigParameterSection.cpp
+++ b/Engine/Plugins/Animation/ControlRig/Source/ControlRig/Private/Sequencer/MovieSceneControlRigParameterSection.cpp
@@ -891,6 +891,42 @@ bool UMovieSceneControlRigParameterSection::HasTransformParameter(FName InParame
 	return false;
 }
 
+bool UMovieSceneControlRigParameterSection::HasSpaceChannel(FName InParameterName) const
+{
+	for (const FSpaceControlNameAndChannel& SpaceChannel : SpaceChannels)
+	{
+		if (SpaceChannel.ControlName == InParameterName)
+		{
+			return true;
+		}
+	}
+	return false;
+}
+
+FSpaceControlNameAndChannel* UMovieSceneControlRigParameterSection::GetSpaceChannel(FName InParameterName) 
+{
+	for (FSpaceControlNameAndChannel& SpaceChannel : SpaceChannels)
+	{
+		if (SpaceChannel.ControlName == InParameterName)
+		{
+			return &SpaceChannel;
+		}
+	}
+	return nullptr;
+}
+
+FName UMovieSceneControlRigParameterSection::FindControlNameFromSpaceChannel(const FMovieSceneControlRigSpaceChannel* InSpaceChannel) const
+{
+	for (const FSpaceControlNameAndChannel& SpaceChannel : SpaceChannels)
+	{
+		if (&(SpaceChannel.SpaceCurve) == InSpaceChannel)
+		{
+			return SpaceChannel.ControlName;
+		}
+	}
+	return NAME_None;
+}
+
 void UMovieSceneControlRigParameterSection::AddScalarParameter(FName InParameterName, TOptional<float> DefaultValue, bool bReconstructChannel)
 {
 	FMovieSceneFloatChannel* ExistingChannel = nullptr;
@@ -1110,16 +1146,51 @@ void UMovieSceneControlRigParameterSection::AddTransformParameter(FName InParame
 	}
 }
 
+void UMovieSceneControlRigParameterSection::AddSpaceChannel(FName InControlName, bool bReconstructChannel)
+{
+	if (!HasSpaceChannel(InControlName))
+	{
+		SpaceChannels.Add(FSpaceControlNameAndChannel(InControlName));
+	}
+	if (bReconstructChannel)
+	{
+		ReconstructChannelProxy();
+	}
+}
+
+TArray<FSpaceControlNameAndChannel>& UMovieSceneControlRigParameterSection::GetSpaceChannels()
+{
+	return SpaceChannels;
+}
+const TArray< FSpaceControlNameAndChannel>& UMovieSceneControlRigParameterSection::GetSpaceChannels() const
+{
+	return SpaceChannels;
+}
+
 
 void UMovieSceneControlRigParameterSection::ReconstructChannelProxy()
 {
 	FMovieSceneChannelProxyData Channels;
 	ControlChannelMap.Empty();
 	// Need to create the channels in sorted orders
-	if (ControlRig)
+	if (ControlRig )
 	{
 		TArray<FRigControlElement*> SortedControls;
 		ControlRig->GetControlsInOrder(SortedControls);
+		bool bIsInUndo = false;
+#if WITH_EDITORONLY_DATA
+		bIsInUndo = GIsTransacting;
+		if (bIsInUndo) //if undo reset spaces @helge to do need to reset to some default not parent since rig may be set up to something else
+		{
+			if (URigHierarchy* RigHierarchy = ControlRig->GetHierarchy())
+			{
+				for (FRigControlElement* ControlElement : SortedControls)
+				{
+					RigHierarchy->SwitchToDefaultParent(ControlElement->GetKey());
+				}
+			}
+		}	
+#endif
 		if (ControlsMask.Num() != SortedControls.Num())
 		{
 			TArray<bool> OnArray;
@@ -1133,10 +1204,12 @@ void UMovieSceneControlRigParameterSection::ReconstructChannelProxy()
 		int32 BoolChannelIndex = 0;
 		int32 EnumChannelIndex = 0;
 		int32 IntegerChannelIndex = 0;
+		int32 SpaceChannelIndeX = 0;
 		const FName BoolChannelTypeName = FMovieSceneBoolChannel::StaticStruct()->GetFName();
 		const FName EnumChannelTypeName = FMovieSceneByteChannel::StaticStruct()->GetFName();
 		const FName IntegerChannelTypeName = FMovieSceneIntegerChannel::StaticStruct()->GetFName();
-			
+		const FName SpaceName = FName(TEXT("Space"));
+
 		for (FRigControlElement* ControlElement : SortedControls)
 		{
 			if (!ControlElement->Settings.bAnimatable)
@@ -1412,22 +1485,48 @@ void UMovieSceneControlRigParameterSection::ReconstructChannelProxy()
 					{
 						if (ControlElement->Settings.ControlType == ERigControlType::Scale)
 						{
-							if (BlendType == EMovieSceneBlendType::Additive)
-							{
-								Vector.XCurve.SetDefault(0.0f);
-								Vector.YCurve.SetDefault(0.0f);
-								Vector.ZCurve.SetDefault(0.0f);
-							}
-							else
+							//only set defaults when not loading or undo
+							if (HasAnyFlags(RF_NeedLoad | RF_NeedPostLoad | RF_NeedPostLoadSubobjects) == false && bIsInUndo == false)
 							{
-								Vector.XCurve.SetDefault(1.0f);
-								Vector.YCurve.SetDefault(1.0f);
-								Vector.ZCurve.SetDefault(1.0f);
+								if (BlendType == EMovieSceneBlendType::Additive)
+								{
+									Vector.XCurve.SetDefault(0.0f);
+									Vector.YCurve.SetDefault(0.0f);
+									Vector.ZCurve.SetDefault(0.0f);
+								}
+								else
+								{
+									Vector.XCurve.SetDefault(1.0f);
+									Vector.YCurve.SetDefault(1.0f);
+									Vector.ZCurve.SetDefault(1.0f);
+								}
 							}
 						}
 						ControlChannelMap.Add(Vector.ParameterName, FChannelMapInfo(ControlIndex, TotalIndex, FloatChannelIndex));
 						bool bEnabled = ControlsMask[ControlIndex++];
 						FText Group = FText::FromName(ControlElement->GetDisplayName());
+						if (FSpaceControlNameAndChannel* SpaceChannel = GetSpaceChannel(Vector.ParameterName))
+						{
+
+							FChannelMapInfo* pChannelIndex = ControlChannelMap.Find(Vector.ParameterName);
+							if (pChannelIndex)
+							{
+								pChannelIndex->bDoesHaveSpace = true;
+							}
+
+							FString TotalName = Vector.ParameterName.ToString(); //need ControlName.Space for selection to work.
+							FString SpaceString = SpaceName.ToString();
+							TotalName += ("." + SpaceString);
+							FMovieSceneChannelMetaData SpaceMetaData(FName(*TotalName), Group, Group, bEnabled);
+							SpaceMetaData.DisplayText = FText::FromName(SpaceName);
+							SpaceChannelIndeX += 1;
+							SpaceMetaData.SortOrder = TotalIndex++;
+							// Prevent single channels from collapsing to the track node
+							SpaceMetaData.bCanCollapseToTrack = false;
+							Channels.Add(SpaceChannel->SpaceCurve, SpaceMetaData);
+						}
+						
+
 						FParameterVectorChannelEditorData EditorData(ControlRig, Vector.ParameterName, bEnabled, Group, TotalIndex, 3);
 						Channels.Add(Vector.XCurve, EditorData.MetaData[0], EditorData.ExternalValues[0]);
 						Channels.Add(Vector.YCurve, EditorData.MetaData[1], EditorData.ExternalValues[1]);
@@ -1452,6 +1551,29 @@ void UMovieSceneControlRigParameterSection::ReconstructChannelProxy()
 						bool bEnabled = ControlsMask[ControlIndex++];
 						FText Group = FText::FromName(ControlElement->GetDisplayName());
 
+						if (FSpaceControlNameAndChannel* SpaceChannel = GetSpaceChannel(Transform.ParameterName))
+						{
+
+							FChannelMapInfo* pChannelIndex = ControlChannelMap.Find(Transform.ParameterName);
+							if (pChannelIndex)
+							{
+								pChannelIndex->bDoesHaveSpace = true;
+							}
+
+							FString TotalName = Transform.ParameterName.ToString(); //need ControlName.Space for selection to work.
+							FString SpaceString = SpaceName.ToString();
+							TotalName += ("." + SpaceString);
+							FMovieSceneChannelMetaData SpaceMetaData(FName(*TotalName), Group, Group, bEnabled);
+							SpaceMetaData.DisplayText = FText::FromName(SpaceName);
+							SpaceChannelIndeX += 1;
+							SpaceMetaData.SortOrder = TotalIndex++;
+							// Prevent single channels from collapsing to the track node
+							SpaceMetaData.bCanCollapseToTrack = false;
+							//TMovieSceneExternalValue<FMovieSceneControlRigSpaceBaseKey> ExternalData;
+							Channels.Add(SpaceChannel->SpaceCurve, SpaceMetaData);
+						}
+					
+
 						FParameterTransformChannelEditorData EditorData(ControlRig, Transform.ParameterName, bEnabled, TransformMask.GetChannels(), Group, 
 							TotalIndex);
 
@@ -1466,17 +1588,21 @@ void UMovieSceneControlRigParameterSection::ReconstructChannelProxy()
 						if (ControlElement->Settings.ControlType == ERigControlType::Transform ||
 							ControlElement->Settings.ControlType == ERigControlType::EulerTransform)
 						{
-							if (BlendType == EMovieSceneBlendType::Additive)
+							//only set defaults when not loading or undo
+							if (HasAnyFlags(RF_NeedLoad | RF_NeedPostLoad | RF_NeedPostLoadSubobjects) == false && bIsInUndo == false)
 							{
-								Transform.Scale[0].SetDefault(0.0f);
-								Transform.Scale[1].SetDefault(0.0f);
-								Transform.Scale[2].SetDefault(0.0f);
-							}
-							else
-							{
-								Transform.Scale[0].SetDefault(1.0f);
-								Transform.Scale[1].SetDefault(1.0f);
-								Transform.Scale[2].SetDefault(1.0f);
+								if (BlendType == EMovieSceneBlendType::Additive)
+								{
+									Transform.Scale[0].SetDefault(0.0f);
+									Transform.Scale[1].SetDefault(0.0f);
+									Transform.Scale[2].SetDefault(0.0f);
+								}
+								else
+								{
+									Transform.Scale[0].SetDefault(1.0f);
+									Transform.Scale[1].SetDefault(1.0f);
+									Transform.Scale[2].SetDefault(1.0f);
+								}
 							}
 							Channels.Add(Transform.Scale[0], EditorData.MetaData[6], EditorData.ExternalValues[6]);
 							Channels.Add(Transform.Scale[1], EditorData.MetaData[7], EditorData.ExternalValues[7]);
@@ -1593,6 +1719,21 @@ void UMovieSceneControlRigParameterSection::ReconstructChannelProxy()
 					if (ControlElement->GetName() == Vector.ParameterName)
 					{
 						ControlChannelMap.Add(Vector.ParameterName, FChannelMapInfo(ControlIndex, TotalIndex, FloatChannelIndex));
+						bool bDoSpaceChannel = true;
+						if (bDoSpaceChannel)
+						{
+							if (FSpaceControlNameAndChannel* SpaceChannel = GetSpaceChannel(Vector.ParameterName))
+							{
+								FChannelMapInfo* pChannelIndex = ControlChannelMap.Find(Vector.ParameterName);
+								if (pChannelIndex)
+								{
+									pChannelIndex->bDoesHaveSpace = true;
+								}
+								SpaceChannelIndeX += 1;
+								Channels.Add(SpaceChannel->SpaceCurve);
+							}
+						}
+
 						Channels.Add(Vector.XCurve);
 						Channels.Add(Vector.YCurve);
 						Channels.Add(Vector.ZCurve);
@@ -1623,6 +1764,23 @@ void UMovieSceneControlRigParameterSection::ReconstructChannelProxy()
 					if (ControlElement->GetName() == Transform.ParameterName)
 					{
 						ControlChannelMap.Add(Transform.ParameterName, FChannelMapInfo(ControlIndex, TotalIndex,FloatChannelIndex));
+						
+						bool bDoSpaceChannel = true;
+						if (bDoSpaceChannel)
+						{
+							if (FSpaceControlNameAndChannel* SpaceChannel = GetSpaceChannel(Transform.ParameterName))
+							{
+
+								FChannelMapInfo* pChannelIndex = ControlChannelMap.Find(Transform.ParameterName);
+								if (pChannelIndex)
+								{
+									pChannelIndex->bDoesHaveSpace = true;
+								}
+								SpaceChannelIndeX += 1;
+								Channels.Add(SpaceChannel->SpaceCurve);
+							}
+						}	
+						
 						Channels.Add(Transform.Translation[0]);
 						Channels.Add(Transform.Translation[1]);
 						Channels.Add(Transform.Translation[2]);
diff --git a/Engine/Plugins/Animation/ControlRig/Source/ControlRig/Private/Sequencer/MovieSceneControlRigParameterTemplate.cpp b/Engine/Plugins/Animation/ControlRig/Source/ControlRig/Private/Sequencer/MovieSceneControlRigParameterTemplate.cpp
--- a/Engine/Plugins/Animation/ControlRig/Source/ControlRig/Private/Sequencer/MovieSceneControlRigParameterTemplate.cpp
+++ b/Engine/Plugins/Animation/ControlRig/Source/ControlRig/Private/Sequencer/MovieSceneControlRigParameterTemplate.cpp
@@ -19,6 +19,7 @@
 #include "SkeletalMeshRestoreState.h"
 #include "Rigs/FKControlRig.h"
 #include "UObject/UObjectAnnotation.h"
+#include "Rigs/RigHierarchy.h"
 
 //#include "Particles/ParticleSystemComponent.h"
 
@@ -84,6 +85,16 @@ struct FIntegerParameterStringAndValue
 	int32 Value;
 };
 
+struct FControlSpaceAndValue
+{
+	FControlSpaceAndValue(FName InControlName, const FMovieSceneControlRigSpaceBaseKey& InValue)
+	{
+		ControlName = InControlName;
+		Value = InValue;
+	}
+	FName ControlName;
+	FMovieSceneControlRigSpaceBaseKey Value;
+};
 /**
  * Structure representing the animated value of a vector2D parameter.
  */
@@ -355,6 +366,8 @@ struct FEvaluatedControlRigParameterSectionValues
 	TArray<FBoolParameterStringAndValue, TInlineAllocator<2>> BoolValues;
 	/** Array of evaluated integer values */
 	TArray<FIntegerParameterStringAndValue, TInlineAllocator<2>> IntegerValues;
+	/** Array of evaluated Spaces */
+	TArray<FControlSpaceAndValue, TInlineAllocator<2>> SpaceValues;
 	/** Array of evaluated vector2d values */
 	TArray<FVector2DParameterStringAndValue, TInlineAllocator<2>> Vector2DValues;
 	/** Array of evaluated vector values */
@@ -648,6 +661,30 @@ struct FControlRigParameterPreAnimatedTokenProducer : IMovieScenePreAnimatedToke
 					{
 						//Restore control rig first
 						const bool bSetupUndo = false;
+						if (URigHierarchy* RigHierarchy = ControlRig->GetHierarchy())
+						{
+							FRigElementKey ControlKey;
+							ControlKey.Type = ERigElementType::Control;
+							for (FControlSpaceAndValue& SpaceNameAndValue : SpaceValues)
+							{
+								ControlKey.Name = SpaceNameAndValue.ControlName;
+								switch (SpaceNameAndValue.Value.SpaceType)
+								{
+								case EMovieSceneControlRigSpaceType::Parent:
+									RigHierarchy->SwitchToDefaultParent(ControlKey);
+									break;
+								case EMovieSceneControlRigSpaceType::World:
+									RigHierarchy->SwitchToWorldSpace(ControlKey);
+									break;
+								case EMovieSceneControlRigSpaceType::ControlRig:
+								{
+									RigHierarchy->SwitchToParent(ControlKey, SpaceNameAndValue.Value.ControlRigElement);
+								}
+								break;
+								}
+							}
+						}
+
 						for (TNameAndValue<float>& Value : ScalarValues)
 						{
 							if (ControlRig->FindControl(Value.Name))
@@ -734,6 +771,7 @@ struct FControlRigParameterPreAnimatedTokenProducer : IMovieScenePreAnimatedToke
 			}
 
 			FMovieSceneSequenceID SequenceID;
+			TArray< FControlSpaceAndValue> SpaceValues;
 			TArray< TNameAndValue<float> > ScalarValues;
 			TArray< TNameAndValue<bool> > BoolValues;
 			TArray< TNameAndValue<int32> > IntegerValues;
@@ -784,6 +822,14 @@ struct FControlRigParameterPreAnimatedTokenProducer : IMovieScenePreAnimatedToke
 					case ERigControlType::Scale:
 					case ERigControlType::Rotator:
 					{
+						URigHierarchy* RigHierarchy = ControlRig->GetHierarchy();
+						FMovieSceneControlRigSpaceBaseKey SpaceValue;
+						//@helge how can we get the correct current space here? this is for restoring it.
+						//for now just using parent space
+						//FRigElementKey DefaultParent = RigHierarchy->GetFirstParent(ControlElement->GetKey());
+						SpaceValue.ControlRigElement = ControlElement->GetKey();
+						SpaceValue.SpaceType = EMovieSceneControlRigSpaceType::Parent;
+						Token.SpaceValues.Add(FControlSpaceAndValue(ControlElement->GetName(), SpaceValue));
 						const FVector3f Val = ControlRig->GetHierarchy()->GetControlValue(ControlElement, ERigControlValueType::Current).Get<FVector3f>();
 						Token.VectorValues.Add(TNameAndValue<FVector>{ ControlElement->GetName(), Val });
 						//mz todo specify rotator special so we can do quat interps
@@ -791,9 +837,14 @@ struct FControlRigParameterPreAnimatedTokenProducer : IMovieScenePreAnimatedToke
 					}
 					case ERigControlType::Transform:
 					{
-						const FTransform Val = 
-							ControlRig->GetHierarchy()
-							->GetControlValue(ControlElement, ERigControlValueType::Current).Get<FRigControlValue::FTransform_Float>().ToTransform();
+						FMovieSceneControlRigSpaceBaseKey SpaceValue;
+						//@helge how can we get the correct current space here? this is for restoring it.
+						//for now just using parent space
+						//FRigElementKey DefaultParent = RigHierarchy->GetFirstParent(ControlElement->GetKey());
+						SpaceValue.ControlRigElement = ControlElement->GetKey();
+						SpaceValue.SpaceType = EMovieSceneControlRigSpaceType::Parent;
+						Token.SpaceValues.Add(FControlSpaceAndValue(ControlElement->GetName(), SpaceValue));
+						const FTransform Val = ControlRig->GetHierarchy()->GetControlValue(ControlElement, ERigControlValueType::Current).Get<FRigControlValue::FTransform_Float>().ToTransform();
 						Token.TransformValues.Add(TNameAndValue<FTransform>{ ControlElement->GetName(), Val });
 						break;
 					}
@@ -856,6 +907,7 @@ struct FControlRigParameterPreAnimatedTokenProducer : IMovieScenePreAnimatedToke
 	}
 
 	FMovieSceneSequenceID SequenceID;
+	TArray< FControlSpaceAndValue> SpaceValues;
 	TArray< TNameAndValue<bool> > BoolValues;
 	TArray< TNameAndValue<int32> > IntegerValues;
 	TArray< TNameAndValue<float> > ScalarValues;
@@ -953,6 +1005,7 @@ struct FControlRigParameterExecutionToken : IMovieSceneExecutionToken
 	{
 		BoolValues = Values.BoolValues;
 		IntegerValues = Values.IntegerValues;
+		SpaceValues = Values.SpaceValues;
 	}
 	FControlRigParameterExecutionToken(FControlRigParameterExecutionToken&&) = default;
 	FControlRigParameterExecutionToken& operator=(FControlRigParameterExecutionToken&&) = default;
@@ -1033,6 +1086,43 @@ struct FControlRigParameterExecutionToken : IMovieSceneExecutionToken
 			{
 				const bool bSetupUndo = false;
 				ControlRig->SetAbsoluteTime((float)Context.GetFrameRate().AsSeconds(Context.GetTime()));
+				for (const FControlSpaceAndValue& SpaceNameAndValue : SpaceValues)
+				{
+					if (Section->ControlsToSet.Num() == 0 || Section->ControlsToSet.Contains(SpaceNameAndValue.ControlName))
+					{
+						if (URigHierarchy* RigHierarchy = ControlRig->GetHierarchy())
+						{
+							if (FRigControlElement* RigControl = ControlRig->FindControl(SpaceNameAndValue.ControlName))
+							{
+								const FRigElementKey ControlKey = RigControl->GetKey();
+
+								switch (SpaceNameAndValue.Value.SpaceType)
+								{
+									case EMovieSceneControlRigSpaceType::Parent:
+										RigHierarchy->SwitchToDefaultParent(ControlKey);
+										break;
+									case EMovieSceneControlRigSpaceType::World:
+										RigHierarchy->SwitchToWorldSpace(ControlKey);
+										break;
+									case EMovieSceneControlRigSpaceType::ControlRig:
+										{
+											RigHierarchy->SwitchToParent(ControlKey, SpaceNameAndValue.Value.ControlRigElement);	
+										}
+										break;
+								}
+							}
+						}
+						/*
+						FRigControlElement* RigControl = ControlRig->FindControl(SpaceNameAndValue.ParameterName);
+						if (RigControl && RigControl->Settings.ControlType == ERigControlType::Bool)
+						{
+							ControlRig->SetControlValue<bool>(SpaceNameAndValue.SpaceNameAndValue, SpaceNameAndValue.Value, true, EControlRigSetKey::Never, bSetupUndo);
+						}
+						*/
+					}
+				}
+
+
 				for (const FBoolParameterStringAndValue& BoolNameAndValue : BoolValues)
 				{
 					if (Section->ControlsToSet.Num() == 0 || Section->ControlsToSet.Contains(BoolNameAndValue.ParameterName))
@@ -1067,12 +1157,14 @@ struct FControlRigParameterExecutionToken : IMovieSceneExecutionToken
 	TArray<FBoolParameterStringAndValue, TInlineAllocator<2>> BoolValues;
 	/** Array of evaluated integer values */
 	TArray<FIntegerParameterStringAndValue, TInlineAllocator<2>> IntegerValues;
+	/** Array of Space Values*/
+	TArray<FControlSpaceAndValue, TInlineAllocator<2>> SpaceValues;
 
 };
 
 
 FMovieSceneControlRigParameterTemplate::FMovieSceneControlRigParameterTemplate(const UMovieSceneControlRigParameterSection& Section, const UMovieSceneControlRigParameterTrack& Track)
-	: FMovieSceneParameterSectionTemplate(Section) , Enums(Section.GetEnumParameterNamesAndCurves()), Integers(Section.GetIntegerParameterNamesAndCurves())
+	: FMovieSceneParameterSectionTemplate(Section) , Enums(Section.GetEnumParameterNamesAndCurves()), Integers(Section.GetIntegerParameterNamesAndCurves()), Spaces(Section.GetSpaceChannels())
 
 {
 
@@ -1512,6 +1604,7 @@ void FMovieSceneControlRigParameterTemplate::EvaluateCurvesWithMasks(const FMovi
 		// Reserve the value arrays to avoid re-allocation
 		Values.ScalarValues.Reserve(Scalars.Num());
 		Values.BoolValues.Reserve(Bools.Num());
+		Values.SpaceValues.Reserve(Spaces.Num());
 		Values.IntegerValues.Reserve(Integers.Num() + Enums.Num()); // Both enums and integers output to the integer value array
 		Values.Vector2DValues.Reserve(Vector2Ds.Num());
 		Values.VectorValues.Reserve(Vectors.Num());
@@ -1538,6 +1631,15 @@ void FMovieSceneControlRigParameterTemplate::EvaluateCurvesWithMasks(const FMovi
 			Values.ScalarValues.Emplace(Scalar.ParameterName, Value);
 		}
 
+		for (int32 Index = 0; Index < Spaces.Num(); ++Index)
+		{
+			FMovieSceneControlRigSpaceBaseKey Value;
+			const FSpaceControlNameAndChannel& Space = Spaces[Index];
+			Space.SpaceCurve.Evaluate(Time, Value);
+			
+			Values.SpaceValues.Emplace(Space.ControlName, Value);
+		}
+
 		for (int32 Index = 0; Index < Bools.Num(); ++Index)
 		{
 			bool Value = false;
diff --git a/Engine/Plugins/Animation/ControlRig/Source/ControlRig/Private/Sequencer/MovieSceneControlRigSpaceChannel.cpp b/Engine/Plugins/Animation/ControlRig/Source/ControlRig/Private/Sequencer/MovieSceneControlRigSpaceChannel.cpp
new file mode 100644
--- /dev/null
+++ b/Engine/Plugins/Animation/ControlRig/Source/ControlRig/Private/Sequencer/MovieSceneControlRigSpaceChannel.cpp
@@ -0,0 +1,153 @@
+// Copyright Epic Games, Inc. All Rights Reserved.
+
+#include "Sequencer/MovieSceneControlRigSpaceChannel.h"
+#include "Channels/MovieSceneChannelProxy.h"
+#include "MovieSceneObjectBindingID.h"
+#include "MovieSceneSequenceID.h"
+#include "Evaluation/MovieSceneSequenceHierarchy.h"
+#include "IMovieScenePlayer.h"
+
+bool FMovieSceneControlRigSpaceChannel::Evaluate(FFrameTime InTime, FMovieSceneControlRigSpaceBaseKey& OutValue) const
+{
+	FMovieSceneControlRigSpaceBaseKey DefaultValue; //this will be in parent space
+	if (KeyTimes.Num())
+	{
+		const int32 Index = FMath::Max(0, Algo::UpperBound(KeyTimes, InTime.FrameNumber) - 1);
+		OutValue = KeyValues[Index];
+		return true;
+	}
+	OutValue = DefaultValue;
+	return true;
+}
+
+void FMovieSceneControlRigSpaceChannel::GetKeys(const TRange<FFrameNumber>& WithinRange, TArray<FFrameNumber>* OutKeyTimes, TArray<FKeyHandle>* DeleteKeyHandles)
+{
+	GetData().GetKeys(WithinRange, OutKeyTimes, DeleteKeyHandles);
+}
+
+void FMovieSceneControlRigSpaceChannel::GetKeyTimes(TArrayView<const FKeyHandle> InHandles, TArrayView<FFrameNumber> OutKeyTimes)
+{
+	GetData().GetKeyTimes(InHandles, OutKeyTimes);
+}
+
+void FMovieSceneControlRigSpaceChannel::SetKeyTimes(TArrayView<const FKeyHandle> InHandles, TArrayView<const FFrameNumber> InKeyTimes)
+{
+	GetData().SetKeyTimes(InHandles, InKeyTimes);
+}
+
+void FMovieSceneControlRigSpaceChannel::DuplicateKeys(TArrayView<const FKeyHandle> InHandles, TArrayView<FKeyHandle> OutNewHandles)
+{
+	GetData().DuplicateKeys(InHandles, OutNewHandles);
+}
+
+void FMovieSceneControlRigSpaceChannel::DeleteKeys(TArrayView<const FKeyHandle> InHandles)
+{
+	GetData().DeleteKeys(InHandles);
+}
+
+void FMovieSceneControlRigSpaceChannel::DeleteKeysFrom(FFrameNumber InTime, bool bDeleteKeysBefore)
+{
+	// Insert a key at the current time to maintain evaluation
+	if (GetData().GetTimes().Num() > 0)
+	{
+		FMovieSceneControlRigSpaceBaseKey Value;
+		Evaluate(InTime, Value);
+		GetData().UpdateOrAddKey(InTime, Value);
+	}
+
+	GetData().DeleteKeysFrom(InTime, bDeleteKeysBefore);
+}
+
+void FMovieSceneControlRigSpaceChannel::ChangeFrameResolution(FFrameRate SourceRate, FFrameRate DestinationRate)
+{
+	GetData().ChangeFrameResolution(SourceRate, DestinationRate);
+}
+
+TRange<FFrameNumber> FMovieSceneControlRigSpaceChannel::ComputeEffectiveRange() const
+{
+	return GetData().GetTotalRange();
+}
+
+int32 FMovieSceneControlRigSpaceChannel::GetNumKeys() const
+{
+	return KeyTimes.Num();
+}
+
+void FMovieSceneControlRigSpaceChannel::Reset()
+{
+	KeyTimes.Reset();
+	KeyValues.Reset();
+	KeyHandles.Reset();
+}
+
+void FMovieSceneControlRigSpaceChannel::Offset(FFrameNumber DeltaPosition)
+{
+	GetData().Offset(DeltaPosition);
+}
+
+//this will also delete any duplicated keys.
+TArray <FSpaceRange> FMovieSceneControlRigSpaceChannel::FindSpaceIntervals()
+{
+	TArray<FSpaceRange> Ranges;
+	TArray<int32> ToDelete;
+	for (int32 Index = 0; Index < KeyTimes.Num(); ++Index)
+	{
+		FSpaceRange Range;
+		Range.Key = KeyValues[Index];
+		if (Index == KeyTimes.Num() - 1)
+		{
+			Range.Range.SetLowerBoundValue(KeyTimes[Index]);
+			Range.Range.SetUpperBoundValue(KeyTimes[Index]);
+		}
+		else
+		{
+			int32 NextIndex = Index;
+			FFrameNumber UpperBound = KeyTimes[Index];
+			while (NextIndex < KeyTimes.Num() -1)
+			{
+				NextIndex = Index + 1;
+				if (Range.Key == KeyValues[NextIndex])
+				{
+					ToDelete.Add(NextIndex);
+				}
+				else
+				{
+					UpperBound = KeyTimes[NextIndex];
+					Index = NextIndex - 1;
+					NextIndex = KeyTimes.Num();
+				}
+			}
+		}
+		Ranges.Add(Range);
+	}
+	//now delete duplicate values
+	if (ToDelete.Num() > 0)
+	{
+		TArray<FKeyHandle> DeleteKeyHandles;
+		for (int32 Index : ToDelete)
+		{
+			FKeyHandle KeyHandle = GetData().GetHandle(Index);
+			DeleteKeyHandles.Add(KeyHandle);
+		}
+
+		GetData().DeleteKeys(DeleteKeyHandles);
+	}
+	return Ranges;
+}
+
+
+/* Is this needed if not remove mz todoo
+TSharedPtr<FStructOnScope> GetKeyStruct(TMovieSceneChannelHandle<FMovieSceneControlRigSpaceChannel> Channel, FKeyHandle InHandle)
+{
+	int32 KeyValueIndex = Channel.Get()->GetData().GetIndex(InHandle);
+	if (KeyValueIndex != INDEX_NONE)
+	{
+		FNiagaraTypeDefinition KeyType = Channel.Get()->GetData().GetValues()[KeyValueIndex].Value.GetType();
+		uint8* KeyData = Channel.Get()->GetData().GetValues()[KeyValueIndex].Value.GetData();
+		return MakeShared<FStructOnScope>(KeyType.GetStruct(), KeyData);
+		
+	}
+	return TSharedPtr<FStructOnScope>();
+}
+*/
+
diff --git a/Engine/Plugins/Animation/ControlRig/Source/ControlRig/Public/Sequencer/MovieSceneControlRigParameterSection.h b/Engine/Plugins/Animation/ControlRig/Source/ControlRig/Public/Sequencer/MovieSceneControlRigParameterSection.h
--- a/Engine/Plugins/Animation/ControlRig/Source/ControlRig/Public/Sequencer/MovieSceneControlRigParameterSection.h
+++ b/Engine/Plugins/Animation/ControlRig/Source/ControlRig/Public/Sequencer/MovieSceneControlRigParameterSection.h
@@ -14,6 +14,7 @@
 #include "Compilation/MovieSceneTemplateInterrogation.h"
 #include "Channels/MovieSceneIntegerChannel.h"
 #include "Channels/MovieSceneByteChannel.h"
+#include "Sequencer/MovieSceneControlRigSpaceChannel.h"
 
 
 #include "MovieSceneControlRigParameterSection.generated.h"
@@ -87,6 +88,21 @@ struct CONTROLRIG_API FIntegerParameterNameAndCurve
 	FMovieSceneIntegerChannel ParameterCurve;
 };
 
+USTRUCT()
+struct CONTROLRIG_API FSpaceControlNameAndChannel
+{
+	GENERATED_USTRUCT_BODY()
+
+	FSpaceControlNameAndChannel(){}
+	FSpaceControlNameAndChannel(FName InControlName) : ControlName(InControlName) {};
+
+	UPROPERTY()
+	FName ControlName;
+
+	UPROPERTY()
+	FMovieSceneControlRigSpaceChannel SpaceCurve;
+};
+
 /**
 *  Data that's queried during an interrogtion
 */
@@ -133,6 +149,8 @@ struct CONTROLRIG_API FChannelMapInfo
 	int32 ParentControlIndex = 0;
 	UPROPERTY()
 	FName ChannelTypeName; 
+	UPROPERTY()
+	bool bDoesHaveSpace = false;
 
 
 };
@@ -159,6 +177,9 @@ public:
 	TArray<FIntegerParameterNameAndCurve>& GetIntegerParameterNamesAndCurves();
 	const TArray<FIntegerParameterNameAndCurve>& GetIntegerParameterNamesAndCurves() const;
 
+	TArray<FSpaceControlNameAndChannel>& GetSpaceChannels();
+	const TArray< FSpaceControlNameAndChannel>& GetSpaceChannels() const;
+	FName FindControlNameFromSpaceChannel(const FMovieSceneControlRigSpaceChannel* SpaceChannel) const;
 private:
 
 	/** Control Rig that controls us*/
@@ -196,6 +217,10 @@ protected:
 	UPROPERTY()
 	TArray<FIntegerParameterNameAndCurve> IntegerParameterNamesAndCurves;
 
+	/** Space Channels*/
+	UPROPERTY()
+	TArray<FSpaceControlNameAndChannel>  SpaceChannels;
+
 public:
 
 	UMovieSceneControlRigParameterSection();
@@ -301,6 +326,12 @@ public:
 	/**  Whether or not this section his scalar*/
 	bool HasTransformParameter(FName InParameterName) const;
 
+	/**  Whether or not this section his space*/
+	bool HasSpaceChannel(FName InParameterName) const;
+
+	/** Get The Space Channel for the Control*/
+	FSpaceControlNameAndChannel* GetSpaceChannel(FName InParameterName);
+
 	/** Adds specified scalar parameter. */
 	void AddScalarParameter(FName InParameterName,  TOptional<float> DefaultValue, bool bReconstructChannel);
 
@@ -325,6 +356,9 @@ public:
 	/** Adds a a key for a specific color parameter*/
 	void AddTransformParameter(FName InParameterName, TOptional<FTransform> DefaultValue, bool bReconstructChannel);
 
+	/** Add Space Parameter for a specified Control, no Default since that is Parent space*/
+	void AddSpaceChannel(FName InControlName, bool bReconstructChannel);
+
 	/** Clear Everything Out*/
 	void ClearAllParameters();
 public:
diff --git a/Engine/Plugins/Animation/ControlRig/Source/ControlRig/Public/Sequencer/MovieSceneControlRigParameterTemplate.h b/Engine/Plugins/Animation/ControlRig/Source/ControlRig/Public/Sequencer/MovieSceneControlRigParameterTemplate.h
--- a/Engine/Plugins/Animation/ControlRig/Source/ControlRig/Public/Sequencer/MovieSceneControlRigParameterTemplate.h
+++ b/Engine/Plugins/Animation/ControlRig/Source/ControlRig/Public/Sequencer/MovieSceneControlRigParameterTemplate.h
@@ -32,8 +32,11 @@ protected:
 	/** The bool parameter names and their associated curves. */
 	UPROPERTY()
 	TArray<FEnumParameterNameAndCurve> Enums;
+	/** The enum parameter names and their associated curves. */
 	UPROPERTY()
 	TArray<FIntegerParameterNameAndCurve> Integers;
-
+	/** Controls and their space keys*/
+	UPROPERTY()
+	TArray<FSpaceControlNameAndChannel> Spaces;
 };
 
diff --git a/Engine/Plugins/Animation/ControlRig/Source/ControlRig/Public/Sequencer/MovieSceneControlRigSpaceChannel.h b/Engine/Plugins/Animation/ControlRig/Source/ControlRig/Public/Sequencer/MovieSceneControlRigSpaceChannel.h
new file mode 100644
--- /dev/null
+++ b/Engine/Plugins/Animation/ControlRig/Source/ControlRig/Public/Sequencer/MovieSceneControlRigSpaceChannel.h
@@ -0,0 +1,146 @@
+// Copyright Epic Games, Inc. All Rights Reserved.
+
+#pragma once
+
+#include "CoreMinimal.h"
+#include "UObject/ObjectMacros.h"
+#include "Misc/Guid.h"
+#include "Curves/KeyHandle.h"
+#include "MovieSceneObjectBindingID.h"
+#include "MovieSceneSection.h"
+#include "Curves/IntegralCurve.h"
+#include "Channels/MovieSceneChannel.h"
+#include "Channels/MovieSceneChannelData.h"
+#include "Channels/MovieSceneChannelTraits.h"
+#include "MovieSceneClipboard.h"
+#include "Rigs/RigHierarchyDefines.h"
+#include "MovieSceneControlRigSpaceChannel.generated.h"
+
+class UControlRig;
+
+UENUM(Blueprintable)
+enum class EMovieSceneControlRigSpaceType : uint8
+{
+	Parent = 0,
+	World,
+	ControlRig
+};
+
+USTRUCT()
+struct CONTROLRIG_API FMovieSceneControlRigSpaceBaseKey
+{
+	GENERATED_BODY()
+	
+	friend bool operator==(const FMovieSceneControlRigSpaceBaseKey& A, const FMovieSceneControlRigSpaceBaseKey& B)
+	{
+		return A.SpaceType == B.SpaceType && (A.SpaceType != EMovieSceneControlRigSpaceType::ControlRig || A.ControlRigElement == B.ControlRigElement);
+	}
+
+	friend bool operator!=(const FMovieSceneControlRigSpaceBaseKey& A, const FMovieSceneControlRigSpaceBaseKey& B)
+	{
+		return A.SpaceType != B.SpaceType || (A.SpaceType == EMovieSceneControlRigSpaceType::ControlRig && A.ControlRigElement != B.ControlRigElement);
+	}
+
+	UPROPERTY(EditAnywhere, Category = "Key")
+	EMovieSceneControlRigSpaceType SpaceType = EMovieSceneControlRigSpaceType::Parent;
+	UPROPERTY(EditAnywhere, Category = "Key")
+	FRigElementKey ControlRigElement;
+
+};
+
+struct FSpaceRange
+{
+	TRange<FFrameNumber> Range;
+	FMovieSceneControlRigSpaceBaseKey Key;
+};
+
+/** A curve of spaces */
+USTRUCT()
+struct CONTROLRIG_API FMovieSceneControlRigSpaceChannel : public FMovieSceneChannel
+{
+	GENERATED_BODY()
+
+	FMovieSceneControlRigSpaceChannel()
+	{}
+
+	/**
+	* Access a mutable interface for this channel's data
+	*
+	* @return An object that is able to manipulate this channel's data
+	*/
+	FORCEINLINE TMovieSceneChannelData<FMovieSceneControlRigSpaceBaseKey> GetData()
+	{
+		return TMovieSceneChannelData<FMovieSceneControlRigSpaceBaseKey>(&KeyTimes, &KeyValues, &KeyHandles);
+	}
+
+	/**
+	* Access a constant interface for this channel's data
+	*
+	* @return An object that is able to interrogate this channel's data
+	*/
+	FORCEINLINE TMovieSceneChannelData<const FMovieSceneControlRigSpaceBaseKey> GetData() const
+	{
+		return TMovieSceneChannelData<const FMovieSceneControlRigSpaceBaseKey>(&KeyTimes, &KeyValues);
+	}
+
+	/**
+	* Evaluate this channel
+	*
+	* @param InTime     The time to evaluate at
+	* @param OutValue   A value to receive the result
+	* @return true if the channel was evaluated successfully, false otherwise
+	*/
+	bool Evaluate(FFrameTime InTime, FMovieSceneControlRigSpaceBaseKey& OutValue) const;
+
+public:
+
+	// ~ FMovieSceneChannel Interface
+	virtual void GetKeys(const TRange<FFrameNumber>& WithinRange, TArray<FFrameNumber>* OutKeyTimes, TArray<FKeyHandle>* OutKeyHandles) override;
+	virtual void GetKeyTimes(TArrayView<const FKeyHandle> InHandles, TArrayView<FFrameNumber> OutKeyTimes) override;
+	virtual void SetKeyTimes(TArrayView<const FKeyHandle> InHandles, TArrayView<const FFrameNumber> InKeyTimes) override;
+	virtual void DuplicateKeys(TArrayView<const FKeyHandle> InHandles, TArrayView<FKeyHandle> OutNewHandles) override;
+	virtual void DeleteKeys(TArrayView<const FKeyHandle> InHandles) override;
+	virtual void DeleteKeysFrom(FFrameNumber InTime, bool bDeleteKeysBefore) override;
+	virtual void ChangeFrameResolution(FFrameRate SourceRate, FFrameRate DestinationRate) override;
+	virtual TRange<FFrameNumber> ComputeEffectiveRange() const override;
+	virtual int32 GetNumKeys() const override;
+	virtual void Reset() override;
+	virtual void Offset(FFrameNumber DeltaPosition) override;
+
+	TArray <FSpaceRange> FindSpaceIntervals();
+private:
+
+	/** Sorted array of key times */
+	UPROPERTY(meta = (KeyTimes))
+	TArray<FFrameNumber> KeyTimes;
+
+	/** Array of values that correspond to each key time */
+	UPROPERTY(meta = (KeyValues))
+	TArray<FMovieSceneControlRigSpaceBaseKey> KeyValues;
+
+	FMovieSceneKeyHandleMap KeyHandles;
+};
+
+template<>
+struct TMovieSceneChannelTraits<FMovieSceneControlRigSpaceChannel> : TMovieSceneChannelTraitsBase<FMovieSceneControlRigSpaceChannel>
+{
+	enum { SupportsDefaults = false };
+
+};
+
+inline bool EvaluateChannel(const FMovieSceneControlRigSpaceChannel* InChannel, FFrameTime InTime, FMovieSceneControlRigSpaceBaseKey& OutValue)
+{
+	return InChannel->Evaluate(InTime, OutValue);
+}
+
+#if WITH_EDITOR
+namespace MovieSceneClipboard
+{
+	template<> inline FName GetKeyTypeName<FMovieSceneControlRigSpaceBaseKey>()
+	{
+		return "FMovieSceneControlRigSpaceBaseKey";
+	}
+}
+#endif
+
+//mz todoo TSharedPtr<FStructOnScope> GetKeyStruct(TMovieSceneChannelHandle<FMovieSceneControlRigSpaceChannel> Channel, FKeyHandle InHandle);
\ No newline at end of file
diff --git a/Engine/Plugins/Animation/ControlRig/Source/ControlRigEditor/Private/ControlRigEditorModule.cpp b/Engine/Plugins/Animation/ControlRig/Source/ControlRigEditor/Private/ControlRigEditorModule.cpp
--- a/Engine/Plugins/Animation/ControlRig/Source/ControlRigEditor/Private/ControlRigEditorModule.cpp
+++ b/Engine/Plugins/Animation/ControlRig/Source/ControlRigEditor/Private/ControlRigEditorModule.cpp
@@ -120,10 +120,15 @@
 #include "EditMode/SControlRigTweenWidget.h"
 #include "EditMode/SControlRigSnapper.h"
 #include "Dialogs/CustomDialog.h"
+#include "Sequencer/MovieSceneControlRigSpaceChannel.h"
+#include "SequencerChannelInterface.h"
 #include "Framework/Notifications/NotificationManager.h"
 #include "Widgets/Notifications/SNotificationList.h"
 #include "Tools/SMotionTrailOptions.h"
 
+
+#include "ControlRigSpaceChannelEditors.h"
+
 #define LOCTEXT_NAMESPACE "ControlRigEditorModule"
 
 DEFINE_LOG_CATEGORY(LogControlRigEditor);
@@ -256,6 +261,7 @@ void FControlRigEditorModule::StartupModule()
 	
 	// Register sequencer track editor
 	ISequencerModule& SequencerModule = FModuleManager::Get().LoadModuleChecked<ISequencerModule>("Sequencer");
+	SequencerModule.RegisterChannelInterface<FMovieSceneControlRigSpaceChannel>();
 	ControlRigParameterTrackCreateEditorHandle = SequencerModule.RegisterTrackEditor(FOnCreateTrackEditor::CreateStatic(&FControlRigParameterTrackEditor::CreateTrackEditor));
 
 	//Register Animation Toolbar Extender
diff --git a/Engine/Plugins/Animation/ControlRig/Source/ControlRigEditor/Private/EditMode/ControlRigEditMode.cpp b/Engine/Plugins/Animation/ControlRig/Source/ControlRigEditor/Private/EditMode/ControlRigEditMode.cpp
--- a/Engine/Plugins/Animation/ControlRig/Source/ControlRigEditor/Private/EditMode/ControlRigEditMode.cpp
+++ b/Engine/Plugins/Animation/ControlRig/Source/ControlRigEditor/Private/EditMode/ControlRigEditMode.cpp
@@ -57,6 +57,7 @@
 #include "Sequencer/MovieSceneControlRigParameterTrack.h"
 #include "Sequencer/MovieSceneControlRigParameterSection.h"
 #include "SRigSpacePickerWidget.h"
+#include "ControlRigSpaceChannelEditors.h"
 
 void UControlRigEditModeDelegateHelper::OnPoseInitialized()
 {
@@ -2057,7 +2058,7 @@ void FControlRigEditMode::ToggleGizmoTransformEdit()
 
 void FControlRigEditMode::OpenSpacePickerWidget()
 {
-	UControlRig* RuntimeRig = GetControlRig(false);
+	UControlRig* RuntimeRig = GetControlRig(false); //helge space picker only works on runtime rig?
 	if (RuntimeRig == nullptr)
 	{
 		return;
@@ -2078,13 +2079,38 @@ void FControlRigEditMode::OpenSpacePickerWidget()
 	{
 		return RuntimeRig->GetControlCustomization(InControlKey);
 	})
-	.OnActiveSpaceChanged_Lambda([this, SelectedControls](URigHierarchy* InHierarchy, const FRigElementKey& InControlKey, const FRigElementKey& InSpaceKey)
+	.OnActiveSpaceChanged_Lambda([this, SelectedControls, RuntimeRig](URigHierarchy* InHierarchy, const FRigElementKey& InControlKey, const FRigElementKey& InSpaceKey)
 	{
 		check(SelectedControls.Contains(InControlKey));
-		
-		const FTransform Transform = InHierarchy->GetGlobalTransform(InControlKey);
-		InHierarchy->SwitchToParent(InControlKey, InSpaceKey);
-		InHierarchy->SetGlobalTransform(InControlKey, Transform);
+		if (IsInLevelEditor())
+		{
+			if (WeakSequencer.IsValid())
+			{
+
+				if (const FRigControlElement* ControlElement = InHierarchy->Find<FRigControlElement>(InControlKey))
+				{
+					ISequencer* Sequencer = WeakSequencer.Pin().Get();
+					if (Sequencer)
+					{
+						FScopedTransaction Transaction(LOCTEXT("KeyControlRigSpace", "Key Control Rig Space"));
+						FSpaceChannelAndSection SpaceChannelAndSection = FControlRigSpaceChannelHelpers::FindSpaceChannelAndSectionForControl(RuntimeRig, InControlKey.Name, Sequencer, true /*bCreateIfNeeded*/);
+						if (SpaceChannelAndSection.SpaceChannel)
+						{
+							const FFrameRate TickResolution = Sequencer->GetFocusedTickResolution();
+							const FFrameTime FrameTime = Sequencer->GetLocalTime().ConvertTo(TickResolution);
+							FFrameNumber CurrentTime = FrameTime.GetFrame();
+							FControlRigSpaceChannelHelpers::SequencerKeyControlRigSpaceChannel(RuntimeRig, Sequencer, SpaceChannelAndSection.SpaceChannel, SpaceChannelAndSection.SectionToKey, CurrentTime, InHierarchy, InControlKey, InSpaceKey);
+						}
+					}
+				}
+			}
+		}
+		else
+		{
+			const FTransform Transform = InHierarchy->GetGlobalTransform(InControlKey);
+			InHierarchy->SwitchToParent(InControlKey, InSpaceKey);
+			InHierarchy->SetGlobalTransform(InControlKey, Transform);
+		}
 		
 	})
 	.OnSpaceListChanged_Lambda([this, SelectedControls, RuntimeRig](URigHierarchy* InHierarchy, const FRigElementKey& InControlKey, const TArray<FRigElementKey>& InSpaceList)
diff --git a/Engine/Plugins/Animation/ControlRig/Source/ControlRigEditor/Private/EditMode/SControlRigEditModeTools.cpp b/Engine/Plugins/Animation/ControlRig/Source/ControlRigEditor/Private/EditMode/SControlRigEditModeTools.cpp
--- a/Engine/Plugins/Animation/ControlRig/Source/ControlRigEditor/Private/EditMode/SControlRigEditModeTools.cpp
+++ b/Engine/Plugins/Animation/ControlRig/Source/ControlRigEditor/Private/EditMode/SControlRigEditModeTools.cpp
@@ -32,6 +32,8 @@
 #include "EditorModeManager.h"
 #include "InteractiveToolManager.h"
 #include "EdModeInteractiveToolsContext.h"
+#include "ControlRigSpaceChannelEditors.h"
+#include "IKeyArea.h"
 
 #define LOCTEXT_NAMESPACE "ControlRigRootCustomization"
 
@@ -621,9 +623,31 @@ const FRigControlElementCustomization* SControlRigEditModeTools::HandleGetContro
 void SControlRigEditModeTools::HandleActiveSpaceChanged(URigHierarchy* InHierarchy, const FRigElementKey& InControlKey,
 	const FRigElementKey& InSpaceKey)
 {
-	// todo: when the user picks a new space
-	// for now we'll just do it directly in the hierarchy
-	InHierarchy->SwitchToParent(InControlKey, InSpaceKey);
+
+	if (WeakSequencer.IsValid())
+	{
+		if (UControlRig* ControlRig = SequencerRig.Get())
+		{
+			if (const FRigControlElement* ControlElement = InHierarchy->Find<FRigControlElement>(InControlKey))
+			{
+				ISequencer* Sequencer = WeakSequencer.Pin().Get();
+				if (Sequencer)
+				{
+					FScopedTransaction Transaction(LOCTEXT("KeyControlRigSpace", "Key Control Rig Space"));
+					ControlRig->Modify();
+
+					FSpaceChannelAndSection SpaceChannelAndSection = FControlRigSpaceChannelHelpers::FindSpaceChannelAndSectionForControl(ControlRig, InControlKey.Name, Sequencer, true /*bCreateIfNeeded*/);
+					if (SpaceChannelAndSection.SpaceChannel)
+					{
+						const FFrameRate TickResolution = Sequencer->GetFocusedTickResolution();
+						const FFrameTime FrameTime = Sequencer->GetLocalTime().ConvertTo(TickResolution);
+						FFrameNumber CurrentTime = FrameTime.GetFrame();
+						FControlRigSpaceChannelHelpers::SequencerKeyControlRigSpaceChannel(ControlRig, Sequencer, SpaceChannelAndSection.SpaceChannel, SpaceChannelAndSection.SectionToKey, CurrentTime, InHierarchy, InControlKey, InSpaceKey);
+					}
+				}
+			}
+		}
+	}
 }
 
 void SControlRigEditModeTools::HandleSpaceListChanged(URigHierarchy* InHierarchy, const FRigElementKey& InControlKey,
@@ -631,16 +655,16 @@ void SControlRigEditModeTools::HandleSpaceListChanged(URigHierarchy* InHierarchy
 {
 	if (UControlRig* ControlRig = SequencerRig.Get())
 	{
-		if(const FRigControlElement* ControlElement = InHierarchy->Find<FRigControlElement>(InControlKey))
+		if (const FRigControlElement* ControlElement = InHierarchy->Find<FRigControlElement>(InControlKey))
 		{
-			FRigControlElementCustomization ControlCustomization = *ControlRig->GetControlCustomization(InControlKey);	
+			FRigControlElementCustomization ControlCustomization = *ControlRig->GetControlCustomization(InControlKey);
 			ControlCustomization.AvailableSpaces = InSpaceList;
 			ControlCustomization.RemovedSpaces.Reset();
 
 			// remember  the elements which are in the asset's available list but removed by the user
-			for(const FRigElementKey& AvailableSpace : ControlElement->Settings.Customization.AvailableSpaces)
+			for (const FRigElementKey& AvailableSpace : ControlElement->Settings.Customization.AvailableSpaces)
 			{
-				if(!ControlCustomization.AvailableSpaces.Contains(AvailableSpace))
+				if (!ControlCustomization.AvailableSpaces.Contains(AvailableSpace))
 				{
 					ControlCustomization.RemovedSpaces.Add(AvailableSpace);
 				}
@@ -654,33 +678,103 @@ void SControlRigEditModeTools::HandleSpaceListChanged(URigHierarchy* InHierarchy
 
 FReply SControlRigEditModeTools::OnBakeControlsToNewSpaceButtonClicked()
 {
-	if(SpacePickerWidget->GetHierarchy() == nullptr)
+	if (SpacePickerWidget->GetHierarchy() == nullptr)
+	{
+		return FReply::Unhandled();
+	}
+	if (SpacePickerWidget->GetControls().Num() == 0)
 	{
 		return FReply::Unhandled();
 	}
-	if(SpacePickerWidget->GetControls().Num() == 0)
+	if (!SequencerRig.Get())
 	{
 		return FReply::Unhandled();
 	}
+	ISequencer* Sequencer = WeakSequencer.Pin().Get();
+	if (Sequencer == nullptr || Sequencer->GetFocusedMovieSceneSequence() == nullptr || Sequencer->GetFocusedMovieSceneSequence()->GetMovieScene() == nullptr)
+	{
+		return FReply::Unhandled();
+	}
+	UControlRig* ControlRig = SequencerRig.Get();
 
 	FRigSpacePickerBakeSettings Settings;
-	// todoo pick the default target space... parent?
-	Settings.TargetSpace = URigHierarchy::GetDefaultParentSocketKey();
-	// todoooo determine the range
-
-	TSharedRef<SRigSpacePickerBakeWidget> BakeWidget =
-		SNew(SRigSpacePickerBakeWidget)
-		.Settings(Settings)
-		.Hierarchy(SpacePickerWidget->GetHierarchy())
-		.Controls(SpacePickerWidget->GetControls())
-		.GetControlCustomization(this, &SControlRigEditModeTools::HandleGetControlElementCustomization)
-		.OnBake_Lambda([](URigHierarchy* InHierarchy, TArray<FRigElementKey> InControls, FRigSpacePickerBakeSettings InSettings)
+	//Find default target space, just use first control and find space at current sequencer time
+	//Then Find range
+	TArray<FRigElementKey> ControlKeys = SpacePickerWidget->GetControls();
+
+	FSpaceChannelAndSection SpaceChannelAndSection = FControlRigSpaceChannelHelpers::FindSpaceChannelAndSectionForControl(ControlRig, ControlKeys[0].Name, Sequencer, true /*bCreateIfNeeded*/);
+	if (SpaceChannelAndSection.SpaceChannel != nullptr)
+	{
+		const FFrameRate TickResolution = Sequencer->GetFocusedTickResolution();
+		const FFrameTime FrameTime = Sequencer->GetLocalTime().ConvertTo(TickResolution);
+		FFrameNumber CurrentTime = FrameTime.GetFrame();
+		FMovieSceneControlRigSpaceBaseKey Value;
+		using namespace UE::MovieScene;
+		URigHierarchy* RigHierarchy = SpacePickerWidget->GetHierarchy();
+		EvaluateChannel(SpaceChannelAndSection.SpaceChannel, CurrentTime, Value);
+		switch (Value.SpaceType)
+		{
+		case EMovieSceneControlRigSpaceType::World:
+			Settings.TargetSpace = RigHierarchy->GetWorldSpaceSocketKey();
+			break;
+		case EMovieSceneControlRigSpaceType::Parent:
+			Settings.TargetSpace = RigHierarchy->GetFirstParent(ControlKeys[0]);
+			break;
+		case EMovieSceneControlRigSpaceType::ControlRig:
+			Settings.TargetSpace = Value.ControlRigElement;
+			break;
+		default:
+			Settings.TargetSpace = URigHierarchy::GetDefaultParentSocketKey();
+		}
+		TRange<FFrameNumber> Range = Sequencer->GetFocusedMovieSceneSequence()->GetMovieScene()->GetPlaybackRange();
+		TArray<FFrameNumber> Keys;
+		TArray < FKeyHandle> KeyHandles;
+
+		Settings.StartFrame = Range.GetLowerBoundValue();
+		Settings.EndFrame = Range.GetUpperBoundValue();
+		if (Keys.Num() > 0)
 		{
-			// todoo react to the bake
-			return FReply::Handled();
-		});
+			int32 Index = Algo::LowerBound(Keys, CurrentTime);
+			if (Index >= 0 && Index < (Keys.Num() - 1))
+			{
+				Settings.StartFrame = Keys[Index];
+				Settings.EndFrame = Keys[Index + 1];
+
+			}
+		}
+
+		TSharedRef<SRigSpacePickerBakeWidget> BakeWidget =
+			SNew(SRigSpacePickerBakeWidget)
+			.Settings(Settings)
+			.Hierarchy(SpacePickerWidget->GetHierarchy())
+			.Controls(SpacePickerWidget->GetControls())
+			.Sequencer(Sequencer)
+			.GetControlCustomization(this, &SControlRigEditModeTools::HandleGetControlElementCustomization)
+			.OnBake_Lambda([Sequencer, ControlRig,TickResolution](URigHierarchy* InHierarchy, TArray<FRigElementKey> InControls, FRigSpacePickerBakeSettings InSettings)
+				{
+					TArray<FFrameNumber> Frames;
+					
+					const FFrameRate& FrameRate = Sequencer->GetFocusedDisplayRate();
+					FFrameNumber FrameRateInFrameNumber = TickResolution.AsFrameNumber(FrameRate.AsInterval());
+					for (FFrameNumber& Frame = InSettings.StartFrame; Frame <= InSettings.EndFrame; Frame += FrameRateInFrameNumber)
+					{
+						Frames.Add(Frame);
+					}
+					FScopedTransaction Transaction(LOCTEXT("BakeControlToSpace", "Bake Control In Space"));
+					for (const FRigElementKey& ControlKey : InControls)
+					{
+						FSpaceChannelAndSection SpaceChannelAndSection = FControlRigSpaceChannelHelpers::FindSpaceChannelAndSectionForControl(ControlRig, ControlKey.Name, Sequencer, false /*bCreateIfNeeded*/);
+
+						FControlRigSpaceChannelHelpers::SequencerBakeControlInSpace(ControlRig, Sequencer, SpaceChannelAndSection.SpaceChannel, SpaceChannelAndSection.SectionToKey,
+							Frames, InHierarchy, ControlKey, InSettings);
+					}
+					return FReply::Handled();
+				});
+
+		return BakeWidget->OpenDialog(true);
+	}
+	return FReply::Unhandled();
 
-	return BakeWidget->OpenDialog(true);
 }
 
 EVisibility SControlRigEditModeTools::GetRigOptionExpanderVisibility() const
diff --git a/Engine/Plugins/Animation/ControlRig/Source/ControlRigEditor/Private/Editor/SRigSpacePickerWidget.cpp b/Engine/Plugins/Animation/ControlRig/Source/ControlRigEditor/Private/Editor/SRigSpacePickerWidget.cpp
--- a/Engine/Plugins/Animation/ControlRig/Source/ControlRigEditor/Private/Editor/SRigSpacePickerWidget.cpp
+++ b/Engine/Plugins/Animation/ControlRig/Source/ControlRigEditor/Private/Editor/SRigSpacePickerWidget.cpp
@@ -8,6 +8,9 @@
 #include "SRigHierarchyTreeView.h"
 #include "ControlRigEditorStyle.h"
 #include "PropertyCustomizationHelpers.h"
+#include "ISequencer.h"
+#include "Widgets/Input/NumericTypeInterface.h"
+#include "FrameNumberDetailsCustomization.h"
 
 #define LOCTEXT_NAMESPACE "SRigSpacePickerWidget"
 
@@ -971,9 +974,13 @@ void SRigSpacePickerBakeWidget::Construct(const FArguments& InArgs)
 {
 	check(InArgs._Hierarchy);
 	check(InArgs._Controls.Num() > 0);
+	check(InArgs._Sequencer);
 	check(InArgs._OnBake.IsBound());
 	
-	Settings = InArgs._Settings;
+	Settings = MakeShared<TStructOnScope<FRigSpacePickerBakeSettings>>();
+	Settings->InitializeAs<FRigSpacePickerBakeSettings>();
+	*Settings = InArgs._Settings;
+	Sequencer = InArgs._Sequencer;
 
 	FStructureDetailsViewArgs StructureViewArgs;
 	StructureViewArgs.bShowObjects = true;
@@ -986,9 +993,13 @@ void SRigSpacePickerBakeWidget::Construct(const FArguments& InArgs)
 	ViewArgs.bHideSelectionTip = false;
 	ViewArgs.bShowObjectLabel = false;
 
-	TSharedPtr<FStructOnScope> StructToDisplay = MakeShareable(new FStructOnScope(FRigSpacePickerBakeSettings::StaticStruct(), (uint8*)&Settings));
-	FPropertyEditorModule& EditModule = FModuleManager::Get().GetModuleChecked<FPropertyEditorModule>("PropertyEditor");
-	DetailsView = EditModule.CreateStructureDetailView(ViewArgs, StructureViewArgs, StructToDisplay, LOCTEXT("Struct", "Struct View"));
+	FPropertyEditorModule& PropertyEditor = FModuleManager::Get().LoadModuleChecked<FPropertyEditorModule>(TEXT("PropertyEditor"));
+
+	DetailsView = PropertyEditor.CreateStructureDetailView(ViewArgs, StructureViewArgs, TSharedPtr<FStructOnScope>());
+	TSharedPtr<INumericTypeInterface<double>> NumericTypeInterface = Sequencer->GetNumericTypeInterface();
+	DetailsView->GetDetailsView()->RegisterInstancedCustomPropertyTypeLayout("FrameNumber",
+		FOnGetPropertyTypeCustomizationInstance::CreateLambda([=]() {return MakeShared<FFrameNumberDetailsCustomization>(NumericTypeInterface); }));
+	DetailsView->SetStructureData(Settings);
 
 	ChildSlot
 	[
@@ -1021,13 +1032,13 @@ void SRigSpacePickerBakeWidget::Construct(const FArguments& InArgs)
 				})
 				.GetActiveSpace_Lambda([this](URigHierarchy*, const FRigElementKey&)
 				{
-					return Settings.TargetSpace;
+					return Settings->Get()->TargetSpace;
 				})
 				.OnActiveSpaceChanged_Lambda([this] (URigHierarchy*, const FRigElementKey&, const FRigElementKey InSpaceKey)
 				{
-					if(Settings.TargetSpace != InSpaceKey)
+					if(Settings->Get()->TargetSpace != InSpaceKey)
 					{
-						Settings.TargetSpace = InSpaceKey;	
+						Settings->Get()->TargetSpace = InSpaceKey;
 						SpacePickerWidget->RefreshContents();
 					}
 				})
@@ -1064,11 +1075,14 @@ void SRigSpacePickerBakeWidget::Construct(const FArguments& InArgs)
 					.Text(LOCTEXT("OK", "OK"))
 					.OnClicked_Lambda([this, InArgs]()
 					{
-						return InArgs._OnBake.Execute(SpacePickerWidget->GetHierarchy(), SpacePickerWidget->GetControls(), Settings);
+						FReply Reply =  InArgs._OnBake.Execute(SpacePickerWidget->GetHierarchy(), SpacePickerWidget->GetControls(),*(Settings->Get()));
+						CloseDialog();
+						return Reply;
+
 					})
 					.IsEnabled_Lambda([this]()
 					{
-						return Settings.TargetSpace.IsValid();
+						return Settings->Get()->TargetSpace.IsValid();
 					})
 				]
 
@@ -1100,7 +1114,7 @@ FReply SRigSpacePickerBakeWidget::OpenDialog(bool bModal)
 	const FVector2D CursorPos = FSlateApplication::Get().GetCursorPos();
 
 	TSharedRef<SRigSpaceDialogWindow> Window = SNew(SRigSpaceDialogWindow)
-	.Title( LOCTEXT("SRigSpacePickerBakeWidgetTitle", "Bake Controls to new space") )
+	.Title( LOCTEXT("SRigSpacePickerBakeWidgetTitle", "Bake Controls To Specified Space") )
 	.CreateTitleBar(true)
 	.Type(EWindowType::Normal)
 	.SizingRule( ESizingRule::Autosized )
diff --git a/Engine/Plugins/Animation/ControlRig/Source/ControlRigEditor/Private/Editor/SRigSpacePickerWidget.h b/Engine/Plugins/Animation/ControlRig/Source/ControlRigEditor/Private/Editor/SRigSpacePickerWidget.h
--- a/Engine/Plugins/Animation/ControlRig/Source/ControlRigEditor/Private/Editor/SRigSpacePickerWidget.h
+++ b/Engine/Plugins/Animation/ControlRig/Source/ControlRigEditor/Private/Editor/SRigSpacePickerWidget.h
@@ -9,6 +9,7 @@
 #include "Widgets/Layout/SBox.h"
 #include "Rigs/RigSpaceHierarchy.h"
 #include "IStructureDetailsView.h"
+#include "Misc/FrameNumber.h"
 #include "SRigSpacePickerWidget.generated.h"
 
 USTRUCT()
@@ -19,22 +20,26 @@ struct CONTROLRIGEDITOR_API FRigSpacePickerBakeSettings
 	FRigSpacePickerBakeSettings()
 	{
 		TargetSpace = FRigElementKey();
-		StartFrame = 1.f;
-		EndFrame = 100.f;
-		bKeyEveryFrame = true;
+		StartFrame = 0;
+		EndFrame = 100;
+		bReduceKeys = false;
+		Tolerance = 0.001f;
 	}
 
 	UPROPERTY()
 	FRigElementKey TargetSpace;
 
 	UPROPERTY(EditAnywhere, Category = "Settings")
-	float StartFrame;
+	FFrameNumber StartFrame;
+
+	UPROPERTY(EditAnywhere, Category = "Settings") // meta = (ShowOnlyInnerProperties)?
+	FFrameNumber EndFrame;
 
 	UPROPERTY(EditAnywhere, Category = "Settings")
-	float EndFrame;
+	bool bReduceKeys;
 
 	UPROPERTY(EditAnywhere, Category = "Settings")
-	bool bKeyEveryFrame;
+	float Tolerance;
 };
 
 DECLARE_DELEGATE_RetVal_TwoParams(FRigElementKey, FRigSpacePickerGetActiveSpace, URigHierarchy*, const FRigElementKey&);
@@ -177,6 +182,8 @@ private:
 	static FRigElementKey InValidKey;
 };
 
+class ISequencer;
+
 /** Widget allowing baking controls from one space to another */
 class CONTROLRIGEDITOR_API SRigSpacePickerBakeWidget : public SCompoundWidget
 {
@@ -185,9 +192,11 @@ public:
 	SLATE_BEGIN_ARGS(SRigSpacePickerBakeWidget)
 	: _Hierarchy(nullptr)
 	, _Controls()
+	, _Sequencer(nullptr)
 	{}
 	SLATE_ARGUMENT(URigHierarchy*, Hierarchy)
 	SLATE_ARGUMENT(TArray<FRigElementKey>, Controls)
+	SLATE_ARGUMENT(ISequencer*, Sequencer)
 	SLATE_EVENT(FRigSpacePickerGetControlCustomization, GetControlCustomization)
 	SLATE_ARGUMENT(FRigSpacePickerBakeSettings, Settings)
 	SLATE_EVENT(SRigSpacePickerOnBake, OnBake)
@@ -202,7 +211,10 @@ public:
 
 private:
 
-	FRigSpacePickerBakeSettings Settings;
+	//FRigSpacePickerBakeSettings Settings;
+	TSharedPtr<TStructOnScope<FRigSpacePickerBakeSettings>> Settings;
+
+	ISequencer* Sequencer;
 	FRigControlElementCustomization Customization;
 	
 	TWeakPtr<SWindow> DialogWindow;
diff --git a/Engine/Plugins/Animation/ControlRig/Source/ControlRigEditor/Private/Sequencer/ControlRigParameterTrackEditor.cpp b/Engine/Plugins/Animation/ControlRig/Source/ControlRigEditor/Private/Sequencer/ControlRigParameterTrackEditor.cpp
--- a/Engine/Plugins/Animation/ControlRig/Source/ControlRigEditor/Private/Sequencer/ControlRigParameterTrackEditor.cpp
+++ b/Engine/Plugins/Animation/ControlRig/Source/ControlRigEditor/Private/Sequencer/ControlRigParameterTrackEditor.cpp
@@ -2064,6 +2064,7 @@ void FControlRigParameterTrackEditor::GetControlRigKeys(UControlRig* InControlRi
 	int32 BoolChannelIndex = 0;
 	int32 EnumChannelIndex = 0;
 	int32 IntChannelIndex = 0;
+	int32 SpaceChannelIndex = 0;
 	for (int32 ControlIndex = 0; ControlIndex < Controls.Num(); ++ControlIndex)
 	{
 		FRigControlElement* ControlElement = Controls[ControlIndex];
@@ -2171,6 +2172,15 @@ void FControlRigParameterTrackEditor::GetControlRigKeys(UControlRig* InControlRi
 			{
 				bKeyZ = false;
 			}
+			if (FChannelMapInfo* pChannelIndex = SectionToKey->ControlChannelMap.Find(ControlElement->GetName()))
+			{
+				if (pChannelIndex->bDoesHaveSpace)
+				{
+					FMovieSceneControlRigSpaceBaseKey NewKey;
+					OutGeneratedKeys.Add(FMovieSceneChannelValueSetter::Create<FMovieSceneControlRigSpaceChannel>(SpaceChannelIndex++, NewKey, false));
+				}
+			}
+
 			OutGeneratedKeys.Add(FMovieSceneChannelValueSetter::Create<FMovieSceneFloatChannel>(ChannelIndex++, CurrentVector.X, bKeyX));
 			OutGeneratedKeys.Add(FMovieSceneChannelValueSetter::Create<FMovieSceneFloatChannel>(ChannelIndex++, CurrentVector.Y, bKeyY));
 			OutGeneratedKeys.Add(FMovieSceneChannelValueSetter::Create<FMovieSceneFloatChannel>(ChannelIndex++, CurrentVector.Z, bKeyZ));
@@ -2494,6 +2504,12 @@ bool FControlRigParameterTrackEditor::ModifyOurGeneratedKeysByCurrentAndWeight(U
 					if (pChannelIndex)
 					{
 						ChannelIndex = pChannelIndex->TotalChannelIndex;
+
+						if (pChannelIndex->bDoesHaveSpace)
+						{
+							++ChannelIndex;
+						}
+						
 						FVector3f CurrentPos = Val.Val.GetTranslation();
 						FRotator CurrentRot = Val.Val.GetRotation().Rotator();
 						GeneratedTotalKeys[ChannelIndex]->ModifyByCurrentAndWeight(Proxy, KeyTime, (void *)&CurrentPos.X, Weight);
diff --git a/Engine/Plugins/Animation/ControlRig/Source/ControlRigEditor/Private/Sequencer/ControlRigSpaceChannelEditors.cpp b/Engine/Plugins/Animation/ControlRig/Source/ControlRigEditor/Private/Sequencer/ControlRigSpaceChannelEditors.cpp
new file mode 100644
--- /dev/null
+++ b/Engine/Plugins/Animation/ControlRig/Source/ControlRigEditor/Private/Sequencer/ControlRigSpaceChannelEditors.cpp
@@ -0,0 +1,612 @@
+// Copyright Epic Games, Inc. All Rights Reserved.
+
+#include "ControlRigSpaceChannelEditors.h"
+#include "MovieSceneSequence.h"
+#include "MovieSceneSequenceEditor.h"
+#include "MovieSceneEventUtils.h"
+#include "Channels/MovieSceneChannelHandle.h"
+#include "Sections/MovieSceneEventSectionBase.h"
+#include "ISequencerChannelInterface.h"
+#include "Widgets/SNullWidget.h"
+#include "IKeyArea.h"
+#include "ISequencer.h"
+#include "SequencerSettings.h"
+#include "MovieSceneCommonHelpers.h"
+#include "GameFramework/Actor.h"
+#include "EditorStyleSet.h"
+#include "Styling/CoreStyle.h"
+#include "KeyDrawParams.h"
+#include "Framework/MultiBox/MultiBoxBuilder.h"
+#include "Channels/MovieSceneChannelProxy.h"
+#include "Channels/MovieSceneChannelEditorData.h"
+#include "Channels/DoubleChannelCurveModel.h"
+#include "Channels/FloatChannelCurveModel.h"
+#include "Channels/IntegerChannelCurveModel.h"
+#include "Channels/BoolChannelCurveModel.h"
+#include "PropertyCustomizationHelpers.h"
+#include "MovieSceneObjectBindingIDCustomization.h"
+#include "MovieSceneObjectBindingIDPicker.h"
+#include "LevelEditor.h"
+#include "Modules/ModuleManager.h"
+#include "Framework/Application/MenuStack.h"
+#include "Framework/Application/SlateApplication.h"
+#include "Editor/SceneOutliner/Private/SSocketChooser.h"
+#include "EntitySystem/MovieSceneDecompositionQuery.h"
+#include "EntitySystem/Interrogation/MovieSceneInterrogationLinker.h"
+#include "EntitySystem/Interrogation/MovieSceneInterrogatedPropertyInstantiator.h"
+#include "Systems/MovieScenePropertyInstantiator.h"
+#include "Tracks/MovieScenePropertyTrack.h"
+#include "MovieSceneSpawnableAnnotation.h"
+#include "ISequencerModule.h"
+#include "MovieSceneTracksComponentTypes.h"
+#include "SRigSpacePickerWidget.h"
+#include "Tools/ControlRigSnapper.h"
+#include "Sequencer/MovieSceneControlRigParameterTrack.h"
+#include "Sequencer/MovieSceneControlRigParameterSection.h"
+#include "MovieScene.h"
+#include "Sequencer/MovieSceneControlRigSpaceChannel.h"
+#include "Channels/MovieSceneChannelProxy.h"
+#include "MovieSceneObjectBindingID.h"
+#include "ControlRig.h"
+#include "IControlRigObjectBinding.h"
+
+#define LOCTEXT_NAMESPACE "ControlRigEditMode"
+
+
+static FKeyHandle SequencerOpenSpaceSwitchDialog(UControlRig* ControlRig, TArray<FRigElementKey> SelectedControls,ISequencer* Sequencer, FMovieSceneControlRigSpaceChannel* Channel, UMovieSceneSection* SectionToKey, FFrameNumber Time)
+{
+	FKeyHandle Handle = FKeyHandle::Invalid();
+	if (ControlRig == nullptr || Sequencer == nullptr)
+	{
+		return Handle;
+	}
+
+	TSharedRef<SRigSpacePickerWidget> PickerWidget =
+	SNew(SRigSpacePickerWidget)
+	.Hierarchy(ControlRig->GetHierarchy())
+	.Controls(SelectedControls)
+	.Title(LOCTEXT("PickSpace", "Pick Space"))
+	.AllowDelete(false)
+	.AllowReorder(false)
+	.AllowAdd(false)
+	.GetControlCustomization_Lambda([ControlRig](URigHierarchy*, const FRigElementKey& InControlKey)
+	{
+		return ControlRig->GetControlCustomization(InControlKey);
+	})
+	.OnActiveSpaceChanged_Lambda([&Handle,ControlRig,Sequencer,Channel,SectionToKey,Time,SelectedControls](URigHierarchy* RigHierarchy, const FRigElementKey& ControlKey, const FRigElementKey& SpaceKey)
+	{
+			
+		Handle = FControlRigSpaceChannelHelpers::SequencerKeyControlRigSpaceChannel(ControlRig, Sequencer, Channel, SectionToKey, Time, RigHierarchy, ControlKey, SpaceKey);
+
+	})
+	.OnSpaceListChanged_Lambda([SelectedControls, ControlRig](URigHierarchy* InHierarchy, const FRigElementKey& InControlKey, const TArray<FRigElementKey>& InSpaceList)
+	{
+		check(SelectedControls.Contains(InControlKey));
+				
+		// update the settings in the control element
+		if (FRigControlElement* ControlElement = InHierarchy->Find<FRigControlElement>(InControlKey))
+		{
+			FScopedTransaction Transaction(LOCTEXT("ControlChangeAvailableSpaces", "Edit Available Spaces"));
+
+			InHierarchy->Modify();
+
+			FRigControlElementCustomization ControlCustomization = *ControlRig->GetControlCustomization(InControlKey);
+			ControlCustomization.AvailableSpaces = InSpaceList;
+			ControlCustomization.RemovedSpaces.Reset();
+
+			// remember  the elements which are in the asset's available list but removed by the user
+			for (const FRigElementKey& AvailableSpace : ControlElement->Settings.Customization.AvailableSpaces)
+			{
+				if (!ControlCustomization.AvailableSpaces.Contains(AvailableSpace))
+				{
+					ControlCustomization.RemovedSpaces.Add(AvailableSpace);
+				}
+			}
+
+			ControlRig->SetControlCustomization(InControlKey, ControlCustomization);
+			InHierarchy->Notify(ERigHierarchyNotification::ControlSettingChanged, ControlElement);
+		}
+
+	});
+	// todo: implement GetAdditionalSpacesDelegate to pull spaces from sequencer
+
+	FReply Reply = PickerWidget->OpenDialog(true);
+	if (Reply.IsEventHandled())
+	{
+		return Handle;
+	}
+	return FKeyHandle::Invalid();
+}
+
+FKeyHandle AddOrUpdateKey(FMovieSceneControlRigSpaceChannel* Channel, UMovieSceneSection* SectionToKey, FFrameNumber Time, ISequencer& Sequencer, const FGuid& InObjectBindingID, FTrackInstancePropertyBindings* PropertyBindings)
+{
+	FKeyHandle Handle = FKeyHandle::Invalid();
+	if (UMovieSceneControlRigParameterSection* Section = Cast<UMovieSceneControlRigParameterSection>(SectionToKey))
+	{
+		if (UControlRig* ControlRig = Section->GetControlRig())
+		{
+			FName ControlName = Section->FindControlNameFromSpaceChannel(Channel);
+			if (ControlName != NAME_None)
+			{
+				if (FRigControlElement* Control = ControlRig->FindControl(ControlName))
+				{
+					TArray<FRigElementKey> Controls;
+					FRigElementKey ControlKey = Control->GetKey();
+					Controls.Add(ControlKey);
+					FMovieSceneControlRigSpaceBaseKey ExistingValue, Value;
+					using namespace UE::MovieScene;
+					EvaluateChannel(Channel, Time, ExistingValue);
+					Value = ExistingValue;
+					URigHierarchy* RigHierarchy = ControlRig->GetHierarchy();
+					Handle = SequencerOpenSpaceSwitchDialog(ControlRig, Controls, &Sequencer, Channel, SectionToKey, Time);
+				}
+			}
+		}
+	}
+	return Handle;
+}
+
+bool CanCreateKeyEditor(const FMovieSceneControlRigSpaceChannel* Channel)
+{
+	return false; //mz todoo maybe change
+}
+TSharedRef<SWidget> CreateKeyEditor(const TMovieSceneChannelHandle<FMovieSceneControlRigSpaceChannel>& Channel, UMovieSceneSection* Section, const FGuid& InObjectBindingID, TWeakPtr<FTrackInstancePropertyBindings> PropertyBindings, TWeakPtr<ISequencer> InSequencer)
+{
+	return SNullWidget::NullWidget;
+}
+
+/*******************************************************************
+*
+* FControlRigSpaceChannelHelpers
+*
+**********************************************************************/
+
+
+FSpaceChannelAndSection FControlRigSpaceChannelHelpers::FindSpaceChannelAndSectionForControl(UControlRig* ControlRig, FName ControlName, ISequencer* Sequencer, bool bCreateIfNeeded)
+{
+	FSpaceChannelAndSection SpaceChannelAndSection;
+	SpaceChannelAndSection.SpaceChannel = nullptr;
+	SpaceChannelAndSection.SectionToKey = nullptr;
+	if (ControlRig == nullptr || Sequencer == nullptr)
+	{
+		return SpaceChannelAndSection;
+	}
+	if (TSharedPtr<IControlRigObjectBinding> ObjectBinding = ControlRig->GetObjectBinding())
+	{
+		USceneComponent* Component = Cast<USceneComponent>(ObjectBinding->GetBoundObject());
+		if (!Component)
+		{
+			return SpaceChannelAndSection;
+		}
+		const bool bCreateHandleIfMissing = false;
+		FName CreatedFolderName = NAME_None;
+		FGuid ObjectHandle = Sequencer->GetHandleToObject(Component, bCreateHandleIfMissing);
+		if (!ObjectHandle.IsValid())
+		{
+			UObject* ActorObject = Component->GetOwner();
+			ObjectHandle = Sequencer->GetHandleToObject(ActorObject, bCreateHandleIfMissing);
+			if (!ObjectHandle.IsValid())
+			{
+				return SpaceChannelAndSection;
+			}
+		}
+		bool bCreateTrack = false;
+		UMovieScene* MovieScene = Sequencer->GetFocusedMovieSceneSequence()->GetMovieScene();
+		if (!MovieScene)
+		{
+			return SpaceChannelAndSection;
+		}
+		if (FMovieSceneBinding* Binding = MovieScene->FindBinding(ObjectHandle))
+		{
+			for (UMovieSceneTrack* Track : Binding->GetTracks())
+			{
+				if (UMovieSceneControlRigParameterTrack* ControlRigParameterTrack = Cast<UMovieSceneControlRigParameterTrack>(Track))
+				{
+					if (ControlRigParameterTrack->GetControlRig() == ControlRig)
+					{
+						UMovieSceneControlRigParameterSection* ActiveSection = Cast<UMovieSceneControlRigParameterSection>(ControlRigParameterTrack->GetSectionToKey());
+						if (ActiveSection)
+						{
+							ActiveSection->Modify();
+							ControlRig->Modify();
+							SpaceChannelAndSection.SectionToKey = ActiveSection;
+							FSpaceControlNameAndChannel* NameAndChannel = ActiveSection->GetSpaceChannel(ControlName);
+							if (NameAndChannel)
+							{
+								SpaceChannelAndSection.SpaceChannel = &NameAndChannel->SpaceCurve;
+							}
+							else if (bCreateIfNeeded)
+							{
+								ActiveSection->AddSpaceChannel(ControlName, true /*ReconstructChannelProxy*/);
+								NameAndChannel = ActiveSection->GetSpaceChannel(ControlName);
+								if (NameAndChannel)
+								{
+									SpaceChannelAndSection.SpaceChannel = &NameAndChannel->SpaceCurve;
+								}
+							}
+						}
+					}
+				}
+			}
+		}
+	}
+	return SpaceChannelAndSection;
+}
+
+FKeyHandle FControlRigSpaceChannelHelpers::SequencerKeyControlRigSpaceChannel(UControlRig* ControlRig, ISequencer* Sequencer, FMovieSceneControlRigSpaceChannel* Channel, UMovieSceneSection* SectionToKey, FFrameNumber Time, URigHierarchy* RigHierarchy, const FRigElementKey& ControlKey, const FRigElementKey& SpaceKey)
+{
+	FKeyHandle Handle = FKeyHandle::Invalid();
+	if (ControlRig == nullptr || Sequencer == nullptr || Sequencer->GetFocusedMovieSceneSequence() == nullptr || Sequencer->GetFocusedMovieSceneSequence()->GetMovieScene() == nullptr)
+	{
+		return Handle;
+	}
+
+	URigHierarchy* Hierarchy = ControlRig->GetHierarchy();
+	FMovieSceneControlRigSpaceBaseKey ExistingValue;
+	using namespace UE::MovieScene;
+	EvaluateChannel(Channel, Time, ExistingValue);
+	FMovieSceneControlRigSpaceBaseKey Value = ExistingValue;
+
+	if (SpaceKey == RigHierarchy->GetWorldSpaceSocketKey())
+	{
+		Value.SpaceType = EMovieSceneControlRigSpaceType::World;
+	}
+	else
+	{
+		FRigElementKey DefaultParent = RigHierarchy->GetFirstParent(ControlKey);
+		if (DefaultParent == SpaceKey)
+		{
+			Value.SpaceType = EMovieSceneControlRigSpaceType::Parent;
+
+		}
+		else if (SpaceKey.Type == ERigElementType::Control)
+		{
+			Value.SpaceType = EMovieSceneControlRigSpaceType::ControlRig;
+			Value.ControlRigElement = SpaceKey;
+		}
+	}
+
+	//we only key if the value is different.
+	if (Value != ExistingValue)
+	{
+		TArray<FFrameNumber> Frames;
+		Frames.Add(Time);
+		
+		TMovieSceneChannelData<FMovieSceneControlRigSpaceBaseKey> ChannelInterface = Channel->GetData();
+		bool bSetPreviousKey = true;
+		//if we have no keys need to set key for current space at start frame, unless setting key at start time, where then don't do previous compensation
+		if (Channel->GetNumKeys() == 0 && Sequencer->GetFocusedMovieSceneSequence() && Sequencer->GetFocusedMovieSceneSequence()->GetMovieScene())
+		{
+			FFrameNumber StartFrame = Sequencer->GetFocusedMovieSceneSequence()->GetMovieScene()->GetPlaybackRange().GetLowerBoundValue();
+			if (StartFrame != Time)
+			{
+				//find all of the times in the space after this time we now need to compensate for
+				TSortedMap<FFrameNumber,FFrameNumber> ExtraFrames;
+				FControlRigSpaceChannelHelpers::GetFramesInThisSpaceAfterThisTime(ControlRig, ControlKey.Name, ExistingValue,
+					Channel, SectionToKey, Time, ExtraFrames);
+				if (ExtraFrames.Num() > 0)
+				{
+					for (const TPair<FFrameNumber,FFrameNumber>& Frame : ExtraFrames)
+					{
+						Frames.Add(Frame.Value);
+					}
+				}
+
+				FMovieSceneControlRigSpaceBaseKey Original = ExistingValue;
+				ChannelInterface.AddKey(StartFrame, Forward<FMovieSceneControlRigSpaceBaseKey>(Original));
+
+			}
+			else
+			{
+				bSetPreviousKey = false;
+			}
+		}
+
+		TArray<FTransform> ControlRigParentWorldTransforms;
+		ControlRigParentWorldTransforms.SetNum(Frames.Num());
+		for (FTransform& WorldTransform : ControlRigParentWorldTransforms)
+		{
+			WorldTransform  = FTransform::Identity;
+		}
+		TArray<FTransform> ControlWorldTransforms;
+		FControlRigSnapper Snapper;
+		Snapper.GetControlRigControlTransforms(Sequencer, ControlRig, ControlKey.Name, Frames, ControlRigParentWorldTransforms, ControlWorldTransforms);
+
+		int32 ExistingIndex = ChannelInterface.FindKey(Time);
+		if (ExistingIndex != INDEX_NONE)
+		{
+			Handle = ChannelInterface.GetHandle(ExistingIndex);
+			using namespace UE::MovieScene;
+			AssignValue(Channel, Handle, Forward<FMovieSceneControlRigSpaceBaseKey>(Value));
+		}
+		else
+		{
+			ExistingIndex = ChannelInterface.AddKey(Time, Forward<FMovieSceneControlRigSpaceBaseKey>(Value));
+			Handle = ChannelInterface.GetHandle(ExistingIndex);
+		}
+
+		FRigControlModifiedContext Context;
+		Context.SetKey = EControlRigSetKey::Always;
+		FFrameRate TickResolution = Sequencer->GetFocusedTickResolution();
+
+		if (bSetPreviousKey)
+		{
+			Context.LocalTime = TickResolution.AsSeconds(FFrameTime(Time - 1));
+			ControlRig->SetControlGlobalTransform(ControlKey.Name, ControlWorldTransforms[0], true, Context);
+		}
+		ControlRig->GetHierarchy()->SwitchToParent(ControlKey, SpaceKey);
+		int32 FramesIndex = 0;
+		for (const FFrameNumber& Frame : Frames)
+		{
+			ControlRig->Evaluate_AnyThread();
+			Context.LocalTime = TickResolution.AsSeconds(FFrameTime(Frame));
+			ControlRig->SetControlGlobalTransform(ControlKey.Name, ControlWorldTransforms[FramesIndex++], true, Context);
+		}
+		Sequencer->NotifyMovieSceneDataChanged(EMovieSceneDataChangeType::MovieSceneStructureItemAdded);
+
+	}
+	return Handle;
+
+}
+
+void FControlRigSpaceChannelHelpers::GetFramesInThisSpaceAfterThisTime(UControlRig* ControlRig, FName ControlName, FMovieSceneControlRigSpaceBaseKey CurrentValue,
+	FMovieSceneControlRigSpaceChannel* Channel, UMovieSceneSection* SectionToKey,
+	FFrameNumber Time, TSortedMap<FFrameNumber,FFrameNumber>& OutMoreFrames)
+{
+	OutMoreFrames.Reset();
+	if (ControlRig && Channel && SectionToKey)
+	{
+		if (UMovieSceneControlRigParameterSection* Section = Cast<UMovieSceneControlRigParameterSection>(SectionToKey))
+		{
+			TArrayView<FMovieSceneFloatChannel*> FloatChannels = Section->GetChannelProxy().GetChannels<FMovieSceneFloatChannel>();
+			FChannelMapInfo* pChannelIndex = Section->ControlChannelMap.Find(ControlName);
+			if (pChannelIndex != nullptr)
+			{
+				int32 ChannelIndex = pChannelIndex->ChannelIndex;
+
+				if (FRigControlElement* ControlElement = ControlRig->FindControl(ControlName))
+				{
+					FMovieSceneControlRigSpaceBaseKey Value;
+					switch (ControlElement->Settings.ControlType)
+					{
+						case ERigControlType::Position:
+						case ERigControlType::Scale:
+						case ERigControlType::Rotator:
+						{
+							const int32 NumChannels = 3;
+							for (int32 Index = 0; Index < NumChannels; ++Index)
+							{
+								for (FFrameNumber Frame : FloatChannels[ChannelIndex++]->GetData().GetTimes())
+								{
+									if (Frame > Time)
+									{
+										EvaluateChannel(Channel, Frame, Value);
+										if (CurrentValue == Value)
+										{
+											if (OutMoreFrames.Find(Frame) == nullptr)
+											{
+												OutMoreFrames.Add(Frame, Frame);
+											}
+										}
+										else
+										{
+											break;
+										}
+									}
+								}
+							}
+							break;
+						}
+
+						case ERigControlType::Transform:
+						case ERigControlType::TransformNoScale:
+						case ERigControlType::EulerTransform:
+						{
+
+							const int32 NumChannels = (ControlElement->Settings.ControlType == ERigControlType::Transform ||
+								ControlElement->Settings.ControlType == ERigControlType::EulerTransform) ? 9 : 6;
+							for (int32 Index = 0; Index < NumChannels; ++Index)
+							{
+								for (FFrameNumber Frame : FloatChannels[ChannelIndex++]->GetData().GetTimes())
+								{
+									if (Frame > Time)
+									{
+										EvaluateChannel(Channel, Frame, Value);
+										if (CurrentValue == Value)
+										{
+											if (OutMoreFrames.Find(Frame) == nullptr)
+											{
+												OutMoreFrames.Add(Frame, Frame);
+											}
+										}
+										else
+										{
+											break;
+										}
+									}
+								}
+							}
+							break;
+						}
+					}
+				}
+			}
+		}
+	}
+}
+
+
+void FControlRigSpaceChannelHelpers::SequencerBakeControlInSpace(UControlRig* ControlRig, ISequencer* Sequencer, FMovieSceneControlRigSpaceChannel* Channel, UMovieSceneSection* SectionToKey,
+	TArray<FFrameNumber> Frames, URigHierarchy* RigHierarchy, const FRigElementKey& ControlKey, FRigSpacePickerBakeSettings Settings)
+{
+	if (ControlRig && Sequencer && Channel && SectionToKey && Frames.Num() > 0 && RigHierarchy
+		&& Sequencer->GetFocusedMovieSceneSequence() && Sequencer->GetFocusedMovieSceneSequence()->GetMovieScene())
+	{
+		TArray<FTransform> ControlRigParentWorldTransforms;
+		ControlRigParentWorldTransforms.SetNum(Frames.Num());
+		for (FTransform& Transform : ControlRigParentWorldTransforms)
+		{
+			Transform = FTransform::Identity;
+		}
+		//Store transforms
+		FControlRigSnapper Snapper;
+		TArray<FTransform> ControlWorldTransforms;
+		Snapper.GetControlRigControlTransforms(Sequencer, ControlRig, ControlKey.Name, Frames, ControlRigParentWorldTransforms, ControlWorldTransforms);
+
+		//Find all space keys in range and delete them since it will get replaced with new space when we components.
+		FFrameNumber StartFrame = Frames[0];
+		FFrameNumber EndFrame = Frames[Frames.Num() - 1];
+		TArray<FFrameNumber> Keys;
+		TArray < FKeyHandle> KeyHandles;
+		TRange<FFrameNumber> Range(StartFrame, EndFrame);
+		Channel->GetKeys(Range, &Keys, &KeyHandles);
+		Channel->DeleteKeys(KeyHandles);
+		URigHierarchy* Hierarchy = ControlRig->GetHierarchy();
+
+		//now find space at start and end see if different than the new space if so we need to compensate
+		FMovieSceneControlRigSpaceBaseKey StartFrameValue, EndFrameValue;
+		using namespace UE::MovieScene;
+		EvaluateChannel(Channel, StartFrame, StartFrameValue);
+		EvaluateChannel(Channel, EndFrame, EndFrameValue);
+
+		FMovieSceneControlRigSpaceBaseKey Value;
+		if (Settings.TargetSpace == RigHierarchy->GetWorldSpaceSocketKey())
+		{
+			Value.SpaceType = EMovieSceneControlRigSpaceType::World;
+		}
+		else
+		{
+			FRigElementKey DefaultParent = RigHierarchy->GetFirstParent(ControlKey);
+			if (DefaultParent == Settings.TargetSpace)
+			{
+				Value.SpaceType = EMovieSceneControlRigSpaceType::Parent;
+			}
+			else
+			{
+				Value.SpaceType = EMovieSceneControlRigSpaceType::ControlRig;
+				Value.ControlRigElement = Settings.TargetSpace;
+			}
+		}
+		const bool bCompensateStart = StartFrameValue != Value;
+		TRange<FFrameNumber> PlaybackRange = Sequencer->GetFocusedMovieSceneSequence()->GetMovieScene()->GetPlaybackRange();
+		const bool bCompensateEnd = (EndFrameValue != Value && PlaybackRange.GetUpperBoundValue() != EndFrame);
+
+		//if compensate at the start we need to set the channel key as the new value
+		if (bCompensateStart)
+		{
+			TMovieSceneChannelData<FMovieSceneControlRigSpaceBaseKey> ChannelInterface = Channel->GetData();
+			ChannelInterface.AddKey(StartFrame, Forward<FMovieSceneControlRigSpaceBaseKey>(Value));
+
+		}
+		//if we compensate at the end we change the last frame to frame -1(tick), and then later set the space to the other one and 
+		if (bCompensateEnd)
+		{
+			Frames[Frames.Num() - 1] = Frames[Frames.Num() - 1] - 1;
+		}
+		//now set all of the key values
+		FRigControlModifiedContext Context;
+		Context.SetKey = EControlRigSetKey::Always;
+		FFrameRate TickResolution = Sequencer->GetFocusedTickResolution();
+
+		RigHierarchy->SwitchToParent(ControlKey, Settings.TargetSpace);
+		ControlRig->Evaluate_AnyThread();
+
+		for (int32 Index = 0; Index < Frames.Num(); ++Index)
+		{
+			const FTransform GlobalTransform = ControlWorldTransforms[Index];
+			const FFrameNumber Frame = Frames[Index];
+			Context.LocalTime = TickResolution.AsSeconds(FFrameTime(Frame));
+			ControlRig->SetControlGlobalTransform(ControlKey.Name, GlobalTransform, true, Context);
+		}
+
+		//if end compensated set the space that was active previously and set the compensated global value
+		if (bCompensateEnd)
+		{
+			//EndFrameValue to SpaceKey todoo move to function
+			switch (EndFrameValue.SpaceType)
+			{
+			case EMovieSceneControlRigSpaceType::Parent:
+				RigHierarchy->SwitchToDefaultParent(ControlKey);
+				break;
+			case EMovieSceneControlRigSpaceType::World:
+				RigHierarchy->SwitchToWorldSpace(ControlKey);
+				break;
+			case EMovieSceneControlRigSpaceType::ControlRig:
+				RigHierarchy->SwitchToParent(ControlKey, EndFrameValue.ControlRigElement);
+				break;
+			}
+			ControlRig->Evaluate_AnyThread();
+
+			TMovieSceneChannelData<FMovieSceneControlRigSpaceBaseKey> ChannelInterface = Channel->GetData();
+			ChannelInterface.AddKey(EndFrame, Forward<FMovieSceneControlRigSpaceBaseKey>(EndFrameValue));
+			const FTransform GlobalTransform = ControlWorldTransforms[Frames.Num() - 1];
+			Context.LocalTime = TickResolution.AsSeconds(FFrameTime(EndFrame));
+			ControlRig->SetControlGlobalTransform(ControlKey.Name, GlobalTransform, true, Context);
+		}
+		if (Settings.bReduceKeys)
+		{
+			if (UMovieSceneControlRigParameterSection* Section = Cast<UMovieSceneControlRigParameterSection>(SectionToKey))
+			{
+				TArrayView<FMovieSceneFloatChannel*> FloatChannels = Section->GetChannelProxy().GetChannels<FMovieSceneFloatChannel>();
+				FChannelMapInfo* pChannelIndex = Section->ControlChannelMap.Find(ControlKey.Name);
+				if (pChannelIndex != nullptr)
+				{
+					int32 ChannelIndex = pChannelIndex->ChannelIndex;
+
+					if (FRigControlElement* ControlElement = ControlRig->FindControl(ControlKey.Name))
+					{
+						FKeyDataOptimizationParams Params;
+						Params.bAutoSetInterpolation = true;
+						Params.Tolerance = Settings.Tolerance;
+						Params.Range = TRange <FFrameNumber>(Frames[0], Frames[Frames.Num() - 1]);
+
+						switch (ControlElement->Settings.ControlType)
+						{
+						case ERigControlType::Position:
+						case ERigControlType::Scale:
+						case ERigControlType::Rotator:
+						{
+							FloatChannels[ChannelIndex]->Optimize(Params);
+							FloatChannels[ChannelIndex + 1]->Optimize(Params);
+							FloatChannels[ChannelIndex + 2]->Optimize(Params);
+							break;
+						}
+
+						case ERigControlType::Transform:
+						case ERigControlType::TransformNoScale:
+						case ERigControlType::EulerTransform:
+						{
+
+							FloatChannels[ChannelIndex]->Optimize(Params);
+							FloatChannels[ChannelIndex + 1]->Optimize(Params);
+							FloatChannels[ChannelIndex + 2]->Optimize(Params);
+							FloatChannels[ChannelIndex + 3]->Optimize(Params);
+							FloatChannels[ChannelIndex + 4]->Optimize(Params);
+							FloatChannels[ChannelIndex + 5]->Optimize(Params);
+
+							if (ControlElement->Settings.ControlType == ERigControlType::Transform ||
+								ControlElement->Settings.ControlType == ERigControlType::EulerTransform)
+							{
+								FloatChannels[ChannelIndex + 6]->Optimize(Params);
+								FloatChannels[ChannelIndex + 7]->Optimize(Params);
+								FloatChannels[ChannelIndex + 8]->Optimize(Params);
+							}
+							break;
+
+						}
+						}
+					}
+				}
+			}
+		}
+		Sequencer->NotifyMovieSceneDataChanged(EMovieSceneDataChangeType::MovieSceneStructureItemAdded); //may have added channel
+	}
+}
+
+struct FSpaceRanges
+{
+	TRange<FFrameNumber> Range;
+	FMovieSceneControlRigSpaceBaseKey Key;
+};
+
+#undef LOCTEXT_NAMESPACE
diff --git a/Engine/Plugins/Animation/ControlRig/Source/ControlRigEditor/Public/ControlRigSpaceChannelEditors.h b/Engine/Plugins/Animation/ControlRig/Source/ControlRigEditor/Public/ControlRigSpaceChannelEditors.h
new file mode 100644
--- /dev/null
+++ b/Engine/Plugins/Animation/ControlRig/Source/ControlRigEditor/Public/ControlRigSpaceChannelEditors.h
@@ -0,0 +1,77 @@
+// Copyright Epic Games, Inc. All Rights Reserved.
+
+#pragma once
+
+#include "CoreTypes.h"
+#include "Templates/SharedPointer.h"
+#include "Containers/Array.h"
+#include "Containers/ContainersFwd.h"
+#include "UObject/ObjectMacros.h"
+#include "UObject/Class.h"
+
+#include "MovieSceneKeyStruct.h"
+#include "SequencerChannelTraits.h"
+#include "Channels/MovieSceneChannelHandle.h"
+
+#include "Sequencer/MovieSceneControlRigSpaceChannel.h"
+#include "Sequencer/MovieSceneControlRigParameterSection.h"
+#include "SRigSpacePickerWidget.h"
+#include "Containers/SortedMap.h"
+
+struct FKeyHandle;
+struct FKeyDrawParams;
+class UControlrig;
+class ISequencer;
+class UMovieSceneSection;
+class URigHierarchy;
+struct FRigElementKey;
+class SWidget;
+class ISequencer;
+class FMenuBuilder;
+
+
+struct FSpaceChannelAndSection
+{
+	FSpaceChannelAndSection() : SectionToKey(nullptr), SpaceChannel(nullptr) {};
+	UMovieSceneSection* SectionToKey;
+	FMovieSceneControlRigSpaceChannel* SpaceChannel;
+};
+
+/*
+* Class that contains helper functions for various space switching activities
+*/
+struct FControlRigSpaceChannelHelpers
+{
+	static FKeyHandle SequencerKeyControlRigSpaceChannel(UControlRig* ControlRig, ISequencer* Sequencer, FMovieSceneControlRigSpaceChannel* Channel, UMovieSceneSection* SectionToKey, FFrameNumber Time, URigHierarchy* RigHierarchy, const FRigElementKey& ControlKey, const FRigElementKey& SpaceKey);
+	static FSpaceChannelAndSection FindSpaceChannelAndSectionForControl(UControlRig* ControlRig, FName ControlName, ISequencer* Sequencer, bool bCreateIfNeeded);
+	static void SequencerBakeControlInSpace(UControlRig* ControlRig, ISequencer* Sequencer, FMovieSceneControlRigSpaceChannel* Channel, UMovieSceneSection* SectionToKey,
+		TArray<FFrameNumber> Frames, URigHierarchy* RigHierarchy, const FRigElementKey& ControlKey, FRigSpacePickerBakeSettings InSettings);
+	static void GetFramesInThisSpaceAfterThisTime(UControlRig* ControlRig, FName ControlName, FMovieSceneControlRigSpaceBaseKey CurrentValue,
+		FMovieSceneControlRigSpaceChannel* Channel, UMovieSceneSection* SectionToKey,
+		FFrameNumber Time, TSortedMap<FFrameNumber,FFrameNumber>& OutMoreFrames);
+};
+
+//template specialization
+FKeyHandle AddOrUpdateKey(FMovieSceneControlRigSpaceChannel* Channel, UMovieSceneSection* SectionToKey, FFrameNumber Time, ISequencer& Sequencer, const FGuid& ObjectBindingID, FTrackInstancePropertyBindings* PropertyBindings);
+
+/** Key editor overrides */
+bool CanCreateKeyEditor(const FMovieSceneControlRigSpaceChannel*       Channel);
+TSharedRef<SWidget> CreateKeyEditor(const TMovieSceneChannelHandle<FMovieSceneControlRigSpaceChannel>&        Channel, UMovieSceneSection* Section, const FGuid& InObjectBindingID, TWeakPtr<FTrackInstancePropertyBindings> PropertyBindings, TWeakPtr<ISequencer> Sequencer);
+
+
+//UMovieSceneKeyStructType* IstanceGeneratedStruct(FMovieSceneControlRigSpaceChannel* Channel, FSequencerKeyStructGenerator* Generator);
+
+//void PostConstructKeystance(const TMovieSceneChannelHandle<FMovieSceneControlRigSpaceChannel>& ChannelHandle, FKeyHandle InHandle, FStructOnScope* Struct);
+
+/** Key drawing overrides */
+//void DrawKeys(FMovieSceneControlRigSpaceChannel*    Channel, TArrayView<const FKeyHandle> InKeyHandles, const UMovieSceneSection* InOwner, TArrayView<FKeyDrawParams> OutKeyDrawParams);
+
+/** Context menu overrides */
+//void ExtendSectionMenu(FMenuBuilder& OuterMenuBuilder, TArray<TMovieSceneChannelHandle<FMovieSceneControlRigSpaceChannel>>&& Channels, TArrayView<UMovieSceneSection* const> Sections, TWeakPtr<ISequencer> InSequencer);
+//void ExtendKeyMenu(FMenuBuilder& OuterMenuBuilder, TArray<TExtendKeyMenuParams<FMovieSceneControlRigSpaceChannel>>&& Channels, TWeakPtr<ISequencer> InSequencer);
+
+/** Curve editor models */
+//inline bool SupportsCurveEditorModels(const TMovieSceneChannelHandle<FMovieSceneControlRigSpaceChannel>& DoubleChannel) { return true; }
+//TUniquePtr<FCurveModel> CreateCurveEditorModel(const TMovieSceneChannelHandle<FMovieSceneControlRigSpaceChannel>& DoubleChannel, UMovieSceneSection* OwningSection, TSharedRef<ISequencer> InSequencer);
+
+
diff --git a/Engine/Source/Editor/MovieSceneTools/Private/Channels/SCurveEditorKeyBarView.cpp b/Engine/Source/Editor/MovieSceneTools/Private/Channels/SCurveEditorKeyBarView.cpp
new file mode 100644
--- /dev/null
+++ b/Engine/Source/Editor/MovieSceneTools/Private/Channels/SCurveEditorKeyBarView.cpp
@@ -0,0 +1,148 @@
+// Copyright Epic Games, Inc. All Rights Reserved.
+
+#include "Channels/SCurveEditorKeyBarView.h"
+#include "Rendering/DrawElements.h"
+#include "Fonts/FontMeasure.h"
+#include "Styling/CoreStyle.h"
+#include "CurveEditor.h"
+#include "EditorStyleSet.h"
+#include "CurveModel.h"
+#include "ICurveEditorBounds.h"
+#include "CurveEditorScreenSpace.h"
+#include "CurveDataAbstraction.h"
+#include "CurveDrawInfo.h"
+
+float SCurveEditorKeyBarView::TrackHeight = 24.f;
+
+void SCurveEditorKeyBarView::Construct(const FArguments& InArgs, TWeakPtr<FCurveEditor> InCurveEditor)
+{
+	bFixedOutputBounds = true;
+	OutputMin = -0.5;
+	OutputMax =  0.5;
+	WeakCurveEditor = InCurveEditor;
+	SortBias = 25;
+
+	SInteractiveCurveEditorView::Construct(InArgs, InCurveEditor);
+}
+
+FVector2D SCurveEditorKeyBarView::ComputeDesiredSize(float LayoutScaleMultiplier) const
+{
+	return FVector2D(100.f, TrackHeight * (CurveInfoByID.Num()));
+}
+
+void SCurveEditorKeyBarView::GetGridLinesY(TSharedRef<const FCurveEditor> CurveEditor, TArray<float>& MajorGridLines, TArray<float>& MinorGridLines, TArray<FText>* MajorGridLabels) const
+{}
+
+void SCurveEditorKeyBarView::Tick(const FGeometry& AllottedGeometry, const double InCurrentTime, const float InDeltaTime)
+{
+	TSharedPtr<FCurveEditor> CurveEditor = WeakCurveEditor.Pin();
+	if (!CurveEditor)
+	{
+		return;
+	}
+
+	double Count = 0.0;
+	for (auto It = CurveInfoByID.CreateIterator(); It; ++It)
+	{
+		FCurveModel* Curve = CurveEditor->FindCurve(It.Key());
+		if (ensureAlways(Curve))
+		{
+			It->Value.ViewToCurveTransform = FTransform2D(FVector2D(0.f, Count));
+		}
+
+		Count += 1.0;
+	}
+
+	OutputMin = OutputMax - FMath::Max(Count, 1e-10);
+	SInteractiveCurveEditorView::Tick(AllottedGeometry, InCurrentTime, InDeltaTime);
+}
+
+void SCurveEditorKeyBarView::PaintView(const FPaintArgs& Args, const FGeometry& AllottedGeometry, const FSlateRect& MyCullingRect, FSlateWindowElementList& OutDrawElements, int32 BaseLayerId, const FWidgetStyle& InWidgetStyle, bool bParentEnabled) const
+{
+	TSharedPtr<FCurveEditor> CurveEditor = WeakCurveEditor.Pin();
+	if (CurveEditor)
+	{
+		const ESlateDrawEffect DrawEffects = ShouldBeEnabled(bParentEnabled) ? ESlateDrawEffect::None : ESlateDrawEffect::DisabledEffect;
+
+		DrawBackground(AllottedGeometry, OutDrawElements, BaseLayerId, DrawEffects);
+		DrawLabels(AllottedGeometry, OutDrawElements, BaseLayerId, DrawEffects);
+		DrawGridLines(CurveEditor.ToSharedRef(), AllottedGeometry, OutDrawElements, BaseLayerId, DrawEffects);
+		DrawCurves(CurveEditor.ToSharedRef(), AllottedGeometry, MyCullingRect, OutDrawElements, BaseLayerId, InWidgetStyle, DrawEffects);
+	}
+}
+
+void SCurveEditorKeyBarView::DrawLabels(const FGeometry& AllottedGeometry, FSlateWindowElementList& OutDrawElements, int32 BaseLayerId, ESlateDrawEffect DrawEffects) const
+{
+	TSharedPtr<FCurveEditor> CurveEditor = WeakCurveEditor.Pin();
+	if (!CurveEditor)
+	{
+		return;
+	}
+
+	const FSlateBrush* WhiteBrush = FEditorStyle::GetBrush("WhiteBrush");
+
+	// Draw some text telling the user how to get retime anchors.
+	const FSlateFontInfo FontInfo = FCoreStyle::Get().GetFontStyle("ToolTip.LargerFont");
+
+	// We have to measure the string so we can draw it centered on the window.
+	const TSharedRef<FSlateFontMeasure> FontMeasure = FSlateApplication::Get().GetRenderer()->GetFontMeasureService();
+
+	const FVector2D LocalSize = AllottedGeometry.GetLocalSize();
+
+	const FCurveEditorScreenSpaceV ViewSpace = GetViewSpace();
+
+	for (auto It = CurveInfoByID.CreateConstIterator(); It; ++It)
+	{
+		FCurveModel* Curve = CurveEditor->FindCurve(It.Key());
+		if (!ensureAlways(Curve))
+		{
+			continue;
+		}
+
+		const int32 CurveIndex = static_cast<int32>(It->Value.ViewToCurveTransform.GetTranslation().Y);
+
+		const FCurveEditorScreenSpaceV CurveSpace = ViewSpace.ToCurveSpace(It->Value.ViewToCurveTransform);
+		const float LaneTop = CurveSpace.ValueToScreen(0.0) - TrackHeight*.5f;
+
+		// Draw the curve color as the background. Event curves set their track color as the curve color.
+		FLinearColor CurveColor = Curve->GetColor();
+
+		// Alpha blend the zebra tint
+		if (CurveIndex%2)
+		{
+			static FLinearColor ZebraTint = FLinearColor::White.CopyWithNewOpacity(0.01f);
+			if (CurveColor == FLinearColor::White)
+			{
+				CurveColor = ZebraTint;
+			}
+			else
+			{
+				CurveColor = CurveColor * (1.f - ZebraTint.A) + ZebraTint * ZebraTint.A;
+			}
+		}
+
+		if (CurveColor != FLinearColor::White)
+		{
+			const FPaintGeometry BoxGeometry = AllottedGeometry.ToPaintGeometry(
+				FVector2D(AllottedGeometry.GetLocalSize().X, TrackHeight),
+				FSlateLayoutTransform(FVector2D(0.f, LaneTop))
+			);
+
+			FSlateDrawElement::MakeBox(OutDrawElements, BaseLayerId + CurveViewConstants::ELayerOffset::Background, BoxGeometry, WhiteBrush, DrawEffects, CurveColor);
+		}
+
+		// Draw the curve label
+		{
+			const FText Label = Curve->GetLongDisplayName();
+
+			const FVector2D TextSize = FontMeasure->Measure(Label, FontInfo);
+			const FVector2D Position(LocalSize.X - TextSize.X - CurveViewConstants::CurveLabelOffsetX, LaneTop + (TrackHeight - TextSize.Y)*.5f);
+
+			const FPaintGeometry LabelGeometry = AllottedGeometry.ToPaintGeometry(
+				FSlateLayoutTransform(Position)
+			);
+
+			FSlateDrawElement::MakeText(OutDrawElements, BaseLayerId + CurveViewConstants::ELayerOffset::Labels, LabelGeometry, Label, FontInfo, DrawEffects, FLinearColor::White);
+		}
+	}
+}
\ No newline at end of file
diff --git a/Engine/Source/Editor/MovieSceneTools/Public/Channels/SCurveEditorKeyBarView.h b/Engine/Source/Editor/MovieSceneTools/Public/Channels/SCurveEditorKeyBarView.h
new file mode 100644
--- /dev/null
+++ b/Engine/Source/Editor/MovieSceneTools/Public/Channels/SCurveEditorKeyBarView.h
@@ -0,0 +1,28 @@
+// Copyright Epic Games, Inc. All Rights Reserved.
+
+#pragma once
+
+#include "Views/SInteractiveCurveEditorView.h"
+
+class MOVIESCENETOOLS_API SCurveEditorKeyBarView : public SInteractiveCurveEditorView
+{
+public:
+
+	void Construct(const FArguments& InArgs, TWeakPtr<FCurveEditor> InCurveEditor);
+
+	virtual void GetGridLinesY(TSharedRef<const FCurveEditor> CurveEditor, TArray<float>& MajorGridLines, TArray<float>& MinorGridLines, TArray<FText>* MajorGridLabels) const override;
+
+private:
+
+	//~ SWidget Interface
+	virtual void Tick(const FGeometry& AllottedGeometry, const double InCurrentTime, const float InDeltaTime) override;
+	virtual FVector2D ComputeDesiredSize(float LayoutScaleMultiplier) const override;
+
+	//~ SInteractiveCurveEditorView interface
+	virtual void PaintView(const FPaintArgs& Args, const FGeometry& AllottedGeometry, const FSlateRect& MyCullingRect, FSlateWindowElementList& OutDrawElements, int32 BaseLayerId, const FWidgetStyle& InWidgetStyle, bool bParentEnabled) const override;
+	virtual bool IsValueSnapEnabled() const override { return false; }
+
+	void DrawLabels(const FGeometry& AllottedGeometry, FSlateWindowElementList& OutDrawElements, int32 BaseLayerId, ESlateDrawEffect DrawEffects) const;
+
+	static float TrackHeight;
+};
\ No newline at end of file
