Control Rig: Cannot interact with duplicated control rig blueprint

diff --git a/Engine/Plugins/Animation/ControlRig/Source/ControlRigDeveloper/Private/ControlRigBlueprint.cpp b/Engine/Plugins/Animation/ControlRig/Source/ControlRigDeveloper/Private/ControlRigBlueprint.cpp
--- a/Engine/Plugins/Animation/ControlRig/Source/ControlRigDeveloper/Private/ControlRigBlueprint.cpp
+++ b/Engine/Plugins/Animation/ControlRig/Source/ControlRigDeveloper/Private/ControlRigBlueprint.cpp
@@ -1842,6 +1842,20 @@ void UControlRigBlueprint::ReplaceDeprecatedNodes()
 	Super::ReplaceDeprecatedNodes();
 }
 
+void UControlRigBlueprint::PostDuplicate(EDuplicateMode::Type DuplicateMode)
+{
+	if(HierarchyController)
+	{
+		HierarchyController->OnModified().RemoveAll(this);
+	}
+	
+	HierarchyController = NewObject<URigHierarchyController>(this);
+	HierarchyController->SetHierarchy(Hierarchy);
+	HierarchyController->OnModified().AddUObject(this, &UControlRigBlueprint::HandleHierarchyModified);
+
+	FBlueprintEditorUtils::MarkBlueprintAsStructurallyModified(this);
+}
+
 FRigVMGraphModifiedEvent& UControlRigBlueprint::OnModified()
 {
 	return ModifiedEvent;
diff --git a/Engine/Plugins/Animation/ControlRig/Source/ControlRigDeveloper/Public/ControlRigBlueprint.h b/Engine/Plugins/Animation/ControlRig/Source/ControlRigDeveloper/Public/ControlRigBlueprint.h
--- a/Engine/Plugins/Animation/ControlRig/Source/ControlRigDeveloper/Public/ControlRigBlueprint.h
+++ b/Engine/Plugins/Animation/ControlRig/Source/ControlRigDeveloper/Public/ControlRigBlueprint.h
@@ -223,6 +223,7 @@ public:
 	virtual bool IsPostLoadThreadSafe() const override { return false; }
 	virtual void PostTransacted(const FTransactionObjectEvent& TransactionEvent) override;
 	virtual void ReplaceDeprecatedNodes() override;
+	virtual void PostDuplicate(EDuplicateMode::Type DuplicateMode) override;
 
 	virtual bool SupportsGlobalVariables() const override { return true; }
 	virtual bool SupportsLocalVariables() const override { return true; }
@@ -353,17 +354,17 @@ protected:
 
 #if WITH_EDITORONLY_DATA
 
-	UPROPERTY(transient)
+	UPROPERTY(transient, DuplicateTransient)
 	TObjectPtr<URigVMGraph> TemplateModel;
 
-	UPROPERTY(transient)
+	UPROPERTY(transient, DuplicateTransient)
 	TObjectPtr<URigVMController> TemplateController;
 
 #endif
 
 public:
 
-	UPROPERTY(transient)
+	UPROPERTY(transient, DuplicateTransient)
 	TMap<FString, FRigVMOperand> PinToOperandMap;
 
 	bool bSuspendModelNotificationsForSelf;
@@ -408,7 +409,7 @@ public:
 	TSoftObjectPtr<UControlRigGizmoLibrary> GizmoLibrary;
 #endif
 
-	UPROPERTY(transient, VisibleAnywhere, Category = "VM", meta = (DisplayName = "VM Statistics", DisplayAfter = "VMCompileSettings"))
+	UPROPERTY(transient, DuplicateTransient, VisibleAnywhere, Category = "VM", meta = (DisplayName = "VM Statistics", DisplayAfter = "VMCompileSettings"))
 	FRigVMStatistics Statistics;
 
 	UPROPERTY(EditAnywhere, Category = "Drawing")
@@ -481,19 +482,19 @@ private:
 	UPROPERTY(AssetRegistrySearchable)
 	bool bExposesAnimatableControls;
 
-	UPROPERTY(transient)
+	UPROPERTY(transient, DuplicateTransient)
 	bool bAutoRecompileVM;
 
-	UPROPERTY(transient)
+	UPROPERTY(transient, DuplicateTransient)
 	bool bVMRecompilationRequired;
 
-	UPROPERTY(transient)
+	UPROPERTY(transient, DuplicateTransient)
 	bool bIsCompiling;
 
-	UPROPERTY(transient)
+	UPROPERTY(transient, DuplicateTransient)
 	int32 VMRecompilationBracket;
 
-	UPROPERTY(transient)
+	UPROPERTY(transient, DuplicateTransient)
 	TObjectPtr<URigHierarchyController> HierarchyController;
 
 	FRigVMGraphModifiedEvent ModifiedEvent;
