Control Rig: Follow up fix for Mac for 17494208

diff --git a/Engine/Plugins/Animation/ControlRig/Source/ControlRigDeveloper/Private/ControlRigBlueprint.cpp b/Engine/Plugins/Animation/ControlRig/Source/ControlRigDeveloper/Private/ControlRigBlueprint.cpp
--- a/Engine/Plugins/Animation/ControlRig/Source/ControlRigDeveloper/Private/ControlRigBlueprint.cpp
+++ b/Engine/Plugins/Animation/ControlRig/Source/ControlRigDeveloper/Private/ControlRigBlueprint.cpp
@@ -1842,8 +1842,10 @@ void UControlRigBlueprint::ReplaceDeprecatedNodes()
 	Super::ReplaceDeprecatedNodes();
 }
 
-void UControlRigBlueprint::PostDuplicate(EDuplicateMode::Type DuplicateMode)
+void UControlRigBlueprint::PostDuplicate(bool bDuplicateForPIE)
 {
+	Super::PostDuplicate(bDuplicateForPIE);
+	
 	if(HierarchyController)
 	{
 		HierarchyController->OnModified().RemoveAll(this);
diff --git a/Engine/Plugins/Animation/ControlRig/Source/ControlRigDeveloper/Public/ControlRigBlueprint.h b/Engine/Plugins/Animation/ControlRig/Source/ControlRigDeveloper/Public/ControlRigBlueprint.h
--- a/Engine/Plugins/Animation/ControlRig/Source/ControlRigDeveloper/Public/ControlRigBlueprint.h
+++ b/Engine/Plugins/Animation/ControlRig/Source/ControlRigDeveloper/Public/ControlRigBlueprint.h
@@ -223,7 +223,7 @@ public:
 	virtual bool IsPostLoadThreadSafe() const override { return false; }
 	virtual void PostTransacted(const FTransactionObjectEvent& TransactionEvent) override;
 	virtual void ReplaceDeprecatedNodes() override;
-	virtual void PostDuplicate(EDuplicateMode::Type DuplicateMode) override;
+	virtual void PostDuplicate(bool bDuplicateForPIE) override;
 
 	virtual bool SupportsGlobalVariables() const override { return true; }
 	virtual bool SupportsLocalVariables() const override { return true; }
