ControlRig: add the necessary flags for RigHierarchy in a Control Rig CDO

Integrated from 17460536

diff --git a/Engine/Plugins/Animation/ControlRig/Source/ControlRig/Public/ControlRig.h b/Engine/Plugins/Animation/ControlRig/Source/ControlRig/Public/ControlRig.h
--- a/Engine/Plugins/Animation/ControlRig/Source/ControlRig/Public/ControlRig.h
+++ b/Engine/Plugins/Animation/ControlRig/Source/ControlRig/Public/ControlRig.h
@@ -205,6 +205,10 @@ public:
 	UFUNCTION(BlueprintPure, Category = "Control Rig")
 	FORCEINLINE_DEBUGGABLE URigHierarchy* GetHierarchy()
 	{
+		// for CDO, when the hierarchy is first created using CreateDefaultSubobject()
+   		// the outer of that hierarchy will be Default_ControlRig instead of "this" because of
+   		// a call to the destructor of FObjectInitializer during CDO constructor that did it.
+   		// therefore, the hierarchy in CDO will be recreated in this function using NewObject() below 
 		if(DynamicHierarchy != nullptr && DynamicHierarchy->GetOuter() != this)
 		{
 			DynamicHierarchy = nullptr;
@@ -217,6 +221,12 @@ public:
 			{
 				DynamicHierarchy->SetFlags(RF_Transient | RF_Transactional);
 			}
+			else
+			{
+				// for CDO, these flags are needed, they should have come with CreateDefaultSubobject()
+				// but since we are recreating a hierarchy here using NewObject, we need make sure the correct flags are manually added
+				DynamicHierarchy->SetFlags(DynamicHierarchy->GetFlags() | RF_Public | RF_DefaultSubObject);
+			}
 #if WITH_EDITOR
 			const TWeakObjectPtr<UControlRig> WeakThis = this;
 			DynamicHierarchy->OnUndoRedo().AddStatic(&UControlRig::OnHierarchyTransformUndoRedoWeak, WeakThis);
