Sequencer: Control Rig : When getting local values time was getting transformed twice to tick rate.

diff --git a/Engine/Plugins/Animation/ControlRig/Source/ControlRigEditor/Private/ControlRigSequencerEditorLibrary.cpp b/Engine/Plugins/Animation/ControlRig/Source/ControlRigEditor/Private/ControlRigSequencerEditorLibrary.cpp
--- a/Engine/Plugins/Animation/ControlRig/Source/ControlRigEditor/Private/ControlRigSequencerEditorLibrary.cpp
+++ b/Engine/Plugins/Animation/ControlRig/Source/ControlRigEditor/Private/ControlRigSequencerEditorLibrary.cpp
@@ -850,12 +850,10 @@ static bool LocalGetControlRigControlValues(IMovieScenePlayer* Player, UMovieSce
 	{
 
 		FFrameRate TickResolution = MovieScene->GetTickResolution();
-		TArray<FFrameNumber> Frames;
-		ConvertFramesToTickResolution(MovieSceneSequence->GetMovieScene(), InFrames, TimeUnit, Frames);
-		OutValues.SetNum(Frames.Num());
-		for (int32 Index = 0; Index < Frames.Num(); ++Index)
+		OutValues.SetNum(InFrames.Num());
+		for (int32 Index = 0; Index < InFrames.Num(); ++Index)
 		{
-			const FFrameNumber& FrameNumber = Frames[Index];
+			const FFrameNumber& FrameNumber = InFrames[Index];
 			FFrameTime GlobalTime(FrameNumber);
 			if (TimeUnit == ESequenceTimeUnit::DisplayRate)
 			{
@@ -879,9 +877,10 @@ static bool GetControlRigValues(ISequencer* Sequencer, UControlRig* ControlRig,
 	{
 		FMovieSceneSequenceIDRef Template = Sequencer->GetFocusedTemplateID();
 		FMovieSceneSequenceTransform RootToLocalTransform;
-		return LocalGetControlRigControlValues(Sequencer, Sequencer->GetFocusedMovieSceneSequence(), Template, RootToLocalTransform,
+		bool bGood = LocalGetControlRigControlValues(Sequencer, Sequencer->GetFocusedMovieSceneSequence(), Template, RootToLocalTransform,
 			ControlRig, ControlName,TimeUnit, Frames, OutValues);
-
+		Sequencer->RequestEvaluate();
+		return bGood;
 	}
 	return false;
 }
