Sequencer: Control Rig: Space Baking, minor fixes for Chase,like bake in parent space and add joint as possible space.
#prefilight 6141177a568b3a0001bc114b

diff --git a/Engine/Plugins/Animation/ControlRig/Source/ControlRigEditor/Private/EditMode/SControlRigEditModeTools.cpp b/Engine/Plugins/Animation/ControlRig/Source/ControlRigEditor/Private/EditMode/SControlRigEditModeTools.cpp
--- a/Engine/Plugins/Animation/ControlRig/Source/ControlRigEditor/Private/EditMode/SControlRigEditModeTools.cpp
+++ b/Engine/Plugins/Animation/ControlRig/Source/ControlRigEditor/Private/EditMode/SControlRigEditModeTools.cpp
@@ -711,6 +711,8 @@ FReply SControlRigEditModeTools::OnBakeControlsToNewSpaceButtonClicked()
 		FMovieSceneControlRigSpaceBaseKey Value;
 		using namespace UE::MovieScene;
 		URigHierarchy* RigHierarchy = SpacePickerWidget->GetHierarchy();
+		Settings.TargetSpace = RigHierarchy->GetFirstParent(ControlKeys[0]);
+		/*  we just set space as parent, but leaving this in case we go back to setting based on current
 		EvaluateChannel(SpaceChannelAndSection.SpaceChannel, CurrentTime, Value);
 		switch (Value.SpaceType)
 		{
@@ -726,6 +728,7 @@ FReply SControlRigEditModeTools::OnBakeControlsToNewSpaceButtonClicked()
 		default:
 			Settings.TargetSpace = URigHierarchy::GetDefaultParentSocketKey();
 		}
+		*/
 		TRange<FFrameNumber> Range = Sequencer->GetFocusedMovieSceneSequence()->GetMovieScene()->GetPlaybackRange();
 		TArray<FFrameNumber> Keys;
 		TArray < FKeyHandle> KeyHandles;
diff --git a/Engine/Plugins/Animation/ControlRig/Source/ControlRigEditor/Private/Sequencer/ControlRigSpaceChannelEditors.cpp b/Engine/Plugins/Animation/ControlRig/Source/ControlRigEditor/Private/Sequencer/ControlRigSpaceChannelEditors.cpp
--- a/Engine/Plugins/Animation/ControlRig/Source/ControlRigEditor/Private/Sequencer/ControlRigSpaceChannelEditors.cpp
+++ b/Engine/Plugins/Animation/ControlRig/Source/ControlRigEditor/Private/Sequencer/ControlRigSpaceChannelEditors.cpp
@@ -259,7 +259,7 @@ FKeyHandle FControlRigSpaceChannelHelpers::SequencerKeyControlRigSpaceChannel(UC
 			Value.SpaceType = EMovieSceneControlRigSpaceType::Parent;
 
 		}
-		else if (SpaceKey.Type == ERigElementType::Control)
+		else  //support all types
 		{
 			Value.SpaceType = EMovieSceneControlRigSpaceType::ControlRig;
 			Value.ControlRigElement = SpaceKey;
