Control Rig: Lack of max-length validation for the function's 'Category' field causes Editor crash

diff --git a/Engine/Plugins/Animation/ControlRig/Source/ControlRigEditor/Private/ControlRigGraphDetails.cpp b/Engine/Plugins/Animation/ControlRig/Source/ControlRigEditor/Private/ControlRigGraphDetails.cpp
--- a/Engine/Plugins/Animation/ControlRig/Source/ControlRigEditor/Private/ControlRigGraphDetails.cpp
+++ b/Engine/Plugins/Animation/ControlRig/Source/ControlRigEditor/Private/ControlRigGraphDetails.cpp
@@ -143,11 +143,25 @@ void FControlRigArgumentLayout::GenerateHeaderRowContent(FDetailWidgetRow& NodeR
 			[
 				SAssignNew(ArgumentNameWidget, SEditableTextBox)
 				.Text(this, &FControlRigArgumentLayout::OnGetArgNameText)
-				.OnTextChanged(this, &FControlRigArgumentLayout::OnArgNameChange)
 				.OnTextCommitted(this, &FControlRigArgumentLayout::OnArgNameTextCommitted)
 				.ToolTipText(this, &FControlRigArgumentLayout::OnGetArgToolTipText)
 				.Font(IDetailLayoutBuilder::GetDetailFont())
 				.IsEnabled(!ShouldPinBeReadOnly())
+				.OnVerifyTextChanged_Lambda([&](const FText& InNewText, FText& OutErrorMessage) -> bool
+				{
+					if (InNewText.IsEmpty())
+					{
+						OutErrorMessage = LOCTEXT("ArgumentNameEmpty", "Cannot have an argument with an emtpy string name.");
+						return false;
+					}
+					else if (InNewText.ToString().Len() >= NAME_SIZE)
+					{
+						OutErrorMessage = LOCTEXT("ArgumentNameTooLong", "Name of argument is too long.");
+						return false;
+					}
+					
+					return true;
+				})
 			]
 		]
 		.ValueContent()
@@ -298,11 +312,6 @@ FText FControlRigArgumentLayout::OnGetArgToolTipText() const
 	return OnGetArgNameText(); // for now since we don't have tooltips
 }
 
-void FControlRigArgumentLayout::OnArgNameChange(const FText& InNewText)
-{
-	// do we need validation?
-}
-
 void FControlRigArgumentLayout::OnArgNameTextCommitted(const FText& NewText, ETextCommit::Type InTextCommit)
 {
 	if (!NewText.IsEmpty() && PinPtr.IsValid() && ControlRigBlueprintPtr.IsValid() && !ShouldPinBeReadOnly())
@@ -663,6 +672,12 @@ void FControlRigGraphDetails::CustomizeDetails(IDetailLayoutBuilder& DetailLayou
 					OutErrorMessage = LOCTEXT("CategoryEmpty", "Cannot add a category with an empty string.");
 					return false;
 				}
+				else if (NewText.ToString().Len() >= NAME_SIZE)
+				{
+					OutErrorMessage = LOCTEXT("CategoryTooLong", "Name of category is too long.");
+					return false;
+				}
+				
 				if (ControlRigBlueprintPtr.IsValid())
 				{
 					if (NewText.EqualTo(FText::FromString(ControlRigBlueprintPtr.Get()->GetName())))
diff --git a/Engine/Plugins/Animation/ControlRig/Source/ControlRigEditor/Private/ControlRigGraphDetails.h b/Engine/Plugins/Animation/ControlRig/Source/ControlRigEditor/Private/ControlRigGraphDetails.h
--- a/Engine/Plugins/Animation/ControlRig/Source/ControlRigEditor/Private/ControlRigGraphDetails.h
+++ b/Engine/Plugins/Animation/ControlRig/Source/ControlRigEditor/Private/ControlRigGraphDetails.h
@@ -83,7 +83,6 @@ private:
 
 	FText OnGetArgNameText() const;
 	FText OnGetArgToolTipText() const;
-	void OnArgNameChange(const FText& InNewText);
 	void OnArgNameTextCommitted(const FText& NewText, ETextCommit::Type InTextCommit);
 
 	FEdGraphPinType OnGetPinInfo() const;
