Control Rig: Adding cycle checking to space switching and offering user feedback

diff --git a/Engine/Plugins/Animation/ControlRig/Source/ControlRig/Private/Rigs/RigHierarchy.cpp b/Engine/Plugins/Animation/ControlRig/Source/ControlRig/Private/Rigs/RigHierarchy.cpp
--- a/Engine/Plugins/Animation/ControlRig/Source/ControlRig/Private/Rigs/RigHierarchy.cpp
+++ b/Engine/Plugins/Animation/ControlRig/Source/ControlRig/Private/Rigs/RigHierarchy.cpp
@@ -1417,30 +1417,90 @@ bool URigHierarchy::SetParentWeightArray(FRigBaseElement* InChild,  const TArray
 	return false;
 }
 
-bool URigHierarchy::SwitchToParent(FRigElementKey InChild, FRigElementKey InParent, bool bInitial, bool bAffectChildren)
+bool URigHierarchy::CanSwitchToParent(FRigElementKey InChild, FRigElementKey InParent, FString* OutFailureReason)
 {
-	if(InParent == GetWorldSpaceSocketKey())
+	InParent = PreprocessParentElementKeyForSpaceSwitching(InChild, InParent);
+
+	FRigBaseElement* Child = Find(InChild);
+	if(Child == nullptr)
 	{
-		InParent = GetOrAddWorldSpaceSocket();
+		if(OutFailureReason)
+		{
+			*OutFailureReason = FString::Printf(TEXT("Child Element %s cannot be found."), *InChild.ToString());
+		}
+		return false;
 	}
-	else if(InParent == GetDefaultParentSocketKey())
+
+	FRigBaseElement* Parent = Find(InParent);
+	if(Parent == nullptr)
 	{
-		const FRigElementKey FirstParent = GetFirstParent(InChild);
-		if(FirstParent == GetWorldSpaceSocketKey())
+		// if we don't specify anything and the element is parented directly to the world,
+		// perfomring this switch means unparenting it from world (since there is no default parent)
+		if(!InParent.IsValid() && GetFirstParent(InChild) == GetWorldSpaceSocketKey())
 		{
-			InParent = FRigElementKey();
+			return true;
 		}
-		else
+		
+		if(OutFailureReason)
 		{
-			InParent = FirstParent;
+			*OutFailureReason = FString::Printf(TEXT("Parent Element %s cannot be found."), *InParent.ToString());
 		}
+		return false;
 	}
-	return SwitchToParent(Find(InChild), Find(InParent), bInitial, bAffectChildren);
+
+	// see if this is already parented to the target parent
+	if(GetFirstParent(Child) == Parent)
+	{
+		return true;
+	}
+
+	const FRigMultiParentElement* MultiParentChild = Cast<FRigMultiParentElement>(Child);
+	if(MultiParentChild == nullptr)
+	{
+		if(OutFailureReason)
+		{
+			*OutFailureReason = FString::Printf(TEXT("Child Element %s does not allow space switching (it's not a multi parent element)."), *InChild.ToString());
+		}
+	}
+
+	const FRigTransformElement* TransformParent = Cast<FRigMultiParentElement>(Parent);
+	if(TransformParent == nullptr)
+	{
+		if(OutFailureReason)
+		{
+			*OutFailureReason = FString::Printf(TEXT("Parent Element %s is not a transform element"), *InParent.ToString());
+		}
+	}
+
+	if(IsParentedTo(InParent, InChild))
+	{
+		if(OutFailureReason)
+		{
+			*OutFailureReason = FString::Printf(TEXT("Cannot switch '%s' to '%s' - would cause a cycle."), *InChild.ToString(), *InParent.ToString());
+		}
+		return false;
+	}
+
+	return true;
+}
+
+bool URigHierarchy::SwitchToParent(FRigElementKey InChild, FRigElementKey InParent, bool bInitial, bool bAffectChildren, FString* OutFailureReason)
+{
+	InParent = PreprocessParentElementKeyForSpaceSwitching(InChild, InParent);
+	return SwitchToParent(Find(InChild), Find(InParent), bInitial, bAffectChildren, OutFailureReason);
 }
 
 bool URigHierarchy::SwitchToParent(FRigBaseElement* InChild, FRigBaseElement* InParent, bool bInitial,
-	bool bAffectChildren)
+	bool bAffectChildren, FString* OutFailureReason)
 {
+	if(InChild && InParent)
+	{
+		if(!CanSwitchToParent(InChild->GetKey(), InParent->GetKey(), OutFailureReason))
+		{
+			return false;
+		}
+	}
+
 	if(const FRigMultiParentElement* MultiParentElement = Cast<FRigMultiParentElement>(InChild))
 	{
 		int32 ParentIndex = INDEX_NONE;
@@ -2856,6 +2916,28 @@ void URigHierarchy::UpdateAllCachedChildren() const
 		}
 	}
 }
+	
+FRigElementKey URigHierarchy::PreprocessParentElementKeyForSpaceSwitching(const FRigElementKey& InChildKey, const FRigElementKey& InParentKey)
+{
+	if(InParentKey == GetWorldSpaceSocketKey())
+	{
+		return GetOrAddWorldSpaceSocket();
+	}
+	else if(InParentKey == GetDefaultParentSocketKey())
+	{
+		const FRigElementKey FirstParent = GetFirstParent(InChildKey);
+		if(FirstParent == GetWorldSpaceSocketKey())
+		{
+			return FRigElementKey();
+		}
+		else
+		{
+			return FirstParent;
+		}
+	}
+
+	return InParentKey;
+}
 
 FRigBaseElement* URigHierarchy::MakeElement(ERigElementType InElementType, int32 InCount, int32* OutStructureSize)
 {
@@ -4054,3 +4136,4 @@ void URigHierarchy::IntegrateParentConstraintQuat(
 	OutMixedRotation.W += InWeight * ParentRotation.W;
 	OutNumMixedRotations++;
 }
+
diff --git a/Engine/Plugins/Animation/ControlRig/Source/ControlRig/Private/Units/Execution/RigUnit_DynamicHierarchy.cpp b/Engine/Plugins/Animation/ControlRig/Source/ControlRig/Private/Units/Execution/RigUnit_DynamicHierarchy.cpp
--- a/Engine/Plugins/Animation/ControlRig/Source/ControlRig/Private/Units/Execution/RigUnit_DynamicHierarchy.cpp
+++ b/Engine/Plugins/Animation/ControlRig/Source/ControlRig/Private/Units/Execution/RigUnit_DynamicHierarchy.cpp
@@ -72,23 +72,37 @@ FRigUnit_SwitchParent_Execute()
 		bMaintainGlobal ? ERigTransformType::CurrentGlobal : ERigTransformType::CurrentLocal;
 	
 	const FTransform Transform = ExecuteContext.Hierarchy->GetTransform(ChildElement, TransformTypeToMaintain);
-	
+
 	switch(Mode)
 	{
 		case ERigSwitchParentMode::World:
 		{
-			ExecuteContext.Hierarchy->SwitchToWorldSpace(ChildElement, false, true);
+			if(!ExecuteContext.Hierarchy->SwitchToWorldSpace(ChildElement, false, true))
+			{
+				return;
+			}
 			break;
 		}
 		case ERigSwitchParentMode::DefaultParent:
 		{
-			ExecuteContext.Hierarchy->SwitchToDefaultParent(ChildElement, false, true);
+			if(!ExecuteContext.Hierarchy->SwitchToDefaultParent(ChildElement, false, true))
+			{
+				return;
+			}
 			break;
 		}
 		case ERigSwitchParentMode::ParentItem:
 		default:
 		{
-			ExecuteContext.Hierarchy->SwitchToParent(ChildElement, ParentElement, false, true);
+			FString FailureReason;
+			if(!ExecuteContext.Hierarchy->SwitchToParent(ChildElement, ParentElement, false, true, &FailureReason))
+			{
+				if(!FailureReason.IsEmpty())
+				{
+					UE_CONTROLRIG_RIGUNIT_REPORT_ERROR(TEXT("%s"), *FailureReason);
+				}
+				return;
+			}
 			break;
 		}
 	}
diff --git a/Engine/Plugins/Animation/ControlRig/Source/ControlRig/Public/Rigs/RigHierarchy.h b/Engine/Plugins/Animation/ControlRig/Source/ControlRig/Public/Rigs/RigHierarchy.h
--- a/Engine/Plugins/Animation/ControlRig/Source/ControlRig/Public/Rigs/RigHierarchy.h
+++ b/Engine/Plugins/Animation/ControlRig/Source/ControlRig/Public/Rigs/RigHierarchy.h
@@ -1728,17 +1728,31 @@ public:
 	bool SetParentWeightArray(FRigBaseElement* InChild,  const TArrayView<const FRigElementWeight>& InWeights, bool bInitial = false, bool bAffectChildren = true);
 
 	/**
-		* Switches a multi parent element to a single parent.
-		* This sets the new parent's weight to 1.0 and disables
-		* weights for all other potential parents.
-		* @param InChild The key of the multi parented element
-		* @param InParent The key of the parent to look up the weight for
-		* @param bInitial If true the initial weights will be used
-		* @param bAffectChildren If set to false children will not move (maintain global).
-		* @return Returns true if changing the weight was successful
-		*/
+	* Determines if the element can be switched to a provided parent
+	* @param InChild The key of the multi parented element
+	* @param InParent The key of the parent to look up the weight for
+	* @param OutFailureReason An optional pointer to retrieve the reason for failure
+	* @return Returns true if changing the weight was successful
+	*/
+	bool CanSwitchToParent(FRigElementKey InChild, FRigElementKey InParent, FString* OutFailureReason = nullptr);
+
+	/**
+	 * Switches a multi parent element to a single parent.
+	 * This sets the new parent's weight to 1.0 and disables
+	 * weights for all other potential parents.
+	 * @param InChild The key of the multi parented element
+	 * @param InParent The key of the parent to look up the weight for
+	 * @param bInitial If true the initial weights will be used
+	 * @param bAffectChildren If set to false children will not move (maintain global).
+	 * @param OutFailureReason An optional pointer to retrieve the reason for failure
+	 * @return Returns true if changing the weight was successful
+	 */
 	UFUNCTION(BlueprintCallable, Category = URigHierarchy)
-	bool SwitchToParent(FRigElementKey InChild, FRigElementKey InParent, bool bInitial = false, bool bAffectChildren = true);
+	bool SwitchToParent(FRigElementKey InChild, FRigElementKey InParent, bool bInitial = false, bool bAffectChildren = true)
+	{
+		return SwitchToParent(InChild, InParent, bInitial, bAffectChildren, nullptr);
+	}
+	bool SwitchToParent(FRigElementKey InChild, FRigElementKey InParent, bool bInitial, bool bAffectChildren, FString* OutFailureReason);
 
 	/**
 	 * Switches a multi parent element to a single parent.
@@ -1748,9 +1762,10 @@ public:
 	 * @param InParent The parent to look up the weight for
 	 * @param bInitial If true the initial weights will be used
 	 * @param bAffectChildren If set to false children will not move (maintain global).
+	 * @param OutFailureReason An optional pointer to retrieve the reason for failure
 	 * @return Returns true if changing the weight was successful
 	 */
-	bool SwitchToParent(FRigBaseElement* InChild, FRigBaseElement* InParent, bool bInitial = false, bool bAffectChildren = true);
+	bool SwitchToParent(FRigBaseElement* InChild, FRigBaseElement* InParent, bool bInitial = false, bool bAffectChildren = true, FString* OutFailureReason = nullptr);
 
 	/**
 	 * Switches a multi parent element to a single parent.
@@ -2452,6 +2467,11 @@ private:
 	*/
 	void UpdateAllCachedChildren() const;
 
+	/**
+	 * Corrects a parent element key for space switching
+	 */
+	FRigElementKey PreprocessParentElementKeyForSpaceSwitching(const FRigElementKey& InChildKey, const FRigElementKey& InParentKey);
+
 	/*
 	 * Helper function to create an element for a given type
 	 */
diff --git a/Engine/Plugins/Animation/ControlRig/Source/ControlRigEditor/Private/EditMode/SControlRigEditModeTools.cpp b/Engine/Plugins/Animation/ControlRig/Source/ControlRigEditor/Private/EditMode/SControlRigEditModeTools.cpp
--- a/Engine/Plugins/Animation/ControlRig/Source/ControlRigEditor/Private/EditMode/SControlRigEditModeTools.cpp
+++ b/Engine/Plugins/Animation/ControlRig/Source/ControlRigEditor/Private/EditMode/SControlRigEditModeTools.cpp
@@ -34,6 +34,8 @@
 #include "EdModeInteractiveToolsContext.h"
 #include "ControlRigSpaceChannelEditors.h"
 #include "IKeyArea.h"
+#include "Widgets/Notifications/SNotificationList.h"
+#include "Framework/Notifications/NotificationManager.h"
 
 #define LOCTEXT_NAMESPACE "ControlRigRootCustomization"
 
@@ -628,6 +630,20 @@ void SControlRigEditModeTools::HandleActiveSpaceChanged(URigHierarchy* InHierarc
 	{
 		if (UControlRig* ControlRig = SequencerRig.Get())
 		{
+			FString FailureReason;
+			if(!InHierarchy->CanSwitchToParent(InControlKey, InSpaceKey, &FailureReason))
+			{
+				// notification
+				FNotificationInfo Info(FText::FromString(FailureReason));
+				Info.bFireAndForget = true;
+				Info.FadeOutDuration = 2.0f;
+				Info.ExpireDuration = 8.0f;
+
+				const TSharedPtr<SNotificationItem> NotificationPtr = FSlateNotificationManager::Get().AddNotification(Info);
+				NotificationPtr->SetCompletionState(SNotificationItem::CS_Fail);
+				return;
+			}
+			
 			if (const FRigControlElement* ControlElement = InHierarchy->Find<FRigControlElement>(InControlKey))
 			{
 				ISequencer* Sequencer = WeakSequencer.Pin().Get();
diff --git a/Engine/Plugins/Animation/ControlRig/Source/ControlRigEditor/Private/Editor/SRigSpacePickerWidget.cpp b/Engine/Plugins/Animation/ControlRig/Source/ControlRigEditor/Private/Editor/SRigSpacePickerWidget.cpp
--- a/Engine/Plugins/Animation/ControlRig/Source/ControlRigEditor/Private/Editor/SRigSpacePickerWidget.cpp
+++ b/Engine/Plugins/Animation/ControlRig/Source/ControlRigEditor/Private/Editor/SRigSpacePickerWidget.cpp
@@ -11,6 +11,8 @@
 #include "ISequencer.h"
 #include "Widgets/Input/NumericTypeInterface.h"
 #include "FrameNumberDetailsCustomization.h"
+#include "Widgets/Notifications/SNotificationList.h"
+#include "Framework/Notifications/NotificationManager.h"
 
 #define LOCTEXT_NAMESPACE "SRigSpacePickerWidget"
 
@@ -595,6 +597,23 @@ FReply SRigSpacePickerWidget::HandleAddElementClicked()
 			const FRigElementKey Key = InItem->Key;
 			if(!IsDefaultSpace(Key) && IsValidKey(Key))
 			{
+				for(const FRigElementKey& ControlKey : ControlKeys)
+				{
+					FString FailureReason;
+					if(!Hierarchy->CanSwitchToParent(ControlKey, Key, &FailureReason))
+					{
+						// notification
+						FNotificationInfo Info(FText::FromString(FailureReason));
+						Info.bFireAndForget = true;
+						Info.FadeOutDuration = 2.0f;
+						Info.ExpireDuration = 8.0f;
+
+						const TSharedPtr<SNotificationItem> NotificationPtr = FSlateNotificationManager::Get().AddNotification(Info);
+						NotificationPtr->SetCompletionState(SNotificationItem::CS_Fail);
+						return;
+					}
+				}
+				
 				TArray<FRigElementKey> ChangedSpaceKeys = CurrentSpaceKeys;
 				ChangedSpaceKeys.AddUnique(Key);
 
