Control Rig: Fix blueprint postload compiling twice

diff --git a/Engine/Plugins/Animation/ControlRig/Source/ControlRigDeveloper/Private/ControlRigBlueprint.cpp b/Engine/Plugins/Animation/ControlRig/Source/ControlRigDeveloper/Private/ControlRigBlueprint.cpp
--- a/Engine/Plugins/Animation/ControlRig/Source/ControlRigDeveloper/Private/ControlRigBlueprint.cpp
+++ b/Engine/Plugins/Animation/ControlRig/Source/ControlRigDeveloper/Private/ControlRigBlueprint.cpp
@@ -344,6 +344,7 @@ void UControlRigBlueprint::PostLoad()
 {
 	Super::PostLoad();
 
+	bVMRecompilationRequired = true;
 	{
 		FControlRigBlueprintVMCompileScope CompileScope(this);
 		TArray<UControlRigBlueprint*> ReferencedBlueprints = GetReferencedControlRigBlueprints();
@@ -508,7 +509,7 @@ void UControlRigBlueprint::PostLoad()
 #endif
 	}
 
-	RecompileVM();
+	RecompileVMIfRequired();
 	RequestControlRigInit();
 
 	FCoreUObjectDelegates::OnObjectModified.RemoveAll(this);
@@ -581,6 +582,7 @@ void UControlRigBlueprint::RecompileVM()
 
 		if (bErrorsDuringCompilation)
 		{
+			bVMRecompilationRequired = false;
 			if(CDO->VM)
 			{
 				VMCompiledEvent.Broadcast(this, CDO->VM);
