Control Rig: change skeletal mesh right click menu wording for creating a control rig from "Create Control Rig" to "Control Rig"

diff --git a/Engine/Plugins/Animation/ControlRig/Source/ControlRigEditor/Private/ControlRigBlueprintActions.cpp b/Engine/Plugins/Animation/ControlRig/Source/ControlRigEditor/Private/ControlRigBlueprintActions.cpp
--- a/Engine/Plugins/Animation/ControlRig/Source/ControlRigEditor/Private/ControlRigBlueprintActions.cpp
+++ b/Engine/Plugins/Animation/ControlRig/Source/ControlRigEditor/Private/ControlRigBlueprintActions.cpp
@@ -104,7 +104,7 @@ void FControlRigBlueprintActions::ExtendSketalMeshToolMenu()
 				{
 					InSection.AddMenuEntry(
 						"CreateControlRig",
-						LOCTEXT("CreateControlRig", "Create Control Rig"),
+						LOCTEXT("CreateControlRig", "Control Rig"),
 						LOCTEXT("CreateControlRig_ToolTip", "Creates a control rig and preconfigures it for this asset"),
 						FSlateIcon(FControlRigEditorStyle::Get().GetStyleSetName(), "ControlRig", "ControlRig.RigUnit"),
 						FExecuteAction::CreateLambda([SelectedObjects]()
