Control Rig: Newly created public variables on Control Rig cannot be used as input pins on Control Rig AnimNode

diff --git a/Engine/Plugins/Animation/ControlRig/Source/ControlRigDeveloper/Private/AnimGraphNode_ControlRig.cpp b/Engine/Plugins/Animation/ControlRig/Source/ControlRigDeveloper/Private/AnimGraphNode_ControlRig.cpp
--- a/Engine/Plugins/Animation/ControlRig/Source/ControlRigDeveloper/Private/AnimGraphNode_ControlRig.cpp
+++ b/Engine/Plugins/Animation/ControlRig/Source/ControlRigDeveloper/Private/AnimGraphNode_ControlRig.cpp
@@ -366,7 +366,7 @@ bool UAnimGraphNode_ControlRig::IsInputProperty(const FName& PropertyName) const
 	DECLARE_SCOPE_HIERARCHICAL_COUNTER_FUNC()
 
 	// this is true for both input variables and controls
-	return !OutputVariables.Contains(PropertyName);
+	return InputVariables.Contains(PropertyName) || !OutputVariables.Contains(PropertyName);
 }
 
 FRigControlElement* UAnimGraphNode_ControlRig::FindControlElement(const FName& InControlName) const
