Control Rig: Fix MathVectorDiv unit test failing in mac

diff --git a/Engine/Plugins/Animation/ControlRig/Source/ControlRig/Private/Units/Math/RigUnit_MathVectorTests.cpp b/Engine/Plugins/Animation/ControlRig/Source/ControlRig/Private/Units/Math/RigUnit_MathVectorTests.cpp
--- a/Engine/Plugins/Animation/ControlRig/Source/ControlRig/Private/Units/Math/RigUnit_MathVectorTests.cpp
+++ b/Engine/Plugins/Animation/ControlRig/Source/ControlRig/Private/Units/Math/RigUnit_MathVectorTests.cpp
@@ -63,7 +63,7 @@ IMPLEMENT_RIGUNIT_AUTOMATION_TEST(FRigUnit_MathVectorDiv)
 	Unit.A = FVector(4.f, 6.f, 8.f);
 	Unit.B = FVector(1.f, 2.f, 3.f);
 	InitAndExecute();
-	AddErrorIfFalse(FRigUnit_MathVectorTest_Utils::IsNearlyEqual(Unit.Result, FVector(4.f, 3.f, 8.f / 3.f)), TEXT("unexpected result"));
+	AddErrorIfFalse(FRigUnit_MathVectorTest_Utils::IsNearlyEqual(Unit.Result, FVector(4.f, 3.f, 8.f / 3.f), KINDA_SMALL_NUMBER), TEXT("unexpected result"));
 	return true;
 }
 
