ControlRig: 1. Made Hierarchy Controller to be part of a Hierarchy so that its lifetime is tied to the Hierarchy. 2. Create the hierarchy controller during cr initialization to avoid uobject creation during on aync thread

diff --git a/Engine/Plugins/Animation/ControlRig/Source/ControlRig/Private/ControlRig.cpp b/Engine/Plugins/Animation/ControlRig/Source/ControlRig/Private/ControlRig.cpp
--- a/Engine/Plugins/Animation/ControlRig/Source/ControlRig/Private/ControlRig.cpp
+++ b/Engine/Plugins/Animation/ControlRig/Source/ControlRig/Private/ControlRig.cpp
@@ -172,6 +172,9 @@ void UControlRig::Initialize(bool bInitRigUnits)
 	// Create the data source registry here to avoid UObject creation from Non-Game Threads
 	GetDataSourceRegistry();
 
+	// Create the Hierarchy Controller here to avoid UObject creation from Non-Game Threads
+	GetHierarchy()->GetController(true);
+	
 	// should refresh mapping 
 	RequestSetup();
 
diff --git a/Engine/Plugins/Animation/ControlRig/Source/ControlRig/Private/Rigs/AdditiveControlRig.cpp b/Engine/Plugins/Animation/ControlRig/Source/ControlRig/Private/Rigs/AdditiveControlRig.cpp
--- a/Engine/Plugins/Animation/ControlRig/Source/ControlRig/Private/Rigs/AdditiveControlRig.cpp
+++ b/Engine/Plugins/Animation/ControlRig/Source/ControlRig/Private/Rigs/AdditiveControlRig.cpp
@@ -91,54 +91,51 @@ void UAdditiveControlRig::Initialize(bool bInitRigUnits /*= true*/)
 
 void UAdditiveControlRig::CreateRigElements(const FReferenceSkeleton& InReferenceSkeleton, const FSmartNameMapping* InSmartNameMapping)
 {
-	if(Controller == nullptr)
-	{
-		Controller = NewObject<URigHierarchyController>(this);
-		Controller->SetHierarchy(GetHierarchy());
-	}
-	
 	GetHierarchy()->Reset();
-	Controller->ImportBones(InReferenceSkeleton, NAME_None, false, false, true, false);
-
-	if (InSmartNameMapping)
+	if (URigHierarchyController* Controller = GetHierarchy()->GetController(true))
 	{
-		TArray<FName> NameArray;
-		InSmartNameMapping->FillNameArray(NameArray);
-		for (int32 Index = 0; Index < NameArray.Num(); ++Index)
-		{
-			Controller->AddCurve(NameArray[Index], 0.f, false);
-		}
-	}
-
-	// add control for all bone hierarchy
-	GetHierarchy()->ForEach<FRigBoneElement>([&](FRigBoneElement* BoneElement) -> bool
-    {
-		const FName BoneName = BoneElement->GetName();
-		const int32 ParentIndex = GetHierarchy()->GetFirstParent(BoneElement->GetIndex());
-		const FName NullName = GetNullName(BoneName);// name conflict?
-		const FName ControlName = GetControlName(BoneName); // name conflict?
+		Controller->ImportBones(InReferenceSkeleton, NAME_None, false, false, true, false);
 
-		FRigElementKey NullKey;
-		
-		if (ParentIndex != INDEX_NONE)
-		{
-			FTransform GlobalTransform = GetHierarchy()->GetGlobalTransform(BoneElement->GetIndex());
-			FTransform ParentTransform = GetHierarchy()->GetGlobalTransform(ParentIndex);
-			FTransform LocalTransform = GlobalTransform.GetRelativeTransform(ParentTransform);
-			NullKey = Controller->AddNull(NullName, GetHierarchy()->GetKey(ParentIndex), LocalTransform, false, false);
-		}
-		else
+		if (InSmartNameMapping)
 		{
-			FTransform GlobalTransform = GetHierarchy()->GetGlobalTransform(BoneElement->GetIndex());
-			NullKey = Controller->AddNull(NullName, FRigElementKey(), GlobalTransform, true, false);
+			TArray<FName> NameArray;
+			InSmartNameMapping->FillNameArray(NameArray);
+			for (int32 Index = 0; Index < NameArray.Num(); ++Index)
+			{
+				 Controller->AddCurve(NameArray[Index], 0.f, false);
+			}
 		}
 
-		FRigControlSettings Settings;
-		Settings.DisplayName = BoneName;
-		Controller->AddControl(ControlName, NullKey, Settings, FRigControlValue::Make(FTransform::Identity), FTransform::Identity, FTransform::Identity, false);
+		// add control for all bone hierarchy
+		GetHierarchy()->ForEach<FRigBoneElement>([&](FRigBoneElement* BoneElement) -> bool
+        {
+            const FName BoneName = BoneElement->GetName();
+            const int32 ParentIndex = GetHierarchy()->GetFirstParent(BoneElement->GetIndex());
+            const FName NullName = GetNullName(BoneName);// name conflict?
+            const FName ControlName = GetControlName(BoneName); // name conflict?
 
-		return true;
-	});
+            FRigElementKey NullKey;
+		
+            if (ParentIndex != INDEX_NONE)
+            {
+                FTransform GlobalTransform = GetHierarchy()->GetGlobalTransform(BoneElement->GetIndex());
+                FTransform ParentTransform = GetHierarchy()->GetGlobalTransform(ParentIndex);
+                FTransform LocalTransform = GlobalTransform.GetRelativeTransform(ParentTransform);
+                NullKey = Controller->AddNull(NullName, GetHierarchy()->GetKey(ParentIndex), LocalTransform, false, false);
+            }
+            else
+            {
+                FTransform GlobalTransform = GetHierarchy()->GetGlobalTransform(BoneElement->GetIndex());
+                NullKey = Controller->AddNull(NullName, FRigElementKey(), GlobalTransform, true, false);
+            }
+
+            FRigControlSettings Settings;
+            Settings.DisplayName = BoneName;
+            Controller->AddControl(ControlName, NullKey, Settings, FRigControlValue::Make(FTransform::Identity), FTransform::Identity, FTransform::Identity, false);
+
+            return true;
+        });
+	}
 }
 
 void UAdditiveControlRig::CreateRigElements(const USkeletalMesh* InReferenceMesh)
diff --git a/Engine/Plugins/Animation/ControlRig/Source/ControlRig/Private/Rigs/FKControlRig.cpp b/Engine/Plugins/Animation/ControlRig/Source/ControlRig/Private/Rigs/FKControlRig.cpp
--- a/Engine/Plugins/Animation/ControlRig/Source/ControlRig/Private/Rigs/FKControlRig.cpp
+++ b/Engine/Plugins/Animation/ControlRig/Source/ControlRig/Private/Rigs/FKControlRig.cpp
@@ -250,69 +250,66 @@ void UFKControlRig::SetControlActive(const TArray<FFKBoneCheckInfo>& BoneChecks)
 
 void UFKControlRig::CreateRigElements(const FReferenceSkeleton& InReferenceSkeleton, const FSmartNameMapping* InSmartNameMapping)
 {
-	if(Controller == nullptr)
-	{
-		Controller = NewObject<URigHierarchyController>(this);
-		Controller->SetHierarchy(GetHierarchy());
-	}
-	
 	GetHierarchy()->Reset();
-	Controller->ImportBones(InReferenceSkeleton, NAME_None, false, false, true, false);
-
-	if (InSmartNameMapping)
+	if (URigHierarchyController* Controller = GetHierarchy()->GetController(true))
 	{
-		TArray<FName> NameArray;
-		InSmartNameMapping->FillNameArray(NameArray);
-		for (int32 Index = 0; Index < NameArray.Num(); ++Index)
+		Controller->ImportBones(InReferenceSkeleton, NAME_None, false, false, true, false);
+
+		if (InSmartNameMapping)
 		{
-			Controller->AddCurve(NameArray[Index], 0.f, false);
+			TArray<FName> NameArray;
+			InSmartNameMapping->FillNameArray(NameArray);
+			for (int32 Index = 0; Index < NameArray.Num(); ++Index)
+			{
+				 Controller->AddCurve(NameArray[Index], 0.f, false);
+			}
 		}
-	}
 
-	// add control for all bone hierarchy 
-	int32 ControlIndex = 0;
+		// add control for all bone hierarchy 
+		int32 ControlIndex = 0;
 
-	GetHierarchy()->ForEach<FRigBoneElement>([&](FRigBoneElement* BoneElement) -> bool
-	{
-		const FName BoneName = BoneElement->GetName();
-		const FName ControlName = GetControlName(BoneName); // name conflict?
-		const FRigElementKey ParentKey = GetHierarchy()->GetFirstParent(BoneElement->GetKey());
+		GetHierarchy()->ForEach<FRigBoneElement>([&](FRigBoneElement* BoneElement) -> bool
+		{
+			const FName BoneName = BoneElement->GetName();
+			const FName ControlName = GetControlName(BoneName); // name conflict?
+			const FRigElementKey ParentKey = GetHierarchy()->GetFirstParent(BoneElement->GetKey());
 
-		FTransform OffsetTransform = FTransform::Identity;
+			FTransform OffsetTransform = FTransform::Identity;
 
-		FRigControlSettings Settings;
-		Settings.ControlType = ERigControlType::EulerTransform;
-		Settings.DisplayName = BoneName;
+			FRigControlSettings Settings;
+			Settings.ControlType = ERigControlType::EulerTransform;
+			Settings.DisplayName = BoneName;
 
-		OffsetTransform.NormalizeRotation();
-		OffsetTransform.SetScale3D(FVector::OneVector);
-		
-		Controller->AddControl(ControlName, ParentKey, Settings, FRigControlValue::Make(FEulerTransform::Identity), OffsetTransform, FTransform::Identity, false);
+			OffsetTransform.NormalizeRotation();
+			OffsetTransform.SetScale3D(FVector::OneVector);
+			
+			Controller->AddControl(ControlName, ParentKey, Settings, FRigControlValue::Make(FEulerTransform::Identity), OffsetTransform, FTransform::Identity, false);
 
-		return true;
-	});
+			return true;
+		});
 
-	SetControlOffsetsFromBoneInitials();
-	
-	GetHierarchy()->ForEach<FRigCurveElement>([&](FRigCurveElement* CurveElement) -> bool
-    {
-        const FName ControlName = GetControlName(CurveElement->GetName()); // name conflict?
+		SetControlOffsetsFromBoneInitials();
+		
+		GetHierarchy()->ForEach<FRigCurveElement>([&](FRigCurveElement* CurveElement) -> bool
+		{
+			const FName ControlName = GetControlName(CurveElement->GetName()); // name conflict?
 
-		FRigControlSettings Settings;
-		Settings.ControlType = ERigControlType::Float;
-		Settings.DisplayName = CurveElement->GetName();
+			FRigControlSettings Settings;
+			Settings.ControlType = ERigControlType::Float;
+			Settings.DisplayName = CurveElement->GetName();
 
-		Controller->AddControl(ControlName, FRigElementKey(), Settings, FRigControlValue::Make(CurveElement->Value), FTransform::Identity, FTransform::Identity, false);
-		
-		return true;
-	});
+			Controller->AddControl(ControlName, FRigElementKey(), Settings, FRigControlValue::Make(CurveElement->Value), FTransform::Identity, FTransform::Identity, false);
+			
+			return true;
+		});
 
-	if (IsControlActive.Num() != GetHierarchy()->Num())
-	{
-		IsControlActive.SetNum(GetHierarchy()->Num());
-		for (bool& bIsActive : IsControlActive)
+		if (IsControlActive.Num() != GetHierarchy()->Num())
 		{
-			bIsActive = true;
+			IsControlActive.SetNum(GetHierarchy()->Num());
+			for (bool& bIsActive : IsControlActive)
+			{
+				 bIsActive = true;
+			}
 		}
 	}
 }
diff --git a/Engine/Plugins/Animation/ControlRig/Source/ControlRig/Private/Rigs/RigHierarchy.cpp b/Engine/Plugins/Animation/ControlRig/Source/ControlRig/Private/Rigs/RigHierarchy.cpp
--- a/Engine/Plugins/Animation/ControlRig/Source/ControlRig/Private/Rigs/RigHierarchy.cpp
+++ b/Engine/Plugins/Animation/ControlRig/Source/ControlRig/Private/Rigs/RigHierarchy.cpp
@@ -91,6 +91,7 @@ URigHierarchy::URigHierarchy()
 , bIsInteracting(false)
 , LastInteractedKey()
 , bSuspendNotifications(false)
+, HierarchyController(nullptr)
 , ResetPoseHash(INDEX_NONE)
 #if WITH_EDITOR
 , bPropagatingChange(false)
@@ -1829,22 +1830,18 @@ void URigHierarchy::SendAutoKeyEvent(FRigElementKey InElement, float InOffsetInS
 
 URigHierarchyController* URigHierarchy::GetController(bool bCreateIfNeeded)
 {
-	if(LastControllerPtr.IsValid())
+	if(HierarchyController)
 	{
-		return Cast<URigHierarchyController>(LastControllerPtr.Get());
+		return HierarchyController;
 	}
 	else if(bCreateIfNeeded)
 	{
-		if(UObject* Outer = GetOuter())
-		{
-			if(!IsGarbageCollecting())
-			{
-				URigHierarchyController* Controller = NewObject<URigHierarchyController>(Outer);
-				Controller->SetHierarchy(this);
-				LastControllerPtr = Controller;
-				return Controller;
-			}
-		}
+		 if(!IsGarbageCollecting())
+		 {
+			 HierarchyController = NewObject<URigHierarchyController>(this, TEXT("HierarchyController"), RF_Transient);
+			 HierarchyController->SetHierarchy(this);
+			 return HierarchyController;
+		 }
 	}
 	return nullptr;
 }
diff --git a/Engine/Plugins/Animation/ControlRig/Source/ControlRig/Private/Rigs/RigHierarchyController.cpp b/Engine/Plugins/Animation/ControlRig/Source/ControlRig/Private/Rigs/RigHierarchyController.cpp
--- a/Engine/Plugins/Animation/ControlRig/Source/ControlRig/Private/Rigs/RigHierarchyController.cpp
+++ b/Engine/Plugins/Animation/ControlRig/Source/ControlRig/Private/Rigs/RigHierarchyController.cpp
@@ -23,21 +23,30 @@ URigHierarchyController::~URigHierarchyController()
 
 void URigHierarchyController::SetHierarchy(URigHierarchy* InHierarchy)
 {
-	if(Hierarchy)
+	// since we changed the controller to be a property of the hierarchy,
+	// controlling a different hierarchy is no longer allowed
+	if (ensure(InHierarchy == GetOuter()))
 	{
-		if(!Hierarchy->HasAnyFlags(RF_BeginDestroyed) && Hierarchy->IsValidLowLevel())
+		// make sure making multiple valid SetHieararchy() calls won't lead to accumulated delegates
+		// though it should not happen in the first place
+		if (Hierarchy.IsValid())
 		{
-			Hierarchy->OnModified().RemoveAll(this);
+			if(!Hierarchy->HasAnyFlags(RF_BeginDestroyed) && Hierarchy->IsValidLowLevel())
+			{
+				  Hierarchy->OnModified().RemoveAll(this);
+			}
+		}
+		
+		URigHierarchy* OuterHierarchy = Cast<URigHierarchy>(GetOuter());
+		if (ensure(OuterHierarchy))
+		{
+			Hierarchy = OuterHierarchy;
+			Hierarchy->OnModified().AddUObject(this, &URigHierarchyController::HandleHierarchyModified);
 		}
-		Hierarchy->LastControllerPtr.Reset();
 	}
-
-	Hierarchy = InHierarchy;
-
-	if(Hierarchy)
+	else
 	{
-		Hierarchy->OnModified().AddUObject(this, &URigHierarchyController::HandleHierarchyModified);
-		Hierarchy->LastControllerPtr = this;
+		UE_LOG(LogControlRig, Error, TEXT("Invalid API Usage, Called URigHierarchyController::SetHierarchy(...) with a Hierarchy that is not the outer of the controller"));
 	}
 }
 
@@ -2523,7 +2532,7 @@ void URigHierarchyController::ReportWarning(const FString& InMessage) const
 	}
 
 	FString Message = InMessage;
-	if (Hierarchy)
+	if (Hierarchy.IsValid())
 	{
 		if (UPackage* Package = Cast<UPackage>(Hierarchy->GetOutermost()))
 		{
@@ -2548,7 +2557,7 @@ void URigHierarchyController::ReportError(const FString& InMessage) const
 	}
 
 	FString Message = InMessage;
-	if (Hierarchy)
+	if (Hierarchy.IsValid())
 	{
 		if (UPackage* Package = Cast<UPackage>(Hierarchy->GetOutermost()))
 		{
diff --git a/Engine/Plugins/Animation/ControlRig/Source/ControlRig/Public/Rigs/AdditiveControlRig.h b/Engine/Plugins/Animation/ControlRig/Source/ControlRig/Public/Rigs/AdditiveControlRig.h
--- a/Engine/Plugins/Animation/ControlRig/Source/ControlRig/Public/Rigs/AdditiveControlRig.h
+++ b/Engine/Plugins/Animation/ControlRig/Source/ControlRig/Public/Rigs/AdditiveControlRig.h
@@ -29,9 +29,6 @@ private:
 	// custom units that are used to execute this rig
 	TArray<FRigUnit_AddBoneTransform> AddBoneRigUnits;
 
-	UPROPERTY(transient)
-	TObjectPtr<URigHierarchyController> Controller;
-
 	/** Create RigElements - bone hierarchy and curves - from incoming skeleton */
 	void CreateRigElements(const USkeletalMesh* InReferenceMesh);
 	void CreateRigElements(const FReferenceSkeleton& InReferenceSkeleton, const FSmartNameMapping* InSmartNameMapping);
diff --git a/Engine/Plugins/Animation/ControlRig/Source/ControlRig/Public/Rigs/FKControlRig.h b/Engine/Plugins/Animation/ControlRig/Source/ControlRig/Public/Rigs/FKControlRig.h
--- a/Engine/Plugins/Animation/ControlRig/Source/ControlRig/Public/Rigs/FKControlRig.h
+++ b/Engine/Plugins/Animation/ControlRig/Source/ControlRig/Public/Rigs/FKControlRig.h
@@ -70,8 +70,5 @@ private:
 	UPROPERTY()
 	EControlRigFKRigExecuteMode ApplyMode;
 
-	UPROPERTY(transient)
-	TObjectPtr<URigHierarchyController> Controller;
-
 	friend class FControlRigInteractionTest;
 };
\ No newline at end of file
diff --git a/Engine/Plugins/Animation/ControlRig/Source/ControlRig/Public/Rigs/RigHierarchy.h b/Engine/Plugins/Animation/ControlRig/Source/ControlRig/Public/Rigs/RigHierarchy.h
--- a/Engine/Plugins/Animation/ControlRig/Source/ControlRig/Public/Rigs/RigHierarchy.h
+++ b/Engine/Plugins/Animation/ControlRig/Source/ControlRig/Public/Rigs/RigHierarchy.h
@@ -2603,7 +2603,9 @@ private:
 	FRigHierarchyUndoRedoTransformEvent UndoRedoEvent;
 
 	TWeakObjectPtr<URigHierarchy> HierarchyForSelectionPtr;
-	TWeakObjectPtr<UObject> LastControllerPtr;
+
+	UPROPERTY(Transient)
+	TObjectPtr<URigHierarchyController> HierarchyController;
 
 	TMap<FRigElementKey, FRigElementKey> PreviousParentMap;
 	TMap<FRigElementKey, FRigElementKey> PreviousNameMap;
diff --git a/Engine/Plugins/Animation/ControlRig/Source/ControlRig/Public/Rigs/RigHierarchyController.h b/Engine/Plugins/Animation/ControlRig/Source/ControlRig/Public/Rigs/RigHierarchyController.h
--- a/Engine/Plugins/Animation/ControlRig/Source/ControlRig/Public/Rigs/RigHierarchyController.h
+++ b/Engine/Plugins/Animation/ControlRig/Source/ControlRig/Public/Rigs/RigHierarchyController.h
@@ -26,7 +26,10 @@ public:
 
 	// Returns the hierarchy currently linked to this controller
 	UFUNCTION(BlueprintCallable, Category = URigHierarchyController)
-	URigHierarchy* GetHierarchy() const { return Hierarchy; }
+	URigHierarchy* GetHierarchy() const
+	{
+		return Hierarchy.Get();
+	}
 
 	// Sets the hierarchy currently linked to this controller
 	UFUNCTION(BlueprintCallable, Category = URigHierarchyController)
@@ -517,7 +520,7 @@ public:
 private:
 
 	UPROPERTY(transient)
-	TObjectPtr<URigHierarchy> Hierarchy;
+	TWeakObjectPtr<URigHierarchy> Hierarchy;
 
 	FRigHierarchyModifiedEvent ModifiedEvent;
 	void Notify(ERigHierarchyNotification InNotifType, const FRigBaseElement* InElement);
@@ -527,7 +530,7 @@ private:
 	 * Returns true if this controller is valid / linked to a valid hierarchy.
 	 * @return Returns true if this controller is valid / linked to a valid hierarchy.
 	 */
-	FORCEINLINE bool IsValid() const { return Hierarchy != nullptr; }
+	FORCEINLINE bool IsValid() const { return Hierarchy.IsValid(); }
 
 	/**
 	 * Adds a new element to the hierarchy
diff --git a/Engine/Plugins/Animation/ControlRig/Source/ControlRig/Public/Units/RigUnitTest.h b/Engine/Plugins/Animation/ControlRig/Source/ControlRig/Public/Units/RigUnitTest.h
--- a/Engine/Plugins/Animation/ControlRig/Source/ControlRig/Public/Units/RigUnitTest.h
+++ b/Engine/Plugins/Animation/ControlRig/Source/ControlRig/Public/Units/RigUnitTest.h
@@ -22,7 +22,7 @@ public:
 	{
 		if (Hierarchy)
 		{
-			Controller->RemoveFromRoot();
+			// we no longer add/remove the controller to/from root since controller is now part of the hierarchy
 			Hierarchy->RemoveFromRoot();
 		}
 	}
@@ -32,11 +32,10 @@ public:
 		if (!Hierarchy)
 		{
 			Hierarchy = NewObject<URigHierarchy>();
-			Controller = NewObject<URigHierarchyController>();
-			Controller->SetHierarchy(Hierarchy);
+			Controller = Hierarchy->GetController(true);
 
 			Hierarchy->AddToRoot();
-			Controller->AddToRoot();
+			// we no longer add the controller to root since controller is now part of the hierarchy
 
 			Context.Hierarchy = Hierarchy;;
 			ExecuteContext.Hierarchy = Hierarchy;
diff --git a/Engine/Plugins/Animation/ControlRig/Source/ControlRigDeveloper/Private/ControlRigBlueprint.cpp b/Engine/Plugins/Animation/ControlRig/Source/ControlRigDeveloper/Private/ControlRigBlueprint.cpp
--- a/Engine/Plugins/Animation/ControlRig/Source/ControlRigDeveloper/Private/ControlRigBlueprint.cpp
+++ b/Engine/Plugins/Animation/ControlRig/Source/ControlRigDeveloper/Private/ControlRigBlueprint.cpp
@@ -106,10 +106,9 @@ UControlRigBlueprint::UControlRigBlueprint(const FObjectInitializer& ObjectIniti
 #endif
 
 	Hierarchy = CreateDefaultSubobject<URigHierarchy>(TEXT("Hierarchy"));
-	HierarchyController = CreateDefaultSubobject<URigHierarchyController>(TEXT("HierarchyController"));
-	HierarchyController->SetHierarchy(Hierarchy);
+	URigHierarchyController* Controller = Hierarchy->GetController(true);
 	// give BP a chance to propagate hierarchy changes to available control rig instances
-	HierarchyController->OnModified().AddUObject(this, &UControlRigBlueprint::HandleHierarchyModified);
+	Controller->OnModified().AddUObject(this, &UControlRigBlueprint::HandleHierarchyModified);
 }
 
 UControlRigBlueprint::UControlRigBlueprint()
@@ -1344,7 +1343,7 @@ TArray<FString> UControlRigBlueprint::GeneratePythonCommands(const FString InNew
 			, *InNewBlueprintName));
 
 	// Hierarchy
-	Commands.Append(HierarchyController->GeneratePythonCommands());
+	Commands.Append(Hierarchy->GetController(true)->GeneratePythonCommands());
 		
 	// Add variables
 	for (const FBPVariableDescription& Variable : NewVariables)
@@ -1848,14 +1847,11 @@ void UControlRigBlueprint::PostDuplicate(bool bDuplicateForPIE)
 {
 	Super::PostDuplicate(bDuplicateForPIE);
 	
-	if(HierarchyController)
+	if (URigHierarchyController* Controller = Hierarchy->GetController(true))
 	{
-		HierarchyController->OnModified().RemoveAll(this);
+		Controller->OnModified().RemoveAll(this);
+		Controller->OnModified().AddUObject(this, &UControlRigBlueprint::HandleHierarchyModified);
 	}
-	
-	HierarchyController = NewObject<URigHierarchyController>(this);
-	HierarchyController->SetHierarchy(Hierarchy);
-	HierarchyController->OnModified().AddUObject(this, &UControlRigBlueprint::HandleHierarchyModified);
 
 	FBlueprintEditorUtils::MarkBlueprintAsStructurallyModified(this);
 }
diff --git a/Engine/Plugins/Animation/ControlRig/Source/ControlRigDeveloper/Public/ControlRigBlueprint.h b/Engine/Plugins/Animation/ControlRig/Source/ControlRigDeveloper/Public/ControlRigBlueprint.h
--- a/Engine/Plugins/Animation/ControlRig/Source/ControlRigDeveloper/Public/ControlRigBlueprint.h
+++ b/Engine/Plugins/Animation/ControlRig/Source/ControlRigDeveloper/Public/ControlRigBlueprint.h
@@ -448,7 +448,7 @@ public:
 	TObjectPtr<URigHierarchy> Hierarchy;
 
 	UFUNCTION(BlueprintCallable, Category = "Hierarchy")
-	URigHierarchyController* GetHierarchyController() { return HierarchyController; }
+	URigHierarchyController* GetHierarchyController() { return Hierarchy->GetController(true); }
 
 private:
 
@@ -494,9 +494,6 @@ private:
 	UPROPERTY(transient, DuplicateTransient)
 	int32 VMRecompilationBracket;
 
-	UPROPERTY(transient, DuplicateTransient)
-	TObjectPtr<URigHierarchyController> HierarchyController;
-
 	FRigVMGraphModifiedEvent ModifiedEvent;
 	void Notify(ERigVMGraphNotifType InNotifType, UObject* InSubject);
 	void HandleModifiedEvent(ERigVMGraphNotifType InNotifType, URigVMGraph* InGraph, UObject* InSubject);
diff --git a/Engine/Plugins/Animation/ControlRig/Source/ControlRigEditor/Private/Editor/ControlRigEditor.cpp b/Engine/Plugins/Animation/ControlRig/Source/ControlRigEditor/Private/Editor/ControlRigEditor.cpp
--- a/Engine/Plugins/Animation/ControlRig/Source/ControlRigEditor/Private/Editor/ControlRigEditor.cpp
+++ b/Engine/Plugins/Animation/ControlRig/Source/ControlRigEditor/Private/Editor/ControlRigEditor.cpp
@@ -1029,7 +1029,7 @@ void FControlRigEditor::ToggleSetupMode()
 		// need to clear selection before remove transient control
 		// because active selection will trigger transient control recreation after removal	
 		PreviousSelection = RigBlueprint->Hierarchy->GetSelectedKeys();
-		RigBlueprint->HierarchyController->ClearSelection();
+		RigBlueprint->GetHierarchyController()->ClearSelection();
 		
 		// need to copy here since the removal changes the iterator
 		if (ControlRig)
diff --git a/Engine/Plugins/Animation/RigLogic/Source/RigLogicModule/Private/RigUnit_RigLogic_Test.cpp b/Engine/Plugins/Animation/RigLogic/Source/RigLogicModule/Private/RigUnit_RigLogic_Test.cpp
--- a/Engine/Plugins/Animation/RigLogic/Source/RigLogicModule/Private/RigUnit_RigLogic_Test.cpp
+++ b/Engine/Plugins/Animation/RigLogic/Source/RigLogicModule/Private/RigUnit_RigLogic_Test.cpp
@@ -48,8 +48,7 @@ TSharedPtr<TestBehaviorReader> FRigUnit_RigLogic::TestAccessor::CreateBehaviorRe
 TStrongObjectPtr<URigHierarchy> FRigUnit_RigLogic::TestAccessor::CreateCurveContainerOneCurve(FString CurveNameStr)
 {
 	TStrongObjectPtr<URigHierarchy> ValidCurveContainer(NewObject<URigHierarchy>());
-	URigHierarchyController* Controller = NewObject<URigHierarchyController>();
-	Controller->SetHierarchy(ValidCurveContainer.Get());
+	URigHierarchyController* Controller = ValidCurveContainer->GetController(true);
 	Controller->AddCurve(FName(*CurveNameStr));
 	return ValidCurveContainer;
 }
@@ -70,8 +69,7 @@ TStrongObjectPtr<URigHierarchy> FRigUnit_RigLogic::TestAccessor::CreateBoneHiera
 TStrongObjectPtr<URigHierarchy> FRigUnit_RigLogic::TestAccessor::CreateBoneHierarchyTwoBones(FString Bone1NameStr, FString Bone2NameStr)
 {
 	TStrongObjectPtr<URigHierarchy> TestHierarchy(NewObject<URigHierarchy>());
-	URigHierarchyController* Controller = NewObject<URigHierarchyController>();
-	Controller->SetHierarchy(TestHierarchy.Get());
+	URigHierarchyController* Controller = TestHierarchy->GetController(true);
 	const FRigElementKey Bone1Key = Controller->AddBone(*Bone1NameStr, FRigElementKey(), FTransform(FVector(1.f, 0.f, 0.f)), true, ERigBoneType::User);
 	Controller->AddBone(*Bone2NameStr, Bone1Key, FTransform(FVector(1.f, 2.f, 3.f)), true, ERigBoneType::User);
 	return TestHierarchy;
@@ -130,8 +128,7 @@ TSharedPtr<TestBehaviorReader> FRigUnit_RigLogic::TestAccessor::CreateBehaviorRe
 TStrongObjectPtr<URigHierarchy> FRigUnit_RigLogic::TestAccessor::CreateCurveContainerOneMorphTarget(FString MorphTargetStr)
 {
 	TStrongObjectPtr<URigHierarchy> ValidCurveContainer(NewObject<URigHierarchy>());
-	URigHierarchyController* Controller = NewObject<URigHierarchyController>();
-	Controller->SetHierarchy(ValidCurveContainer.Get());
+	URigHierarchyController* Controller = ValidCurveContainer->GetController(true);
 	Controller->AddCurve(FName(*MorphTargetStr));
 	return ValidCurveContainer;
 }
@@ -139,8 +136,7 @@ TStrongObjectPtr<URigHierarchy> FRigUnit_RigLogic::TestAccessor::CreateCurveCont
 TStrongObjectPtr<URigHierarchy> FRigUnit_RigLogic::TestAccessor::CreateCurveContainerTwoMorphTargets(FString MorphTarget1Str, FString MorphTarget2Str)
 {
 	TStrongObjectPtr<URigHierarchy> ValidCurveContainer(NewObject<URigHierarchy>());
-	URigHierarchyController* Controller = NewObject<URigHierarchyController>();
-	Controller->SetHierarchy(ValidCurveContainer.Get());
+	URigHierarchyController* Controller = ValidCurveContainer->GetController(true);
 	Controller->AddCurve(FName(*MorphTarget1Str));
 	Controller->AddCurve(FName(*MorphTarget2Str));
 	return ValidCurveContainer;
