Sequencer: Space Switching: Add Events to for key moving, deletion, and adding which is needed for Space Switching workflows where moving Space Keys will also move corresponding transform keys and delete Space keys will auto-bake transform keys. Adding added for completeness. Events are fired by the channel, but usually come from the channel data so we added a constructor.
Space Channel is only one currently using this.

Also no added how these events are used by the TrackEditor to handle the events. Note since MovieSceneChannel is just a USTRUCT we can't use dynamic delegates on it and so on undo we lose anything that was bound to it, so the track editor now rebiinds the events on any space channel on undo/redo.

diff --git a/Engine/Plugins/Animation/ControlRig/Source/ControlRig/Private/Sequencer/MovieSceneControlRigParameterSection.cpp b/Engine/Plugins/Animation/ControlRig/Source/ControlRig/Private/Sequencer/MovieSceneControlRigParameterSection.cpp
--- a/Engine/Plugins/Animation/ControlRig/Source/ControlRig/Private/Sequencer/MovieSceneControlRigParameterSection.cpp
+++ b/Engine/Plugins/Animation/ControlRig/Source/ControlRig/Private/Sequencer/MovieSceneControlRigParameterSection.cpp
@@ -1151,6 +1151,11 @@ void UMovieSceneControlRigParameterSection::AddSpaceChannel(FName InControlName,
 	if (!HasSpaceChannel(InControlName))
 	{
 		SpaceChannels.Add(FSpaceControlNameAndChannel(InControlName));
+		if (OnSpaceChannelAdded.IsBound())
+		{
+			FSpaceControlNameAndChannel& NameAndChannel = SpaceChannels[SpaceChannels.Num() - 1];
+			OnSpaceChannelAdded.Broadcast(this, &NameAndChannel.SpaceCurve);
+		}
 	}
 	if (bReconstructChannel)
 	{
diff --git a/Engine/Plugins/Animation/ControlRig/Source/ControlRig/Private/Sequencer/MovieSceneControlRigParameterTrack.cpp b/Engine/Plugins/Animation/ControlRig/Source/ControlRig/Private/Sequencer/MovieSceneControlRigParameterTrack.cpp
--- a/Engine/Plugins/Animation/ControlRig/Source/ControlRig/Private/Sequencer/MovieSceneControlRigParameterTrack.cpp
+++ b/Engine/Plugins/Animation/ControlRig/Source/ControlRig/Private/Sequencer/MovieSceneControlRigParameterTrack.cpp
@@ -48,6 +48,7 @@ UMovieSceneSection* UMovieSceneControlRigParameterTrack::CreateNewSection()
 	{
 		NewSection->SetBlendType(EMovieSceneBlendType::Additive);
 	}
+	NewSection->SpaceChannelAdded().AddUObject(this, &UMovieSceneControlRigParameterTrack::HandleOnSpaceAdded);
 	if (ControlRig)
 	{
 		NewSection->RecreateWithThisControlRig(ControlRig,bSetDefault);
@@ -55,6 +56,11 @@ UMovieSceneSection* UMovieSceneControlRigParameterTrack::CreateNewSection()
 	return  NewSection;
 }
 
+void UMovieSceneControlRigParameterTrack::HandleOnSpaceAdded(UMovieSceneControlRigParameterSection* Section, FMovieSceneControlRigSpaceChannel* Channel)
+{
+	OnSpaceChannelAdded.Broadcast(Section, Channel);
+}
+
 void UMovieSceneControlRigParameterTrack::RemoveAllAnimationData()
 {
 	Sections.Empty();
@@ -342,6 +348,10 @@ void UMovieSceneControlRigParameterTrack::ReconstructControlRig()
 				UMovieSceneControlRigParameterSection* CRSection = Cast<UMovieSceneControlRigParameterSection>(Section);
 				if (CRSection)
 				{
+					if (CRSection->SpaceChannelAdded().IsBoundToObject(this) == false)
+					{
+						CRSection->SpaceChannelAdded().AddUObject(this, &UMovieSceneControlRigParameterTrack::HandleOnSpaceAdded);
+					}
 					CRSection->RecreateWithThisControlRig(ControlRig, false);
 				}
 			}
@@ -363,6 +373,7 @@ void UMovieSceneControlRigParameterTrack::PostEditImport()
 	ReconstructControlRig();
 }
 
+
 CONTROLRIG_API void UMovieSceneControlRigParameterTrack::ReplaceControlRig(UControlRig* NewControlRig, bool RecreateChannels)
 {
 	ControlRig = NewControlRig;
diff --git a/Engine/Plugins/Animation/ControlRig/Source/ControlRig/Public/Sequencer/MovieSceneControlRigParameterSection.h b/Engine/Plugins/Animation/ControlRig/Source/ControlRig/Public/Sequencer/MovieSceneControlRigParameterSection.h
--- a/Engine/Plugins/Animation/ControlRig/Source/ControlRig/Public/Sequencer/MovieSceneControlRigParameterSection.h
+++ b/Engine/Plugins/Animation/ControlRig/Source/ControlRig/Public/Sequencer/MovieSceneControlRigParameterSection.h
@@ -155,6 +155,9 @@ struct CONTROLRIG_API FChannelMapInfo
 
 };
 
+
+struct FMovieSceneControlRigSpaceChannel;
+
 /**
  * Movie scene section that controls animation controller animation
  */
@@ -165,6 +168,9 @@ class CONTROLRIG_API UMovieSceneControlRigParameterSection : public UMovieSceneP
 
 public:
 
+	/** Bindable event for when we add a space channel*/
+	DECLARE_MULTICAST_DELEGATE_TwoParams(FSpaceChannelAddedEvent, UMovieSceneControlRigParameterSection*, FMovieSceneControlRigSpaceChannel*);
+
 	void AddEnumParameterKey(FName InParameterName, FFrameNumber InTime, uint8 InValue);
 	void AddIntegerParameterKey(FName InParameterName, FFrameNumber InTime, int32 InValue);
 
@@ -180,8 +186,12 @@ public:
 	TArray<FSpaceControlNameAndChannel>& GetSpaceChannels();
 	const TArray< FSpaceControlNameAndChannel>& GetSpaceChannels() const;
 	FName FindControlNameFromSpaceChannel(const FMovieSceneControlRigSpaceChannel* SpaceChannel) const;
+	
+	FSpaceChannelAddedEvent& SpaceChannelAdded() { return OnSpaceChannelAdded; }
 private:
 
+	FSpaceChannelAddedEvent OnSpaceChannelAdded;
+
 	/** Control Rig that controls us*/
 	UPROPERTY()
 	TObjectPtr<UControlRig> ControlRig;
diff --git a/Engine/Plugins/Animation/ControlRig/Source/ControlRig/Public/Sequencer/MovieSceneControlRigParameterTrack.h b/Engine/Plugins/Animation/ControlRig/Source/ControlRig/Public/Sequencer/MovieSceneControlRigParameterTrack.h
--- a/Engine/Plugins/Animation/ControlRig/Source/ControlRig/Public/Sequencer/MovieSceneControlRigParameterTrack.h
+++ b/Engine/Plugins/Animation/ControlRig/Source/ControlRig/Public/Sequencer/MovieSceneControlRigParameterTrack.h
@@ -9,6 +9,7 @@
 #include "ControlRig.h"
 #include "Compilation/IMovieSceneTrackTemplateProducer.h"
 #include "INodeAndChannelMappings.h"
+#include "MovieSceneControlRigParameterSection.h"
 #include "MovieSceneControlRigParameterTrack.generated.h"
 
 /**
@@ -107,8 +108,16 @@ public:
 
 	CONTROLRIG_API void SetTrackName(FName InName) { TrackName = InName; }
 
+	UMovieSceneControlRigParameterSection::FSpaceChannelAddedEvent& SpaceChannelAdded() { return OnSpaceChannelAdded; }
+
 private:
 
+	//we register this with sections.
+	void HandleOnSpaceAdded(UMovieSceneControlRigParameterSection* Section, FMovieSceneControlRigSpaceChannel* Channel);
+	//then send this event out to the track editor.
+	UMovieSceneControlRigParameterSection::FSpaceChannelAddedEvent OnSpaceChannelAdded;
+
+
 	void ReconstructControlRig();
 
 private:
diff --git a/Engine/Plugins/Animation/ControlRig/Source/ControlRig/Public/Sequencer/MovieSceneControlRigSpaceChannel.h b/Engine/Plugins/Animation/ControlRig/Source/ControlRig/Public/Sequencer/MovieSceneControlRigSpaceChannel.h
--- a/Engine/Plugins/Animation/ControlRig/Source/ControlRig/Public/Sequencer/MovieSceneControlRigSpaceChannel.h
+++ b/Engine/Plugins/Animation/ControlRig/Source/ControlRig/Public/Sequencer/MovieSceneControlRigSpaceChannel.h
@@ -70,7 +70,7 @@ struct CONTROLRIG_API FMovieSceneControlRigSpaceChannel : public FMovieSceneChan
 	*/
 	FORCEINLINE TMovieSceneChannelData<FMovieSceneControlRigSpaceBaseKey> GetData()
 	{
-		return TMovieSceneChannelData<FMovieSceneControlRigSpaceBaseKey>(&KeyTimes, &KeyValues, &KeyHandles);
+		return TMovieSceneChannelData<FMovieSceneControlRigSpaceBaseKey>(&KeyTimes, &KeyValues, &KeyHandles, this);
 	}
 
 	/**
@@ -121,6 +121,7 @@ private:
 	FMovieSceneKeyHandleMap KeyHandles;
 };
 
+
 template<>
 struct TMovieSceneChannelTraits<FMovieSceneControlRigSpaceChannel> : TMovieSceneChannelTraitsBase<FMovieSceneControlRigSpaceChannel>
 {
diff --git a/Engine/Plugins/Animation/ControlRig/Source/ControlRigEditor/Private/Sequencer/ControlRigParameterTrackEditor.cpp b/Engine/Plugins/Animation/ControlRig/Source/ControlRigEditor/Private/Sequencer/ControlRigParameterTrackEditor.cpp
--- a/Engine/Plugins/Animation/ControlRig/Source/ControlRigEditor/Private/Sequencer/ControlRigParameterTrackEditor.cpp
+++ b/Engine/Plugins/Animation/ControlRig/Source/ControlRigEditor/Private/Sequencer/ControlRigParameterTrackEditor.cpp
@@ -77,6 +77,8 @@
 #include "Editor/UnrealEdEngine.h"
 #include "Toolkits/IToolkitHost.h"
 #include "ControlRigEditModeSettings.h"
+#include "ControlRigSpaceChannelEditors.h"
+
 #define LOCTEXT_NAMESPACE "FControlRigParameterTrackEditor"
 
 static USkeletalMeshComponent* AcquireSkeletalMeshFromObject(UObject* BoundObject, TSharedPtr<ISequencer> SequencerPtr)
@@ -182,6 +184,11 @@ FControlRigParameterTrackEditor::FControlRigParameterTrackEditor(TSharedRef<ISeq
 {
 	FMovieSceneToolsModule::Get().RegisterAnimationBakeHelper(this);
 
+	if (GEditor != nullptr)
+	{
+		GEditor->RegisterForUndo(this);
+	}
+
 	UMovieScene* MovieScene = InSequencer->GetFocusedMovieSceneSequence()->GetMovieScene();
 
 	SelectionChangedHandle = InSequencer->GetSelectionChangedTracks().AddRaw(this, &FControlRigParameterTrackEditor::OnSelectionChanged);
@@ -309,6 +316,11 @@ FControlRigParameterTrackEditor::FControlRigParameterTrackEditor(TSharedRef<ISeq
 
 FControlRigParameterTrackEditor::~FControlRigParameterTrackEditor()
 {
+	if (GEditor != nullptr)
+	{
+		GEditor->UnregisterForUndo(this);
+	}
+
 	UnbindAllControlRigs();
 	if (GetSequencer().IsValid())
 	{
@@ -326,12 +338,22 @@ void FControlRigParameterTrackEditor::BindControlRig(UControlRig* ControlRig)
 		ControlRig->OnInitialized_AnyThread().AddRaw(this, &FControlRigParameterTrackEditor::HandleOnInitialized);
 		ControlRig->ControlSelected().AddRaw(this, &FControlRigParameterTrackEditor::HandleControlSelected);
 		BoundControlRigs.Add(ControlRig);
+		UMovieSceneControlRigParameterTrack* Track = FindTrack(ControlRig);
+		if (Track)
+		{
+			Track->SpaceChannelAdded().AddRaw(this, &FControlRigParameterTrackEditor::HandleOnSpaceAdded);
+		}
 	}
 }
 void FControlRigParameterTrackEditor::UnbindControlRig(UControlRig* ControlRig)
 {
 	if (ControlRig && BoundControlRigs.Contains(ControlRig) == true)
 	{
+		UMovieSceneControlRigParameterTrack* Track = FindTrack(ControlRig);
+		if (Track)
+		{
+			Track->SpaceChannelAdded().RemoveAll(this);
+		}
 		ControlRig->ControlModified().RemoveAll(this);
 		ControlRig->OnInitialized_AnyThread().RemoveAll(this);
 		ControlRig->ControlSelected().RemoveAll(this);
@@ -1910,6 +1932,114 @@ FMovieSceneTrackEditor::FFindOrCreateTrackResult FControlRigParameterTrackEditor
 	return Result;
 }
 
+UMovieSceneControlRigParameterTrack* FControlRigParameterTrackEditor::FindTrack(UControlRig* InControlRig)
+{	
+	UMovieScene* MovieScene = GetSequencer()->GetFocusedMovieSceneSequence()->GetMovieScene();
+	if (MovieScene)
+	{
+		const TArray<FMovieSceneBinding>& Bindings = MovieScene->GetBindings();
+		for (const FMovieSceneBinding& Binding : Bindings)
+		{
+			TArray<UMovieSceneTrack*> Tracks = MovieScene->FindTracks(UMovieSceneControlRigParameterTrack::StaticClass(), Binding.GetObjectGuid(), NAME_None);
+			for (UMovieSceneTrack* AnyOleTrack : Tracks)
+			{
+				UMovieSceneControlRigParameterTrack* Track = Cast<UMovieSceneControlRigParameterTrack>(AnyOleTrack);
+				if (Track && Track->GetControlRig() == InControlRig)
+				{
+					return Track;
+				}
+			}
+		}
+	}
+	return nullptr;
+}
+
+void FControlRigParameterTrackEditor::HandleOnSpaceAdded(UMovieSceneControlRigParameterSection* Section, FMovieSceneControlRigSpaceChannel* SpaceChannel)
+{
+	if (SpaceChannel)
+	{
+		SpaceChannel->OnKeyMovedEvent().AddLambda([this, SpaceChannel,Section](FMovieSceneChannel* Channel, const  TArray<FKeyMoveEventItem>& MovedItems)
+			{
+				HandleSpaceKeyMoved(Section, SpaceChannel,MovedItems);
+			});
+		SpaceChannel->OnKeyDeletedEvent().AddLambda([this,SpaceChannel, Section](FMovieSceneChannel* Channel, const  TArray<FKeyAddOrDeleteEventItem>& Items)
+			{
+				HandleSpaceKeyDeleted(Section, SpaceChannel,Items);
+			});
+	}
+	//todoo do we need to remove this or not mz
+}
+
+bool FControlRigParameterTrackEditor::MatchesContext(const FTransactionContext& InContext, const TArray<TPair<UObject*, FTransactionObjectEvent>>& TransactionObjects) const
+{
+	SectionsGettingUndone.SetNum(0);
+	// Check if we care about the undo/redo
+	bool bGettingUndone = false;
+	for (const TPair<UObject*, FTransactionObjectEvent>& TransactionObjectPair : TransactionObjects)
+	{
+
+		UObject* Object = TransactionObjectPair.Key;
+		while (Object != nullptr)
+		{
+			if (Object->GetClass()->IsChildOf(UMovieSceneControlRigParameterSection::StaticClass()))
+			{
+				UMovieSceneControlRigParameterSection* Section = Cast< UMovieSceneControlRigParameterSection>(Object);
+				if (Section)
+				{
+					SectionsGettingUndone.Add(Section);
+				}
+				bGettingUndone =  true;
+				break;
+			}
+			Object = Object->GetOuter();
+		}
+	}
+	
+	return bGettingUndone;
+}
+
+void FControlRigParameterTrackEditor::PostUndo(bool bSuccess)
+{
+	for (UMovieSceneControlRigParameterSection* Section : SectionsGettingUndone)
+	{
+		if (Section->GetControlRig())
+		{
+			TArray<FSpaceControlNameAndChannel>& SpaceChannels = Section->GetSpaceChannels();
+			for(FSpaceControlNameAndChannel& Channel: SpaceChannels)
+			{ 
+				HandleOnSpaceAdded(Section, &(Channel.SpaceCurve));
+			}
+		}
+	}
+}
+
+
+void FControlRigParameterTrackEditor::HandleSpaceKeyDeleted(UMovieSceneControlRigParameterSection* Section, FMovieSceneControlRigSpaceChannel* Channel, const TArray<FKeyAddOrDeleteEventItem>& DeletedItems)
+{
+	const TSharedPtr<ISequencer> ParentSequencer = GetSequencer();
+
+	if (Section && Section->GetControlRig() && Channel && ParentSequencer.IsValid())
+	{
+		FName ControlName = Section->FindControlNameFromSpaceChannel(Channel);
+		for (const FKeyAddOrDeleteEventItem& EventItem : DeletedItems)
+		{
+			FControlRigSpaceChannelHelpers::SequencerSpaceChannelKeyDeleted(Section->GetControlRig(), ParentSequencer.Get(), ControlName, Channel, Section,EventItem.Frame);
+		}
+	}
+}
+
+void FControlRigParameterTrackEditor::HandleSpaceKeyMoved(UMovieSceneControlRigParameterSection* Section, FMovieSceneControlRigSpaceChannel* SpaceChannel, const  TArray<FKeyMoveEventItem>& MovedItems)
+{
+	if (Section && Section->GetControlRig() && SpaceChannel)
+	{
+		FName ControlName = Section->FindControlNameFromSpaceChannel(SpaceChannel);
+		for (const FKeyMoveEventItem& MoveEventItem : MovedItems)
+		{
+			FControlRigSpaceChannelHelpers::HandleSpaceKeyTimeChanged(Section->GetControlRig(), ControlName, SpaceChannel, Section,
+				MoveEventItem.Frame, MoveEventItem.NewFrame);
+		}
+	}
+}
 
 void FControlRigParameterTrackEditor::HandleControlSelected(UControlRig* Subject, FRigControlElement* ControlElement, bool bSelected)
 {
@@ -2039,7 +2169,17 @@ void FControlRigParameterTrackEditor::HandleControlModified(UControlRig* Control
 					{
 						KeyMode = ESequencerKeyMode::ManualKeyForced;
 					}
-					AddControlKeys(Component, ControlRig, Name, ControlElement->GetName(), (EControlRigContextChannelToKey)Context.KeyMask, KeyMode, Context.LocalTime);
+					AddControlKeys(Component, ControlRig, Name, ControlElement->GetName(), (EControlRigContextChannelToKey)Context.KeyMask, 
+						KeyMode, Context.LocalTime);
+				}
+				UMovieSceneControlRigParameterSection* ParamSection = Cast<UMovieSceneControlRigParameterSection>(Track->GetSectionToKey());
+				if (ParamSection)  
+				{
+					
+					FFrameNumber KeyTime = (Context.LocalTime == FLT_MAX) ? GetTimeForKey() : GetSequencer()->GetFocusedTickResolution().AsFrameTime((double)Context.LocalTime).RoundToFrame();
+
+					FControlRigSpaceChannelHelpers::CompensatePreviousFrameIfNeeded(ControlRig, GetSequencer().Get(), ParamSection,
+						ControlElement->GetName(), KeyTime);
 				}
 			}
 		}
@@ -2431,7 +2571,8 @@ bool FControlRigParameterTrackEditor::ModifyOurGeneratedKeysByCurrentAndWeight(U
 					if (pChannelIndex)
 					{
 						ChannelIndex = pChannelIndex->TotalChannelIndex;
-						GeneratedTotalKeys[ChannelIndex]->ModifyByCurrentAndWeight(Proxy, KeyTime, (void *)&Val.Val, Weight);
+						float FVal = (float)Val.Val;
+						GeneratedTotalKeys[ChannelIndex]->ModifyByCurrentAndWeight(Proxy, KeyTime, (void *)&FVal, Weight);
 					}
 					break;
 				}
@@ -2455,8 +2596,10 @@ bool FControlRigParameterTrackEditor::ModifyOurGeneratedKeysByCurrentAndWeight(U
 					if (pChannelIndex)
 					{
 						ChannelIndex = pChannelIndex->TotalChannelIndex;
-						GeneratedTotalKeys[ChannelIndex]->ModifyByCurrentAndWeight(Proxy, KeyTime, (void *)&Val.Val.X, Weight);
-						GeneratedTotalKeys[ChannelIndex + 1]->ModifyByCurrentAndWeight(Proxy, KeyTime, (void *)&Val.Val.Y, Weight);
+						float FVal = (float)Val.Val.X;
+						GeneratedTotalKeys[ChannelIndex]->ModifyByCurrentAndWeight(Proxy, KeyTime, (void *)&FVal, Weight);
+						FVal = (float)Val.Val.Y;
+						GeneratedTotalKeys[ChannelIndex + 1]->ModifyByCurrentAndWeight(Proxy, KeyTime, (void *)&FVal, Weight);
 					}
 					break;
 				}
@@ -2476,13 +2619,17 @@ bool FControlRigParameterTrackEditor::ModifyOurGeneratedKeysByCurrentAndWeight(U
 					{
 						ChannelIndex = pChannelIndex->TotalChannelIndex;
 
+						
 						/* @Mike.Zyracki why is this causing the value to continuously grow?
 						*/
 						if (ControlElement->Settings.ControlType != ERigControlType::Rotator)
 						{
-							GeneratedTotalKeys[ChannelIndex]->ModifyByCurrentAndWeight(Proxy, KeyTime, (void*)&Val.Val.X, Weight);
-							GeneratedTotalKeys[ChannelIndex + 1]->ModifyByCurrentAndWeight(Proxy, KeyTime, (void*)&Val.Val.Y, Weight);
-							GeneratedTotalKeys[ChannelIndex + 2]->ModifyByCurrentAndWeight(Proxy, KeyTime, (void*)&Val.Val.Z, Weight);
+							float FVal = (float)Val.Val.X;
+							GeneratedTotalKeys[ChannelIndex]->ModifyByCurrentAndWeight(Proxy, KeyTime, (void*)&FVal, Weight);
+							FVal = (float)Val.Val.Y;
+							GeneratedTotalKeys[ChannelIndex + 1]->ModifyByCurrentAndWeight(Proxy, KeyTime, (void*)&FVal, Weight);
+							FVal = (float)Val.Val.Z;
+							GeneratedTotalKeys[ChannelIndex + 2]->ModifyByCurrentAndWeight(Proxy, KeyTime, (void*)&FVal, Weight);
 						}
 					}
 					break;
diff --git a/Engine/Plugins/Animation/ControlRig/Source/ControlRigEditor/Private/Sequencer/ControlRigParameterTrackEditor.h b/Engine/Plugins/Animation/ControlRig/Source/ControlRigEditor/Private/Sequencer/ControlRigParameterTrackEditor.h
--- a/Engine/Plugins/Animation/ControlRig/Source/ControlRigEditor/Private/Sequencer/ControlRigParameterTrackEditor.h
+++ b/Engine/Plugins/Animation/ControlRig/Source/ControlRigEditor/Private/Sequencer/ControlRigParameterTrackEditor.h
@@ -17,6 +17,7 @@
 #include "MovieSceneToolHelpers.h"
 #include "MovieSceneToolsModule.h"
 #include "Engine/EngineTypes.h"
+#include "EditorUndoClient.h"
 
 struct FAssetData;
 class FMenuBuilder;
@@ -26,11 +27,16 @@ class UFKControlRig;
 class FEditorModeTools;
 class FControlRigEditMode;
 struct FRigControlModifiedContext;
+struct FMovieSceneControlRigSpaceChannel;
+struct FMovieSceneChannel;
+struct FKeyAddOrDeleteEventItem;
+struct FKeyMoveEventItem;
 
 /**
  * Tools for animation tracks
  */
-class FControlRigParameterTrackEditor : public FKeyframeTrackEditor<UMovieSceneControlRigParameterTrack>, public IMovieSceneToolsAnimationBakeHelper
+class FControlRigParameterTrackEditor : public FKeyframeTrackEditor<UMovieSceneControlRigParameterTrack>, public IMovieSceneToolsAnimationBakeHelper,
+	public FEditorUndoClient
 {
 public:
 	/**
@@ -71,6 +77,11 @@ public:
 	//IMovieSceneToolsAnimationBakeHelper
 	virtual void PostEvaluation(UMovieScene* MovieScene, FFrameNumber Frame);
 
+	//FEditorUndoClient Interface
+	virtual bool MatchesContext(const FTransactionContext& InContext, const TArray<TPair<UObject*, FTransactionObjectEvent>>& TransactionObjects) const override;
+	virtual void PostUndo(bool bSuccess) override;
+	virtual void PostRedo(bool bSuccess) override { PostUndo(bSuccess); }
+
 private:
 
 	void HandleAddTrackSubMenu(FMenuBuilder& MenuBuilder, TArray<FGuid> ObjectBindings, UMovieSceneTrack* Track);
@@ -118,6 +129,11 @@ private:
 	void HandleControlSelected(UControlRig* Subject, FRigControlElement* ControlElement, bool bSelected);
 	void HandleOnInitialized(UControlRig* Subject, const EControlRigState InState, const FName& InEventName);
 
+	/** SpaceChannel Delegates*/
+	void HandleOnSpaceAdded(UMovieSceneControlRigParameterSection* Section, FMovieSceneControlRigSpaceChannel* Channel);
+	void HandleSpaceKeyDeleted(UMovieSceneControlRigParameterSection* Section, FMovieSceneControlRigSpaceChannel* Channel, const TArray<FKeyAddOrDeleteEventItem>& DeletedItems);
+	void HandleSpaceKeyMoved(UMovieSceneControlRigParameterSection* Section, FMovieSceneControlRigSpaceChannel* Channel, const  TArray<FKeyMoveEventItem>& MovedItems);
+
 	/** Select control rig if not selected, select controls from key areas */
 	void SelectRigsAndControls(UControlRig* Subject, const TArray<const IKeyArea*>& KeyAreas);
 
@@ -130,6 +146,8 @@ private:
 	/** Import FBX*/
 	void ImportFBX(UMovieSceneControlRigParameterTrack* InTrack, UMovieSceneControlRigParameterSection* InSection, 
 		TArray<FFBXNodeAndChannels>* NodeAndChannels);
+	/** Find Track for given ControlRig*/
+	UMovieSceneControlRigParameterTrack* FindTrack(UControlRig* InControlRig);
 
 	/** Select Bones to Animate on FK Rig*/
 	void SelectFKBonesToAnimate(UFKControlRig* FKControlRig, UMovieSceneControlRigParameterTrack* Track);
@@ -150,7 +168,6 @@ private:
 
 	/** Command Bindings added by the Transform Track Editor to Sequencer and curve editor. */
 	TSharedPtr<FUICommandList> CommandBindings;
-
 	FAcquiredResources AcquiredResources;
 
 public:
@@ -185,6 +202,7 @@ private:
 	FDelegateHandle OnMovieSceneChannelChangedHandle;
 	FDelegateHandle OnActorAddedToSequencerHandle;
 
+	//map used for space channel delegates since we loose them
 
 	//used to sync curve editor selections/displays on next tick for performance reasons
 	TSet<FName> DisplayedControls;
@@ -206,10 +224,15 @@ private:
 	/** Whether or not we should check for Animatable Controls when filtering*/
 	bool bFilterAssetByAnimatableControls;
 
+	/** Handle to help updating selection on tick tick to avoid too many flooded selections*/
 	FTimerHandle UpdateSelectionTimerHandle;
+
+	/** Array of sections that are getting undone, we need to recreate any space channel add,move key delegates to them*/
+	mutable TArray<UMovieSceneControlRigParameterSection*> SectionsGettingUndone;
 };
 
 
+
 /** Class for control rig sections */
 class FControlRigParameterSection : public FSequencerSection
 {
diff --git a/Engine/Plugins/Animation/ControlRig/Source/ControlRigEditor/Private/Sequencer/ControlRigSpaceChannelEditors.cpp b/Engine/Plugins/Animation/ControlRig/Source/ControlRigEditor/Private/Sequencer/ControlRigSpaceChannelEditors.cpp
--- a/Engine/Plugins/Animation/ControlRig/Source/ControlRigEditor/Private/Sequencer/ControlRigSpaceChannelEditors.cpp
+++ b/Engine/Plugins/Animation/ControlRig/Source/ControlRigEditor/Private/Sequencer/ControlRigSpaceChannelEditors.cpp
@@ -340,7 +340,8 @@ FKeyHandle FControlRigSpaceChannelHelpers::SequencerKeyControlRigSpaceChannel(UC
 		{
 			ControlRig->Evaluate_AnyThread();
 			Context.LocalTime = TickResolution.AsSeconds(FFrameTime(Frame));
-			ControlRig->SetControlGlobalTransform(ControlKey.Name, ControlWorldTransforms[FramesIndex++], true, Context);
+			ControlRig->SetControlGlobalTransform(ControlKey.Name, ControlWorldTransforms[FramesIndex], true, Context);
+			FramesIndex++;
 		}
 		Sequencer->NotifyMovieSceneDataChanged(EMovieSceneDataChangeType::MovieSceneStructureItemAdded);
 
@@ -349,6 +350,133 @@ FKeyHandle FControlRigSpaceChannelHelpers::SequencerKeyControlRigSpaceChannel(UC
 
 }
 
+void  FControlRigSpaceChannelHelpers::SequencerSpaceChannelKeyDeleted(UControlRig* ControlRig, ISequencer* Sequencer, FName ControlName, FMovieSceneControlRigSpaceChannel* Channel, UMovieSceneControlRigParameterSection* SectionToKey,
+	FFrameNumber TimeOfDeletion)
+{
+	FMovieSceneControlRigSpaceBaseKey ExistingValue, PreviousValue;
+	using namespace UE::MovieScene;
+	EvaluateChannel(Channel, TimeOfDeletion -1, PreviousValue);
+	EvaluateChannel(Channel, TimeOfDeletion, ExistingValue);
+	if (ExistingValue != PreviousValue) //if they are the same no need to do anything
+	{
+		//find all key frames we need to compensate
+		TArray<FFrameNumber> Frames;
+		Frames.Add(TimeOfDeletion);
+		TSortedMap<FFrameNumber, FFrameNumber> ExtraFrames;
+		FControlRigSpaceChannelHelpers::GetFramesInThisSpaceAfterThisTime(ControlRig, ControlName, ExistingValue,
+			Channel, SectionToKey, TimeOfDeletion, ExtraFrames);
+		if (ExtraFrames.Num() > 0)
+		{
+			for (const TPair<FFrameNumber, FFrameNumber>& Frame : ExtraFrames)
+			{
+				Frames.Add(Frame.Value);
+			}
+		}
+		TArray<FTransform> ControlRigParentWorldTransforms;
+		ControlRigParentWorldTransforms.SetNum(Frames.Num());
+		for (FTransform& WorldTransform : ControlRigParentWorldTransforms)
+		{
+			WorldTransform = FTransform::Identity;
+		}
+		TArray<FTransform> ControlWorldTransforms;
+		FControlRigSnapper Snapper;
+		Snapper.GetControlRigControlTransforms(Sequencer, ControlRig, ControlName, Frames, ControlRigParentWorldTransforms, ControlWorldTransforms);
+		FRigElementKey ControlKey;
+		ControlKey.Name = ControlName;
+		ControlKey.Type = ERigElementType::Control;  
+		URigHierarchy* RigHierarchy = ControlRig->GetHierarchy();
+		switch (PreviousValue.SpaceType)
+		{
+		case EMovieSceneControlRigSpaceType::Parent:
+			RigHierarchy->SwitchToDefaultParent(ControlKey);
+			break;
+		case EMovieSceneControlRigSpaceType::World:
+			RigHierarchy->SwitchToWorldSpace(ControlKey);
+			break;
+		case EMovieSceneControlRigSpaceType::ControlRig:
+			RigHierarchy->SwitchToParent(ControlKey, PreviousValue.ControlRigElement);
+			break;
+		}
+		int32 FramesIndex = 0;
+		FRigControlModifiedContext Context;
+		Context.SetKey = EControlRigSetKey::Always;
+		FFrameRate TickResolution = Sequencer->GetFocusedTickResolution();
+		for (const FFrameNumber& Frame : Frames)
+		{
+			ControlRig->Evaluate_AnyThread();
+			Context.LocalTime = TickResolution.AsSeconds(FFrameTime(Frame));
+			ControlRig->SetControlGlobalTransform(ControlKey.Name, ControlWorldTransforms[FramesIndex++], true, Context);
+		}
+		//now delete any extra TimeOfDelete -1
+		FControlRigSpaceChannelHelpers::DeleteTransformKeysAtThisTime(ControlRig, SectionToKey, ControlName, TimeOfDeletion - 1);
+	}
+}
+
+void FControlRigSpaceChannelHelpers::DeleteTransformKeysAtThisTime(UControlRig* ControlRig, UMovieSceneControlRigParameterSection* Section, FName ControlName, FFrameNumber Time)
+{
+	if (Section && ControlRig)
+	{
+		TArrayView<FMovieSceneFloatChannel*> FloatChannels = Section->GetChannelProxy().GetChannels<FMovieSceneFloatChannel>();
+		FChannelMapInfo* pChannelIndex = Section->ControlChannelMap.Find(ControlName);
+		if (pChannelIndex != nullptr)
+		{
+			int32 ChannelIndex = pChannelIndex->ChannelIndex;
+
+			if (FRigControlElement* ControlElement = ControlRig->FindControl(ControlName))
+			{
+				FMovieSceneControlRigSpaceBaseKey Value;
+				switch (ControlElement->Settings.ControlType)
+				{
+					case ERigControlType::Position:
+					case ERigControlType::Scale:
+					case ERigControlType::Rotator:
+					case ERigControlType::Transform:
+					case ERigControlType::TransformNoScale:
+					case ERigControlType::EulerTransform:
+					{
+						int NumChannels = 0;
+						if (ControlElement->Settings.ControlType == ERigControlType::Transform ||
+							ControlElement->Settings.ControlType == ERigControlType::EulerTransform)
+						{
+							NumChannels = 9;
+						}
+						else if (ControlElement->Settings.ControlType == ERigControlType::TransformNoScale)
+						{
+							NumChannels = 6;
+						}
+						else //vectors
+						{
+							NumChannels = 3;
+						}
+						for (int32 Index = 0; Index < NumChannels; ++Index)
+						{
+							int32 KeyIndex = 0;
+							for (FFrameNumber Frame : FloatChannels[ChannelIndex]->GetData().GetTimes())
+							{
+								if (Frame == Time)
+								{
+									FloatChannels[ChannelIndex]->GetData().RemoveKey(KeyIndex);
+									break;
+								}
+								else if (Frame > Time)
+								{
+									break;
+								}
+								++KeyIndex;
+							}
+							++ChannelIndex;
+						}
+						break;
+					}
+					default:
+						break;
+
+				}
+			}
+		}
+	}
+}
+
 void FControlRigSpaceChannelHelpers::GetFramesInThisSpaceAfterThisTime(UControlRig* ControlRig, FName ControlName, FMovieSceneControlRigSpaceBaseKey CurrentValue,
 	FMovieSceneControlRigSpaceChannel* Channel, UMovieSceneSection* SectionToKey,
 	FFrameNumber Time, TSortedMap<FFrameNumber,FFrameNumber>& OutMoreFrames)
@@ -372,45 +500,31 @@ void FControlRigSpaceChannelHelpers::GetFramesInThisSpaceAfterThisTime(UControlR
 						case ERigControlType::Position:
 						case ERigControlType::Scale:
 						case ERigControlType::Rotator:
-						{
-							const int32 NumChannels = 3;
-							for (int32 Index = 0; Index < NumChannels; ++Index)
-							{
-								for (FFrameNumber Frame : FloatChannels[ChannelIndex++]->GetData().GetTimes())
-								{
-									if (Frame > Time)
-									{
-										EvaluateChannel(Channel, Frame, Value);
-										if (CurrentValue == Value)
-										{
-											if (OutMoreFrames.Find(Frame) == nullptr)
-											{
-												OutMoreFrames.Add(Frame, Frame);
-											}
-										}
-										else
-										{
-											break;
-										}
-									}
-								}
-							}
-							break;
-						}
-
 						case ERigControlType::Transform:
 						case ERigControlType::TransformNoScale:
 						case ERigControlType::EulerTransform:
 						{
-
-							const int32 NumChannels = (ControlElement->Settings.ControlType == ERigControlType::Transform ||
-								ControlElement->Settings.ControlType == ERigControlType::EulerTransform) ? 9 : 6;
+							int NumChannels = 0;
+							if (ControlElement->Settings.ControlType == ERigControlType::Transform ||
+								ControlElement->Settings.ControlType == ERigControlType::EulerTransform)
+							{
+								NumChannels = 9;
+							}
+							else if (ControlElement->Settings.ControlType == ERigControlType::TransformNoScale)
+							{
+								NumChannels = 6;
+							}
+							else //vectors
+							{
+								NumChannels = 3;
+							}
 							for (int32 Index = 0; Index < NumChannels; ++Index)
 							{
 								for (FFrameNumber Frame : FloatChannels[ChannelIndex++]->GetData().GetTimes())
 								{
 									if (Frame > Time)
 									{
+										using namespace UE::MovieScene;
 										EvaluateChannel(Channel, Frame, Value);
 										if (CurrentValue == Value)
 										{
@@ -428,6 +542,7 @@ void FControlRigSpaceChannelHelpers::GetFramesInThisSpaceAfterThisTime(UControlR
 							}
 							break;
 						}
+
 					}
 				}
 			}
@@ -603,10 +718,153 @@ void FControlRigSpaceChannelHelpers::SequencerBakeControlInSpace(UControlRig* Co
 	}
 }
 
-struct FSpaceRanges
+void FControlRigSpaceChannelHelpers::HandleSpaceKeyTimeChanged(UControlRig* ControlRig, FName ControlName, FMovieSceneControlRigSpaceChannel* Channel, UMovieSceneSection* SectionToKey,
+	FFrameNumber CurrentFrame, FFrameNumber NextFrame)
+{
+	if (ControlRig && Channel && SectionToKey && (CurrentFrame != NextFrame))
+	{
+		if (UMovieSceneControlRigParameterSection* Section = Cast<UMovieSceneControlRigParameterSection>(SectionToKey))
+		{
+			TArrayView<FMovieSceneFloatChannel*> FloatChannels = Section->GetChannelProxy().GetChannels<FMovieSceneFloatChannel>();
+			FChannelMapInfo* pChannelIndex = Section->ControlChannelMap.Find(ControlName);
+			if (pChannelIndex != nullptr)
+			{
+				int32 ChannelIndex = pChannelIndex->ChannelIndex;
+				FFrameNumber Delta = NextFrame - CurrentFrame;
+				if (FRigControlElement* ControlElement = ControlRig->FindControl(ControlName))
+				{
+					FMovieSceneControlRigSpaceBaseKey Value;
+					switch (ControlElement->Settings.ControlType)
+					{
+						case ERigControlType::Position:
+						case ERigControlType::Scale:
+						case ERigControlType::Rotator:
+						case ERigControlType::Transform:
+						case ERigControlType::TransformNoScale:
+						case ERigControlType::EulerTransform:
+						{
+							int NumChannels = 0;
+							if (ControlElement->Settings.ControlType == ERigControlType::Transform ||
+								ControlElement->Settings.ControlType == ERigControlType::EulerTransform)
+							{
+								NumChannels = 9;
+							}
+							else if (ControlElement->Settings.ControlType == ERigControlType::TransformNoScale)
+							{
+								NumChannels = 6;
+							}
+							else //vectors
+							{
+								NumChannels = 3;
+							}
+							for (int32 Index = 0; Index < NumChannels; ++Index)
+							{
+								FMovieSceneFloatChannel* FloatChannel = FloatChannels[ChannelIndex++];
+								if (Delta > 0) //if we are moving keys positively in time we start from end frames and move them so we can use indices
+								{
+									for (int32 KeyIndex = FloatChannel->GetData().GetTimes().Num() - 1; KeyIndex >= 0; --KeyIndex)
+									{
+										const FFrameNumber& Frame = FloatChannel->GetData().GetTimes()[KeyIndex];
+										FFrameNumber Diff = Frame - CurrentFrame;
+										FFrameNumber AbsDiff = Diff < 0 ? -Diff : Diff;
+										if (AbsDiff <= 1)
+										{
+											FFrameNumber NewKeyTime = Frame + Delta;
+											FloatChannel->GetData().MoveKey(KeyIndex, NewKeyTime);
+										}
+									}
+								}
+								else
+								{
+									for (int32 KeyIndex = 0; KeyIndex < FloatChannel->GetData().GetTimes().Num(); ++KeyIndex)
+									{
+										const FFrameNumber& Frame = FloatChannel->GetData().GetTimes()[KeyIndex];
+										FFrameNumber Diff = Frame - CurrentFrame;
+										FFrameNumber AbsDiff = Diff < 0 ? -Diff : Diff;
+										if (AbsDiff <= 1)
+										{
+											FFrameNumber NewKeyTime = Frame + Delta;
+											FloatChannel->GetData().MoveKey(KeyIndex, NewKeyTime);
+										}
+									}
+								}
+							}
+							break;
+						}
+						default:
+							break;
+					}
+				}
+			}
+		}
+	}
+}
+
+
+void FControlRigSpaceChannelHelpers::CompensatePreviousFrameIfNeeded(UControlRig* ControlRig, ISequencer* Sequencer, UMovieSceneControlRigParameterSection* Section, FName ControlName, FFrameNumber Time)
 {
-	TRange<FFrameNumber> Range;
-	FMovieSceneControlRigSpaceBaseKey Key;
-};
+	
+	//we need to check all controls for 1) space and 2) previous frame and if so we automatically compensate.
+	TArray<FRigControlElement*> Controls = ControlRig->GetHierarchy()->GetControls();
+	bool bDidIt = false;
+	for(FRigControlElement* Control: Controls)
+	{ 
+		if(Control && Control->GetName() != ControlName)
+		{ 
+			//only if we have a channel
+			if (FSpaceControlNameAndChannel* Channel = Section->GetSpaceChannel(Control->GetName()))
+			{
+				FMovieSceneControlRigSpaceBaseKey ExistingValue, PreviousValue;
+				using namespace UE::MovieScene;
+				EvaluateChannel(&(Channel->SpaceCurve), Time - 1, PreviousValue);
+				EvaluateChannel(&(Channel->SpaceCurve), Time, ExistingValue);
+				if (ExistingValue != PreviousValue) //if they are the same no need to do anything
+				{
+					//find global value at curren time
+					TArray<FFrameNumber> Frames;
+					Frames.Add(Time);
+					TArray<FTransform> ControlRigParentWorldTransforms;
+					ControlRigParentWorldTransforms.SetNum(Frames.Num());
+					for (FTransform& WorldTransform : ControlRigParentWorldTransforms)
+					{
+						WorldTransform = FTransform::Identity;
+					}
+					TArray<FTransform> ControlWorldTransforms;
+					FControlRigSnapper Snapper;
+					Snapper.GetControlRigControlTransforms(Sequencer, ControlRig, Control->GetName(), Frames, ControlRigParentWorldTransforms, ControlWorldTransforms);
+					FRigElementKey ControlKey;
+					ControlKey.Name = Control->GetName();
+					ControlKey.Type = ERigElementType::Control;
+					//set space to previous space value that's different.
+					URigHierarchy* RigHierarchy = ControlRig->GetHierarchy();
+					switch (PreviousValue.SpaceType)
+					{
+					case EMovieSceneControlRigSpaceType::Parent:
+						RigHierarchy->SwitchToDefaultParent(ControlKey);
+						break;
+					case EMovieSceneControlRigSpaceType::World:
+						RigHierarchy->SwitchToWorldSpace(ControlKey);
+						break;
+					case EMovieSceneControlRigSpaceType::ControlRig:
+						RigHierarchy->SwitchToParent(ControlKey, PreviousValue.ControlRigElement);
+						break;
+					}
+					//now set time -1 frame value
+					FRigControlModifiedContext Context;
+					Context.SetKey = EControlRigSetKey::Always;
+					FFrameRate TickResolution = Sequencer->GetFocusedTickResolution();
+					ControlRig->Evaluate_AnyThread();
+					Context.LocalTime = TickResolution.AsSeconds(FFrameTime(Time - 1));
+					ControlRig->SetControlGlobalTransform(ControlKey.Name, ControlWorldTransforms[0], true, Context);
+					bDidIt = true;
+				}
+			}
+		}
+	}
+	if (bDidIt == true)
+	{
+		Sequencer->ForceEvaluate();
+	}
+}
 
 #undef LOCTEXT_NAMESPACE
diff --git a/Engine/Plugins/Animation/ControlRig/Source/ControlRigEditor/Public/ControlRigSpaceChannelEditors.h b/Engine/Plugins/Animation/ControlRig/Source/ControlRigEditor/Public/ControlRigSpaceChannelEditors.h
--- a/Engine/Plugins/Animation/ControlRig/Source/ControlRigEditor/Public/ControlRigSpaceChannelEditors.h
+++ b/Engine/Plugins/Animation/ControlRig/Source/ControlRigEditor/Public/ControlRigSpaceChannelEditors.h
@@ -37,18 +37,24 @@ struct FSpaceChannelAndSection
 	FMovieSceneControlRigSpaceChannel* SpaceChannel;
 };
 
+
 /*
 * Class that contains helper functions for various space switching activities
 */
 struct FControlRigSpaceChannelHelpers
 {
 	static FKeyHandle SequencerKeyControlRigSpaceChannel(UControlRig* ControlRig, ISequencer* Sequencer, FMovieSceneControlRigSpaceChannel* Channel, UMovieSceneSection* SectionToKey, FFrameNumber Time, URigHierarchy* RigHierarchy, const FRigElementKey& ControlKey, const FRigElementKey& SpaceKey);
+	static void SequencerSpaceChannelKeyDeleted(UControlRig* ControlRig, ISequencer* Sequencer, FName ControlName, FMovieSceneControlRigSpaceChannel* Channel, UMovieSceneControlRigParameterSection* SectionToKey, FFrameNumber TimeOfDeletion);
+	static void CompensatePreviousFrameIfNeeded(UControlRig* ControlRig, ISequencer* Sequencer, UMovieSceneControlRigParameterSection* Section, FName ControlName, FFrameNumber Time);
 	static FSpaceChannelAndSection FindSpaceChannelAndSectionForControl(UControlRig* ControlRig, FName ControlName, ISequencer* Sequencer, bool bCreateIfNeeded);
 	static void SequencerBakeControlInSpace(UControlRig* ControlRig, ISequencer* Sequencer, FMovieSceneControlRigSpaceChannel* Channel, UMovieSceneSection* SectionToKey,
 		TArray<FFrameNumber> Frames, URigHierarchy* RigHierarchy, const FRigElementKey& ControlKey, FRigSpacePickerBakeSettings InSettings);
 	static void GetFramesInThisSpaceAfterThisTime(UControlRig* ControlRig, FName ControlName, FMovieSceneControlRigSpaceBaseKey CurrentValue,
 		FMovieSceneControlRigSpaceChannel* Channel, UMovieSceneSection* SectionToKey,
 		FFrameNumber Time, TSortedMap<FFrameNumber,FFrameNumber>& OutMoreFrames);
+	static void HandleSpaceKeyTimeChanged(UControlRig* ControlRig, FName ControlName,FMovieSceneControlRigSpaceChannel* Channel, UMovieSceneSection* SectionToKey,
+		FFrameNumber CurrentFrame, FFrameNumber NextFrame);
+	static void DeleteTransformKeysAtThisTime(UControlRig* ControlRig, UMovieSceneControlRigParameterSection* Section, FName ControlName, FFrameNumber Time);
 };
 
 //template specialization
diff --git a/Engine/Source/Runtime/MovieScene/Private/Channels/MovieSceneChannelData.cpp b/Engine/Source/Runtime/MovieScene/Private/Channels/MovieSceneChannelData.cpp
--- a/Engine/Source/Runtime/MovieScene/Private/Channels/MovieSceneChannelData.cpp
+++ b/Engine/Source/Runtime/MovieScene/Private/Channels/MovieSceneChannelData.cpp
@@ -89,8 +89,8 @@ namespace MovieScene
 } // namespace MovieScene
 } // namespace UE
 
-FMovieSceneChannelData::FMovieSceneChannelData(TArray<FFrameNumber>* InTimes, FKeyHandleLookupTable* InKeyHandles)
-	: Times(InTimes), KeyHandles(InKeyHandles)
+FMovieSceneChannelData::FMovieSceneChannelData(TArray<FFrameNumber>* InTimes, FKeyHandleLookupTable* InKeyHandles, FMovieSceneChannel* InChannel)
+	: Times(InTimes), KeyHandles(InKeyHandles), OwningChannel(InChannel)
 {}
 
 FKeyHandle FMovieSceneChannelData::GetHandle(int32 Index)
@@ -138,7 +138,7 @@ int32 FMovieSceneChannelData::AddKeyInternal(FFrameNumber InTime)
 int32 FMovieSceneChannelData::MoveKeyInternal(int32 KeyIndex, FFrameNumber InNewTime)
 {
 	check(Times->IsValidIndex(KeyIndex));
-
+	FFrameNumber OldTime = (*Times)[KeyIndex];
 	int32 NewIndex = Algo::LowerBound(*Times, InNewTime);
 	if (NewIndex < KeyIndex || NewIndex > KeyIndex+1)
 	{
@@ -157,10 +157,22 @@ int32 FMovieSceneChannelData::MoveKeyInternal(int32 KeyIndex, FFrameNumber InNew
 		{
 			KeyHandles->MoveHandle(KeyIndex, NewIndex);
 		}
+		if (OwningChannel && OwningChannel->OnKeyMovedEvent().IsBound())
+		{
+			TArray<FKeyMoveEventItem> Items;
+			Items.Add(FKeyMoveEventItem(KeyIndex, OldTime, NewIndex,InNewTime));
+			OwningChannel->OnKeyMovedEvent().Broadcast(OwningChannel, Items);
+		}
 		return NewIndex;
 	}
 	else
 	{
+		if (OwningChannel && OwningChannel->OnKeyMovedEvent().IsBound())
+		{
+			TArray<FKeyMoveEventItem> Items;
+			Items.Add(FKeyMoveEventItem(KeyIndex, OldTime, KeyIndex, InNewTime));
+			OwningChannel->OnKeyMovedEvent().Broadcast(OwningChannel, Items);
+		}
 		(*Times)[KeyIndex] = InNewTime;
 		return KeyIndex;
 	}
@@ -227,8 +239,26 @@ void FMovieSceneChannelData::GetKeyTimes(TArrayView<const FKeyHandle> InHandles,
 
 void FMovieSceneChannelData::Offset(FFrameNumber DeltaTime)
 {
-	for (FFrameNumber& Time : *Times)
+	if (OwningChannel == nullptr || OwningChannel->OnKeyMovedEvent().IsBound() == false)
+	{
+		for (FFrameNumber& Time : *Times)
+		{
+			Time += DeltaTime;
+		}
+	}
+	else
 	{
-		Time += DeltaTime;
+		TArray<FKeyMoveEventItem> Items;
+		for (int32 Index = 0; Index < Times->Num(); ++Index)
+		{
+			const FFrameNumber OldTime = (*Times)[Index];
+			FFrameNumber NewTime = OldTime + DeltaTime;
+			Items.Add(FKeyMoveEventItem(Index, OldTime, Index, NewTime));
+			(*Times)[Index] = NewTime;
+		}
+		OwningChannel->OnKeyMovedEvent().Broadcast(OwningChannel, Items);
+
 	}
+	
+
 }
diff --git a/Engine/Source/Runtime/MovieScene/Public/Channels/MovieSceneChannel.h b/Engine/Source/Runtime/MovieScene/Public/Channels/MovieSceneChannel.h
--- a/Engine/Source/Runtime/MovieScene/Public/Channels/MovieSceneChannel.h
+++ b/Engine/Source/Runtime/MovieScene/Public/Channels/MovieSceneChannel.h
@@ -15,6 +15,38 @@ struct FFrameRate;
 struct FFrameNumber;
 struct FKeyDataOptimizationParams;
 
+/*
+*  Events that are fired when adding, deleting or moving keys
+*/
+/**
+*  Item that's sent  when an event is fired when we delete or add a key
+*/
+struct FKeyAddOrDeleteEventItem
+{
+	FKeyAddOrDeleteEventItem(int32 InIndex, FFrameNumber InFrame) : Index(InIndex), Frame(InFrame) {};
+	int32 Index;
+	FFrameNumber Frame;;
+};
+
+/**
+*  Item that's sent when an event is fired when we move a key in time
+*/
+struct FKeyMoveEventItem
+{
+	FKeyMoveEventItem(int32 InIndex, FFrameNumber InFrame, int32 InNewIndex, FFrameNumber InNewFrame) : Index(InIndex), Frame(InFrame), NewIndex(InNewIndex), NewFrame(InNewFrame) {};
+	int32 Index;
+	FFrameNumber Frame;
+	int32 NewIndex;
+	FFrameNumber NewFrame;
+};
+
+/*
+* Note if any Channel uses these delegate's they need a custom serializer to make sure they stick around on undo/redo. Dynamic delegates are too heavy.
+*/
+DECLARE_MULTICAST_DELEGATE_TwoParams(FMovieSceneChannelDataKeyAddedEvent, FMovieSceneChannel*, const TArray<FKeyAddOrDeleteEventItem>& Items);
+DECLARE_MULTICAST_DELEGATE_TwoParams(FMovieSceneChannelDataKeyDeletedEvent, FMovieSceneChannel*, const TArray<FKeyAddOrDeleteEventItem>& Items);
+DECLARE_MULTICAST_DELEGATE_TwoParams(FMovieSceneChannelDataKeyMovedEvent, FMovieSceneChannel*, const TArray<FKeyMoveEventItem>& Items);
+
 USTRUCT()
 struct MOVIESCENE_API FMovieSceneChannel
 {
@@ -156,4 +188,32 @@ struct MOVIESCENE_API FMovieSceneChannel
 	 */
 	virtual void PostEditChange() {}
 
+public:
+
+	/**
+	* Get delegate that's called when a key is added
+	*@return Returns key added event delegate
+	*/
+	FMovieSceneChannelDataKeyAddedEvent& OnKeyAddedEvent() { return KeyAddedEvent; }
+	/**
+	* Get delegate that's called when a key is deleted
+	*@return Returns key deleted event delegate
+	*/
+	FMovieSceneChannelDataKeyDeletedEvent& OnKeyDeletedEvent() { return KeyDeletedEvent; }
+	/**
+	* Get delegate that's called when a key is moved
+	*@return Returns key moved event delegate
+	*/
+	FMovieSceneChannelDataKeyMovedEvent& OnKeyMovedEvent() { return KeyMovedEvent; }
+
+protected:
+	/** Broadcasts a notification whenever key is added */
+	FMovieSceneChannelDataKeyAddedEvent KeyAddedEvent;
+
+	/** Broadcasts a notification whenever key is deleted */
+	FMovieSceneChannelDataKeyDeletedEvent KeyDeletedEvent;
+
+	/** Broadcasts a notification whenever key is moved */
+	FMovieSceneChannelDataKeyMovedEvent KeyMovedEvent;
+
 };
\ No newline at end of file
diff --git a/Engine/Source/Runtime/MovieScene/Public/Channels/MovieSceneChannelData.h b/Engine/Source/Runtime/MovieScene/Public/Channels/MovieSceneChannelData.h
--- a/Engine/Source/Runtime/MovieScene/Public/Channels/MovieSceneChannelData.h
+++ b/Engine/Source/Runtime/MovieScene/Public/Channels/MovieSceneChannelData.h
@@ -2,11 +2,12 @@
 
 #pragma once
 
+#include "UObject/ObjectMacros.h"
 #include "Containers/ArrayView.h"
 #include "Curves/KeyHandle.h"
 #include "Misc/FrameNumber.h"
 #include "Misc/FrameTime.h"
-
+#include "Channels/MovieSceneChannel.h"
 #include "MovieSceneChannelData.generated.h"
 
 struct FFrameRate;
@@ -186,8 +187,9 @@ protected:
 	 *
 	 * @param InTimes        A pointer to an array that should be operated on by this class. Externally owned.
 	 * @param InKeyHandles   A key handle map used for persistent, order independent identification of keys
+	 * @param InChannel      A optional pointer to the owning channel.
 	 */
-	FMovieSceneChannelData(TArray<FFrameNumber>* InTimes, FKeyHandleLookupTable* InKeyHandles);
+	FMovieSceneChannelData(TArray<FFrameNumber>* InTimes, FKeyHandleLookupTable* InKeyHandles, FMovieSceneChannel* InChannel = nullptr);
 
 	/**
 	 * Move the key at index KeyIndex to a new time
@@ -203,6 +205,7 @@ protected:
 	 */
 	int32 AddKeyInternal(FFrameNumber InTime);
 
+
 protected:
 
 	/** Pointer to an external array of sorted times. Must be kept in sync with a corresponding value array. */
@@ -210,8 +213,11 @@ protected:
 
 	/** Pointer to an external key handle map */
 	FKeyHandleLookupTable* KeyHandles;
-};
 
+	/** Optional Pointer to the owning FMovieSceneChannel, should be set if the add,move, and delete callbacks are needed */
+	FMovieSceneChannel* OwningChannel;
+
+};
 
 /**
  * Templated channel data utility class that provides a consistent interface for interacting with a channel's keys and values.
@@ -224,6 +230,21 @@ struct TMovieSceneChannelData : FMovieSceneChannelData
 	typedef typename TCallTraits<ValueType>::ParamType ParamType;
 
 
+	/**
+	 * Constructor that takes a non-owning pointer to an array of times and values, and a key handle map
+	 *
+	 * @param InTimes        A pointer to an array of times that should be operated on by this class. Externally owned.
+	 * @param InValues       A pointer to an array of values that should be operated on by this class. Externally owned.
+	 * @param InKeyHandles   A key handle map used for persistent, order independent identification of keys
+	 * @param InChannel A option point to the owning channel, should be set if the move,add, delete delegates are utilizaed
+	 */
+	TMovieSceneChannelData(TArray<FFrameNumber>* InTimes, TArray<ValueType>* InValues, FKeyHandleLookupTable* InKeyHandles, FMovieSceneChannel* InChannel = nullptr)
+		: FMovieSceneChannelData(InTimes, InKeyHandles, InChannel), Values(InValues) 
+	{
+		check(Times && Values);
+	}
+
+
 	/**
 	 * Constructor that takes a non-owning pointer to an array of times and values, and a key handle map
 	 *
@@ -231,11 +252,12 @@ struct TMovieSceneChannelData : FMovieSceneChannelData
 	 * @param InValues       A pointer to an array of values that should be operated on by this class. Externally owned.
 	 * @param InKeyHandles   A key handle map used for persistent, order independent identification of keys
 	 */
-	TMovieSceneChannelData(TArray<FFrameNumber>* InTimes, TArray<ValueType>* InValues, FKeyHandleLookupTable* InKeyHandles)
+	TMovieSceneChannelData(TArray<FFrameNumber>* InTimes, TArray<ValueType>* InValues, FMovieSceneChannel* InChannel, FKeyHandleLookupTable* InKeyHandles)
 		: FMovieSceneChannelData(InTimes, InKeyHandles), Values(InValues)
 	{
 		check(Times && Values);
 	}
+	
 
 	/**
 	 * Conversion to a constant version of this class
@@ -272,6 +294,13 @@ struct TMovieSceneChannelData : FMovieSceneChannelData
 	{
 		int32 KeyIndex = AddKeyInternal(InTime);
 		Values->Insert(InValue, KeyIndex);
+		if (OwningChannel && OwningChannel->OnKeyAddedEvent().IsBound())
+		{
+			TArray<FKeyAddOrDeleteEventItem> Items;
+			Items.Add(FKeyAddOrDeleteEventItem(KeyIndex, InTime));
+			OwningChannel->OnKeyAddedEvent().Broadcast(OwningChannel, Items);
+			
+		}
 		return KeyIndex;
 	}
 
@@ -316,6 +345,13 @@ struct TMovieSceneChannelData : FMovieSceneChannelData
 	void RemoveKey(int32 KeyIndex)
 	{
 		check(Times->IsValidIndex(KeyIndex));
+		if (OwningChannel && OwningChannel->OnKeyDeletedEvent().IsBound())
+		{
+			const FFrameNumber Time = (*Times)[KeyIndex];
+			TArray<FKeyAddOrDeleteEventItem> Items;
+			Items.Add(FKeyAddOrDeleteEventItem(KeyIndex, Time));
+			OwningChannel->OnKeyDeletedEvent().Broadcast(OwningChannel, Items);
+		}
 		Times->RemoveAt(KeyIndex, 1, false);
 		Values->RemoveAt(KeyIndex, 1, false);
 
@@ -450,6 +486,16 @@ struct TMovieSceneChannelData : FMovieSceneChannelData
 	 */
 	void Reset()
 	{
+		if (OwningChannel && OwningChannel->OnKeyDeletedEvent().IsBound())
+		{
+			TArray<FKeyAddOrDeleteEventItem> Items;
+			for (int32 Index = 0; Index < Times->Num(); ++Index)
+			{
+				const FFrameNumber Time = (*Times)[Index];
+				Items.Add(FKeyAddOrDeleteEventItem(Index, Time));
+			}
+			OwningChannel->OnKeyDeletedEvent().Broadcast(OwningChannel, Items);
+		}
 		Times->Reset();
 		Values->Reset();
 		if (KeyHandles)
@@ -462,6 +508,7 @@ private:
 
 	/** Pointer to an external array of values, to be kept in sync with FMovieSceneChannelData::Times */
 	TArray<ValueType>* Values;
+
 };
 
 
