Sequencer: Space Switching: Compensation works with curve editor or sequencer keying, baking does euler filter.
review pending

diff --git a/Engine/Plugins/Animation/ControlRig/Source/ControlRig/Private/Sequencer/MovieSceneControlRigParameterSection.cpp b/Engine/Plugins/Animation/ControlRig/Source/ControlRig/Private/Sequencer/MovieSceneControlRigParameterSection.cpp
--- a/Engine/Plugins/Animation/ControlRig/Source/ControlRig/Private/Sequencer/MovieSceneControlRigParameterSection.cpp
+++ b/Engine/Plugins/Animation/ControlRig/Source/ControlRig/Private/Sequencer/MovieSceneControlRigParameterSection.cpp
@@ -2033,6 +2033,53 @@ void UMovieSceneControlRigParameterSection::SetControlRig(UControlRig* InControl
 	ControlRigClass = ControlRig ? ControlRig->GetClass() : nullptr;
 }
 
+void UMovieSceneControlRigParameterSection::FixRotationWinding(FName ControlName, FFrameNumber StartFrame, FFrameNumber EndFrame)
+{
+	FChannelMapInfo* pChannelIndex = ControlChannelMap.Find(ControlName);
+	if (pChannelIndex == nullptr || GetControlRig() == nullptr)
+	{
+		return;
+	}
+	int32 ChannelIndex = pChannelIndex->ChannelIndex;
+	TArrayView<FMovieSceneFloatChannel*> FloatChannels = GetChannelProxy().GetChannels<FMovieSceneFloatChannel>();
+	if (FRigControlElement* ControlElement = GetControlRig()->FindControl(ControlName))
+	{
+		if (ControlElement->Settings.ControlType == ERigControlType::Rotator  ||
+			ControlElement->Settings.ControlType == ERigControlType::EulerTransform ||
+			ControlElement->Settings.ControlType == ERigControlType::Transform || 
+			ControlElement->Settings.ControlType == ERigControlType::TransformNoScale)
+		{
+			int32 StartIndex = (ControlElement->Settings.ControlType == ERigControlType::Rotator) ? 0 : 3;
+			for (int32 Index = 0; Index < 3; ++Index)
+			{
+				int32 RealIndex = StartIndex + Index + ChannelIndex;
+				int32 NumKeys = FloatChannels[RealIndex]->GetNumKeys();
+				bool bDidFrame = false;
+				float PrevVal = 0.0f;
+				for (int32 KeyIndex = 0; KeyIndex < NumKeys; ++KeyIndex)
+				{
+					const FFrameNumber& Frame = FloatChannels[RealIndex]->GetData().GetTimes()[KeyIndex];
+					if (Frame >= StartFrame && Frame <= EndFrame)
+					{
+						FMovieSceneFloatValue Val = FloatChannels[RealIndex]->GetData().GetValues()[KeyIndex];
+						if (bDidFrame == true)
+						{
+							FMath::WindRelativeAnglesDegrees(PrevVal, Val.Value);
+							FloatChannels[RealIndex]->GetData().GetValues()[KeyIndex].Value = Val.Value;
+						}
+						else
+						{
+							bDidFrame = true;
+						}
+						PrevVal = Val.Value;
+					}
+					
+				}
+			}
+		}
+	}
+}
+
 #if WITH_EDITOR
 
 void UMovieSceneControlRigParameterSection::RecordControlRigKey(FFrameNumber FrameNumber, bool bSetDefault, bool bAutoKey)
diff --git a/Engine/Plugins/Animation/ControlRig/Source/ControlRig/Public/Sequencer/MovieSceneControlRigParameterSection.h b/Engine/Plugins/Animation/ControlRig/Source/ControlRig/Public/Sequencer/MovieSceneControlRigParameterSection.h
--- a/Engine/Plugins/Animation/ControlRig/Source/ControlRig/Public/Sequencer/MovieSceneControlRigParameterSection.h
+++ b/Engine/Plugins/Animation/ControlRig/Source/ControlRig/Public/Sequencer/MovieSceneControlRigParameterSection.h
@@ -183,6 +183,8 @@ public:
 	TArray<FIntegerParameterNameAndCurve>& GetIntegerParameterNamesAndCurves();
 	const TArray<FIntegerParameterNameAndCurve>& GetIntegerParameterNamesAndCurves() const;
 
+	void FixRotationWinding(FName ControlName, FFrameNumber StartFrame, FFrameNumber EndFrame);
+
 	TArray<FSpaceControlNameAndChannel>& GetSpaceChannels();
 	const TArray< FSpaceControlNameAndChannel>& GetSpaceChannels() const;
 	FName FindControlNameFromSpaceChannel(const FMovieSceneControlRigSpaceChannel* SpaceChannel) const;
diff --git a/Engine/Plugins/Animation/ControlRig/Source/ControlRigEditor/Private/Sequencer/ControlRigParameterTrackEditor.cpp b/Engine/Plugins/Animation/ControlRig/Source/ControlRigEditor/Private/Sequencer/ControlRigParameterTrackEditor.cpp
--- a/Engine/Plugins/Animation/ControlRig/Source/ControlRigEditor/Private/Sequencer/ControlRigParameterTrackEditor.cpp
+++ b/Engine/Plugins/Animation/ControlRig/Source/ControlRigEditor/Private/Sequencer/ControlRigParameterTrackEditor.cpp
@@ -1437,6 +1437,10 @@ void FControlRigParameterTrackEditor::OnChannelChanged(const FMovieSceneChannelM
 			GetSequencer()->ForceEvaluate(); //now run sequencer...
 			Section->GetControlRig()->Evaluate_AnyThread();
 			Section->ControlsToSet.Empty();
+
+			TOptional<FFrameNumber> Optional;
+			FControlRigSpaceChannelHelpers::CompensateIfNeeded(Section->GetControlRig(), GetSequencer().Get(), Section,
+				ControlName, Optional);
 		}
 	}
 }
@@ -2178,8 +2182,9 @@ void FControlRigParameterTrackEditor::HandleControlModified(UControlRig* Control
 					
 					FFrameNumber KeyTime = (Context.LocalTime == FLT_MAX) ? GetTimeForKey() : GetSequencer()->GetFocusedTickResolution().AsFrameTime((double)Context.LocalTime).RoundToFrame();
 
-					FControlRigSpaceChannelHelpers::CompensatePreviousFrameIfNeeded(ControlRig, GetSequencer().Get(), ParamSection,
-						ControlElement->GetName(), KeyTime);
+					TOptional<FFrameNumber> OptionalKeyTime = KeyTime;
+					FControlRigSpaceChannelHelpers::CompensateIfNeeded(ControlRig, GetSequencer().Get(), ParamSection,
+						ControlElement->GetName(), OptionalKeyTime);
 				}
 			}
 		}
diff --git a/Engine/Plugins/Animation/ControlRig/Source/ControlRigEditor/Private/Sequencer/ControlRigSpaceChannelEditors.cpp b/Engine/Plugins/Animation/ControlRig/Source/ControlRigEditor/Private/Sequencer/ControlRigSpaceChannelEditors.cpp
--- a/Engine/Plugins/Animation/ControlRig/Source/ControlRigEditor/Private/Sequencer/ControlRigSpaceChannelEditors.cpp
+++ b/Engine/Plugins/Animation/ControlRig/Source/ControlRigEditor/Private/Sequencer/ControlRigSpaceChannelEditors.cpp
@@ -658,9 +658,12 @@ void FControlRigSpaceChannelHelpers::SequencerBakeControlInSpace(UControlRig* Co
 			Context.LocalTime = TickResolution.AsSeconds(FFrameTime(EndFrame));
 			ControlRig->SetControlGlobalTransform(ControlKey.Name, GlobalTransform, true, Context);
 		}
-		if (Settings.bReduceKeys)
-		{
-			if (UMovieSceneControlRigParameterSection* Section = Cast<UMovieSceneControlRigParameterSection>(SectionToKey))
+		if (UMovieSceneControlRigParameterSection* Section = Cast<UMovieSceneControlRigParameterSection>(SectionToKey))
+		{ 		
+			// Fix any Rotation Channels
+			Section->FixRotationWinding(ControlKey.Name, Frames[0], Frames[Frames.Num() - 1]);
+			// Then reduce
+			if (Settings.bReduceKeys)
 			{
 				TArrayView<FMovieSceneFloatChannel*> FloatChannels = Section->GetChannelProxy().GetChannels<FMovieSceneFloatChannel>();
 				FChannelMapInfo* pChannelIndex = Section->ControlChannelMap.Find(ControlKey.Name);
@@ -713,6 +716,7 @@ void FControlRigSpaceChannelHelpers::SequencerBakeControlInSpace(UControlRig* Co
 					}
 				}
 			}
+			
 		}
 		Sequencer->NotifyMovieSceneDataChanged(EMovieSceneDataChangeType::MovieSceneStructureItemAdded); //may have added channel
 	}
@@ -801,7 +805,7 @@ void FControlRigSpaceChannelHelpers::HandleSpaceKeyTimeChanged(UControlRig* Cont
 }
 
 
-void FControlRigSpaceChannelHelpers::CompensatePreviousFrameIfNeeded(UControlRig* ControlRig, ISequencer* Sequencer, UMovieSceneControlRigParameterSection* Section, FName ControlName, FFrameNumber Time)
+void FControlRigSpaceChannelHelpers::CompensateIfNeeded(UControlRig* ControlRig, ISequencer* Sequencer, UMovieSceneControlRigParameterSection* Section, FName ControlName, TOptional<FFrameNumber>& OptionalTime)
 {
 	
 	//we need to check all controls for 1) space and 2) previous frame and if so we automatically compensate.
@@ -809,54 +813,70 @@ void FControlRigSpaceChannelHelpers::CompensatePreviousFrameIfNeeded(UControlRig
 	bool bDidIt = false;
 	for(FRigControlElement* Control: Controls)
 	{ 
-		if(Control && Control->GetName() != ControlName)
+		if(Control)// ac && Control->GetName() != ControlName)
 		{ 
 			//only if we have a channel
 			if (FSpaceControlNameAndChannel* Channel = Section->GetSpaceChannel(Control->GetName()))
 			{
-				FMovieSceneControlRigSpaceBaseKey ExistingValue, PreviousValue;
-				using namespace UE::MovieScene;
-				EvaluateChannel(&(Channel->SpaceCurve), Time - 1, PreviousValue);
-				EvaluateChannel(&(Channel->SpaceCurve), Time, ExistingValue);
-				if (ExistingValue != PreviousValue) //if they are the same no need to do anything
+				TArray<FFrameNumber> AllFrames; 
+				if (OptionalTime.IsSet())
 				{
-					//find global value at curren time
-					TArray<FFrameNumber> Frames;
-					Frames.Add(Time);
-					TArray<FTransform> ControlRigParentWorldTransforms;
-					ControlRigParentWorldTransforms.SetNum(Frames.Num());
-					for (FTransform& WorldTransform : ControlRigParentWorldTransforms)
-					{
-						WorldTransform = FTransform::Identity;
-					}
-					TArray<FTransform> ControlWorldTransforms;
-					FControlRigSnapper Snapper;
-					Snapper.GetControlRigControlTransforms(Sequencer, ControlRig, Control->GetName(), Frames, ControlRigParentWorldTransforms, ControlWorldTransforms);
-					FRigElementKey ControlKey;
-					ControlKey.Name = Control->GetName();
-					ControlKey.Type = ERigElementType::Control;
-					//set space to previous space value that's different.
-					URigHierarchy* RigHierarchy = ControlRig->GetHierarchy();
-					switch (PreviousValue.SpaceType)
+					FFrameNumber Time = OptionalTime.GetValue();
+					AllFrames.Add(Time);
+				}
+				else
+				{
+					AllFrames = Channel->SpaceCurve.GetData().GetTimes();
+				}
+				if (AllFrames.Num() > 0)
+				{
+					for (FFrameNumber& Time : AllFrames)
 					{
-					case EMovieSceneControlRigSpaceType::Parent:
-						RigHierarchy->SwitchToDefaultParent(ControlKey);
-						break;
-					case EMovieSceneControlRigSpaceType::World:
-						RigHierarchy->SwitchToWorldSpace(ControlKey);
-						break;
-					case EMovieSceneControlRigSpaceType::ControlRig:
-						RigHierarchy->SwitchToParent(ControlKey, PreviousValue.ControlRigElement);
-						break;
+						FMovieSceneControlRigSpaceBaseKey ExistingValue, PreviousValue;
+						using namespace UE::MovieScene;
+						EvaluateChannel(&(Channel->SpaceCurve), Time - 1, PreviousValue);
+						EvaluateChannel(&(Channel->SpaceCurve), Time, ExistingValue);
+						if (ExistingValue != PreviousValue) //if they are the same no need to do anything
+						{
+							//find global value at curren time
+							TArray<FFrameNumber> Frames;
+							Frames.Add(Time);
+							TArray<FTransform> ControlRigParentWorldTransforms;
+							ControlRigParentWorldTransforms.SetNum(Frames.Num());
+							for (FTransform& WorldTransform : ControlRigParentWorldTransforms)
+							{
+								WorldTransform = FTransform::Identity;
+							}
+							TArray<FTransform> ControlWorldTransforms;
+							FControlRigSnapper Snapper;
+							Snapper.GetControlRigControlTransforms(Sequencer, ControlRig, Control->GetName(), Frames, ControlRigParentWorldTransforms, ControlWorldTransforms);
+							FRigElementKey ControlKey;
+							ControlKey.Name = Control->GetName();
+							ControlKey.Type = ERigElementType::Control;
+							//set space to previous space value that's different.
+							URigHierarchy* RigHierarchy = ControlRig->GetHierarchy();
+							switch (PreviousValue.SpaceType)
+							{
+							case EMovieSceneControlRigSpaceType::Parent:
+								RigHierarchy->SwitchToDefaultParent(ControlKey);
+								break;
+							case EMovieSceneControlRigSpaceType::World:
+								RigHierarchy->SwitchToWorldSpace(ControlKey);
+								break;
+							case EMovieSceneControlRigSpaceType::ControlRig:
+								RigHierarchy->SwitchToParent(ControlKey, PreviousValue.ControlRigElement);
+								break;
+							}
+							//now set time -1 frame value
+							FRigControlModifiedContext Context;
+							Context.SetKey = EControlRigSetKey::Always;
+							FFrameRate TickResolution = Sequencer->GetFocusedTickResolution();
+							ControlRig->Evaluate_AnyThread();
+							Context.LocalTime = TickResolution.AsSeconds(FFrameTime(Time - 1));
+							ControlRig->SetControlGlobalTransform(ControlKey.Name, ControlWorldTransforms[0], true, Context);
+							bDidIt = true;
+						}
 					}
-					//now set time -1 frame value
-					FRigControlModifiedContext Context;
-					Context.SetKey = EControlRigSetKey::Always;
-					FFrameRate TickResolution = Sequencer->GetFocusedTickResolution();
-					ControlRig->Evaluate_AnyThread();
-					Context.LocalTime = TickResolution.AsSeconds(FFrameTime(Time - 1));
-					ControlRig->SetControlGlobalTransform(ControlKey.Name, ControlWorldTransforms[0], true, Context);
-					bDidIt = true;
 				}
 			}
 		}
diff --git a/Engine/Plugins/Animation/ControlRig/Source/ControlRigEditor/Public/ControlRigSpaceChannelEditors.h b/Engine/Plugins/Animation/ControlRig/Source/ControlRigEditor/Public/ControlRigSpaceChannelEditors.h
--- a/Engine/Plugins/Animation/ControlRig/Source/ControlRigEditor/Public/ControlRigSpaceChannelEditors.h
+++ b/Engine/Plugins/Animation/ControlRig/Source/ControlRigEditor/Public/ControlRigSpaceChannelEditors.h
@@ -45,7 +45,7 @@ struct FControlRigSpaceChannelHelpers
 {
 	static FKeyHandle SequencerKeyControlRigSpaceChannel(UControlRig* ControlRig, ISequencer* Sequencer, FMovieSceneControlRigSpaceChannel* Channel, UMovieSceneSection* SectionToKey, FFrameNumber Time, URigHierarchy* RigHierarchy, const FRigElementKey& ControlKey, const FRigElementKey& SpaceKey);
 	static void SequencerSpaceChannelKeyDeleted(UControlRig* ControlRig, ISequencer* Sequencer, FName ControlName, FMovieSceneControlRigSpaceChannel* Channel, UMovieSceneControlRigParameterSection* SectionToKey, FFrameNumber TimeOfDeletion);
-	static void CompensatePreviousFrameIfNeeded(UControlRig* ControlRig, ISequencer* Sequencer, UMovieSceneControlRigParameterSection* Section, FName ControlName, FFrameNumber Time);
+	static void CompensateIfNeeded(UControlRig* ControlRig, ISequencer* Sequencer, UMovieSceneControlRigParameterSection* Section, FName ControlName, TOptional<FFrameNumber>& Time);
 	static FSpaceChannelAndSection FindSpaceChannelAndSectionForControl(UControlRig* ControlRig, FName ControlName, ISequencer* Sequencer, bool bCreateIfNeeded);
 	static void SequencerBakeControlInSpace(UControlRig* ControlRig, ISequencer* Sequencer, FMovieSceneControlRigSpaceChannel* Channel, UMovieSceneSection* SectionToKey,
 		TArray<FFrameNumber> Frames, URigHierarchy* RigHierarchy, const FRigElementKey& ControlKey, FRigSpacePickerBakeSettings InSettings);
diff --git a/Engine/Source/Editor/MovieSceneTools/Private/Channels/ChannelCurveModel.cpp b/Engine/Source/Editor/MovieSceneTools/Private/Channels/ChannelCurveModel.cpp
--- a/Engine/Source/Editor/MovieSceneTools/Private/Channels/ChannelCurveModel.cpp
+++ b/Engine/Source/Editor/MovieSceneTools/Private/Channels/ChannelCurveModel.cpp
@@ -187,6 +187,11 @@ void FChannelCurveModel<ChannelType, ChannelValue, KeyType>::SetKeyPositions(TAr
 			}
 		}
 		Channel->PostEditChange();
+		if(WeakSequencer.IsValid())
+		{ 
+			const FMovieSceneChannelMetaData* MetaData = ChannelHandle.GetMetaData();
+			WeakSequencer.Pin()->OnChannelChanged().Broadcast(MetaData, Section);
+		}
 		CurveModifiedDelegate.Broadcast();
 	}
 }
@@ -360,6 +365,11 @@ void FChannelCurveModel<ChannelType, ChannelValue, KeyType>::AddKeys(TArrayView<
 		// For now we need to duplicate the new key handle array due to API mismatch. This will auto calculate tangents if needed.
 		SetKeyAttributes(NewKeyHandles, InKeyAttributes);
 		Channel->PostEditChange();
+		if (WeakSequencer.IsValid())
+		{
+			const FMovieSceneChannelMetaData* MetaData = ChannelHandle.GetMetaData();
+			WeakSequencer.Pin()->OnChannelChanged().Broadcast(MetaData, Section);
+		}
 		CurveModifiedDelegate.Broadcast();
 	}
 }
@@ -384,6 +394,11 @@ void FChannelCurveModel<ChannelType, ChannelValue, KeyType>::RemoveKeys(TArrayVi
 			}
 		}
 		Channel->PostEditChange();
+		if (WeakSequencer.IsValid())
+		{
+			const FMovieSceneChannelMetaData* MetaData = ChannelHandle.GetMetaData();
+			WeakSequencer.Pin()->OnChannelChanged().Broadcast(MetaData, Section);
+		}
 		CurveModifiedDelegate.Broadcast();
 	}
 }
