Static Analysis fixes for ControlRigParameterSectionand ControlRigSpaceChannelEditors.

diff --git a/Engine/Plugins/Animation/ControlRig/Source/ControlRig/Private/Sequencer/MovieSceneControlRigParameterSection.cpp b/Engine/Plugins/Animation/ControlRig/Source/ControlRig/Private/Sequencer/MovieSceneControlRigParameterSection.cpp
--- a/Engine/Plugins/Animation/ControlRig/Source/ControlRig/Private/Sequencer/MovieSceneControlRigParameterSection.cpp
+++ b/Engine/Plugins/Animation/ControlRig/Source/ControlRig/Private/Sequencer/MovieSceneControlRigParameterSection.cpp
@@ -2058,7 +2058,7 @@ void UMovieSceneControlRigParameterSection::FixRotationWinding(FName ControlName
 				float PrevVal = 0.0f;
 				for (int32 KeyIndex = 0; KeyIndex < NumKeys; ++KeyIndex)
 				{
-					const FFrameNumber& Frame = FloatChannels[RealIndex]->GetData().GetTimes()[KeyIndex];
+					const FFrameNumber Frame = FloatChannels[RealIndex]->GetData().GetTimes()[KeyIndex];
 					if (Frame >= StartFrame && Frame <= EndFrame)
 					{
 						FMovieSceneFloatValue Val = FloatChannels[RealIndex]->GetData().GetValues()[KeyIndex];
diff --git a/Engine/Plugins/Animation/ControlRig/Source/ControlRigEditor/Private/Sequencer/ControlRigSpaceChannelEditors.cpp b/Engine/Plugins/Animation/ControlRig/Source/ControlRigEditor/Private/Sequencer/ControlRigSpaceChannelEditors.cpp
--- a/Engine/Plugins/Animation/ControlRig/Source/ControlRigEditor/Private/Sequencer/ControlRigSpaceChannelEditors.cpp
+++ b/Engine/Plugins/Animation/ControlRig/Source/ControlRigEditor/Private/Sequencer/ControlRigSpaceChannelEditors.cpp
@@ -768,7 +768,7 @@ void FControlRigSpaceChannelHelpers::HandleSpaceKeyTimeChanged(UControlRig* Cont
 								{
 									for (int32 KeyIndex = FloatChannel->GetData().GetTimes().Num() - 1; KeyIndex >= 0; --KeyIndex)
 									{
-										const FFrameNumber& Frame = FloatChannel->GetData().GetTimes()[KeyIndex];
+										const FFrameNumber Frame = FloatChannel->GetData().GetTimes()[KeyIndex];
 										FFrameNumber Diff = Frame - CurrentFrame;
 										FFrameNumber AbsDiff = Diff < 0 ? -Diff : Diff;
 										if (AbsDiff <= 1)
@@ -782,7 +782,7 @@ void FControlRigSpaceChannelHelpers::HandleSpaceKeyTimeChanged(UControlRig* Cont
 								{
 									for (int32 KeyIndex = 0; KeyIndex < FloatChannel->GetData().GetTimes().Num(); ++KeyIndex)
 									{
-										const FFrameNumber& Frame = FloatChannel->GetData().GetTimes()[KeyIndex];
+										const FFrameNumber Frame = FloatChannel->GetData().GetTimes()[KeyIndex];
 										FFrameNumber Diff = Frame - CurrentFrame;
 										FFrameNumber AbsDiff = Diff < 0 ? -Diff : Diff;
 										if (AbsDiff <= 1)
