Control Rig: Refactoring VM and Hierarchy ownership, delay compile until package end load

diff --git a/Engine/Plugins/Animation/ControlRig/Source/ControlRig/Private/ControlRig.cpp b/Engine/Plugins/Animation/ControlRig/Source/ControlRig/Private/ControlRig.cpp
--- a/Engine/Plugins/Animation/ControlRig/Source/ControlRig/Private/ControlRig.cpp
+++ b/Engine/Plugins/Animation/ControlRig/Source/ControlRig/Private/ControlRig.cpp
@@ -79,9 +79,6 @@ UControlRig::UControlRig(const FObjectInitializer& ObjectInitializer)
 #endif
 	, DebugBoneRadiusMultiplier(1.f)
 {
-	SetVM(ObjectInitializer.CreateDefaultSubobject<URigVM>(this, TEXT("VM")));
-	DynamicHierarchy = ObjectInitializer.CreateDefaultSubobject<URigHierarchy>(this, TEXT("DynamicHierarchy"));
-
 	EventQueue.Add(FRigUnit_BeginExecution::EventName);
 }
 
@@ -182,7 +179,6 @@ void UControlRig::Initialize(bool bInitRigUnits)
 	{
 		RequestInit();
 	}
-
 	
 	GetHierarchy()->OnModified().RemoveAll(this);
 	GetHierarchy()->OnModified().AddUObject(this, &UControlRig::HandleHierarchyModified);
@@ -478,19 +474,6 @@ void UControlRig::InstantiateVMFromCDO()
 {
 	DECLARE_SCOPE_HIERARCHICAL_COUNTER_FUNC()
 
-	if (VM == nullptr || VM->GetOuter() != this)
-	{
-		SetVM(NewObject<URigVM>(this, TEXT("VM")));
-		
-#if !UE_RIGVM_UCLASS_BASED_STORAGE_DISABLED
-		// ensure memory storage objects are created during initialization in game thread
-		// delaying this step may lead to async uobject creation and causing GC to complain
-		VM->GetMemoryByType(ERigVMMemoryType::Work, true);
-		VM->GetMemoryByType(ERigVMMemoryType::Literal, true);
-		VM->GetMemoryByType(ERigVMMemoryType::Debug, true);
-#endif
-	}
-	
 	if (!HasAnyFlags(RF_ClassDefaultObject))
 	{
 		UControlRig* CDO = GetClass()->GetDefaultObject<UControlRig>();
@@ -1137,27 +1120,6 @@ void UControlRig::SetEventQueue(const TArray<FName>& InEventNames)
 	EventQueue = InEventNames;
 }
 
-void UControlRig::SetVM(URigVM* NewVM)
-{
-	if (VM)
-	{
-		VM->ExecutionReachedExit().RemoveAll(this);
-		VM->SetRuntimeSettings(FRigVMRuntimeSettings());
-	}
-	
-	if (NewVM)
-	{
-		if (!NewVM->ExecutionReachedExit().IsBoundToObject(this))
-		{
-			NewVM->ExecutionReachedExit().AddUObject(this, &UControlRig::HandleExecutionReachedExit);
-		}
-	}
-
-	VM = NewVM;
-
-	UpdateVMSettings();
-}
-
 void UControlRig::UpdateVMSettings()
 {
 	if(VM)
@@ -2809,6 +2771,53 @@ void UControlRig::SetControlCustomization(const FRigElementKey& InControl, const
 	ControlCustomizations.FindOrAdd(InControl) = InCustomization;
 }
 
+void UControlRig::PostInitInstanceIfRequired()
+{
+	if(GetHierarchy() == nullptr || GetVM() == nullptr)
+	{
+		if(HasAnyFlags(RF_ClassDefaultObject))
+		{
+			PostInitInstance(nullptr);
+		}
+		else
+		{
+			UControlRig* CDO = GetClass()->GetDefaultObject<UControlRig>();
+			PostInitInstance(CDO);
+		}
+	}
+}
+
+void UControlRig::PostInitInstance(UControlRig* InCDO)
+{
+	const EObjectFlags SubObjectFlags =
+	HasAnyFlags(RF_ClassDefaultObject) ?
+		RF_Public | RF_DefaultSubObject :
+		RF_Transient | RF_Transactional;
+
+	// set up the VM
+	VM = NewObject<URigVM>(this, TEXT("VM"), SubObjectFlags);
+	VM->GetMemoryByType(ERigVMMemoryType::Work, true);
+	VM->GetMemoryByType(ERigVMMemoryType::Literal, true);
+	VM->GetMemoryByType(ERigVMMemoryType::Debug, true);
+
+	VM->ExecutionReachedExit().AddUObject(this, &UControlRig::HandleExecutionReachedExit);
+	UpdateVMSettings();
+
+	// set up the hierarchy
+	DynamicHierarchy = NewObject<URigHierarchy>(this, TEXT("DynamicHierarchy"), SubObjectFlags);
+
+#if WITH_EDITOR
+	const TWeakObjectPtr<UControlRig> WeakThis = this;
+	DynamicHierarchy->OnUndoRedo().AddStatic(&UControlRig::OnHierarchyTransformUndoRedoWeak, WeakThis);
+#endif
+
+	if(!HasAnyFlags(RF_ClassDefaultObject) && InCDO)
+	{
+		VM->CopyFrom(InCDO->GetVM());
+		DynamicHierarchy->CopyHierarchy(InCDO->GetHierarchy());
+	}
+}
+
 void UControlRig::OnHierarchyTransformUndoRedo(URigHierarchy* InHierarchy, const FRigElementKey& InKey, ERigTransformType::Type InTransformType, const FTransform& InTransform, bool bIsUndo)
 {
 	if(InKey.Type == ERigElementType::Control)
diff --git a/Engine/Plugins/Animation/ControlRig/Source/ControlRig/Private/ControlRigBlueprintGeneratedClass.cpp b/Engine/Plugins/Animation/ControlRig/Source/ControlRig/Private/ControlRigBlueprintGeneratedClass.cpp
--- a/Engine/Plugins/Animation/ControlRig/Source/ControlRig/Private/ControlRigBlueprintGeneratedClass.cpp
+++ b/Engine/Plugins/Animation/ControlRig/Source/ControlRig/Private/ControlRigBlueprintGeneratedClass.cpp
@@ -20,6 +20,20 @@ uint8* UControlRigBlueprintGeneratedClass::GetPersistentUberGraphFrame(UObject*
 	return Super::GetPersistentUberGraphFrame(Obj, FuncToCheck);
 }
 
+void UControlRigBlueprintGeneratedClass::PostInitInstance(UObject* InObj)
+{
+	UControlRig* ControlRig = Cast<UControlRig>(InObj);
+	check(ControlRig);
+
+	UControlRig* CDO = nullptr;
+	if(!ControlRig->HasAnyFlags(RF_ClassDefaultObject))
+	{
+		CDO = Cast<UControlRig>(GetDefaultObject());;
+	}
+
+	ControlRig->PostInitInstance(CDO);
+}
+
 void UControlRigBlueprintGeneratedClass::Serialize(FArchive& Ar)
 {
 	DECLARE_SCOPE_HIERARCHICAL_COUNTER_FUNC()
@@ -49,14 +63,6 @@ void UControlRigBlueprintGeneratedClass::Serialize(FArchive& Ar)
 	{
 		if (Ar.IsLoading())
 		{
-			if (CDO->VM == nullptr)
-			{
-				CDO->SetVM(NewObject<URigVM>(CDO, TEXT("VM")));
-			}
-			if (CDO->VM->GetOuter() != CDO)
-			{
-				CDO->SetVM(NewObject<URigVM>(CDO, TEXT("VM")));
-			}
 			CDO->VM->CopyFrom(VM);
 		}
 	}
diff --git a/Engine/Plugins/Animation/ControlRig/Source/ControlRig/Private/ControlRigComponent.cpp b/Engine/Plugins/Animation/ControlRig/Source/ControlRig/Private/ControlRigComponent.cpp
--- a/Engine/Plugins/Animation/ControlRig/Source/ControlRig/Private/ControlRigComponent.cpp
+++ b/Engine/Plugins/Animation/ControlRig/Source/ControlRig/Private/ControlRigComponent.cpp
@@ -1138,7 +1138,6 @@ UControlRig* UControlRigComponent::SetupControlRigIfRequired()
 	if(ControlRigClass)
 	{
 		ControlRig = NewObject<UControlRig>(this, ControlRigClass);
-		ControlRig->SetVM(NewObject<URigVM>(ControlRig));
 
 		SetControlRig(ControlRig);
 
diff --git a/Engine/Plugins/Animation/ControlRig/Source/ControlRig/Private/Rigs/AdditiveControlRig.cpp b/Engine/Plugins/Animation/ControlRig/Source/ControlRig/Private/Rigs/AdditiveControlRig.cpp
--- a/Engine/Plugins/Animation/ControlRig/Source/ControlRig/Private/Rigs/AdditiveControlRig.cpp
+++ b/Engine/Plugins/Animation/ControlRig/Source/ControlRig/Private/Rigs/AdditiveControlRig.cpp
@@ -59,6 +59,8 @@ void UAdditiveControlRig::ExecuteUnits(FRigUnitContext& InOutContext, const FNam
 
 void UAdditiveControlRig::Initialize(bool bInitRigUnits /*= true*/)
 {
+	PostInitInstanceIfRequired();
+	
 	Super::Initialize(bInitRigUnits);
 
 	if (GetObjectBinding() == nullptr)
@@ -91,6 +93,8 @@ void UAdditiveControlRig::Initialize(bool bInitRigUnits /*= true*/)
 
 void UAdditiveControlRig::CreateRigElements(const FReferenceSkeleton& InReferenceSkeleton, const FSmartNameMapping* InSmartNameMapping)
 {
+	PostInitInstanceIfRequired();
+	
 	GetHierarchy()->Reset();
 	if (URigHierarchyController* Controller = GetHierarchy()->GetController(true))
 	{
diff --git a/Engine/Plugins/Animation/ControlRig/Source/ControlRig/Private/Rigs/FKControlRig.cpp b/Engine/Plugins/Animation/ControlRig/Source/ControlRig/Private/Rigs/FKControlRig.cpp
--- a/Engine/Plugins/Animation/ControlRig/Source/ControlRig/Private/Rigs/FKControlRig.cpp
+++ b/Engine/Plugins/Animation/ControlRig/Source/ControlRig/Private/Rigs/FKControlRig.cpp
@@ -192,6 +192,8 @@ void UFKControlRig::SetBoneInitialTransformsFromSkeletalMeshComponent(USkeletalM
 
 void UFKControlRig::Initialize(bool bInitRigUnits /*= true*/)
 {
+	PostInitInstanceIfRequired();
+
 	Super::Initialize(bInitRigUnits);
 
 	if (GetObjectBinding() == nullptr)
@@ -250,6 +252,8 @@ void UFKControlRig::SetControlActive(const TArray<FFKBoneCheckInfo>& BoneChecks)
 
 void UFKControlRig::CreateRigElements(const FReferenceSkeleton& InReferenceSkeleton, const FSmartNameMapping* InSmartNameMapping)
 {
+	PostInitInstanceIfRequired();
+
 	GetHierarchy()->Reset();
 	if (URigHierarchyController* Controller = GetHierarchy()->GetController(true))
 	{
diff --git a/Engine/Plugins/Animation/ControlRig/Source/ControlRig/Public/ControlRig.h b/Engine/Plugins/Animation/ControlRig/Source/ControlRig/Public/ControlRig.h
--- a/Engine/Plugins/Animation/ControlRig/Source/ControlRig/Public/ControlRig.h
+++ b/Engine/Plugins/Animation/ControlRig/Source/ControlRig/Public/ControlRig.h
@@ -205,33 +205,6 @@ public:
 	UFUNCTION(BlueprintPure, Category = "Control Rig")
 	FORCEINLINE_DEBUGGABLE URigHierarchy* GetHierarchy()
 	{
-		// for CDO, when the hierarchy is first created using CreateDefaultSubobject()
-   		// the outer of that hierarchy will be Default_ControlRig instead of "this" because of
-   		// a call to the destructor of FObjectInitializer during CDO constructor that did it.
-   		// therefore, the hierarchy in CDO will be recreated in this function using NewObject() below 
-		if(DynamicHierarchy != nullptr && DynamicHierarchy->GetOuter() != this)
-		{
-			DynamicHierarchy = nullptr;
-		}
-		
-		if(DynamicHierarchy == nullptr && !IsGarbageCollecting())
-		{
-			DynamicHierarchy = NewObject<URigHierarchy>(this);
-			if (!HasAnyFlags(RF_ClassDefaultObject))
-			{
-				DynamicHierarchy->SetFlags(RF_Transient | RF_Transactional);
-			}
-			else
-			{
-				// for CDO, these flags are needed, they should have come with CreateDefaultSubobject()
-				// but since we are recreating a hierarchy here using NewObject, we need make sure the correct flags are manually added
-				DynamicHierarchy->SetFlags(DynamicHierarchy->GetFlags() | RF_Public | RF_DefaultSubObject);
-			}
-#if WITH_EDITOR
-			const TWeakObjectPtr<UControlRig> WeakThis = this;
-			DynamicHierarchy->OnUndoRedo().AddStatic(&UControlRig::OnHierarchyTransformUndoRedoWeak, WeakThis);
-#endif
-		}
 		return DynamicHierarchy;
 	}
 
@@ -274,9 +247,6 @@ public:
 	/** Sets the queue of events to run */
 	void SetEventQueue(const TArray<FName>& InEventNames);
 
-	/** Set the VM, and bind its commands */
-	void SetVM(URigVM* VM);
-
 	/** Update the settings such as array bound and log facilities */
 	void UpdateVMSettings();
 
@@ -426,9 +396,13 @@ public:
 
 	const FRigControlElementCustomization* GetControlCustomization(const FRigElementKey& InControl) const;
 	void SetControlCustomization(const FRigElementKey& InControl, const FRigControlElementCustomization& InCustomization);
-	
+
+	void PostInitInstanceIfRequired();
+
 private:
 
+	void PostInitInstance(UControlRig* InCDO);
+
 	UPROPERTY()
 	TMap<FRigElementKey, FRigControlElementCustomization> ControlCustomizations;
 
diff --git a/Engine/Plugins/Animation/ControlRig/Source/ControlRig/Public/ControlRigBlueprintGeneratedClass.h b/Engine/Plugins/Animation/ControlRig/Source/ControlRig/Public/ControlRigBlueprintGeneratedClass.h
--- a/Engine/Plugins/Animation/ControlRig/Source/ControlRig/Public/ControlRigBlueprintGeneratedClass.h
+++ b/Engine/Plugins/Animation/ControlRig/Source/ControlRig/Public/ControlRigBlueprintGeneratedClass.h
@@ -18,6 +18,7 @@ public:
 
 	// UClass interface
 	virtual uint8* GetPersistentUberGraphFrame(UObject* Obj, UFunction* FuncToCheck) const override;
+	virtual void PostInitInstance(UObject* InObj) override;
 
 	// UObject interface
 	void Serialize(FArchive& Ar);
diff --git a/Engine/Plugins/Animation/ControlRig/Source/ControlRigDeveloper/Private/ControlRigBlueprint.cpp b/Engine/Plugins/Animation/ControlRig/Source/ControlRigDeveloper/Private/ControlRigBlueprint.cpp
--- a/Engine/Plugins/Animation/ControlRig/Source/ControlRigDeveloper/Private/ControlRigBlueprint.cpp
+++ b/Engine/Plugins/Animation/ControlRig/Source/ControlRigDeveloper/Private/ControlRigBlueprint.cpp
@@ -345,7 +345,8 @@ void UControlRigBlueprint::PostLoad()
 
 	bVMRecompilationRequired = true;
 	{
-		FControlRigBlueprintVMCompileScope CompileScope(this);
+		TGuardValue<bool> IsCompilingGuard(bIsCompiling, true);
+		
 		TArray<UControlRigBlueprint*> ReferencedBlueprints = GetReferencedControlRigBlueprints();
 
 		// PostLoad all referenced BPs so that their function graphs are fully loaded 
@@ -508,7 +509,12 @@ void UControlRigBlueprint::PostLoad()
 #endif
 	}
 
+#if WITH_EDITOR
+	// delay compilation until the package has been loaded
+	FCoreUObjectDelegates::OnEndLoadPackage.AddUObject(this, &UControlRigBlueprint::HandlePackageDone);
+#else
 	RecompileVMIfRequired();
+#endif
 	RequestControlRigInit();
 
 	FCoreUObjectDelegates::OnObjectModified.RemoveAll(this);
@@ -522,6 +528,21 @@ void UControlRigBlueprint::PostLoad()
 	}
 }
 
+#if WITH_EDITOR
+void UControlRigBlueprint::HandlePackageDone(TConstArrayView<UPackage*> InPackages)
+{
+	if (!InPackages.Contains(GetPackage()))
+	{
+		return;
+	}
+
+	FCoreUObjectDelegates::OnEndLoadPackage.RemoveAll(this);
+	
+	RecompileVM();
+	RequestControlRigInit();
+}
+#endif
+
 void UControlRigBlueprint::RecompileVM()
 {
 	if(bIsCompiling)
@@ -545,11 +566,6 @@ void UControlRigBlueprint::RecompileVM()
 		CDO->VMRuntimeSettings = VMRuntimeSettings;
 		CDO->GetHierarchy()->CopyHierarchy(Hierarchy);
 
-		if (CDO->VM->GetOuter() != CDO)
-		{
-			CDO->SetVM(NewObject<URigVM>(CDO, TEXT("VM")));
-		}
-
 		if (!HasAnyFlags(RF_Transient | RF_Transactional))
 		{
 			CDO->Modify(false);
diff --git a/Engine/Plugins/Animation/ControlRig/Source/ControlRigDeveloper/Private/ControlRigBlueprintCompiler.cpp b/Engine/Plugins/Animation/ControlRig/Source/ControlRigDeveloper/Private/ControlRigBlueprintCompiler.cpp
--- a/Engine/Plugins/Animation/ControlRig/Source/ControlRigDeveloper/Private/ControlRigBlueprintCompiler.cpp
+++ b/Engine/Plugins/Animation/ControlRig/Source/ControlRigDeveloper/Private/ControlRigBlueprintCompiler.cpp
@@ -94,6 +94,7 @@ void FControlRigBlueprintCompilerContext::CopyTermDefaultsToDefaultObject(UObjec
 		// and in UControlRig::InitializeFromCDO,
 		// other Control Rig Instances are then initialized From the CDO 
 		UControlRig* ControlRig = CastChecked<UControlRig>(DefaultObject);
+		ControlRig->PostInitInstanceIfRequired();
 		
 		// copy hierarchy
 		{
diff --git a/Engine/Plugins/Animation/ControlRig/Source/ControlRigDeveloper/Public/ControlRigBlueprint.h b/Engine/Plugins/Animation/ControlRig/Source/ControlRigDeveloper/Public/ControlRigBlueprint.h
--- a/Engine/Plugins/Animation/ControlRig/Source/ControlRigDeveloper/Public/ControlRigBlueprint.h
+++ b/Engine/Plugins/Animation/ControlRig/Source/ControlRigDeveloper/Public/ControlRigBlueprint.h
@@ -558,6 +558,10 @@ private:
 
 	void HandleHierarchyModified(ERigHierarchyNotification InNotification, URigHierarchy* InHierarchy, const FRigBaseElement* InElement);
 
+#if WITH_EDITOR
+	void HandlePackageDone(TConstArrayView<UPackage*> InPackages);
+#endif
+
 	// Class used to temporarily cache all 
 	// current control values and reapply them
 	// on destruction
diff --git a/Engine/Source/Runtime/RigVM/Private/RigVMCore/RigVM.cpp b/Engine/Source/Runtime/RigVM/Private/RigVMCore/RigVM.cpp
--- a/Engine/Source/Runtime/RigVM/Private/RigVMCore/RigVM.cpp
+++ b/Engine/Source/Runtime/RigVM/Private/RigVMCore/RigVM.cpp
@@ -115,6 +115,16 @@ void URigVM::Serialize(FArchive& Ar)
 		return;
 	}
 
+#if !UE_RIGVM_UCLASS_BASED_STORAGE_DISABLED
+
+	// call into the super class to serialize any uproperty
+	if(Ar.IsObjectReferenceCollector() || Ar.IsCountingMemory())
+	{
+		Super::Serialize(Ar);
+	}
+	
+#endif
+
 	ensure(ExecutingThreadId == INDEX_NONE);
 
 	if (Ar.IsSaving() || Ar.IsObjectReferenceCollector() || Ar.IsCountingMemory())
diff --git a/Engine/Source/Runtime/RigVM/Private/RigVMCore/RigVMMemoryStorage.cpp b/Engine/Source/Runtime/RigVM/Private/RigVMCore/RigVMMemoryStorage.cpp
--- a/Engine/Source/Runtime/RigVM/Private/RigVMCore/RigVMMemoryStorage.cpp
+++ b/Engine/Source/Runtime/RigVM/Private/RigVMCore/RigVMMemoryStorage.cpp
@@ -309,9 +309,7 @@ URigVMMemoryStorageGeneratorClass* URigVMMemoryStorageGeneratorClass::CreateStor
 	URigVMMemoryStorageGeneratorClass* OldClass = FindObject<URigVMMemoryStorageGeneratorClass>(Package, *ClassName);
 	if(OldClass)
 	{
-		OldClass->RemoveFromRoot();
 		OldClass->Rename(nullptr, GetTransientPackage(), REN_ForceNoResetLoaders | REN_DoNotDirty | REN_DontCreateRedirectors | REN_NonTransactional);
-		OldClass->MarkPendingKill();
 	}
 
 	// create the new class
@@ -321,9 +319,6 @@ URigVMMemoryStorageGeneratorClass* URigVMMemoryStorageGeneratorClass::CreateStor
 		RF_Standalone | RF_Public
 	);
 
-	// add it to the root so that the garbage collector doesn't pick it up
-	Class->AddToRoot();
-
 	// clear the class (sets relevant flags)
 	Class->PurgeClass(false);
 
@@ -390,9 +385,7 @@ bool URigVMMemoryStorageGeneratorClass::RemoveStorageClass(UObject* InOuter, ERi
 	URigVMMemoryStorageGeneratorClass* OldClass = FindObject<URigVMMemoryStorageGeneratorClass>(Package, *ClassName);
 	if(OldClass)
 	{
-		OldClass->RemoveFromRoot();
 		OldClass->Rename(nullptr, GetTransientPackage(), REN_ForceNoResetLoaders | REN_DoNotDirty | REN_DontCreateRedirectors | REN_NonTransactional);
-		OldClass->MarkPendingKill();
 		return true;
 	}
 
diff --git a/Engine/Source/Runtime/RigVM/Public/RigVMCore/RigVM.h b/Engine/Source/Runtime/RigVM/Public/RigVMCore/RigVM.h
--- a/Engine/Source/Runtime/RigVM/Public/RigVMCore/RigVM.h
+++ b/Engine/Source/Runtime/RigVM/Public/RigVMCore/RigVM.h
@@ -247,13 +247,13 @@ public:
 
 #endif
 
-	UPROPERTY(transient, DuplicateTransient)
+	UPROPERTY()
 	TObjectPtr<URigVMMemoryStorage> WorkMemoryStorageObject;
 
-	UPROPERTY(transient, DuplicateTransient)
+	UPROPERTY()
 	TObjectPtr<URigVMMemoryStorage> LiteralMemoryStorageObject;
 
-	UPROPERTY(transient, DuplicateTransient)
+	UPROPERTY()
 	TObjectPtr<URigVMMemoryStorage> DebugMemoryStorageObject;
 
 	TArray<FRigVMPropertyPathDescription> ExternalPropertyPathDescriptions;
