Control Rig: Invoke PostInitInstance in post load cases

diff --git a/Engine/Plugins/Animation/ControlRig/Source/ControlRig/Private/ControlRig.cpp b/Engine/Plugins/Animation/ControlRig/Source/ControlRig/Private/ControlRig.cpp
--- a/Engine/Plugins/Animation/ControlRig/Source/ControlRig/Private/ControlRig.cpp
+++ b/Engine/Plugins/Animation/ControlRig/Source/ControlRig/Private/ControlRig.cpp
@@ -2813,6 +2813,7 @@ void UControlRig::PostInitInstance(UControlRig* InCDO)
 
 	if(!HasAnyFlags(RF_ClassDefaultObject) && InCDO)
 	{
+		InCDO->PostInitInstanceIfRequired();
 		VM->CopyFrom(InCDO->GetVM());
 		DynamicHierarchy->CopyHierarchy(InCDO->GetHierarchy());
 	}
diff --git a/Engine/Plugins/Animation/ControlRig/Source/ControlRigDeveloper/Private/ControlRigBlueprint.cpp b/Engine/Plugins/Animation/ControlRig/Source/ControlRigDeveloper/Private/ControlRigBlueprint.cpp
--- a/Engine/Plugins/Animation/ControlRig/Source/ControlRigDeveloper/Private/ControlRigBlueprint.cpp
+++ b/Engine/Plugins/Animation/ControlRig/Source/ControlRigDeveloper/Private/ControlRigBlueprint.cpp
@@ -563,6 +563,7 @@ void UControlRigBlueprint::RecompileVM()
 		TGuardValue<bool> ReentrantGuardSelf(bSuspendModelNotificationsForSelf, true);
 		TGuardValue<bool> ReentrantGuardOthers(bSuspendModelNotificationsForOthers, true);
 
+		CDO->PostInitInstanceIfRequired();
 		CDO->VMRuntimeSettings = VMRuntimeSettings;
 		CDO->GetHierarchy()->CopyHierarchy(Hierarchy);
 
@@ -3352,6 +3353,7 @@ void UControlRigBlueprint::PropagatePoseFromBPToInstances()
 	{
 		if (UControlRig* DefaultObject = Cast<UControlRig>(MyControlRigClass->GetDefaultObject(false)))
 		{
+			DefaultObject->PostInitInstanceIfRequired();
 			DefaultObject->GetHierarchy()->CopyPose(Hierarchy, true, true);
 
 			TArray<UObject*> ArchetypeInstances;
@@ -3360,6 +3362,7 @@ void UControlRigBlueprint::PropagatePoseFromBPToInstances()
 			{
 				if (UControlRig* InstanceRig = Cast<UControlRig>(ArchetypeInstance))
 				{
+					InstanceRig->PostInitInstanceIfRequired();
 					InstanceRig->GetHierarchy()->CopyPose(Hierarchy, true, true);
 				}
 			}
@@ -3373,6 +3376,7 @@ void UControlRigBlueprint::PropagateHierarchyFromBPToInstances()
 	{
 		if (UControlRig* DefaultObject = Cast<UControlRig>(MyControlRigClass->GetDefaultObject(false)))
 		{
+			DefaultObject->PostInitInstanceIfRequired();
 			DefaultObject->GetHierarchy()->CopyHierarchy(Hierarchy);
 			DefaultObject->Initialize(true);
 
@@ -3382,6 +3386,7 @@ void UControlRigBlueprint::PropagateHierarchyFromBPToInstances()
 			{
 				if (UControlRig* InstanceRig = Cast<UControlRig>(ArchetypeInstance))
 				{
+					InstanceRig->PostInitInstanceIfRequired();
 					InstanceRig->GetHierarchy()->CopyHierarchy(Hierarchy);
 					InstanceRig->Initialize(true);
 				}
@@ -3482,6 +3487,7 @@ void UControlRigBlueprint::PropagatePropertyFromBPToInstances(FRigElementKey InR
 			{
 				if (UControlRig* InstanceRig = Cast<UControlRig>(ArchetypeInstance))
 				{
+					InstanceRig->PostInitInstanceIfRequired();
 					uint8* Dest = ((uint8*)InstanceRig->GetHierarchy()->Get(ElementIndex)) + PropertyOffset;
 					FMemory::Memcpy(Dest, Source, PropertySize);
 				}
