Control Rig: Fix losing debug object when it is replaced

diff --git a/Engine/Plugins/Animation/ControlRig/Source/ControlRigEditor/Private/Editor/ControlRigEditor.cpp b/Engine/Plugins/Animation/ControlRig/Source/ControlRigEditor/Private/Editor/ControlRigEditor.cpp
--- a/Engine/Plugins/Animation/ControlRig/Source/ControlRigEditor/Private/Editor/ControlRigEditor.cpp
+++ b/Engine/Plugins/Animation/ControlRig/Source/ControlRigEditor/Private/Editor/ControlRigEditor.cpp
@@ -1271,6 +1271,7 @@ void FControlRigEditor::HandleSetObjectBeingDebugged(UObject* InObject)
 	}
 
 	RefreshDetailView();
+	LastDebuggedRig = GetCustomDebugObjectLabel(DebuggedControlRig);
 }
 
 FString FControlRigEditor::GetCustomDebugObjectLabel(UObject* ObjectBeingDebugged) const
@@ -1797,8 +1798,6 @@ void FControlRigEditor::Compile()
 			ValueScope = MakeUnique<UControlRigBlueprint::FControlValueScope>(GetControlRigBlueprint());
 		}
 
-		LastDebuggedRig.Empty();
-
 		// force to disable the supended notif brackets
 		if (UControlRigBlueprint* RigBlueprint = Cast<UControlRigBlueprint>(GetBlueprintObj()))
 		{
@@ -2978,6 +2977,21 @@ void FControlRigEditor::Tick(float DeltaTime)
 			ControlRig->Evaluate_AnyThread();
 			bDrawHierarchyBones = true;
 		}
+
+		if (LastDebuggedRig != GetCustomDebugObjectLabel(Blueprint->GetObjectBeingDebugged()))
+		{
+			TArray<FCustomDebugObject> DebugList;
+			GetCustomDebugObjects(DebugList);
+
+			for (const FCustomDebugObject& DebugObject : DebugList)
+			{
+				if (DebugObject.NameOverride == LastDebuggedRig)
+				{
+					GetBlueprintObj()->SetObjectBeingDebugged(DebugObject.Object);
+					break;
+				}
+			}
+		}
 	}
 
 	if (FControlRigEditorEditMode* EditMode = GetEditMode())
