Control Rig: Guard importing text onto transform with error pipe

diff --git a/Engine/Plugins/Animation/ControlRig/Source/ControlRigEditor/Private/Editor/SRigHierarchy.cpp b/Engine/Plugins/Animation/ControlRig/Source/ControlRigEditor/Private/Editor/SRigHierarchy.cpp
--- a/Engine/Plugins/Animation/ControlRig/Source/ControlRigEditor/Private/Editor/SRigHierarchy.cpp
+++ b/Engine/Plugins/Animation/ControlRig/Source/ControlRigEditor/Private/Editor/SRigHierarchy.cpp
@@ -1685,6 +1685,25 @@ void SRigHierarchy::HandlePasteItems()
 	RefreshTreeView();
 }
 
+class SRigHierarchyPasteTransformsErrorPipe : public FOutputDevice
+{
+public:
+
+	int32 NumErrors;
+
+	SRigHierarchyPasteTransformsErrorPipe()
+		: FOutputDevice()
+		, NumErrors(0)
+	{
+	}
+
+	virtual void Serialize(const TCHAR* V, ELogVerbosity::Type Verbosity, const class FName& Category) override
+	{
+		UE_LOG(LogControlRig, Error, TEXT("Error importing transforms to Hierarchy: %s"), V);
+		NumErrors++;
+	}
+};
+
 void SRigHierarchy::HandlePasteLocalTransforms()
 {
 	HandlePasteTransforms(ERigTransformType::CurrentLocal, true);
@@ -1704,9 +1723,14 @@ void SRigHierarchy::HandlePasteTransforms(ERigTransformType::Type InTransformTyp
 
 		FScopedTransaction Transaction(LOCTEXT("HierarchyTreePaste", "Pasted transforms."));
 
+		SRigHierarchyPasteTransformsErrorPipe ErrorPipe;
 		FRigHierarchyCopyPasteContent Data;
-		FRigHierarchyCopyPasteContent::StaticStruct()->ImportText(*Content, &Data, nullptr, EPropertyPortFlags::PPF_None, nullptr, FRigHierarchyCopyPasteContent::StaticStruct()->GetName(), true);
-
+		FRigHierarchyCopyPasteContent::StaticStruct()->ImportText(*Content, &Data, nullptr, EPropertyPortFlags::PPF_None, &ErrorPipe, FRigHierarchyCopyPasteContent::StaticStruct()->GetName(), true);
+		if(ErrorPipe.NumErrors > 0)
+		{
+			return;
+		}
+		
 		URigHierarchy* DebuggedHierarchy = GetDebuggedHierarchy();
 
 		const TArray<FRigElementKey> CurrentSelection = Hierarchy->GetSelectedKeys();
