Control Rig: Change order of elements within Space Picker widget : show controls first

diff --git a/Engine/Plugins/Animation/ControlRig/Source/ControlRigEditor/Private/Editor/SRigHierarchyTreeView.cpp b/Engine/Plugins/Animation/ControlRig/Source/ControlRigEditor/Private/Editor/SRigHierarchyTreeView.cpp
--- a/Engine/Plugins/Animation/ControlRig/Source/ControlRigEditor/Private/Editor/SRigHierarchyTreeView.cpp
+++ b/Engine/Plugins/Animation/ControlRig/Source/ControlRigEditor/Private/Editor/SRigHierarchyTreeView.cpp
@@ -465,6 +465,14 @@ void SRigHierarchyTreeView::RefreshTreeView(bool bRebuildContent)
 				RestoreSparseItemInfos(Pair.Value);
 			}
 
+			if(Delegates.OnCompareKeys.IsBound())
+			{
+				Algo::Sort(RootElements, [&](const TSharedPtr<FRigTreeElement>& A, const TSharedPtr<FRigTreeElement>& B)
+				{
+					return Delegates.OnCompareKeys.Execute(A->Key, B->Key);
+				});
+			}
+
 			if (RootElements.Num() > 0)
 			{
 				AddSpacerElement();
diff --git a/Engine/Plugins/Animation/ControlRig/Source/ControlRigEditor/Private/Editor/SRigHierarchyTreeView.h b/Engine/Plugins/Animation/ControlRig/Source/ControlRigEditor/Private/Editor/SRigHierarchyTreeView.h
--- a/Engine/Plugins/Animation/ControlRig/Source/ControlRigEditor/Private/Editor/SRigHierarchyTreeView.h
+++ b/Engine/Plugins/Animation/ControlRig/Source/ControlRigEditor/Private/Editor/SRigHierarchyTreeView.h
@@ -61,6 +61,7 @@ DECLARE_DELEGATE_RetVal(const URigHierarchy*, FOnGetRigTreeHierarchy);
 DECLARE_DELEGATE_RetVal(const FRigTreeDisplaySettings&, FOnGetRigTreeDisplaySettings);
 DECLARE_DELEGATE_RetVal_TwoParams(FName, FOnRigTreeRenameElement, const FRigElementKey& /*OldKey*/, const FString& /*NewName*/);
 DECLARE_DELEGATE_RetVal_ThreeParams(bool, FOnRigTreeVerifyElementNameChanged, const FRigElementKey& /*OldKey*/, const FString& /*NewName*/, FText& /*OutErrorMessage*/);
+DECLARE_DELEGATE_RetVal_TwoParams(bool, FOnRigTreeCompareKeys, const FRigElementKey& /*A*/, const FRigElementKey& /*B*/);
 
 typedef STableRow<TSharedPtr<FRigTreeElement>>::FOnCanAcceptDrop FOnRigTreeCanAcceptDrop;
 typedef STableRow<TSharedPtr<FRigTreeElement>>::FOnAcceptDrop FOnRigTreeAcceptDrop;
@@ -83,6 +84,7 @@ struct CONTROLRIGEDITOR_API FRigTreeDelegates
 	FOnRigTreeMouseButtonClick OnMouseButtonClick;
 	FOnRigTreeMouseButtonDoubleClick OnMouseButtonDoubleClick;
 	FOnRigTreeSetExpansionRecursive OnSetExpansionRecursive;
+	FOnRigTreeCompareKeys OnCompareKeys;
 
 	FRigTreeDelegates()
 	{
diff --git a/Engine/Plugins/Animation/ControlRig/Source/ControlRigEditor/Private/Editor/SRigSpacePickerWidget.cpp b/Engine/Plugins/Animation/ControlRig/Source/ControlRigEditor/Private/Editor/SRigSpacePickerWidget.cpp
--- a/Engine/Plugins/Animation/ControlRig/Source/ControlRigEditor/Private/Editor/SRigSpacePickerWidget.cpp
+++ b/Engine/Plugins/Animation/ControlRig/Source/ControlRigEditor/Private/Editor/SRigSpacePickerWidget.cpp
@@ -631,6 +631,20 @@ FReply SRigSpacePickerWidget::HandleAddElementClicked()
 		}
 	});
 
+	TreeDelegates.OnCompareKeys = FOnRigTreeCompareKeys::CreateLambda([](const FRigElementKey& A, const FRigElementKey& B) -> bool
+	{
+		// controls should always show up first - so we'll sort them to the start of the list
+		if(A.Type == ERigElementType::Control && B.Type != ERigElementType::Control)
+		{
+			return true;
+		}
+		if(B.Type == ERigElementType::Control && A.Type != ERigElementType::Control)
+		{
+			return false;
+		}
+		return A < B;
+	});
+
 	TSharedPtr<SSearchableRigHierarchyTreeView> SearchableTreeView = SNew(SSearchableRigHierarchyTreeView)
 	.RigTreeDelegates(TreeDelegates);
 	SearchableTreeView->GetTreeView()->RefreshTreeView(true);
