Control Rig: Change label on bake button for space switching

diff --git a/Engine/Plugins/Animation/ControlRig/Source/ControlRigEditor/Private/Editor/SRigSpacePickerWidget.cpp b/Engine/Plugins/Animation/ControlRig/Source/ControlRigEditor/Private/Editor/SRigSpacePickerWidget.cpp
--- a/Engine/Plugins/Animation/ControlRig/Source/ControlRigEditor/Private/Editor/SRigSpacePickerWidget.cpp
+++ b/Engine/Plugins/Animation/ControlRig/Source/ControlRigEditor/Private/Editor/SRigSpacePickerWidget.cpp
@@ -169,7 +169,7 @@ void SRigSpacePickerWidget::Construct(const FArguments& InArgs)
 			[
 				SNew(SButton)
 				.ButtonStyle(FEditorStyle::Get(), "FlatButton.Default")
-				.Text(LOCTEXT("BakeButton", "Bake"))
+				.Text(LOCTEXT("BakeButton", "Bake..."))
 				.OnClicked(InArgs._OnBakeButtonClicked)
 				.ToolTipText(LOCTEXT("BakeButtonToolTip", "Allows to bake the animation of one or more controls to a single space."))
 			];
