Control Rig: Move add button for space switching to top right expandable area

diff --git a/Engine/Plugins/Animation/ControlRig/Source/ControlRigEditor/Private/EditMode/SControlRigEditModeTools.cpp b/Engine/Plugins/Animation/ControlRig/Source/ControlRigEditor/Private/EditMode/SControlRigEditModeTools.cpp
--- a/Engine/Plugins/Animation/ControlRig/Source/ControlRigEditor/Private/EditMode/SControlRigEditModeTools.cpp
+++ b/Engine/Plugins/Animation/ControlRig/Source/ControlRigEditor/Private/EditMode/SControlRigEditModeTools.cpp
@@ -319,12 +319,51 @@ void SControlRigEditModeTools::Construct(const FArguments& InArgs, FControlRigEd
 				.AreaTitleFont(FEditorStyle::GetFontStyle("DetailsView.CategoryFontStyle"))
 				.BorderBackgroundColor(FLinearColor(.6f, .6f, .6f))
 				.Padding(FMargin(8.f))
+				.HeaderContent()
+				[
+					SNew(SHorizontalBox)
+
+					+SHorizontalBox::Slot()
+					.AutoWidth()
+					.HAlign(HAlign_Left)
+					.VAlign(VAlign_Center)
+					.Padding(0.f, 0.f, 0.f, 0.f)
+					[
+						SNew(STextBlock)
+						.Text(LOCTEXT("Picker_SpaceWidget", "Spaces"))
+						.Font(FCoreStyle::Get().GetFontStyle("ExpandableArea.TitleFont"))
+					]
+					
+					+SHorizontalBox::Slot()
+					.FillWidth(1.f)
+					[
+						SNew(SSpacer)
+					]
+
+					+SHorizontalBox::Slot()
+					.AutoWidth()
+					.HAlign(HAlign_Right)
+					.VAlign(VAlign_Center)
+					.Padding(0.f, 2.f, 8.f, 2.f)
+					[
+						SNew(SButton)
+						.ContentPadding(0.0f)
+						.ButtonStyle(FEditorStyle::Get(), "NoBorder")
+						.OnClicked(this, &SControlRigEditModeTools::HandleAddSpaceClicked)
+						.Cursor(EMouseCursor::Default)
+						.ToolTipText(LOCTEXT("AddSpace", "Add Space"))
+						[
+							SNew(SImage)
+							.Image(FEditorStyle::GetBrush(TEXT("Icons.PlusCircle")))
+						]
+					]
+				]
 				.BodyContent()
 				[
 					SAssignNew(SpacePickerWidget, SRigSpacePickerWidget)
 					.AllowDelete(true)
 					.AllowReorder(true)
-					.AllowAdd(true)
+					.AllowAdd(false)
 					.ShowBakeButton(true)
 					.GetControlCustomization(this, &SControlRigEditModeTools::HandleGetControlElementCustomization)
 					.OnActiveSpaceChanged(this, &SControlRigEditModeTools::HandleActiveSpaceChanged)
@@ -692,6 +731,11 @@ void SControlRigEditModeTools::HandleSpaceListChanged(URigHierarchy* InHierarchy
 	}
 }
 
+FReply SControlRigEditModeTools::HandleAddSpaceClicked()
+{
+	return SpacePickerWidget->HandleAddElementClicked();
+}
+
 FReply SControlRigEditModeTools::OnBakeControlsToNewSpaceButtonClicked()
 {
 	if (SpacePickerWidget->GetHierarchy() == nullptr)
diff --git a/Engine/Plugins/Animation/ControlRig/Source/ControlRigEditor/Private/EditMode/SControlRigEditModeTools.h b/Engine/Plugins/Animation/ControlRig/Source/ControlRigEditor/Private/EditMode/SControlRigEditModeTools.h
--- a/Engine/Plugins/Animation/ControlRig/Source/ControlRigEditor/Private/EditMode/SControlRigEditModeTools.h
+++ b/Engine/Plugins/Animation/ControlRig/Source/ControlRigEditor/Private/EditMode/SControlRigEditModeTools.h
@@ -85,6 +85,7 @@ private:
 	const FRigControlElementCustomization* HandleGetControlElementCustomization(URigHierarchy* InHierarchy, const FRigElementKey& InControlKey);
 	void HandleActiveSpaceChanged(URigHierarchy* InHierarchy, const FRigElementKey& InControlKey, const FRigElementKey& InSpaceKey);
 	void HandleSpaceListChanged(URigHierarchy* InHierarchy, const FRigElementKey& InControlKey, const TArray<FRigElementKey>& InSpaceList);
+	FReply HandleAddSpaceClicked();
 	FReply OnBakeControlsToNewSpaceButtonClicked();
 
 	EVisibility GetRigOptionExpanderVisibility() const;
diff --git a/Engine/Plugins/Animation/ControlRig/Source/ControlRigEditor/Private/Editor/SRigSpacePickerWidget.h b/Engine/Plugins/Animation/ControlRig/Source/ControlRigEditor/Private/Editor/SRigSpacePickerWidget.h
--- a/Engine/Plugins/Animation/ControlRig/Source/ControlRigEditor/Private/Editor/SRigSpacePickerWidget.h
+++ b/Engine/Plugins/Animation/ControlRig/Source/ControlRigEditor/Private/Editor/SRigSpacePickerWidget.h
@@ -138,7 +138,13 @@ private:
 	FReply HandleSpaceMoveUp(FRigElementKey InKey);
 	FReply HandleSpaceMoveDown(FRigElementKey InKey);
 	void HandleSpaceDelete(FRigElementKey InKey);
+
+public:
+	
 	FReply HandleAddElementClicked();
+
+private:
+	
 	bool IsSpaceMoveUpEnabled(FRigElementKey InKey) const;
 	bool IsSpaceMoveDownEnabled(FRigElementKey InKey) const;
 
