Control Rig: only mark blueprint as modified if needed - avoid markup if autorecompile is turned on

diff --git a/Engine/Plugins/Animation/ControlRig/Source/ControlRigDeveloper/Private/ControlRigBlueprint.cpp b/Engine/Plugins/Animation/ControlRig/Source/ControlRigDeveloper/Private/ControlRigBlueprint.cpp
--- a/Engine/Plugins/Animation/ControlRig/Source/ControlRigDeveloper/Private/ControlRigBlueprint.cpp
+++ b/Engine/Plugins/Animation/ControlRig/Source/ControlRigDeveloper/Private/ControlRigBlueprint.cpp
@@ -2905,7 +2905,13 @@ void UControlRigBlueprint::HandleModifiedEvent(ERigVMGraphNotifType InNotifType,
 				ClearTransientControls();
 				RequestAutoVMRecompilation();
 				MarkPackageDirty();
-				FBlueprintEditorUtils::MarkBlueprintAsModified(this);
+
+				// we don't need to mark the blueprint as modified since we only
+				// need to recompile the VM here - unless we don't auto recompile.
+				if(!bAutoRecompileVM)
+				{
+					FBlueprintEditorUtils::MarkBlueprintAsModified(this);
+				}
 				break;
 			}
 			case ERigVMGraphNotifType::PinWatchedChanged:
