#UE Do not clear RigVM memory class pointers on load in cooked builds. Let them be loaded from disk

Original CL 17586012

diff --git a/Engine/Plugins/Animation/ControlRig/Source/ControlRig/Private/ControlRig.cpp b/Engine/Plugins/Animation/ControlRig/Source/ControlRig/Private/ControlRig.cpp
--- a/Engine/Plugins/Animation/ControlRig/Source/ControlRig/Private/ControlRig.cpp
+++ b/Engine/Plugins/Animation/ControlRig/Source/ControlRig/Private/ControlRig.cpp
@@ -2796,9 +2796,14 @@ void UControlRig::PostInitInstance(UControlRig* InCDO)
 
 	// set up the VM
 	VM = NewObject<URigVM>(this, TEXT("VM"), SubObjectFlags);
-	VM->GetMemoryByType(ERigVMMemoryType::Work, true);
-	VM->GetMemoryByType(ERigVMMemoryType::Literal, true);
-	VM->GetMemoryByType(ERigVMMemoryType::Debug, true);
+
+	// Cooked platforms will load these pointers from disk
+	if (!FPlatformProperties::RequiresCookedData())
+	{
+		VM->GetMemoryByType(ERigVMMemoryType::Work, true);
+		VM->GetMemoryByType(ERigVMMemoryType::Literal, true);
+		VM->GetMemoryByType(ERigVMMemoryType::Debug, true);
+	}
 
 	VM->ExecutionReachedExit().AddUObject(this, &UControlRig::HandleExecutionReachedExit);
 	UpdateVMSettings();
diff --git a/Engine/Source/Runtime/RigVM/Private/RigVMCore/RigVM.cpp b/Engine/Source/Runtime/RigVM/Private/RigVMCore/RigVM.cpp
--- a/Engine/Source/Runtime/RigVM/Private/RigVMCore/RigVM.cpp
+++ b/Engine/Source/Runtime/RigVM/Private/RigVMCore/RigVM.cpp
@@ -203,7 +203,11 @@ void URigVM::Load(FArchive& Ar)
 #if !UE_RIGVM_UCLASS_BASED_STORAGE_DISABLED
 
 	// requesting the memory types will create them
-	ClearMemory();
+	// Cooked platforms will just load the objects and do no need to clear the referenes
+	if (!FPlatformProperties::RequiresCookedData())
+	{
+		ClearMemory();
+	}
 	
 	Ar << ExternalPropertyPathDescriptions;
 	Ar << FunctionNamesStorage;
