Control Rig: Sustain mapped direction from component to element

diff --git a/Engine/Plugins/Animation/ControlRig/Source/ControlRig/Private/ControlRigComponent.cpp b/Engine/Plugins/Animation/ControlRig/Source/ControlRig/Private/ControlRigComponent.cpp
--- a/Engine/Plugins/Animation/ControlRig/Source/ControlRig/Private/ControlRigComponent.cpp
+++ b/Engine/Plugins/Animation/ControlRig/Source/ControlRig/Private/ControlRigComponent.cpp
@@ -509,6 +509,7 @@ void UControlRigComponent::AddMappedComponents(TArray<FControlRigComponentMapped
 
 		ElementToMap.ElementName = ComponentToMap.ElementName;
 		ElementToMap.ElementType = ComponentToMap.ElementType;
+		ElementToMap.Direction = ComponentToMap.Direction;
 
 		ElementsToMap.Add(ElementToMap);
 	}
