Control Rig: Hide Transform and TransformNoScale as available control types

diff --git a/Engine/Plugins/Animation/ControlRig/Source/ControlRig/Public/Rigs/RigHierarchyDefines.h b/Engine/Plugins/Animation/ControlRig/Source/ControlRig/Public/Rigs/RigHierarchyDefines.h
--- a/Engine/Plugins/Animation/ControlRig/Source/ControlRig/Public/Rigs/RigHierarchyDefines.h
+++ b/Engine/Plugins/Animation/ControlRig/Source/ControlRig/Public/Rigs/RigHierarchyDefines.h
@@ -87,8 +87,8 @@ enum class ERigControlType : uint8
     Position,
     Scale,
     Rotator,
-    Transform,
-    TransformNoScale,
+    Transform UMETA(Hidden),
+    TransformNoScale UMETA(Hidden),
     EulerTransform,
 };
 
