Control Rig: Minor Menu Item Name change for the menu that shows up when drag and drop from the rig hierarchy to the graph + Refactor this menu to also use UToolMenu system

diff --git a/Engine/Plugins/Animation/ControlRig/Source/ControlRigEditor/Private/Editor/ControlRigContextMenuContext.cpp b/Engine/Plugins/Animation/ControlRig/Source/ControlRigEditor/Private/Editor/ControlRigContextMenuContext.cpp
--- a/Engine/Plugins/Animation/ControlRig/Source/ControlRigEditor/Private/Editor/ControlRigContextMenuContext.cpp
+++ b/Engine/Plugins/Animation/ControlRig/Source/ControlRigEditor/Private/Editor/ControlRigContextMenuContext.cpp
@@ -6,11 +6,20 @@
 #include "ControlRigEditor.h"
 #include "Slate/Public/Framework/Application/SlateApplication.h"
 
-void UControlRigContextMenuContext::Init(TWeakObjectPtr<UControlRigBlueprint> InControlRigBlueprint, const FControlRigRigHierarchyDragAndDropContext& InDragAndDropContext, const FControlRigGraphNodeContextMenuContext& InGraphNodeContext)
+FString FControlRigRigHierarchyToGraphDragAndDropContext::GetSectionTitle() const
+{
+	TArray<FString> ElementNameStrings;
+	for (const FRigElementKey& Element: DraggedElementKeys)
+	{
+		ElementNameStrings.Add(Element.Name.ToString());
+	}
+	return FString::Join(ElementNameStrings, TEXT(","));
+}
+
+void UControlRigContextMenuContext::Init(TWeakObjectPtr<UControlRigBlueprint> InControlRigBlueprint, const FControlRigMenuSpecificContext& InMenuSpecificContext)
 {
 	ControlRigBlueprint = InControlRigBlueprint;
-	DragAndDropContext = InDragAndDropContext;
-	GraphNodeContextMenuContext = InGraphNodeContext;
+	MenuSpecificContext = InMenuSpecificContext;
 }
 
 UControlRigBlueprint* UControlRigContextMenuContext::GetControlRigBlueprint() const
@@ -37,10 +46,15 @@ bool UControlRigContextMenuContext::IsAltDown() const
 
 FControlRigRigHierarchyDragAndDropContext UControlRigContextMenuContext::GetRigHierarchyDragAndDropContext()
 {
-	return DragAndDropContext;
+	return MenuSpecificContext.RigHierarchyDragAndDropContext;
 }
 
 FControlRigGraphNodeContextMenuContext UControlRigContextMenuContext::GetGraphNodeContextMenuContext()
 {
-	return GraphNodeContextMenuContext;
+	return MenuSpecificContext.GraphNodeContextMenuContext;
 }
+
+FControlRigRigHierarchyToGraphDragAndDropContext UControlRigContextMenuContext::GetRigHierarchyToGraphDragAndDropContext()
+{
+	return MenuSpecificContext.RigHierarchyToGraphDragAndDropContext;
+}
\ No newline at end of file
diff --git a/Engine/Plugins/Animation/ControlRig/Source/ControlRigEditor/Private/Editor/ControlRigContextMenuContext.h b/Engine/Plugins/Animation/ControlRig/Source/ControlRigEditor/Private/Editor/ControlRigContextMenuContext.h
--- a/Engine/Plugins/Animation/ControlRig/Source/ControlRigEditor/Private/Editor/ControlRigContextMenuContext.h
+++ b/Engine/Plugins/Animation/ControlRig/Source/ControlRigEditor/Private/Editor/ControlRigContextMenuContext.h
@@ -14,6 +14,7 @@ class FControlRigEditor;
 class URigVMGraph;
 class URigVMNode;
 class URigVMPin;
+class UEdGraph;
 
 USTRUCT(BlueprintType)
 struct FControlRigRigHierarchyDragAndDropContext
@@ -67,6 +68,38 @@ struct FControlRigGraphNodeContextMenuContext
 	TObjectPtr<const URigVMPin> Pin;
 };
 
+USTRUCT(BlueprintType)
+struct FControlRigRigHierarchyToGraphDragAndDropContext
+{
+	GENERATED_BODY()
+
+	FControlRigRigHierarchyToGraphDragAndDropContext() = default;
+	
+	FControlRigRigHierarchyToGraphDragAndDropContext(const TArray<FRigElementKey>& InDraggedElementKeys, UEdGraph* InGraph, const FVector2D& InNodePosition)
+		:DraggedElementKeys(InDraggedElementKeys)
+		,Graph(InGraph)
+		,NodePosition(InNodePosition)
+	{ };
+
+	UPROPERTY(BlueprintReadOnly, Category = ControlRigEditorExtensions)
+	TArray<FRigElementKey> DraggedElementKeys;
+
+	TWeakObjectPtr<UEdGraph> Graph;
+
+	FVector2D NodePosition;
+
+	FString GetSectionTitle() const;
+};
+
+struct FControlRigMenuSpecificContext
+{
+	FControlRigRigHierarchyDragAndDropContext RigHierarchyDragAndDropContext;
+	
+	FControlRigGraphNodeContextMenuContext GraphNodeContextMenuContext;
+
+	FControlRigRigHierarchyToGraphDragAndDropContext RigHierarchyToGraphDragAndDropContext;
+};
+
 UCLASS(BlueprintType)
 class UControlRigContextMenuContext : public UObject
 {
@@ -76,10 +109,9 @@ public:
 	/**
 	 *	Initialize the Context
 	 * @param InControlRigBlueprint 	The Control Rig Bluerpint currently opened in the Editor
-	 * @param InDragAndDropContext 		Optional, only used for the menu that shows up after a drag & drop action
-	 * @param InGraphNodeContext 		Optional, only used for the graph node context menu
+	 * @param InMenuSpecificContext 	Additional context for specific menus
 	*/
-	void Init(TWeakObjectPtr<UControlRigBlueprint> InControlRigBlueprint, const FControlRigRigHierarchyDragAndDropContext& InDragAndDropContext = FControlRigRigHierarchyDragAndDropContext(), const FControlRigGraphNodeContextMenuContext& InGraphNodeContext = FControlRigGraphNodeContextMenuContext());
+	void Init(TWeakObjectPtr<UControlRigBlueprint> InControlRigBlueprint, const FControlRigMenuSpecificContext& InMenuSpecificContext = FControlRigMenuSpecificContext());
 
 	/** Get the control rig blueprint that we are editing */
 	UFUNCTION(BlueprintCallable, Category = ControlRigEditorExtensions)
@@ -100,11 +132,13 @@ public:
 	/** Returns context for graph node context menu */
 	UFUNCTION(BlueprintCallable, Category = ControlRigEditorExtensions)
 	FControlRigGraphNodeContextMenuContext GetGraphNodeContextMenuContext();
+	
+	/** Returns context for the menu that shows up when drag and drop from Rig Hierarchy to the Rig Graph*/
+	UFUNCTION(BlueprintCallable, Category = ControlRigEditorExtensions)
+    FControlRigRigHierarchyToGraphDragAndDropContext GetRigHierarchyToGraphDragAndDropContext();
 
 private:
 	TWeakObjectPtr<UControlRigBlueprint> ControlRigBlueprint;
 
-	FControlRigRigHierarchyDragAndDropContext DragAndDropContext;
-
-	FControlRigGraphNodeContextMenuContext GraphNodeContextMenuContext;
+	FControlRigMenuSpecificContext MenuSpecificContext;
 };
\ No newline at end of file
diff --git a/Engine/Plugins/Animation/ControlRig/Source/ControlRigEditor/Private/Editor/ControlRigEditor.cpp b/Engine/Plugins/Animation/ControlRig/Source/ControlRigEditor/Private/Editor/ControlRigEditor.cpp
--- a/Engine/Plugins/Animation/ControlRig/Source/ControlRigEditor/Private/Editor/ControlRigEditor.cpp
+++ b/Engine/Plugins/Animation/ControlRig/Source/ControlRigEditor/Private/Editor/ControlRigEditor.cpp
@@ -93,6 +93,7 @@
 #include "Types/ISlateMetaData.h"
 #include "ControlRigGraphDetails.h"
 #include "Kismet2/KismetDebugUtilities.h"
+#include "ToolMenus.h"
 
 #define LOCTEXT_NAMESPACE "ControlRigEditor"
 
@@ -2034,9 +2035,9 @@ void FControlRigEditor::InitToolMenuContext(FToolMenuContext& MenuContext)
 		}
 		
 		UControlRigContextMenuContext* ControlRigContext = NewObject<UControlRigContextMenuContext>();
-		ControlRigContext->Init(RigBlueprint,
-			FControlRigRigHierarchyDragAndDropContext(),
-			FControlRigGraphNodeContextMenuContext(Model, Node, Pin));
+		FControlRigMenuSpecificContext MenuSpecificContext;	
+		MenuSpecificContext.GraphNodeContextMenuContext = FControlRigGraphNodeContextMenuContext(Model, Node, Pin);
+		ControlRigContext->Init(RigBlueprint, MenuSpecificContext);
 
 		MenuContext.AddObject(ControlRigContext);
 	}
@@ -5165,192 +5166,301 @@ void FControlRigEditor::OnCurveContainerChanged()
 	NotificationPtr->SetCompletionState(SNotificationItem::CS_Success);
 }
 
-void FControlRigEditor::OnGraphNodeDropToPerform(TSharedPtr<FGraphNodeDragDropOp> DragDropOp, UEdGraph* Graph, const FVector2D& NodePosition, const FVector2D& ScreenPosition)
+void FControlRigEditor::OnGraphNodeDropToPerform(TSharedPtr<FGraphNodeDragDropOp> InDragDropOp, UEdGraph* InGraph, const FVector2D& InNodePosition, const FVector2D& InScreenPosition)
 {
 	DECLARE_SCOPE_HIERARCHICAL_COUNTER_FUNC()
 
-	if (DragDropOp->IsOfType<FRigElementHierarchyDragDropOp>())
+	if (InDragDropOp->IsOfType<FRigElementHierarchyDragDropOp>())
 	{
-		UControlRigBlueprint* Blueprint = Cast<UControlRigBlueprint>(GetBlueprintObj());
-		TSharedPtr<FRigElementHierarchyDragDropOp> RigHierarchyOp = StaticCastSharedPtr<FRigElementHierarchyDragDropOp>(DragDropOp);
+		TSharedPtr<FRigElementHierarchyDragDropOp> RigHierarchyOp = StaticCastSharedPtr<FRigElementHierarchyDragDropOp>(InDragDropOp);
 
-		TArray<FRigElementKey> DraggedKeys = RigHierarchyOp->GetElements();
-		uint8 DraggedTypes = 0;
-		for (const FRigElementKey& DraggedKey : DraggedKeys)
+		if (RigHierarchyOp->GetElements().Num() > 0 && FocusedGraphEdPtr.IsValid())
 		{
-			DraggedTypes = DraggedTypes | (uint8)DraggedKey.Type;
-		}
-
-		if (DraggedTypes != 0 && FocusedGraphEdPtr.IsValid())
-		{
-			FMenuBuilder MenuBuilder(true, NULL);
-			const FText SectionText = FText::FromString(RigHierarchyOp->GetJoinedElementNames());
-
-			MenuBuilder.BeginSection("RigHierarchyDroppedOn", SectionText);
+			const FName MenuName = TEXT("ControlRigEditor.RigHierarchyToGraphDragAndDropMenu");
+			UToolMenus* ToolMenus = UToolMenus::Get();
 			
-			if ((DraggedTypes & (uint8)ERigElementType::Control) != 0)
-			{
-				MenuBuilder.AddMenuEntry(
-					LOCTEXT("CreateGetControl", "Get Control"),
-					LOCTEXT("CreateGetControlTooltip", "Getter for control\n"),
-					FSlateIcon(),
-					FUIAction(
-						FExecuteAction::CreateSP(this, &FControlRigEditor::HandleMakeElementGetterSetter, ERigElementGetterSetterType_Transform, true, DraggedKeys, Graph, NodePosition),
-						FCanExecuteAction()
-					)
-				);
-				MenuBuilder.AddMenuEntry(
-					LOCTEXT("CreateSetControl", "Set Control"),
-					LOCTEXT("CreateSetControlTooltip", "Setter for control\n"),
-					FSlateIcon(),
-					FUIAction(
-						FExecuteAction::CreateSP(this, &FControlRigEditor::HandleMakeElementGetterSetter, ERigElementGetterSetterType_Transform, false, DraggedKeys, Graph, NodePosition),
-						FCanExecuteAction()
-					)
-				);
-			}
-
-			if (((DraggedTypes & (uint8)ERigElementType::Bone) != 0) ||
-				((DraggedTypes & (uint8)ERigElementType::Null) != 0))
-			{
-				MenuBuilder.AddMenuEntry(
-					LOCTEXT("CreateGetTransform", "Get Transform"),
-					LOCTEXT("CreateGetTransformTooltip", "Getter for transform\n"),
-					FSlateIcon(),
-					FUIAction(
-						FExecuteAction::CreateSP(this, &FControlRigEditor::HandleMakeElementGetterSetter, ERigElementGetterSetterType_Transform, true, DraggedKeys, Graph, NodePosition),
-						FCanExecuteAction()
-					)
-				);
-				MenuBuilder.AddMenuEntry(
-					LOCTEXT("CreateSetTransform", "Set Transform"),
-					LOCTEXT("CreateSetTransformTooltip", "Setter for transform\n"),
-					FSlateIcon(),
-					FUIAction(
-						FExecuteAction::CreateSP(this, &FControlRigEditor::HandleMakeElementGetterSetter, ERigElementGetterSetterType_Transform, false, DraggedKeys, Graph, NodePosition),
-						FCanExecuteAction()
-					)
-				);
-			}
-
-			if (((DraggedTypes & (uint8)ERigElementType::Bone) != 0) ||
-				((DraggedTypes & (uint8)ERigElementType::Control) != 0) ||
-				((DraggedTypes & (uint8)ERigElementType::Null) != 0))
+			if (!ToolMenus->IsMenuRegistered(MenuName))
 			{
-				MenuBuilder.AddMenuSeparator();
-
-				MenuBuilder.AddMenuEntry(
-					LOCTEXT("CreateSetRotation", "Set Rotation"),
-					LOCTEXT("CreateSetRotationTooltip", "Setter for Rotation\n"),
-					FSlateIcon(),
-					FUIAction(
-						FExecuteAction::CreateSP(this, &FControlRigEditor::HandleMakeElementGetterSetter, ERigElementGetterSetterType_Rotation, false, DraggedKeys, Graph, NodePosition),
-						FCanExecuteAction()
-					)
-				);
-
-				MenuBuilder.AddMenuEntry(
-					LOCTEXT("CreateSetTranslation", "Set Translation"),
-					LOCTEXT("CreateSetTranslationTooltip", "Setter for translation\n"),
-					FSlateIcon(),
-					FUIAction(
-						FExecuteAction::CreateSP(this, &FControlRigEditor::HandleMakeElementGetterSetter, ERigElementGetterSetterType_Translation, false, DraggedKeys, Graph, NodePosition),
-						FCanExecuteAction()
-					)
-				);
-
-				MenuBuilder.AddMenuEntry(
-					LOCTEXT("CreateSetOffset", "Add Offset"),
-					LOCTEXT("CreateSetOffsetTooltip", "Setter for offset\n"),
-					FSlateIcon(),
-					FUIAction(
-						FExecuteAction::CreateSP(this, &FControlRigEditor::HandleMakeElementGetterSetter, ERigElementGetterSetterType_Offset, false, DraggedKeys, Graph, NodePosition),
-						FCanExecuteAction()
-					)
-				);
-
-				MenuBuilder.AddMenuSeparator();
-
-				MenuBuilder.AddMenuEntry(
-					LOCTEXT("CreateGetRelativeTransform", "Get Relative Transform"),
-					LOCTEXT("CreateGetRelativeTransformTooltip", "Getter for relative transform\n"),
-					FSlateIcon(),
-					FUIAction(
-						FExecuteAction::CreateSP(this, &FControlRigEditor::HandleMakeElementGetterSetter, ERigElementGetterSetterType_Relative, true, DraggedKeys, Graph, NodePosition),
-						FCanExecuteAction()
-					)
-				);
-				MenuBuilder.AddMenuEntry(
-					LOCTEXT("CreateSetRelativeTransform", "Set Relative Transform"),
-					LOCTEXT("CreateSetRelativeTransformTooltip", "Setter for relative transform\n"),
-					FSlateIcon(),
-					FUIAction(
-						FExecuteAction::CreateSP(this, &FControlRigEditor::HandleMakeElementGetterSetter, ERigElementGetterSetterType_Relative, false, DraggedKeys, Graph, NodePosition),
-						FCanExecuteAction()
-					)
-				);
-
-			}
-
-			if (DraggedKeys.Num() > 0 && Blueprint != nullptr)
-			{
-				MenuBuilder.AddMenuSeparator();
-
-				MenuBuilder.AddMenuEntry(
-					LOCTEXT("CreateItemArrayFromKeys", "Create Item Array"),
-					LOCTEXT("CreateItemArrayFromKeysTooltip", "Creates an item array from the selected elements in the hierarchy"),
-					FSlateIcon(),
-					FUIAction(
-						FExecuteAction::CreateLambda([this, Blueprint, DraggedKeys, NodePosition]()
+				UToolMenu* Menu = ToolMenus->RegisterMenu(MenuName);
+				
+				const FName SectionName = TEXT("RigHierarchyToGraphDragAndDropSection");
+				const FName DynamicSectionName = *(SectionName.ToString() + TEXT("Dynamic"));
+				
+				Menu->AddDynamicSection( DynamicSectionName,
+                    FNewToolMenuDelegate::CreateLambda([
+                    	SectionName,
+                    	WeakThis = TWeakPtr<FControlRigEditor>(SharedThis(this))](UToolMenu* InMenu)
+                    {
+                        if (const TSharedPtr<FControlRigEditor> ControlRigEditor = WeakThis.Pin())
+                        {
+							UControlRigContextMenuContext* MainContext = InMenu->FindContext<UControlRigContextMenuContext>();
+							const FControlRigRigHierarchyToGraphDragAndDropContext& DragDropContext = MainContext->GetRigHierarchyToGraphDragAndDropContext();
+
+                        	UControlRigBlueprint* Blueprint = MainContext->GetControlRigBlueprint();
+							const TArray<FRigElementKey>& DraggedKeys = DragDropContext.DraggedElementKeys;
+                        	UEdGraph* Graph = DragDropContext.Graph.Get();
+                        	const FVector2D& NodePosition = DragDropContext.NodePosition;
+                        	
+                        	// if multiple types are selected, we show Get Elements/Set Elements
+							bool bMultipleTypeSelected = false;
+
+							ERigElementType LastType = ERigElementType::None;
+					
+							if (DraggedKeys.Num() > 0)
 							{
-								if (URigVMController* Controller = GetFocusedController())
+								LastType = DraggedKeys[0].Type;
+							}
+					
+							uint8 DraggedTypes = 0;
+							for (const FRigElementKey& DraggedKey : DragDropContext.DraggedElementKeys)
+							{
+								if (DraggedKey.Type != LastType)
 								{
-									Controller->OpenUndoBracket(TEXT("Create Item Array From Selection"));
-
-									static const FString ItemArrayCPPType = FString::Printf(TEXT("TArray<%s>"), *FRigElementKey::StaticStruct()->GetStructCPPName());
-									static const FString ItemArrayObjectPath = FRigElementKey::StaticStruct()->GetPathName();
-									
-									if (URigVMNode* ItemsNode = Controller->AddFreeRerouteNode(true, ItemArrayCPPType, *ItemArrayObjectPath, false, NAME_None, FString(), NodePosition))
-									{
-										if (URigVMPin* ItemsPin = ItemsNode->FindPin(TEXT("Value")))
-										{
-											Controller->SetArrayPinSize(ItemsPin->GetPinPath(), DraggedKeys.Num());
-
-											TArray<URigVMPin*> ItemPins = ItemsPin->GetSubPins();
-											ensure(ItemPins.Num() == DraggedKeys.Num());
-
-											for (int32 ItemIndex = 0; ItemIndex < DraggedKeys.Num(); ItemIndex++)
-											{
-												FString DefaultValue;
-												FRigElementKey::StaticStruct()->ExportText(DefaultValue, &DraggedKeys[ItemIndex], nullptr, nullptr, PPF_None, nullptr);
-												Controller->SetPinDefaultValue(ItemPins[ItemIndex]->GetPinPath(), DefaultValue, true, true, false, true);
-												Controller->SetPinExpansion(ItemPins[ItemIndex]->GetPinPath(), true, true, true);
-											}
-
-											Controller->SetPinExpansion(ItemsPin->GetPinPath(), true, true, true);
-										}
-									}
-
-									Controller->CloseUndoBracket();
+									bMultipleTypeSelected = true;
 								}
+						
+								DraggedTypes = DraggedTypes | (uint8)DraggedKey.Type;
 							}
-						)
-					)
+                        	
+							const FText SectionText = FText::FromString(DragDropContext.GetSectionTitle());
+							FToolMenuSection& Section = InMenu->AddSection(SectionName, SectionText);
+
+                        	FText GetterLabel = LOCTEXT("GetElement","Get Element");
+                        	FText GetterTooltip = LOCTEXT("GetElement_ToolTip", "Getter For Element");
+                        	FText SetterLabel = LOCTEXT("SetElement","Set Element");
+                        	FText SetterTooltip = LOCTEXT("SetElement_ToolTip", "Setter For Element");
+							// if multiple types are selected, we show Get Elements/Set Elements
+                        	if (bMultipleTypeSelected)
+                        	{
+                        		GetterLabel = LOCTEXT("GetElements","Get Elements");
+                        		GetterTooltip = LOCTEXT("GetElements_ToolTip", "Getter For Elements");
+                        		SetterLabel = LOCTEXT("SetElements","Set Elements");
+                        		SetterTooltip = LOCTEXT("SetElements_ToolTip", "Setter For Elements");
+                        	}
+                        	else
+                        	{
+								// otherwise, we show "Get Bone/NUll/Control"
+                        		if ((DraggedTypes & (uint8)ERigElementType::Bone) != 0)
+                        		{
+                        			GetterLabel = LOCTEXT("GetBone","Get Bone");
+                        			GetterTooltip = LOCTEXT("GetBone_ToolTip", "Getter For Bone");
+                        			SetterLabel = LOCTEXT("SetBone","Set Bone");
+                        			SetterTooltip = LOCTEXT("SetBone_ToolTip", "Setter For Bone");
+                        		}
+                        		else if ((DraggedTypes & (uint8)ERigElementType::Null) != 0)
+                        		{
+                        			GetterLabel = LOCTEXT("GetNull","Get Null");
+                        			GetterTooltip = LOCTEXT("GetNull_ToolTip", "Getter For Null");
+                        			SetterLabel = LOCTEXT("SetNull","Set Null");
+                        			SetterTooltip = LOCTEXT("SetNull_eoolTip", "Setter For Null");
+                        		}
+                        		else if ((DraggedTypes & (uint8)ERigElementType::Control) != 0)
+                        		{
+                        			GetterLabel = LOCTEXT("GetControl","Get Control");
+                        			GetterTooltip = LOCTEXT("GetControl_ToolTip", "Getter For Control");
+                        			SetterLabel = LOCTEXT("SetControl","Set Control");
+                        			SetterTooltip = LOCTEXT("SetControl_ToolTip", "Setter For Control");
+                        		}
+                        	}
+
+                        	FToolMenuEntry GetElementsEntry = FToolMenuEntry::InitMenuEntry(
+								TEXT("GetElements"),
+								GetterLabel,
+								GetterTooltip,
+								FSlateIcon(),
+								FUIAction(
+									FExecuteAction::CreateSP(ControlRigEditor.Get(), &FControlRigEditor::HandleMakeElementGetterSetter, ERigElementGetterSetterType_Transform, true, DraggedKeys, Graph, NodePosition),
+									FCanExecuteAction()
+								)
+                            );
+                        	GetElementsEntry.InsertPosition.Name = NAME_None;
+                        	GetElementsEntry.InsertPosition.Position = EToolMenuInsertType::First;
+                        	
+
+                        	Section.AddEntry(GetElementsEntry);
+
+                        	FToolMenuEntry SetElementsEntry = FToolMenuEntry::InitMenuEntry(
+								TEXT("SetElements"),
+								SetterLabel,
+								SetterTooltip,
+								FSlateIcon(),
+								FUIAction(
+									FExecuteAction::CreateSP(ControlRigEditor.Get(), &FControlRigEditor::HandleMakeElementGetterSetter, ERigElementGetterSetterType_Transform, false, DraggedKeys, Graph, NodePosition),
+									FCanExecuteAction()
+								)
+                            );
+                        	SetElementsEntry.InsertPosition.Name = GetElementsEntry.Name;
+							SetElementsEntry.InsertPosition.Position = EToolMenuInsertType::After;	
+
+                        	Section.AddEntry(SetElementsEntry);
+
+                        	if (((DraggedTypes & (uint8)ERigElementType::Bone) != 0) ||
+								((DraggedTypes & (uint8)ERigElementType::Control) != 0) ||
+								((DraggedTypes & (uint8)ERigElementType::Null) != 0))
+                        	{
+                        		FToolMenuEntry& RotationTranslationSeparator = Section.AddSeparator(TEXT("RotationTranslationSeparator"));
+                        		RotationTranslationSeparator.InsertPosition.Name = SetElementsEntry.Name;
+                        		RotationTranslationSeparator.InsertPosition.Position = EToolMenuInsertType::After;
+
+								FToolMenuEntry SetRotationEntry = FToolMenuEntry::InitMenuEntry(
+									TEXT("SetRotation"),
+									LOCTEXT("SetRotation","Set Rotation"),
+									LOCTEXT("SetRotation_ToolTip","Setter for Rotation"),
+									FSlateIcon(),
+									FUIAction(
+										FExecuteAction::CreateSP(ControlRigEditor.Get(), &FControlRigEditor::HandleMakeElementGetterSetter, ERigElementGetterSetterType_Rotation, false, DraggedKeys, Graph, NodePosition),
+										FCanExecuteAction()
+									)
+								);
+                        		
+                        		SetRotationEntry.InsertPosition.Name = RotationTranslationSeparator.Name;
+								SetRotationEntry.InsertPosition.Position = EToolMenuInsertType::After;		
+								Section.AddEntry(SetRotationEntry);
+
+                        		FToolMenuEntry SetTranslationEntry = FToolMenuEntry::InitMenuEntry(
+									TEXT("SetTranslation"),
+									LOCTEXT("SetTranslation","Set Translation"),
+									LOCTEXT("SetTranslation_ToolTip","Setter for Translation"),
+									FSlateIcon(),
+									FUIAction(
+										FExecuteAction::CreateSP(ControlRigEditor.Get(), &FControlRigEditor::HandleMakeElementGetterSetter, ERigElementGetterSetterType_Translation, false, DraggedKeys, Graph, NodePosition),
+										FCanExecuteAction()
+									)
+								);
+
+                        		SetTranslationEntry.InsertPosition.Name = SetRotationEntry.Name;
+								SetTranslationEntry.InsertPosition.Position = EToolMenuInsertType::After;		
+								Section.AddEntry(SetTranslationEntry);
+
+                        		FToolMenuEntry AddOffsetEntry = FToolMenuEntry::InitMenuEntry(
+									TEXT("AddOffset"),
+									LOCTEXT("AddOffset","Add Offset"),
+									LOCTEXT("AddOffset_ToolTip","Setter for Offset"),
+									FSlateIcon(),
+									FUIAction(
+										FExecuteAction::CreateSP(ControlRigEditor.Get(), &FControlRigEditor::HandleMakeElementGetterSetter, ERigElementGetterSetterType_Offset, false, DraggedKeys, Graph, NodePosition),
+										FCanExecuteAction()
+									)
+								);
+                        		
+								AddOffsetEntry.InsertPosition.Name = SetTranslationEntry.Name;
+                                AddOffsetEntry.InsertPosition.Position = EToolMenuInsertType::After;						
+								Section.AddEntry(AddOffsetEntry);
+
+                        		FToolMenuEntry& RelativeTransformSeparator = Section.AddSeparator(TEXT("RelativeTransformSeparator"));
+                        		RelativeTransformSeparator.InsertPosition.Name = AddOffsetEntry.Name;
+                        		RelativeTransformSeparator.InsertPosition.Position = EToolMenuInsertType::After;
+                        		
+                        		FToolMenuEntry GetRelativeTransformEntry = FToolMenuEntry::InitMenuEntry(
+									TEXT("GetRelativeTransformEntry"),
+									LOCTEXT("GetRelativeTransform", "Get Relative Transform"),
+									LOCTEXT("GetRelativeTransform_ToolTip", "Getter for Relative Transform"),
+									FSlateIcon(),
+									FUIAction(
+										FExecuteAction::CreateSP(ControlRigEditor.Get(), &FControlRigEditor::HandleMakeElementGetterSetter, ERigElementGetterSetterType_Relative, true, DraggedKeys, Graph, NodePosition),
+										FCanExecuteAction()
+									)
+								);
+                        		
+								GetRelativeTransformEntry.InsertPosition.Name = RelativeTransformSeparator.Name;
+                                GetRelativeTransformEntry.InsertPosition.Position = EToolMenuInsertType::After;	
+								Section.AddEntry(GetRelativeTransformEntry);
+                        		
+                        	    FToolMenuEntry SetRelativeTransformEntry = FToolMenuEntry::InitMenuEntry(
+									TEXT("SetRelativeTransformEntry"),
+									LOCTEXT("SetRelativeTransform", "Set Relative Transform"),
+									LOCTEXT("SetRelativeTransform_ToolTip", "Setter for Relative Transform"),
+									FSlateIcon(),
+									FUIAction(
+										FExecuteAction::CreateSP(ControlRigEditor.Get(), &FControlRigEditor::HandleMakeElementGetterSetter, ERigElementGetterSetterType_Relative, false, DraggedKeys, Graph, NodePosition),
+										FCanExecuteAction()
+									)
+								);
+                        		SetRelativeTransformEntry.InsertPosition.Name = GetRelativeTransformEntry.Name;
+								SetRelativeTransformEntry.InsertPosition.Position = EToolMenuInsertType::After;		
+                        		Section.AddEntry(SetRelativeTransformEntry);
+                        	}
+
+                        	if (DraggedKeys.Num() > 0 && Blueprint != nullptr)
+                        	{
+                                FToolMenuEntry& ItemArraySeparator = Section.AddSeparator(TEXT("ItemArraySeparator"));
+                        		ItemArraySeparator.InsertPosition.Name = TEXT("SetRelativeTransformEntry"),
+								ItemArraySeparator.InsertPosition.Position = EToolMenuInsertType::After;
+                        		
+                        		FToolMenuEntry CreateItemArrayEntry = FToolMenuEntry::InitMenuEntry(
+                                    TEXT("CreateItemArray"),
+                                    LOCTEXT("CreateItemArray", "Create Item Array"),
+                                    LOCTEXT("CreateItemArray_ToolTip", "Creates an item array from the selected elements in the hierarchy"),
+                                    FSlateIcon(),
+                                    FUIAction(
+                                        FExecuteAction::CreateLambda([ControlRigEditor, Blueprint, DraggedKeys, NodePosition]()
+                                            {
+                                                if (URigVMController* Controller = ControlRigEditor->GetFocusedController())
+                                                {
+                                                    Controller->OpenUndoBracket(TEXT("Create Item Array From Selection"));
+
+                                                    static const FString ItemArrayCPPType = FString::Printf(TEXT("TArray<%s>"), *FRigElementKey::StaticStruct()->GetStructCPPName());
+                                                    static const FString ItemArrayObjectPath = FRigElementKey::StaticStruct()->GetPathName();
+									
+                                                    if (URigVMNode* ItemsNode = Controller->AddFreeRerouteNode(true, ItemArrayCPPType, *ItemArrayObjectPath, false, NAME_None, FString(), NodePosition))
+                                                    {
+                                                        if (URigVMPin* ItemsPin = ItemsNode->FindPin(TEXT("Value")))
+                                                        {
+                                                            Controller->SetArrayPinSize(ItemsPin->GetPinPath(), DraggedKeys.Num());
+
+                                                            TArray<URigVMPin*> ItemPins = ItemsPin->GetSubPins();
+                                                            ensure(ItemPins.Num() == DraggedKeys.Num());
+
+                                                            for (int32 ItemIndex = 0; ItemIndex < DraggedKeys.Num(); ItemIndex++)
+                                                            {
+                                                                FString DefaultValue;
+                                                                FRigElementKey::StaticStruct()->ExportText(DefaultValue, &DraggedKeys[ItemIndex], nullptr, nullptr, PPF_None, nullptr);
+                                                                Controller->SetPinDefaultValue(ItemPins[ItemIndex]->GetPinPath(), DefaultValue, true, true, false, true);
+                                                                Controller->SetPinExpansion(ItemPins[ItemIndex]->GetPinPath(), true, true, true);
+                                                            }
+
+                                                            Controller->SetPinExpansion(ItemsPin->GetPinPath(), true, true, true);
+                                                        }
+                                                    }
+
+                                                    Controller->CloseUndoBracket();
+                                                }
+                                            }
+                                        )
+                                    )
+                        		);
+                        		
+                        		CreateItemArrayEntry.InsertPosition.Name = ItemArraySeparator.Name,
+								CreateItemArrayEntry.InsertPosition.Position = EToolMenuInsertType::After;
+                        		Section.AddEntry(CreateItemArrayEntry);
+                            }
+                        }
+                    })
+                );
+			}
+
+			UControlRigContextMenuContext* MenuContext = NewObject<UControlRigContextMenuContext>();
+			FControlRigMenuSpecificContext MenuSpecificContext;
+			MenuSpecificContext.RigHierarchyToGraphDragAndDropContext =
+				FControlRigRigHierarchyToGraphDragAndDropContext(
+					RigHierarchyOp->GetElements(),
+					InGraph,
+					InNodePosition
 				);
-			}
-
+			MenuContext->Init(Cast<UControlRigBlueprint>(GetBlueprintObj()), MenuSpecificContext);
+			TSharedRef<SWidget>	MenuWidget = ToolMenus->GenerateWidget(MenuName, FToolMenuContext(MenuContext));
+			
 			TSharedRef<SWidget> GraphEditorPanel = FocusedGraphEdPtr.Pin().ToSharedRef();
 
-			// Show dialog to choose getter vs setter
+			// Show menu to choose getter vs setter
 			FSlateApplication::Get().PushMenu(
 				GraphEditorPanel,
 				FWidgetPath(),
-				MenuBuilder.MakeWidget(),
-				ScreenPosition,
+				MenuWidget,
+				InScreenPosition,
 				FPopupTransitionEffect(FPopupTransitionEffect::ContextMenu)
 			);
 
-			MenuBuilder.EndSection();
 		}
 	}
 }
diff --git a/Engine/Plugins/Animation/ControlRig/Source/ControlRigEditor/Private/Editor/ControlRigEditor.h b/Engine/Plugins/Animation/ControlRig/Source/ControlRigEditor/Private/Editor/ControlRigEditor.h
--- a/Engine/Plugins/Animation/ControlRig/Source/ControlRigEditor/Private/Editor/ControlRigEditor.h
+++ b/Engine/Plugins/Animation/ControlRig/Source/ControlRigEditor/Private/Editor/ControlRigEditor.h
@@ -189,7 +189,7 @@ public:
 	void OnHierarchyModified(ERigHierarchyNotification InNotif, URigHierarchy* InHierarchy, const FRigBaseElement* InElement);
 	void OnHierarchyModified_AnyThread(ERigHierarchyNotification InNotif, URigHierarchy* InHierarchy, const FRigBaseElement* InElement);
 
-	void OnGraphNodeDropToPerform(TSharedPtr<FGraphNodeDragDropOp> DragDropOp, UEdGraph* Graph, const FVector2D& NodePosition, const FVector2D& ScreenPosition);
+	void OnGraphNodeDropToPerform(TSharedPtr<FGraphNodeDragDropOp> InDragDropOp, UEdGraph* InGraph, const FVector2D& InNodePosition, const FVector2D& InScreenPosition);
 
 	FPersonaViewportKeyDownDelegate& GetKeyDownDelegate() { return OnKeyDownDelegate; }
 	FOnGetContextMenu& OnGetViewportContextMenu() { return OnGetViewportContextMenuDelegate; }
diff --git a/Engine/Plugins/Animation/ControlRig/Source/ControlRigEditor/Private/Editor/SRigHierarchy.cpp b/Engine/Plugins/Animation/ControlRig/Source/ControlRigEditor/Private/Editor/SRigHierarchy.cpp
--- a/Engine/Plugins/Animation/ControlRig/Source/ControlRigEditor/Private/Editor/SRigHierarchy.cpp
+++ b/Engine/Plugins/Animation/ControlRig/Source/ControlRigEditor/Private/Editor/SRigHierarchy.cpp
@@ -967,7 +967,9 @@ UToolMenu* SRigHierarchy::GetOrCreateDragDropMenu(const TArray<FRigElementKey>&
 	}
 
 	UControlRigContextMenuContext* MenuContext = NewObject<UControlRigContextMenuContext>();
-	MenuContext->Init(ControlRigBlueprint, FControlRigRigHierarchyDragAndDropContext(DraggedKeys, TargetKey));
+	FControlRigMenuSpecificContext MenuSpecificContext;
+	MenuSpecificContext.RigHierarchyDragAndDropContext = FControlRigRigHierarchyDragAndDropContext(DraggedKeys, TargetKey);
+	MenuContext->Init(ControlRigBlueprint, MenuSpecificContext);
 	
 	UToolMenu* Menu = ToolMenus->GenerateMenu(MenuName, FToolMenuContext(MenuContext));
 
