Control Rig: Fix python interface for dealing with FTRansform based controls

diff --git a/Engine/Plugins/Animation/ControlRig/Source/ControlRig/Private/Rigs/RigHierarchyController.cpp b/Engine/Plugins/Animation/ControlRig/Source/ControlRig/Private/Rigs/RigHierarchyController.cpp
--- a/Engine/Plugins/Animation/ControlRig/Source/ControlRig/Private/Rigs/RigHierarchyController.cpp
+++ b/Engine/Plugins/Animation/ControlRig/Source/ControlRig/Private/Rigs/RigHierarchyController.cpp
@@ -1354,10 +1354,10 @@ TArray<FString> URigHierarchyController::GetAddControlPythonCommands(FRigControl
 			case ERigControlType::Position: TypeStr = TEXT("POSITION"); break;
 			case ERigControlType::Rotator: TypeStr = TEXT("POSITION"); break;
 			case ERigControlType::Scale: TypeStr = TEXT("SCALE"); break;
-			case ERigControlType::Transform: TypeStr = TEXT("TRANSFORM"); break;
+			case ERigControlType::Transform: TypeStr = TEXT("EULER_TRANSFORM"); break;
 			case ERigControlType::EulerTransform: TypeStr = TEXT("EULER_TRANSFORM"); break;
 			case ERigControlType::Vector2D: TypeStr = TEXT("VECTOR2D"); break;
-			case ERigControlType::TransformNoScale: TypeStr = TEXT("TRANSFORM_NO_SCALE"); break;
+			case ERigControlType::TransformNoScale: TypeStr = TEXT("EULER_TRANSFORM"); break;
 			default: ensure(false);
 		}
 		Commands.Add(FString::Printf(TEXT("control_settings_%s.control_type = unreal.RigControlType.%s"),
diff --git a/Engine/Plugins/Animation/ControlRig/Source/ControlRig/Public/Rigs/RigHierarchyDefines.h b/Engine/Plugins/Animation/ControlRig/Source/ControlRig/Public/Rigs/RigHierarchyDefines.h
--- a/Engine/Plugins/Animation/ControlRig/Source/ControlRig/Public/Rigs/RigHierarchyDefines.h
+++ b/Engine/Plugins/Animation/ControlRig/Source/ControlRig/Public/Rigs/RigHierarchyDefines.h
@@ -298,6 +298,7 @@ public:
 				Get<FVector3f>().X, Get<FVector3f>().Z, Get<FVector3f>().Y); break;
 			case ERigControlType::Scale: ValueStr = FString::Printf(TEXT("unreal.RigHierarchy.make_control_value_from_vector(unreal.Vector(%.6f, %.6f, %.6f))"),
 				Get<FVector3f>().X, Get<FVector3f>().Y, Get<FVector3f>().Z); break;
+			/*
 			case ERigControlType::Transform: ValueStr = FString::Printf(TEXT("unreal.RigHierarchy.make_control_value_from_transform(unreal.Transform(location=[%.6f,%.6f,%.6f],rotation=[%.6f,%.6f,%.6f],scale=[%.6f,%.6f,%.6f]))"),
 				Get<FTransform_Float>().ToTransform().GetLocation().X,
 				Get<FTransform_Float>().ToTransform().GetLocation().Y,
@@ -308,6 +309,17 @@ public:
 				Get<FTransform_Float>().ToTransform().GetScale3D().X,
 				Get<FTransform_Float>().ToTransform().GetScale3D().Y,
 				Get<FTransform_Float>().ToTransform().GetScale3D().Z); break;
+			*/
+			case ERigControlType::Transform: ValueStr = ValueStr = FString::Printf(TEXT("unreal.RigHierarchy.make_control_value_from_euler_transform(unreal.EulerTransform(location=[%.6f,%.6f,%.6f],rotation=[%.6f,%.6f,%.6f],scale=[%.6f,%.6f,%.6f]))"),
+				FEulerTransform(Get<FTransform_Float>().ToTransform()).Location.X,
+				FEulerTransform(Get<FTransform_Float>().ToTransform()).Location.Y,
+				FEulerTransform(Get<FTransform_Float>().ToTransform()).Location.Z,
+				FEulerTransform(Get<FTransform_Float>().ToTransform()).Rotation.Euler().X,
+				FEulerTransform(Get<FTransform_Float>().ToTransform()).Rotation.Euler().Y,
+				FEulerTransform(Get<FTransform_Float>().ToTransform()).Rotation.Euler().Z,
+				FEulerTransform(Get<FTransform_Float>().ToTransform()).Scale.X,
+				FEulerTransform(Get<FTransform_Float>().ToTransform()).Scale.Y,
+				FEulerTransform(Get<FTransform_Float>().ToTransform()).Scale.Z); break;
 			case ERigControlType::EulerTransform: ValueStr = FString::Printf(TEXT("unreal.RigHierarchy.make_control_value_from_euler_transform(unreal.EulerTransform(location=[%.6f,%.6f,%.6f],rotation=[%.6f,%.6f,%.6f],scale=[%.6f,%.6f,%.6f]))"),
 				Get<FEulerTransform_Float>().ToTransform().Location.X,
 				Get<FEulerTransform_Float>().ToTransform().Location.Y,
@@ -320,6 +332,7 @@ public:
 				Get<FEulerTransform_Float>().ToTransform().Scale.Z); break;
 			case ERigControlType::Vector2D: ValueStr = FString::Printf(TEXT("unreal.RigHierarchy.make_control_value_from_vector2d(unreal.Vector2D(%.6f, %.6f))"),
 				Get<FVector3f>().X, Get<FVector3f>().Y); break;
+			/*
 			case ERigControlType::TransformNoScale: ValueStr = FString::Printf(TEXT("unreal.RigHierarchy.make_control_value_from_transform_no_scale(unreal.TransformNoScale(location=[%.6f,%.6f,%.6f],rotation=[%.6f,%.6f,%.6f]))"),
 				Get<FTransformNoScale_Float>().ToTransform().ToFTransform().GetLocation().X,
 				Get<FTransformNoScale_Float>().ToTransform().ToFTransform().GetLocation().Y,
@@ -327,6 +340,15 @@ public:
 				Get<FTransformNoScale_Float>().ToTransform().ToFTransform().Rotator().Pitch,
 				Get<FTransformNoScale_Float>().ToTransform().ToFTransform().Rotator().Yaw,
 				Get<FTransformNoScale_Float>().ToTransform().ToFTransform().Rotator().Roll); break;
+			*/
+			case ERigControlType::TransformNoScale: ValueStr = ValueStr = FString::Printf(TEXT("unreal.RigHierarchy.make_control_value_from_euler_transform(unreal.EulerTransform(location=[%.6f,%.6f,%.6f],rotation=[%.6f,%.6f,%.6f],scale=[%.6f,%.6f,%.6f]))"),
+				FEulerTransform(Get<FTransformNoScale_Float>().ToTransform().ToFTransform()).Location.X,
+				FEulerTransform(Get<FTransformNoScale_Float>().ToTransform().ToFTransform()).Location.Y,
+				FEulerTransform(Get<FTransformNoScale_Float>().ToTransform().ToFTransform()).Location.Z,
+				FEulerTransform(Get<FTransformNoScale_Float>().ToTransform().ToFTransform()).Rotation.Euler().X,
+				FEulerTransform(Get<FTransformNoScale_Float>().ToTransform().ToFTransform()).Rotation.Euler().Y,
+				FEulerTransform(Get<FTransformNoScale_Float>().ToTransform().ToFTransform()).Rotation.Euler().Z,
+				1.f, 1.f, 1.f); break;
 			default: ensure(false);
 		}
 
