Control Rig: Deprecate Static Mesh Control Node and Transform Constraint Node

diff --git a/Engine/Plugins/Animation/ControlRig/Source/ControlRig/Private/Units/Highlevel/Hierarchy/RigUnit_TransformConstraint.h b/Engine/Plugins/Animation/ControlRig/Source/ControlRig/Private/Units/Highlevel/Hierarchy/RigUnit_TransformConstraint.h
--- a/Engine/Plugins/Animation/ControlRig/Source/ControlRig/Private/Units/Highlevel/Hierarchy/RigUnit_TransformConstraint.h
+++ b/Engine/Plugins/Animation/ControlRig/Source/ControlRig/Private/Units/Highlevel/Hierarchy/RigUnit_TransformConstraint.h
@@ -114,7 +114,7 @@ private:
 /**
  * Constrains an item's transform to multiple items' transforms
  */
-USTRUCT(meta=(DisplayName="Transform Constraint", Category="Transforms", Keywords = "Parent,Orient,Scale"))
+USTRUCT(meta=(DisplayName="Transform Constraint", Category="Transforms", Keywords = "Parent,Orient,Scale", Deprecated="5.0"))
 struct CONTROLRIG_API FRigUnit_TransformConstraintPerItem : public FRigUnit_HighlevelBaseMutable
 {
 	GENERATED_BODY()
diff --git a/Engine/Plugins/Animation/ControlRig/Source/ControlRig/Public/Units/Control/RigUnit_Control_StaticMesh.h b/Engine/Plugins/Animation/ControlRig/Source/ControlRig/Public/Units/Control/RigUnit_Control_StaticMesh.h
--- a/Engine/Plugins/Animation/ControlRig/Source/ControlRig/Public/Units/Control/RigUnit_Control_StaticMesh.h
+++ b/Engine/Plugins/Animation/ControlRig/Source/ControlRig/Public/Units/Control/RigUnit_Control_StaticMesh.h
@@ -9,7 +9,7 @@ class UStaticMesh;
 class UMaterialInterface;
 
 /** A control unit used to drive a transform from an external source */
-USTRUCT(meta=(DisplayName="Static Mesh Control", Category="Controls", ShowVariableNameInTitle))
+USTRUCT(meta=(DisplayName="Static Mesh Control", Category="Controls", ShowVariableNameInTitle, Deprecated="5.0"))
 struct CONTROLRIG_API FRigUnit_Control_StaticMesh : public FRigUnit_Control
 {
 	GENERATED_BODY()
