Control Rig: change comment move mode by holding down the Shift key

#UE-120113

diff --git a/Engine/Plugins/Animation/ControlRig/Source/ControlRigEditor/Private/Graph/SControlRigGraphNodeComment.cpp b/Engine/Plugins/Animation/ControlRig/Source/ControlRigEditor/Private/Graph/SControlRigGraphNodeComment.cpp
--- a/Engine/Plugins/Animation/ControlRig/Source/ControlRigEditor/Private/Graph/SControlRigGraphNodeComment.cpp
+++ b/Engine/Plugins/Animation/ControlRig/Source/ControlRigEditor/Private/Graph/SControlRigGraphNodeComment.cpp
@@ -145,8 +145,11 @@ void SControlRigGraphNodeComment::MoveTo(const FVector2D& NewPosition, FNodeSet&
 				{
 					const FVector2D OldPosition = RigSchema->GetNodePositionAtStartOfInteraction(NodeToMove);
 
-					if((NodeToMove == GraphNode) || IsNodeUnderComment(CommentOldPosition, OldPosition))
-					{
+					// Don't drag note content if either of the shift keys are down.
+					FModifierKeysState KeysState = FSlateApplication::Get().GetModifierKeys();
+
+					if((NodeToMove == GraphNode) || (IsNodeUnderComment(CommentOldPosition, OldPosition) && !KeysState.IsShiftDown()))
+					{ 
 						const FVector2D AccumulatedPosition = OldPosition + Delta; 
 						RigSchema->SetNodePosition(NodeToMove, AccumulatedPosition, false);
 					}
