Control Rig: Fix static analysis warnings

diff --git a/Engine/Plugins/Animation/ControlRig/Source/ControlRig/Public/Rigs/RigHierarchyDefines.h b/Engine/Plugins/Animation/ControlRig/Source/ControlRig/Public/Rigs/RigHierarchyDefines.h
--- a/Engine/Plugins/Animation/ControlRig/Source/ControlRig/Public/Rigs/RigHierarchyDefines.h
+++ b/Engine/Plugins/Animation/ControlRig/Source/ControlRig/Public/Rigs/RigHierarchyDefines.h
@@ -310,7 +310,7 @@ public:
 				Get<FTransform_Float>().ToTransform().GetScale3D().Y,
 				Get<FTransform_Float>().ToTransform().GetScale3D().Z); break;
 			*/
-			case ERigControlType::Transform: ValueStr = ValueStr = FString::Printf(TEXT("unreal.RigHierarchy.make_control_value_from_euler_transform(unreal.EulerTransform(location=[%.6f,%.6f,%.6f],rotation=[%.6f,%.6f,%.6f],scale=[%.6f,%.6f,%.6f]))"),
+			case ERigControlType::Transform: ValueStr = FString::Printf(TEXT("unreal.RigHierarchy.make_control_value_from_euler_transform(unreal.EulerTransform(location=[%.6f,%.6f,%.6f],rotation=[%.6f,%.6f,%.6f],scale=[%.6f,%.6f,%.6f]))"),
 				FEulerTransform(Get<FTransform_Float>().ToTransform()).Location.X,
 				FEulerTransform(Get<FTransform_Float>().ToTransform()).Location.Y,
 				FEulerTransform(Get<FTransform_Float>().ToTransform()).Location.Z,
@@ -341,7 +341,7 @@ public:
 				Get<FTransformNoScale_Float>().ToTransform().ToFTransform().Rotator().Yaw,
 				Get<FTransformNoScale_Float>().ToTransform().ToFTransform().Rotator().Roll); break;
 			*/
-			case ERigControlType::TransformNoScale: ValueStr = ValueStr = FString::Printf(TEXT("unreal.RigHierarchy.make_control_value_from_euler_transform(unreal.EulerTransform(location=[%.6f,%.6f,%.6f],rotation=[%.6f,%.6f,%.6f],scale=[%.6f,%.6f,%.6f]))"),
+			case ERigControlType::TransformNoScale: ValueStr = FString::Printf(TEXT("unreal.RigHierarchy.make_control_value_from_euler_transform(unreal.EulerTransform(location=[%.6f,%.6f,%.6f],rotation=[%.6f,%.6f,%.6f],scale=[%.6f,%.6f,%.6f]))"),
 				FEulerTransform(Get<FTransformNoScale_Float>().ToTransform().ToFTransform()).Location.X,
 				FEulerTransform(Get<FTransformNoScale_Float>().ToTransform().ToFTransform()).Location.Y,
 				FEulerTransform(Get<FTransformNoScale_Float>().ToTransform().ToFTransform()).Location.Z,
