Changed Gizmo property labels

diff --git a/Engine/Plugins/Animation/ControlRig/Source/ControlRigEditor/Private/ControlRigElementDetails.cpp b/Engine/Plugins/Animation/ControlRig/Source/ControlRigEditor/Private/ControlRigElementDetails.cpp
--- a/Engine/Plugins/Animation/ControlRig/Source/ControlRigEditor/Private/ControlRigElementDetails.cpp
+++ b/Engine/Plugins/Animation/ControlRig/Source/ControlRigEditor/Private/ControlRigElementDetails.cpp
@@ -1856,8 +1856,10 @@ void FRigControlElementDetails::CustomizeChildren(TSharedRef<class IPropertyHand
 		IsAnyControlOfType(ERigControlType::TransformNoScale) ||
 		IsAnyControlOfType(ERigControlType::EulerTransform))
 	{
-		GizmoGroup->AddPropertyRow(SettingsHandle->GetChildHandle(TEXT("bGizmoEnabled")).ToSharedRef());
-		GizmoGroup->AddPropertyRow(SettingsHandle->GetChildHandle(TEXT("bGizmoVisible")).ToSharedRef());
+		GizmoGroup->AddPropertyRow(SettingsHandle->GetChildHandle(TEXT("bGizmoEnabled")).ToSharedRef())
+		.DisplayName(FText::FromString(TEXT("Enabled")));
+		GizmoGroup->AddPropertyRow(SettingsHandle->GetChildHandle(TEXT("bGizmoVisible")).ToSharedRef())
+		.DisplayName(FText::FromString(TEXT("Visible")));
 
 		// setup gizmo transform
 		{
@@ -1865,16 +1867,16 @@ void FRigControlElementDetails::CustomizeChildren(TSharedRef<class IPropertyHand
 			const TSharedPtr<IPropertyHandle> InitialHandle = GizmoHandle->GetChildHandle(TEXT("Initial"));
 			const TSharedPtr<IPropertyHandle> LocalHandle = InitialHandle->GetChildHandle(TEXT("Local"));
 			const TSharedPtr<IPropertyHandle> TransformHandle = LocalHandle->GetChildHandle(TEXT("Transform"));
-			GizmoGroup->AddPropertyRow(TransformHandle.ToSharedRef()).DisplayName(FText::FromString(TEXT("Gizmo Transform")))
+			GizmoGroup->AddPropertyRow(TransformHandle.ToSharedRef())
 			.IsEnabled(TAttribute<bool>::CreateSP(this, &FRigControlElementDetails::IsGizmoEnabled));
 		}
 
-		const TSharedPtr<IPropertyHandle> GizmoNameHandle = SettingsHandle->GetChildHandle(TEXT("GizmoName"));
-		GizmoGroup->AddPropertyRow(GizmoNameHandle.ToSharedRef()).CustomWidget()
+		const TSharedPtr<IPropertyHandle> GizmoShapeHandle = SettingsHandle->GetChildHandle(TEXT("GizmoName"));
+		GizmoGroup->AddPropertyRow(GizmoShapeHandle.ToSharedRef()).CustomWidget()
 		.NameContent()
 		[
 			SNew(STextBlock)
-			.Text(FText::FromString(TEXT("Gizmo Name")))
+			.Text(FText::FromString(TEXT("Shape")))
 			.Font(IDetailLayoutBuilder::GetDetailFont())
 			.IsEnabled(this, &FRigControlElementDetails::IsGizmoEnabled)
 		]
@@ -1885,7 +1887,8 @@ void FRigControlElementDetails::CustomizeChildren(TSharedRef<class IPropertyHand
 			.IsEnabled(this, &FRigControlElementDetails::IsGizmoEnabled)
 		];
 
-		GizmoGroup->AddPropertyRow(SettingsHandle->GetChildHandle(TEXT("GizmoColor")).ToSharedRef());
+		GizmoGroup->AddPropertyRow(SettingsHandle->GetChildHandle(TEXT("GizmoColor")).ToSharedRef())
+		.DisplayName(FText::FromString(TEXT("Color")));
 	}
 
 	if(IsAnyControlOfType(ERigControlType::Integer))
