Control Rig: Add null check in SRigHierarchy.cpp

diff --git a/Engine/Plugins/Animation/ControlRig/Source/ControlRigEditor/Private/Editor/SRigHierarchy.cpp b/Engine/Plugins/Animation/ControlRig/Source/ControlRigEditor/Private/Editor/SRigHierarchy.cpp
--- a/Engine/Plugins/Animation/ControlRig/Source/ControlRigEditor/Private/Editor/SRigHierarchy.cpp
+++ b/Engine/Plugins/Animation/ControlRig/Source/ControlRigEditor/Private/Editor/SRigHierarchy.cpp
@@ -594,6 +594,11 @@ void SRigHierarchy::OnSelectionChanged(TSharedPtr<FRigTreeElement> Selection, ES
 
 void SRigHierarchy::OnHierarchyModified(ERigHierarchyNotification InNotif, URigHierarchy* InHierarchy, const FRigBaseElement* InElement)
 {
+	if (!InElement)
+	{
+		return;
+	}
+
 	if(!ControlRigBlueprint.IsValid())
 	{
 		return;
