Control Rig: Delay setup of control rig withins sequencer to package end load

#Checked in on behalf of Helge.mathee

diff --git a/Engine/Plugins/Animation/ControlRig/Source/ControlRig/Private/ControlRig.cpp b/Engine/Plugins/Animation/ControlRig/Source/ControlRig/Private/ControlRig.cpp
--- a/Engine/Plugins/Animation/ControlRig/Source/ControlRig/Private/ControlRig.cpp
+++ b/Engine/Plugins/Animation/ControlRig/Source/ControlRig/Private/ControlRig.cpp
@@ -200,6 +200,8 @@ void UControlRig::InitializeFromCDO()
 
 		// copy hierarchy
 		{
+			PostInitInstanceIfRequired();
+			
 			TGuardValue<bool> Guard(GetHierarchy()->GetSuspendNotificationsFlag(), true);
 			GetHierarchy()->CopyHierarchy(CDO->GetHierarchy());
 			GetHierarchy()->ResetPoseToInitial(ERigElementType::All);
diff --git a/Engine/Plugins/Animation/ControlRig/Source/ControlRig/Private/Sequencer/MovieSceneControlRigParameterTrack.cpp b/Engine/Plugins/Animation/ControlRig/Source/ControlRig/Private/Sequencer/MovieSceneControlRigParameterTrack.cpp
--- a/Engine/Plugins/Animation/ControlRig/Source/ControlRig/Private/Sequencer/MovieSceneControlRigParameterTrack.cpp
+++ b/Engine/Plugins/Animation/ControlRig/Source/ControlRig/Private/Sequencer/MovieSceneControlRigParameterTrack.cpp
@@ -363,8 +363,27 @@ void UMovieSceneControlRigParameterTrack::PostLoad()
 {
 	Super::PostLoad();
 
+#if WITH_EDITOR
+	FCoreUObjectDelegates::OnEndLoadPackage.AddUObject(this, &UMovieSceneControlRigParameterTrack::HandlePackageDone);
+#else
+	ReconstructControlRig();
+#endif
+}
+
+#if WITH_EDITOR
+void UMovieSceneControlRigParameterTrack::HandlePackageDone(TConstArrayView<UPackage*> InPackages)
+{
+	if (!InPackages.Contains(GetPackage()))
+	{
+		return;
+	}
+
+	FCoreUObjectDelegates::OnEndLoadPackage.RemoveAll(this);
+
 	ReconstructControlRig();
 }
+#endif
+
 
 void UMovieSceneControlRigParameterTrack::PostEditImport()
 {
diff --git a/Engine/Plugins/Animation/ControlRig/Source/ControlRig/Public/Sequencer/MovieSceneControlRigParameterTrack.h b/Engine/Plugins/Animation/ControlRig/Source/ControlRig/Public/Sequencer/MovieSceneControlRigParameterTrack.h
--- a/Engine/Plugins/Animation/ControlRig/Source/ControlRig/Public/Sequencer/MovieSceneControlRigParameterTrack.h
+++ b/Engine/Plugins/Animation/ControlRig/Source/ControlRig/Public/Sequencer/MovieSceneControlRigParameterTrack.h
@@ -49,6 +49,10 @@ public:
 	virtual TArray<FFBXNodeAndChannels>* GetNodeAndChannelMappings()  override;
 	virtual void GetSelectedNodes(TArray<FName>& OutSelectedNodes) override;
 
+#if WITH_EDITOR
+	void HandlePackageDone(TConstArrayView<UPackage*> InPackages);
+#endif
+
 public:
 	/** Add a section at that start time*/
 	CONTROLRIG_API UMovieSceneSection* CreateControlRigSection(FFrameNumber StartTime, UControlRig* InControlRig, bool bInOwnsControlRig);
diff --git a/Engine/Plugins/Animation/ControlRig/Source/ControlRigDeveloper/Private/ControlRigBlueprint.cpp b/Engine/Plugins/Animation/ControlRig/Source/ControlRigDeveloper/Private/ControlRigBlueprint.cpp
--- a/Engine/Plugins/Animation/ControlRig/Source/ControlRigDeveloper/Private/ControlRigBlueprint.cpp
+++ b/Engine/Plugins/Animation/ControlRig/Source/ControlRigDeveloper/Private/ControlRigBlueprint.cpp
@@ -403,8 +403,6 @@ void UControlRigBlueprint::PostLoad()
 			GetHierarchyController()->ImportFromHierarchyContainer(HierarchyContainer_DEPRECATED, false);
 		}
 
-		PropagateHierarchyFromBPToInstances();
-		
 		// remove all non-controlrig-graphs
 		TArray<UEdGraph*> NewUberGraphPages;
 		for (UEdGraph* Graph : UbergraphPages)
