Control rig: allows user to clear function category

diff --git a/Engine/Plugins/Animation/ControlRig/Source/ControlRigEditor/Private/ControlRigGraphDetails.cpp b/Engine/Plugins/Animation/ControlRig/Source/ControlRigEditor/Private/ControlRigGraphDetails.cpp
--- a/Engine/Plugins/Animation/ControlRig/Source/ControlRigEditor/Private/ControlRigGraphDetails.cpp
+++ b/Engine/Plugins/Animation/ControlRig/Source/ControlRigEditor/Private/ControlRigGraphDetails.cpp
@@ -667,12 +667,7 @@ void FControlRigGraphDetails::CustomizeDetails(IDetailLayoutBuilder& DetailLayou
 			.OnVerifyTextChanged_Lambda([&](const FText& InNewText, FText& OutErrorMessage) -> bool
 			{
 				const FText NewText = FEditorCategoryUtils::GetCategoryDisplayString(InNewText);
-				if (NewText.IsEmpty())
-				{
-					OutErrorMessage = LOCTEXT("CategoryEmpty", "Cannot add a category with an empty string.");
-					return false;
-				}
-				else if (NewText.ToString().Len() >= NAME_SIZE)
+				if (NewText.ToString().Len() >= NAME_SIZE)
 				{
 					OutErrorMessage = LOCTEXT("CategoryTooLong", "Name of category is too long.");
 					return false;
