Control Rig: minor static analysis fix

diff --git a/Engine/Plugins/Animation/ControlRig/Source/ControlRig/Public/Rigs/RigHierarchyDefines.h b/Engine/Plugins/Animation/ControlRig/Source/ControlRig/Public/Rigs/RigHierarchyDefines.h
--- a/Engine/Plugins/Animation/ControlRig/Source/ControlRig/Public/Rigs/RigHierarchyDefines.h
+++ b/Engine/Plugins/Animation/ControlRig/Source/ControlRig/Public/Rigs/RigHierarchyDefines.h
@@ -310,16 +310,7 @@ public:
 				Get<FTransform_Float>().ToTransform().GetScale3D().Y,
 				Get<FTransform_Float>().ToTransform().GetScale3D().Z); break;
 			*/
-			case ERigControlType::Transform: ValueStr = FString::Printf(TEXT("unreal.RigHierarchy.make_control_value_from_euler_transform(unreal.EulerTransform(location=[%.6f,%.6f,%.6f],rotation=[%.6f,%.6f,%.6f],scale=[%.6f,%.6f,%.6f]))"),
-				FEulerTransform(Get<FTransform_Float>().ToTransform()).Location.X,
-				FEulerTransform(Get<FTransform_Float>().ToTransform()).Location.Y,
-				FEulerTransform(Get<FTransform_Float>().ToTransform()).Location.Z,
-				FEulerTransform(Get<FTransform_Float>().ToTransform()).Rotation.Euler().X,
-				FEulerTransform(Get<FTransform_Float>().ToTransform()).Rotation.Euler().Y,
-				FEulerTransform(Get<FTransform_Float>().ToTransform()).Rotation.Euler().Z,
-				FEulerTransform(Get<FTransform_Float>().ToTransform()).Scale.X,
-				FEulerTransform(Get<FTransform_Float>().ToTransform()).Scale.Y,
-				FEulerTransform(Get<FTransform_Float>().ToTransform()).Scale.Z); break;
+			case ERigControlType::Transform:
 			case ERigControlType::EulerTransform: ValueStr = FString::Printf(TEXT("unreal.RigHierarchy.make_control_value_from_euler_transform(unreal.EulerTransform(location=[%.6f,%.6f,%.6f],rotation=[%.6f,%.6f,%.6f],scale=[%.6f,%.6f,%.6f]))"),
 				Get<FEulerTransform_Float>().ToTransform().Location.X,
 				Get<FEulerTransform_Float>().ToTransform().Location.Y,
