Control Rig: Fixed reorder of exposed pins

diff --git a/Engine/Plugins/Animation/ControlRig/Source/ControlRigEditor/Private/ControlRigGraphDetails.cpp b/Engine/Plugins/Animation/ControlRig/Source/ControlRigEditor/Private/ControlRigGraphDetails.cpp
--- a/Engine/Plugins/Animation/ControlRig/Source/ControlRigEditor/Private/ControlRigGraphDetails.cpp
+++ b/Engine/Plugins/Animation/ControlRig/Source/ControlRigEditor/Private/ControlRigGraphDetails.cpp
@@ -256,7 +256,32 @@ FReply FControlRigArgumentLayout::OnArgMoveUp()
 		{
 			if (URigVMController* Controller = Blueprint->GetController(LibraryNode->GetContainedGraph()))
 			{
-				Controller->SetExposedPinIndex(Pin->GetFName(), Pin->GetPinIndex() - 1, true, true);
+				bool bIsInput = Pin->GetDirection() == ERigVMPinDirection::Input || Pin->GetDirection() == ERigVMPinDirection::IO;
+				
+				int32 NewPinIndex = Pin->GetPinIndex() - 1;
+				while (NewPinIndex != INDEX_NONE)
+				{
+					URigVMPin* OtherPin = LibraryNode->GetPins()[NewPinIndex];
+					if (bIsInput)
+					{
+						if (OtherPin->GetDirection() == ERigVMPinDirection::Input || OtherPin->GetDirection() == ERigVMPinDirection::IO)
+						{
+							break;
+						}
+					}	
+					else
+					{
+						if (OtherPin->GetDirection() == ERigVMPinDirection::Output || OtherPin->GetDirection() == ERigVMPinDirection::IO)
+						{
+							break;
+						}
+					}
+					--NewPinIndex;
+				}
+				if (NewPinIndex != INDEX_NONE)
+				{
+					Controller->SetExposedPinIndex(Pin->GetFName(), NewPinIndex, true, true);
+				}
 				return FReply::Handled();
 			}
 		}
@@ -274,7 +299,32 @@ FReply FControlRigArgumentLayout::OnArgMoveDown()
 		{
 			if (URigVMController* Controller = Blueprint->GetController(LibraryNode->GetContainedGraph()))
 			{
-				Controller->SetExposedPinIndex(Pin->GetFName(), Pin->GetPinIndex() + 1, true, true);
+				bool bIsInput = Pin->GetDirection() == ERigVMPinDirection::Input || Pin->GetDirection() == ERigVMPinDirection::IO;
+				
+				int32 NewPinIndex = Pin->GetPinIndex() + 1;
+				while (NewPinIndex < LibraryNode->GetPins().Num())
+				{
+					URigVMPin* OtherPin = LibraryNode->GetPins()[NewPinIndex];
+					if (bIsInput)
+					{
+						if (OtherPin->GetDirection() == ERigVMPinDirection::Input || OtherPin->GetDirection() == ERigVMPinDirection::IO)
+						{
+							break;
+						}
+					}	
+					else
+					{
+						if (OtherPin->GetDirection() == ERigVMPinDirection::Output || OtherPin->GetDirection() == ERigVMPinDirection::IO)
+						{
+							break;
+						}
+					}
+					++NewPinIndex;
+				}
+				if (NewPinIndex < LibraryNode->GetPins().Num())
+				{
+					Controller->SetExposedPinIndex(Pin->GetFName(), NewPinIndex, true, true);
+				}
 				return FReply::Handled();
 			}
 		}
