Control Rig: RigHierarchy Context Menu was holding invalidate pointers

diff --git a/Engine/Plugins/Animation/ControlRig/Source/ControlRigEditor/Private/Editor/SRigHierarchy.cpp b/Engine/Plugins/Animation/ControlRig/Source/ControlRigEditor/Private/Editor/SRigHierarchy.cpp
--- a/Engine/Plugins/Animation/ControlRig/Source/ControlRigEditor/Private/Editor/SRigHierarchy.cpp
+++ b/Engine/Plugins/Animation/ControlRig/Source/ControlRigEditor/Private/Editor/SRigHierarchy.cpp
@@ -259,10 +259,13 @@ void SRigHierarchy::Construct(const FArguments& InArgs, TSharedRef<FControlRigEd
 		ControlRigEditor.Pin()->GetKeyDownDelegate().BindLambda([&](const FGeometry& MyGeometry, const FKeyEvent& InKeyEvent)->FReply {
 			return OnKeyDown(MyGeometry, InKeyEvent);
 		});
-		ControlRigEditor.Pin()->OnGetViewportContextMenu().BindSP(this, &SRigHierarchy::GetOrCreateContextMenu);
+		ControlRigEditor.Pin()->OnGetViewportContextMenu().BindSP(this, &SRigHierarchy::GetContextMenu);
 		ControlRigEditor.Pin()->OnViewportContextMenuCommands().BindSP(this, &SRigHierarchy::GetContextMenuCommands);
 		ControlRigEditor.Pin()->OnControlRigEditorClosed().AddSP(this, &SRigHierarchy::OnEditorClose);
 	}
+	
+	CreateContextMenu();
+	CreateDragDropMenu();
 }
 
 void SRigHierarchy::OnEditorClose(const FControlRigEditor* InEditor, UControlRigBlueprint* InBlueprint)
@@ -860,7 +863,7 @@ TSharedPtr< SWidget > SRigHierarchy::CreateContextMenuWidget()
 {
 	UToolMenus* ToolMenus = UToolMenus::Get();
 
-	if (UToolMenu* Menu = GetOrCreateContextMenu())
+	if (UToolMenu* Menu = GetContextMenu())
 	{
 		return ToolMenus->GenerateWidget(Menu);
 	}
@@ -922,15 +925,28 @@ void SRigHierarchy::OnSetExpansionRecursive(TSharedPtr<FRigTreeElement> InItem,
 	TreeView->SetExpansionRecursive(InItem, false, bShouldBeExpanded);
 }
 
-UToolMenu* SRigHierarchy::GetOrCreateDragDropMenu(const TArray<FRigElementKey>& DraggedKeys, FRigElementKey TargetKey)
+void SRigHierarchy::CreateDragDropMenu()
 {
-	const FName MenuName = TEXT("ControlRigEditor.RigHierarchy.DragDropMenu");
+	const FName MenuName = DragDropMenuName;
 	UToolMenus* ToolMenus = UToolMenus::Get();
 
+	if (!ensure(ToolMenus))
+	{
+		return;
+	}
+
+	UToolMenu* Menu	 = nullptr;
 	if (!ToolMenus->IsMenuRegistered(MenuName))
 	{
-		UToolMenu* Menu = ToolMenus->RegisterMenu(MenuName);
+		Menu = ToolMenus->RegisterMenu(MenuName);
+	}
+	else
+	{
+		Menu = ToolMenus->FindMenu(MenuName);
+	}
 
+	if (ensure(Menu))
+	{
 		FToolMenuEntry ParentEntry = FToolMenuEntry::InitMenuEntry(
 			TEXT("Parent"),
 			LOCTEXT("DragDropMenu_Parent", "Parent"),
@@ -970,7 +986,17 @@ UToolMenu* SRigHierarchy::GetOrCreateDragDropMenu(const TArray<FRigElementKey>&
 
 		AlignMenu->AddMenuEntry(NAME_None, AlignAllEntry);
 	}
+}
 
+UToolMenu* SRigHierarchy::GetDragDropMenu(const TArray<FRigElementKey>& DraggedKeys, FRigElementKey TargetKey) const
+{
+	UToolMenus* ToolMenus = UToolMenus::Get();
+	if (!ensure(ToolMenus))
+	{
+		return nullptr;
+	}
+	
+	const FName MenuName = DragDropMenuName;
 	UControlRigContextMenuContext* MenuContext = NewObject<UControlRigContextMenuContext>();
 	FControlRigMenuSpecificContext MenuSpecificContext;
 	MenuSpecificContext.RigHierarchyDragAndDropContext = FControlRigRigHierarchyDragAndDropContext(DraggedKeys, TargetKey);
@@ -981,18 +1007,33 @@ UToolMenu* SRigHierarchy::GetOrCreateDragDropMenu(const TArray<FRigElementKey>&
 	return Menu;
 }
 
-UToolMenu* SRigHierarchy::GetOrCreateContextMenu()
+void SRigHierarchy::CreateContextMenu()
 {
-	const FName MenuName = TEXT("ControlRigEditor.RigHierarchy.ContextMenu");
+	const FName MenuName = ContextMenuName;
 	const FName InteractionSectionName = TEXT("Interaction");
 
 	UToolMenus* ToolMenus = UToolMenus::Get();
+	
+	if (!ensure(ToolMenus))
+	{
+		return;
+	}
+	
 	const FControlRigHierarchyCommands& Commands = FControlRigHierarchyCommands::Get();
 	
+	UToolMenu* Menu = nullptr;
+	
 	if (!ToolMenus->IsMenuRegistered(MenuName))
 	{
-		UToolMenu* Menu = ToolMenus->RegisterMenu(MenuName);
+		Menu = ToolMenus->RegisterMenu(MenuName);
+	}
+	else
+	{
+		Menu = ToolMenus->FindMenu(MenuName);
+	}
 
+	if (ensure(Menu))
+	{
 		struct FLocalMenuBuilder
 		{
 			static void FillNewMenu(FMenuBuilder& InSubMenuBuilder, TSharedPtr<SRigHierarchyTreeView> InTreeView)
@@ -1017,7 +1058,7 @@ UToolMenu* SRigHierarchy::GetOrCreateContextMenu()
 
 		FToolMenuSection& ElementsSection = Menu->AddSection(TEXT("Elements"), LOCTEXT("ElementsHeader", "Elements"));
 		ElementsSection.AddSubMenu(TEXT("New"), LOCTEXT("New", "New"), LOCTEXT("New_ToolTip", "Create New Elements"),
-                                   FNewMenuDelegate::CreateStatic(&FLocalMenuBuilder::FillNewMenu, TreeView));
+								   FNewMenuDelegate::CreateStatic(&FLocalMenuBuilder::FillNewMenu, TreeView));
 		ElementsSection.AddMenuEntry(Commands.DeleteItem);
 		ElementsSection.AddMenuEntry(Commands.DuplicateItem);
 		ElementsSection.AddMenuEntry(Commands.RenameItem);
@@ -1062,23 +1103,34 @@ UToolMenu* SRigHierarchy::GetOrCreateContextMenu()
 
 		FToolMenuSection& AssetsSection = Menu->AddSection(TEXT("Assets"), LOCTEXT("AssetsHeader", "Assets"));
 		AssetsSection.AddSubMenu(TEXT("Import"), LOCTEXT("ImportSubMenu", "Import"),
-            LOCTEXT("ImportSubMenu_ToolTip", "Import hierarchy to the current rig. This only imports non-existing node. For example, if there is hand_r, it won't import hand_r. If you want to reimport whole new hiearchy, delete all nodes, and use import hierarchy."),
-            FNewMenuDelegate::CreateSP(this, &SRigHierarchy::CreateImportMenu)
-        );
+			LOCTEXT("ImportSubMenu_ToolTip", "Import hierarchy to the current rig. This only imports non-existing node. For example, if there is hand_r, it won't import hand_r. If you want to reimport whole new hiearchy, delete all nodes, and use import hierarchy."),
+			FNewMenuDelegate::CreateSP(this, &SRigHierarchy::CreateImportMenu)
+		);
 		
 		AssetsSection.AddSubMenu(TEXT("Refresh"), LOCTEXT("RefreshSubMenu", "Refresh"),
 			LOCTEXT("RefreshSubMenu_ToolTip", "Refresh the existing initial transform from the selected mesh. This only updates if the node is found."),
-            FNewMenuDelegate::CreateSP(this, &SRigHierarchy::CreateRefreshMenu)
-        );
+			FNewMenuDelegate::CreateSP(this, &SRigHierarchy::CreateRefreshMenu)
+		);
+	}
+}
+
+UToolMenu* SRigHierarchy::GetContextMenu() const
+{
+	const FName MenuName = ContextMenuName;
+	UToolMenus* ToolMenus = UToolMenus::Get();
+
+	if(!ensure(ToolMenus))
+	{
+		return nullptr;
 	}
 
 	// individual entries in this menu can access members of this context, particularly useful for editor scripting
 	UControlRigContextMenuContext* ContextMenuContext = NewObject<UControlRigContextMenuContext>();
 	ContextMenuContext->Init(ControlRigBlueprint);
-	
+
 	FToolMenuContext MenuContext(CommandList);
 	MenuContext.AddObject(ContextMenuContext);
-	
+
 	UToolMenu* Menu = ToolMenus->GenerateMenu(MenuName, MenuContext);
 
 	return Menu;
@@ -1929,7 +1981,7 @@ FReply SRigHierarchy::OnAcceptDrop(const FDragDropEvent& DragDropEvent, EItemDro
 			const FVector2D& SummonLocation = DragDropEvent.GetScreenSpacePosition();
 
 			// Get the context menu content. If NULL, don't open a menu.
-			UToolMenu* DragDropMenu = GetOrCreateDragDropMenu(RigDragDropOp->GetElements(), TargetItem->Key);
+			UToolMenu* DragDropMenu = GetDragDropMenu(RigDragDropOp->GetElements(), TargetItem->Key);
 			const TSharedPtr<SWidget> MenuContent = UToolMenus::Get()->GenerateWidget(DragDropMenu);
 
 			if (MenuContent.IsValid())
diff --git a/Engine/Plugins/Animation/ControlRig/Source/ControlRigEditor/Private/Editor/SRigHierarchy.h b/Engine/Plugins/Animation/ControlRig/Source/ControlRigEditor/Private/Editor/SRigHierarchy.h
--- a/Engine/Plugins/Animation/ControlRig/Source/ControlRigEditor/Private/Editor/SRigHierarchy.h
+++ b/Engine/Plugins/Animation/ControlRig/Source/ControlRigEditor/Private/Editor/SRigHierarchy.h
@@ -147,10 +147,14 @@ private:
 	TOptional<EItemDropZone> OnCanAcceptDrop(const FDragDropEvent& DragDropEvent, EItemDropZone DropZone, TSharedPtr<FRigTreeElement> TargetItem);
 	FReply OnAcceptDrop(const FDragDropEvent& DragDropEvent, EItemDropZone DropZone, TSharedPtr<FRigTreeElement> TargetItem);
 
-	UToolMenu* GetOrCreateContextMenu();
+	const FName ContextMenuName = TEXT("ControlRigEditor.RigHierarchy.ContextMenu");
+	void CreateContextMenu();
+	UToolMenu* GetContextMenu() const;
 	TSharedPtr<FUICommandList> GetContextMenuCommands() const;
-
-	UToolMenu* GetOrCreateDragDropMenu(const TArray<FRigElementKey>& DraggedKeys, FRigElementKey TargetKey);
+	
+	const FName DragDropMenuName = TEXT("ControlRigEditor.RigHierarchy.DragDropMenu");
+	void CreateDragDropMenu();
+	UToolMenu* GetDragDropMenu(const TArray<FRigElementKey>& DraggedKeys, FRigElementKey TargetKey) const;
 
 	/** Our owning control rig editor */
 	TWeakPtr<FControlRigEditor> ControlRigEditor;
