ControlRig: make Preview world floor look like the one in the Anim BP Editor but with collision enabled

diff --git a/Engine/Plugins/Animation/ControlRig/Source/ControlRigEditor/Private/Editor/ControlRigEditor.cpp b/Engine/Plugins/Animation/ControlRig/Source/ControlRigEditor/Private/Editor/ControlRigEditor.cpp
--- a/Engine/Plugins/Animation/ControlRig/Source/ControlRigEditor/Private/Editor/ControlRigEditor.cpp
+++ b/Engine/Plugins/Animation/ControlRig/Source/ControlRigEditor/Private/Editor/ControlRigEditor.cpp
@@ -3029,8 +3029,6 @@ void FControlRigEditor::Tick(float DeltaTime)
 	if(WeakGroundActorPtr.IsValid())
 	{
 		const TSharedRef<IPersonaPreviewScene> CurrentPreviewScene = GetPersonaToolkit()->GetPreviewScene();
-		WeakGroundActorPtr->GetStaticMeshComponent()->SetVisibility(CurrentPreviewScene->GetFloorVisibility());
-
 		const float FloorOffset = CurrentPreviewScene->GetFloorOffset();
 		const FTransform FloorTransform(FRotator(0, 0, 0), FVector(0, 0, -(FloorOffset)), FVector(4.0f, 4.0f, 1.0f));
 		WeakGroundActorPtr->GetStaticMeshComponent()->SetRelativeTransform(FloorTransform);
@@ -3762,6 +3760,8 @@ void FControlRigEditor::HandlePreviewSceneCreated(const TSharedRef<IPersonaPrevi
 	GroundActor->GetStaticMeshComponent()->SetCollisionEnabled(ECollisionEnabled::QueryAndPhysics);
 	GroundActor->GetStaticMeshComponent()->SetCollisionProfileName(UCollisionProfile::BlockAll_ProfileName);
 	GroundActor->GetStaticMeshComponent()->bSelectable = false;
+	// this will be an invisible collision box that users can use to test traces
+	GroundActor->GetStaticMeshComponent()->SetVisibility(false);
 
 	WeakGroundActorPtr = GroundActor;
 
