Control Rig: Refactor all menus in Control Rig to use dynamic sections to avoid caching invalid poiners

diff --git a/Engine/Plugins/Animation/ControlRig/Source/ControlRigEditor/Private/Editor/ControlRigContextMenuContext.cpp b/Engine/Plugins/Animation/ControlRig/Source/ControlRigEditor/Private/Editor/ControlRigContextMenuContext.cpp
--- a/Engine/Plugins/Animation/ControlRig/Source/ControlRigEditor/Private/Editor/ControlRigContextMenuContext.cpp
+++ b/Engine/Plugins/Animation/ControlRig/Source/ControlRigEditor/Private/Editor/ControlRigContextMenuContext.cpp
@@ -16,15 +16,20 @@ FString FControlRigRigHierarchyToGraphDragAndDropContext::GetSectionTitle() cons
 	return FString::Join(ElementNameStrings, TEXT(","));
 }
 
-void UControlRigContextMenuContext::Init(TWeakObjectPtr<UControlRigBlueprint> InControlRigBlueprint, const FControlRigMenuSpecificContext& InMenuSpecificContext)
+void UControlRigContextMenuContext::Init(TWeakPtr<FControlRigEditor> InControlRigEditor, const FControlRigMenuSpecificContext& InMenuSpecificContext)
 {
-	ControlRigBlueprint = InControlRigBlueprint;
+	WeakControlRigEditor = InControlRigEditor;
 	MenuSpecificContext = InMenuSpecificContext;
 }
 
 UControlRigBlueprint* UControlRigContextMenuContext::GetControlRigBlueprint() const
 {
-	return ControlRigBlueprint.Get();
+	if (const TSharedPtr<FControlRigEditor> Editor = WeakControlRigEditor.Pin())
+	{
+		return Editor->GetControlRigBlueprint();
+	}
+	
+	return nullptr;
 }
 
 UControlRig* UControlRigContextMenuContext::GetControlRig() const
@@ -57,4 +62,22 @@ FControlRigGraphNodeContextMenuContext UControlRigContextMenuContext::GetGraphNo
 FControlRigRigHierarchyToGraphDragAndDropContext UControlRigContextMenuContext::GetRigHierarchyToGraphDragAndDropContext()
 {
 	return MenuSpecificContext.RigHierarchyToGraphDragAndDropContext;
-}
\ No newline at end of file
+}
+
+SRigHierarchy* UControlRigContextMenuContext::GetRigHierarchyPanel() const
+{
+	if (const TSharedPtr<SRigHierarchy> RigHierarchyPanel = MenuSpecificContext.RigHierarchyPanel.Pin())
+	{
+		return RigHierarchyPanel.Get();
+	}
+	return nullptr;
+}
+
+FControlRigEditor* UControlRigContextMenuContext::GetControlRigEditor() const
+{
+	if (const TSharedPtr<FControlRigEditor> Editor = WeakControlRigEditor.Pin())
+	{
+		return Editor.Get();
+	}
+	return nullptr;
+}
diff --git a/Engine/Plugins/Animation/ControlRig/Source/ControlRigEditor/Private/Editor/ControlRigContextMenuContext.h b/Engine/Plugins/Animation/ControlRig/Source/ControlRigEditor/Private/Editor/ControlRigContextMenuContext.h
--- a/Engine/Plugins/Animation/ControlRig/Source/ControlRigEditor/Private/Editor/ControlRigContextMenuContext.h
+++ b/Engine/Plugins/Animation/ControlRig/Source/ControlRigEditor/Private/Editor/ControlRigContextMenuContext.h
@@ -15,6 +15,7 @@ class URigVMGraph;
 class URigVMNode;
 class URigVMPin;
 class UEdGraph;
+class SRigHierarchy;
 
 USTRUCT(BlueprintType)
 struct FControlRigRigHierarchyDragAndDropContext
@@ -93,6 +94,8 @@ struct FControlRigRigHierarchyToGraphDragAndDropContext
 
 struct FControlRigMenuSpecificContext
 {
+	TWeakPtr<SRigHierarchy> RigHierarchyPanel;
+	
 	FControlRigRigHierarchyDragAndDropContext RigHierarchyDragAndDropContext;
 	
 	FControlRigGraphNodeContextMenuContext GraphNodeContextMenuContext;
@@ -108,10 +111,10 @@ class UControlRigContextMenuContext : public UObject
 public:
 	/**
 	 *	Initialize the Context
-	 * @param InControlRigBlueprint 	The Control Rig Bluerpint currently opened in the Editor
+	 * @param InControlRigEditor 	    The Control Rig Editor hosting the menus
 	 * @param InMenuSpecificContext 	Additional context for specific menus
 	*/
-	void Init(TWeakObjectPtr<UControlRigBlueprint> InControlRigBlueprint, const FControlRigMenuSpecificContext& InMenuSpecificContext = FControlRigMenuSpecificContext());
+	void Init(TWeakPtr<FControlRigEditor> InControlRigEditor, const FControlRigMenuSpecificContext& InMenuSpecificContext = FControlRigMenuSpecificContext());
 
 	/** Get the control rig blueprint that we are editing */
 	UFUNCTION(BlueprintCallable, Category = ControlRigEditorExtensions)
@@ -137,8 +140,12 @@ public:
 	UFUNCTION(BlueprintCallable, Category = ControlRigEditorExtensions)
     FControlRigRigHierarchyToGraphDragAndDropContext GetRigHierarchyToGraphDragAndDropContext();
 
+	SRigHierarchy* GetRigHierarchyPanel() const;
+	
+	FControlRigEditor* GetControlRigEditor() const;
+
 private:
-	TWeakObjectPtr<UControlRigBlueprint> ControlRigBlueprint;
+	TWeakPtr<FControlRigEditor> WeakControlRigEditor;
 
 	FControlRigMenuSpecificContext MenuSpecificContext;
 };
\ No newline at end of file
diff --git a/Engine/Plugins/Animation/ControlRig/Source/ControlRigEditor/Private/Editor/ControlRigEditor.cpp b/Engine/Plugins/Animation/ControlRig/Source/ControlRigEditor/Private/Editor/ControlRigEditor.cpp
--- a/Engine/Plugins/Animation/ControlRig/Source/ControlRigEditor/Private/Editor/ControlRigEditor.cpp
+++ b/Engine/Plugins/Animation/ControlRig/Source/ControlRigEditor/Private/Editor/ControlRigEditor.cpp
@@ -529,6 +529,8 @@ void FControlRigEditor::InitControlRigEditor(const EToolkitMode::Type Mode, cons
 	}, TStatId(), NULL, ENamedThreads::GameThread);
 
 	UpdateCapsules();
+	
+	CreateRigHierarchyToGraphDragAndDropMenu();
 }
 
 void FControlRigEditor::BindCommands()
@@ -2034,12 +2036,12 @@ void FControlRigEditor::InitToolMenuContext(FToolMenuContext& MenuContext)
 			}
 		}
 		
-		UControlRigContextMenuContext* ControlRigContext = NewObject<UControlRigContextMenuContext>();
+		UControlRigContextMenuContext* ControlRigMenuContext = NewObject<UControlRigContextMenuContext>();
 		FControlRigMenuSpecificContext MenuSpecificContext;	
 		MenuSpecificContext.GraphNodeContextMenuContext = FControlRigGraphNodeContextMenuContext(Model, Node, Pin);
-		ControlRigContext->Init(RigBlueprint, MenuSpecificContext);
+		ControlRigMenuContext->Init(SharedThis(this), MenuSpecificContext);
 
-		MenuContext.AddObject(ControlRigContext);
+		MenuContext.AddObject(ControlRigMenuContext);
 	}
 }
 
@@ -5166,279 +5168,284 @@ void FControlRigEditor::OnCurveContainerChanged()
 	NotificationPtr->SetCompletionState(SNotificationItem::CS_Success);
 }
 
-void FControlRigEditor::OnGraphNodeDropToPerform(TSharedPtr<FGraphNodeDragDropOp> InDragDropOp, UEdGraph* InGraph, const FVector2D& InNodePosition, const FVector2D& InScreenPosition)
+void FControlRigEditor::CreateRigHierarchyToGraphDragAndDropMenu() const
 {
-	DECLARE_SCOPE_HIERARCHICAL_COUNTER_FUNC()
+	const FName MenuName = RigHierarchyToGraphDragAndDropMenuName;
+	UToolMenus* ToolMenus = UToolMenus::Get();
+	
+	if(!ensure(ToolMenus))
+	{
+		return;
+	}
 
-	if (InDragDropOp->IsOfType<FRigElementHierarchyDragDropOp>())
+	if (!ToolMenus->IsMenuRegistered(MenuName))
 	{
-		TSharedPtr<FRigElementHierarchyDragDropOp> RigHierarchyOp = StaticCastSharedPtr<FRigElementHierarchyDragDropOp>(InDragDropOp);
+		UToolMenu* Menu = ToolMenus->RegisterMenu(MenuName);
 
-		if (RigHierarchyOp->GetElements().Num() > 0 && FocusedGraphEdPtr.IsValid())
-		{
-			const FName MenuName = TEXT("ControlRigEditor.RigHierarchyToGraphDragAndDropMenu");
-			UToolMenus* ToolMenus = UToolMenus::Get();
-			
-			if (!ToolMenus->IsMenuRegistered(MenuName))
+		Menu->AddDynamicSection(NAME_None, FNewToolMenuDelegate::CreateLambda([](UToolMenu* InMenu)
 			{
-				UToolMenu* Menu = ToolMenus->RegisterMenu(MenuName);
+				UControlRigContextMenuContext* MainContext = InMenu->FindContext<UControlRigContextMenuContext>();
 				
-				const FName SectionName = TEXT("RigHierarchyToGraphDragAndDropSection");
-				const FName DynamicSectionName = *(SectionName.ToString() + TEXT("Dynamic"));
+				if (FControlRigEditor* ControlRigEditor = MainContext->GetControlRigEditor())
+				{
+					const FControlRigRigHierarchyToGraphDragAndDropContext& DragDropContext = MainContext->GetRigHierarchyToGraphDragAndDropContext();
+
+					UControlRigBlueprint* Blueprint = MainContext->GetControlRigBlueprint();
+					const TArray<FRigElementKey>& DraggedKeys = DragDropContext.DraggedElementKeys;
+					UEdGraph* Graph = DragDropContext.Graph.Get();
+					const FVector2D& NodePosition = DragDropContext.NodePosition;
+					
+					// if multiple types are selected, we show Get Elements/Set Elements
+					bool bMultipleTypeSelected = false;
+
+					ERigElementType LastType = ERigElementType::None;
+			
+					if (DraggedKeys.Num() > 0)
+					{
+						LastType = DraggedKeys[0].Type;
+					}
+			
+					uint8 DraggedTypes = 0;
+					for (const FRigElementKey& DraggedKey : DragDropContext.DraggedElementKeys)
+					{
+						if (DraggedKey.Type != LastType)
+						{
+							bMultipleTypeSelected = true;
+						}
 				
-				Menu->AddDynamicSection( DynamicSectionName,
-                    FNewToolMenuDelegate::CreateLambda([
-                    	SectionName,
-                    	WeakThis = TWeakPtr<FControlRigEditor>(SharedThis(this))](UToolMenu* InMenu)
-                    {
-                        if (const TSharedPtr<FControlRigEditor> ControlRigEditor = WeakThis.Pin())
-                        {
-							UControlRigContextMenuContext* MainContext = InMenu->FindContext<UControlRigContextMenuContext>();
-							const FControlRigRigHierarchyToGraphDragAndDropContext& DragDropContext = MainContext->GetRigHierarchyToGraphDragAndDropContext();
-
-                        	UControlRigBlueprint* Blueprint = MainContext->GetControlRigBlueprint();
-							const TArray<FRigElementKey>& DraggedKeys = DragDropContext.DraggedElementKeys;
-                        	UEdGraph* Graph = DragDropContext.Graph.Get();
-                        	const FVector2D& NodePosition = DragDropContext.NodePosition;
-                        	
-                        	// if multiple types are selected, we show Get Elements/Set Elements
-							bool bMultipleTypeSelected = false;
-
-							ERigElementType LastType = ERigElementType::None;
+						DraggedTypes = DraggedTypes | (uint8)DraggedKey.Type;
+					}
 					
-							if (DraggedKeys.Num() > 0)
-							{
-								LastType = DraggedKeys[0].Type;
-							}
+					const FText SectionText = FText::FromString(DragDropContext.GetSectionTitle());
+					FToolMenuSection& Section = InMenu->AddSection(NAME_None, SectionText);
+
+					FText GetterLabel = LOCTEXT("GetElement","Get Element");
+					FText GetterTooltip = LOCTEXT("GetElement_ToolTip", "Getter For Element");
+					FText SetterLabel = LOCTEXT("SetElement","Set Element");
+					FText SetterTooltip = LOCTEXT("SetElement_ToolTip", "Setter For Element");
+					// if multiple types are selected, we show Get Elements/Set Elements
+					if (bMultipleTypeSelected)
+					{
+						GetterLabel = LOCTEXT("GetElements","Get Elements");
+						GetterTooltip = LOCTEXT("GetElements_ToolTip", "Getter For Elements");
+						SetterLabel = LOCTEXT("SetElements","Set Elements");
+						SetterTooltip = LOCTEXT("SetElements_ToolTip", "Setter For Elements");
+					}
+					else
+					{
+						// otherwise, we show "Get Bone/NUll/Control"
+						if ((DraggedTypes & (uint8)ERigElementType::Bone) != 0)
+						{
+							GetterLabel = LOCTEXT("GetBone","Get Bone");
+							GetterTooltip = LOCTEXT("GetBone_ToolTip", "Getter For Bone");
+							SetterLabel = LOCTEXT("SetBone","Set Bone");
+							SetterTooltip = LOCTEXT("SetBone_ToolTip", "Setter For Bone");
+						}
+						else if ((DraggedTypes & (uint8)ERigElementType::Null) != 0)
+						{
+							GetterLabel = LOCTEXT("GetNull","Get Null");
+							GetterTooltip = LOCTEXT("GetNull_ToolTip", "Getter For Null");
+							SetterLabel = LOCTEXT("SetNull","Set Null");
+							SetterTooltip = LOCTEXT("SetNull_eoolTip", "Setter For Null");
+						}
+						else if ((DraggedTypes & (uint8)ERigElementType::Control) != 0)
+						{
+							GetterLabel = LOCTEXT("GetControl","Get Control");
+							GetterTooltip = LOCTEXT("GetControl_ToolTip", "Getter For Control");
+							SetterLabel = LOCTEXT("SetControl","Set Control");
+							SetterTooltip = LOCTEXT("SetControl_ToolTip", "Setter For Control");
+						}
+					}
+
+					FToolMenuEntry GetElementsEntry = FToolMenuEntry::InitMenuEntry(
+						TEXT("GetElements"),
+						GetterLabel,
+						GetterTooltip,
+						FSlateIcon(),
+						FUIAction(
+							FExecuteAction::CreateSP(ControlRigEditor, &FControlRigEditor::HandleMakeElementGetterSetter, ERigElementGetterSetterType_Transform, true, DraggedKeys, Graph, NodePosition),
+							FCanExecuteAction()
+						)
+					);
+					GetElementsEntry.InsertPosition.Name = NAME_None;
+					GetElementsEntry.InsertPosition.Position = EToolMenuInsertType::First;
 					
-							uint8 DraggedTypes = 0;
-							for (const FRigElementKey& DraggedKey : DragDropContext.DraggedElementKeys)
-							{
-								if (DraggedKey.Type != LastType)
-								{
-									bMultipleTypeSelected = true;
-								}
+
+					Section.AddEntry(GetElementsEntry);
+
+					FToolMenuEntry SetElementsEntry = FToolMenuEntry::InitMenuEntry(
+						TEXT("SetElements"),
+						SetterLabel,
+						SetterTooltip,
+						FSlateIcon(),
+						FUIAction(
+							FExecuteAction::CreateSP(ControlRigEditor, &FControlRigEditor::HandleMakeElementGetterSetter, ERigElementGetterSetterType_Transform, false, DraggedKeys, Graph, NodePosition),
+							FCanExecuteAction()
+						)
+					);
+					SetElementsEntry.InsertPosition.Name = GetElementsEntry.Name;
+					SetElementsEntry.InsertPosition.Position = EToolMenuInsertType::After;	
+
+					Section.AddEntry(SetElementsEntry);
+
+					if (((DraggedTypes & (uint8)ERigElementType::Bone) != 0) ||
+						((DraggedTypes & (uint8)ERigElementType::Control) != 0) ||
+						((DraggedTypes & (uint8)ERigElementType::Null) != 0))
+					{
+						FToolMenuEntry& RotationTranslationSeparator = Section.AddSeparator(TEXT("RotationTranslationSeparator"));
+						RotationTranslationSeparator.InsertPosition.Name = SetElementsEntry.Name;
+						RotationTranslationSeparator.InsertPosition.Position = EToolMenuInsertType::After;
+
+						FToolMenuEntry SetRotationEntry = FToolMenuEntry::InitMenuEntry(
+							TEXT("SetRotation"),
+							LOCTEXT("SetRotation","Set Rotation"),
+							LOCTEXT("SetRotation_ToolTip","Setter for Rotation"),
+							FSlateIcon(),
+							FUIAction(
+								FExecuteAction::CreateSP(ControlRigEditor, &FControlRigEditor::HandleMakeElementGetterSetter, ERigElementGetterSetterType_Rotation, false, DraggedKeys, Graph, NodePosition),
+								FCanExecuteAction()
+							)
+						);
 						
-								DraggedTypes = DraggedTypes | (uint8)DraggedKey.Type;
-							}
-                        	
-							const FText SectionText = FText::FromString(DragDropContext.GetSectionTitle());
-							FToolMenuSection& Section = InMenu->AddSection(SectionName, SectionText);
-
-                        	FText GetterLabel = LOCTEXT("GetElement","Get Element");
-                        	FText GetterTooltip = LOCTEXT("GetElement_ToolTip", "Getter For Element");
-                        	FText SetterLabel = LOCTEXT("SetElement","Set Element");
-                        	FText SetterTooltip = LOCTEXT("SetElement_ToolTip", "Setter For Element");
-							// if multiple types are selected, we show Get Elements/Set Elements
-                        	if (bMultipleTypeSelected)
-                        	{
-                        		GetterLabel = LOCTEXT("GetElements","Get Elements");
-                        		GetterTooltip = LOCTEXT("GetElements_ToolTip", "Getter For Elements");
-                        		SetterLabel = LOCTEXT("SetElements","Set Elements");
-                        		SetterTooltip = LOCTEXT("SetElements_ToolTip", "Setter For Elements");
-                        	}
-                        	else
-                        	{
-								// otherwise, we show "Get Bone/NUll/Control"
-                        		if ((DraggedTypes & (uint8)ERigElementType::Bone) != 0)
-                        		{
-                        			GetterLabel = LOCTEXT("GetBone","Get Bone");
-                        			GetterTooltip = LOCTEXT("GetBone_ToolTip", "Getter For Bone");
-                        			SetterLabel = LOCTEXT("SetBone","Set Bone");
-                        			SetterTooltip = LOCTEXT("SetBone_ToolTip", "Setter For Bone");
-                        		}
-                        		else if ((DraggedTypes & (uint8)ERigElementType::Null) != 0)
-                        		{
-                        			GetterLabel = LOCTEXT("GetNull","Get Null");
-                        			GetterTooltip = LOCTEXT("GetNull_ToolTip", "Getter For Null");
-                        			SetterLabel = LOCTEXT("SetNull","Set Null");
-                        			SetterTooltip = LOCTEXT("SetNull_eoolTip", "Setter For Null");
-                        		}
-                        		else if ((DraggedTypes & (uint8)ERigElementType::Control) != 0)
-                        		{
-                        			GetterLabel = LOCTEXT("GetControl","Get Control");
-                        			GetterTooltip = LOCTEXT("GetControl_ToolTip", "Getter For Control");
-                        			SetterLabel = LOCTEXT("SetControl","Set Control");
-                        			SetterTooltip = LOCTEXT("SetControl_ToolTip", "Setter For Control");
-                        		}
-                        	}
-
-                        	FToolMenuEntry GetElementsEntry = FToolMenuEntry::InitMenuEntry(
-								TEXT("GetElements"),
-								GetterLabel,
-								GetterTooltip,
-								FSlateIcon(),
-								FUIAction(
-									FExecuteAction::CreateSP(ControlRigEditor.Get(), &FControlRigEditor::HandleMakeElementGetterSetter, ERigElementGetterSetterType_Transform, true, DraggedKeys, Graph, NodePosition),
-									FCanExecuteAction()
-								)
-                            );
-                        	GetElementsEntry.InsertPosition.Name = NAME_None;
-                        	GetElementsEntry.InsertPosition.Position = EToolMenuInsertType::First;
-                        	
-
-                        	Section.AddEntry(GetElementsEntry);
-
-                        	FToolMenuEntry SetElementsEntry = FToolMenuEntry::InitMenuEntry(
-								TEXT("SetElements"),
-								SetterLabel,
-								SetterTooltip,
-								FSlateIcon(),
-								FUIAction(
-									FExecuteAction::CreateSP(ControlRigEditor.Get(), &FControlRigEditor::HandleMakeElementGetterSetter, ERigElementGetterSetterType_Transform, false, DraggedKeys, Graph, NodePosition),
-									FCanExecuteAction()
+						SetRotationEntry.InsertPosition.Name = RotationTranslationSeparator.Name;
+						SetRotationEntry.InsertPosition.Position = EToolMenuInsertType::After;		
+						Section.AddEntry(SetRotationEntry);
+
+						FToolMenuEntry SetTranslationEntry = FToolMenuEntry::InitMenuEntry(
+							TEXT("SetTranslation"),
+							LOCTEXT("SetTranslation","Set Translation"),
+							LOCTEXT("SetTranslation_ToolTip","Setter for Translation"),
+							FSlateIcon(),
+							FUIAction(
+								FExecuteAction::CreateSP(ControlRigEditor, &FControlRigEditor::HandleMakeElementGetterSetter, ERigElementGetterSetterType_Translation, false, DraggedKeys, Graph, NodePosition),
+								FCanExecuteAction()
+							)
+						);
+
+						SetTranslationEntry.InsertPosition.Name = SetRotationEntry.Name;
+						SetTranslationEntry.InsertPosition.Position = EToolMenuInsertType::After;		
+						Section.AddEntry(SetTranslationEntry);
+
+						FToolMenuEntry AddOffsetEntry = FToolMenuEntry::InitMenuEntry(
+							TEXT("AddOffset"),
+							LOCTEXT("AddOffset","Add Offset"),
+							LOCTEXT("AddOffset_ToolTip","Setter for Offset"),
+							FSlateIcon(),
+							FUIAction(
+								FExecuteAction::CreateSP(ControlRigEditor, &FControlRigEditor::HandleMakeElementGetterSetter, ERigElementGetterSetterType_Offset, false, DraggedKeys, Graph, NodePosition),
+								FCanExecuteAction()
+							)
+						);
+						
+						AddOffsetEntry.InsertPosition.Name = SetTranslationEntry.Name;
+						AddOffsetEntry.InsertPosition.Position = EToolMenuInsertType::After;						
+						Section.AddEntry(AddOffsetEntry);
+
+						FToolMenuEntry& RelativeTransformSeparator = Section.AddSeparator(TEXT("RelativeTransformSeparator"));
+						RelativeTransformSeparator.InsertPosition.Name = AddOffsetEntry.Name;
+						RelativeTransformSeparator.InsertPosition.Position = EToolMenuInsertType::After;
+						
+						FToolMenuEntry GetRelativeTransformEntry = FToolMenuEntry::InitMenuEntry(
+							TEXT("GetRelativeTransformEntry"),
+							LOCTEXT("GetRelativeTransform", "Get Relative Transform"),
+							LOCTEXT("GetRelativeTransform_ToolTip", "Getter for Relative Transform"),
+							FSlateIcon(),
+							FUIAction(
+								FExecuteAction::CreateSP(ControlRigEditor, &FControlRigEditor::HandleMakeElementGetterSetter, ERigElementGetterSetterType_Relative, true, DraggedKeys, Graph, NodePosition),
+								FCanExecuteAction()
+							)
+						);
+						
+						GetRelativeTransformEntry.InsertPosition.Name = RelativeTransformSeparator.Name;
+						GetRelativeTransformEntry.InsertPosition.Position = EToolMenuInsertType::After;	
+						Section.AddEntry(GetRelativeTransformEntry);
+						
+						FToolMenuEntry SetRelativeTransformEntry = FToolMenuEntry::InitMenuEntry(
+							TEXT("SetRelativeTransformEntry"),
+							LOCTEXT("SetRelativeTransform", "Set Relative Transform"),
+							LOCTEXT("SetRelativeTransform_ToolTip", "Setter for Relative Transform"),
+							FSlateIcon(),
+							FUIAction(
+								FExecuteAction::CreateSP(ControlRigEditor, &FControlRigEditor::HandleMakeElementGetterSetter, ERigElementGetterSetterType_Relative, false, DraggedKeys, Graph, NodePosition),
+								FCanExecuteAction()
+							)
+						);
+						SetRelativeTransformEntry.InsertPosition.Name = GetRelativeTransformEntry.Name;
+						SetRelativeTransformEntry.InsertPosition.Position = EToolMenuInsertType::After;		
+						Section.AddEntry(SetRelativeTransformEntry);
+					}
+
+					if (DraggedKeys.Num() > 0 && Blueprint != nullptr)
+					{
+						FToolMenuEntry& ItemArraySeparator = Section.AddSeparator(TEXT("ItemArraySeparator"));
+						ItemArraySeparator.InsertPosition.Name = TEXT("SetRelativeTransformEntry"),
+						ItemArraySeparator.InsertPosition.Position = EToolMenuInsertType::After;
+						
+						FToolMenuEntry CreateItemArrayEntry = FToolMenuEntry::InitMenuEntry(
+							TEXT("CreateItemArray"),
+							LOCTEXT("CreateItemArray", "Create Item Array"),
+							LOCTEXT("CreateItemArray_ToolTip", "Creates an item array from the selected elements in the hierarchy"),
+							FSlateIcon(),
+							FUIAction(
+								FExecuteAction::CreateLambda([ControlRigEditor, Blueprint, DraggedKeys, NodePosition]()
+									{
+										if (URigVMController* Controller = ControlRigEditor->GetFocusedController())
+										{
+											Controller->OpenUndoBracket(TEXT("Create Item Array From Selection"));
+
+											static const FString ItemArrayCPPType = FString::Printf(TEXT("TArray<%s>"), *FRigElementKey::StaticStruct()->GetStructCPPName());
+											static const FString ItemArrayObjectPath = FRigElementKey::StaticStruct()->GetPathName();
+							
+											if (URigVMNode* ItemsNode = Controller->AddFreeRerouteNode(true, ItemArrayCPPType, *ItemArrayObjectPath, false, NAME_None, FString(), NodePosition))
+											{
+												if (URigVMPin* ItemsPin = ItemsNode->FindPin(TEXT("Value")))
+												{
+													Controller->SetArrayPinSize(ItemsPin->GetPinPath(), DraggedKeys.Num());
+
+													TArray<URigVMPin*> ItemPins = ItemsPin->GetSubPins();
+													ensure(ItemPins.Num() == DraggedKeys.Num());
+
+													for (int32 ItemIndex = 0; ItemIndex < DraggedKeys.Num(); ItemIndex++)
+													{
+														FString DefaultValue;
+														FRigElementKey::StaticStruct()->ExportText(DefaultValue, &DraggedKeys[ItemIndex], nullptr, nullptr, PPF_None, nullptr);
+														Controller->SetPinDefaultValue(ItemPins[ItemIndex]->GetPinPath(), DefaultValue, true, true, false, true);
+														Controller->SetPinExpansion(ItemPins[ItemIndex]->GetPinPath(), true, true, true);
+													}
+
+													Controller->SetPinExpansion(ItemsPin->GetPinPath(), true, true, true);
+												}
+											}
+
+											Controller->CloseUndoBracket();
+										}
+									}
 								)
-                            );
-                        	SetElementsEntry.InsertPosition.Name = GetElementsEntry.Name;
-							SetElementsEntry.InsertPosition.Position = EToolMenuInsertType::After;	
-
-                        	Section.AddEntry(SetElementsEntry);
-
-                        	if (((DraggedTypes & (uint8)ERigElementType::Bone) != 0) ||
-								((DraggedTypes & (uint8)ERigElementType::Control) != 0) ||
-								((DraggedTypes & (uint8)ERigElementType::Null) != 0))
-                        	{
-                        		FToolMenuEntry& RotationTranslationSeparator = Section.AddSeparator(TEXT("RotationTranslationSeparator"));
-                        		RotationTranslationSeparator.InsertPosition.Name = SetElementsEntry.Name;
-                        		RotationTranslationSeparator.InsertPosition.Position = EToolMenuInsertType::After;
-
-								FToolMenuEntry SetRotationEntry = FToolMenuEntry::InitMenuEntry(
-									TEXT("SetRotation"),
-									LOCTEXT("SetRotation","Set Rotation"),
-									LOCTEXT("SetRotation_ToolTip","Setter for Rotation"),
-									FSlateIcon(),
-									FUIAction(
-										FExecuteAction::CreateSP(ControlRigEditor.Get(), &FControlRigEditor::HandleMakeElementGetterSetter, ERigElementGetterSetterType_Rotation, false, DraggedKeys, Graph, NodePosition),
-										FCanExecuteAction()
-									)
-								);
-                        		
-                        		SetRotationEntry.InsertPosition.Name = RotationTranslationSeparator.Name;
-								SetRotationEntry.InsertPosition.Position = EToolMenuInsertType::After;		
-								Section.AddEntry(SetRotationEntry);
-
-                        		FToolMenuEntry SetTranslationEntry = FToolMenuEntry::InitMenuEntry(
-									TEXT("SetTranslation"),
-									LOCTEXT("SetTranslation","Set Translation"),
-									LOCTEXT("SetTranslation_ToolTip","Setter for Translation"),
-									FSlateIcon(),
-									FUIAction(
-										FExecuteAction::CreateSP(ControlRigEditor.Get(), &FControlRigEditor::HandleMakeElementGetterSetter, ERigElementGetterSetterType_Translation, false, DraggedKeys, Graph, NodePosition),
-										FCanExecuteAction()
-									)
-								);
-
-                        		SetTranslationEntry.InsertPosition.Name = SetRotationEntry.Name;
-								SetTranslationEntry.InsertPosition.Position = EToolMenuInsertType::After;		
-								Section.AddEntry(SetTranslationEntry);
-
-                        		FToolMenuEntry AddOffsetEntry = FToolMenuEntry::InitMenuEntry(
-									TEXT("AddOffset"),
-									LOCTEXT("AddOffset","Add Offset"),
-									LOCTEXT("AddOffset_ToolTip","Setter for Offset"),
-									FSlateIcon(),
-									FUIAction(
-										FExecuteAction::CreateSP(ControlRigEditor.Get(), &FControlRigEditor::HandleMakeElementGetterSetter, ERigElementGetterSetterType_Offset, false, DraggedKeys, Graph, NodePosition),
-										FCanExecuteAction()
-									)
-								);
-                        		
-								AddOffsetEntry.InsertPosition.Name = SetTranslationEntry.Name;
-                                AddOffsetEntry.InsertPosition.Position = EToolMenuInsertType::After;						
-								Section.AddEntry(AddOffsetEntry);
-
-                        		FToolMenuEntry& RelativeTransformSeparator = Section.AddSeparator(TEXT("RelativeTransformSeparator"));
-                        		RelativeTransformSeparator.InsertPosition.Name = AddOffsetEntry.Name;
-                        		RelativeTransformSeparator.InsertPosition.Position = EToolMenuInsertType::After;
-                        		
-                        		FToolMenuEntry GetRelativeTransformEntry = FToolMenuEntry::InitMenuEntry(
-									TEXT("GetRelativeTransformEntry"),
-									LOCTEXT("GetRelativeTransform", "Get Relative Transform"),
-									LOCTEXT("GetRelativeTransform_ToolTip", "Getter for Relative Transform"),
-									FSlateIcon(),
-									FUIAction(
-										FExecuteAction::CreateSP(ControlRigEditor.Get(), &FControlRigEditor::HandleMakeElementGetterSetter, ERigElementGetterSetterType_Relative, true, DraggedKeys, Graph, NodePosition),
-										FCanExecuteAction()
-									)
-								);
-                        		
-								GetRelativeTransformEntry.InsertPosition.Name = RelativeTransformSeparator.Name;
-                                GetRelativeTransformEntry.InsertPosition.Position = EToolMenuInsertType::After;	
-								Section.AddEntry(GetRelativeTransformEntry);
-                        		
-                        	    FToolMenuEntry SetRelativeTransformEntry = FToolMenuEntry::InitMenuEntry(
-									TEXT("SetRelativeTransformEntry"),
-									LOCTEXT("SetRelativeTransform", "Set Relative Transform"),
-									LOCTEXT("SetRelativeTransform_ToolTip", "Setter for Relative Transform"),
-									FSlateIcon(),
-									FUIAction(
-										FExecuteAction::CreateSP(ControlRigEditor.Get(), &FControlRigEditor::HandleMakeElementGetterSetter, ERigElementGetterSetterType_Relative, false, DraggedKeys, Graph, NodePosition),
-										FCanExecuteAction()
-									)
-								);
-                        		SetRelativeTransformEntry.InsertPosition.Name = GetRelativeTransformEntry.Name;
-								SetRelativeTransformEntry.InsertPosition.Position = EToolMenuInsertType::After;		
-                        		Section.AddEntry(SetRelativeTransformEntry);
-                        	}
-
-                        	if (DraggedKeys.Num() > 0 && Blueprint != nullptr)
-                        	{
-                                FToolMenuEntry& ItemArraySeparator = Section.AddSeparator(TEXT("ItemArraySeparator"));
-                        		ItemArraySeparator.InsertPosition.Name = TEXT("SetRelativeTransformEntry"),
-								ItemArraySeparator.InsertPosition.Position = EToolMenuInsertType::After;
-                        		
-                        		FToolMenuEntry CreateItemArrayEntry = FToolMenuEntry::InitMenuEntry(
-                                    TEXT("CreateItemArray"),
-                                    LOCTEXT("CreateItemArray", "Create Item Array"),
-                                    LOCTEXT("CreateItemArray_ToolTip", "Creates an item array from the selected elements in the hierarchy"),
-                                    FSlateIcon(),
-                                    FUIAction(
-                                        FExecuteAction::CreateLambda([ControlRigEditor, Blueprint, DraggedKeys, NodePosition]()
-                                            {
-                                                if (URigVMController* Controller = ControlRigEditor->GetFocusedController())
-                                                {
-                                                    Controller->OpenUndoBracket(TEXT("Create Item Array From Selection"));
-
-                                                    static const FString ItemArrayCPPType = FString::Printf(TEXT("TArray<%s>"), *FRigElementKey::StaticStruct()->GetStructCPPName());
-                                                    static const FString ItemArrayObjectPath = FRigElementKey::StaticStruct()->GetPathName();
-									
-                                                    if (URigVMNode* ItemsNode = Controller->AddFreeRerouteNode(true, ItemArrayCPPType, *ItemArrayObjectPath, false, NAME_None, FString(), NodePosition))
-                                                    {
-                                                        if (URigVMPin* ItemsPin = ItemsNode->FindPin(TEXT("Value")))
-                                                        {
-                                                            Controller->SetArrayPinSize(ItemsPin->GetPinPath(), DraggedKeys.Num());
-
-                                                            TArray<URigVMPin*> ItemPins = ItemsPin->GetSubPins();
-                                                            ensure(ItemPins.Num() == DraggedKeys.Num());
-
-                                                            for (int32 ItemIndex = 0; ItemIndex < DraggedKeys.Num(); ItemIndex++)
-                                                            {
-                                                                FString DefaultValue;
-                                                                FRigElementKey::StaticStruct()->ExportText(DefaultValue, &DraggedKeys[ItemIndex], nullptr, nullptr, PPF_None, nullptr);
-                                                                Controller->SetPinDefaultValue(ItemPins[ItemIndex]->GetPinPath(), DefaultValue, true, true, false, true);
-                                                                Controller->SetPinExpansion(ItemPins[ItemIndex]->GetPinPath(), true, true, true);
-                                                            }
-
-                                                            Controller->SetPinExpansion(ItemsPin->GetPinPath(), true, true, true);
-                                                        }
-                                                    }
-
-                                                    Controller->CloseUndoBracket();
-                                                }
-                                            }
-                                        )
-                                    )
-                        		);
-                        		
-                        		CreateItemArrayEntry.InsertPosition.Name = ItemArraySeparator.Name,
-								CreateItemArrayEntry.InsertPosition.Position = EToolMenuInsertType::After;
-                        		Section.AddEntry(CreateItemArrayEntry);
-                            }
-                        }
-                    })
-                );
-			}
+							)
+						);
+						
+						CreateItemArrayEntry.InsertPosition.Name = ItemArraySeparator.Name,
+						CreateItemArrayEntry.InsertPosition.Position = EToolMenuInsertType::After;
+						Section.AddEntry(CreateItemArrayEntry);
+					}
+				}
+			})
+		);
+	}
+}
+
+void FControlRigEditor::OnGraphNodeDropToPerform(TSharedPtr<FGraphNodeDragDropOp> InDragDropOp, UEdGraph* InGraph, const FVector2D& InNodePosition, const FVector2D& InScreenPosition)
+{
+	DECLARE_SCOPE_HIERARCHICAL_COUNTER_FUNC()
+
+	if (InDragDropOp->IsOfType<FRigElementHierarchyDragDropOp>())
+	{
+		TSharedPtr<FRigElementHierarchyDragDropOp> RigHierarchyOp = StaticCastSharedPtr<FRigElementHierarchyDragDropOp>(InDragDropOp);
 
+		if (RigHierarchyOp->GetElements().Num() > 0 && FocusedGraphEdPtr.IsValid())
+		{
+			const FName MenuName = RigHierarchyToGraphDragAndDropMenuName;
+			
 			UControlRigContextMenuContext* MenuContext = NewObject<UControlRigContextMenuContext>();
 			FControlRigMenuSpecificContext MenuSpecificContext;
 			MenuSpecificContext.RigHierarchyToGraphDragAndDropContext =
@@ -5447,7 +5454,9 @@ void FControlRigEditor::OnGraphNodeDropToPerform(TSharedPtr<FGraphNodeDragDropOp
 					InGraph,
 					InNodePosition
 				);
-			MenuContext->Init(Cast<UControlRigBlueprint>(GetBlueprintObj()), MenuSpecificContext);
+			MenuContext->Init(SharedThis(this), MenuSpecificContext);
+			
+			UToolMenus* ToolMenus = UToolMenus::Get();
 			TSharedRef<SWidget>	MenuWidget = ToolMenus->GenerateWidget(MenuName, FToolMenuContext(MenuContext));
 			
 			TSharedRef<SWidget> GraphEditorPanel = FocusedGraphEdPtr.Pin().ToSharedRef();
diff --git a/Engine/Plugins/Animation/ControlRig/Source/ControlRigEditor/Private/Editor/ControlRigEditor.h b/Engine/Plugins/Animation/ControlRig/Source/ControlRigEditor/Private/Editor/ControlRigEditor.h
--- a/Engine/Plugins/Animation/ControlRig/Source/ControlRigEditor/Private/Editor/ControlRigEditor.h
+++ b/Engine/Plugins/Animation/ControlRig/Source/ControlRigEditor/Private/Editor/ControlRigEditor.h
@@ -189,6 +189,8 @@ public:
 	void OnHierarchyModified(ERigHierarchyNotification InNotif, URigHierarchy* InHierarchy, const FRigBaseElement* InElement);
 	void OnHierarchyModified_AnyThread(ERigHierarchyNotification InNotif, URigHierarchy* InHierarchy, const FRigBaseElement* InElement);
 
+	const FName RigHierarchyToGraphDragAndDropMenuName = TEXT("ControlRigEditor.RigHierarchyToGraphDragAndDropMenu");
+	void CreateRigHierarchyToGraphDragAndDropMenu() const;
 	void OnGraphNodeDropToPerform(TSharedPtr<FGraphNodeDragDropOp> InDragDropOp, UEdGraph* InGraph, const FVector2D& InNodePosition, const FVector2D& InScreenPosition);
 
 	FPersonaViewportKeyDownDelegate& GetKeyDownDelegate() { return OnKeyDownDelegate; }
diff --git a/Engine/Plugins/Animation/ControlRig/Source/ControlRigEditor/Private/Editor/SRigHierarchy.cpp b/Engine/Plugins/Animation/ControlRig/Source/ControlRigEditor/Private/Editor/SRigHierarchy.cpp
--- a/Engine/Plugins/Animation/ControlRig/Source/ControlRigEditor/Private/Editor/SRigHierarchy.cpp
+++ b/Engine/Plugins/Animation/ControlRig/Source/ControlRigEditor/Private/Editor/SRigHierarchy.cpp
@@ -925,7 +925,7 @@ void SRigHierarchy::OnSetExpansionRecursive(TSharedPtr<FRigTreeElement> InItem,
 	TreeView->SetExpansionRecursive(InItem, false, bShouldBeExpanded);
 }
 
-void SRigHierarchy::CreateDragDropMenu()
+void SRigHierarchy::CreateDragDropMenu() const
 {
 	const FName MenuName = DragDropMenuName;
 	UToolMenus* ToolMenus = UToolMenus::Get();
@@ -935,60 +935,62 @@ void SRigHierarchy::CreateDragDropMenu()
 		return;
 	}
 
-	UToolMenu* Menu	 = nullptr;
 	if (!ToolMenus->IsMenuRegistered(MenuName))
 	{
-		Menu = ToolMenus->RegisterMenu(MenuName);
-	}
-	else
-	{
-		Menu = ToolMenus->FindMenu(MenuName);
-	}
-
-	if (ensure(Menu))
-	{
-		FToolMenuEntry ParentEntry = FToolMenuEntry::InitMenuEntry(
-			TEXT("Parent"),
-			LOCTEXT("DragDropMenu_Parent", "Parent"),
-			LOCTEXT("DragDropMenu_Parent_ToolTip", "Parent Selected Items to the Target Item"),
-			FSlateIcon(),
-			FToolMenuExecuteAction::CreateSP(this, &SRigHierarchy::HandleParent)
-		);
-
-		ParentEntry.InsertPosition.Position = EToolMenuInsertType::First;
-		Menu->AddMenuEntry(NAME_None, ParentEntry);
+		UToolMenu* Menu = ToolMenus->RegisterMenu(MenuName);
 
-		UToolMenu* AlignMenu = Menu->AddSubMenu(
-			ToolMenus->CurrentOwner(),
-			NAME_None,
-			TEXT("Align"),
-			LOCTEXT("DragDropMenu_Align", "Align"),
-			LOCTEXT("DragDropMenu_Align_ToolTip", "Align Selected Items' Transforms to Target Item's Transform")
-		);
-
-		if (FToolMenuSection* DefaultSection = Menu->FindSection(NAME_None))
-		{
-			if (FToolMenuEntry* AlignMenuEntry = DefaultSection->FindEntry(TEXT("Align")))
+		Menu->AddDynamicSection(NAME_None, FNewToolMenuDelegate::CreateLambda([](UToolMenu* InMenu)
 			{
-				AlignMenuEntry->InsertPosition.Name = ParentEntry.Name;
-				AlignMenuEntry->InsertPosition.Position = EToolMenuInsertType::After;
-			}
-		}
-
-		FToolMenuEntry AlignAllEntry = FToolMenuEntry::InitMenuEntry(
-			TEXT("All"),
-			LOCTEXT("DragDropMenu_Align_All", "All"),
-			LOCTEXT("DragDropMenu_Align_All_ToolTip", "Align Selected Items' Transforms to Target Item's Transform"),
-			FSlateIcon(),
-			FToolMenuExecuteAction::CreateSP(this, &SRigHierarchy::HandleAlign)
+				UToolMenus* ToolMenus = UToolMenus::Get();
+				UControlRigContextMenuContext* MainContext = InMenu->FindContext<UControlRigContextMenuContext>();
+				
+				if (SRigHierarchy* RigHierarchyPanel = MainContext->GetRigHierarchyPanel())
+				{
+					FToolMenuEntry ParentEntry = FToolMenuEntry::InitMenuEntry(
+                        TEXT("Parent"),
+                        LOCTEXT("DragDropMenu_Parent", "Parent"),
+                        LOCTEXT("DragDropMenu_Parent_ToolTip", "Parent Selected Items to the Target Item"),
+                        FSlateIcon(),
+                        FToolMenuExecuteAction::CreateSP(RigHierarchyPanel, &SRigHierarchy::HandleParent)
+                    );
+		
+                    ParentEntry.InsertPosition.Position = EToolMenuInsertType::First;
+                    InMenu->AddMenuEntry(NAME_None, ParentEntry);
+		
+                    UToolMenu* AlignMenu = InMenu->AddSubMenu(
+                        ToolMenus->CurrentOwner(),
+                        NAME_None,
+                        TEXT("Align"),
+                        LOCTEXT("DragDropMenu_Align", "Align"),
+                        LOCTEXT("DragDropMenu_Align_ToolTip", "Align Selected Items' Transforms to Target Item's Transform")
+                    );
+		
+                    if (FToolMenuSection* DefaultSection = InMenu->FindSection(NAME_None))
+                    {
+                        if (FToolMenuEntry* AlignMenuEntry = DefaultSection->FindEntry(TEXT("Align")))
+                        {
+                            AlignMenuEntry->InsertPosition.Name = ParentEntry.Name;
+                            AlignMenuEntry->InsertPosition.Position = EToolMenuInsertType::After;
+                        }
+                    }
+		
+                    FToolMenuEntry AlignAllEntry = FToolMenuEntry::InitMenuEntry(
+                        TEXT("All"),
+                        LOCTEXT("DragDropMenu_Align_All", "All"),
+                        LOCTEXT("DragDropMenu_Align_All_ToolTip", "Align Selected Items' Transforms to Target Item's Transform"),
+                        FSlateIcon(),
+                        FToolMenuExecuteAction::CreateSP(RigHierarchyPanel, &SRigHierarchy::HandleAlign)
+                    );
+                    AlignAllEntry.InsertPosition.Position = EToolMenuInsertType::First;
+		
+                    AlignMenu->AddMenuEntry(NAME_None, AlignAllEntry);	
+				}
+			})
 		);
-		AlignAllEntry.InsertPosition.Position = EToolMenuInsertType::First;
-
-		AlignMenu->AddMenuEntry(NAME_None, AlignAllEntry);
 	}
 }
 
-UToolMenu* SRigHierarchy::GetDragDropMenu(const TArray<FRigElementKey>& DraggedKeys, FRigElementKey TargetKey) const
+UToolMenu* SRigHierarchy::GetDragDropMenu(const TArray<FRigElementKey>& DraggedKeys, FRigElementKey TargetKey)
 {
 	UToolMenus* ToolMenus = UToolMenus::Get();
 	if (!ensure(ToolMenus))
@@ -1000,17 +1002,17 @@ UToolMenu* SRigHierarchy::GetDragDropMenu(const TArray<FRigElementKey>& DraggedK
 	UControlRigContextMenuContext* MenuContext = NewObject<UControlRigContextMenuContext>();
 	FControlRigMenuSpecificContext MenuSpecificContext;
 	MenuSpecificContext.RigHierarchyDragAndDropContext = FControlRigRigHierarchyDragAndDropContext(DraggedKeys, TargetKey);
-	MenuContext->Init(ControlRigBlueprint, MenuSpecificContext);
+	MenuSpecificContext.RigHierarchyPanel = SharedThis(this);
+	MenuContext->Init(ControlRigEditor, MenuSpecificContext);
 	
 	UToolMenu* Menu = ToolMenus->GenerateMenu(MenuName, FToolMenuContext(MenuContext));
 
 	return Menu;
 }
 
-void SRigHierarchy::CreateContextMenu()
+void SRigHierarchy::CreateContextMenu() const
 {
 	const FName MenuName = ContextMenuName;
-	const FName InteractionSectionName = TEXT("Interaction");
 
 	UToolMenus* ToolMenus = UToolMenus::Get();
 	
@@ -1019,60 +1021,47 @@ void SRigHierarchy::CreateContextMenu()
 		return;
 	}
 	
-	const FControlRigHierarchyCommands& Commands = FControlRigHierarchyCommands::Get();
-	
-	UToolMenu* Menu = nullptr;
-	
 	if (!ToolMenus->IsMenuRegistered(MenuName))
 	{
-		Menu = ToolMenus->RegisterMenu(MenuName);
-	}
-	else
-	{
-		Menu = ToolMenus->FindMenu(MenuName);
-	}
-
-	if (ensure(Menu))
-	{
-		struct FLocalMenuBuilder
-		{
-			static void FillNewMenu(FMenuBuilder& InSubMenuBuilder, TSharedPtr<SRigHierarchyTreeView> InTreeView)
+		UToolMenu* Menu = ToolMenus->RegisterMenu(MenuName);
+		
+		Menu->AddDynamicSection(NAME_None, FNewToolMenuDelegate::CreateLambda([](UToolMenu* InMenu)
 			{
-				const FControlRigHierarchyCommands& Actions = FControlRigHierarchyCommands::Get();
-
-				FRigElementKey SelectedKey;
-				TArray<TSharedPtr<FRigTreeElement>> SelectedItems = InTreeView->GetSelectedItems();
-				if (SelectedItems.Num() > 0)
+				UControlRigContextMenuContext* MainContext = InMenu->FindContext<UControlRigContextMenuContext>();
+				
+				if (SRigHierarchy* RigHierarchyPanel = MainContext->GetRigHierarchyPanel())
 				{
-					SelectedKey = SelectedItems[0]->Key;
-				}
+					const FControlRigHierarchyCommands& Commands = FControlRigHierarchyCommands::Get(); 
+				
+					FToolMenuSection& ElementsSection = InMenu->AddSection(TEXT("Elements"), LOCTEXT("ElementsHeader", "Elements"));
+					ElementsSection.AddSubMenu(TEXT("New"), LOCTEXT("New", "New"), LOCTEXT("New_ToolTip", "Create New Elements"),
+						FNewToolMenuDelegate::CreateLambda([Commands, RigHierarchyPanel](UToolMenu* InSubMenu)
+						{
+							FToolMenuSection& DefaultSection = InSubMenu->AddSection(NAME_None);
+							FRigElementKey SelectedKey;
+							TArray<TSharedPtr<FRigTreeElement>> SelectedItems = RigHierarchyPanel->TreeView->GetSelectedItems();
+							if (SelectedItems.Num() > 0)
+							{
+								SelectedKey = SelectedItems[0]->Key;
+							}
+							
+							if (!SelectedKey || SelectedKey.Type == ERigElementType::Bone)
+							{
+								DefaultSection.AddMenuEntry(Commands.AddBoneItem);
+							}
+							DefaultSection.AddMenuEntry(Commands.AddControlItem);
+							DefaultSection.AddMenuEntry(Commands.AddNullItem);
+						})
+					);
+					
+					ElementsSection.AddMenuEntry(Commands.DeleteItem);
+					ElementsSection.AddMenuEntry(Commands.DuplicateItem);
+					ElementsSection.AddMenuEntry(Commands.RenameItem);
+					ElementsSection.AddMenuEntry(Commands.MirrorItem);
 
-				if (!SelectedKey || SelectedKey.Type == ERigElementType::Bone)
-				{
-					InSubMenuBuilder.AddMenuEntry(Actions.AddBoneItem);
-				}
-				InSubMenuBuilder.AddMenuEntry(Actions.AddControlItem);
-				InSubMenuBuilder.AddMenuEntry(Actions.AddNullItem);
-			}
-		};
-
-		FToolMenuSection& ElementsSection = Menu->AddSection(TEXT("Elements"), LOCTEXT("ElementsHeader", "Elements"));
-		ElementsSection.AddSubMenu(TEXT("New"), LOCTEXT("New", "New"), LOCTEXT("New_ToolTip", "Create New Elements"),
-								   FNewMenuDelegate::CreateStatic(&FLocalMenuBuilder::FillNewMenu, TreeView));
-		ElementsSection.AddMenuEntry(Commands.DeleteItem);
-		ElementsSection.AddMenuEntry(Commands.DuplicateItem);
-		ElementsSection.AddMenuEntry(Commands.RenameItem);
-		ElementsSection.AddMenuEntry(Commands.MirrorItem);
-
-		// dynamic section is used here so that the whole section can hide when a condition is not met
-		Menu->AddDynamicSection(InteractionSectionName,
-			FNewToolMenuDelegate::CreateLambda([InteractionSectionName, Commands, WeakThis = TWeakPtr<SRigHierarchy>(SharedThis(this)) ](UToolMenu* InMenu)
-			{
-				if (const TSharedPtr<SRigHierarchy> RigHierarchyPanel = WeakThis.Pin())
-				{
 					if (RigHierarchyPanel->IsSingleBoneSelected() || RigHierarchyPanel->IsControlSelected())
 					{
-						FToolMenuSection& InteractionSection = InMenu->AddSection(InteractionSectionName, LOCTEXT("InteractionHeader", "Interaction"));
+						FToolMenuSection& InteractionSection = InMenu->AddSection(TEXT("Interaction"), LOCTEXT("InteractionHeader", "Interaction"));
 						if(RigHierarchyPanel->IsSingleBoneSelected())
 						{
 							InteractionSection.AddMenuEntry(Commands.ControlBoneTransform);
@@ -1082,39 +1071,38 @@ void SRigHierarchy::CreateContextMenu()
 							InteractionSection.AddMenuEntry(Commands.TestSpaceSwitching);
 						}
 					}
+					
+					FToolMenuSection& CopyPasteSection = InMenu->AddSection(TEXT("Copy&Paste"), LOCTEXT("Copy&PasteHeader", "Copy & Paste"));
+					CopyPasteSection.AddMenuEntry(Commands.CopyItems);
+					CopyPasteSection.AddMenuEntry(Commands.PasteItems);
+					CopyPasteSection.AddMenuEntry(Commands.PasteLocalTransforms);
+					CopyPasteSection.AddMenuEntry(Commands.PasteGlobalTransforms);
+					
+					FToolMenuSection& TransformsSection = InMenu->AddSection(TEXT("Transforms"), LOCTEXT("TransformsHeader", "Transforms"));
+					TransformsSection.AddMenuEntry(Commands.ResetTransform);
+					TransformsSection.AddMenuEntry(Commands.ResetAllTransforms);
+					TransformsSection.AddMenuEntry(Commands.SetInitialTransformFromCurrentTransform);
+					TransformsSection.AddMenuEntry(Commands.SetInitialTransformFromClosestBone);
+					TransformsSection.AddMenuEntry(Commands.SetGizmoTransformFromCurrent);
+					TransformsSection.AddMenuEntry(Commands.Unparent);
+
+					FToolMenuSection& AssetsSection = InMenu->AddSection(TEXT("Assets"), LOCTEXT("AssetsHeader", "Assets"));
+					AssetsSection.AddSubMenu(TEXT("Import"), LOCTEXT("ImportSubMenu", "Import"),
+						LOCTEXT("ImportSubMenu_ToolTip", "Import hierarchy to the current rig. This only imports non-existing node. For example, if there is hand_r, it won't import hand_r. If you want to reimport whole new hiearchy, delete all nodes, and use import hierarchy."),
+						FNewMenuDelegate::CreateSP(RigHierarchyPanel, &SRigHierarchy::CreateImportMenu)
+					);
+					
+					AssetsSection.AddSubMenu(TEXT("Refresh"), LOCTEXT("RefreshSubMenu", "Refresh"),
+						LOCTEXT("RefreshSubMenu_ToolTip", "Refresh the existing initial transform from the selected mesh. This only updates if the node is found."),
+						FNewMenuDelegate::CreateSP(RigHierarchyPanel, &SRigHierarchy::CreateRefreshMenu)
+					);	
 				}
-			}),
-			FToolMenuInsert(TEXT("Elements"), EToolMenuInsertType::After)
-		);
-		
-		FToolMenuSection& CopyPasteSection = Menu->AddSection(TEXT("Copy&Paste"), LOCTEXT("Copy&PasteHeader", "Copy & Paste"));
-		CopyPasteSection.AddMenuEntry(Commands.CopyItems);
-		CopyPasteSection.AddMenuEntry(Commands.PasteItems);
-		CopyPasteSection.AddMenuEntry(Commands.PasteLocalTransforms);
-		CopyPasteSection.AddMenuEntry(Commands.PasteGlobalTransforms);
-		
-		FToolMenuSection& TransformsSection = Menu->AddSection(TEXT("Transforms"), LOCTEXT("TransformsHeader", "Transforms"));
-		TransformsSection.AddMenuEntry(Commands.ResetTransform);
-		TransformsSection.AddMenuEntry(Commands.ResetAllTransforms);
-		TransformsSection.AddMenuEntry(Commands.SetInitialTransformFromCurrentTransform);
-		TransformsSection.AddMenuEntry(Commands.SetInitialTransformFromClosestBone);
-		TransformsSection.AddMenuEntry(Commands.SetGizmoTransformFromCurrent);
-		TransformsSection.AddMenuEntry(Commands.Unparent);
-
-		FToolMenuSection& AssetsSection = Menu->AddSection(TEXT("Assets"), LOCTEXT("AssetsHeader", "Assets"));
-		AssetsSection.AddSubMenu(TEXT("Import"), LOCTEXT("ImportSubMenu", "Import"),
-			LOCTEXT("ImportSubMenu_ToolTip", "Import hierarchy to the current rig. This only imports non-existing node. For example, if there is hand_r, it won't import hand_r. If you want to reimport whole new hiearchy, delete all nodes, and use import hierarchy."),
-			FNewMenuDelegate::CreateSP(this, &SRigHierarchy::CreateImportMenu)
-		);
-		
-		AssetsSection.AddSubMenu(TEXT("Refresh"), LOCTEXT("RefreshSubMenu", "Refresh"),
-			LOCTEXT("RefreshSubMenu_ToolTip", "Refresh the existing initial transform from the selected mesh. This only updates if the node is found."),
-			FNewMenuDelegate::CreateSP(this, &SRigHierarchy::CreateRefreshMenu)
+			})
 		);
 	}
 }
 
-UToolMenu* SRigHierarchy::GetContextMenu() const
+UToolMenu* SRigHierarchy::GetContextMenu()
 {
 	const FName MenuName = ContextMenuName;
 	UToolMenus* ToolMenus = UToolMenus::Get();
@@ -1126,7 +1114,9 @@ UToolMenu* SRigHierarchy::GetContextMenu() const
 
 	// individual entries in this menu can access members of this context, particularly useful for editor scripting
 	UControlRigContextMenuContext* ContextMenuContext = NewObject<UControlRigContextMenuContext>();
-	ContextMenuContext->Init(ControlRigBlueprint);
+	FControlRigMenuSpecificContext MenuSpecificContext;
+	MenuSpecificContext.RigHierarchyPanel = SharedThis(this);
+	ContextMenuContext->Init(ControlRigEditor, MenuSpecificContext);
 
 	FToolMenuContext MenuContext(CommandList);
 	MenuContext.AddObject(ContextMenuContext);
diff --git a/Engine/Plugins/Animation/ControlRig/Source/ControlRigEditor/Private/Editor/SRigHierarchy.h b/Engine/Plugins/Animation/ControlRig/Source/ControlRigEditor/Private/Editor/SRigHierarchy.h
--- a/Engine/Plugins/Animation/ControlRig/Source/ControlRigEditor/Private/Editor/SRigHierarchy.h
+++ b/Engine/Plugins/Animation/ControlRig/Source/ControlRigEditor/Private/Editor/SRigHierarchy.h
@@ -148,13 +148,13 @@ private:
 	FReply OnAcceptDrop(const FDragDropEvent& DragDropEvent, EItemDropZone DropZone, TSharedPtr<FRigTreeElement> TargetItem);
 
 	const FName ContextMenuName = TEXT("ControlRigEditor.RigHierarchy.ContextMenu");
-	void CreateContextMenu();
-	UToolMenu* GetContextMenu() const;
+	void CreateContextMenu() const;
+	UToolMenu* GetContextMenu();
 	TSharedPtr<FUICommandList> GetContextMenuCommands() const;
 	
 	const FName DragDropMenuName = TEXT("ControlRigEditor.RigHierarchy.DragDropMenu");
-	void CreateDragDropMenu();
-	UToolMenu* GetDragDropMenu(const TArray<FRigElementKey>& DraggedKeys, FRigElementKey TargetKey) const;
+	void CreateDragDropMenu() const;
+	UToolMenu* GetDragDropMenu(const TArray<FRigElementKey>& DraggedKeys, FRigElementKey TargetKey);
 
 	/** Our owning control rig editor */
 	TWeakPtr<FControlRigEditor> ControlRigEditor;
