Allow renaming of a single item with case change

we only support case change for the item curently edited to avoid having both ctrlA and ctrla in the hierarchy
note that both FString and FName operator== are case insensitive resulting in misunderstanding when renaming items

diff --git a/Engine/Plugins/Animation/ControlRig/Source/ControlRig/Private/Rigs/RigHierarchyController.cpp b/Engine/Plugins/Animation/ControlRig/Source/ControlRig/Private/Rigs/RigHierarchyController.cpp
--- a/Engine/Plugins/Animation/ControlRig/Source/ControlRig/Private/Rigs/RigHierarchyController.cpp
+++ b/Engine/Plugins/Animation/ControlRig/Source/ControlRig/Private/Rigs/RigHierarchyController.cpp
@@ -1762,7 +1762,7 @@ bool URigHierarchyController::RenameElement(FRigBaseElement* InElement, const FN
 		return false;
 	}
 
-	if(InElement->GetName() == InName)
+	if (InElement->GetName().IsEqual(InName, ENameCase::CaseSensitive))
 	{
 		return false;
 	}
@@ -1775,8 +1775,16 @@ bool URigHierarchyController::RenameElement(FRigBaseElement* InElement, const FN
 	{
 		DeselectElement(OldKey);
 	}
+
+	{
+		// create a temp copy of the map and remove the current item's key
+		TMap<FRigElementKey, int32> TemporaryMap = Hierarchy->IndexLookup;
+		TemporaryMap.Remove(OldKey);
+   
+		TGuardValue<TMap<FRigElementKey, int32>> MapGuard(Hierarchy->IndexLookup, TemporaryMap);
+		InElement->Key.Name = Hierarchy->GetSafeNewName(InName.ToString(), InElement->GetType());
+	}
 	
-	InElement->Key.Name = Hierarchy->GetSafeNewName(InName.ToString(), InElement->GetType());
 	const FRigElementKey NewKey = InElement->GetKey();
 
 	Hierarchy->IndexLookup.Remove(OldKey);
diff --git a/Engine/Plugins/Animation/ControlRig/Source/ControlRigEditor/Private/ControlRigElementDetails.cpp b/Engine/Plugins/Animation/ControlRig/Source/ControlRigEditor/Private/ControlRigElementDetails.cpp
--- a/Engine/Plugins/Animation/ControlRig/Source/ControlRigEditor/Private/ControlRigElementDetails.cpp
+++ b/Engine/Plugins/Animation/ControlRig/Source/ControlRigEditor/Private/ControlRigElementDetails.cpp
@@ -851,7 +851,7 @@ bool FRigBaseElementDetails::OnVerifyNameChanged(const FText& InText, FText& Out
 		return false;
 	}
 
-	if (*GetElementKey().Name.ToString() == InText.ToString())
+	if (GetElementKey().Name.ToString() == InText.ToString())
 	{
 		return true;
 	}
diff --git a/Engine/Plugins/Animation/ControlRig/Source/ControlRigEditor/Private/Editor/SRigHierarchy.cpp b/Engine/Plugins/Animation/ControlRig/Source/ControlRigEditor/Private/Editor/SRigHierarchy.cpp
--- a/Engine/Plugins/Animation/ControlRig/Source/ControlRigEditor/Private/Editor/SRigHierarchy.cpp
+++ b/Engine/Plugins/Animation/ControlRig/Source/ControlRigEditor/Private/Editor/SRigHierarchy.cpp
@@ -1996,7 +1996,7 @@ FName SRigHierarchy::HandleRenameElement(const FRigElementKey& OldKey, const FSt
 {
 	ClearDetailPanel();
 
-	if (OldKey.Name.ToString() == NewName)
+	if (OldKey.Name.ToString().Equals(NewName, ESearchCase::CaseSensitive))	
 	{
 		return NAME_None;
 	}
