Changed duplicate shortcut from Ctrl+W to Ctrl+D for consistency

diff --git a/Engine/Plugins/Animation/ControlRig/Source/ControlRigEditor/Private/ControlRigHierarchyCommands.cpp b/Engine/Plugins/Animation/ControlRig/Source/ControlRigEditor/Private/ControlRigHierarchyCommands.cpp
--- a/Engine/Plugins/Animation/ControlRig/Source/ControlRigEditor/Private/ControlRigHierarchyCommands.cpp
+++ b/Engine/Plugins/Animation/ControlRig/Source/ControlRigEditor/Private/ControlRigHierarchyCommands.cpp
@@ -9,7 +9,7 @@ void FControlRigHierarchyCommands::RegisterCommands()
 	UI_COMMAND(AddBoneItem, "New Bone", "Add new bone at the origin (0, 0, 0) to the hierarchy.", EUserInterfaceActionType::Button, FInputChord());
 	UI_COMMAND(AddControlItem, "New Control", "Add new control at the origin (0, 0, 0) to the hierarchy.", EUserInterfaceActionType::Button, FInputChord(EKeys::N, EModifierKey::Control));
 	UI_COMMAND(AddNullItem, "New Null", "Add new null at the origin (0, 0, 0) to the hierarchy.", EUserInterfaceActionType::Button, FInputChord());
-	UI_COMMAND(DuplicateItem, "Duplicate", "Duplicate the selected items in the hierarchy.", EUserInterfaceActionType::Button, FInputChord(EKeys::W, EModifierKey::Control));
+	UI_COMMAND(DuplicateItem, "Duplicate", "Duplicate the selected items in the hierarchy.", EUserInterfaceActionType::Button, FInputChord(EKeys::D, EModifierKey::Control));
 	UI_COMMAND(MirrorItem, "Mirror", "Mirror the selected items in the hierarchy.", EUserInterfaceActionType::Button, FInputChord());
 	UI_COMMAND(DeleteItem, "Delete", "Deletes the selected items from the hierarchy.", EUserInterfaceActionType::Button, FInputChord(EKeys::Delete));
 	UI_COMMAND(RenameItem, "Rename", "Rename the selected item.", EUserInterfaceActionType::Button, FInputChord(EKeys::F2));
