Prefix new control with bone's name

diff --git a/Engine/Plugins/Animation/ControlRig/Source/ControlRigEditor/Private/Editor/SRigHierarchy.cpp b/Engine/Plugins/Animation/ControlRig/Source/ControlRigEditor/Private/Editor/SRigHierarchy.cpp
--- a/Engine/Plugins/Animation/ControlRig/Source/ControlRigEditor/Private/Editor/SRigHierarchy.cpp
+++ b/Engine/Plugins/Animation/ControlRig/Source/ControlRigEditor/Private/Editor/SRigHierarchy.cpp
@@ -1442,7 +1442,7 @@ bool SRigHierarchy::CanDeleteItem() const
 	return IsMultiSelected();
 }
 
-/** Delete Item */
+/** Create Item */
 void SRigHierarchy::HandleNewItem(ERigElementType InElementType)
 {
 	if(!ControlRigEditor.IsValid())
@@ -1472,7 +1472,20 @@ void SRigHierarchy::HandleNewItem(ERigElementType InElementType)
 			ParentTransform = Hierarchy->GetGlobalTransform(ParentKey);
 		}
 
-		FString NewNameTemplate = FString::Printf(TEXT("New%s"), *StaticEnum<ERigElementType>()->GetNameStringByValue((int64)InElementType));
+		// use bone's name as prefix if creating a control
+		FString NewNameTemplate;
+		const bool bIsParentABone = ParentKey.IsValid() && ParentKey.Type == ERigElementType::Bone;
+		if( InElementType == ERigElementType::Control && bIsParentABone )
+		{
+			static const FString CtrlSuffix(TEXT("_ctrl"));
+			NewNameTemplate = ParentKey.Name.ToString();
+			NewNameTemplate += CtrlSuffix;
+		}
+		else
+		{
+			NewNameTemplate = FString::Printf(TEXT("New%s"), *StaticEnum<ERigElementType>()->GetNameStringByValue((int64)InElementType));
+		}
+		
 		const FName NewElementName = CreateUniqueName(*NewNameTemplate, InElementType);
 		{
 			TGuardValue<bool> GuardRigHierarchyChanges(bIsChangingRigHierarchy, true);
