Removed 'Output' on ControlRig details panel

diff --git a/Engine/Plugins/Animation/ControlRig/Source/ControlRigDeveloper/Private/AnimGraphNode_ControlRig.cpp b/Engine/Plugins/Animation/ControlRig/Source/ControlRigDeveloper/Private/AnimGraphNode_ControlRig.cpp
--- a/Engine/Plugins/Animation/ControlRig/Source/ControlRigDeveloper/Private/AnimGraphNode_ControlRig.cpp
+++ b/Engine/Plugins/Animation/ControlRig/Source/ControlRigDeveloper/Private/AnimGraphNode_ControlRig.cpp
@@ -474,22 +474,6 @@ void UAnimGraphNode_ControlRig::CustomizeDetails(IDetailLayoutBuilder& DetailBui
 		.OnPinIsEnabledCheckState(FOnPinIsCheckEnabled::CreateUObject(this, &UAnimGraphNode_ControlRig::IsPropertyExposeEnabled))
 	];
 
-	IDetailCategoryBuilder& OutputCategoryBuilder = DetailBuilder.EditCategory(FName(TEXT("Output")));
-
-	FDetailWidgetRow& OutputWidgetRow = OutputCategoryBuilder.AddCustomRow(FText::FromString(TEXT("Output")));
-	OutputWidgetRow.WholeRowContent()
-	[
-		SNew(SVariableMappingWidget)
-		.OnVariableMappingChanged(FOnVariableMappingChanged::CreateUObject(this, &UAnimGraphNode_ControlRig::OnVariableMappingChanged, false))
-		.OnGetVariableMapping(FOnGetVariableMapping::CreateUObject(this, &UAnimGraphNode_ControlRig::GetVariableMapping, false))
-		.OnGetAvailableMapping(FOnGetAvailableMapping::CreateUObject(this, &UAnimGraphNode_ControlRig::GetAvailableMapping, false))
-		.OnCreateVariableMapping(FOnCreateVariableMapping::CreateUObject(this, &UAnimGraphNode_ControlRig::CreateVariableMapping, false))
-		.OnVariableOptionAvailable(FOnVarOptionAvailable::CreateUObject(this, &UAnimGraphNode_ControlRig::IsAvailableToMapToCurve, false))
-		.OnPinGetCheckState(FOnPinGetCheckState::CreateUObject(this, &UAnimGraphNode_ControlRig::IsPropertyExposed))
-		.OnPinCheckStateChanged(FOnPinCheckStateChanged::CreateUObject(this, &UAnimGraphNode_ControlRig::OnPropertyExposeCheckboxChanged))
-		.OnPinIsEnabledCheckState(FOnPinIsCheckEnabled::CreateUObject(this, &UAnimGraphNode_ControlRig::IsPropertyExposeEnabled))
-	];
-
 	TSharedRef<IPropertyHandle> ClassHandle = DetailBuilder.GetProperty(TEXT("Node.ControlRigClass"), GetClass());
 	if (ClassHandle->IsValidHandle())
 	{
