Control Rig: Fix when creating a new RigElement in RigHierarchy, Control RIg Python Log creates a new page calleed "Hierarchy"

diff --git a/Engine/Plugins/Animation/ControlRig/Source/ControlRig/Private/Rigs/RigHierarchy.cpp b/Engine/Plugins/Animation/ControlRig/Source/ControlRig/Private/Rigs/RigHierarchy.cpp
--- a/Engine/Plugins/Animation/ControlRig/Source/ControlRig/Private/Rigs/RigHierarchy.cpp
+++ b/Engine/Plugins/Animation/ControlRig/Source/ControlRig/Private/Rigs/RigHierarchy.cpp
@@ -2112,30 +2112,34 @@ void URigHierarchy::SetTransform(FRigTransformElement* InTransformElement, const
 
 	if (bPrintPythonCommands)
 	{
-		FString MethodName;
-		switch (InTransformType)
+		UBlueprint* Blueprint = GetTypedOuter<UBlueprint>();
+		if (Blueprint)
 		{
-			case ERigTransformType::InitialLocal: 
-			case ERigTransformType::CurrentLocal:
+			FString MethodName;
+			switch (InTransformType)
 			{
-				MethodName = TEXT("set_local_transform");
-				break;
-			}
-			case ERigTransformType::InitialGlobal: 
-			case ERigTransformType::CurrentGlobal:
-			{
-				MethodName = TEXT("set_global_transform");
-				break;
+				case ERigTransformType::InitialLocal: 
+				case ERigTransformType::CurrentLocal:
+				{
+					MethodName = TEXT("set_local_transform");
+					break;
+				}
+				case ERigTransformType::InitialGlobal: 
+				case ERigTransformType::CurrentGlobal:
+				{
+					MethodName = TEXT("set_global_transform");
+					break;
+				}
 			}
-		}
 
-		RigVMPythonUtils::Print(GetOuter()->GetFName().ToString(),
-			FString::Printf(TEXT("hierarchy.%s(%s, %s, %s, %s)"),
-			*MethodName,
-			*InTransformElement->GetKey().ToPythonString(),
-			*RigVMPythonUtils::TransformToPythonString(InTransform),
-			(InTransformType == ERigTransformType::InitialGlobal || InTransformType == ERigTransformType::InitialLocal) ? TEXT("True") : TEXT("False"),
-			(bAffectChildren) ? TEXT("True") : TEXT("False")));
+			RigVMPythonUtils::Print(Blueprint->GetFName().ToString(),
+				FString::Printf(TEXT("hierarchy.%s(%s, %s, %s, %s)"),
+				*MethodName,
+				*InTransformElement->GetKey().ToPythonString(),
+				*RigVMPythonUtils::TransformToPythonString(InTransform),
+				(InTransformType == ERigTransformType::InitialGlobal || InTransformType == ERigTransformType::InitialLocal) ? TEXT("True") : TEXT("False"),
+				(bAffectChildren) ? TEXT("True") : TEXT("False")));
+		}
 	}
 #endif
 }
@@ -2243,13 +2247,16 @@ void URigHierarchy::SetControlOffsetTransform(FRigControlElement* InControlEleme
 
 	if (bPrintPythonCommands)
 	{
-		//FRigElementKey InKey, FTransform InTransform, bool bInitial = false, bool bAffectChildren = true, bool bSetupUndo = false, bool bPrintPythonCommands = false)
-		RigVMPythonUtils::Print(GetOuter()->GetFName().ToString(),
-			FString::Printf(TEXT("hierarchy.set_control_offset_transform(%s, %s, %s, %s)"),
-			*InControlElement->GetKey().ToPythonString(),
-			*RigVMPythonUtils::TransformToPythonString(InTransform),
-			(ERigTransformType::IsInitial(InTransformType)) ? TEXT("True") : TEXT("False"),
-			(bAffectChildren) ? TEXT("True") : TEXT("False")));
+		UBlueprint* Blueprint = GetTypedOuter<UBlueprint>();
+		if (Blueprint)
+		{
+			RigVMPythonUtils::Print(Blueprint->GetFName().ToString(),
+				FString::Printf(TEXT("hierarchy.set_control_offset_transform(%s, %s, %s, %s)"),
+				*InControlElement->GetKey().ToPythonString(),
+				*RigVMPythonUtils::TransformToPythonString(InTransform),
+				(ERigTransformType::IsInitial(InTransformType)) ? TEXT("True") : TEXT("False"),
+				(bAffectChildren) ? TEXT("True") : TEXT("False")));
+		}
 	}
 #endif
 }
diff --git a/Engine/Plugins/Animation/ControlRig/Source/ControlRig/Private/Rigs/RigHierarchyController.cpp b/Engine/Plugins/Animation/ControlRig/Source/ControlRig/Private/Rigs/RigHierarchyController.cpp
--- a/Engine/Plugins/Animation/ControlRig/Source/ControlRig/Private/Rigs/RigHierarchyController.cpp
+++ b/Engine/Plugins/Animation/ControlRig/Source/ControlRig/Private/Rigs/RigHierarchyController.cpp
@@ -181,11 +181,15 @@ FRigElementKey URigHierarchyController::AddBone(FName InName, FRigElementKey InP
 
 	if (bPrintPythonCommand)
 	{
-		TArray<FString> Commands = GetAddBonePythonCommands(NewElement);
-		for (const FString& Command : Commands)
-		{
-			RigVMPythonUtils::Print(GetOuter()->GetFName().ToString(),
-				FString::Printf(TEXT("%s"), *Command));
+		UBlueprint* Blueprint = GetTypedOuter<UBlueprint>();
+		if (Blueprint)
+		{
+			TArray<FString> Commands = GetAddBonePythonCommands(NewElement);
+			for (const FString& Command : Commands)
+			{			
+				RigVMPythonUtils::Print(Blueprint->GetFName().ToString(),
+					FString::Printf(TEXT("%s"), *Command));
+			}
 		}
 	}
 #endif
@@ -233,11 +237,15 @@ FRigElementKey URigHierarchyController::AddNull(FName InName, FRigElementKey InP
 
 	if (bPrintPythonCommand)
 	{
-		TArray<FString> Commands = GetAddNullPythonCommands(NewElement);
-		for (const FString& Command : Commands)
+		UBlueprint* Blueprint = GetTypedOuter<UBlueprint>();
+		if (Blueprint)
 		{
-			RigVMPythonUtils::Print(GetOuter()->GetFName().ToString(),
-				FString::Printf(TEXT("%s"), *Command));
+			TArray<FString> Commands = GetAddNullPythonCommands(NewElement);
+			for (const FString& Command : Commands)
+			{
+				RigVMPythonUtils::Print(Blueprint->GetFName().ToString(),
+					FString::Printf(TEXT("%s"), *Command));
+			}
 		}
 	}
 #endif
@@ -292,11 +300,15 @@ FRigElementKey URigHierarchyController::AddControl(
 
 	if (bPrintPythonCommand)
 	{
-		TArray<FString> Commands = GetAddControlPythonCommands(NewElement);
-		for (const FString& Command : Commands)
+		UBlueprint* Blueprint = GetTypedOuter<UBlueprint>();
+		if (Blueprint)
 		{
-			RigVMPythonUtils::Print(GetOuter()->GetFName().ToString(),
-				FString::Printf(TEXT("%s"), *Command));
+			TArray<FString> Commands = GetAddControlPythonCommands(NewElement);
+			for (const FString& Command : Commands)
+			{
+				RigVMPythonUtils::Print(Blueprint->GetFName().ToString(),
+					FString::Printf(TEXT("%s"), *Command));
+			}
 		}
 	}
 #endif
@@ -331,11 +343,15 @@ FRigElementKey URigHierarchyController::AddCurve(FName InName, float InValue, bo
 
 	if (bPrintPythonCommand)
 	{
-		TArray<FString> Commands = GetAddCurvePythonCommands(NewElement);
-		for (const FString& Command : Commands)
+		UBlueprint* Blueprint = GetTypedOuter<UBlueprint>();
+		if (Blueprint)
 		{
-			RigVMPythonUtils::Print(GetOuter()->GetFName().ToString(),
-				FString::Printf(TEXT("%s"), *Command));
+			TArray<FString> Commands = GetAddCurvePythonCommands(NewElement);
+			for (const FString& Command : Commands)
+			{
+				RigVMPythonUtils::Print(Blueprint->GetFName().ToString(),
+					FString::Printf(TEXT("%s"), *Command));
+			}
 		}
 	}
 #endif
@@ -374,11 +390,15 @@ FRigElementKey URigHierarchyController::AddRigidBody(FName InName, FRigElementKe
 
 	if (bPrintPythonCommand)
 	{
-		TArray<FString> Commands = GetAddRigidBodyPythonCommands(NewElement);
-		for (const FString& Command : Commands)
+		UBlueprint* Blueprint = GetTypedOuter<UBlueprint>();
+		if (Blueprint)
 		{
-			RigVMPythonUtils::Print(GetOuter()->GetFName().ToString(),
-				FString::Printf(TEXT("%s"), *Command));
+			TArray<FString> Commands = GetAddRigidBodyPythonCommands(NewElement);
+			for (const FString& Command : Commands)
+			{
+				RigVMPythonUtils::Print(Blueprint->GetFName().ToString(),
+					FString::Printf(TEXT("%s"), *Command));
+			}
 		}
 	}
 #endif
@@ -662,13 +682,17 @@ TArray<FRigElementKey> URigHierarchyController::ImportBones(USkeleton* InSkeleto
 	{
 		if (bPrintPythonCommand)
 		{
-			RigVMPythonUtils::Print(GetOuter()->GetFName().ToString(),
-				FString::Printf(TEXT("hierarchy_controller.import_bones_from_asset('%s', '%s', %s, %s, %s)"),
-				*InSkeleton->GetPathName(),
-				*InNameSpace.ToString(),
-				(bReplaceExistingBones) ? TEXT("True") : TEXT("False"),
-				(bRemoveObsoleteBones) ? TEXT("True") : TEXT("False"),
-				(bSelectBones) ? TEXT("True") : TEXT("False")));
+			UBlueprint* Blueprint = GetTypedOuter<UBlueprint>();
+			if (Blueprint)
+			{
+				RigVMPythonUtils::Print(Blueprint->GetFName().ToString(),
+					FString::Printf(TEXT("hierarchy_controller.import_bones_from_asset('%s', '%s', %s, %s, %s)"),
+					*InSkeleton->GetPathName(),
+					*InNameSpace.ToString(),
+					(bReplaceExistingBones) ? TEXT("True") : TEXT("False"),
+					(bRemoveObsoleteBones) ? TEXT("True") : TEXT("False"),
+					(bSelectBones) ? TEXT("True") : TEXT("False")));
+			}
 		}
 	}
 #endif
@@ -768,11 +792,15 @@ TArray<FRigElementKey> URigHierarchyController::ImportCurves(USkeleton* InSkelet
 	{
 		if (bPrintPythonCommand)
 		{
-			RigVMPythonUtils::Print(GetOuter()->GetFName().ToString(),
-				FString::Printf(TEXT("hierarchy_controller.import_curves_from_asset('%s', '%s', %s)"),
-				*InSkeleton->GetPathName(),
-				*InNameSpace.ToString(),
-				(bSelectCurves) ? TEXT("True") : TEXT("False")));
+			UBlueprint* Blueprint = GetTypedOuter<UBlueprint>();
+			if (Blueprint)
+			{
+				RigVMPythonUtils::Print(Blueprint->GetFName().ToString(),
+					FString::Printf(TEXT("hierarchy_controller.import_curves_from_asset('%s', '%s', %s)"),
+					*InSkeleton->GetPathName(),
+					*InNameSpace.ToString(),
+					(bSelectCurves) ? TEXT("True") : TEXT("False")));
+			}
 		}
 	}
 #endif
@@ -1106,13 +1134,17 @@ TArray<FRigElementKey> URigHierarchyController::ImportFromText(FString InContent
 #if WITH_EDITOR
 	if (bPrintPythonCommands)
 	{
-		FString PythonContent = InContent.Replace(TEXT("\\\""), TEXT("\\\\\""));
+		UBlueprint* Blueprint = GetTypedOuter<UBlueprint>();
+		if (Blueprint)
+		{
+			FString PythonContent = InContent.Replace(TEXT("\\\""), TEXT("\\\\\""));
 		
-		RigVMPythonUtils::Print(GetOuter()->GetFName().ToString(),
-			FString::Printf(TEXT("hierarchy_controller.import_from_text('%s', %s, %s)"),
-			*PythonContent,
-			(bReplaceExistingElements) ? TEXT("True") : TEXT("False"),
-			(bSelectNewElements) ? TEXT("True") : TEXT("False")));
+			RigVMPythonUtils::Print(Blueprint->GetFName().ToString(),
+				FString::Printf(TEXT("hierarchy_controller.import_from_text('%s', %s, %s)"),
+				*PythonContent,
+				(bReplaceExistingElements) ? TEXT("True") : TEXT("False"),
+				(bSelectNewElements) ? TEXT("True") : TEXT("False")));
+		}
 	}
 #endif
 
@@ -1584,9 +1616,13 @@ bool URigHierarchyController::RemoveElement(FRigElementKey InElement, bool bSetu
 
 	if (bRemoved && bPrintPythonCommand)
 	{
-		RigVMPythonUtils::Print(GetOuter()->GetFName().ToString(),
-			FString::Printf(TEXT("hierarchy_controller.remove_element(%s)"),
-			*InElement.ToPythonString()));
+		UBlueprint* Blueprint = GetTypedOuter<UBlueprint>();
+		if (Blueprint)
+		{
+			RigVMPythonUtils::Print(Blueprint->GetFName().ToString(),
+				FString::Printf(TEXT("hierarchy_controller.remove_element(%s)"),
+				*InElement.ToPythonString()));
+		}
 	}
 #endif
 	
@@ -1745,10 +1781,14 @@ FRigElementKey URigHierarchyController::RenameElement(FRigElementKey InElement,
 
 	if (bRenamed && bPrintPythonCommand)
 	{
-		RigVMPythonUtils::Print(GetOuter()->GetFName().ToString(), 
-			FString::Printf(TEXT("hierarchy_controller.rename_element(%s, '%s')"),
-			*InElement.ToPythonString(),
-			*InName.ToString()));
+		UBlueprint* Blueprint = GetTypedOuter<UBlueprint>();
+		if (Blueprint)
+		{
+			RigVMPythonUtils::Print(Blueprint->GetFName().ToString(), 
+				FString::Printf(TEXT("hierarchy_controller.rename_element(%s, '%s')"),
+				*InElement.ToPythonString(),
+				*InName.ToString()));
+		}
 	}
 #endif
 
@@ -2039,11 +2079,15 @@ bool URigHierarchyController::RemoveParent(FRigElementKey InChild, FRigElementKe
 
 	if (bRemoved && bPrintPythonCommand)
 	{
-		RigVMPythonUtils::Print(GetOuter()->GetFName().ToString(),
-			FString::Printf(TEXT("hierarchy_controller.remove_parent(%s, %s, %s)"),
-			*InChild.ToPythonString(),
-			*InParent.ToPythonString(),
-			(bMaintainGlobalTransform) ? TEXT("True") : TEXT("False")));
+		UBlueprint* Blueprint = GetTypedOuter<UBlueprint>();
+		if (Blueprint)
+		{
+			RigVMPythonUtils::Print(Blueprint->GetFName().ToString(),
+				FString::Printf(TEXT("hierarchy_controller.remove_parent(%s, %s, %s)"),
+				*InChild.ToPythonString(),
+				*InParent.ToPythonString(),
+				(bMaintainGlobalTransform) ? TEXT("True") : TEXT("False")));
+		}
 	}
 #endif
 	
@@ -2211,10 +2255,14 @@ bool URigHierarchyController::RemoveAllParents(FRigElementKey InChild, bool bMai
 	
 	if (bRemoved && bPrintPythonCommand)
 	{
-		RigVMPythonUtils::Print(GetOuter()->GetFName().ToString(),
-			FString::Printf(TEXT("hierarchy_controller.remove_all_parents(%s, %s)"),
-			*InChild.ToPythonString(),
-			(bMaintainGlobalTransform) ? TEXT("True") : TEXT("False")));
+		UBlueprint* Blueprint = GetTypedOuter<UBlueprint>();
+		if (Blueprint)
+		{
+			RigVMPythonUtils::Print(Blueprint->GetFName().ToString(),
+				FString::Printf(TEXT("hierarchy_controller.remove_all_parents(%s, %s)"),
+				*InChild.ToPythonString(),
+				(bMaintainGlobalTransform) ? TEXT("True") : TEXT("False")));
+		}
 	}
 #endif
 
@@ -2286,11 +2334,15 @@ bool URigHierarchyController::SetParent(FRigElementKey InChild, FRigElementKey I
 
 	if (bParentSet && bPrintPythonCommand)
 	{
-		RigVMPythonUtils::Print(GetOuter()->GetFName().ToString(),
-			FString::Printf(TEXT("hierarchy_controller.set_parent(%s, %s, %s)"),
-			*InChild.ToPythonString(),
-			*InParent.ToPythonString(),
-			(bMaintainGlobalTransform) ? TEXT("True") : TEXT("False")));
+		UBlueprint* Blueprint = GetTypedOuter<UBlueprint>();
+		if (Blueprint)
+		{
+			RigVMPythonUtils::Print(Blueprint->GetFName().ToString(),
+				FString::Printf(TEXT("hierarchy_controller.set_parent(%s, %s, %s)"),
+				*InChild.ToPythonString(),
+				*InParent.ToPythonString(),
+				(bMaintainGlobalTransform) ? TEXT("True") : TEXT("False")));
+		}
 	}
 #endif
 
@@ -2305,21 +2357,25 @@ TArray<FRigElementKey> URigHierarchyController::DuplicateElements(TArray<FRigEle
 #if WITH_EDITOR
 	if (!Result.IsEmpty() && bPrintPythonCommands)
 	{
-		FString ArrayStr = TEXT("[");
-		for (auto It = InKeys.CreateConstIterator(); It; ++It)
+		UBlueprint* Blueprint = GetTypedOuter<UBlueprint>();
+		if (Blueprint)
 		{
-			ArrayStr += It->ToPythonString();
-			if (It.GetIndex() < InKeys.Num() - 1)
+			FString ArrayStr = TEXT("[");
+			for (auto It = InKeys.CreateConstIterator(); It; ++It)
 			{
-				ArrayStr += TEXT(", ");
+				ArrayStr += It->ToPythonString();
+				if (It.GetIndex() < InKeys.Num() - 1)
+				{
+					ArrayStr += TEXT(", ");
+				}
 			}
-		}
-		ArrayStr += TEXT("]");		
+			ArrayStr += TEXT("]");		
 		
-		RigVMPythonUtils::Print(GetOuter()->GetFName().ToString(),
-			FString::Printf(TEXT("hierarchy_controller.duplicate_elements(%s, %s)"),
-			*ArrayStr,
-			(bSelectNewElements) ? TEXT("True") : TEXT("False")));
+			RigVMPythonUtils::Print(Blueprint->GetFName().ToString(),
+				FString::Printf(TEXT("hierarchy_controller.duplicate_elements(%s, %s)"),
+				*ArrayStr,
+				(bSelectNewElements) ? TEXT("True") : TEXT("False")));
+		}
 	}
 #endif
 
@@ -2421,25 +2477,29 @@ TArray<FRigElementKey> URigHierarchyController::MirrorElements(TArray<FRigElemen
 #if WITH_EDITOR
 	if (!DuplicatedKeys.IsEmpty() && bPrintPythonCommands)
 	{
-		FString ArrayStr = TEXT("[");
-		for (auto It = InKeys.CreateConstIterator(); It; ++It)
+		UBlueprint* Blueprint = GetTypedOuter<UBlueprint>();
+		if (Blueprint)
 		{
-			ArrayStr += It->ToPythonString();
-			if (It.GetIndex() < InKeys.Num() - 1)
+			FString ArrayStr = TEXT("[");
+			for (auto It = InKeys.CreateConstIterator(); It; ++It)
 			{
-				ArrayStr += TEXT(", ");
+				ArrayStr += It->ToPythonString();
+				if (It.GetIndex() < InKeys.Num() - 1)
+				{
+					ArrayStr += TEXT(", ");
+				}
 			}
-		}
-		ArrayStr += TEXT("]");
+			ArrayStr += TEXT("]");
 
-		RigVMPythonUtils::Print(GetOuter()->GetFName().ToString(),
-			FString::Printf(TEXT("hierarchy_controller.mirror_elements(%s, unreal.RigMirrorSettings(unreal.AxisType.%s, unreal.AxisType.%s, '%s', '%s'), %s)"),
-			*ArrayStr,
-			(InSettings.MirrorAxis.GetValue() == EAxis::X) ? TEXT("X") : (InSettings.MirrorAxis.GetValue() == EAxis::Y) ? TEXT("Y") : TEXT("Z"),
-			(InSettings.AxisToFlip.GetValue() == EAxis::X) ? TEXT("X") : (InSettings.AxisToFlip.GetValue() == EAxis::Y) ? TEXT("Y") : TEXT("Z"),
-			*InSettings.OldName,
-			*InSettings.NewName,
-			(bSelectNewElements) ? TEXT("True") : TEXT("False")));
+			RigVMPythonUtils::Print(Blueprint->GetFName().ToString(),
+				FString::Printf(TEXT("hierarchy_controller.mirror_elements(%s, unreal.RigMirrorSettings(unreal.AxisType.%s, unreal.AxisType.%s, '%s', '%s'), %s)"),
+				*ArrayStr,
+				(InSettings.MirrorAxis.GetValue() == EAxis::X) ? TEXT("X") : (InSettings.MirrorAxis.GetValue() == EAxis::Y) ? TEXT("Y") : TEXT("Z"),
+				(InSettings.AxisToFlip.GetValue() == EAxis::X) ? TEXT("X") : (InSettings.AxisToFlip.GetValue() == EAxis::Y) ? TEXT("Y") : TEXT("Z"),
+				*InSettings.OldName,
+				*InSettings.NewName,
+				(bSelectNewElements) ? TEXT("True") : TEXT("False")));
+		}
 	}
 #endif
 	
