Control Rig: Deprecate bezier nodes

diff --git a/Engine/Plugins/Animation/ControlRig/Source/ControlRig/Private/Units/Highlevel/Hierarchy/RigUnit_FitChainToCurve.h b/Engine/Plugins/Animation/ControlRig/Source/ControlRig/Private/Units/Highlevel/Hierarchy/RigUnit_FitChainToCurve.h
--- a/Engine/Plugins/Animation/ControlRig/Source/ControlRig/Private/Units/Highlevel/Hierarchy/RigUnit_FitChainToCurve.h
+++ b/Engine/Plugins/Animation/ControlRig/Source/ControlRig/Private/Units/Highlevel/Hierarchy/RigUnit_FitChainToCurve.h
@@ -368,7 +368,7 @@ struct CONTROLRIG_API FRigUnit_FitChainToCurvePerItem : public FRigUnit_Highleve
  * Fits a given chain to a four point bezier curve.
  * Additionally provides rotational control matching the features of the Distribute Rotation node.
  */
-USTRUCT(meta=(DisplayName="Fit Chain on Curve", Category="Hierarchy", Keywords="Fit,Resample,Bezier"))
+USTRUCT(meta=(DisplayName="Fit Chain on Curve", Category="Hierarchy", Keywords="Fit,Resample,Bezier", Deprecated = "5.0"))
 struct CONTROLRIG_API FRigUnit_FitChainToCurveItemArray : public FRigUnit_HighlevelBaseMutable
 {
 	GENERATED_BODY()
diff --git a/Engine/Plugins/Animation/ControlRig/Source/ControlRig/Private/Units/Math/RigUnit_MathVector.h b/Engine/Plugins/Animation/ControlRig/Source/ControlRig/Private/Units/Math/RigUnit_MathVector.h
--- a/Engine/Plugins/Animation/ControlRig/Source/ControlRig/Private/Units/Math/RigUnit_MathVector.h
+++ b/Engine/Plugins/Animation/ControlRig/Source/ControlRig/Private/Units/Math/RigUnit_MathVector.h
@@ -836,7 +836,7 @@ struct CONTROLRIG_API FRigUnit_MathVectorOrthogonal : public FRigUnit_MathVector
 /**
  * Returns the 4 point bezier interpolation
  */
-USTRUCT(meta=(DisplayName="Bezier Four Point"))
+USTRUCT(meta=(DisplayName="Bezier Four Point", Deprecated = "5.0"))
 struct CONTROLRIG_API FRigUnit_MathVectorBezierFourPoint : public FRigUnit_MathVectorBase
 {
 	GENERATED_BODY()
@@ -867,7 +867,7 @@ struct CONTROLRIG_API FRigUnit_MathVectorBezierFourPoint : public FRigUnit_MathV
 /**
  * Creates a bezier four point
  */
-USTRUCT(meta = (DisplayName = "Make Bezier Four Point", Constant))
+USTRUCT(meta = (DisplayName = "Make Bezier Four Point", Constant, Deprecated = "5.0"))
 struct CONTROLRIG_API FRigUnit_MathVectorMakeBezierFourPoint : public FRigUnit_MathVectorBase
 {
 	GENERATED_BODY()
