Control Rig: Python Log does not update for RigElement editing

diff --git a/Engine/Plugins/Animation/ControlRig/Source/ControlRig/Private/Rigs/RigHierarchy.cpp b/Engine/Plugins/Animation/ControlRig/Source/ControlRig/Private/Rigs/RigHierarchy.cpp
--- a/Engine/Plugins/Animation/ControlRig/Source/ControlRig/Private/Rigs/RigHierarchy.cpp
+++ b/Engine/Plugins/Animation/ControlRig/Source/ControlRig/Private/Rigs/RigHierarchy.cpp
@@ -2048,7 +2048,7 @@ void URigHierarchy::SetTransform(FRigTransformElement* InTransformElement, const
 		{
 			FTransform LocalTransform = ComputeLocalControlValue(ControlElement, InTransform, InTransformType);
 			ControlElement->Settings.ApplyLimits(LocalTransform);
-			SetTransform(ControlElement, LocalTransform, MakeLocal(InTransformType), bAffectChildren);
+			SetTransform(ControlElement, LocalTransform, MakeLocal(InTransformType), bAffectChildren, false, false, bPrintPythonCommands);
 			return;
 		}
 	}
@@ -2289,7 +2289,7 @@ FTransform URigHierarchy::GetControlGizmoTransform(FRigControlElement* InControl
 }
 
 void URigHierarchy::SetControlGizmoTransform(FRigControlElement* InControlElement, const FTransform& InTransform,
-	const ERigTransformType::Type InTransformType, bool bSetupUndo, bool bForce)
+	const ERigTransformType::Type InTransformType, bool bSetupUndo, bool bForce, bool bPrintPythonCommands)
 {
 	if(InControlElement == nullptr)
 	{
@@ -2355,10 +2355,21 @@ void URigHierarchy::SetControlGizmoTransform(FRigControlElement* InControlElemen
 			}
 		}
 	}
+
+	if (bPrintPythonCommands)
+	{
+		FString BlueprintName = GetOuter()->GetFName().ToString();
+		RigVMPythonUtils::Print(BlueprintName,
+			FString::Printf(TEXT("hierarchy.set_control_gizmo_transform(%s, %s, %s)"),
+			*InControlElement->GetKey().ToPythonString(),
+			*RigVMPythonUtils::TransformToPythonString(InTransform),
+			ERigTransformType::IsInitial(InTransformType) ? TEXT("True") : TEXT("False")));
+	
+	}
 #endif
 }
 
-void URigHierarchy::SetControlSettings(FRigControlElement* InControlElement, FRigControlSettings InSettings, bool bSetupUndo, bool bForce)
+void URigHierarchy::SetControlSettings(FRigControlElement* InControlElement, FRigControlSettings InSettings, bool bSetupUndo, bool bForce, bool bPrintPythonCommands)
 {
 	if(InControlElement == nullptr)
 	{
@@ -2392,6 +2403,25 @@ void URigHierarchy::SetControlSettings(FRigControlElement* InControlElement, FRi
 			}
 		}
 	}
+
+	if (bPrintPythonCommands)
+	{
+		FString ControlNamePythonized = RigVMPythonUtils::NameToPep8(InControlElement->GetName().ToString());
+		FString SettingsName = FString::Printf(TEXT("control_settings_%s"),
+			*ControlNamePythonized);
+		TArray<FString> Commands = ControlSettingsToPythonCommands(InControlElement->Settings, SettingsName);
+
+		FString BlueprintName = GetOuter()->GetFName().ToString();
+		for (const FString& Command : Commands)
+		{
+			RigVMPythonUtils::Print(BlueprintName, Command);
+		}
+		
+		RigVMPythonUtils::Print(BlueprintName,
+			FString::Printf(TEXT("hierarchy.set_control_settings(%s, %s)"),
+			*InControlElement->GetKey().ToPythonString(),
+			*SettingsName));
+	}
 #endif
 }
 
@@ -2543,6 +2573,22 @@ void URigHierarchy::SetControlValue(FRigControlElement* InControlElement, const
 						}
 					}
 				}
+
+				if (bPrintPythonCommands)
+				{
+					FString TypeStr;
+					switch (InValueType)
+					{
+						case ERigControlValueType::Minimum: TypeStr = TEXT("MINIMUM"); break;
+						case ERigControlValueType::Maximum: TypeStr = TEXT("MAXIMUM"); break;
+						default: ensure(false);
+					}
+					RigVMPythonUtils::Print(GetOuter()->GetFName().ToString(),
+						FString::Printf(TEXT("hierarchy.set_control_value(%s, %s, unreal.RigControlValueType.%s)"),
+						*InControlElement->GetKey().ToPythonString(),
+						*InValue.ToPythonString(InControlElement->Settings.ControlType),
+						*TypeStr));
+				}
 #endif
 				break;
 			}
@@ -4141,3 +4187,68 @@ void URigHierarchy::IntegrateParentConstraintQuat(
 	OutNumMixedRotations++;
 }
 
+#if WITH_EDITOR
+TArray<FString> URigHierarchy::ControlSettingsToPythonCommands(const FRigControlSettings& Settings, const FString& NameSettings)
+{
+	TArray<FString> Commands;
+	Commands.Add(FString::Printf(TEXT("%s = unreal.RigControlSettings()"),
+			*NameSettings));
+	FString TypeStr;
+	switch (Settings.ControlType)
+	{
+		case ERigControlType::Bool: TypeStr = TEXT("BOOL"); break;							
+		case ERigControlType::Float: TypeStr = TEXT("FLOAT"); break;
+		case ERigControlType::Integer: TypeStr = TEXT("INTEGER"); break;
+		case ERigControlType::Position: TypeStr = TEXT("POSITION"); break;
+		case ERigControlType::Rotator: TypeStr = TEXT("POSITION"); break;
+		case ERigControlType::Scale: TypeStr = TEXT("SCALE"); break;
+		case ERigControlType::Transform: TypeStr = TEXT("EULER_TRANSFORM"); break;
+		case ERigControlType::EulerTransform: TypeStr = TEXT("EULER_TRANSFORM"); break;
+		case ERigControlType::Vector2D: TypeStr = TEXT("VECTOR2D"); break;
+		case ERigControlType::TransformNoScale: TypeStr = TEXT("EULER_TRANSFORM"); break;
+		default: ensure(false);
+	}
+	Commands.Add(FString::Printf(TEXT("%s.control_type = unreal.RigControlType.%s"),
+									*NameSettings,
+									*TypeStr));
+	Commands.Add(FString::Printf(TEXT("%s.animatable = %s"),
+		*NameSettings,
+		Settings.bAnimatable ? TEXT("True") : TEXT("False")));
+	Commands.Add(FString::Printf(TEXT("%s.display_name = '%s'"),
+		*NameSettings,
+		*Settings.DisplayName.ToString()));
+	Commands.Add(FString::Printf(TEXT("%s.draw_limits = %s"),
+		*NameSettings,
+		Settings.bDrawLimits ? TEXT("True") : TEXT("False")));
+	Commands.Add(FString::Printf(TEXT("%s.gizmo_color = %s"),
+		*NameSettings,
+		*RigVMPythonUtils::LinearColorToPythonString(Settings.GizmoColor)));
+	Commands.Add(FString::Printf(TEXT("%s.gizmo_enabled = %s"),
+		*NameSettings,
+		Settings.bGizmoEnabled ? TEXT("True") : TEXT("False")));
+	Commands.Add(FString::Printf(TEXT("%s.gizmo_name = '%s'"),
+		*NameSettings,
+		*Settings.GizmoName.ToString()));
+	Commands.Add(FString::Printf(TEXT("%s.gizmo_visible = %s"),
+		*NameSettings,
+		Settings.bGizmoVisible ? TEXT("True") : TEXT("False")));
+	Commands.Add(FString::Printf(TEXT("%s.is_transient_control = %s"),
+		*NameSettings,
+		Settings.bIsTransientControl ? TEXT("True") : TEXT("False")));
+	Commands.Add(FString::Printf(TEXT("%s.limit_rotation = %s"),
+		*NameSettings,
+		Settings.bLimitRotation ? TEXT("True") : TEXT("False")));
+	Commands.Add(FString::Printf(TEXT("%s.limit_translation = %s"),
+		*NameSettings,
+		Settings.bLimitTranslation ? TEXT("True") : TEXT("False")));
+	Commands.Add(FString::Printf(TEXT("%s.limit_scale = %s"),
+		*NameSettings,
+		Settings.bLimitScale ? TEXT("True") : TEXT("False")));
+	Commands.Add(FString::Printf(TEXT("%s.primary_axis = unreal.RigControlAxis.%s"),
+		*NameSettings,
+		Settings.PrimaryAxis == ERigControlAxis::X ? TEXT("X") : Settings.PrimaryAxis == ERigControlAxis::Y ? TEXT("Y") : TEXT("Z")));
+
+	return Commands;
+}
+
+#endif
\ No newline at end of file
diff --git a/Engine/Plugins/Animation/ControlRig/Source/ControlRig/Private/Rigs/RigHierarchyController.cpp b/Engine/Plugins/Animation/ControlRig/Source/ControlRig/Private/Rigs/RigHierarchyController.cpp
--- a/Engine/Plugins/Animation/ControlRig/Source/ControlRig/Private/Rigs/RigHierarchyController.cpp
+++ b/Engine/Plugins/Animation/ControlRig/Source/ControlRig/Private/Rigs/RigHierarchyController.cpp
@@ -1375,62 +1375,7 @@ TArray<FString> URigHierarchyController::GetAddControlPythonCommands(FRigControl
 		SettingsStr = FString::Printf(TEXT("control_settings_%s"),
 			*ControlNamePythonized);
 			
-		Commands.Add(FString::Printf(TEXT("control_settings_%s = unreal.RigControlSettings()"),
-			*ControlNamePythonized));
-		FString TypeStr;
-		switch (Settings.ControlType)
-		{
-			case ERigControlType::Bool: TypeStr = TEXT("BOOL"); break;							
-			case ERigControlType::Float: TypeStr = TEXT("FLOAT"); break;
-			case ERigControlType::Integer: TypeStr = TEXT("INTEGER"); break;
-			case ERigControlType::Position: TypeStr = TEXT("POSITION"); break;
-			case ERigControlType::Rotator: TypeStr = TEXT("POSITION"); break;
-			case ERigControlType::Scale: TypeStr = TEXT("SCALE"); break;
-			case ERigControlType::Transform: TypeStr = TEXT("EULER_TRANSFORM"); break;
-			case ERigControlType::EulerTransform: TypeStr = TEXT("EULER_TRANSFORM"); break;
-			case ERigControlType::Vector2D: TypeStr = TEXT("VECTOR2D"); break;
-			case ERigControlType::TransformNoScale: TypeStr = TEXT("EULER_TRANSFORM"); break;
-			default: ensure(false);
-		}
-		Commands.Add(FString::Printf(TEXT("control_settings_%s.control_type = unreal.RigControlType.%s"),
-										*ControlNamePythonized,
-										*TypeStr));
-		Commands.Add(FString::Printf(TEXT("control_settings_%s.animatable = %s"),
-			*ControlNamePythonized,
-			Settings.bAnimatable ? TEXT("True") : TEXT("False")));
-		Commands.Add(FString::Printf(TEXT("control_settings_%s.display_name = '%s'"),
-			*ControlNamePythonized,
-			*Settings.DisplayName.ToString()));
-		Commands.Add(FString::Printf(TEXT("control_settings_%s.draw_limits = %s"),
-			*ControlNamePythonized,
-			Settings.bDrawLimits ? TEXT("True") : TEXT("False")));
-		Commands.Add(FString::Printf(TEXT("control_settings_%s.gizmo_color = %s"),
-			*ControlNamePythonized,
-			*RigVMPythonUtils::LinearColorToPythonString(Settings.GizmoColor)));
-		Commands.Add(FString::Printf(TEXT("control_settings_%s.gizmo_enabled = %s"),
-			*ControlNamePythonized,
-			Settings.bGizmoEnabled ? TEXT("True") : TEXT("False")));
-		Commands.Add(FString::Printf(TEXT("control_settings_%s.gizmo_name = '%s'"),
-			*ControlNamePythonized,
-			*Settings.GizmoName.ToString()));
-		Commands.Add(FString::Printf(TEXT("control_settings_%s.gizmo_visible = %s"),
-			*ControlNamePythonized,
-			Settings.bGizmoVisible ? TEXT("True") : TEXT("False")));
-		Commands.Add(FString::Printf(TEXT("control_settings_%s.is_transient_control = %s"),
-			*ControlNamePythonized,
-			Settings.bIsTransientControl ? TEXT("True") : TEXT("False")));
-		Commands.Add(FString::Printf(TEXT("control_settings_%s.limit_rotation = %s"),
-				*ControlNamePythonized,
-				Settings.bLimitRotation ? TEXT("True") : TEXT("False")));
-		Commands.Add(FString::Printf(TEXT("control_settings_%s.limit_translation = %s"),
-				*ControlNamePythonized,
-				Settings.bLimitTranslation ? TEXT("True") : TEXT("False")));
-		Commands.Add(FString::Printf(TEXT("control_settings_%s.limit_scale = %s"),
-				*ControlNamePythonized,
-				Settings.bLimitScale ? TEXT("True") : TEXT("False")));
-		Commands.Add(FString::Printf(TEXT("control_settings_%s.primary_axis = unreal.RigControlAxis.%s"),
-			*ControlNamePythonized,
-			Settings.PrimaryAxis == ERigControlAxis::X ? TEXT("X") : Settings.PrimaryAxis == ERigControlAxis::Y ? TEXT("Y") : TEXT("Z")));				
+		Commands.Append(URigHierarchy::ControlSettingsToPythonCommands(Settings, SettingsStr));	
 	}
 		
 	FRigControlValue Value = Hierarchy->GetControlValue(Control->GetKey(), ERigControlValueType::Initial);
diff --git a/Engine/Plugins/Animation/ControlRig/Source/ControlRig/Public/Rigs/RigHierarchy.h b/Engine/Plugins/Animation/ControlRig/Source/ControlRig/Public/Rigs/RigHierarchy.h
--- a/Engine/Plugins/Animation/ControlRig/Source/ControlRig/Public/Rigs/RigHierarchy.h
+++ b/Engine/Plugins/Animation/ControlRig/Source/ControlRig/Public/Rigs/RigHierarchy.h
@@ -7,6 +7,9 @@
 #include "RigHierarchyPose.h"
 #include "UObject/WeakObjectPtrTemplates.h"
 #include "EdGraph/EdGraphPin.h"
+#if WITH_EDITOR
+#include "RigVMPythonUtils.h"
+#endif
 #include "RigHierarchy.generated.h"
 
 class URigHierarchy;
@@ -1203,9 +1206,9 @@ public:
 	 * @param bSetupUndo If true the transform stack will be setup for undo / redo
 	 */
 	UFUNCTION(BlueprintCallable, Category = URigHierarchy)
-    FORCEINLINE_DEBUGGABLE void SetControlValue(FRigElementKey InKey, FRigControlValue InValue, ERigControlValueType InValueType = ERigControlValueType::Current, bool bSetupUndo = false)
+    FORCEINLINE_DEBUGGABLE void SetControlValue(FRigElementKey InKey, FRigControlValue InValue, ERigControlValueType InValueType = ERigControlValueType::Current, bool bSetupUndo = false, bool bPrintPythonCommands = false)
 	{
-		SetControlValueByIndex(GetIndex(InKey), InValue, InValueType, bSetupUndo);
+		SetControlValueByIndex(GetIndex(InKey), InValue, InValueType, bSetupUndo, bPrintPythonCommands);
 	}
 
 	/**
@@ -1216,9 +1219,9 @@ public:
 	 * @param bSetupUndo If true the transform stack will be setup for undo / redo
 	 */
 	template<typename T>
-	FORCEINLINE_DEBUGGABLE void SetControlValue(FRigElementKey InKey, const T& InValue, ERigControlValueType InValueType = ERigControlValueType::Current, bool bSetupUndo = false) const
+	FORCEINLINE_DEBUGGABLE void SetControlValue(FRigElementKey InKey, const T& InValue, ERigControlValueType InValueType = ERigControlValueType::Current, bool bSetupUndo = false, bool bPrintPythonCommands = false) const
 	{
-		return SetControlValue(InKey, FRigControlValue::Make<T>(InValue), InValueType, bSetupUndo);
+		return SetControlValue(InKey, FRigControlValue::Make<T>(InValue), InValueType, bSetupUndo, bPrintPythonCommands);
 	}
 
 	/**
@@ -1229,13 +1232,13 @@ public:
   	 * @param bSetupUndo If true the transform stack will be setup for undo / redo
 	 */
 	UFUNCTION(BlueprintCallable, Category = URigHierarchy)
-	FORCEINLINE_DEBUGGABLE void SetControlValueByIndex(int32 InElementIndex, FRigControlValue InValue, ERigControlValueType InValueType = ERigControlValueType::Current, bool bSetupUndo = false)
+	FORCEINLINE_DEBUGGABLE void SetControlValueByIndex(int32 InElementIndex, FRigControlValue InValue, ERigControlValueType InValueType = ERigControlValueType::Current, bool bSetupUndo = false, bool bPrintPythonCommands = false)
 	{
 		if(Elements.IsValidIndex(InElementIndex))
 		{
 			if(FRigControlElement* ControlElement = Cast<FRigControlElement>(Elements[InElementIndex]))
 			{
-				SetControlValue(ControlElement, InValue, InValueType, bSetupUndo);
+				SetControlValue(ControlElement, InValue, InValueType, bSetupUndo, false, bPrintPythonCommands);
 			}
 		}
 	}
@@ -1247,9 +1250,9 @@ public:
 	 * @param InValueType The type of value to set
 	 * @param bSetupUndo If true the transform stack will be setup for undo / redo
 	 */
-	FORCEINLINE_DEBUGGABLE void SetControlValue(int32 InElementIndex, const FRigControlValue& InValue, ERigControlValueType InValueType = ERigControlValueType::Current, bool bSetupUndo = false)
+	FORCEINLINE_DEBUGGABLE void SetControlValue(int32 InElementIndex, const FRigControlValue& InValue, ERigControlValueType InValueType = ERigControlValueType::Current, bool bSetupUndo = false, bool bPrintPythonCommands = false)
 	{
-		SetControlValueByIndex(InElementIndex, InValue, InValueType, bSetupUndo);
+		SetControlValueByIndex(InElementIndex, InValue, InValueType, bSetupUndo, bPrintPythonCommands);
 	}
 
 	/**
@@ -1463,9 +1466,9 @@ public:
 	 * @param bSetupUndo If true the transform stack will be setup for undo / redo
 	 */
 	UFUNCTION(BlueprintCallable, Category = URigHierarchy)
-	FORCEINLINE_DEBUGGABLE void SetControlSettings(FRigElementKey InKey, FRigControlSettings InSettings, bool bSetupUndo = false, bool bForce = false)
+	FORCEINLINE_DEBUGGABLE void SetControlSettings(FRigElementKey InKey, FRigControlSettings InSettings, bool bSetupUndo = false, bool bForce = false, bool bPrintPythonCommands = false)
 	{
-		return SetControlSettingsByIndex(GetIndex(InKey), InSettings, bSetupUndo, bForce);
+		return SetControlSettingsByIndex(GetIndex(InKey), InSettings, bSetupUndo, bForce, bPrintPythonCommands);
 	}
 
 	/**
@@ -1475,13 +1478,13 @@ public:
 	 * @param bSetupUndo If true the transform stack will be setup for undo / redo
 	 */
 	UFUNCTION(BlueprintCallable, Category = URigHierarchy)
-	FORCEINLINE_DEBUGGABLE void SetControlSettingsByIndex(int32 InElementIndex, FRigControlSettings InSettings, bool bSetupUndo = false, bool bForce = false)
+	FORCEINLINE_DEBUGGABLE void SetControlSettingsByIndex(int32 InElementIndex, FRigControlSettings InSettings, bool bSetupUndo = false, bool bForce = false, bool bPrintPythonCommands = false)
 	{
 		if(Elements.IsValidIndex(InElementIndex))
 		{
 			if(FRigControlElement* ControlElement = Cast<FRigControlElement>(Elements[InElementIndex]))
 			{
-				SetControlSettings(ControlElement, InSettings, bSetupUndo, bForce);
+				SetControlSettings(ControlElement, InSettings, bSetupUndo, bForce, bPrintPythonCommands);
 			}
 		}
 	}
@@ -2315,7 +2318,7 @@ public:
 	 * @param bSetupUndo If true the transform stack will be setup for undo / redo
 	 * @param bForce Set the transform even if it is the same as the previously set one
 	 */
-	void SetControlGizmoTransform(FRigControlElement* InControlElement, const FTransform& InTransform, const ERigTransformType::Type InTransformType, bool bSetupUndo = false, bool bForce = false);
+	void SetControlGizmoTransform(FRigControlElement* InControlElement, const FTransform& InTransform, const ERigTransformType::Type InTransformType, bool bSetupUndo = false, bool bForce = false, bool bPrintPythonCommands = false);
 
 	/**
 	 * Sets the control settings for a given control element
@@ -2323,7 +2326,7 @@ public:
 	 * @param InSettings The new control settings value to set
 	 * @param bSetupUndo If true the transform stack will be setup for undo / redo
 	 */
-	void SetControlSettings(FRigControlElement* InControlElement, FRigControlSettings InSettings, bool bSetupUndo = false, bool bForce = false);
+	void SetControlSettings(FRigControlElement* InControlElement, FRigControlSettings InSettings, bool bSetupUndo = false, bool bForce = false, bool bPrintPythonCommands = false);
 
 	/**
 	 * Returns a control's current value
@@ -2836,6 +2839,10 @@ protected:
 		const FTransform& InTransform,
 		float InWeight);
 
+#if WITH_EDITOR
+	static TArray<FString> ControlSettingsToPythonCommands(const FRigControlSettings& Settings, const FString& NameSettings);
+#endif
+	
 	friend class URigHierarchyController;
 	friend class UControlRig;
 	friend class FControlRigEditor;
diff --git a/Engine/Plugins/Animation/ControlRig/Source/ControlRigEditor/Private/ControlRigElementDetails.cpp b/Engine/Plugins/Animation/ControlRig/Source/ControlRigEditor/Private/ControlRigElementDetails.cpp
--- a/Engine/Plugins/Animation/ControlRig/Source/ControlRigEditor/Private/ControlRigElementDetails.cpp
+++ b/Engine/Plugins/Animation/ControlRig/Source/ControlRigEditor/Private/ControlRigElementDetails.cpp
@@ -1476,7 +1476,7 @@ void FRigControlElementDetails_SetupStructValueWidget(IDetailGroup& InGroup, IDe
 		if(HierarchyPtr.IsValid())
 		{
             const FRigControlValue ControlValue = FRigControlElementDetails_PackageValue(*(T*)StructToDisplay->GetStructMemory());
-			HierarchyPtr->SetControlValue(Key, ControlValue, InValueType, true);
+			HierarchyPtr->SetControlValue(Key, ControlValue, InValueType, true, true);
 			if(InValueType == ERigControlValueType::Initial)
 			{
                 if(UControlRigBlueprint* Blueprint = RigElementDetails_GetBlueprintFromHierarchy(HierarchyPtr.Get()))
@@ -2052,7 +2052,7 @@ void FRigControlElementDetails::SetDisplayName(const FText& InNewText, ETextComm
 
 			if(GetHierarchy())
 			{
-				GetHierarchy()->SetControlSettings(ControlElement->GetKey(), Settings);
+				GetHierarchy()->SetControlSettings(ControlElement->GetKey(), Settings, false, false, true);
 			}
 		}
 	}
diff --git a/Engine/Plugins/Animation/ControlRig/Source/ControlRigEditor/Private/Editor/ControlRigEditor.cpp b/Engine/Plugins/Animation/ControlRig/Source/ControlRigEditor/Private/Editor/ControlRigEditor.cpp
--- a/Engine/Plugins/Animation/ControlRig/Source/ControlRigEditor/Private/Editor/ControlRigEditor.cpp
+++ b/Engine/Plugins/Animation/ControlRig/Source/ControlRigEditor/Private/Editor/ControlRigEditor.cpp
@@ -4439,7 +4439,7 @@ void FControlRigEditor::OnWrappedPropertyChangedChainEvent(UDetailsViewWrapperOb
 			}
 			else
 			{
-				Hierarchy->SetTransform(TransformElement, Transform, TransformType, true, true, false);
+				Hierarchy->SetTransform(TransformElement, Transform, TransformType, true, true, false, true);
 			}
 		}
 		else if(PropertyPath.RemoveFromStart(OffsetString))
@@ -4454,7 +4454,7 @@ void FControlRigEditor::OnWrappedPropertyChangedChainEvent(UDetailsViewWrapperOb
 			bIsInitial = bIsInitial || ERigTransformType::IsInitial(TransformType);
 			const FTransform Transform = CastChecked<FRigControlElement>(FirstWrappedElement)->Offset.Get(TransformType);
 			
-			ControlRigBP->Hierarchy->SetControlOffsetTransform(ControlElement, Transform, ERigTransformType::MakeInitial(TransformType), true, true, false);
+			ControlRigBP->Hierarchy->SetControlOffsetTransform(ControlElement, Transform, ERigTransformType::MakeInitial(TransformType), true, true, false, true);
 		}
 		else if(PropertyPath.RemoveFromStart(GizmoString))
 		{
@@ -4468,7 +4468,7 @@ void FControlRigEditor::OnWrappedPropertyChangedChainEvent(UDetailsViewWrapperOb
 			bIsInitial = bIsInitial || ERigTransformType::IsInitial(TransformType);
 			const FTransform Transform = CastChecked<FRigControlElement>(FirstWrappedElement)->Gizmo.Get(TransformType);
 			
-			ControlRigBP->Hierarchy->SetControlGizmoTransform(ControlElement, Transform, ERigTransformType::MakeInitial(TransformType), true, false);
+			ControlRigBP->Hierarchy->SetControlGizmoTransform(ControlElement, Transform, ERigTransformType::MakeInitial(TransformType), true, false, true);
 		}
 		else if(PropertyPath.RemoveFromStart(SettingsString))
 		{
@@ -4480,7 +4480,7 @@ void FControlRigEditor::OnWrappedPropertyChangedChainEvent(UDetailsViewWrapperOb
 				return;
 			}
 
-			ControlRigBP->Hierarchy->SetControlSettings(ControlElement, Settings, true, false);
+			ControlRigBP->Hierarchy->SetControlSettings(ControlElement, Settings, true, false, true);
 		}
 
 		if(bSetupModeEnabled || bIsInitial)
diff --git a/Engine/Plugins/Animation/ControlRig/Source/ControlRigEditor/Private/Editor/SControlRigGizmoNameList.cpp b/Engine/Plugins/Animation/ControlRig/Source/ControlRigEditor/Private/Editor/SControlRigGizmoNameList.cpp
--- a/Engine/Plugins/Animation/ControlRig/Source/ControlRigEditor/Private/Editor/SControlRigGizmoNameList.cpp
+++ b/Engine/Plugins/Animation/ControlRig/Source/ControlRigEditor/Private/Editor/SControlRigGizmoNameList.cpp
@@ -113,7 +113,7 @@ void SControlRigGizmoNameList::SetNameListText(const FText& NewTypeInValue, ETex
 
 				FRigControlSettings Settings = ControlElement->Settings;
 				Settings.GizmoName = NewName;
-				Hierarchy->SetControlSettings(ControlElement, Settings, true, true);
+				Hierarchy->SetControlSettings(ControlElement, Settings, true, true, true);
 			}
 		}
 	}
