Control Rig: Fix Python Output Log outputs Control Rigs that are manipulating in Level Editor

diff --git a/Engine/Plugins/Animation/ControlRig/Source/ControlRig/Private/Rigs/RigHierarchy.cpp b/Engine/Plugins/Animation/ControlRig/Source/ControlRig/Private/Rigs/RigHierarchy.cpp
--- a/Engine/Plugins/Animation/ControlRig/Source/ControlRig/Private/Rigs/RigHierarchy.cpp
+++ b/Engine/Plugins/Animation/ControlRig/Source/ControlRig/Private/Rigs/RigHierarchy.cpp
@@ -2112,8 +2112,19 @@ void URigHierarchy::SetTransform(FRigTransformElement* InTransformElement, const
 
 	if (bPrintPythonCommands)
 	{
-		UBlueprint* Blueprint = GetTypedOuter<UBlueprint>();
-		if (Blueprint)
+		FString BlueprintName;
+		if (UBlueprint* Blueprint = GetTypedOuter<UBlueprint>())
+		{
+			BlueprintName = Blueprint->GetFName().ToString();
+		}
+		else if (UControlRig* Rig = Cast<UControlRig>(GetOuter()))
+		{
+			if (UBlueprint* BlueprintCR = Cast<UBlueprint>(Rig->GetClass()->ClassGeneratedBy))
+			{
+				BlueprintName = BlueprintCR->GetFName().ToString();
+			}
+		}
+		if (!BlueprintName.IsEmpty())
 		{
 			FString MethodName;
 			switch (InTransformType)
@@ -2132,7 +2143,7 @@ void URigHierarchy::SetTransform(FRigTransformElement* InTransformElement, const
 				}
 			}
 
-			RigVMPythonUtils::Print(Blueprint->GetFName().ToString(),
+			RigVMPythonUtils::Print(BlueprintName,
 				FString::Printf(TEXT("hierarchy.%s(%s, %s, %s, %s)"),
 				*MethodName,
 				*InTransformElement->GetKey().ToPythonString(),
@@ -2247,10 +2258,21 @@ void URigHierarchy::SetControlOffsetTransform(FRigControlElement* InControlEleme
 
 	if (bPrintPythonCommands)
 	{
-		UBlueprint* Blueprint = GetTypedOuter<UBlueprint>();
-		if (Blueprint)
+		FString BlueprintName;
+		if (UBlueprint* Blueprint = GetTypedOuter<UBlueprint>())
+		{
+			BlueprintName = Blueprint->GetFName().ToString();
+		}
+		else if (UControlRig* Rig = Cast<UControlRig>(GetOuter()))
+		{
+			if (UBlueprint* BlueprintCR = Cast<UBlueprint>(Rig->GetClass()->ClassGeneratedBy))
+			{
+				BlueprintName = BlueprintCR->GetFName().ToString();
+			}
+		}
+		if (!BlueprintName.IsEmpty())
 		{
-			RigVMPythonUtils::Print(Blueprint->GetFName().ToString(),
+			RigVMPythonUtils::Print(BlueprintName,
 				FString::Printf(TEXT("hierarchy.set_control_offset_transform(%s, %s, %s, %s)"),
 				*InControlElement->GetKey().ToPythonString(),
 				*RigVMPythonUtils::TransformToPythonString(InTransform),
@@ -2358,12 +2380,26 @@ void URigHierarchy::SetControlGizmoTransform(FRigControlElement* InControlElemen
 
 	if (bPrintPythonCommands)
 	{
-		FString BlueprintName = GetOuter()->GetFName().ToString();
-		RigVMPythonUtils::Print(BlueprintName,
-			FString::Printf(TEXT("hierarchy.set_control_gizmo_transform(%s, %s, %s)"),
-			*InControlElement->GetKey().ToPythonString(),
-			*RigVMPythonUtils::TransformToPythonString(InTransform),
-			ERigTransformType::IsInitial(InTransformType) ? TEXT("True") : TEXT("False")));
+		FString BlueprintName;
+		if (UBlueprint* Blueprint = GetTypedOuter<UBlueprint>())
+		{
+			BlueprintName = Blueprint->GetFName().ToString();
+		}
+		else if (UControlRig* Rig = Cast<UControlRig>(GetOuter()))
+		{
+			if (UBlueprint* BlueprintCR = Cast<UBlueprint>(Rig->GetClass()->ClassGeneratedBy))
+			{
+				BlueprintName = BlueprintCR->GetFName().ToString();
+			}
+		}
+		if (!BlueprintName.IsEmpty())
+		{
+			RigVMPythonUtils::Print(BlueprintName,
+				FString::Printf(TEXT("hierarchy.set_control_gizmo_transform(%s, %s, %s)"),
+				*InControlElement->GetKey().ToPythonString(),
+				*RigVMPythonUtils::TransformToPythonString(InTransform),
+				ERigTransformType::IsInitial(InTransformType) ? TEXT("True") : TEXT("False")));
+		}
 	
 	}
 #endif
@@ -2406,21 +2442,35 @@ void URigHierarchy::SetControlSettings(FRigControlElement* InControlElement, FRi
 
 	if (bPrintPythonCommands)
 	{
-		FString ControlNamePythonized = RigVMPythonUtils::NameToPep8(InControlElement->GetName().ToString());
-		FString SettingsName = FString::Printf(TEXT("control_settings_%s"),
-			*ControlNamePythonized);
-		TArray<FString> Commands = ControlSettingsToPythonCommands(InControlElement->Settings, SettingsName);
-
-		FString BlueprintName = GetOuter()->GetFName().ToString();
-		for (const FString& Command : Commands)
+		FString BlueprintName;
+		if (UBlueprint* Blueprint = GetTypedOuter<UBlueprint>())
 		{
-			RigVMPythonUtils::Print(BlueprintName, Command);
+			BlueprintName = Blueprint->GetFName().ToString();
+		}
+		else if (UControlRig* Rig = Cast<UControlRig>(GetOuter()))
+		{
+			if (UBlueprint* BlueprintCR = Cast<UBlueprint>(Rig->GetClass()->ClassGeneratedBy))
+			{
+				BlueprintName = BlueprintCR->GetFName().ToString();
+			}
+		}
+		if (!BlueprintName.IsEmpty())
+		{
+			FString ControlNamePythonized = RigVMPythonUtils::NameToPep8(InControlElement->GetName().ToString());
+			FString SettingsName = FString::Printf(TEXT("control_settings_%s"),
+				*ControlNamePythonized);
+			TArray<FString> Commands = ControlSettingsToPythonCommands(InControlElement->Settings, SettingsName);
+
+			for (const FString& Command : Commands)
+			{
+				RigVMPythonUtils::Print(BlueprintName, Command);
+			}
+			
+			RigVMPythonUtils::Print(BlueprintName,
+				FString::Printf(TEXT("hierarchy.set_control_settings(%s, %s)"),
+				*InControlElement->GetKey().ToPythonString(),
+				*SettingsName));
 		}
-		
-		RigVMPythonUtils::Print(BlueprintName,
-			FString::Printf(TEXT("hierarchy.set_control_settings(%s, %s)"),
-			*InControlElement->GetKey().ToPythonString(),
-			*SettingsName));
 	}
 #endif
 }
@@ -2576,18 +2626,33 @@ void URigHierarchy::SetControlValue(FRigControlElement* InControlElement, const
 
 				if (bPrintPythonCommands)
 				{
-					FString TypeStr;
-					switch (InValueType)
+					FString BlueprintName;
+					if (UBlueprint* Blueprint = GetTypedOuter<UBlueprint>())
+					{
+						BlueprintName = Blueprint->GetFName().ToString();
+					}
+					else if (UControlRig* Rig = Cast<UControlRig>(GetOuter()))
+					{
+						if (UBlueprint* BlueprintCR = Cast<UBlueprint>(Rig->GetClass()->ClassGeneratedBy))
+						{
+							BlueprintName = BlueprintCR->GetFName().ToString();
+						}
+					}
+					if (!BlueprintName.IsEmpty())
 					{
+						FString TypeStr;
+						switch (InValueType)
+						{
 						case ERigControlValueType::Minimum: TypeStr = TEXT("MINIMUM"); break;
 						case ERigControlValueType::Maximum: TypeStr = TEXT("MAXIMUM"); break;
 						default: ensure(false);
+						}
+						RigVMPythonUtils::Print(BlueprintName,
+							FString::Printf(TEXT("hierarchy.set_control_value(%s, %s, unreal.RigControlValueType.%s)"),
+							*InControlElement->GetKey().ToPythonString(),
+							*InValue.ToPythonString(InControlElement->Settings.ControlType),
+							*TypeStr));
 					}
-					RigVMPythonUtils::Print(GetOuter()->GetFName().ToString(),
-						FString::Printf(TEXT("hierarchy.set_control_value(%s, %s, unreal.RigControlValueType.%s)"),
-						*InControlElement->GetKey().ToPythonString(),
-						*InValue.ToPythonString(InControlElement->Settings.ControlType),
-						*TypeStr));
 				}
 #endif
 				break;
diff --git a/Engine/Plugins/Animation/ControlRig/Source/ControlRigEditor/Private/EditMode/ControlRigEditMode.cpp b/Engine/Plugins/Animation/ControlRig/Source/ControlRigEditor/Private/EditMode/ControlRigEditMode.cpp
--- a/Engine/Plugins/Animation/ControlRig/Source/ControlRigEditor/Private/EditMode/ControlRigEditMode.cpp
+++ b/Engine/Plugins/Animation/ControlRig/Source/ControlRigEditor/Private/EditMode/ControlRigEditMode.cpp
@@ -2668,7 +2668,13 @@ void FControlRigEditMode::MoveGizmo(AControlRigGizmoActor* GizmoActor, const boo
 				{
 					InOutLocal = InteractionControlRig->GetControlLocalTransform(GizmoActor->ControlName);
 				}
-				InteractionControlRig->SetControlGlobalTransform(GizmoActor->ControlName, NewTransform, true, Context, true, true);			// assumes it's attached to actor
+
+				bool bPrintPythonCommands = false;
+				if (UWorld* World = InteractionControlRig->GetWorld())
+				{
+					bPrintPythonCommands = World->IsPreviewWorld();
+				}
+				InteractionControlRig->SetControlGlobalTransform(GizmoActor->ControlName, NewTransform, true, Context, true, bPrintPythonCommands);			// assumes it's attached to actor
 				GizmoActor->SetGlobalTransform(CurrentTransform);
 				if (bCalcLocal)
 				{
