Control Rig: Ensure is thrown when redoing deletion of category with variables

diff --git a/Engine/Plugins/Animation/ControlRig/Source/ControlRigDeveloper/Private/ControlRigBlueprint.cpp b/Engine/Plugins/Animation/ControlRig/Source/ControlRigDeveloper/Private/ControlRigBlueprint.cpp
--- a/Engine/Plugins/Animation/ControlRig/Source/ControlRigDeveloper/Private/ControlRigBlueprint.cpp
+++ b/Engine/Plugins/Animation/ControlRig/Source/ControlRigDeveloper/Private/ControlRigBlueprint.cpp
@@ -3697,7 +3697,12 @@ void UControlRigBlueprint::OnVariableRemoved(const FName& InVarName)
 	{
 		if (URigVMController* Controller = GetOrCreateController(Graph))
 		{
-			Controller->OnExternalVariableRemoved(InVarName, true);
+#if WITH_EDITOR
+			const bool bSetupUndoRedo = !GIsTransacting;
+#else
+			const bool bSetupUndoRedo = false;
+#endif
+			Controller->OnExternalVariableRemoved(InVarName, bSetupUndoRedo);
 		}
 	}
 
@@ -3715,7 +3720,12 @@ void UControlRigBlueprint::OnVariableRenamed(const FName& InOldVarName, const FN
 	{
 		if (URigVMController* Controller = GetOrCreateController(Graph))
 		{
-			Controller->OnExternalVariableRenamed(InOldVarName, InNewVarName, true);
+#if WITH_EDITOR
+			const bool bSetupUndoRedo = !GIsTransacting;
+#else
+			const bool bSetupUndoRedo = false;
+#endif
+			Controller->OnExternalVariableRenamed(InOldVarName, InNewVarName, bSetupUndoRedo);
 		}
 	}
 
