Control Rig: Orphan Pin Unit Test BP does not need to display errors for PIE-triggered compilation

#jria UE-128641

diff --git a/Engine/Plugins/Animation/ControlRig/Source/ControlRigDeveloper/Private/ControlRigBlueprint.cpp b/Engine/Plugins/Animation/ControlRig/Source/ControlRigDeveloper/Private/ControlRigBlueprint.cpp
--- a/Engine/Plugins/Animation/ControlRig/Source/ControlRigDeveloper/Private/ControlRigBlueprint.cpp
+++ b/Engine/Plugins/Animation/ControlRig/Source/ControlRigDeveloper/Private/ControlRigBlueprint.cpp
@@ -596,6 +596,11 @@ void UControlRigBlueprint::RecompileVM()
 
 		if (bErrorsDuringCompilation)
 		{
+			if(Compiler->Settings.SurpressErrors)
+			{
+				Compiler->Settings.Reportf(EMessageSeverity::Info, this,
+					TEXT("Compilation Errors may be suppressed for ControlRigBlueprint: %s. See VM Compile Setting in Class Settings for more Details"), *this->GetName());
+			}
 			bVMRecompilationRequired = false;
 			if(CDO->VM)
 			{
@@ -687,6 +692,12 @@ void UControlRigBlueprint::HandleReportFromCompiler(EMessageSeverity::Type InSev
 		Status = BS_Error;
 		MarkPackageDirty();
 
+		// see UnitTest "ControlRig.Basics.OrphanedPins" to learn why errors are suppressed this way
+		if (VMCompileSettings.SurpressErrors)
+		{
+			Log->bSilentMode = true;
+		}
+
 		if(InMessage.Contains(TEXT("@@")))
 		{
 			Log->Error(*InMessage, SubjectForMessage);
@@ -697,7 +708,13 @@ void UControlRigBlueprint::HandleReportFromCompiler(EMessageSeverity::Type InSev
 		}
 
 		BroadCastReportCompilerMessage(InSeverity, InSubject, InMessage);
-		FScriptExceptionHandler::Get().HandleException(ELogVerbosity::Error, *InMessage, *FString());
+
+		// see UnitTest "ControlRig.Basics.OrphanedPins" to learn why errors are suppressed this way
+		if (!VMCompileSettings.SurpressErrors)
+		{ 
+			FScriptExceptionHandler::Get().HandleException(ELogVerbosity::Error, *InMessage, *FString());
+		}
+		
 		bErrorsDuringCompilation = true;
 	}
 	else if (InSeverity == EMessageSeverity::Warning)
