Control Rig: Introduce hierarchy validity checker

diff --git a/Engine/Plugins/Animation/ControlRig/Source/ControlRig/Private/AnimNode_ControlRigBase.cpp b/Engine/Plugins/Animation/ControlRig/Source/ControlRig/Private/AnimNode_ControlRigBase.cpp
--- a/Engine/Plugins/Animation/ControlRig/Source/ControlRig/Private/AnimNode_ControlRigBase.cpp
+++ b/Engine/Plugins/Animation/ControlRig/Source/ControlRig/Private/AnimNode_ControlRigBase.cpp
@@ -138,6 +138,10 @@ void FAnimNode_ControlRigBase::UpdateInput(UControlRig* ControlRig, const FPoseC
 		// so the reset here ensures excluded bones are also reset
 		if(!ControlRigBoneInputMappingByName.IsEmpty() || bResetInputPoseToInitial)
 		{
+			{
+				FRigHierarchyValidityBracket ValidityBracket(ControlRig->GetHierarchy());
+			}
+			FRigHierarchyValidityBracket ValidityBracket(ControlRig->GetHierarchy());
 			ControlRig->GetHierarchy()->ResetPoseToInitial(ERigElementType::Bone);
 		}
 
diff --git a/Engine/Plugins/Animation/ControlRig/Source/ControlRig/Private/ControlRig.cpp b/Engine/Plugins/Animation/ControlRig/Source/ControlRig/Private/ControlRig.cpp
--- a/Engine/Plugins/Animation/ControlRig/Source/ControlRig/Private/ControlRig.cpp
+++ b/Engine/Plugins/Animation/ControlRig/Source/ControlRig/Private/ControlRig.cpp
@@ -201,6 +201,9 @@ void UControlRig::InitializeFromCDO()
 		// copy hierarchy
 		{
 			PostInitInstanceIfRequired();
+
+			FRigHierarchyValidityBracket ValidityBracketA(GetHierarchy());
+			FRigHierarchyValidityBracket ValidityBracketB(CDO->GetHierarchy());
 			
 			TGuardValue<bool> Guard(GetHierarchy()->GetSuspendNotificationsFlag(), true);
 			GetHierarchy()->CopyHierarchy(CDO->GetHierarchy());
@@ -860,6 +863,8 @@ void UControlRig::Execute(const EControlRigState InState, const FName& InEventNa
 	{
 		Context.DrawInterface->Instructions.Append(Context.DrawContainer->Instructions);
 
+		FRigHierarchyValidityBracket ValidityBracket(GetHierarchy());
+		
 		GetHierarchy()->ForEach<FRigControlElement>([this](FRigControlElement* ControlElement) -> bool
 		{
 			const FRigControlSettings& Settings = ControlElement->Settings;
diff --git a/Engine/Plugins/Animation/ControlRig/Source/ControlRig/Private/Rigs/RigHierarchy.cpp b/Engine/Plugins/Animation/ControlRig/Source/ControlRig/Private/Rigs/RigHierarchy.cpp
--- a/Engine/Plugins/Animation/ControlRig/Source/ControlRig/Private/Rigs/RigHierarchy.cpp
+++ b/Engine/Plugins/Animation/ControlRig/Source/ControlRig/Private/Rigs/RigHierarchy.cpp
@@ -99,6 +99,12 @@ URigHierarchy::URigHierarchy()
 , TraceFramesLeft(0)
 , TraceFramesCaptured(0)
 #endif
+#if URIGHIERARCHY_ENSURE_CACHE_VALIDITY
+, bEnableCacheValidityCheck(true)
+#else
+, bEnableCacheValidityCheck(false)
+#endif
+, HierarchyForCacheValidation()
 {
 	Reset();
 #if WITH_EDITOR
@@ -144,6 +150,8 @@ void URigHierarchy::Save(FArchive& Ar)
 
 	// make sure all parts of pose are valid.
 	// this ensures cache validity.
+	EnsureCacheValidity();
+
 	ComputeAllTransforms();
 
 	int32 ElementCount = Elements.Num();
@@ -344,7 +352,7 @@ void URigHierarchy::Reset()
 	IndexLookup.Reset();
 
 	ResetPoseHash = INDEX_NONE;
-	ResetPoseHasFilteredChildren.Reset();
+	ResetPoseIsFilteredOut.Reset();
 
 	if(!IsGarbageCollecting())
 	{
@@ -418,6 +426,8 @@ void URigHierarchy::CopyHierarchy(URigHierarchy* InHierarchy)
 
 	TopologyVersion = InHierarchy->GetTopologyVersion();
 	UpdateAllCachedChildren();
+	
+	EnsureCacheValidity();
 }
 
 #if WITH_EDITOR
@@ -484,6 +494,8 @@ void URigHierarchy::CopyPose(URigHierarchy* InHierarchy, bool bCurrent, bool bIn
 			Element->CopyPose(OtherElement, bCurrent, bInitial);
 		}
 	}
+
+	EnsureCacheValidity();
 }
 
 void URigHierarchy::UpdateSockets(const FRigUnitContext* InContext)
@@ -525,44 +537,33 @@ void URigHierarchy::ResetPoseToInitial(ERigElementType InTypeFilter)
 		const int32 Hash = HashCombine(GetTopologyVersion(), (int32)InTypeFilter);
 		if(Hash != ResetPoseHash)
 		{
-			ResetPoseHasFilteredChildren.Reset();
+			ResetPoseIsFilteredOut.Reset();
 			ResetPoseHash = Hash;
 
 			// let's look at all elements and mark all parent of unaffected children
-			bool bHitAnyParentWithFilteredChildren = false;
-			ResetPoseHasFilteredChildren.AddZeroed(Elements.Num());
+			ResetPoseIsFilteredOut.AddZeroed(Elements.Num());
 
-			Traverse([this, &bHitAnyParentWithFilteredChildren, InTypeFilter](FRigBaseElement* InElement, bool& bContinue)
+			Traverse([this, InTypeFilter](FRigBaseElement* InElement, bool& bContinue)
 			{
 				bContinue = true;
+				ResetPoseIsFilteredOut[InElement->GetIndex()] = !InElement->IsTypeOf(InTypeFilter);
 
-				const bool bFilteredOut = (!InElement->IsTypeOf(InTypeFilter)) || ResetPoseHasFilteredChildren[InElement->Index];
-				if(bFilteredOut)
+				// make sure to distribute the filtering options from
+				// the parent to the children of the part of the tree
+				const FRigBaseElementParentArray Parents = GetParents(InElement);
+				for(const FRigBaseElement* Parent : Parents)
 				{
-					const FRigBaseElementParentArray Parents = GetParents(InElement);
-					for(const FRigBaseElement* Parent : Parents)
+					if(!ResetPoseIsFilteredOut[Parent->GetIndex()])
 					{
-						// only mark this up if the parent is not filtered out / 
-						// if we want the parent to reset its pose to initial.
-						if(Parent->IsTypeOf(InTypeFilter))
-						{
-							bHitAnyParentWithFilteredChildren = true;
-						}
-						ResetPoseHasFilteredChildren[Parent->GetIndex()] = true;
+						ResetPoseIsFilteredOut[InElement->GetIndex()] = false;
 					}
 				}
-
-			}, false);
-
-			if(!bHitAnyParentWithFilteredChildren)
-			{
-				ResetPoseHasFilteredChildren.Reset();
-			}
+			});
 		}
 
 		// if the per element state is empty
 		// it means that the filter doesn't affect 
-		if(ResetPoseHasFilteredChildren.IsEmpty())
+		if(ResetPoseIsFilteredOut.IsEmpty())
 		{
 			bPerformFiltering = false;
 		}
@@ -570,58 +571,32 @@ void URigHierarchy::ResetPoseToInitial(ERigElementType InTypeFilter)
 	
 	for(int32 ElementIndex=0; ElementIndex<Elements.Num(); ElementIndex++)
 	{
-		bool bHasFilteredChildren = bPerformFiltering;
-		if(bHasFilteredChildren)
-		{
-			bHasFilteredChildren = ResetPoseHasFilteredChildren[ElementIndex];
-		}
-		
-		if(!Elements[ElementIndex]->IsTypeOf(InTypeFilter))
+		if(!ResetPoseIsFilteredOut.IsEmpty() && bPerformFiltering)
 		{
-			continue;
+			if(ResetPoseIsFilteredOut[ElementIndex])
+			{
+				continue;
+			}
 		}
 
 		if(FRigControlElement* ControlElement = Cast<FRigControlElement>(Elements[ElementIndex]))
 		{
-			if(bHasFilteredChildren)
-			{
-				const FTransform OffsetTransform = GetControlOffsetTransform(ControlElement, ERigTransformType::InitialLocal);
-				SetControlOffsetTransform(ControlElement, OffsetTransform, ERigTransformType::CurrentLocal, true, false, true);
-				const FTransform GizmoTransform = GetControlGizmoTransform(ControlElement, ERigTransformType::InitialLocal);
-				SetControlGizmoTransform(ControlElement, GizmoTransform, ERigTransformType::CurrentLocal, false, true);
-			}
-			else
-			{
-				ControlElement->Offset.Current = ControlElement->Offset.Initial;
-				ControlElement->Gizmo.Current = ControlElement->Gizmo.Initial;
-			}
+			ControlElement->Offset.Current = ControlElement->Offset.Initial;
+			ControlElement->Gizmo.Current = ControlElement->Gizmo.Initial;
 		}
 
 		if(FRigTransformElement* TransformElement = Cast<FRigTransformElement>(Elements[ElementIndex]))
 		{
-			if(bHasFilteredChildren)
-			{
-				const FTransform Transform = GetTransform(TransformElement, ERigTransformType::InitialLocal);
-				SetTransform(TransformElement, Transform, ERigTransformType::CurrentLocal, true);
-			}
-			else
-			{
-				TransformElement->Pose.Current = TransformElement->Pose.Initial;
-			}
+			TransformElement->Pose.Current = TransformElement->Pose.Initial;
 		}
 
 		if(FRigMultiParentElement* MultiParentElement = Cast<FRigMultiParentElement>(Elements[ElementIndex]))
 		{
-			if(bHasFilteredChildren)
-			{
-				MultiParentElement->Parent.MarkDirty(ERigTransformType::CurrentGlobal);
-			}
-			else
-			{
-				MultiParentElement->Parent.Current = MultiParentElement->Parent.Initial;
-			}
+			MultiParentElement->Parent.Current = MultiParentElement->Parent.Initial;
 		}
 	}
+	
+	EnsureCacheValidity();
 }
 
 void URigHierarchy::ResetCurveValues()
@@ -1263,6 +1238,7 @@ bool URigHierarchy::SetParentWeight(FRigBaseElement* InChild, int32 InParentInde
 			}
 
 			PropagateDirtyFlags(MultiParentElement, ERigTransformType::IsInitial(LocalType), bAffectChildren);
+			EnsureCacheValidity();
 			
 #if WITH_EDITOR
 			if (ensure(!bPropagatingChange))
@@ -2031,6 +2007,8 @@ FTransform URigHierarchy::GetTransform(FRigTransformElement* InTransformElement,
 				InTransformElement->Pose.Set(InTransformType, NewTransform);
 			}
 		}
+
+		EnsureCacheValidity();
 	}
 	return InTransformElement->Pose.Get(InTransformType);
 }
@@ -2074,6 +2052,8 @@ void URigHierarchy::SetTransform(FRigTransformElement* InTransformElement, const
 		ControlElement->Gizmo.MarkDirty(MakeGlobal(InTransformType));
 	}
 
+	EnsureCacheValidity();
+	
 #if WITH_EDITOR
 	if(bSetupUndo || IsTracingChanges())
 	{
@@ -2184,6 +2164,8 @@ FTransform URigHierarchy::GetControlOffsetTransform(FRigControlElement* InContro
 				FTransform::Identity, false);
 			InControlElement->Offset.Set(InTransformType, GlobalTransform);
 		}
+
+		EnsureCacheValidity();
 	}
 	return InControlElement->Offset.Get(InTransformType);
 }
@@ -2216,6 +2198,8 @@ void URigHierarchy::SetControlOffsetTransform(FRigControlElement* InControlEleme
 	InControlElement->Offset.MarkDirty(OpposedType);
 	InControlElement->Gizmo.MarkDirty(MakeGlobal(InTransformType));
 
+	EnsureCacheValidity();
+
 	if (ERigTransformType::IsInitial(InTransformType))
 	{
 		// control's offset transform is considered a special type of transform
@@ -2306,6 +2290,8 @@ FTransform URigHierarchy::GetControlGizmoTransform(FRigControlElement* InControl
 		{
 			InControlElement->Gizmo.Set(InTransformType, InControlElement->Gizmo.Get(OpposedType) * ParentTransform);
 		}
+
+		EnsureCacheValidity();
 	}
 	return InControlElement->Gizmo.Get(InTransformType);
 }
@@ -2340,6 +2326,8 @@ void URigHierarchy::SetControlGizmoTransform(FRigControlElement* InControlElemen
 		SetControlGizmoTransform(InControlElement, InTransform, ERigTransformType::MakeCurrent(InTransformType), false, bForce);
 	}
 	
+	EnsureCacheValidity();
+	
 #if WITH_EDITOR
 	if(bSetupUndo || IsTracingChanges())
 	{
@@ -2496,6 +2484,8 @@ FTransform URigHierarchy::GetParentTransform(FRigBaseElement* InElement, const E
 				false
 			);
 			MultiParentElement->Parent.Set(InTransformType, OutputTransform);
+
+			EnsureCacheValidity();
 		}
 		return Output.Transform;
 	}
@@ -3413,6 +3403,213 @@ void URigHierarchy::PropagateDirtyFlags(FRigTransformElement* InTransformElement
 	}
 }
 
+void URigHierarchy::EnsureCacheValidityImpl()
+{
+	if(!bEnableCacheValidityCheck)
+	{
+		return;
+	}
+	TGuardValue<bool> Guard(bEnableCacheValidityCheck, false);
+
+	static TArray<FString> TransformTypeStrings;
+	if(TransformTypeStrings.IsEmpty())
+	{
+		for(int32 TransformTypeIndex = 0; TransformTypeIndex < (int32) ERigTransformType::NumTransformTypes; TransformTypeIndex++)
+		{
+			TransformTypeStrings.Add(StaticEnum<ERigTransformType::Type>()->GetDisplayNameTextByValue((int64)TransformTypeIndex).ToString());
+		}
+	}
+
+	// make sure that elements which are marked as dirty don't have fully cached children
+	ForEach<FRigTransformElement>([](FRigTransformElement* TransformElement)
+    {
+		for(int32 TransformTypeIndex = 0; TransformTypeIndex < (int32) ERigTransformType::NumTransformTypes; TransformTypeIndex++)
+		{
+			const ERigTransformType::Type GlobalType = (ERigTransformType::Type)TransformTypeIndex;
+			const ERigTransformType::Type LocalType = ERigTransformType::SwapLocalAndGlobal(GlobalType);
+			const FString& TransformTypeString = TransformTypeStrings[TransformTypeIndex];
+
+			if(ERigTransformType::IsLocal(GlobalType))
+			{
+				continue;
+			}
+
+			if(!TransformElement->Pose.IsDirty(GlobalType))
+			{
+				continue;
+			}
+
+			for(const FRigTransformElement::FElementToDirty& ElementToDirty : TransformElement->ElementsToDirty)
+			{
+				if(FRigMultiParentElement* MultiParentElementToDirty = Cast<FRigMultiParentElement>(ElementToDirty.Element))
+				{
+					check(MultiParentElementToDirty->Parent.IsDirty(GlobalType) ||
+                        MultiParentElementToDirty->Parent.IsDirty(LocalType));
+
+                    if(FRigControlElement* ControlElementToDirty = Cast<FRigControlElement>(ElementToDirty.Element))
+                    {
+                        if(ControlElementToDirty->Parent.IsDirty(GlobalType))
+                        {
+                            checkf(ControlElementToDirty->Offset.IsDirty(GlobalType) ||
+                                    ControlElementToDirty->Offset.IsDirty(LocalType),
+                                    TEXT("Control '%s' %s Parent Cache is dirty, but the Offset is not."),
+                                    *ControlElementToDirty->GetKey().ToString(),
+                                    *TransformTypeString);
+                        }
+
+                        if(ControlElementToDirty->Offset.IsDirty(GlobalType))
+                        {
+                            checkf(ControlElementToDirty->Pose.IsDirty(GlobalType) ||
+                                    ControlElementToDirty->Pose.IsDirty(LocalType),
+                                    TEXT("Control '%s' %s Offset Cache is dirty, but the Pose is not."),
+									*ControlElementToDirty->GetKey().ToString(),
+									*TransformTypeString);
+						}
+
+                        if(ControlElementToDirty->Pose.IsDirty(GlobalType))
+                        {
+                            checkf(ControlElementToDirty->Gizmo.IsDirty(GlobalType) ||
+                                    ControlElementToDirty->Gizmo.IsDirty(LocalType),
+                                    TEXT("Control '%s' %s Pose Cache is dirty, but the Gizmo is not."),
+									*ControlElementToDirty->GetKey().ToString(),
+									*TransformTypeString);
+						}
+                    }
+                    else
+                    {
+                        if(MultiParentElementToDirty->Parent.IsDirty(GlobalType))
+                        {
+                            checkf(MultiParentElementToDirty->Pose.IsDirty(GlobalType) ||
+                                    MultiParentElementToDirty->Pose.IsDirty(LocalType),
+                                    TEXT("MultiParent '%s' %s Parent Cache is dirty, but the Pose is not."),
+									*MultiParentElementToDirty->GetKey().ToString(),
+									*TransformTypeString);
+						}
+                    }
+				}
+				else
+				{
+					checkf(ElementToDirty.Element->Pose.IsDirty(GlobalType) ||
+						ElementToDirty.Element->Pose.IsDirty(LocalType),
+						TEXT("SingleParent '%s' %s Pose is not dirty in Local or Global"),
+						*ElementToDirty.Element->GetKey().ToString(),
+						*TransformTypeString);
+				}
+			}
+		}
+		
+        return true;
+    });
+
+	// store our own pose in a transient hierarchy used for cache validation
+	if(HierarchyForCacheValidation == nullptr)
+	{
+		HierarchyForCacheValidation = NewObject<URigHierarchy>(this, NAME_None, RF_Transient);
+		HierarchyForCacheValidation->bEnableCacheValidityCheck = false;
+	}
+	if(HierarchyForCacheValidation->GetTopologyVersion() != GetTopologyVersion())
+	{
+		HierarchyForCacheValidation->CopyHierarchy(this);
+	}
+	HierarchyForCacheValidation->CopyPose(this, true, true);
+
+	// traverse the copied hierarchy and compare cached vs computed values
+	URigHierarchy* HierarchyForLambda = HierarchyForCacheValidation;
+	HierarchyForLambda->Traverse([HierarchyForLambda](FRigBaseElement* Element, bool& bContinue)
+	{
+		bContinue = true;
+
+		if(FRigMultiParentElement* MultiParentElement = Cast<FRigMultiParentElement>(Element))
+		{
+			for(int32 TransformTypeIndex = 0; TransformTypeIndex < (int32) ERigTransformType::NumTransformTypes; TransformTypeIndex++)
+			{
+				const ERigTransformType::Type TransformType = (ERigTransformType::Type)TransformTypeIndex;
+				const ERigTransformType::Type OpposedType = ERigTransformType::SwapLocalAndGlobal(TransformType);
+				const FString& TransformTypeString = TransformTypeStrings[TransformTypeIndex];
+
+				if(ERigTransformType::IsLocal(TransformType))
+				{
+					continue;
+				}
+
+				if(!MultiParentElement->Parent.IsDirty(TransformType))
+				{
+					const FTransform CachedTransform = HierarchyForLambda->GetParentTransform(MultiParentElement, TransformType);
+					MultiParentElement->Parent.MarkDirty(TransformType);
+					const FTransform ComputedTransform = HierarchyForLambda->GetParentTransform(MultiParentElement, TransformType);
+					checkf(FRigComputedTransform::Equals(CachedTransform, ComputedTransform),
+						TEXT("Element '%s' Parent %s Cached vs Computed doesn't match."),
+						*Element->GetName().ToString(),
+						*TransformTypeString);
+				}
+			}
+		}
+
+		if(FRigControlElement* ControlElement = Cast<FRigControlElement>(Element))
+		{
+			for(int32 TransformTypeIndex = 0; TransformTypeIndex < (int32) ERigTransformType::NumTransformTypes; TransformTypeIndex++)
+			{
+				const ERigTransformType::Type TransformType = (ERigTransformType::Type)TransformTypeIndex;
+				const ERigTransformType::Type OpposedType = ERigTransformType::SwapLocalAndGlobal(TransformType);
+				const FString& TransformTypeString = TransformTypeStrings[TransformTypeIndex];
+
+				if(!ControlElement->Offset.IsDirty(TransformType) && !ControlElement->Offset.IsDirty(OpposedType))
+				{
+					const FTransform CachedTransform = HierarchyForLambda->GetControlOffsetTransform(ControlElement, TransformType);
+					ControlElement->Offset.MarkDirty(TransformType);
+					const FTransform ComputedTransform = HierarchyForLambda->GetControlOffsetTransform(ControlElement, TransformType);
+					checkf(FRigComputedTransform::Equals(CachedTransform, ComputedTransform),
+						TEXT("Element '%s' Offset %s Cached vs Computed doesn't match."),
+						*Element->GetName().ToString(),
+						*TransformTypeString);
+				}
+			}
+		}
+
+		if(FRigTransformElement* TransformElement = Cast<FRigTransformElement>(Element))
+		{
+			for(int32 TransformTypeIndex = 0; TransformTypeIndex < (int32) ERigTransformType::NumTransformTypes; TransformTypeIndex++)
+			{
+				const ERigTransformType::Type TransformType = (ERigTransformType::Type)TransformTypeIndex;
+				const ERigTransformType::Type OpposedType = ERigTransformType::SwapLocalAndGlobal(TransformType);
+				const FString& TransformTypeString = TransformTypeStrings[TransformTypeIndex];
+
+				if(!TransformElement->Pose.IsDirty(TransformType) && !TransformElement->Pose.IsDirty(OpposedType))
+				{
+					const FTransform CachedTransform = HierarchyForLambda->GetTransform(TransformElement, TransformType);
+					TransformElement->Pose.MarkDirty(TransformType);
+					const FTransform ComputedTransform = HierarchyForLambda->GetTransform(TransformElement, TransformType);
+					checkf(FRigComputedTransform::Equals(CachedTransform, ComputedTransform),
+						TEXT("Element '%s' Pose %s Cached vs Computed doesn't match."),
+						*Element->GetName().ToString(),
+						*TransformTypeString);
+				}
+			}
+		}
+
+		if(FRigControlElement* ControlElement = Cast<FRigControlElement>(Element))
+		{
+			for(int32 TransformTypeIndex = 0; TransformTypeIndex < (int32) ERigTransformType::NumTransformTypes; TransformTypeIndex++)
+			{
+				const ERigTransformType::Type TransformType = (ERigTransformType::Type)TransformTypeIndex;
+				const ERigTransformType::Type OpposedType = ERigTransformType::SwapLocalAndGlobal(TransformType);
+				const FString& TransformTypeString = TransformTypeStrings[TransformTypeIndex];
+
+				if(!ControlElement->Gizmo.IsDirty(TransformType) && !ControlElement->Gizmo.IsDirty(OpposedType))
+				{
+					const FTransform CachedTransform = HierarchyForLambda->GetControlGizmoTransform(ControlElement, TransformType);
+					ControlElement->Gizmo.MarkDirty(TransformType);
+					const FTransform ComputedTransform = HierarchyForLambda->GetControlGizmoTransform(ControlElement, TransformType);
+					checkf(FRigComputedTransform::Equals(CachedTransform, ComputedTransform),
+						TEXT("Element '%s' Gizmo %s Cached vs Computed doesn't match."),
+						*Element->GetName().ToString(),
+						*TransformTypeString);
+				}
+			}
+		}
+	});
+}
+
 void URigHierarchy::PushTransformToStack(const FRigElementKey& InKey, ERigTransformStackEntryType InEntryType,
 	ERigTransformType::Type InTransformType, const FTransform& InOldTransform, const FTransform& InNewTransform,
 	bool bAffectChildren, bool bModify)
diff --git a/Engine/Plugins/Animation/ControlRig/Source/ControlRig/Private/Rigs/RigHierarchyController.cpp b/Engine/Plugins/Animation/ControlRig/Source/ControlRig/Private/Rigs/RigHierarchyController.cpp
--- a/Engine/Plugins/Animation/ControlRig/Source/ControlRig/Private/Rigs/RigHierarchyController.cpp
+++ b/Engine/Plugins/Animation/ControlRig/Source/ControlRig/Private/Rigs/RigHierarchyController.cpp
@@ -160,21 +160,24 @@ FRigElementKey URigHierarchyController::AddBone(FName InName, FRigElementKey InP
 #endif
 
 	FRigBoneElement* NewElement = new FRigBoneElement();
-	NewElement->Key.Type = ERigElementType::Bone;
-	NewElement->Key.Name = Hierarchy->GetSafeNewName(InName.ToString(), NewElement->Key.Type);
-	NewElement->BoneType = InBoneType;
-	AddElement(NewElement, Hierarchy->Get(Hierarchy->GetIndex(InParent)), true);
-
-	if(bTransformInGlobal)
-	{
-		Hierarchy->SetTransform(NewElement, InTransform, ERigTransformType::InitialGlobal, true, false);
-	}
-	else
 	{
-		Hierarchy->SetTransform(NewElement, InTransform, ERigTransformType::InitialLocal, true, false);
-	}
+		TGuardValue<bool> DisableCacheValidityChecks(Hierarchy->bEnableCacheValidityCheck, false);
+		NewElement->Key.Type = ERigElementType::Bone;
+		NewElement->Key.Name = Hierarchy->GetSafeNewName(InName.ToString(), NewElement->Key.Type);
+		NewElement->BoneType = InBoneType;
+		AddElement(NewElement, Hierarchy->Get(Hierarchy->GetIndex(InParent)), true);
 
-	NewElement->Pose.Current = NewElement->Pose.Initial;
+		if(bTransformInGlobal)
+		{
+			Hierarchy->SetTransform(NewElement, InTransform, ERigTransformType::InitialGlobal, true, false);
+		}
+		else
+		{
+			Hierarchy->SetTransform(NewElement, InTransform, ERigTransformType::InitialLocal, true, false);
+		}
+
+		NewElement->Pose.Current = NewElement->Pose.Initial;
+	}
 
 #if WITH_EDITOR
 	TransactionPtr.Reset();
@@ -194,6 +197,8 @@ FRigElementKey URigHierarchyController::AddBone(FName InName, FRigElementKey InP
 	}
 #endif
 
+	Hierarchy->EnsureCacheValidity();
+		
 	return NewElement->Key;
 }
 
@@ -215,22 +220,25 @@ FRigElementKey URigHierarchyController::AddNull(FName InName, FRigElementKey InP
 #endif
 
 	FRigNullElement* NewElement = new FRigNullElement();
-	NewElement->Key.Type = ERigElementType::Null;
-	NewElement->Key.Name = Hierarchy->GetSafeNewName(InName.ToString(), NewElement->Key.Type);
-	AddElement(NewElement, Hierarchy->Get(Hierarchy->GetIndex(InParent)), false);
-
-	if(bTransformInGlobal)
-	{
-		Hierarchy->SetTransform(NewElement, InTransform, ERigTransformType::InitialGlobal, true, false);
-	}
-	else
 	{
-		Hierarchy->SetTransform(NewElement, InTransform, ERigTransformType::InitialLocal, true, false);
-	}
+		TGuardValue<bool> DisableCacheValidityChecks(Hierarchy->bEnableCacheValidityCheck, false);		
+		NewElement->Key.Type = ERigElementType::Null;
+		NewElement->Key.Name = Hierarchy->GetSafeNewName(InName.ToString(), NewElement->Key.Type);
+		AddElement(NewElement, Hierarchy->Get(Hierarchy->GetIndex(InParent)), false);
 
-	NewElement->Parent.MarkDirty(ERigTransformType::InitialGlobal);
-	NewElement->Parent.Current = NewElement->Parent.Initial;
-	NewElement->Pose.Current = NewElement->Pose.Initial;
+		if(bTransformInGlobal)
+		{
+			Hierarchy->SetTransform(NewElement, InTransform, ERigTransformType::InitialGlobal, true, false);
+		}
+		else
+		{
+			Hierarchy->SetTransform(NewElement, InTransform, ERigTransformType::InitialLocal, true, false);
+		}
+
+		NewElement->Parent.MarkDirty(ERigTransformType::InitialGlobal);
+		NewElement->Parent.Current = NewElement->Parent.Initial;
+		NewElement->Pose.Current = NewElement->Pose.Initial;
+	}
 
 #if WITH_EDITOR
 	TransactionPtr.Reset();
@@ -250,6 +258,8 @@ FRigElementKey URigHierarchyController::AddNull(FName InName, FRigElementKey InP
 	}
 #endif
 
+	Hierarchy->EnsureCacheValidity();
+
 	return NewElement->Key;
 }
 
@@ -277,23 +287,26 @@ FRigElementKey URigHierarchyController::AddControl(
 #endif
 
 	FRigControlElement* NewElement = new FRigControlElement();
-	NewElement->Key.Type = ERigElementType::Control;
-	NewElement->Key.Name = Hierarchy->GetSafeNewName(InName.ToString(), NewElement->Key.Type);
-	NewElement->Settings = InSettings;
-	AddElement(NewElement, Hierarchy->Get(Hierarchy->GetIndex(InParent)), false);
-	
-	NewElement->Offset.Set(ERigTransformType::InitialLocal, InOffsetTransform);  
-	NewElement->Gizmo.Set(ERigTransformType::InitialLocal, InGizmoTransform);  
-	Hierarchy->SetControlValue(NewElement, InValue, ERigControlValueType::Initial, false);
-
-	NewElement->Parent.MarkDirty(ERigTransformType::InitialGlobal);
-	NewElement->Offset.MarkDirty(ERigTransformType::InitialGlobal);
-	NewElement->Pose.MarkDirty(ERigTransformType::InitialGlobal);
-	NewElement->Gizmo.MarkDirty(ERigTransformType::InitialGlobal);
-	NewElement->Parent.Current = NewElement->Parent.Initial;
-	NewElement->Offset.Current = NewElement->Offset.Initial;
-	NewElement->Pose.Current = NewElement->Pose.Initial;
-	NewElement->Gizmo.Current = NewElement->Gizmo.Initial;
+	{
+		TGuardValue<bool> DisableCacheValidityChecks(Hierarchy->bEnableCacheValidityCheck, false);		
+		NewElement->Key.Type = ERigElementType::Control;
+		NewElement->Key.Name = Hierarchy->GetSafeNewName(InName.ToString(), NewElement->Key.Type);
+		NewElement->Settings = InSettings;
+		AddElement(NewElement, Hierarchy->Get(Hierarchy->GetIndex(InParent)), false);
+		
+		NewElement->Offset.Set(ERigTransformType::InitialLocal, InOffsetTransform);  
+		NewElement->Gizmo.Set(ERigTransformType::InitialLocal, InGizmoTransform);  
+		Hierarchy->SetControlValue(NewElement, InValue, ERigControlValueType::Initial, false);
+
+		NewElement->Parent.MarkDirty(ERigTransformType::InitialGlobal);
+		NewElement->Offset.MarkDirty(ERigTransformType::InitialGlobal);
+		NewElement->Pose.MarkDirty(ERigTransformType::InitialGlobal);
+		NewElement->Gizmo.MarkDirty(ERigTransformType::InitialGlobal);
+		NewElement->Parent.Current = NewElement->Parent.Initial;
+		NewElement->Offset.Current = NewElement->Offset.Initial;
+		NewElement->Pose.Current = NewElement->Pose.Initial;
+		NewElement->Gizmo.Current = NewElement->Gizmo.Initial;
+	}
 
 #if WITH_EDITOR
 	TransactionPtr.Reset();
@@ -313,6 +326,8 @@ FRigElementKey URigHierarchyController::AddControl(
 	}
 #endif
 
+	Hierarchy->EnsureCacheValidity();
+
 	return NewElement->Key;
 }
 
@@ -333,10 +348,13 @@ FRigElementKey URigHierarchyController::AddCurve(FName InName, float InValue, bo
 #endif
 
 	FRigCurveElement* NewElement = new FRigCurveElement();
-	NewElement->Key.Type = ERigElementType::Curve;
-	NewElement->Key.Name = Hierarchy->GetSafeNewName(InName.ToString(), NewElement->Key.Type);
-	NewElement->Value = InValue;
-	AddElement(NewElement, nullptr, false);
+	{
+		TGuardValue<bool> DisableCacheValidityChecks(Hierarchy->bEnableCacheValidityCheck, false);		
+		NewElement->Key.Type = ERigElementType::Curve;
+		NewElement->Key.Name = Hierarchy->GetSafeNewName(InName.ToString(), NewElement->Key.Type);
+		NewElement->Value = InValue;
+		AddElement(NewElement, nullptr, false);
+	}
 
 #if WITH_EDITOR
 	TransactionPtr.Reset();
@@ -356,6 +374,8 @@ FRigElementKey URigHierarchyController::AddCurve(FName InName, float InValue, bo
 	}
 #endif
 
+	Hierarchy->EnsureCacheValidity();
+
 	return NewElement->Key;
 }
 
@@ -377,14 +397,17 @@ FRigElementKey URigHierarchyController::AddRigidBody(FName InName, FRigElementKe
 #endif
 
 	FRigRigidBodyElement* NewElement = new FRigRigidBodyElement();
-	NewElement->Key.Type = ERigElementType::RigidBody;
-	NewElement->Key.Name = Hierarchy->GetSafeNewName(InName.ToString(), NewElement->Key.Type);
-	NewElement->Settings = InSettings;
-	AddElement(NewElement, Hierarchy->Get(Hierarchy->GetIndex(InParent)), true);
-
-	Hierarchy->SetTransform(NewElement, InLocalTransform, ERigTransformType::InitialLocal, true, false);
-	NewElement->Pose.Current = NewElement->Pose.Initial;
+	{
+		TGuardValue<bool> DisableCacheValidityChecks(Hierarchy->bEnableCacheValidityCheck, false);		
+		NewElement->Key.Type = ERigElementType::RigidBody;
+		NewElement->Key.Name = Hierarchy->GetSafeNewName(InName.ToString(), NewElement->Key.Type);
+		NewElement->Settings = InSettings;
+		AddElement(NewElement, Hierarchy->Get(Hierarchy->GetIndex(InParent)), true);
 
+		Hierarchy->SetTransform(NewElement, InLocalTransform, ERigTransformType::InitialLocal, true, false);
+		NewElement->Pose.Current = NewElement->Pose.Initial;
+	}
+	
 #if WITH_EDITOR
 	TransactionPtr.Reset();
 
@@ -403,6 +426,8 @@ FRigElementKey URigHierarchyController::AddRigidBody(FName InName, FRigElementKe
 	}
 #endif
 
+	Hierarchy->EnsureCacheValidity();
+
 	return NewElement->Key;
 }
 
@@ -424,18 +449,23 @@ FRigElementKey URigHierarchyController::AddSocket(FName InName, FRigElementKey I
 #endif
 
 	FRigSocketElement* NewElement = new FRigSocketElement();
-	NewElement->Key.Type = ERigElementType::Socket;
-	NewElement->Key.Name = Hierarchy->GetSafeNewName(InName.ToString(), NewElement->Key.Type);
-	NewElement->GetWorldTransformDelegate = InDelegate;
-	AddElement(NewElement, Hierarchy->Get(Hierarchy->GetIndex(InParent)), true);
+	{
+		TGuardValue<bool> DisableCacheValidityChecks(Hierarchy->bEnableCacheValidityCheck, false);		
+		NewElement->Key.Type = ERigElementType::Socket;
+		NewElement->Key.Name = Hierarchy->GetSafeNewName(InName.ToString(), NewElement->Key.Type);
+		NewElement->GetWorldTransformDelegate = InDelegate;
+		AddElement(NewElement, Hierarchy->Get(Hierarchy->GetIndex(InParent)), true);
 
-	Hierarchy->SetTransform(NewElement, FTransform::Identity, ERigTransformType::InitialLocal, true, false);
-	NewElement->Pose.Current = NewElement->Pose.Initial;
+		Hierarchy->SetTransform(NewElement, FTransform::Identity, ERigTransformType::InitialLocal, true, false);
+		NewElement->Pose.Current = NewElement->Pose.Initial;
+	}
 
 #if WITH_EDITOR
 	TransactionPtr.Reset();
 #endif
 
+	Hierarchy->EnsureCacheValidity();
+
 	return NewElement->Key;
 }
 
@@ -494,6 +524,8 @@ bool URigHierarchyController::SetControlSettings(FRigElementKey InKey, FRigContr
     TransactionPtr.Reset();
 #endif
 
+	Hierarchy->EnsureCacheValidity();
+	
 	return true;
 }
 
@@ -661,6 +693,8 @@ TArray<FRigElementKey> URigHierarchyController::ImportBones(const FReferenceSkel
 		SetSelection(BonesToSelect);
 	}
 
+	Hierarchy->EnsureCacheValidity();
+
 	return AddedBones;
 }
 
@@ -805,6 +839,8 @@ TArray<FRigElementKey> URigHierarchyController::ImportCurves(USkeleton* InSkelet
 	}
 #endif
 
+	Hierarchy->EnsureCacheValidity();
+
 	return Keys;
 }
 
@@ -1148,6 +1184,8 @@ TArray<FRigElementKey> URigHierarchyController::ImportFromText(FString InContent
 	}
 #endif
 
+	Hierarchy->EnsureCacheValidity();
+
 	return PastedKeys;
 }
 
@@ -1267,6 +1305,8 @@ TArray<FRigElementKey> URigHierarchyController::ImportFromHierarchyContainer(con
 	}
 #endif
 
+	Hierarchy->EnsureCacheValidity();
+
 	TArray<FRigElementKey> AddedKeys;
 	KeyMap.GenerateValueArray(AddedKeys);
 	return AddedKeys;
@@ -1570,7 +1610,7 @@ bool URigHierarchyController::RemoveElement(FRigElementKey InElement, bool bSetu
 		}
 	}
 #endif
-	
+
 	return bRemoved;
 }
 
@@ -1684,6 +1724,8 @@ bool URigHierarchyController::RemoveElement(FRigBaseElement* InElement)
 
 	delete InElement;
 
+	Hierarchy->EnsureCacheValidity();
+
 	return NumElementsRemoved == 1;
 }
 
@@ -1942,6 +1984,9 @@ bool URigHierarchyController::AddParent(FRigBaseElement* InChild, FRigBaseElemen
 		}
 
 		Notify(ERigHierarchyNotification::ParentChanged, SingleParentElement);
+		
+		Hierarchy->EnsureCacheValidity();
+		
 		return true;
 	}
 	else if(FRigMultiParentElement* MultiParentElement = Cast<FRigMultiParentElement>(InChild))
@@ -1977,6 +2022,9 @@ bool URigHierarchyController::AddParent(FRigBaseElement* InChild, FRigBaseElemen
 		}
 
 		Notify(ERigHierarchyNotification::ParentChanged, MultiParentElement);
+
+		Hierarchy->EnsureCacheValidity();
+		
 		return true;
 	}
 	
@@ -2088,6 +2136,8 @@ bool URigHierarchyController::RemoveParent(FRigBaseElement* InChild, FRigBaseEle
 
 			Notify(ERigHierarchyNotification::ParentChanged, SingleParentElement);
 
+			Hierarchy->EnsureCacheValidity();
+			
 			return true;
 		}
 	}
@@ -2159,6 +2209,8 @@ bool URigHierarchyController::RemoveParent(FRigBaseElement* InChild, FRigBaseEle
 
 			Notify(ERigHierarchyNotification::ParentChanged, MultiParentElement);
 
+			Hierarchy->EnsureCacheValidity();
+			
 			return true;
 		}
 	}
@@ -2324,6 +2376,8 @@ TArray<FRigElementKey> URigHierarchyController::DuplicateElements(TArray<FRigEle
 	}
 #endif
 
+	Hierarchy->EnsureCacheValidity();
+	
 	return Result;
 }
 
@@ -2447,6 +2501,8 @@ TArray<FRigElementKey> URigHierarchyController::MirrorElements(TArray<FRigElemen
 		}
 	}
 #endif
+
+	Hierarchy->EnsureCacheValidity();
 	
 	return DuplicatedKeys;
 }
diff --git a/Engine/Plugins/Animation/ControlRig/Source/ControlRig/Public/Rigs/RigHierarchy.h b/Engine/Plugins/Animation/ControlRig/Source/ControlRig/Public/Rigs/RigHierarchy.h
--- a/Engine/Plugins/Animation/ControlRig/Source/ControlRig/Public/Rigs/RigHierarchy.h
+++ b/Engine/Plugins/Animation/ControlRig/Source/ControlRig/Public/Rigs/RigHierarchy.h
@@ -15,8 +15,15 @@
 class URigHierarchy;
 class URigHierarchyController;
 
+// Define to switch between recursive calls for dirty propagation or flat iteration.
+// Based on this each element will either contains only the next tier of elements to dirty (recursive)
+// or a full flattened list of elements to dirty (no need to recurse)
 #define URIGHIERARCHY_RECURSIVE_DIRTY_PROPAGATION 1
 
+// Debug define which performs a full check on the cache validity for all elements of the hierarchy.
+// This can be useful for debugging cache validity bugs.
+#define URIGHIERARCHY_ENSURE_CACHE_VALIDITY 1
+
 DECLARE_MULTICAST_DELEGATE_ThreeParams(FRigHierarchyModifiedEvent, ERigHierarchyNotification /* type */, URigHierarchy* /* hierarchy */, const FRigBaseElement* /* element */);
 DECLARE_EVENT_FiveParams(URigHierarchy, FRigHierarchyUndoRedoTransformEvent, URigHierarchy*, const FRigElementKey&, ERigTransformType::Type, const FTransform&, bool /* bUndo */);
 
@@ -2497,6 +2504,20 @@ private:
 	void PropagateDirtyFlags(FRigTransformElement* InTransformElement, bool bInitial, bool bAffectChildren) const;
 #endif
 
+	/**
+	* Performs validation of the cache within the hierarchy on any mutation.
+	*/
+	FORCEINLINE void EnsureCacheValidity() const
+	{
+#if WITH_EDITOR
+		if(bEnableCacheValidityCheck)
+		{
+			URigHierarchy* MutableThis = (URigHierarchy*)this; 
+			MutableThis->EnsureCacheValidityImpl();
+		}
+#endif
+	}
+	
 	/**
 	 * The topology version of the hierarchy changes when elements are
 	 * added, removed, re-parented or renamed.
@@ -2614,7 +2635,7 @@ private:
 	TMap<FRigElementKey, FRigElementKey> PreviousNameMap;
 
 	int32 ResetPoseHash;
-	TArray<bool> ResetPoseHasFilteredChildren;
+	TArray<bool> ResetPoseIsFilteredOut;
 
 #if WITH_EDITOR
 
@@ -2843,9 +2864,52 @@ protected:
 	static TArray<FString> ControlSettingsToPythonCommands(const FRigControlSettings& Settings, const FString& NameSettings);
 #endif
 	
+protected:
+	
+	bool bEnableCacheValidityCheck;
+
+	UPROPERTY(transient)
+	TObjectPtr<URigHierarchy> HierarchyForCacheValidation;
+	
+private:
+	
+	void EnsureCacheValidityImpl();
+	
 	friend class URigHierarchyController;
 	friend class UControlRig;
 	friend class FControlRigEditor;
+	friend struct FRigHierarchyValidityBracket;
+};
+
+struct CONTROLRIG_API FRigHierarchyValidityBracket
+{
+	public:
+	FRigHierarchyValidityBracket(URigHierarchy* InHierarchy)
+		: bPreviousValue(false)
+		, HierarchyPtr() 
+	{
+		if(InHierarchy)
+		{
+			bPreviousValue = InHierarchy->bEnableCacheValidityCheck;
+			InHierarchy->bEnableCacheValidityCheck = false;
+			HierarchyPtr = InHierarchy;
+		}
+	}
+
+	~FRigHierarchyValidityBracket()
+	{
+		if(HierarchyPtr.IsValid())
+		{
+			URigHierarchy* Hierarchy = HierarchyPtr.Get();
+			Hierarchy->bEnableCacheValidityCheck = bPreviousValue;
+			Hierarchy->EnsureCacheValidity();
+		}
+	}
+
+	private:
+
+	bool bPreviousValue;
+	TWeakObjectPtr<URigHierarchy> HierarchyPtr;
 };
 
 #if !UE_LARGE_WORLD_COORDINATES_DISABLED
diff --git a/Engine/Plugins/Animation/RigLogic/Source/RigLogicModule/Private/RigUnit_RigLogic.cpp b/Engine/Plugins/Animation/RigLogic/Source/RigLogicModule/Private/RigUnit_RigLogic.cpp
--- a/Engine/Plugins/Animation/RigLogic/Source/RigLogicModule/Private/RigUnit_RigLogic.cpp
+++ b/Engine/Plugins/Animation/RigLogic/Source/RigLogicModule/Private/RigUnit_RigLogic.cpp
@@ -240,6 +240,7 @@ void FRigUnit_RigLogic_Data::CalculateRigLogic(const URigHierarchy* InHierarchy)
 
 void FRigUnit_RigLogic_Data::UpdateJoints(URigHierarchy* Hierarchy, const FRigUnit_RigLogic_JointUpdateParams& JointUpdateParams)
 {
+	FRigHierarchyValidityBracket ValidityBracket(Hierarchy);
 	for (const uint16 JointIndex : SharedRigRuntimeContext->VariableJointIndices[CurrentLOD].Values)
 	{
 		const int32 BoneIndex = HierarchyBoneIndices[JointIndex];
