Control Rig: Disable hierarchy validity checker by default

diff --git a/Engine/Plugins/Animation/ControlRig/Source/ControlRig/Public/Rigs/RigHierarchy.h b/Engine/Plugins/Animation/ControlRig/Source/ControlRig/Public/Rigs/RigHierarchy.h
--- a/Engine/Plugins/Animation/ControlRig/Source/ControlRig/Public/Rigs/RigHierarchy.h
+++ b/Engine/Plugins/Animation/ControlRig/Source/ControlRig/Public/Rigs/RigHierarchy.h
@@ -22,7 +22,7 @@ class URigHierarchyController;
 
 // Debug define which performs a full check on the cache validity for all elements of the hierarchy.
 // This can be useful for debugging cache validity bugs.
-#define URIGHIERARCHY_ENSURE_CACHE_VALIDITY 1
+#define URIGHIERARCHY_ENSURE_CACHE_VALIDITY 0
 
 DECLARE_MULTICAST_DELEGATE_ThreeParams(FRigHierarchyModifiedEvent, ERigHierarchyNotification /* type */, URigHierarchy* /* hierarchy */, const FRigBaseElement* /* element */);
 DECLARE_EVENT_FiveParams(URigHierarchy, FRigHierarchyUndoRedoTransformEvent, URigHierarchy*, const FRigElementKey&, ERigTransformType::Type, const FTransform&, bool /* bUndo */);
