Control Rig: Fix not able to debug Control Rig when used in AnimBP in Editor

diff --git a/Engine/Plugins/Animation/ControlRig/Source/ControlRigEditor/Private/Editor/ControlRigEditor.cpp b/Engine/Plugins/Animation/ControlRig/Source/ControlRigEditor/Private/Editor/ControlRigEditor.cpp
--- a/Engine/Plugins/Animation/ControlRig/Source/ControlRigEditor/Private/Editor/ControlRigEditor.cpp
+++ b/Engine/Plugins/Animation/ControlRig/Source/ControlRigEditor/Private/Editor/ControlRigEditor.cpp
@@ -1148,11 +1148,6 @@ void FControlRigEditor::GetCustomDebugObjects(TArray<FCustomDebugObject>& DebugL
 						continue;
 					}
 
-					if (!World->IsGameWorld() && !World->IsPreviewWorld())
-					{
-						continue;
-					}
-
 					// ensure to only allow preview actors in preview worlds
 					if (World->IsPreviewWorld())
 					{
