Sequencer: Space Switching UX For curve Editor

diff --git a/Engine/Plugins/Animation/ControlRig/Source/ControlRig/Private/Sequencer/MovieSceneControlRigSpaceChannel.cpp b/Engine/Plugins/Animation/ControlRig/Source/ControlRig/Private/Sequencer/MovieSceneControlRigSpaceChannel.cpp
--- a/Engine/Plugins/Animation/ControlRig/Source/ControlRig/Private/Sequencer/MovieSceneControlRigSpaceChannel.cpp
+++ b/Engine/Plugins/Animation/ControlRig/Source/ControlRig/Private/Sequencer/MovieSceneControlRigSpaceChannel.cpp
@@ -96,13 +96,13 @@ TArray <FSpaceRange> FMovieSceneControlRigSpaceChannel::FindSpaceIntervals()
 		Range.Key = KeyValues[Index];
 		if (Index == KeyTimes.Num() - 1)
 		{
-			Range.Range.SetLowerBoundValue(KeyTimes[Index]);
-			Range.Range.SetUpperBoundValue(KeyTimes[Index]);
+			Range.Range.SetLowerBound(TRangeBound<FFrameNumber>(FFrameNumber(KeyTimes[Index])));
+			Range.Range.SetUpperBound(TRangeBound<FFrameNumber>(FFrameNumber(KeyTimes[Index])));
 		}
 		else
 		{
 			int32 NextIndex = Index;
-			FFrameNumber UpperBound = KeyTimes[Index];
+			FFrameNumber LowerBound = KeyTimes[Index],UpperBound = KeyTimes[Index];
 			while (NextIndex < KeyTimes.Num() -1)
 			{
 				NextIndex = Index + 1;
@@ -117,6 +117,8 @@ TArray <FSpaceRange> FMovieSceneControlRigSpaceChannel::FindSpaceIntervals()
 					NextIndex = KeyTimes.Num();
 				}
 			}
+			Range.Range.SetLowerBound(TRangeBound<FFrameNumber>(LowerBound));
+			Range.Range.SetUpperBound(TRangeBound<FFrameNumber>(UpperBound));
 		}
 		Ranges.Add(Range);
 	}
@@ -135,6 +137,24 @@ TArray <FSpaceRange> FMovieSceneControlRigSpaceChannel::FindSpaceIntervals()
 	return Ranges;
 }
 
+FName FMovieSceneControlRigSpaceBaseKey::GetName() const
+{
+	switch (SpaceType)
+	{
+	case EMovieSceneControlRigSpaceType::Parent:
+		return FName(TEXT("Parent"));
+		break;
+
+	case EMovieSceneControlRigSpaceType::World:
+		return FName(TEXT("World"));
+		break;
+
+	case EMovieSceneControlRigSpaceType::ControlRig:
+		return ControlRigElement.Name;
+		break;
+	};
+	return NAME_None;
+}
 
 /* Is this needed if not remove mz todoo
 TSharedPtr<FStructOnScope> GetKeyStruct(TMovieSceneChannelHandle<FMovieSceneControlRigSpaceChannel> Channel, FKeyHandle InHandle)
diff --git a/Engine/Plugins/Animation/ControlRig/Source/ControlRig/Public/Sequencer/MovieSceneControlRigSpaceChannel.h b/Engine/Plugins/Animation/ControlRig/Source/ControlRig/Public/Sequencer/MovieSceneControlRigSpaceChannel.h
--- a/Engine/Plugins/Animation/ControlRig/Source/ControlRig/Public/Sequencer/MovieSceneControlRigSpaceChannel.h
+++ b/Engine/Plugins/Animation/ControlRig/Source/ControlRig/Public/Sequencer/MovieSceneControlRigSpaceChannel.h
@@ -41,6 +41,8 @@ struct CONTROLRIG_API FMovieSceneControlRigSpaceBaseKey
 		return A.SpaceType != B.SpaceType || (A.SpaceType == EMovieSceneControlRigSpaceType::ControlRig && A.ControlRigElement != B.ControlRigElement);
 	}
 
+	FName GetName() const;
+
 	UPROPERTY(EditAnywhere, Category = "Key")
 	EMovieSceneControlRigSpaceType SpaceType = EMovieSceneControlRigSpaceType::Parent;
 	UPROPERTY(EditAnywhere, Category = "Key")
diff --git a/Engine/Plugins/Animation/ControlRig/Source/ControlRigEditor/Private/ControlRigEditorModule.cpp b/Engine/Plugins/Animation/ControlRig/Source/ControlRigEditor/Private/ControlRigEditorModule.cpp
--- a/Engine/Plugins/Animation/ControlRig/Source/ControlRigEditor/Private/ControlRigEditorModule.cpp
+++ b/Engine/Plugins/Animation/ControlRig/Source/ControlRigEditor/Private/ControlRigEditorModule.cpp
@@ -125,8 +125,9 @@
 #include "Framework/Notifications/NotificationManager.h"
 #include "Widgets/Notifications/SNotificationList.h"
 #include "Tools/SMotionTrailOptions.h"
-
-
+#include "ICurveEditorModule.h"
+#include "Channels/SCurveEditorKeyBarView.h"
+#include "ControlRigSpaceChannelCurveModel.h"
 #include "ControlRigSpaceChannelEditors.h"
 
 #define LOCTEXT_NAMESPACE "ControlRigEditorModule"
@@ -332,6 +333,14 @@ void FControlRigEditorModule::StartupModule()
 	};
 #endif
 
+	ICurveEditorModule& CurveEditorModule = FModuleManager::LoadModuleChecked<ICurveEditorModule>("CurveEditor");
+	FControlRigSpaceChannelCurveModel::ViewID = CurveEditorModule.RegisterView(FOnCreateCurveEditorView::CreateStatic(
+		[](TWeakPtr<FCurveEditor> WeakCurveEditor) -> TSharedRef<SCurveEditorView>
+		{
+			return SNew(SCurveEditorKeyBarView, WeakCurveEditor);
+		}
+	));
+
 	FControlRigBlueprintActions::ExtendSketalMeshToolMenu();
 	ExtendAnimSequenceMenu();
 
@@ -349,6 +358,12 @@ void FControlRigEditorModule::StartupModule()
 
 void FControlRigEditorModule::ShutdownModule()
 {
+
+	if (ICurveEditorModule* CurveEditorModule = FModuleManager::GetModulePtr<ICurveEditorModule>("CurveEditor"))
+	{
+		CurveEditorModule->UnregisterView(FControlRigSpaceChannelCurveModel::ViewID);
+	}
+
 #if WITH_EDITOR
 	if (FSlateApplication::IsInitialized())
 	{
diff --git a/Engine/Plugins/Animation/ControlRig/Source/ControlRigEditor/Private/ControlRigSpaceChannelCurveModel.cpp b/Engine/Plugins/Animation/ControlRig/Source/ControlRigEditor/Private/ControlRigSpaceChannelCurveModel.cpp
new file mode 100644
--- /dev/null
+++ b/Engine/Plugins/Animation/ControlRig/Source/ControlRigEditor/Private/ControlRigSpaceChannelCurveModel.cpp
@@ -0,0 +1,312 @@
+// Copyright Epic Games, Inc. All Rights Reserved.
+
+#include "ControlRigSpaceChannelCurveModel.h"
+#include "HAL/PlatformMath.h"
+#include "Channels/MovieSceneFloatChannel.h"
+#include "MovieSceneSection.h"
+#include "MovieScene.h"
+#include "CurveDrawInfo.h"
+#include "CurveDataAbstraction.h"
+#include "CurveEditor.h"
+#include "CurveEditorScreenSpace.h"
+#include "CurveEditorSnapMetrics.h"
+#include "EditorStyleSet.h"
+#include "SequencerChannelTraits.h"
+#include "SequencerSectionPainter.h"
+#include "ISequencer.h"
+#include "ControlRigSpaceChannelEditors.h"
+
+ECurveEditorViewID FControlRigSpaceChannelCurveModel::ViewID = ECurveEditorViewID::Invalid;
+
+FControlRigSpaceChannelCurveModel::FControlRigSpaceChannelCurveModel(TMovieSceneChannelHandle<FMovieSceneControlRigSpaceChannel> InChannel, UMovieSceneSection* OwningSection, TWeakPtr<ISequencer> InWeakSequencer)
+{
+	ChannelHandle = InChannel;
+	WeakSection = OwningSection;
+	WeakSequencer = InWeakSequencer;
+	SupportedViews = ViewID;
+
+	Color = FSequencerSectionPainter::BlendColor(OwningSection->GetTypedOuter<UMovieSceneTrack>()->GetColorTint());
+}
+
+TArray<FKeyBarCurveModel::FBarRange> FControlRigSpaceChannelCurveModel::FindRanges()
+{
+	TArray<FKeyBarCurveModel::FBarRange> Range;
+	FMovieSceneControlRigSpaceChannel* Channel = ChannelHandle.Get();
+	UMovieSceneSection* Section = WeakSection.Get();
+	if (Channel && Section)
+	{
+		FFrameRate TickResolution = Section->GetTypedOuter<UMovieScene>()->GetTickResolution();
+		TArray<FSpaceRange> SpaceRanges = Channel->FindSpaceIntervals();
+		for (const FSpaceRange& SpaceRange : SpaceRanges)
+		{
+			FKeyBarCurveModel::FBarRange BarRange;
+			double LowerValue = SpaceRange.Range.GetLowerBoundValue() / TickResolution;
+			double UpperValue = SpaceRange.Range.GetUpperBoundValue() / TickResolution;
+
+			BarRange.Range.SetLowerBound(TRangeBound<double>(LowerValue));
+			BarRange.Range.SetUpperBound(TRangeBound<double>(UpperValue));
+
+			BarRange.Name = SpaceRange.Key.GetName();
+			BarRange.Color = FControlRigSpaceChannelHelpers::GetColor(SpaceRange.Key);
+
+			Range.Add(BarRange);
+		}
+	}
+	return  Range;
+}
+
+const void* FControlRigSpaceChannelCurveModel::GetCurve() const
+{
+	return ChannelHandle.Get();
+}
+
+void FControlRigSpaceChannelCurveModel::Modify()
+{
+	if (UMovieSceneSection* Section = WeakSection.Get())
+	{
+		Section->Modify();
+	}
+}
+
+void FControlRigSpaceChannelCurveModel::AddKeys(TArrayView<const FKeyPosition> InKeyPositions, TArrayView<const FKeyAttributes> InKeyAttributes, TArrayView<TOptional<FKeyHandle>>* OutKeyHandles)
+{
+	check(InKeyPositions.Num() == InKeyAttributes.Num() && (!OutKeyHandles || OutKeyHandles->Num() == InKeyPositions.Num()));
+
+	FMovieSceneControlRigSpaceChannel* Channel = ChannelHandle.Get();
+	UMovieSceneSection*      Section = WeakSection.Get();
+	if (Channel && Section)
+	{
+		Section->Modify();
+		TMovieSceneChannelData<FMovieSceneControlRigSpaceBaseKey> ChannelData = Channel->GetData();
+		FFrameRate TickResolution = Section->GetTypedOuter<UMovieScene>()->GetTickResolution();
+
+		for (int32 Index = 0; Index < InKeyPositions.Num(); ++Index)
+		{
+			FKeyPosition   Position = InKeyPositions[Index];
+			FKeyAttributes Attributes = InKeyAttributes[Index];
+
+			FFrameNumber Time = (Position.InputValue * TickResolution).RoundToFrame();
+			Section->ExpandToFrame(Time);
+
+			FMovieSceneControlRigSpaceBaseKey Value;
+			FKeyHandle NewHandle = ChannelData.UpdateOrAddKey(Time, Value);
+			if (NewHandle != FKeyHandle::Invalid())
+			{
+				if (OutKeyHandles)
+				{
+					(*OutKeyHandles)[Index] = NewHandle;
+				}
+			}
+		}
+	}
+}
+
+bool FControlRigSpaceChannelCurveModel::Evaluate(double Time, double& OutValue) const
+{
+	// Spaces don't evaluate into a valid value
+	OutValue = 0.0;
+	return false;
+}
+
+void FControlRigSpaceChannelCurveModel::RemoveKeys(TArrayView<const FKeyHandle> InKeys)
+{
+	FMovieSceneControlRigSpaceChannel* Channel = ChannelHandle.Get();
+	UMovieSceneSection*      Section = WeakSection.Get();
+	if (Channel && Section)
+	{
+		Section->Modify();
+
+		TMovieSceneChannelData<FMovieSceneControlRigSpaceBaseKey> ChannelData = Channel->GetData();
+
+		for (FKeyHandle Handle : InKeys)
+		{
+			int32 KeyIndex = ChannelData.GetIndex(Handle);
+			if (KeyIndex != INDEX_NONE)
+			{
+				ChannelData.RemoveKey(KeyIndex);
+			}
+		}
+	}
+}
+
+void FControlRigSpaceChannelCurveModel::DrawCurve(const FCurveEditor& CurveEditor, const FCurveEditorScreenSpace& ScreenSpace, TArray<TTuple<double, double>>& InterpolatingPoints) const
+{
+	// Space Channels don't draw any lines so there's no need to fill out the Interpolating Points array.
+}
+
+void FControlRigSpaceChannelCurveModel::GetKeys(const FCurveEditor& CurveEditor, double MinTime, double MaxTime, double MinValue, double MaxValue, TArray<FKeyHandle>& OutKeyHandles) const
+{
+	FMovieSceneControlRigSpaceChannel* Channel = ChannelHandle.Get();
+	UMovieSceneSection*      Section = WeakSection.Get();
+
+	if (Channel && Section)
+	{
+		FFrameRate TickResolution = Section->GetTypedOuter<UMovieScene>()->GetTickResolution();
+
+		TMovieSceneChannelData<FMovieSceneControlRigSpaceBaseKey> ChannelData = Channel->GetData();
+		TArrayView<const FFrameNumber>          Times = ChannelData.GetTimes();
+
+		const FFrameNumber StartFrame = MinTime <= MIN_int32 ? MIN_int32 : (MinTime * TickResolution).CeilToFrame();
+		const FFrameNumber EndFrame = MaxTime >= MAX_int32 ? MAX_int32 : (MaxTime * TickResolution).FloorToFrame();
+
+		const int32 StartingIndex = Algo::LowerBound(Times, StartFrame);
+		const int32 EndingIndex = Algo::UpperBound(Times, EndFrame);
+
+		for (int32 KeyIndex = StartingIndex; KeyIndex < EndingIndex; ++KeyIndex)
+		{
+			// Space Channels don't have Values associated with them so we ignore the Min/Max value and always return the key.
+			OutKeyHandles.Add(ChannelData.GetHandle(KeyIndex));
+		}
+	}
+}
+
+//we actually still draw the 'key' as a the border which you can move and right click, context menu with.
+void FControlRigSpaceChannelCurveModel::GetKeyDrawInfo(ECurvePointType PointType, const FKeyHandle InKeyHandle, FKeyDrawInfo& OutDrawInfo) const
+{
+	OutDrawInfo.Brush = FEditorStyle::Get().GetBrush("FilledBorder");
+	OutDrawInfo.ScreenSize = FVector2D(10, 100);
+}
+
+void FControlRigSpaceChannelCurveModel::GetKeyPositions(TArrayView<const FKeyHandle> InKeys, TArrayView<FKeyPosition> OutKeyPositions) const
+{
+	FMovieSceneControlRigSpaceChannel* Channel = ChannelHandle.Get();
+	UMovieSceneSection*      Section = WeakSection.Get();
+
+	if (Channel && Section)
+	{
+		FFrameRate TickResolution = Section->GetTypedOuter<UMovieScene>()->GetTickResolution();
+
+		TMovieSceneChannelData<FMovieSceneControlRigSpaceBaseKey> ChannelData = Channel->GetData();
+		TArrayView<const FFrameNumber>          Times = ChannelData.GetTimes();
+		TArrayView<const FMovieSceneControlRigSpaceBaseKey> Values = ChannelData.GetValues();
+
+		for (int32 Index = 0; Index < InKeys.Num(); ++Index)
+		{
+			int32 KeyIndex = ChannelData.GetIndex(InKeys[Index]);
+			if (KeyIndex != INDEX_NONE)
+			{
+				OutKeyPositions[Index].InputValue = Times[KeyIndex] / TickResolution;
+				// Spaces have no values so we just output zero.
+				OutKeyPositions[Index].OutputValue = 0.0;
+			}
+		}
+	}
+}
+
+void FControlRigSpaceChannelCurveModel::SetKeyPositions(TArrayView<const FKeyHandle> InKeys, TArrayView<const FKeyPosition> InKeyPositions, EPropertyChangeType::Type ChangeType)
+{
+	FMovieSceneControlRigSpaceChannel* Channel = ChannelHandle.Get();
+	UMovieSceneSection*      Section = WeakSection.Get();
+
+	if (Channel && Section)
+	{
+		Section->Modify();
+
+		FFrameRate TickResolution = Section->GetTypedOuter<UMovieScene>()->GetTickResolution();
+
+		TMovieSceneChannelData<FMovieSceneControlRigSpaceBaseKey> ChannelData = Channel->GetData();
+		for (int32 Index = 0; Index < InKeys.Num(); ++Index)
+		{
+			int32 KeyIndex = ChannelData.GetIndex(InKeys[Index]);
+			if (KeyIndex != INDEX_NONE)
+			{
+				FFrameNumber NewTime = (InKeyPositions[Index].InputValue * TickResolution).FloorToFrame();
+
+				KeyIndex = ChannelData.MoveKey(KeyIndex, NewTime);
+				Section->ExpandToFrame(NewTime);
+			}
+		}
+	}
+}
+
+void FControlRigSpaceChannelCurveModel::GetKeyAttributes(TArrayView<const FKeyHandle> InKeys, TArrayView<FKeyAttributes> OutAttributes) const
+{
+}
+
+void FControlRigSpaceChannelCurveModel::SetKeyAttributes(TArrayView<const FKeyHandle> InKeys, TArrayView<const FKeyAttributes> InAttributes, EPropertyChangeType::Type ChangeType)
+{
+}
+
+void FControlRigSpaceChannelCurveModel::GetCurveAttributes(FCurveAttributes& OutCurveAttributes) const
+{
+	// Space Channels have no Pre/Post Extrapolation
+	OutCurveAttributes.SetPreExtrapolation(RCCE_None);
+	OutCurveAttributes.SetPostExtrapolation(RCCE_None);
+}
+
+void FControlRigSpaceChannelCurveModel::SetCurveAttributes(const FCurveAttributes& InCurveAttributes)
+{
+}
+
+void FControlRigSpaceChannelCurveModel::CreateKeyProxies(TArrayView<const FKeyHandle> InKeyHandles, TArrayView<UObject*> OutObjects)
+{
+}
+
+void FControlRigSpaceChannelCurveModel::GetTimeRange(double& MinTime, double& MaxTime) const
+{
+	FMovieSceneControlRigSpaceChannel* Channel = ChannelHandle.Get();
+	UMovieSceneSection*      Section = WeakSection.Get();
+
+	if (Channel && Section)
+	{
+		TArrayView<const FFrameNumber> Times = Channel->GetData().GetTimes();
+		if (Times.Num() == 0)
+		{
+			MinTime = 0.f;
+			MaxTime = 0.f;
+		}
+		else
+		{
+			FFrameRate TickResolution = Section->GetTypedOuter<UMovieScene>()->GetTickResolution();
+			double ToTime = TickResolution.AsInterval();
+			MinTime = static_cast<double> (Times[0].Value) * ToTime;
+			MaxTime = static_cast<double>(Times[Times.Num() - 1].Value) * ToTime;
+		}
+	}
+}
+
+void FControlRigSpaceChannelCurveModel::GetValueRange(double& MinValue, double& MaxValue) const
+{
+	// Space channels have no value, thus their value range is just zero.
+	MinValue = MaxValue = 0.0;
+}
+
+int32 FControlRigSpaceChannelCurveModel::GetNumKeys() const
+{
+	FMovieSceneControlRigSpaceChannel* Channel = ChannelHandle.Get();
+
+	if (Channel)
+	{
+		return Channel->GetData().GetTimes().Num();
+	}
+
+	return 0;
+}
+
+void FControlRigSpaceChannelCurveModel::BuildContextMenu(const FCurveEditor& CurveEditor,FMenuBuilder& MenuBuilder, TOptional<FCurvePointHandle> ClickedPoint)
+{
+	if (IsReadOnly())
+	{
+		return;
+	}
+	if (ClickedPoint.IsSet())
+	{
+		MenuBuilder.BeginSection("SpaceContextMenuSection", NSLOCTEXT("CurveEditor", "Space", "Space"));
+		{
+			FMovieSceneControlRigSpaceChannel* Channel = ChannelHandle.Get();
+			const FKeyHandle KeyHandle = ClickedPoint.GetValue().KeyHandle;
+			if (Channel && KeyHandle != FKeyHandle::Invalid())
+			{
+				TMovieSceneChannelData<FMovieSceneControlRigSpaceBaseKey> ChannelData = Channel->GetData();
+				int32 KeyIndex = ChannelData.GetIndex(KeyHandle);
+				TArrayView<const FMovieSceneControlRigSpaceBaseKey> Values = ChannelData.GetValues();
+				if (KeyIndex >= 0 && KeyIndex < Values.Num())
+				{
+					FControlRigSpaceChannelHelpers::OpenBakeDialog(WeakSequencer.Pin().Get(), Channel, KeyIndex, WeakSection.Get());
+				}
+
+			}
+		}
+		MenuBuilder.EndSection();
+	}
+}
diff --git a/Engine/Plugins/Animation/ControlRig/Source/ControlRigEditor/Private/EditMode/SControlRigEditModeTools.cpp b/Engine/Plugins/Animation/ControlRig/Source/ControlRigEditor/Private/EditMode/SControlRigEditModeTools.cpp
--- a/Engine/Plugins/Animation/ControlRig/Source/ControlRigEditor/Private/EditMode/SControlRigEditModeTools.cpp
+++ b/Engine/Plugins/Animation/ControlRig/Source/ControlRigEditor/Private/EditMode/SControlRigEditModeTools.cpp
@@ -771,24 +771,8 @@ FReply SControlRigEditModeTools::OnBakeControlsToNewSpaceButtonClicked()
 		FMovieSceneControlRigSpaceBaseKey Value;
 		using namespace UE::MovieScene;
 		URigHierarchy* RigHierarchy = SpacePickerWidget->GetHierarchy();
-		Settings.TargetSpace = RigHierarchy->GetFirstParent(ControlKeys[0]);
-		/*  we just set space as parent, but leaving this in case we go back to setting based on current
-		EvaluateChannel(SpaceChannelAndSection.SpaceChannel, CurrentTime, Value);
-		switch (Value.SpaceType)
-		{
-		case EMovieSceneControlRigSpaceType::World:
-			Settings.TargetSpace = RigHierarchy->GetWorldSpaceSocketKey();
-			break;
-		case EMovieSceneControlRigSpaceType::Parent:
-			Settings.TargetSpace = RigHierarchy->GetFirstParent(ControlKeys[0]);
-			break;
-		case EMovieSceneControlRigSpaceType::ControlRig:
-			Settings.TargetSpace = Value.ControlRigElement;
-			break;
-		default:
-			Settings.TargetSpace = URigHierarchy::GetDefaultParentSocketKey();
-		}
-		*/
+		Settings.TargetSpace = URigHierarchy::GetDefaultParentSocketKey();
+		
 		TRange<FFrameNumber> Range = Sequencer->GetFocusedMovieSceneSequence()->GetMovieScene()->GetPlaybackRange();
 		TArray<FFrameNumber> Keys;
 		TArray < FKeyHandle> KeyHandles;
diff --git a/Engine/Plugins/Animation/ControlRig/Source/ControlRigEditor/Private/Editor/SRigSpacePickerWidget.cpp b/Engine/Plugins/Animation/ControlRig/Source/ControlRigEditor/Private/Editor/SRigSpacePickerWidget.cpp
--- a/Engine/Plugins/Animation/ControlRig/Source/ControlRigEditor/Private/Editor/SRigSpacePickerWidget.cpp
+++ b/Engine/Plugins/Animation/ControlRig/Source/ControlRigEditor/Private/Editor/SRigSpacePickerWidget.cpp
@@ -514,7 +514,9 @@ FReply SRigSpacePickerWidget::HandleWorldSpaceClicked()
 
 FReply SRigSpacePickerWidget::HandleElementSpaceClicked(FRigElementKey InKey)
 {
-	for(const FRigElementKey& ControlKey : ControlKeys)
+	//need to make copy since array may get shrunk during the event broadcast
+	TArray<FRigElementKey> ControlKeysCopy = ControlKeys;
+	for (const FRigElementKey& ControlKey : ControlKeysCopy)
 	{
 		ActiveSpaceChangedEvent.Broadcast(Hierarchy, ControlKey, InKey);
 	}
diff --git a/Engine/Plugins/Animation/ControlRig/Source/ControlRigEditor/Private/Sequencer/ControlRigSpaceChannelEditors.cpp b/Engine/Plugins/Animation/ControlRig/Source/ControlRigEditor/Private/Sequencer/ControlRigSpaceChannelEditors.cpp
--- a/Engine/Plugins/Animation/ControlRig/Source/ControlRigEditor/Private/Sequencer/ControlRigSpaceChannelEditors.cpp
+++ b/Engine/Plugins/Animation/ControlRig/Source/ControlRigEditor/Private/Sequencer/ControlRigSpaceChannelEditors.cpp
@@ -49,7 +49,7 @@
 #include "MovieSceneObjectBindingID.h"
 #include "ControlRig.h"
 #include "IControlRigObjectBinding.h"
-
+#include "ControlRigSpaceChannelCurveModel.h"
 #define LOCTEXT_NAMESPACE "ControlRigEditMode"
 
 
@@ -196,6 +196,8 @@ FSpaceChannelAndSection FControlRigSpaceChannelHelpers::FindSpaceChannelAndSecti
 		{
 			return SpaceChannelAndSection;
 		}
+		bool bRecreateCurves = false;
+		TArray<TPair<UControlRig*, FName>> ControlRigPairsToReselect;
 		if (FMovieSceneBinding* Binding = MovieScene->FindBinding(ObjectHandle))
 		{
 			for (UMovieSceneTrack* Track : Binding->GetTracks())
@@ -217,17 +219,38 @@ FSpaceChannelAndSection FControlRigSpaceChannelHelpers::FindSpaceChannelAndSecti
 							}
 							else if (bCreateIfNeeded)
 							{
+								if (ControlRig->IsControlSelected(ControlName))
+								{
+									TPair<UControlRig*, FName> Pair;
+									Pair.Key = ControlRig;
+									Pair.Value = ControlName;
+									ControlRigPairsToReselect.Add(Pair);
+								}
 								ActiveSection->AddSpaceChannel(ControlName, true /*ReconstructChannelProxy*/);
 								NameAndChannel = ActiveSection->GetSpaceChannel(ControlName);
 								if (NameAndChannel)
 								{
 									SpaceChannelAndSection.SpaceChannel = &NameAndChannel->SpaceCurve;
+									bRecreateCurves = true;
 								}
 							}
 						}
 					}
 				}
 			}
+			if (bRecreateCurves)
+			{
+				Sequencer->RecreateCurveEditor(); //this will require the curve editor to get recreated so the ordering is correct
+				for (TPair<UControlRig*, FName>& Pair : ControlRigPairsToReselect)
+				{
+
+					GEditor->GetTimerManager()->SetTimerForNextTick([Pair]()
+					{
+						Pair.Key->SelectControl(Pair.Value);
+					});
+					
+				}
+			}
 		}
 	}
 	return SpaceChannelAndSection;
@@ -251,6 +274,10 @@ FKeyHandle FControlRigSpaceChannelHelpers::SequencerKeyControlRigSpaceChannel(UC
 	{
 		Value.SpaceType = EMovieSceneControlRigSpaceType::World;
 	}
+	else if (SpaceKey == RigHierarchy->GetDefaultParentSocketKey())
+	{
+		Value.SpaceType = EMovieSceneControlRigSpaceType::Parent;
+	}
 	else
 	{
 		FRigElementKey DefaultParent = RigHierarchy->GetFirstParent(ControlKey);
@@ -589,6 +616,10 @@ void FControlRigSpaceChannelHelpers::SequencerBakeControlInSpace(UControlRig* Co
 		{
 			Value.SpaceType = EMovieSceneControlRigSpaceType::World;
 		}
+		else if (Settings.TargetSpace == RigHierarchy->GetDefaultParentSocketKey())
+		{
+			Value.SpaceType = EMovieSceneControlRigSpaceType::Parent;
+		}
 		else
 		{
 			FRigElementKey DefaultParent = RigHierarchy->GetFirstParent(ControlKey);
@@ -887,4 +918,137 @@ void FControlRigSpaceChannelHelpers::CompensateIfNeeded(UControlRig* ControlRig,
 	}
 }
 
+//mz todo, we may want to expose these colors in editor user preferences, okay fo rnow
+static FLinearColor GetNextIndexedColor()
+{
+	static TArray<FLinearColor> IndexedColor;
+	static int32 NextIndex = 0;
+	if (IndexedColor.Num() == 0)
+	{
+		IndexedColor.Add(FLinearColor(FColor::Orange));
+		IndexedColor.Add(FLinearColor(FColor::Silver));
+		IndexedColor.Add(FLinearColor(FColor::Purple));
+		IndexedColor.Add(FLinearColor(FColor::Emerald));
+		IndexedColor.Add(FLinearColor(FColor::Red));
+		IndexedColor.Add(FLinearColor(FColor::Green));
+		IndexedColor.Add(FLinearColor(FColor::Blue));
+	}
+	FLinearColor Color = IndexedColor[NextIndex];
+	int32 TheNextIndex = NextIndex + 1;
+	NextIndex = (TheNextIndex) % IndexedColor.Num();
+	return Color;
+
+}
+FLinearColor FControlRigSpaceChannelHelpers::GetColor(const FMovieSceneControlRigSpaceBaseKey& Key)
+{
+	static TMap<FName, FLinearColor> ColorsForSpaces;
+
+	switch (Key.SpaceType)
+	{
+		case EMovieSceneControlRigSpaceType::Parent:
+			return FLinearColor(FColor::Magenta);
+		case EMovieSceneControlRigSpaceType::World:
+			return  FLinearColor(FColor::Turquoise);
+		case EMovieSceneControlRigSpaceType::ControlRig:
+		{
+			if (FLinearColor* Color = ColorsForSpaces.Find(Key.ControlRigElement.Name))
+			{
+				return *Color;
+			}
+			else
+			{
+				FLinearColor RandoColor = GetNextIndexedColor();
+				ColorsForSpaces.Add(Key.ControlRigElement.Name, RandoColor);
+				return RandoColor;
+			}
+		}
+	};
+	return  FLinearColor(FColor::White);
+
+}
+
+FReply FControlRigSpaceChannelHelpers::OpenBakeDialog(ISequencer* Sequencer, FMovieSceneControlRigSpaceChannel* Channel,int32 KeyIndex, UMovieSceneSection* SectionToKey )
+{
+
+	if (!Sequencer || !Sequencer->GetFocusedMovieSceneSequence()|| !Sequencer->GetFocusedMovieSceneSequence()->GetMovieScene())
+	{
+		return FReply::Unhandled();
+	}
+	if (UMovieSceneControlRigParameterSection* Section = Cast<UMovieSceneControlRigParameterSection>(SectionToKey))
+	{
+		if (UControlRig* ControlRig = Section->GetControlRig())
+		{
+			FName ControlName = Section->FindControlNameFromSpaceChannel(Channel);
+
+			if (ControlName == NAME_None)
+			{
+				return FReply::Unhandled();
+			}
+			FRigSpacePickerBakeSettings Settings;
+
+			const FFrameRate TickResolution = Sequencer->GetFocusedTickResolution();
+			FMovieSceneControlRigSpaceBaseKey Value;
+
+			URigHierarchy* RigHierarchy = ControlRig->GetHierarchy();
+			Settings.TargetSpace = URigHierarchy::GetDefaultParentSocketKey();
+
+			TRange<FFrameNumber> Range = Sequencer->GetFocusedMovieSceneSequence()->GetMovieScene()->GetPlaybackRange();
+			Settings.StartFrame = Range.GetLowerBoundValue();
+			Settings.EndFrame = Range.GetUpperBoundValue();
+			TMovieSceneChannelData<FMovieSceneControlRigSpaceBaseKey> ChannelData = Channel->GetData();
+			TArrayView<const FFrameNumber> Times = ChannelData.GetTimes();
+			if (KeyIndex >= 0 && KeyIndex < Times.Num())
+			{
+				Settings.StartFrame = Times[KeyIndex];
+				if (KeyIndex + 1 < Times.Num())
+				{
+					Settings.EndFrame = Times[KeyIndex + 1];
+				}
+			}
+			TArray<FRigElementKey> Controls;
+			FRigElementKey Key;
+			Key.Name = ControlName;
+			Key.Type = ERigElementType::Control;
+			Controls.Add(Key);
+			TSharedRef<SRigSpacePickerBakeWidget> BakeWidget =
+				SNew(SRigSpacePickerBakeWidget)
+				.Settings(Settings)
+				.Hierarchy(ControlRig->GetHierarchy())
+				.Controls(Controls)
+				.Sequencer(Sequencer)
+				.GetControlCustomization_Lambda([ControlRig](URigHierarchy*, const FRigElementKey& InControlKey)
+				{
+					return ControlRig->GetControlCustomization(InControlKey);
+				})
+				.OnBake_Lambda([Sequencer, ControlRig, TickResolution](URigHierarchy* InHierarchy, TArray<FRigElementKey> InControls, FRigSpacePickerBakeSettings InSettings)
+				{
+					TArray<FFrameNumber> Frames;
+
+					const FFrameRate& FrameRate = Sequencer->GetFocusedDisplayRate();
+					FFrameNumber FrameRateInFrameNumber = TickResolution.AsFrameNumber(FrameRate.AsInterval());
+					for (FFrameNumber& Frame = InSettings.StartFrame; Frame <= InSettings.EndFrame; Frame += FrameRateInFrameNumber)
+					{
+						Frames.Add(Frame);
+					}
+					FScopedTransaction Transaction(LOCTEXT("BakeControlToSpace", "Bake Control In Space"));
+					for (const FRigElementKey& ControlKey : InControls)
+					{
+						FSpaceChannelAndSection SpaceChannelAndSection = FControlRigSpaceChannelHelpers::FindSpaceChannelAndSectionForControl(ControlRig, ControlKey.Name, Sequencer, false /*bCreateIfNeeded*/);
+
+						FControlRigSpaceChannelHelpers::SequencerBakeControlInSpace(ControlRig, Sequencer, SpaceChannelAndSection.SpaceChannel, SpaceChannelAndSection.SectionToKey,
+							Frames, InHierarchy, ControlKey, InSettings);
+					}
+					return FReply::Handled();
+				});
+
+			return BakeWidget->OpenDialog(true);
+			
+		}
+	}
+	return FReply::Unhandled();
+}
+TUniquePtr<FCurveModel> CreateCurveEditorModel(const TMovieSceneChannelHandle<FMovieSceneControlRigSpaceChannel>& Channel, UMovieSceneSection* OwningSection, TSharedRef<ISequencer> InSequencer)
+{
+	return MakeUnique<FControlRigSpaceChannelCurveModel>(Channel, OwningSection, InSequencer);
+}
 #undef LOCTEXT_NAMESPACE
diff --git a/Engine/Plugins/Animation/ControlRig/Source/ControlRigEditor/Public/ControlRigSpaceChannelCurveModel.h b/Engine/Plugins/Animation/ControlRig/Source/ControlRigEditor/Public/ControlRigSpaceChannelCurveModel.h
new file mode 100644
--- /dev/null
+++ b/Engine/Plugins/Animation/ControlRig/Source/ControlRigEditor/Public/ControlRigSpaceChannelCurveModel.h
@@ -0,0 +1,57 @@
+// Copyright Epic Games, Inc. All Rights Reserved.
+
+#pragma once
+
+#include "CoreTypes.h"
+#include "UObject/WeakObjectPtr.h"
+#include "UObject/WeakObjectPtrTemplates.h"
+#include "KeyBarCurveModel.h"
+#include "Sequencer/MovieSceneControlRigSpaceChannel.h"
+#include "Channels/MovieSceneChannelHandle.h"
+
+class UMovieSceneSection;
+class ISequencer;
+
+class FControlRigSpaceChannelCurveModel : public FKeyBarCurveModel
+{
+public:
+	static ECurveEditorViewID ViewID;
+
+	FControlRigSpaceChannelCurveModel(TMovieSceneChannelHandle<FMovieSceneControlRigSpaceChannel> InChannel, UMovieSceneSection* InOwningSection, TWeakPtr<ISequencer> InWeakSequencer);
+
+	//FCurveModel
+	virtual const void* GetCurve() const override;
+
+	virtual void Modify() override;
+
+	virtual void DrawCurve(const FCurveEditor& CurveEditor, const FCurveEditorScreenSpace& ScreenSpace, TArray<TTuple<double, double>>& InterpolatingPoints) const override;
+	virtual void GetKeys(const FCurveEditor& CurveEditor, double MinTime, double MaxTime, double MinValue, double MaxValue, TArray<FKeyHandle>& OutKeyHandles) const override;
+	virtual void GetKeyDrawInfo(ECurvePointType PointType, const FKeyHandle InKeyHandle, FKeyDrawInfo& OutDrawInfo) const override;
+
+	virtual void GetKeyPositions(TArrayView<const FKeyHandle> InKeys, TArrayView<FKeyPosition> OutKeyPositions) const override;
+	virtual void SetKeyPositions(TArrayView<const FKeyHandle> InKeys, TArrayView<const FKeyPosition> InKeyPositions, EPropertyChangeType::Type ChangeType) override;
+
+	virtual void GetKeyAttributes(TArrayView<const FKeyHandle> InKeys, TArrayView<FKeyAttributes> OutAttributes) const override;
+	virtual void SetKeyAttributes(TArrayView<const FKeyHandle> InKeys, TArrayView<const FKeyAttributes> InAttributes, EPropertyChangeType::Type ChangeType = EPropertyChangeType::Unspecified) override;
+
+	virtual void GetCurveAttributes(FCurveAttributes& OutCurveAttributes) const override;
+	virtual void SetCurveAttributes(const FCurveAttributes& InCurveAttributes) override;
+	virtual void GetTimeRange(double& MinTime, double& MaxTime) const override;
+	virtual void GetValueRange(double& MinValue, double& MaxValue) const override;
+	virtual int32 GetNumKeys() const override;
+	virtual void GetNeighboringKeys(const FKeyHandle InKeyHandle, TOptional<FKeyHandle>& OutPreviousKeyHandle, TOptional<FKeyHandle>& OutNextKeyHandle) const override {}
+	virtual bool Evaluate(double ProspectiveTime, double& OutValue) const override;
+	virtual void AddKeys(TArrayView<const FKeyPosition> InKeyPositions, TArrayView<const FKeyAttributes> InAttributes, TArrayView<TOptional<FKeyHandle>>* OutKeyHandles) override;
+	virtual void RemoveKeys(TArrayView<const FKeyHandle> InKeys) override;
+
+	virtual void CreateKeyProxies(TArrayView<const FKeyHandle> InKeyHandles, TArrayView<UObject*> OutObjects) override;
+
+	//FKeyBarCurveModel
+	virtual void BuildContextMenu(const FCurveEditor& CurveEditor,FMenuBuilder& MenuBuilder, TOptional<FCurvePointHandle> ClickedPoint);
+	virtual TArray<FKeyBarCurveModel::FBarRange> FindRanges();
+
+private:
+	TMovieSceneChannelHandle<FMovieSceneControlRigSpaceChannel> ChannelHandle;
+	TWeakObjectPtr<UMovieSceneSection> WeakSection;
+	TWeakPtr<ISequencer> WeakSequencer;
+};
\ No newline at end of file
diff --git a/Engine/Plugins/Animation/ControlRig/Source/ControlRigEditor/Public/ControlRigSpaceChannelEditors.h b/Engine/Plugins/Animation/ControlRig/Source/ControlRigEditor/Public/ControlRigSpaceChannelEditors.h
--- a/Engine/Plugins/Animation/ControlRig/Source/ControlRigEditor/Public/ControlRigSpaceChannelEditors.h
+++ b/Engine/Plugins/Animation/ControlRig/Source/ControlRigEditor/Public/ControlRigSpaceChannelEditors.h
@@ -55,6 +55,8 @@ struct FControlRigSpaceChannelHelpers
 	static void HandleSpaceKeyTimeChanged(UControlRig* ControlRig, FName ControlName,FMovieSceneControlRigSpaceChannel* Channel, UMovieSceneSection* SectionToKey,
 		FFrameNumber CurrentFrame, FFrameNumber NextFrame);
 	static void DeleteTransformKeysAtThisTime(UControlRig* ControlRig, UMovieSceneControlRigParameterSection* Section, FName ControlName, FFrameNumber Time);
+	static FLinearColor GetColor(const FMovieSceneControlRigSpaceBaseKey& Key);
+	static FReply OpenBakeDialog(ISequencer* Sequencer, FMovieSceneControlRigSpaceChannel* Channel, int32 KeyIndex, UMovieSceneSection* SectionToKey);
 };
 
 //template specialization
@@ -77,7 +79,7 @@ TSharedRef<SWidget> CreateKeyEditor(const TMovieSceneChannelHandle<FMovieSceneCo
 //void ExtendKeyMenu(FMenuBuilder& OuterMenuBuilder, TArray<TExtendKeyMenuParams<FMovieSceneControlRigSpaceChannel>>&& Channels, TWeakPtr<ISequencer> InSequencer);
 
 /** Curve editor models */
-//inline bool SupportsCurveEditorModels(const TMovieSceneChannelHandle<FMovieSceneControlRigSpaceChannel>& DoubleChannel) { return true; }
-//TUniquePtr<FCurveModel> CreateCurveEditorModel(const TMovieSceneChannelHandle<FMovieSceneControlRigSpaceChannel>& DoubleChannel, UMovieSceneSection* OwningSection, TSharedRef<ISequencer> InSequencer);
+inline bool SupportsCurveEditorModels(const TMovieSceneChannelHandle<FMovieSceneControlRigSpaceChannel>& Channel) { return true; }
+TUniquePtr<FCurveModel> CreateCurveEditorModel(const TMovieSceneChannelHandle<FMovieSceneControlRigSpaceChannel>& Channel, UMovieSceneSection* OwningSection, TSharedRef<ISequencer> InSequencer);
 
 
diff --git a/Engine/Source/Editor/CurveEditor/Private/KeyBarCurveModel.cpp b/Engine/Source/Editor/CurveEditor/Private/KeyBarCurveModel.cpp
new file mode 100644
--- /dev/null
+++ b/Engine/Source/Editor/CurveEditor/Private/KeyBarCurveModel.cpp
@@ -0,0 +1,15 @@
+// Copyright Epic Games, Inc. All Rights Reserved.
+
+
+#include "KeyBarCurveModel.h"
+
+FKeyBarCurveModel::FKeyBarCurveModel()
+{
+
+}
+
+FKeyBarCurveModel::~FKeyBarCurveModel()
+{
+
+}
+
diff --git a/Engine/Source/Editor/CurveEditor/Private/Views/SInteractiveCurveEditorView.cpp b/Engine/Source/Editor/CurveEditor/Private/Views/SInteractiveCurveEditorView.cpp
--- a/Engine/Source/Editor/CurveEditor/Private/Views/SInteractiveCurveEditorView.cpp
+++ b/Engine/Source/Editor/CurveEditor/Private/Views/SInteractiveCurveEditorView.cpp
@@ -1221,14 +1221,23 @@ void SInteractiveCurveEditorView::CreateContextMenu(const FGeometry& MyGeometry,
 
 	const bool bCloseAfterSelection = true;
 	FMenuBuilder MenuBuilder(bCloseAfterSelection, EditorPanel->GetCommands());
-
-	FCurveEditorContextMenu::BuildMenu(MenuBuilder, CurveEditor.ToSharedRef(), MouseUpPoint, GetHoveredCurve());
+	
+	BuildContextMenu(MenuBuilder, MouseUpPoint, GetHoveredCurve());
 
 	// Push the context menu
 	FWidgetPath WidgetPath = MouseEvent.GetEventPath() != nullptr ? *MouseEvent.GetEventPath() : FWidgetPath();
 	ActiveContextMenu = FSlateApplication::Get().PushMenu(AsShared(), WidgetPath, MenuBuilder.MakeWidget(), FSlateApplication::Get().GetCursorPos(), FPopupTransitionEffect(FPopupTransitionEffect::ContextMenu));
 }
 
+void SInteractiveCurveEditorView::BuildContextMenu(FMenuBuilder& MenuBuilder, TOptional<FCurvePointHandle> ClickedPoint, TOptional<FCurveModelID> HoveredCurveID)
+{
+	TSharedPtr<FCurveEditor> CurveEditor = WeakCurveEditor.Pin();
+	if (CurveEditor)
+	{
+		FCurveEditorContextMenu::BuildMenu(MenuBuilder, CurveEditor.ToSharedRef(), ClickedPoint, HoveredCurveID);
+	}
+}
+
 TOptional<FCurvePointHandle> SInteractiveCurveEditorView::HitPoint(FVector2D MousePixel) const
 {
 	TOptional<FCurvePointHandle> HitPoint;
diff --git a/Engine/Source/Editor/CurveEditor/Public/KeyBarCurveModel.h b/Engine/Source/Editor/CurveEditor/Public/KeyBarCurveModel.h
new file mode 100644
--- /dev/null
+++ b/Engine/Source/Editor/CurveEditor/Public/KeyBarCurveModel.h
@@ -0,0 +1,59 @@
+// Copyright Epic Games, Inc. All Rights Reserved.
+
+#pragma once
+
+#include "CoreTypes.h"
+#include "CurveModel.h"
+
+class FMenuBuilder;
+struct FCurvePointHandle;
+
+/**
+ * Class that provides functionality to represents keys as continuous ranges.
+ */
+class CURVEEDITOR_API FKeyBarCurveModel: public FCurveModel
+{
+public:
+
+	FKeyBarCurveModel();
+	~FKeyBarCurveModel();
+
+	virtual const void* GetCurve() const override { return nullptr; }
+	virtual void Modify() override {};
+	virtual void DrawCurve(const FCurveEditor& CurveEditor, const FCurveEditorScreenSpace& ScreenSpace, TArray<TTuple<double, double>>& InterpolatingPoints) const override{};
+	virtual void GetKeys(const FCurveEditor& CurveEditor, double MinTime, double MaxTime, double MinValue, double MaxValue, TArray<FKeyHandle>& OutKeyHandles) const override {};
+	virtual void AddKeys(TArrayView<const FKeyPosition> InPositions, TArrayView<const FKeyAttributes> InAttributes, TArrayView<TOptional<FKeyHandle>>* OutKeyHandles = nullptr) override {};
+	virtual void RemoveKeys(TArrayView<const FKeyHandle> InKeys) override {};
+	virtual void GetKeyPositions(TArrayView<const FKeyHandle> InKeys, TArrayView<FKeyPosition> OutKeyPositions) const override {};
+	virtual void SetKeyPositions(TArrayView<const FKeyHandle> InKeys, TArrayView<const FKeyPosition> InKeyPositions,const EPropertyChangeType::Type ChangeType = EPropertyChangeType::Unspecified) override {};
+	virtual void GetKeyDrawInfo(ECurvePointType PointType, const FKeyHandle InKeyHandle, FKeyDrawInfo& OutDrawInfo) const override {};
+	virtual void GetTimeRange(double& MinTime, double& MaxTime) const override {};
+	virtual void GetValueRange(double& MinValue, double& MaxValue) const override {};
+	virtual int32 GetNumKeys() const override { return 0; }
+	virtual void GetNeighboringKeys(const FKeyHandle InKeyHandle, TOptional<FKeyHandle>& OutPreviousKeyHandle, TOptional<FKeyHandle>& OutNextKeyHandle) const override {};
+	virtual bool Evaluate(double InTime, double& OutValue) const override { return false; }
+
+	/*
+	*  New Virtuals for Key Bar Curve Model
+	*/
+
+	virtual void BuildContextMenu(const FCurveEditor& CurveEditor, FMenuBuilder& MenuBuilder, TOptional<FCurvePointHandle> ClickedPoint) {};
+
+	/**
+	*  Bar Ranges for the Keys
+	*/
+	struct CURVEEDITOR_API FBarRange
+	{
+		TRange<double> Range;
+		FName  Name;
+		FLinearColor Color;
+	};
+	/**
+	 * Find all of the ranges in this model. This should return an increasing set of ranges.
+	 */
+	virtual TArray<FKeyBarCurveModel::FBarRange> FindRanges() 
+	{
+		TArray<FKeyBarCurveModel::FBarRange> Temp; 
+		return Temp;
+	}
+};
diff --git a/Engine/Source/Editor/CurveEditor/Public/Views/SInteractiveCurveEditorView.h b/Engine/Source/Editor/CurveEditor/Public/Views/SInteractiveCurveEditorView.h
--- a/Engine/Source/Editor/CurveEditor/Public/Views/SInteractiveCurveEditorView.h
+++ b/Engine/Source/Editor/CurveEditor/Public/Views/SInteractiveCurveEditorView.h
@@ -77,6 +77,8 @@ public:
 	virtual void GetGridLinesX(TSharedRef<const FCurveEditor> CurveEditor, TArray<float>& MajorGridLines, TArray<float>& MinorGridLines, TArray<FText>* MajorGridLabels = nullptr) const override;
 	virtual void GetGridLinesY(TSharedRef<const FCurveEditor> CurveEditor, TArray<float>& MajorGridLines, TArray<float>& MinorGridLines, TArray<FText>* MajorGridLabels = nullptr) const override;
 
+	virtual void BuildContextMenu(FMenuBuilder& MenuBuilder, TOptional<FCurvePointHandle> ClickedPoint, TOptional<FCurveModelID> HoveredCurveID);
+
 protected:
 
 	// ~SCurveEditorView Interface
diff --git a/Engine/Source/Editor/MovieSceneTools/Private/Channels/SCurveEditorKeyBarView.cpp b/Engine/Source/Editor/MovieSceneTools/Private/Channels/SCurveEditorKeyBarView.cpp
--- a/Engine/Source/Editor/MovieSceneTools/Private/Channels/SCurveEditorKeyBarView.cpp
+++ b/Engine/Source/Editor/MovieSceneTools/Private/Channels/SCurveEditorKeyBarView.cpp
@@ -11,6 +11,9 @@
 #include "CurveEditorScreenSpace.h"
 #include "CurveDataAbstraction.h"
 #include "CurveDrawInfo.h"
+#include "KeyBarCurveModel.h"
+#include "Framework/MultiBox/MultiBoxBuilder.h"
+#include "SCurveEditorPanel.h"
 
 float SCurveEditorKeyBarView::TrackHeight = 24.f;
 
@@ -89,7 +92,7 @@ void SCurveEditorKeyBarView::DrawLabels(const FGeometry& AllottedGeometry, FSlat
 
 	const FVector2D LocalSize = AllottedGeometry.GetLocalSize();
 
-	const FCurveEditorScreenSpaceV ViewSpace = GetViewSpace();
+	const FCurveEditorScreenSpace ViewSpace = GetViewSpace();
 
 	for (auto It = CurveInfoByID.CreateConstIterator(); It; ++It)
 	{
@@ -101,48 +104,99 @@ void SCurveEditorKeyBarView::DrawLabels(const FGeometry& AllottedGeometry, FSlat
 
 		const int32 CurveIndex = static_cast<int32>(It->Value.ViewToCurveTransform.GetTranslation().Y);
 
-		const FCurveEditorScreenSpaceV CurveSpace = ViewSpace.ToCurveSpace(It->Value.ViewToCurveTransform);
+		FCurveEditorScreenSpace CurveSpace = GetCurveSpace(It.Key());
+
 		const float LaneTop = CurveSpace.ValueToScreen(0.0) - TrackHeight*.5f;
 
-		// Draw the curve color as the background. Event curves set their track color as the curve color.
-		FLinearColor CurveColor = Curve->GetColor();
 
-		// Alpha blend the zebra tint
-		if (CurveIndex%2)
+		double InputMin = 0, InputMax = 1;
+		GetInputBounds(InputMin, InputMax);//in Sequencer Seconds
+
+		// Draw the curve label and background
+		FKeyBarCurveModel* KeyBarCurveModel = static_cast<FKeyBarCurveModel*>(Curve);
+		if (KeyBarCurveModel)
 		{
-			static FLinearColor ZebraTint = FLinearColor::White.CopyWithNewOpacity(0.01f);
-			if (CurveColor == FLinearColor::White)
-			{
-				CurveColor = ZebraTint;
-			}
-			else
+			TArray<FKeyBarCurveModel::FBarRange> Ranges = KeyBarCurveModel->FindRanges();
+			int32 LastLabelPos = -1;
+			for (int32 Index = 0;Index < Ranges.Num(); ++Index)
 			{
-				CurveColor = CurveColor * (1.f - ZebraTint.A) + ZebraTint * ZebraTint.A;
+				const FKeyBarCurveModel::FBarRange& Range = Ranges[Index];
+				const double LowerSeconds = Range.Range.GetLowerBoundValue();
+				const double UpperSeconds = Range.Range.GetUpperBoundValue();
+				const bool bOutsideUpper = (Index != Ranges.Num() - 1) && UpperSeconds < InputMin;
+				if (LowerSeconds > InputMax || bOutsideUpper) //out of range
+				{
+					continue;
+				}
+				
+				FLinearColor CurveColor = Range.Color;
+
+				// Alpha blend the zebra tint
+				if (CurveIndex % 2)
+				{
+					static FLinearColor ZebraTint = FLinearColor::White.CopyWithNewOpacity(0.01f);
+					if (CurveColor == FLinearColor::White)
+					{
+						CurveColor = ZebraTint;
+					}
+					else
+					{
+						CurveColor = CurveColor * (1.f - ZebraTint.A) + ZebraTint * ZebraTint.A;
+					}
+				}
+
+				if (CurveColor != FLinearColor::White)
+				{
+					const double LowerSecondsForBox =(Index == 0 && LowerSeconds > InputMin )? InputMin : LowerSeconds;
+					const double BoxPos = CurveSpace.SecondsToScreen(LowerSecondsForBox);
+
+					const FPaintGeometry BoxGeometry = AllottedGeometry.ToPaintGeometry(
+						FVector2D(AllottedGeometry.GetLocalSize().X, TrackHeight),
+						FSlateLayoutTransform(FVector2D(BoxPos, LaneTop))
+					);
+
+					FSlateDrawElement::MakeBox(OutDrawElements, BaseLayerId + CurveViewConstants::ELayerOffset::Background, BoxGeometry, WhiteBrush, DrawEffects, CurveColor);
+				}
+				const double LowerSecondsForLabel = (InputMin > LowerSeconds) ? InputMin : LowerSeconds;
+				double LabelPos = CurveSpace.SecondsToScreen(LowerSecondsForLabel) + 10;
+
+				const FText Label = FText::FromName(Range.Name);
+				const FVector2D TextSize = FontMeasure->Measure(Label, FontInfo);
+				if (Index > 0)
+				{
+					LabelPos = (LabelPos < LastLabelPos) ? LastLabelPos + 5 : LabelPos;
+				}
+				LastLabelPos = LabelPos + TextSize.X + 15;
+				const FVector2D Position(LabelPos, LaneTop + (TrackHeight - TextSize.Y) * .5f);
+
+				const FPaintGeometry LabelGeometry = AllottedGeometry.ToPaintGeometry(
+					FSlateLayoutTransform(Position)
+				);
+
+				FSlateDrawElement::MakeText(OutDrawElements, BaseLayerId + CurveViewConstants::ELayerOffset::Labels, LabelGeometry, Label, FontInfo, DrawEffects, FLinearColor::White);
+
 			}
 		}
+	}
+}
 
-		if (CurveColor != FLinearColor::White)
-		{
-			const FPaintGeometry BoxGeometry = AllottedGeometry.ToPaintGeometry(
-				FVector2D(AllottedGeometry.GetLocalSize().X, TrackHeight),
-				FSlateLayoutTransform(FVector2D(0.f, LaneTop))
-			);
+void SCurveEditorKeyBarView::BuildContextMenu(FMenuBuilder & MenuBuilder, TOptional<FCurvePointHandle> ClickedPoint, TOptional<FCurveModelID> HoveredCurveID)
+{
+	TSharedPtr<FCurveEditor> CurveEditor = WeakCurveEditor.Pin();
+	if (!CurveEditor)
+	{
+		return;
+	}
+	int32 NumSelectedKeys = CurveEditor->GetSelection().Count();
 
-			FSlateDrawElement::MakeBox(OutDrawElements, BaseLayerId + CurveViewConstants::ELayerOffset::Background, BoxGeometry, WhiteBrush, DrawEffects, CurveColor);
-		}
+	FCurveModel* Curve = HoveredCurveID.IsSet() ? CurveEditor->FindCurve(HoveredCurveID.GetValue()) : nullptr;
+	FKeyBarCurveModel* KeyBarCurveModel = static_cast<FKeyBarCurveModel*>(Curve);
+	if (KeyBarCurveModel)
+	{
+		KeyBarCurveModel->BuildContextMenu(*CurveEditor,MenuBuilder, ClickedPoint);
+	}
+}
 
-		// Draw the curve label
-		{
-			const FText Label = Curve->GetLongDisplayName();
 
-			const FVector2D TextSize = FontMeasure->Measure(Label, FontInfo);
-			const FVector2D Position(LocalSize.X - TextSize.X - CurveViewConstants::CurveLabelOffsetX, LaneTop + (TrackHeight - TextSize.Y)*.5f);
 
-			const FPaintGeometry LabelGeometry = AllottedGeometry.ToPaintGeometry(
-				FSlateLayoutTransform(Position)
-			);
 
-			FSlateDrawElement::MakeText(OutDrawElements, BaseLayerId + CurveViewConstants::ELayerOffset::Labels, LabelGeometry, Label, FontInfo, DrawEffects, FLinearColor::White);
-		}
-	}
-}
\ No newline at end of file
diff --git a/Engine/Source/Editor/MovieSceneTools/Public/Channels/SCurveEditorKeyBarView.h b/Engine/Source/Editor/MovieSceneTools/Public/Channels/SCurveEditorKeyBarView.h
--- a/Engine/Source/Editor/MovieSceneTools/Public/Channels/SCurveEditorKeyBarView.h
+++ b/Engine/Source/Editor/MovieSceneTools/Public/Channels/SCurveEditorKeyBarView.h
@@ -21,6 +21,7 @@ private:
 	//~ SInteractiveCurveEditorView interface
 	virtual void PaintView(const FPaintArgs& Args, const FGeometry& AllottedGeometry, const FSlateRect& MyCullingRect, FSlateWindowElementList& OutDrawElements, int32 BaseLayerId, const FWidgetStyle& InWidgetStyle, bool bParentEnabled) const override;
 	virtual bool IsValueSnapEnabled() const override { return false; }
+	virtual void BuildContextMenu(FMenuBuilder& MenuBuilder, TOptional<FCurvePointHandle> ClickedPoint, TOptional<FCurveModelID> HoveredCurveID) override;
 
 	void DrawLabels(const FGeometry& AllottedGeometry, FSlateWindowElementList& OutDrawElements, int32 BaseLayerId, ESlateDrawEffect DrawEffects) const;
 
diff --git a/Engine/Source/Editor/Sequencer/Private/Sequencer.cpp b/Engine/Source/Editor/Sequencer/Private/Sequencer.cpp
--- a/Engine/Source/Editor/Sequencer/Private/Sequencer.cpp
+++ b/Engine/Source/Editor/Sequencer/Private/Sequencer.cpp
@@ -2694,6 +2694,11 @@ void FSequencer::RefreshTree()
 	}
 }
 
+void FSequencer::RecreateCurveEditor()
+{
+	GetNodeTree()->RecreateCurveEditor();
+}
+
 FAnimatedRange FSequencer::GetViewRange() const
 {
 	FAnimatedRange AnimatedRange(FMath::Lerp(LastViewRange.GetLowerBoundValue(), TargetViewRange.GetLowerBoundValue(), ZoomCurve.GetLerp()),
diff --git a/Engine/Source/Editor/Sequencer/Private/Sequencer.h b/Engine/Source/Editor/Sequencer/Private/Sequencer.h
--- a/Engine/Source/Editor/Sequencer/Private/Sequencer.h
+++ b/Engine/Source/Editor/Sequencer/Private/Sequencer.h
@@ -773,6 +773,7 @@ public:
 	virtual void SetPerspectiveViewportPossessionEnabled(bool bEnabled) override;
 	virtual void SetPerspectiveViewportCameraCutEnabled(bool bEnabled) override;
 	virtual void RenderMovie(const TArray<UMovieSceneCinematicShotSection*>& InSections) const override;
+	virtual void RecreateCurveEditor() override;
 	virtual void EnterSilentMode() override;
 	virtual void ExitSilentMode() override;
 	virtual bool IsInSilentMode() const override { return SilentModeCount != 0; }
diff --git a/Engine/Source/Editor/Sequencer/Private/SequencerNodeTree.cpp b/Engine/Source/Editor/Sequencer/Private/SequencerNodeTree.cpp
--- a/Engine/Source/Editor/Sequencer/Private/SequencerNodeTree.cpp
+++ b/Engine/Source/Editor/Sequencer/Private/SequencerNodeTree.cpp
@@ -1487,11 +1487,25 @@ TArray< TSharedRef<FSequencerDisplayNode> > FSequencerNodeTree::GetAllNodes() co
 	return AllNodes;
 }
 
+// this will just clear out the curve editor, it will be recreated on next tick which is safer
+void FSequencerNodeTree::RecreateCurveEditor()
+{
+	FCurveEditor* CurveEditor = Sequencer.GetCurveEditor().Get();
+	if (CurveEditor)
+	{
+		for (auto It = CurveEditorTreeItemIDs.CreateIterator(); It; ++It)
+		{
+			CurveEditor->RemoveTreeItem(It->Value);
+		}
+
+		CurveEditorTreeItemIDs.Empty();
+	}
+}
+
 void FSequencerNodeTree::UpdateCurveEditorTree()
 {
 	FCurveEditor*     CurveEditor     = Sequencer.GetCurveEditor().Get();
 	FCurveEditorTree* CurveEditorTree = CurveEditor->GetTree();
-
 	// Guard against multiple broadcasts here and defer them until the end of this function
 	FScopedCurveEditorTreeEventGuard ScopedEventGuard = CurveEditorTree->ScopedEventGuard();
 
diff --git a/Engine/Source/Editor/Sequencer/Private/SequencerNodeTree.h b/Engine/Source/Editor/Sequencer/Private/SequencerNodeTree.h
--- a/Engine/Source/Editor/Sequencer/Private/SequencerNodeTree.h
+++ b/Engine/Source/Editor/Sequencer/Private/SequencerNodeTree.h
@@ -217,6 +217,9 @@ public:
 	/** Returns whether any of the currently selected nodes are muted */
 	bool IsSelectedNodesMute() const;
 
+	/**Make sure associated curve editor get's recreated, usually due to an item getting added that needs to be sorted*/
+	void RecreateCurveEditor();
+
 private:
 
 	/** Returns whether this NodeTree should only display selected nodes */
diff --git a/Engine/Source/Editor/Sequencer/Public/ISequencer.h b/Engine/Source/Editor/Sequencer/Public/ISequencer.h
--- a/Engine/Source/Editor/Sequencer/Public/ISequencer.h
+++ b/Engine/Source/Editor/Sequencer/Public/ISequencer.h
@@ -434,6 +434,11 @@ public:
 	 */
 	virtual void RenderMovie(const TArray<UMovieSceneCinematicShotSection*>& InSections) const = 0;
 
+	/*
+	* Recreate any associated Curve Editor 
+	*/
+	virtual void RecreateCurveEditor() {};
+
 	/*
 	 * Puts sequencer in a silent state (whereby it will not redraw viewports, or attempt to update external state besides the sequence itself)
 	 */
