Control Rig: Use Component property name instead of Component FName for Component Reference

diff --git a/Engine/Plugins/Animation/ControlRig/Source/ControlRig/Private/ControlRigComponent.cpp b/Engine/Plugins/Animation/ControlRig/Source/ControlRig/Private/ControlRigComponent.cpp
--- a/Engine/Plugins/Animation/ControlRig/Source/ControlRig/Private/ControlRigComponent.cpp
+++ b/Engine/Plugins/Animation/ControlRig/Source/ControlRig/Private/ControlRigComponent.cpp
@@ -505,7 +505,7 @@ void UControlRigComponent::AddMappedComponents(TArray<FControlRigComponentMapped
 
 		FControlRigComponentMappedElement ElementToMap;
 		ElementToMap.ComponentReference.OtherActor = Component->GetOwner() != GetOwner() ? Component->GetOwner() : nullptr;
-		ElementToMap.ComponentReference.ComponentProperty = Component->GetFName();
+		ElementToMap.ComponentReference.ComponentProperty = GetComponentNameWithinActor(Component);
 
 		ElementToMap.ElementName = ComponentToMap.ElementName;
 		ElementToMap.ElementType = ComponentToMap.ElementType;
@@ -612,7 +612,7 @@ void UControlRigComponent::AddMappedSkeletalMesh(USkeletalMeshComponent* Skeleta
 
 		FControlRigComponentMappedElement ElementToMap;
 		ElementToMap.ComponentReference.OtherActor = SkeletalMeshComponent->GetOwner() != GetOwner() ? SkeletalMeshComponent->GetOwner() : nullptr;
-		ElementToMap.ComponentReference.ComponentProperty = SkeletalMeshComponent->GetFName();
+		ElementToMap.ComponentReference.ComponentProperty = GetComponentNameWithinActor(SkeletalMeshComponent);
 
 		ElementToMap.ElementName = BoneToMap.Source;
 		ElementToMap.ElementType = ERigElementType::Bone;
@@ -631,7 +631,7 @@ void UControlRigComponent::AddMappedSkeletalMesh(USkeletalMeshComponent* Skeleta
 
 		FControlRigComponentMappedElement ElementToMap;
 		ElementToMap.ComponentReference.OtherActor = SkeletalMeshComponent->GetOwner() != GetOwner() ? SkeletalMeshComponent->GetOwner() : nullptr;
-		ElementToMap.ComponentReference.ComponentProperty = SkeletalMeshComponent->GetFName();
+		ElementToMap.ComponentReference.ComponentProperty = GetComponentNameWithinActor(SkeletalMeshComponent);
 
 		ElementToMap.ElementName = CurveToMap.Source;
 		ElementToMap.ElementType = ERigElementType::Curve;
@@ -1527,6 +1527,43 @@ void UControlRigComponent::TransferOutputs()
 	}
 }
 
+FName UControlRigComponent::GetComponentNameWithinActor(UActorComponent* InComponent)
+{
+	check(InComponent);
+	
+	FName ComponentProperty = InComponent->GetFName(); 
+
+	if(AActor* Owner = InComponent->GetOwner())
+	{
+		// we need to see if the owner stores this component as a property
+		for (TFieldIterator<FProperty> PropertyIt(Owner->GetClass()); PropertyIt; ++PropertyIt)
+		{
+			if(const FObjectPropertyBase* Property = CastField<FObjectPropertyBase>(*PropertyIt))
+			{
+				if(Property->GetObjectPropertyValue_InContainer(Owner) == InComponent)
+				{
+					ComponentProperty = Property->GetFName();
+					break;
+				}
+			}
+		}
+
+#if WITH_EDITOR
+
+		// validate that the property storage will return the right component.
+		// this is a sanity check ensuring that ComponentReference will find the right component later.
+		FObjectPropertyBase* Property = FindFProperty<FObjectPropertyBase>(Owner->GetClass(), ComponentProperty);
+		if(Property != nullptr)
+		{
+			UActorComponent* FoundComponent = Cast<UActorComponent>(Property->GetObjectPropertyValue_InContainer(Owner));
+			check(FoundComponent == InComponent);
+		}
+		
+#endif
+	}
+	return ComponentProperty;
+}
+
 void UControlRigComponent::HandleControlRigInitializedEvent(UControlRig* InControlRig, const EControlRigState InState, const FName& InEventName)
 {
 #if WITH_EDITOR
diff --git a/Engine/Plugins/Animation/ControlRig/Source/ControlRig/Public/ControlRigComponent.h b/Engine/Plugins/Animation/ControlRig/Source/ControlRig/Public/ControlRigComponent.h
--- a/Engine/Plugins/Animation/ControlRig/Source/ControlRig/Public/ControlRigComponent.h
+++ b/Engine/Plugins/Animation/ControlRig/Source/ControlRig/Public/ControlRigComponent.h
@@ -629,6 +629,7 @@ private:
 	void ValidateMappingData();
 	void TransferInputs();
 	void TransferOutputs();
+	static FName GetComponentNameWithinActor(UActorComponent* InComponent);
 
 	void HandleControlRigInitializedEvent(UControlRig* InControlRig, const EControlRigState InState, const FName& InEventName);
 	void HandleControlRigPreSetupEvent(UControlRig* InControlRig, const EControlRigState InState, const FName& InEventName);
