Control Rig: Pose Library: Alt-Double Click on a pose will select the stored controls in the pose now

diff --git a/Engine/Plugins/Animation/ControlRig/Source/ControlRigEditor/Private/EditMode/SControlRigBaseListWidget.cpp b/Engine/Plugins/Animation/ControlRig/Source/ControlRigEditor/Private/EditMode/SControlRigBaseListWidget.cpp
--- a/Engine/Plugins/Animation/ControlRig/Source/ControlRigEditor/Private/EditMode/SControlRigBaseListWidget.cpp
+++ b/Engine/Plugins/Animation/ControlRig/Source/ControlRigEditor/Private/EditMode/SControlRigBaseListWidget.cpp
@@ -932,6 +932,16 @@ void SControlRigBaseListWidget::OnAssetsActivated(const TArray<FAssetData>& Sele
 			if (PoseAsset)
 			{
 				ExecutePastePose(PoseAsset);
+
+				// If alt is down, select controls
+				if (FSlateApplication::Get().GetModifierKeys().IsAltDown())
+				{
+					UControlRig* ControlRig = SControlRigPoseView::GetFirstControlRigInLevelSequence(GetControlRig());
+					if (ControlRig)
+					{
+						PoseAsset->SelectControls(ControlRig, SControlRigPoseView::IsMirror());
+					}
+				}
 			}
 		}
 		SelectThisAsset(Asset);
diff --git a/Engine/Plugins/Animation/ControlRig/Source/ControlRigEditor/Private/EditMode/SControlRigControlViews.cpp b/Engine/Plugins/Animation/ControlRig/Source/ControlRigEditor/Private/EditMode/SControlRigControlViews.cpp
--- a/Engine/Plugins/Animation/ControlRig/Source/ControlRigEditor/Private/EditMode/SControlRigControlViews.cpp
+++ b/Engine/Plugins/Animation/ControlRig/Source/ControlRigEditor/Private/EditMode/SControlRigControlViews.cpp
@@ -531,23 +531,10 @@ FReply SControlRigPoseView::OnPastePose()
 }
 
 FReply SControlRigPoseView::OnSelectControls()
-{
+{	
 	if (PoseAsset.IsValid())
 	{
-		UControlRig* ControlRig = GetControlRig();
-		if (ControlRig == nullptr)
-		{
-			ULevelSequence* LevelSequence = ULevelSequenceEditorBlueprintLibrary::GetCurrentLevelSequence();
-			if (LevelSequence)
-			{
-				TArray<FControlRigSequencerBindingProxy> Proxies = UControlRigSequencerEditorLibrary::GetControlRigs(LevelSequence);
-				if (Proxies.Num() > 0)
-				{
-					//MZ TODO when we have Mutliple Control Rig's active select more than one.
-					ControlRig = Proxies[0].ControlRig;
-				}
-			}
-		}
+		UControlRig* ControlRig = GetFirstControlRigInLevelSequence(GetControlRig());
 		if (ControlRig)
 		{
 			PoseAsset->SelectControls(ControlRig, bIsMirror);
@@ -638,6 +625,24 @@ TSharedRef<SWidget> SControlRigPoseView::GetThumbnailWidget()
 		];
 }
 
+UControlRig* SControlRigPoseView::GetFirstControlRigInLevelSequence(UControlRig* ControlRig)
+{
+	if (ControlRig == nullptr)
+	{
+		ULevelSequence* LevelSequence = ULevelSequenceEditorBlueprintLibrary::GetCurrentLevelSequence();
+		if (LevelSequence)
+		{
+			TArray<FControlRigSequencerBindingProxy> Proxies = UControlRigSequencerEditorLibrary::GetControlRigs(LevelSequence);
+			if (Proxies.Num() > 0)
+			{
+				//MZ TODO when we have Mutliple Control Rig's active select more than one.
+				ControlRig = Proxies[0].ControlRig;
+			}
+		}
+	}
+	return ControlRig;
+}
+
 UControlRig* SControlRigPoseView::GetControlRig()
 {
 	UControlRig* NewControlRig = nullptr;
diff --git a/Engine/Plugins/Animation/ControlRig/Source/ControlRigEditor/Private/EditMode/SControlRigControlViews.h b/Engine/Plugins/Animation/ControlRig/Source/ControlRigEditor/Private/EditMode/SControlRigControlViews.h
--- a/Engine/Plugins/Animation/ControlRig/Source/ControlRigEditor/Private/EditMode/SControlRigControlViews.h
+++ b/Engine/Plugins/Animation/ControlRig/Source/ControlRigEditor/Private/EditMode/SControlRigControlViews.h
@@ -97,5 +97,9 @@ private:
 			]
 			];
 	}
+
+public:
+	static bool IsMirror() { return bIsMirror; }
+	static UControlRig* GetFirstControlRigInLevelSequence(UControlRig* ControlRig);
 };
 
