Control Rig: Allow space in rig element names

diff --git a/Engine/Plugins/Animation/ControlRig/Source/ControlRig/Private/Rigs/RigHierarchy.cpp b/Engine/Plugins/Animation/ControlRig/Source/ControlRig/Private/Rigs/RigHierarchy.cpp
--- a/Engine/Plugins/Animation/ControlRig/Source/ControlRig/Private/Rigs/RigHierarchy.cpp
+++ b/Engine/Plugins/Animation/ControlRig/Source/ControlRig/Private/Rigs/RigHierarchy.cpp
@@ -675,9 +675,10 @@ void URigHierarchy::SanitizeName(FString& InOutName)
 	{
 		TCHAR& C = InOutName[i];
 
-		const bool bGoodChar = FChar::IsAlpha(C) ||							// Any letter
-			(C == '_') || (C == '-') || (C == '.') ||						// _  - . anytime
-			((i > 0) && (FChar::IsDigit(C)));								// 0-9 after the first character
+		const bool bGoodChar = FChar::IsAlpha(C) ||				// Any letter
+			(C == '_') || (C == '-') || (C == '.') ||			// _  - . anytime
+			((i > 0) && (FChar::IsDigit(C))) ||					// 0-9 after the first character
+			((i > 0) && (C== ' '));								// Space after the first character to support virtual bones
 
 		if (!bGoodChar)
 		{
