ControlRig & Sequencer: make sure sequencer updates its control rig instance after control rig bp package is fully loaded

diff --git a/Engine/Plugins/Animation/ControlRig/Source/ControlRig/Private/Sequencer/MovieSceneControlRigParameterTrack.cpp b/Engine/Plugins/Animation/ControlRig/Source/ControlRig/Private/Sequencer/MovieSceneControlRigParameterTrack.cpp
--- a/Engine/Plugins/Animation/ControlRig/Source/ControlRig/Private/Sequencer/MovieSceneControlRigParameterTrack.cpp
+++ b/Engine/Plugins/Animation/ControlRig/Source/ControlRig/Private/Sequencer/MovieSceneControlRigParameterTrack.cpp
@@ -365,6 +365,10 @@ void UMovieSceneControlRigParameterTrack::PostLoad()
 
 #if WITH_EDITOR
 	FCoreUObjectDelegates::OnEndLoadPackage.AddUObject(this, &UMovieSceneControlRigParameterTrack::HandlePackageDone);
+	if (ControlRig)
+	{
+		ControlRig->OnEndLoadPackage().AddUObject(this, &UMovieSceneControlRigParameterTrack::HandleControlRigPackageDone);
+	}
 #else
 	ReconstructControlRig();
 #endif
@@ -382,6 +386,15 @@ void UMovieSceneControlRigParameterTrack::HandlePackageDone(TConstArrayView<UPac
 
 	ReconstructControlRig();
 }
+
+void UMovieSceneControlRigParameterTrack::HandleControlRigPackageDone(UControlRig* InControlRig)
+{
+	if (ensure(ControlRig == InControlRig))
+	{
+		ControlRig->OnEndLoadPackage().RemoveAll(this);
+		ReconstructControlRig();
+	}
+}
 #endif
 
 
diff --git a/Engine/Plugins/Animation/ControlRig/Source/ControlRig/Public/ControlRig.h b/Engine/Plugins/Animation/ControlRig/Source/ControlRig/Public/ControlRig.h
--- a/Engine/Plugins/Animation/ControlRig/Source/ControlRig/Public/ControlRig.h
+++ b/Engine/Plugins/Animation/ControlRig/Source/ControlRig/Public/ControlRig.h
@@ -62,6 +62,11 @@ public:
 	/** Bindable event for external objects to be notified that a Control is Selected */
 	DECLARE_EVENT_ThreeParams(UControlRig, FControlSelectedEvent, UControlRig*, FRigControlElement*, bool);
 
+#if WITH_EDITOR
+	/** Bindable event for external objects to be notified that a control rig is fully end-loaded*/
+	DECLARE_EVENT_OneParam(UControlRig, FOnEndLoadPackage, UControlRig*);
+#endif
+
 	static const FName OwnerComponent;
 
 	UFUNCTION(BlueprintCallable, Category = ControlRig)
@@ -83,10 +88,21 @@ private:
 	/** Latest state being processed */
 	EControlRigState LatestExecutedState;
 
+	
 public:
 	virtual void Serialize(FArchive& Ar) override;
 	virtual void PostLoad() override;
 #if WITH_EDITOR
+
+private:
+	FOnEndLoadPackage EndLoadPackageEvent;
+
+public:
+	// these are needed so that sequencer can have a chance to update its 
+	// ControlRig instances after the package is fully end-loaded
+	void BroadCastEndLoadPackage() { EndLoadPackageEvent.Broadcast(this); }
+	FOnEndLoadPackage& OnEndLoadPackage() { return EndLoadPackageEvent; }
+
 	virtual void PostEditUndo() override;
 #endif
 
diff --git a/Engine/Plugins/Animation/ControlRig/Source/ControlRig/Public/Sequencer/MovieSceneControlRigParameterTrack.h b/Engine/Plugins/Animation/ControlRig/Source/ControlRig/Public/Sequencer/MovieSceneControlRigParameterTrack.h
--- a/Engine/Plugins/Animation/ControlRig/Source/ControlRig/Public/Sequencer/MovieSceneControlRigParameterTrack.h
+++ b/Engine/Plugins/Animation/ControlRig/Source/ControlRig/Public/Sequencer/MovieSceneControlRigParameterTrack.h
@@ -51,6 +51,8 @@ public:
 
 #if WITH_EDITOR
 	void HandlePackageDone(TConstArrayView<UPackage*> InPackages);
+	// control Rigs are ready only after its package is fully end-loaded
+	void HandleControlRigPackageDone(UControlRig* InControlRig);
 #endif
 
 public:
diff --git a/Engine/Plugins/Animation/ControlRig/Source/ControlRigDeveloper/Private/ControlRigBlueprint.cpp b/Engine/Plugins/Animation/ControlRig/Source/ControlRigDeveloper/Private/ControlRigBlueprint.cpp
--- a/Engine/Plugins/Animation/ControlRig/Source/ControlRigDeveloper/Private/ControlRigBlueprint.cpp
+++ b/Engine/Plugins/Animation/ControlRig/Source/ControlRigDeveloper/Private/ControlRigBlueprint.cpp
@@ -536,8 +536,27 @@ void UControlRigBlueprint::HandlePackageDone(TConstArrayView<UPackage*> InPackag
 
 	FCoreUObjectDelegates::OnEndLoadPackage.RemoveAll(this);
 	
+	PropagateHierarchyFromBPToInstances();
 	RecompileVM();
 	RequestControlRigInit();
+	BroadcastControlRigPackageDone();
+}
+
+void UControlRigBlueprint::BroadcastControlRigPackageDone()
+{
+	UControlRigBlueprintGeneratedClass* RigClass = GetControlRigBlueprintGeneratedClass();
+	UControlRig* CDO = Cast<UControlRig>(RigClass->GetDefaultObject(true /* create if needed */));
+	CDO->BroadCastEndLoadPackage();
+
+	TArray<UObject*> ArchetypeInstances;
+	CDO->GetArchetypeInstances(ArchetypeInstances);
+	for (UObject* Instance : ArchetypeInstances)
+	{
+		if (UControlRig* InstanceRig = Cast<UControlRig>(Instance))
+		{
+			InstanceRig->BroadCastEndLoadPackage();
+		}
+	}
 }
 #endif
 
diff --git a/Engine/Plugins/Animation/ControlRig/Source/ControlRigDeveloper/Public/ControlRigBlueprint.h b/Engine/Plugins/Animation/ControlRig/Source/ControlRigDeveloper/Public/ControlRigBlueprint.h
--- a/Engine/Plugins/Animation/ControlRig/Source/ControlRigDeveloper/Public/ControlRigBlueprint.h
+++ b/Engine/Plugins/Animation/ControlRig/Source/ControlRigDeveloper/Public/ControlRigBlueprint.h
@@ -560,6 +560,8 @@ private:
 
 #if WITH_EDITOR
 	void HandlePackageDone(TConstArrayView<UPackage*> InPackages);
+	// ControlRigBP, once end-loaded, will inform other ControlRig-Dependent systems that ControlRig instances are ready.
+	void BroadcastControlRigPackageDone();
 #endif
 
 	// Class used to temporarily cache all 
