Control Rig: [RigVM] - fix infinite recursion blocking everything from running

diff --git a/Engine/Plugins/Animation/ControlRig/Source/ControlRig/Private/ControlRig.cpp b/Engine/Plugins/Animation/ControlRig/Source/ControlRig/Private/ControlRig.cpp
--- a/Engine/Plugins/Animation/ControlRig/Source/ControlRig/Private/ControlRig.cpp
+++ b/Engine/Plugins/Animation/ControlRig/Source/ControlRig/Private/ControlRig.cpp
@@ -2780,7 +2780,9 @@ void UControlRig::SetControlCustomization(const FRigElementKey& InControl, const
 
 void UControlRig::PostInitInstanceIfRequired()
 {
-	if(GetHierarchy() == nullptr || GetVM() == nullptr)
+	static int re_entry = 0;
+	re_entry++;
+	if(GetHierarchy() == nullptr || VM == nullptr)
 	{
 		if(HasAnyFlags(RF_ClassDefaultObject))
 		{
@@ -2792,6 +2794,7 @@ void UControlRig::PostInitInstanceIfRequired()
 			PostInitInstance(CDO);
 		}
 	}
+	re_entry--;
 }
 
 void UControlRig::PostInitInstance(UControlRig* InCDO)
