splittable ControlRigGizmoDefinition struct in Blueprints

#preflght https://horde.devtools.epicgames.com/job/615da5f99a74fe0001c5afdf

diff --git a/Engine/Plugins/Animation/ControlRig/Source/ControlRig/Public/ControlRigGizmoLibrary.h b/Engine/Plugins/Animation/ControlRig/Source/ControlRig/Public/ControlRigGizmoLibrary.h
--- a/Engine/Plugins/Animation/ControlRig/Source/ControlRig/Public/ControlRigGizmoLibrary.h
+++ b/Engine/Plugins/Animation/ControlRig/Source/ControlRig/Public/ControlRigGizmoLibrary.h
@@ -22,13 +22,13 @@ struct CONTROLRIG_API FControlRigGizmoDefinition
 		Transform = FTransform::Identity;
 	}
 
-	UPROPERTY(EditAnywhere, Category = "Gizmo")
+	UPROPERTY(BlueprintReadOnly, EditAnywhere, Category = "Gizmo")
 	FName GizmoName;
 
-	UPROPERTY(EditAnywhere, Category = "Gizmo")
+	UPROPERTY(BlueprintReadOnly, EditAnywhere, Category = "Gizmo")
 	TSoftObjectPtr<UStaticMesh> StaticMesh;
 
-	UPROPERTY(EditAnywhere, Category = "Gizmo")
+	UPROPERTY(BlueprintReadOnly, EditAnywhere, Category = "Gizmo")
 	FTransform Transform;
 };
 
