fix default-enabled ed modes showing up in random order

diff --git a/Engine/Plugins/Animation/ControlRig/Source/ControlRigEditor/Private/ControlRigEditorModule.cpp b/Engine/Plugins/Animation/ControlRig/Source/ControlRigEditor/Private/ControlRigEditorModule.cpp
--- a/Engine/Plugins/Animation/ControlRig/Source/ControlRigEditor/Private/ControlRigEditorModule.cpp
+++ b/Engine/Plugins/Animation/ControlRig/Source/ControlRigEditor/Private/ControlRigEditorModule.cpp
@@ -276,13 +276,15 @@ void FControlRigEditorModule::StartupModule()
 		FControlRigEditMode::ModeName,
 		NSLOCTEXT("AnimationModeToolkit", "DisplayName", "Animation"),
 		FSlateIcon(FControlRigEditorStyle::Get().GetStyleSetName(), "ControlRigEditMode", "ControlRigEditMode.Small"),
-		true);
+		true,
+		8000);
 
 	FEditorModeRegistry::Get().RegisterMode<FControlRigEditorEditMode>(
 		FControlRigEditorEditMode::ModeName,
 		NSLOCTEXT("RiggingModeToolkit", "DisplayName", "Rigging"),
 		FSlateIcon(FControlRigEditorStyle::Get().GetStyleSetName(), "ControlRigEditMode", "ControlRigEditMode.Small"),
-		false);
+		false,
+		8500);
 
 
 	ControlRigGraphPanelNodeFactory = MakeShared<FControlRigGraphPanelNodeFactory>();
diff --git a/Engine/Plugins/Editor/GeometryMode/Source/GeometryMode/Private/GeometryModeModule.cpp b/Engine/Plugins/Editor/GeometryMode/Source/GeometryMode/Private/GeometryModeModule.cpp
--- a/Engine/Plugins/Editor/GeometryMode/Source/GeometryMode/Private/GeometryModeModule.cpp
+++ b/Engine/Plugins/Editor/GeometryMode/Source/GeometryMode/Private/GeometryModeModule.cpp
@@ -21,7 +21,8 @@ void FGeometryModeModule::StartupModule()
 		FGeometryEditingModes::EM_Geometry,
 		NSLOCTEXT("EditorModes", "GeometryMode", "Brush Editing"),
 		FSlateIcon(FEditorStyle::GetStyleSetName(), "LevelEditor.BrushEdit"),
-		true
+		true,
+		7000
 	);
 }
 
diff --git a/Engine/Plugins/Editor/ModelingToolsEditorMode/Source/ModelingToolsEditorMode/Private/ModelingToolsEditorMode.cpp b/Engine/Plugins/Editor/ModelingToolsEditorMode/Source/ModelingToolsEditorMode/Private/ModelingToolsEditorMode.cpp
--- a/Engine/Plugins/Editor/ModelingToolsEditorMode/Source/ModelingToolsEditorMode/Private/ModelingToolsEditorMode.cpp
+++ b/Engine/Plugins/Editor/ModelingToolsEditorMode/Source/ModelingToolsEditorMode/Private/ModelingToolsEditorMode.cpp
@@ -140,7 +140,8 @@ UModelingToolsEditorMode::UModelingToolsEditorMode()
 		EM_ModelingToolsEditorModeId,
 		LOCTEXT("ModelingToolsEditorModeName", "Modeling"),
 		FSlateIcon("ModelingToolsStyle", "LevelEditor.ModelingToolsMode", "LevelEditor.ModelingToolsMode.Small"),
-		true);
+		true,
+		5000);
 }
 
 UModelingToolsEditorMode::UModelingToolsEditorMode(FVTableHelper& Helper)
diff --git a/Engine/Plugins/Experimental/Animation/ContextualAnimation/Source/ContextualAnimationEditor/Private/ContextualAnimationEditor.cpp b/Engine/Plugins/Experimental/Animation/ContextualAnimation/Source/ContextualAnimationEditor/Private/ContextualAnimationEditor.cpp
--- a/Engine/Plugins/Experimental/Animation/ContextualAnimation/Source/ContextualAnimationEditor/Private/ContextualAnimationEditor.cpp
+++ b/Engine/Plugins/Experimental/Animation/ContextualAnimation/Source/ContextualAnimationEditor/Private/ContextualAnimationEditor.cpp
@@ -14,7 +14,8 @@ void FContextualAnimationEditorModule::StartupModule()
 	FEditorModeRegistry::Get().RegisterMode<FContextualAnimEdMode>(FContextualAnimEdMode::EM_ContextualAnimEdModeId, 
 		LOCTEXT("ContextualAnimEdModeEdModeName", "ContextualAnim"), 
 		FSlateIcon(FContextualAnimEditorStyle::GetStyleSetName(), "ContextualAnimEditor.Icon", "ContextualAnimEditor.Icon"),
-		true);
+		true,
+		9000);
 }
 
 void FContextualAnimationEditorModule::ShutdownModule()
diff --git a/Engine/Plugins/Experimental/ChaosEditor/Source/FractureEditor/Private/FractureEditorMode.cpp b/Engine/Plugins/Experimental/ChaosEditor/Source/FractureEditor/Private/FractureEditorMode.cpp
--- a/Engine/Plugins/Experimental/ChaosEditor/Source/FractureEditor/Private/FractureEditorMode.cpp
+++ b/Engine/Plugins/Experimental/ChaosEditor/Source/FractureEditor/Private/FractureEditorMode.cpp
@@ -42,7 +42,8 @@ UFractureEditorMode::UFractureEditorMode()
 		EM_FractureEditorModeId,
 		LOCTEXT("FractureToolsEditorModeName", "Fracture"),
 		FSlateIcon("FractureEditorStyle", "LevelEditor.FractureMode", "LevelEditor.FractureMode.Small"),
-		true);
+		true,
+		6000);
 }
 
 UFractureEditorMode::~UFractureEditorMode()
