optional effector transform for ControlRig's FABRIK node for consistency with Blueprint

#prefilght https://horde.devtools.epicgames.com/job/615da6cfe3c4b000019e5764

diff --git a/Engine/Plugins/Animation/ControlRig/Source/ControlRig/Private/Units/Highlevel/Hierarchy/RigUnit_FABRIK.cpp b/Engine/Plugins/Animation/ControlRig/Source/ControlRig/Private/Units/Highlevel/Hierarchy/RigUnit_FABRIK.cpp
--- a/Engine/Plugins/Animation/ControlRig/Source/ControlRig/Private/Units/Highlevel/Hierarchy/RigUnit_FABRIK.cpp
+++ b/Engine/Plugins/Animation/ControlRig/Source/ControlRig/Private/Units/Highlevel/Hierarchy/RigUnit_FABRIK.cpp
@@ -31,13 +31,14 @@ FRigUnit_FABRIK_Execute()
 		bPropagateToChildren,
 		MaxIterations,
 		WorkData,
+		bSetEffectorTransform,
 		ExecuteContext, 
 		Context);
 }
 
 FRigUnit_FABRIKPerItem_Execute()
 {
-	FRigUnit_FABRIKItemArray::StaticExecute(RigVMExecuteContext, Items.Keys, EffectorTransform, Precision, Weight, bPropagateToChildren, MaxIterations, WorkData, ExecuteContext, Context);
+	FRigUnit_FABRIKItemArray::StaticExecute(RigVMExecuteContext, Items.Keys, EffectorTransform, Precision, Weight, bPropagateToChildren, MaxIterations, WorkData, bSetEffectorTransform,ExecuteContext, Context);
 }
 
 FRigUnit_FABRIKItemArray_Execute()
@@ -167,16 +168,19 @@ FRigUnit_FABRIKItemArray_Execute()
 				}
 			}
 
-			if (FMath::IsNearlyEqual(Weight, 1.f))
+			if (bSetEffectorTransform)
 			{
-				Hierarchy->SetGlobalTransform(CachedEffector, EffectorTransform, bPropagateToChildren);
-			}
-			else
-			{
-				float T = FMath::Clamp<float>(Weight, 0.f, 1.f);
-				FTransform PreviousXfo = Hierarchy->GetGlobalTransform(CachedEffector);
-				FTransform Xfo = FControlRigMathLibrary::LerpTransform(PreviousXfo, EffectorTransform, T);
-				Hierarchy->SetGlobalTransform(CachedEffector, Xfo, bPropagateToChildren);
+				if (FMath::IsNearlyEqual(Weight, 1.f))
+				{
+					Hierarchy->SetGlobalTransform(CachedEffector, EffectorTransform, bPropagateToChildren);
+				}
+				else
+				{
+					float T = FMath::Clamp<float>(Weight, 0.f, 1.f);
+					FTransform PreviousXfo = Hierarchy->GetGlobalTransform(CachedEffector);
+					FTransform Xfo = FControlRigMathLibrary::LerpTransform(PreviousXfo, EffectorTransform, T);
+					Hierarchy->SetGlobalTransform(CachedEffector, Xfo, bPropagateToChildren);
+				}
 			}
 		}
 	}
diff --git a/Engine/Plugins/Animation/ControlRig/Source/ControlRig/Private/Units/Highlevel/Hierarchy/RigUnit_FABRIK.h b/Engine/Plugins/Animation/ControlRig/Source/ControlRig/Private/Units/Highlevel/Hierarchy/RigUnit_FABRIK.h
--- a/Engine/Plugins/Animation/ControlRig/Source/ControlRig/Private/Units/Highlevel/Hierarchy/RigUnit_FABRIK.h
+++ b/Engine/Plugins/Animation/ControlRig/Source/ControlRig/Private/Units/Highlevel/Hierarchy/RigUnit_FABRIK.h
@@ -46,6 +46,7 @@ struct CONTROLRIG_API FRigUnit_FABRIK : public FRigUnit_HighlevelBaseMutable
 		MaxIterations = 10;
 		EffectorTransform = FTransform::Identity;
 		bPropagateToChildren = true;
+		bSetEffectorTransform = true;
 	}
 
 	/**
@@ -94,6 +95,12 @@ struct CONTROLRIG_API FRigUnit_FABRIK : public FRigUnit_HighlevelBaseMutable
 
 	UPROPERTY(transient)
 	FRigUnit_FABRIK_WorkData WorkData;
+
+	/**
+	* The option to set the effector transform
+	*/
+	UPROPERTY(meta = (Input))
+	bool bSetEffectorTransform;
 };
 
 /**
@@ -116,6 +123,7 @@ struct CONTROLRIG_API FRigUnit_FABRIKPerItem : public FRigUnit_HighlevelBaseMuta
 		bPropagateToChildren = true;
 		MaxIterations = 10;
 		EffectorTransform = FTransform::Identity;
+		bSetEffectorTransform = true;
 	}
 
 	/**
@@ -158,6 +166,12 @@ struct CONTROLRIG_API FRigUnit_FABRIKPerItem : public FRigUnit_HighlevelBaseMuta
 
 	UPROPERTY(transient)
 	FRigUnit_FABRIK_WorkData WorkData;
+
+	/**
+	* The option to set the effector transform
+	*/
+	UPROPERTY(meta = (Input))
+	bool bSetEffectorTransform;
 };
 
 /**
@@ -180,6 +194,7 @@ struct CONTROLRIG_API FRigUnit_FABRIKItemArray : public FRigUnit_HighlevelBaseMu
 		bPropagateToChildren = true;
 		MaxIterations = 10;
 		EffectorTransform = FTransform::Identity;
+		bSetEffectorTransform = true;
 	}
 
 	/**
@@ -222,4 +237,10 @@ struct CONTROLRIG_API FRigUnit_FABRIKItemArray : public FRigUnit_HighlevelBaseMu
 
 	UPROPERTY(transient)
 	FRigUnit_FABRIK_WorkData WorkData;
+	
+	/**
+	* The option to set the effector transform
+	*/
+	UPROPERTY(meta = (Input))
+	bool bSetEffectorTransform;
 };
