Control Rig: Fix Transient Controls for bones not showing up in the viewport

diff --git a/Engine/Plugins/Animation/ControlRig/Source/ControlRig/Private/AnimNode_ControlRigBase.cpp b/Engine/Plugins/Animation/ControlRig/Source/ControlRig/Private/AnimNode_ControlRigBase.cpp
--- a/Engine/Plugins/Animation/ControlRig/Source/ControlRig/Private/AnimNode_ControlRigBase.cpp
+++ b/Engine/Plugins/Animation/ControlRig/Source/ControlRig/Private/AnimNode_ControlRigBase.cpp
@@ -142,7 +142,14 @@ void FAnimNode_ControlRigBase::UpdateInput(UControlRig* ControlRig, const FPoseC
 				FRigHierarchyValidityBracket ValidityBracket(ControlRig->GetHierarchy());
 			}
 			FRigHierarchyValidityBracket ValidityBracket(ControlRig->GetHierarchy());
-			ControlRig->GetHierarchy()->ResetPoseToInitial(ERigElementType::Bone);
+
+			{
+#if WITH_EDITOR
+				// make sure transient controls don't get reset
+				UControlRig::FTransientControlPoseScope PoseScope(ControlRig);
+#endif 
+				ControlRig->GetHierarchy()->ResetPoseToInitial(ERigElementType::Bone);
+			}
 		}
 
 		if(bTransferPoseInGlobalSpace || NodeMappingContainer.IsValid())
@@ -211,9 +218,6 @@ void FAnimNode_ControlRigBase::UpdateInput(UControlRig* ControlRig, const FPoseC
 			}
 		}
 		
-#if WITH_EDITOR
-		ControlRig->ApplyTransformOverrideForUserCreatedBones();
-#endif
 	}
 
 	if (InputSettings.bUpdateCurves && bTransferInputCurves)
diff --git a/Engine/Plugins/Animation/ControlRig/Source/ControlRig/Private/ControlRig.cpp b/Engine/Plugins/Animation/ControlRig/Source/ControlRig/Private/ControlRig.cpp
--- a/Engine/Plugins/Animation/ControlRig/Source/ControlRig/Private/ControlRig.cpp
+++ b/Engine/Plugins/Animation/ControlRig/Source/ControlRig/Private/ControlRig.cpp
@@ -724,39 +724,52 @@ void UControlRig::Execute(const EControlRigState InState, const FName& InEventNa
 					GetHierarchy()->CopyPose(CDO->GetHierarchy(), false, true);
 				}
 			}
-			
+
 			{
-				// save the current state of all pose elements to preserve user intention
-				// the saved pose is reapplied to the rig after setup event as the pose scope goes out of scope
-				TUniquePtr<UControlRig::FPoseScope> PoseScope;
-				if (!bSetupModeEnabled)
+#if WITH_EDITOR
+				TUniquePtr<FTransientControlPoseScope> TransientControlPoseScope;
+				if (bSetupModeEnabled)
 				{
-					// only do this in non-setup mode because 
-					// when setup mode is enabled, the control values are cleared before reaching here (too late to save them)
-					PoseScope = MakeUnique<UControlRig::FPoseScope>(this, ERigElementType::ToResetAfterSetupEvent);
+					// save the transient control value, it should not be constantly reset in setup mode
+					TransientControlPoseScope = MakeUnique<FTransientControlPoseScope>(this);
 				}
+#endif
+				{
+					// save the current state of all pose elements to preserve user intention
+					// the saved pose is reapplied to the rig after setup event as the pose scope goes out of scope
+					TUniquePtr<UControlRig::FPoseScope> PoseScope;
+					if (!bSetupModeEnabled)
+					{
+						// only do this in non-setup mode because 
+						// when setup mode is enabled, the control values are cleared before reaching here (too late to save them)
+						PoseScope = MakeUnique<UControlRig::FPoseScope>(this, ERigElementType::ToResetAfterSetupEvent);
+					}
 
-				// reset the pose to initial such that setup event can run from a deterministic initial state
-				GetHierarchy()->ResetPoseToInitial(ERigElementType::All);
+					{
 
-				if (PreSetupEvent.IsBound())
-				{
-					FControlRigBracketScope BracketScope(PreSetupBracket);
-					PreSetupEvent.Broadcast(this, EControlRigState::Update, FRigUnit_PrepareForExecution::EventName);
-				}
+						// reset the pose to initial such that setup event can run from a deterministic initial state
+						GetHierarchy()->ResetPoseToInitial(ERigElementType::All);
+					}
 
-				ExecuteUnits(Context, FRigUnit_PrepareForExecution::EventName);
+					if (PreSetupEvent.IsBound())
+					{
+						FControlRigBracketScope BracketScope(PreSetupBracket);
+						PreSetupEvent.Broadcast(this, EControlRigState::Update, FRigUnit_PrepareForExecution::EventName);
+					}
 
-				if (PostSetupEvent.IsBound())
-				{
-					FControlRigBracketScope BracketScope(PostSetupBracket);
-					PostSetupEvent.Broadcast(this, EControlRigState::Update, FRigUnit_PrepareForExecution::EventName);
+					ExecuteUnits(Context, FRigUnit_PrepareForExecution::EventName);
+
+					if (PostSetupEvent.IsBound())
+					{
+						FControlRigBracketScope BracketScope(PostSetupBracket);
+						PostSetupEvent.Broadcast(this, EControlRigState::Update, FRigUnit_PrepareForExecution::EventName);
+					}
 				}
-			}
 
-			if (bSetupModeEnabled)
-			{
-				GetHierarchy()->ResetPoseToInitial(ERigElementType::Bone);
+				if (bSetupModeEnabled)
+				{
+					GetHierarchy()->ResetPoseToInitial(ERigElementType::Bone);
+				}
 			}
 		}
 		else
@@ -777,6 +790,9 @@ void UControlRig::Execute(const EControlRigState InState, const FName& InEventNa
 			{
 				VM->SetFirstEntryEventInEventQueue(EventQueue[0]);
 			}
+
+			// Transform Overrride is generated using a Transient Control 
+			ApplyTransformOverrideForUserCreatedBones();
 #endif
 
 			if (InState == EControlRigState::Update && InEventName == FRigUnit_BeginExecution::EventName)
@@ -2780,6 +2796,8 @@ void UControlRig::SetControlCustomization(const FRigElementKey& InControl, const
 
 void UControlRig::PostInitInstanceIfRequired()
 {
+	static int re_entry = 0;
+	re_entry++;
 	if(GetHierarchy() == nullptr || VM == nullptr)
 	{
 		if(HasAnyFlags(RF_ClassDefaultObject))
@@ -2792,6 +2810,7 @@ void UControlRig::PostInitInstanceIfRequired()
 			PostInitInstance(CDO);
 		}
 	}
+	re_entry--;
 }
 
 void UControlRig::PostInitInstance(UControlRig* InCDO)
@@ -2895,6 +2914,27 @@ UControlRig::FTransientControlScope::~FTransientControlScope()
 	}
 }
 
+UControlRig::FTransientControlPoseScope::FTransientControlPoseScope(TObjectPtr<UControlRig> InControlRig)
+{
+	ControlRig = InControlRig;
+
+	TArray<FRigControlElement*> TransientControls = ControlRig->GetHierarchy()->GetTransientControls();
+	TArray<FRigElementKey> Keys;
+	for(FRigControlElement* TransientControl : TransientControls)
+	{
+		Keys.Add(TransientControl->GetKey());
+	}
+	
+	CachedPose = ControlRig->GetHierarchy()->GetPose(false, ERigElementType::Control, TArrayView<FRigElementKey>(Keys));
+}
+
+UControlRig::FTransientControlPoseScope::~FTransientControlPoseScope()
+{
+	check(ControlRig);
+
+	ControlRig->GetHierarchy()->SetPose(CachedPose);
+}
+
 #endif
  
 #undef LOCTEXT_NAMESPACE
diff --git a/Engine/Plugins/Animation/ControlRig/Source/ControlRig/Public/ControlRig.h b/Engine/Plugins/Animation/ControlRig/Source/ControlRig/Public/ControlRig.h
--- a/Engine/Plugins/Animation/ControlRig/Source/ControlRig/Public/ControlRig.h
+++ b/Engine/Plugins/Animation/ControlRig/Source/ControlRig/Public/ControlRig.h
@@ -769,6 +769,24 @@ private:
 		TObjectPtr<URigHierarchy> Hierarchy;
 	};
 
+public:
+	// Class used to temporarily cache current pose of transient controls
+	// restore them after a ResetPoseToInitial call,
+	// which allows user to move bones in setup mode
+	class FTransientControlPoseScope
+	{
+	public:
+		FTransientControlPoseScope(TObjectPtr<UControlRig> InControlRig);
+		~FTransientControlPoseScope();
+	
+	private:
+		
+		UControlRig* ControlRig;
+		FRigPose CachedPose;	
+	};	
+
+private:
+	
 #endif
 
 	friend class FControlRigBlueprintCompilerContext;
diff --git a/Engine/Plugins/Animation/ControlRig/Source/ControlRigEditor/Private/EditMode/ControlRigEditMode.cpp b/Engine/Plugins/Animation/ControlRig/Source/ControlRigEditor/Private/EditMode/ControlRigEditMode.cpp
--- a/Engine/Plugins/Animation/ControlRig/Source/ControlRigEditor/Private/EditMode/ControlRigEditMode.cpp
+++ b/Engine/Plugins/Animation/ControlRig/Source/ControlRigEditor/Private/EditMode/ControlRigEditMode.cpp
@@ -1658,16 +1658,35 @@ void FControlRigEditMode::SetRigElementSelection(ERigElementType Type, const TAr
 
 TArray<FRigElementKey> FControlRigEditMode::GetSelectedRigElements() const
 {
+	TArray<FRigElementKey> SelectedKeys;
+	
 	if (UControlRig* ControlRig = GetControlRig(true))
 	{
 		if (ControlRig->GetHierarchy())
 		{
-			return ControlRig->GetHierarchy()->GetSelectedKeys();
+			SelectedKeys = ControlRig->GetHierarchy()->GetSelectedKeys();
+		}
+
+		// currently only 1 transient control is allowed at a time
+		// Transient Control's bSelected flag is never set to true, probably to avoid confusing other parts of the system
+		// But since Edit Mode directly deals with transient controls, its selection status is given special treatment here.
+		// So basically, whenever a bone is selected, and there is a transient control present, we consider both selected.
+		if (SelectedKeys.Num() == 1)
+		{
+			if (SelectedKeys[0].Type == ERigElementType::Bone || SelectedKeys[0].Type == ERigElementType::Null)
+			{
+				const FName ControlName = UControlRig::GetNameForTransientControl(SelectedKeys[0]);
+				const FRigElementKey TransientControlKey = FRigElementKey(ControlName, ERigElementType::Control);
+				if(ControlRig->GetHierarchy()->Contains(TransientControlKey))
+				{
+					SelectedKeys.Add(TransientControlKey);
+				}
+
+			}
 		}
 	}
 
-	TArray<FRigElementKey> Empty;
-	return Empty;
+	return SelectedKeys;
 }
 
 bool FControlRigEditMode::AreRigElementsSelected(uint32 InTypes) const
