Control Rig: Fix stepping into loop unit test not passing

diff --git a/Engine/Plugins/Animation/ControlRig/Source/ControlRig/Public/ControlRig.h b/Engine/Plugins/Animation/ControlRig/Source/ControlRig/Public/ControlRig.h
--- a/Engine/Plugins/Animation/ControlRig/Source/ControlRig/Public/ControlRig.h
+++ b/Engine/Plugins/Animation/ControlRig/Source/ControlRig/Public/ControlRig.h
@@ -179,13 +179,22 @@ public:
 	template<class T>
 	FORCEINLINE T GetPublicVariableValue(const FName& InVariableName)
 	{
-		return GetPublicVariableByName(InVariableName).GetValue<T>();
+		FRigVMExternalVariable Variable = GetPublicVariableByName(InVariableName);
+		if (Variable.IsValid())
+		{
+			return Variable.GetValue<T>();
+		}
+		return T();
 	}
 
 	template<class T>
 	FORCEINLINE void SetPublicVariableValue(const FName& InVariableName, const T& InValue)
 	{
-		GetPublicVariableByName(InVariableName).SetValue<T>(InValue);
+		FRigVMExternalVariable Variable = GetPublicVariableByName(InVariableName);
+		if (Variable.IsValid())
+		{
+			Variable.SetValue<T>(InValue);
+		}
 	}
 
 	template<class T>
