Control Rig: Post fix and first compile, compile warning in viewport still is there

diff --git a/Engine/Plugins/Animation/ControlRig/Source/ControlRigEditor/Private/Editor/ControlRigEditor.cpp b/Engine/Plugins/Animation/ControlRig/Source/ControlRigEditor/Private/Editor/ControlRigEditor.cpp
--- a/Engine/Plugins/Animation/ControlRig/Source/ControlRigEditor/Private/Editor/ControlRigEditor.cpp
+++ b/Engine/Plugins/Animation/ControlRig/Source/ControlRigEditor/Private/Editor/ControlRigEditor.cpp
@@ -2738,6 +2738,8 @@ void FControlRigEditor::HandleVMCompiledEvent(UBlueprint* InBlueprint, URigVM* I
 	}
 	
 #endif
+
+	UpdateGraphCompilerErrors();
 }
 
 void FControlRigEditor::HandleControlRigExecutedEvent(UControlRig* InControlRig, const EControlRigState InState, const FName& InEventName)
@@ -6145,11 +6147,6 @@ void FControlRigEditor::UpdateGraphCompilerErrors()
 	UControlRigBlueprint* Blueprint = Cast<UControlRigBlueprint>(GetBlueprintObj());
 	if (Blueprint && ControlRig && ControlRig->VM)
 	{
-		if (Blueprint->Status == BS_Error)
-		{
-			return;
-		}
-
 		if (ControlRigLog.Entries.Num() == 0 && !bAnyErrorsLeft)
 		{
 			return;
