Control Rig: Fixed a bug with Position Contraint Node

diff --git a/Engine/Plugins/Animation/ControlRig/Source/ControlRig/Private/Units/Highlevel/Hierarchy/RigUnit_TransformConstraint.cpp b/Engine/Plugins/Animation/ControlRig/Source/ControlRig/Private/Units/Highlevel/Hierarchy/RigUnit_TransformConstraint.cpp
--- a/Engine/Plugins/Animation/ControlRig/Source/ControlRig/Private/Units/Highlevel/Hierarchy/RigUnit_TransformConstraint.cpp
+++ b/Engine/Plugins/Animation/ControlRig/Source/ControlRig/Private/Units/Highlevel/Hierarchy/RigUnit_TransformConstraint.cpp
@@ -490,7 +490,7 @@ FRigUnit_PositionConstraint_Execute()
 
 			// handle filtering, performed in local space
 			FTransform ChildParentGlobalTransform = Hierarchy->GetParentTransform(Child, false);
-			FVector MixedPosition = ChildParentGlobalTransform.Inverse().TransformVector(MixedGlobalPosition);
+			FVector MixedPosition = ChildParentGlobalTransform.Inverse().TransformPosition(MixedGlobalPosition);
 			
 			FTransform ChildCurrentLocalTransform = Hierarchy->GetLocalTransform(Child);
 			
