Control Rig: Fix edited default values after adding more than 1 element are not saved post compile

diff --git a/Engine/Plugins/Animation/ControlRig/Source/ControlRigEditor/Private/Editor/ControlRigEditor.cpp b/Engine/Plugins/Animation/ControlRig/Source/ControlRigEditor/Private/Editor/ControlRigEditor.cpp
--- a/Engine/Plugins/Animation/ControlRig/Source/ControlRigEditor/Private/Editor/ControlRigEditor.cpp
+++ b/Engine/Plugins/Animation/ControlRig/Source/ControlRigEditor/Private/Editor/ControlRigEditor.cpp
@@ -4639,10 +4639,7 @@ void FControlRigEditor::UpdateDefaultValueForVariable(FBPVariableDescription& In
 				{
 					FString NewDefaultValue;
 					TargetProperty->ExportTextItem(NewDefaultValue, OldPropertyAddr, OldPropertyAddr, nullptr, PPF_None);
-					if(NewDefaultValue.Len() < 80)
-					{
-						InVariable.DefaultValue = NewDefaultValue;
-					}
+					InVariable.DefaultValue = NewDefaultValue;
 				}
 			}
 		}
