dump python command when selecting elements in the RigHierarchy

diff --git a/Engine/Plugins/Animation/ControlRig/Source/ControlRig/Private/Rigs/RigHierarchyController.cpp b/Engine/Plugins/Animation/ControlRig/Source/ControlRig/Private/Rigs/RigHierarchyController.cpp
--- a/Engine/Plugins/Animation/ControlRig/Source/ControlRig/Private/Rigs/RigHierarchyController.cpp
+++ b/Engine/Plugins/Animation/ControlRig/Source/ControlRig/Private/Rigs/RigHierarchyController.cpp
@@ -97,7 +97,7 @@ bool URigHierarchyController::SelectElement(FRigElementKey InKey, bool bSelect,
 	return true;
 }
 
-bool URigHierarchyController::SetSelection(TArray<FRigElementKey> InKeys)
+bool URigHierarchyController::SetSelection(const TArray<FRigElementKey>& InKeys, bool bPrintPythonCommand)
 {
 	if(!IsValid())
 	{
@@ -115,30 +115,53 @@ bool URigHierarchyController::SetSelection(TArray<FRigElementKey> InKeys)
 	TArray<FRigElementKey> PreviousSelection = Hierarchy->GetSelectedKeys();
 	bool bResult = false;
 
-	for(const FRigElementKey& KeyToDeselect : PreviousSelection)
 	{
-		if(!InKeys.Contains(KeyToDeselect))
+		// disable python printing here as we only want to print a single command instead of one per selected item
+		const TGuardValue<bool> Guard(bSuspendPythonPrinting, true);
+
+		for(const FRigElementKey& KeyToDeselect : PreviousSelection)
 		{
-			if(!SelectElement(KeyToDeselect, false))
+			if(!InKeys.Contains(KeyToDeselect))
 			{
-				return false;
+				if(!SelectElement(KeyToDeselect, false))
+				{
+					return false;
+				}
+				bResult = true;
 			}
-			bResult = true;
 		}
-	}
 
-	for(const FRigElementKey& KeyToSelect : InKeys)
-	{
-		if(!PreviousSelection.Contains(KeyToSelect))
+		for(const FRigElementKey& KeyToSelect : InKeys)
 		{
-			if(!SelectElement(KeyToSelect, true))
+			if(!PreviousSelection.Contains(KeyToSelect))
 			{
-				return false;
+				if(!SelectElement(KeyToSelect, true))
+				{
+					return false;
+				}
+				bResult = true;
 			}
-			bResult = true;
 		}
 	}
-
+	
+#if WITH_EDITOR
+	if (bPrintPythonCommand && !bSuspendPythonPrinting)
+	{
+		UBlueprint* Blueprint = GetTypedOuter<UBlueprint>();
+		if (Blueprint)
+		{
+			const FString Selection = FString::JoinBy( InKeys, TEXT(", "), [](const FRigElementKey& Key)
+			{
+				return Key.ToPythonString();
+			});
+			
+			RigVMPythonUtils::Print(Blueprint->GetFName().ToString(),
+				FString::Printf(TEXT("hierarchy_controller.set_selection([%s])"),
+				*Selection ) );
+		}
+	}
+#endif
+	
 	return bResult;
 }
 
@@ -182,7 +205,7 @@ FRigElementKey URigHierarchyController::AddBone(FName InName, FRigElementKey InP
 #if WITH_EDITOR
 	TransactionPtr.Reset();
 
-	if (bPrintPythonCommand)
+	if (bPrintPythonCommand && !bSuspendPythonPrinting)
 	{
 		UBlueprint* Blueprint = GetTypedOuter<UBlueprint>();
 		if (Blueprint)
@@ -243,7 +266,7 @@ FRigElementKey URigHierarchyController::AddNull(FName InName, FRigElementKey InP
 #if WITH_EDITOR
 	TransactionPtr.Reset();
 
-	if (bPrintPythonCommand)
+	if (bPrintPythonCommand && !bSuspendPythonPrinting)
 	{
 		UBlueprint* Blueprint = GetTypedOuter<UBlueprint>();
 		if (Blueprint)
@@ -311,7 +334,7 @@ FRigElementKey URigHierarchyController::AddControl(
 #if WITH_EDITOR
 	TransactionPtr.Reset();
 
-	if (bPrintPythonCommand)
+	if (bPrintPythonCommand && !bSuspendPythonPrinting)
 	{
 		UBlueprint* Blueprint = GetTypedOuter<UBlueprint>();
 		if (Blueprint)
@@ -359,7 +382,7 @@ FRigElementKey URigHierarchyController::AddCurve(FName InName, float InValue, bo
 #if WITH_EDITOR
 	TransactionPtr.Reset();
 
-	if (bPrintPythonCommand)
+	if (bPrintPythonCommand && !bSuspendPythonPrinting)
 	{
 		UBlueprint* Blueprint = GetTypedOuter<UBlueprint>();
 		if (Blueprint)
@@ -411,7 +434,7 @@ FRigElementKey URigHierarchyController::AddRigidBody(FName InName, FRigElementKe
 #if WITH_EDITOR
 	TransactionPtr.Reset();
 
-	if (bPrintPythonCommand)
+	if (bPrintPythonCommand && !bSuspendPythonPrinting)
 	{
 		UBlueprint* Blueprint = GetTypedOuter<UBlueprint>();
 		if (Blueprint)
@@ -712,21 +735,18 @@ TArray<FRigElementKey> URigHierarchyController::ImportBones(USkeleton* InSkeleto
 	                   bSelectBones, bSetupUndo);
 
 #if WITH_EDITOR
-	if (!BoneKeys.IsEmpty() && bPrintPythonCommand)
+	if (!BoneKeys.IsEmpty() && bPrintPythonCommand && !bSuspendPythonPrinting)
 	{
-		if (bPrintPythonCommand)
+		UBlueprint* Blueprint = GetTypedOuter<UBlueprint>();
+		if (Blueprint)
 		{
-			UBlueprint* Blueprint = GetTypedOuter<UBlueprint>();
-			if (Blueprint)
-			{
-				RigVMPythonUtils::Print(Blueprint->GetFName().ToString(),
-					FString::Printf(TEXT("hierarchy_controller.import_bones_from_asset('%s', '%s', %s, %s, %s)"),
-					*InSkeleton->GetPathName(),
-					*InNameSpace.ToString(),
-					(bReplaceExistingBones) ? TEXT("True") : TEXT("False"),
-					(bRemoveObsoleteBones) ? TEXT("True") : TEXT("False"),
-					(bSelectBones) ? TEXT("True") : TEXT("False")));
-			}
+			RigVMPythonUtils::Print(Blueprint->GetFName().ToString(),
+				FString::Printf(TEXT("hierarchy_controller.import_bones_from_asset('%s', '%s', %s, %s, %s)"),
+				*InSkeleton->GetPathName(),
+				*InNameSpace.ToString(),
+				(bReplaceExistingBones) ? TEXT("True") : TEXT("False"),
+				(bRemoveObsoleteBones) ? TEXT("True") : TEXT("False"),
+				(bSelectBones) ? TEXT("True") : TEXT("False")));
 		}
 	}
 #endif
@@ -822,19 +842,16 @@ TArray<FRigElementKey> URigHierarchyController::ImportCurves(USkeleton* InSkelet
 	}
 	
 #if WITH_EDITOR
-	if (!Keys.IsEmpty() && bPrintPythonCommand)
+	if (!Keys.IsEmpty() && bPrintPythonCommand && !bSuspendPythonPrinting)
 	{
-		if (bPrintPythonCommand)
+		UBlueprint* Blueprint = GetTypedOuter<UBlueprint>();
+		if (Blueprint)
 		{
-			UBlueprint* Blueprint = GetTypedOuter<UBlueprint>();
-			if (Blueprint)
-			{
-				RigVMPythonUtils::Print(Blueprint->GetFName().ToString(),
-					FString::Printf(TEXT("hierarchy_controller.import_curves_from_asset('%s', '%s', %s)"),
-					*InSkeleton->GetPathName(),
-					*InNameSpace.ToString(),
-					(bSelectCurves) ? TEXT("True") : TEXT("False")));
-			}
+			RigVMPythonUtils::Print(Blueprint->GetFName().ToString(),
+				FString::Printf(TEXT("hierarchy_controller.import_curves_from_asset('%s', '%s', %s)"),
+				*InSkeleton->GetPathName(),
+				*InNameSpace.ToString(),
+				(bSelectCurves) ? TEXT("True") : TEXT("False")));
 		}
 	}
 #endif
@@ -1168,7 +1185,7 @@ TArray<FRigElementKey> URigHierarchyController::ImportFromText(FString InContent
 	}
 
 #if WITH_EDITOR
-	if (bPrintPythonCommands)
+	if (bPrintPythonCommands && !bSuspendPythonPrinting)
 	{
 		UBlueprint* Blueprint = GetTypedOuter<UBlueprint>();
 		if (Blueprint)
@@ -1599,7 +1616,7 @@ bool URigHierarchyController::RemoveElement(FRigElementKey InElement, bool bSetu
 	}
 	TransactionPtr.Reset();
 
-	if (bRemoved && bPrintPythonCommand)
+	if (bRemoved && bPrintPythonCommand && !bSuspendPythonPrinting)
 	{
 		UBlueprint* Blueprint = GetTypedOuter<UBlueprint>();
 		if (Blueprint)
@@ -1766,7 +1783,7 @@ FRigElementKey URigHierarchyController::RenameElement(FRigElementKey InElement,
 		ClearSelection();
 	}
 
-	if (bRenamed && bPrintPythonCommand)
+	if (bRenamed && bPrintPythonCommand && !bSuspendPythonPrinting)
 	{
 		UBlueprint* Blueprint = GetTypedOuter<UBlueprint>();
 		if (Blueprint)
@@ -2070,7 +2087,7 @@ bool URigHierarchyController::RemoveParent(FRigElementKey InChild, FRigElementKe
 	}
 	TransactionPtr.Reset();
 
-	if (bRemoved && bPrintPythonCommand)
+	if (bRemoved && bPrintPythonCommand && !bSuspendPythonPrinting)
 	{
 		UBlueprint* Blueprint = GetTypedOuter<UBlueprint>();
 		if (Blueprint)
@@ -2250,7 +2267,7 @@ bool URigHierarchyController::RemoveAllParents(FRigElementKey InChild, bool bMai
 	}
 	TransactionPtr.Reset();
 	
-	if (bRemoved && bPrintPythonCommand)
+	if (bRemoved && bPrintPythonCommand && !bSuspendPythonPrinting)
 	{
 		UBlueprint* Blueprint = GetTypedOuter<UBlueprint>();
 		if (Blueprint)
@@ -2329,7 +2346,7 @@ bool URigHierarchyController::SetParent(FRigElementKey InChild, FRigElementKey I
 	}
 	TransactionPtr.Reset();
 
-	if (bParentSet && bPrintPythonCommand)
+	if (bParentSet && bPrintPythonCommand && !bSuspendPythonPrinting)
 	{
 		UBlueprint* Blueprint = GetTypedOuter<UBlueprint>();
 		if (Blueprint)
@@ -2352,7 +2369,7 @@ TArray<FRigElementKey> URigHierarchyController::DuplicateElements(TArray<FRigEle
 	TArray<FRigElementKey> Result = ImportFromText(Content, false, bSelectNewElements, bSetupUndo);
 
 #if WITH_EDITOR
-	if (!Result.IsEmpty() && bPrintPythonCommands)
+	if (!Result.IsEmpty() && bPrintPythonCommands && !bSuspendPythonPrinting)
 	{
 		UBlueprint* Blueprint = GetTypedOuter<UBlueprint>();
 		if (Blueprint)
@@ -2474,7 +2491,7 @@ TArray<FRigElementKey> URigHierarchyController::MirrorElements(TArray<FRigElemen
 	}
 
 #if WITH_EDITOR
-	if (!DuplicatedKeys.IsEmpty() && bPrintPythonCommands)
+	if (!DuplicatedKeys.IsEmpty() && bPrintPythonCommands && !bSuspendPythonPrinting)
 	{
 		UBlueprint* Blueprint = GetTypedOuter<UBlueprint>();
 		if (Blueprint)
diff --git a/Engine/Plugins/Animation/ControlRig/Source/ControlRig/Public/Rigs/RigHierarchyController.h b/Engine/Plugins/Animation/ControlRig/Source/ControlRig/Public/Rigs/RigHierarchyController.h
--- a/Engine/Plugins/Animation/ControlRig/Source/ControlRig/Public/Rigs/RigHierarchyController.h
+++ b/Engine/Plugins/Animation/ControlRig/Source/ControlRig/Public/Rigs/RigHierarchyController.h
@@ -20,6 +20,7 @@ public:
 	: bReportWarningsAndErrors(true)
 	, Hierarchy(nullptr)
 	, bSuspendNotifications(false)
+	, bSuspendPythonPrinting(false)
 	{}
 
 	virtual ~URigHierarchyController();
@@ -61,7 +62,7 @@ public:
 	 * @return Returns true if the selection was applied
 	 */
 	UFUNCTION(BlueprintCallable, Category = URigHierarchyController)
-    bool SetSelection(TArray<FRigElementKey> InKeys);
+    bool SetSelection(const TArray<FRigElementKey>& InKeys, bool bPrintPythonCommand = false);
 
 	/**
 	 * Clears the selection
@@ -608,6 +609,11 @@ private:
 	 */
 	bool bSuspendNotifications;
 
+	/** 
+	* If set to true all python printing can be disabled.  
+	*/
+	bool bSuspendPythonPrinting;
+
 	/**
 	 * This function can be used to override the controller's logging mechanism
 	 */
diff --git a/Engine/Plugins/Animation/ControlRig/Source/ControlRigEditor/Private/Editor/SRigHierarchy.cpp b/Engine/Plugins/Animation/ControlRig/Source/ControlRigEditor/Private/Editor/SRigHierarchy.cpp
--- a/Engine/Plugins/Animation/ControlRig/Source/ControlRigEditor/Private/Editor/SRigHierarchy.cpp
+++ b/Engine/Plugins/Animation/ControlRig/Source/ControlRigEditor/Private/Editor/SRigHierarchy.cpp
@@ -563,7 +563,7 @@ void SRigHierarchy::OnSelectionChanged(TSharedPtr<FRigTreeElement> Selection, ES
 		TGuardValue<bool> SuspendDetailsPanelRefreshGuard(*SuspensionFlagPtr, true);
 		
 		const TArray<FRigElementKey> NewSelection = GetSelectedKeys();
-		if(!Controller->SetSelection(NewSelection))
+		if(!Controller->SetSelection(NewSelection, true))
 		{
 			return;
 		}
