only fit in ControlRig's viewport when necessary

diff --git a/Engine/Plugins/Animation/ControlRig/Source/ControlRigEditor/Private/EditMode/ControlRigEditMode.cpp b/Engine/Plugins/Animation/ControlRig/Source/ControlRigEditor/Private/EditMode/ControlRigEditMode.cpp
--- a/Engine/Plugins/Animation/ControlRig/Source/ControlRigEditor/Private/EditMode/ControlRigEditMode.cpp
+++ b/Engine/Plugins/Animation/ControlRig/Source/ControlRigEditor/Private/EditMode/ControlRigEditMode.cpp
@@ -1983,6 +1983,7 @@ void FControlRigEditMode::FrameSelection()
 			Bounds += Sphere.Center + FVector::OneVector * Sphere.W;
 			Bounds += Sphere.Center - FVector::OneVector * Sphere.W;
 			CurrentViewportClient->FocusViewportOnBox(Bounds);
+			return;
 		}
     }
 
