Spherical Pose Reader UX fixes
- added option to flip axes for mirroring
- scaling factors encapsulated in struct
- falloff region independent of inner region
- added debug drawing of coordinate axes

diff --git a/Engine/Plugins/Animation/ControlRig/Source/ControlRig/Private/Units/Highlevel/Hierarchy/RigUnit_SphericalPoseReader.cpp b/Engine/Plugins/Animation/ControlRig/Source/ControlRig/Private/Units/Highlevel/Hierarchy/RigUnit_SphericalPoseReader.cpp
--- a/Engine/Plugins/Animation/ControlRig/Source/ControlRig/Private/Units/Highlevel/Hierarchy/RigUnit_SphericalPoseReader.cpp
+++ b/Engine/Plugins/Animation/ControlRig/Source/ControlRig/Private/Units/Highlevel/Hierarchy/RigUnit_SphericalPoseReader.cpp
@@ -26,17 +26,23 @@ FRigUnit_SphericalPoseReader_Execute()
 
 	// remap/clamp inputs
 	RemapAndConvertInputs(
-		InnerRegion, OuterRegion,
-		ActiveRegionSize, PositiveWidth, NegativeWidth, PositiveHeight, NegativeHeight,
-		FalloffSize, PositiveWidthFalloff, NegativeWidthFalloff, PositiveHeightFalloff, NegativeHeightFalloff);
+		InnerRegion,
+		OuterRegion,
+		ActiveRegionSize,
+		ActiveRegionScaleFactors,
+		FalloffSize,
+		FalloffRegionScaleFactors,
+		FlipWidthScaling,
+		FlipHeightScaling);
 	
-	// get the world offset
-	const FTransform LocalDriverTransformInit = Hierarchy->GetInitialLocalTransform(DriverItem);
+	// get local rotation space of driver
+	FTransform LocalDriverTransformInit = Hierarchy->GetInitialLocalTransform(DriverItem);
+	// apply static offset in local space
+	FQuat RotationOffsetQuat = FQuat::MakeFromEuler(RotationOffset);
+	LocalDriverTransformInit.SetRotation(LocalDriverTransformInit.GetRotation() * RotationOffsetQuat);
+	// put into global space
 	const FTransform GlobalDriverParentTransform = Hierarchy->GetParentTransform(DriverItem);
 	FTransform WorldOffset = LocalDriverTransformInit * GlobalDriverParentTransform;
-	// rotate by optional static offset
-	FQuat RotationOffsetQuat = FQuat::MakeFromEuler(RotationOffset);
-	WorldOffset.SetRotation(WorldOffset.GetRotation()*RotationOffsetQuat);
 
 	// get driver axis
 	const FTransform GlobalDriverTransform = Hierarchy->GetGlobalTransform(DriverItem);
@@ -90,45 +96,50 @@ void FRigUnit_SphericalPoseReader::RemapAndConvertInputs(
 	FSphericalRegion& InnerRegion,
 	FSphericalRegion& OuterRegion,
 	const float ActiveRegionSize,
-	const float PositiveWidth,
-	const float NegativeWidth,
-	const float PositiveHeight,
-	const float NegativeHeight,
+	const FRegionScaleFactors& ActiveRegionScaleFactors,
 	const float FalloffSize,
-	const float PositiveWidthFalloff,
-	const float NegativeWidthFalloff,
-	const float PositiveHeightFalloff,
-	const float NegativeHeightFalloff)
+	const FRegionScaleFactors& FalloffRegionScaleFactors,
+	const bool FlipWidth,
+	const bool FlipHeight)
 {
 	// remap normalized inputs to angles
 	float RegionAngle = ActiveRegionSize * 180.0f;
 	RegionAngle = FMath::Min(FMath::Max(0.5f, RegionAngle), 178.0f);
 
 	InnerRegion.RegionAngleRadians = FMath::DegreesToRadians(RegionAngle);
-	InnerRegion.PosWidth = PositiveWidth;
-	InnerRegion.NegWidth = NegativeWidth;
-	InnerRegion.PosHeight = PositiveHeight;
-	InnerRegion.NegHeight = NegativeHeight;
+	InnerRegion.ScaleFactors = ActiveRegionScaleFactors;
 
 	// clamp outer falloff angle to always be greater than inner
-	float FalloffAngle = RegionAngle + (FalloffSize * 180.0f);
-	FalloffAngle = FMath::Min(FMath::Max(1.0f, FalloffAngle), 179.0f);
+	float FalloffAngle = FalloffSize * 180.0f;
+	FalloffAngle = FMath::Min(FMath::Max(RegionAngle+1.0f, FalloffAngle), 179.0f);
 	OuterRegion.RegionAngleRadians = FMath::DegreesToRadians(FalloffAngle);
 
 	// clamp falloff to always be outside inner angle
 	const float InvOuterAngleRadians = 1.0f / OuterRegion.RegionAngleRadians;
 	//
-	const float PosWidthMin = (InnerRegion.RegionAngleRadians * PositiveWidth) * InvOuterAngleRadians;
-	OuterRegion.PosWidth = FMath::Lerp(PosWidthMin, 1.0f, PositiveWidthFalloff);
+	const float PosWidthMin = (InnerRegion.RegionAngleRadians * ActiveRegionScaleFactors.PositiveWidth) * InvOuterAngleRadians;
+	OuterRegion.ScaleFactors.PositiveWidth = FMath::Lerp(PosWidthMin, 1.0f, FalloffRegionScaleFactors.PositiveWidth);
 	//
-	const float NegWidthMin = (InnerRegion.RegionAngleRadians * NegativeWidth) * InvOuterAngleRadians;
-	OuterRegion.NegWidth = FMath::Lerp(NegWidthMin, 1.0f, NegativeWidthFalloff);
+	const float NegWidthMin = (InnerRegion.RegionAngleRadians * ActiveRegionScaleFactors.NegativeWidth) * InvOuterAngleRadians;
+	OuterRegion.ScaleFactors.NegativeWidth = FMath::Lerp(NegWidthMin, 1.0f, FalloffRegionScaleFactors.NegativeWidth);
 	//
-	const float PosHeightMin = (InnerRegion.RegionAngleRadians * PositiveHeight) * InvOuterAngleRadians;
-	OuterRegion.PosHeight = FMath::Lerp(PosHeightMin, 1.0f, PositiveHeightFalloff);
+	const float PosHeightMin = (InnerRegion.RegionAngleRadians * ActiveRegionScaleFactors.PositiveHeight) * InvOuterAngleRadians;
+	OuterRegion.ScaleFactors.PositiveHeight = FMath::Lerp(PosHeightMin, 1.0f, FalloffRegionScaleFactors.PositiveHeight);
 	//
-	const float NegHeightMin = (InnerRegion.RegionAngleRadians * NegativeHeight) * InvOuterAngleRadians;
-	OuterRegion.NegHeight = FMath::Lerp(NegHeightMin, 1.0f, NegativeHeightFalloff);
+	const float NegHeightMin = (InnerRegion.RegionAngleRadians * ActiveRegionScaleFactors.NegativeHeight) * InvOuterAngleRadians;
+	OuterRegion.ScaleFactors.NegativeHeight = FMath::Lerp(NegHeightMin, 1.0f, FalloffRegionScaleFactors.NegativeHeight);
+
+	// optionally flip the scaling factors to better support mirrored setups
+	if (FlipWidth)
+	{
+		Swap(InnerRegion.ScaleFactors.PositiveWidth, InnerRegion.ScaleFactors.NegativeWidth);
+		Swap(OuterRegion.ScaleFactors.PositiveWidth, OuterRegion.ScaleFactors.NegativeWidth);
+	}
+	if (FlipHeight)
+	{
+		Swap(InnerRegion.ScaleFactors.PositiveHeight, InnerRegion.ScaleFactors.NegativeHeight);
+		Swap(OuterRegion.ScaleFactors.PositiveHeight, OuterRegion.ScaleFactors.NegativeHeight);
+	}
 }
 
 float FRigUnit_SphericalPoseReader::CalcOutputParam(
diff --git a/Engine/Plugins/Animation/ControlRig/Source/ControlRig/Private/Units/Highlevel/Hierarchy/RigUnit_SphericalPoseReader.h b/Engine/Plugins/Animation/ControlRig/Source/ControlRig/Private/Units/Highlevel/Hierarchy/RigUnit_SphericalPoseReader.h
--- a/Engine/Plugins/Animation/ControlRig/Source/ControlRig/Private/Units/Highlevel/Hierarchy/RigUnit_SphericalPoseReader.h
+++ b/Engine/Plugins/Animation/ControlRig/Source/ControlRig/Private/Units/Highlevel/Hierarchy/RigUnit_SphericalPoseReader.h
@@ -16,29 +16,46 @@ struct FEllipseQuery
 };
 
 USTRUCT()
-struct CONTROLRIG_API FSphericalRegion
+struct CONTROLRIG_API FRegionScaleFactors
 {
 	GENERATED_BODY()
 
-	FSphericalRegion()
-	: RegionAngleRadians(0.1f),
-	PosWidth(1.0f),
-	NegWidth(1.0f),
-	PosHeight(1.0f),
-	NegHeight(1.0f)
+	FRegionScaleFactors()
+	: PositiveWidth(1.0f),
+	NegativeWidth(1.0f),
+	PositiveHeight(1.0f),
+	NegativeHeight(1.0f)
 	{
 	}
+
+	// Scale the region in the POSITIVE width direction. Range is 0-1. Default is 1.0.
+	UPROPERTY(meta = (Input, ClampMin = "0.0", ClampMax = "1.0", UIMin = "0.0", UIMax = "1.0"))
+	float PositiveWidth = 1.0f;
+	// Scale the region in the NEGATIVE width direction. Range is 0-1. Default is 1.0.
+	UPROPERTY(meta = (Input, ClampMin = "0.0", ClampMax = "1.0", UIMin = "0.0", UIMax = "1.0"))
+	float NegativeWidth = 1.0f;
+	// Scale the region in the POSITIVE height direction. Range is 0-1. Default is 1.0.
+	UPROPERTY(meta = (Input, ClampMin = "0.0", ClampMax = "1.0", UIMin = "0.0", UIMax = "1.0"))
+	float PositiveHeight = 1.0f;
+	// Scale the region in the NEGATIVE height direction. Range is 0-1. Default is 1.0.
+	UPROPERTY(meta = (Input, ClampMin = "0.0", ClampMax = "1.0", UIMin = "0.0", UIMax = "1.0"))
+	float NegativeHeight = 1.0f;
+};
+
+USTRUCT()
+struct CONTROLRIG_API FSphericalRegion
+{
+	GENERATED_BODY()
+
+	FSphericalRegion() : RegionAngleRadians(0.1f), ScaleFactors(FRegionScaleFactors()){}
 	
 	float RegionAngleRadians;
-	float PosWidth;
-	float NegWidth;
-	float PosHeight;
-	float NegHeight;
+	FRegionScaleFactors ScaleFactors;
 
 	void GetEllipseWidthAndHeight(const float PointX, const float PointY, float& OutWidth, float& OutHeight)
 	{
-		OutWidth = ((PointX > 0 ? PosWidth : NegWidth) * RegionAngleRadians) * INV_PI;
-		OutHeight = ((PointY > 0 ? PosHeight : NegHeight) * RegionAngleRadians) * INV_PI;
+		OutWidth = ((PointX > 0 ? ScaleFactors.PositiveWidth : ScaleFactors.NegativeWidth) * RegionAngleRadians) * INV_PI;
+		OutHeight = ((PointY > 0 ? ScaleFactors.PositiveHeight : ScaleFactors.NegativeHeight) * RegionAngleRadians) * INV_PI;
 	}
 };
 
@@ -62,6 +79,9 @@ struct CONTROLRIG_API FSphericalPoseReaderDebugSettings
 	UPROPERTY(meta = (Input))
 	bool bDraw2D = false;
 	
+	UPROPERTY(meta = (Input))
+	bool bDrawLocalAxes = false;
+	
 	UPROPERTY(meta = (Input, ClampMin = "0.0", UIMin = "0.01", UIMax = "100.0"))
 	float DebugScale = 25.0f;
 	
@@ -85,6 +105,13 @@ struct CONTROLRIG_API FSphericalPoseReaderDebugSettings
 			return;
 		}
 
+		if (bDrawLocalAxes)
+		{
+			FTransform AxisOffset = WorldOffset;
+			AxisOffset.AddToTranslation(AxisOffset.GetUnitAxis(EAxis::Z) * DebugScale);
+			DrawInterface->DrawAxes(AxisOffset, FTransform::Identity, DebugScale * 0.5f, DebugThickness);	
+		}
+
 		// make colors
 		FLinearColor DriverColor = FLinearColor::Red;
 		FLinearColor Yellow = FLinearColor::Yellow;
@@ -208,8 +235,8 @@ struct CONTROLRIG_API FSphericalPoseReaderDebugSettings
 		FVector Point = SpanAxis * (Region.RegionAngleRadians / PI);
 
 		// stretch cone point non-uniformly
-		const float Width = Point.X > 0 ? Region.PosWidth : Region.NegWidth;
-		const float Height = Point.Y > 0 ? Region.PosHeight : Region.NegHeight;
+		const float Width = Point.X > 0 ? Region.ScaleFactors.PositiveWidth : Region.ScaleFactors.NegativeWidth;
+		const float Height = Point.Y > 0 ? Region.ScaleFactors.PositiveHeight : Region.ScaleFactors.NegativeHeight;
 		Point.X *= Width;
 		Point.Y *= Height;
 		Point.Z = 0.0f;
@@ -249,15 +276,7 @@ struct CONTROLRIG_API FRigUnit_SphericalPoseReader: public FRigUnit_HighlevelBas
 			DriverAxis(FVector(1.0f,0.0f,0.0f)),
 			RotationOffset(FVector(90.0f,0.0f,0.0f)),
 			ActiveRegionSize(0.1f),
-			PositiveWidth(1.0f),
-			NegativeWidth(1.0f),
-			PositiveHeight(1.0f),
-			NegativeHeight(1.0f),
 			FalloffSize(0.2f),
-			PositiveWidthFalloff(1.0f),
-			NegativeWidthFalloff(1.0f),
-			PositiveHeightFalloff(1.0f),
-			NegativeHeightFalloff(1.0f),
 			Debug(FSphericalPoseReaderDebugSettings()),
 			InnerRegion(FSphericalRegion()),
 			OuterRegion(FSphericalRegion()),
@@ -273,15 +292,11 @@ struct CONTROLRIG_API FRigUnit_SphericalPoseReader: public FRigUnit_HighlevelBas
 		FSphericalRegion& InnerRegion,
 		FSphericalRegion& OuterRegion,
 		const float ActiveRegionSize,
-		const float PositiveWidth,
-		const float NegativeWidth,
-		const float PositiveHeight,
-		const float NegativeHeight,
+		const FRegionScaleFactors& ActiveRegionScaleFactors,
 		const float FalloffSize,
-		const float PositiveWidthFalloff,
-		const float NegativeWidthFalloff,
-		const float PositiveHeightFalloff,
-		const float NegativeHeightFalloff);
+		const FRegionScaleFactors& FalloffRegionScaleFactors,
+		const bool FlipWidth,
+		const bool FlipHeight);
 	static void EvaluateInterpolation(FRigUnit_SphericalPoseReader& Node, const FTransform& DriverTransform);
 	
 	static void DrawDebug(FRigUnit_SphericalPoseReader& PoseReader);
@@ -323,34 +338,24 @@ struct CONTROLRIG_API FRigUnit_SphericalPoseReader: public FRigUnit_HighlevelBas
 	// The size of the active region (green) that outputs the full value (1.0). Range is 0-1. Default is 0.1.
 	UPROPERTY(meta = (Input, ClampMin = "0.0", ClampMax = "1.0", UIMin = "0.0", UIMax = "1.0"))
     float ActiveRegionSize = 0.1f;
-	// Scale the active region in the POSITIVE width direction. Range is 0-1. Default is 1.0.
-	UPROPERTY(meta = (Input, ClampMin = "0.0", ClampMax = "1.0", UIMin = "0.0", UIMax = "1.0"))
-    float PositiveWidth = 1.0f;
-	// Scale the active region in the NEGATIVE width direction. Range is 0-1. Default is 1.0.
-	UPROPERTY(meta = (Input, ClampMin = "0.0", ClampMax = "1.0", UIMin = "0.0", UIMax = "1.0"))
-    float NegativeWidth = 1.0f;
-	// Scale the active region in the POSITIVE height direction. Range is 0-1. Default is 1.0.
-    UPROPERTY(meta = (Input, ClampMin = "0.0", ClampMax = "1.0", UIMin = "0.0", UIMax = "1.0"))
-    float PositiveHeight = 1.0f;
-	// Scale the active region in the NEGATIVE height direction. Range is 0-1. Default is 1.0.
-    UPROPERTY(meta = (Input, ClampMin = "0.0", ClampMax = "1.0", UIMin = "0.0", UIMax = "1.0"))
-    float NegativeHeight = 1.0f;
+	// The directional scaling parameters for the active region (green).
+	UPROPERTY(meta = (Input))
+	FRegionScaleFactors ActiveRegionScaleFactors;
 
 	// The size of the falloff region (yellow) that defines the start of the output range. A value of 1 wraps the entire sphere with falloff. Range is 0-1. Default is 0.2.
 	UPROPERTY(meta = (Input, ClampMin = "0.0", ClampMax = "1.0", UIMin = "0.0", UIMax = "1.0"))
     float FalloffSize = 0.2f;
-	// Scale the falloff region in the POSITIVE width direction. Range is 0-1. Default is 1.0.
-	UPROPERTY(meta = (Input, ClampMin = "0.0", ClampMax = "1.0", UIMin = "0.0", UIMax = "1.0"))
-    float PositiveWidthFalloff = 1.0f;
-	// Scale the falloff region in the NEGATIVE width direction. Range is 0-1. Default is 1.0.
-	UPROPERTY(meta = (Input, ClampMin = "0.0", ClampMax = "1.0", UIMin = "0.0", UIMax = "1.0"))
-    float NegativeWidthFalloff = 1.0f;
-	// Scale the falloff region in the POSITIVE height direction. Range is 0-1. Default is 1.0.
-	UPROPERTY(meta = (Input, ClampMin = "0.0", ClampMax = "1.0", UIMin = "0.0", UIMax = "1.0"))
-    float PositiveHeightFalloff = 1.0f;
-	// Scale the falloff region in the NEGATIVE height direction. Range is 0-1. Default is 1.0.
-	UPROPERTY(meta = (Input, ClampMin = "0.0", ClampMax = "1.0", UIMin = "0.0", UIMax = "1.0"))
-    float NegativeHeightFalloff = 1.0f;
+	// The directional scaling parameters for the falloff region (yellow).
+	UPROPERTY(meta = (Input))
+	FRegionScaleFactors FalloffRegionScaleFactors;
+
+	// Flip the positive / negative directions of the width scale factors.
+	UPROPERTY(meta = (Input))
+	bool FlipWidthScaling = false;
+
+	// Flip the positive / negative directions of the height scale factors.
+	UPROPERTY(meta = (Input))
+	bool FlipHeightScaling = false;
 
 	UPROPERTY(meta = (Input))
 	FSphericalPoseReaderDebugSettings Debug;
