Control Rig: Fix compilation issues SRigSpacePickerWidget.cpp not finding FModuleManager

diff --git a/Engine/Plugins/Animation/ControlRig/Source/ControlRigEditor/Private/Editor/ControlRigContextMenuContext.h b/Engine/Plugins/Animation/ControlRig/Source/ControlRigEditor/Private/Editor/ControlRigContextMenuContext.h
--- a/Engine/Plugins/Animation/ControlRig/Source/ControlRigEditor/Private/Editor/ControlRigContextMenuContext.h
+++ b/Engine/Plugins/Animation/ControlRig/Source/ControlRigEditor/Private/Editor/ControlRigContextMenuContext.h
@@ -14,7 +14,6 @@ class FControlRigEditor;
 class URigVMGraph;
 class URigVMNode;
 class URigVMPin;
-class UEdGraph;
 class SRigHierarchy;
 
 USTRUCT(BlueprintType)
diff --git a/Engine/Plugins/Animation/ControlRig/Source/ControlRigEditor/Private/Editor/SRigSpacePickerWidget.cpp b/Engine/Plugins/Animation/ControlRig/Source/ControlRigEditor/Private/Editor/SRigSpacePickerWidget.cpp
--- a/Engine/Plugins/Animation/ControlRig/Source/ControlRigEditor/Private/Editor/SRigSpacePickerWidget.cpp
+++ b/Engine/Plugins/Animation/ControlRig/Source/ControlRigEditor/Private/Editor/SRigSpacePickerWidget.cpp
@@ -14,6 +14,7 @@
 #include "FrameNumberDetailsCustomization.h"
 #include "Widgets/Notifications/SNotificationList.h"
 #include "Framework/Notifications/NotificationManager.h"
+#include "Modules/ModuleManager.h"
 
 #define LOCTEXT_NAMESPACE "SRigSpacePickerWidget"
 
diff --git a/Engine/Plugins/Animation/ControlRig/Source/ControlRigEditor/Public/ControlRigSpaceChannelCurveModel.h b/Engine/Plugins/Animation/ControlRig/Source/ControlRigEditor/Public/ControlRigSpaceChannelCurveModel.h
--- a/Engine/Plugins/Animation/ControlRig/Source/ControlRigEditor/Public/ControlRigSpaceChannelCurveModel.h
+++ b/Engine/Plugins/Animation/ControlRig/Source/ControlRigEditor/Public/ControlRigSpaceChannelCurveModel.h
@@ -8,6 +8,7 @@
 #include "KeyBarCurveModel.h"
 #include "Sequencer/MovieSceneControlRigSpaceChannel.h"
 #include "Channels/MovieSceneChannelHandle.h"
+#include "Framework/MultiBox/MultiBoxBuilder.h"
 
 class UMovieSceneSection;
 class ISequencer;
