Control Rig: Fix Editor hangs and becomes unresponsive when undoing 'Save Camera As Default' in Control Rig viewport

diff --git a/Engine/Plugins/Animation/ControlRig/Source/ControlRigEditor/Private/Editor/ControlRigSkeletalMeshComponent.cpp b/Engine/Plugins/Animation/ControlRig/Source/ControlRigEditor/Private/Editor/ControlRigSkeletalMeshComponent.cpp
--- a/Engine/Plugins/Animation/ControlRig/Source/ControlRigEditor/Private/Editor/ControlRigSkeletalMeshComponent.cpp
+++ b/Engine/Plugins/Animation/ControlRig/Source/ControlRigEditor/Private/Editor/ControlRigSkeletalMeshComponent.cpp
@@ -15,14 +15,29 @@ UControlRigSkeletalMeshComponent::UControlRigSkeletalMeshComponent(const FObject
 
 void UControlRigSkeletalMeshComponent::InitAnim(bool bForceReinit)
 {
-	// skip preview init entirely, just init the super class
-	Super::InitAnim(bForceReinit);
+	UAnimInstance* AnimInstance = GetAnimInstance();
+	TObjectPtr<class UAnimPreviewInstance> LastPreviewInstance = PreviewInstance;
 
-	UControlRigLayerInstance* ControlRigInstance = Cast<UControlRigLayerInstance>(GetAnimInstance());
+	// Super::InitAnim might trigger an evaluation, in which case the source animation instance must be set
+	UControlRigLayerInstance* ControlRigInstance = Cast<UControlRigLayerInstance>(AnimInstance);
 	if (ControlRigInstance)
 	{
 		ControlRigInstance->SetSourceAnimInstance(PreviewInstance);
 	}
+	
+	// skip preview init entirely, just init the super class
+	Super::InitAnim(bForceReinit);
+
+	// The preview instance or anim instance might have been created in Super::InitAnim, in which case
+	// the source animation instance must be set now
+	if (AnimInstance != GetAnimInstance() || LastPreviewInstance != PreviewInstance)
+	{
+		ControlRigInstance = Cast<UControlRigLayerInstance>(GetAnimInstance());
+		if (ControlRigInstance)
+		{
+			ControlRigInstance->SetSourceAnimInstance(PreviewInstance);
+		}
+	}
 
 	RebuildDebugDrawSkeleton();
 }
