Control Rig: Control Rig being debugged holds a pointer to editor's Control Rig Log, it needs to be cleared when control rig editor closes.

diff --git a/Engine/Plugins/Experimental/ControlRig/Source/ControlRigEditor/Private/Editor/ControlRigEditor.cpp b/Engine/Plugins/Experimental/ControlRig/Source/ControlRigEditor/Private/Editor/ControlRigEditor.cpp
--- a/Engine/Plugins/Experimental/ControlRig/Source/ControlRigEditor/Private/Editor/ControlRigEditor.cpp
+++ b/Engine/Plugins/Experimental/ControlRig/Source/ControlRigEditor/Private/Editor/ControlRigEditor.cpp
@@ -119,6 +119,9 @@ FControlRigEditor::~FControlRigEditor()
 	UControlRigBlueprint* RigBlueprint = GetControlRigBlueprint();
 	if (RigBlueprint)
 	{
+		// clear editor related data from the debugged control rig instance 
+		RigBlueprint->SetObjectBeingDebugged(nullptr);
+
 		UControlRigBlueprint::sCurrentlyOpenedRigBlueprints.Remove(RigBlueprint);
 
 		RigBlueprint->Hierarchy->OnModified().RemoveAll(this);
