Control Rig: Fix for control rig component, it was using the target's element type when dealing with source element in validateMappedElements()

diff --git a/Engine/Plugins/Experimental/ControlRig/Source/ControlRig/Private/ControlRigComponent.cpp b/Engine/Plugins/Experimental/ControlRig/Source/ControlRig/Private/ControlRigComponent.cpp
--- a/Engine/Plugins/Experimental/ControlRig/Source/ControlRig/Private/ControlRigComponent.cpp
+++ b/Engine/Plugins/Experimental/ControlRig/Source/ControlRig/Private/ControlRigComponent.cpp
@@ -1111,11 +1111,7 @@ void UControlRigComponent::ValidateMappingData()
 					{
 						if (USkeleton* Skeleton = SkeletalMesh->GetSkeleton())
 						{
-							if (MappedElement.ElementType == ERigElementType::Bone)
-							{
-								MappedElement.SubIndex = Skeleton->GetReferenceSkeleton().FindBoneIndex(MappedElement.TransformName);
-							}
-							else if (MappedElement.ElementType == ERigElementType::Curve)
+							if (MappedElement.ElementType == ERigElementType::Curve)
 							{
 								const FSmartNameMapping* CurveNameMapping = Skeleton->GetSmartNameContainer(USkeleton::AnimCurveMappingName);
 								if (CurveNameMapping)
@@ -1126,6 +1122,10 @@ void UControlRigComponent::ValidateMappingData()
 										MappedElement.SubIndex = (int32)SmartName.UID;
 									}
 								}
+							} 
+							else
+							{
+								MappedElement.SubIndex = Skeleton->GetReferenceSkeleton().FindBoneIndex(MappedElement.TransformName);
 							}
 						}
 						else
