Sequencer: Control Rig: Don't use bool Controls on Additive Sections.

diff --git a/Engine/Plugins/Experimental/ControlRig/Source/ControlRig/Private/Sequencer/MovieSceneControlRigParameterTemplate.cpp b/Engine/Plugins/Experimental/ControlRig/Source/ControlRig/Private/Sequencer/MovieSceneControlRigParameterTemplate.cpp
--- a/Engine/Plugins/Experimental/ControlRig/Source/ControlRig/Private/Sequencer/MovieSceneControlRigParameterTemplate.cpp
+++ b/Engine/Plugins/Experimental/ControlRig/Source/ControlRig/Private/Sequencer/MovieSceneControlRigParameterTemplate.cpp
@@ -954,40 +954,43 @@ struct FControlRigParameterExecutionToken : IMovieSceneExecutionToken
 		}
 
 		//Do Bool straight up no blending
-		bool bWasDoNotKey = false;
+		if (Section->GetBlendType().Get() != EMovieSceneBlendType::Additive)
+		{
+			bool bWasDoNotKey = false;
 
-		bWasDoNotKey = Section->GetDoNotKey();
-		Section->SetDoNotKey(true);
+			bWasDoNotKey = Section->GetDoNotKey();
+			Section->SetDoNotKey(true);
 
-		if (Section->GetControlRig())
-		{
-			const bool bSetupUndo = false;
-			Section->GetControlRig()->SetAbsoluteTime((float)Context.GetFrameRate().AsSeconds(Context.GetTime()));
-			for (const FBoolParameterStringAndValue& BoolNameAndValue : BoolValues)
+			if (Section->GetControlRig())
 			{
-				if (Section->ControlsToSet.Num() == 0 || Section->ControlsToSet.Contains(BoolNameAndValue.ParameterName))
+				const bool bSetupUndo = false;
+				Section->GetControlRig()->SetAbsoluteTime((float)Context.GetFrameRate().AsSeconds(Context.GetTime()));
+				for (const FBoolParameterStringAndValue& BoolNameAndValue : BoolValues)
 				{
-					FRigControlElement* RigControl = Section->GetControlRig()->FindControl(BoolNameAndValue.ParameterName);
-					if (RigControl && RigControl->Settings.ControlType == ERigControlType::Bool)
+					if (Section->ControlsToSet.Num() == 0 || Section->ControlsToSet.Contains(BoolNameAndValue.ParameterName))
 					{
-						Section->GetControlRig()->SetControlValue<bool>(BoolNameAndValue.ParameterName, BoolNameAndValue.Value, true, EControlRigSetKey::Never,bSetupUndo);
+						FRigControlElement* RigControl = Section->GetControlRig()->FindControl(BoolNameAndValue.ParameterName);
+						if (RigControl && RigControl->Settings.ControlType == ERigControlType::Bool)
+						{
+							Section->GetControlRig()->SetControlValue<bool>(BoolNameAndValue.ParameterName, BoolNameAndValue.Value, true, EControlRigSetKey::Never,bSetupUndo);
+						}
 					}
 				}
-			}
 
-			for (const FIntegerParameterStringAndValue& IntegerNameAndValue : IntegerValues)
-			{
-				if (Section->ControlsToSet.Num() == 0 || Section->ControlsToSet.Contains(IntegerNameAndValue.ParameterName))
+				for (const FIntegerParameterStringAndValue& IntegerNameAndValue : IntegerValues)
 				{
-					FRigControlElement* RigControl = Section->GetControlRig()->FindControl(IntegerNameAndValue.ParameterName);
-					if (RigControl && RigControl->Settings.ControlType == ERigControlType::Integer)
+					if (Section->ControlsToSet.Num() == 0 || Section->ControlsToSet.Contains(IntegerNameAndValue.ParameterName))
 					{
-						Section->GetControlRig()->SetControlValue<int32>(IntegerNameAndValue.ParameterName, IntegerNameAndValue.Value, true, EControlRigSetKey::Never,bSetupUndo);
+						FRigControlElement* RigControl = Section->GetControlRig()->FindControl(IntegerNameAndValue.ParameterName);
+						if (RigControl && RigControl->Settings.ControlType == ERigControlType::Integer)
+						{
+							Section->GetControlRig()->SetControlValue<int32>(IntegerNameAndValue.ParameterName, IntegerNameAndValue.Value, true, EControlRigSetKey::Never,bSetupUndo);
+						}
 					}
 				}
 			}
+			Section->SetDoNotKey(bWasDoNotKey);
 		}
-		Section->SetDoNotKey(bWasDoNotKey);
 
 	}
 
