Control Rig: Fix bad merge from 16328409

diff --git a/Engine/Plugins/Experimental/ControlRig/Source/ControlRig/Private/ControlRigComponent.cpp b/Engine/Plugins/Experimental/ControlRig/Source/ControlRig/Private/ControlRigComponent.cpp
--- a/Engine/Plugins/Experimental/ControlRig/Source/ControlRig/Private/ControlRigComponent.cpp
+++ b/Engine/Plugins/Experimental/ControlRig/Source/ControlRig/Private/ControlRigComponent.cpp
@@ -1125,7 +1125,6 @@ void UControlRigComponent::ValidateMappingData()
 										MappedElement.SubIndex = (int32)SmartName.UID;
 									}
 								}
-							} 
 							}
 							else
 							{
