Sequencer: Fix so that keying doesn't affect masked channels.

Rename "Filter Channels" to "Active Channels" to be consistent with the transform mask menu for static meshes.

diff --git a/Engine/Plugins/Experimental/ControlRig/Source/ControlRigEditor/Private/Sequencer/ControlRigParameterTrackEditor.cpp b/Engine/Plugins/Experimental/ControlRig/Source/ControlRigEditor/Private/Sequencer/ControlRigParameterTrackEditor.cpp
--- a/Engine/Plugins/Experimental/ControlRig/Source/ControlRigEditor/Private/Sequencer/ControlRigParameterTrackEditor.cpp
+++ b/Engine/Plugins/Experimental/ControlRig/Source/ControlRigEditor/Private/Sequencer/ControlRigParameterTrackEditor.cpp
@@ -2086,8 +2086,11 @@ void FControlRigParameterTrackEditor::HandleControlModified(UControlRig* Control
 	}
 }
 
-void FControlRigParameterTrackEditor::GetControlRigKeys(UControlRig* InControlRig, FName ParameterName, EMovieSceneTransformChannel ChannelsToKey, FGeneratedTrackKeys& OutGeneratedKeys)
+void FControlRigParameterTrackEditor::GetControlRigKeys(UControlRig* InControlRig, FName ParameterName, EMovieSceneTransformChannel ChannelsToKey, UMovieSceneControlRigParameterSection* SectionToKey, FGeneratedTrackKeys& OutGeneratedKeys)
 {
+	const TArray<bool>& ControlsMask = SectionToKey->GetControlsMask();
+	EMovieSceneTransformChannel TransformMask = SectionToKey->GetTransformMask().GetChannels();
+
 	TArray<FRigControlElement*> Controls;
 	InControlRig->GetControlsInOrder(Controls);
 	// If key all is enabled, for a key on all the channels
@@ -2101,12 +2104,20 @@ void FControlRigParameterTrackEditor::GetControlRigKeys(UControlRig* InControlRi
 	int32 BoolChannelIndex = 0;
 	int32 EnumChannelIndex = 0;
 	int32 IntChannelIndex = 0;
-	for (FRigControlElement* ControlElement : Controls)
+	for (int32 ControlIndex = 0; ControlIndex < Controls.Num(); ++ControlIndex)
 	{
+		FRigControlElement* ControlElement = Controls[ControolIndex];
+
 		if (!ControlElement->Settings.bAnimatable)
 		{
 			continue;
 		}
+
+		if (ControlIndex < ControlsMask.Num() && !ControlsMask[ControlIndex])
+		{
+			continue;
+		}
+
 		bool bSetKey = ControlElement->GetName() == ParameterName;
 
 		FRigControlValue ControlValue = InControlRig->GetHierarchy()->GetControlValue(ControlElement, ERigControlValueType::Current);
@@ -2192,6 +2203,18 @@ void FControlRigParameterTrackEditor::GetControlRigKeys(UControlRig* InControlRi
 			{
 				bKeyX = bKeyY = bKeyZ = true;
 			}
+			if (!EnumHasAnyFlags(TransformMask, EMovieSceneTransformChannel::TranslationX))
+			{
+				bKeyX = false;
+			}
+			if (!EnumHasAnyFlags(TransformMask, EMovieSceneTransformChannel::TranslationY))
+			{
+				bKeyY = false;
+			}
+			if (!EnumHasAnyFlags(TransformMask, EMovieSceneTransformChannel::TranslationZ))
+			{
+				bKeyZ = false;
+			}
 			OutGeneratedKeys.Add(FMovieSceneChannelValueSetter::Create<FMovieSceneFloatChannel>(ChannelIndex++, CurrentVector.X, bKeyX));
 			OutGeneratedKeys.Add(FMovieSceneChannelValueSetter::Create<FMovieSceneFloatChannel>(ChannelIndex++, CurrentVector.Y, bKeyY));
 			OutGeneratedKeys.Add(FMovieSceneChannelValueSetter::Create<FMovieSceneFloatChannel>(ChannelIndex++, CurrentVector.Z, bKeyZ));
@@ -2204,6 +2227,18 @@ void FControlRigParameterTrackEditor::GetControlRigKeys(UControlRig* InControlRi
 			{
 				bKeyX = bKeyY = bKeyZ = true;
 			}
+			if (!EnumHasAnyFlags(TransformMask, EMovieSceneTransformChannel::RotationX))
+			{
+				bKeyX = false;
+			}
+			if (!EnumHasAnyFlags(TransformMask, EMovieSceneTransformChannel::RotationY))
+			{
+				bKeyY = false;
+			}
+			if (!EnumHasAnyFlags(TransformMask, EMovieSceneTransformChannel::RotationZ))
+			{
+				bKeyZ = false;
+			}
 
 			/* @Mike.Zyracki this is my gut feeling - we should run SetClosestToMe on the rotator SOMEWHERE....
 			FMovieSceneInterrogationData InterrogationData;
@@ -2232,6 +2267,18 @@ void FControlRigParameterTrackEditor::GetControlRigKeys(UControlRig* InControlRi
 				{
 					bKeyX = bKeyY = bKeyZ = true;
 				}
+				if (!EnumHasAnyFlags(TransformMask, EMovieSceneTransformChannel::ScaleX))
+				{
+					bKeyX = false;
+				}
+				if (!EnumHasAnyFlags(TransformMask, EMovieSceneTransformChannel::ScaleY))
+				{
+					bKeyY = false;
+				}
+				if (!EnumHasAnyFlags(TransformMask, EMovieSceneTransformChannel::ScaleZ))
+				{
+					bKeyZ = false;
+				}
 				OutGeneratedKeys.Add(FMovieSceneChannelValueSetter::Create<FMovieSceneFloatChannel>(ChannelIndex++, CurrentVector.X, bKeyX));
 				OutGeneratedKeys.Add(FMovieSceneChannelValueSetter::Create<FMovieSceneFloatChannel>(ChannelIndex++, CurrentVector.Y, bKeyY));
 				OutGeneratedKeys.Add(FMovieSceneChannelValueSetter::Create<FMovieSceneFloatChannel>(ChannelIndex++, CurrentVector.Z, bKeyZ));
@@ -2345,11 +2392,12 @@ void FControlRigParameterTrackEditor::AddControlKeys(USceneComponent *InSceneCom
 	}
 	FFindOrCreateTrackResult TrackResult = FindOrCreateControlRigTrackForObject(ObjectHandle, InControlRig, ControlRigName, bCreateTrack);
 	UMovieSceneControlRigParameterTrack* Track = CastChecked<UMovieSceneControlRigParameterTrack>(TrackResult.Track, ECastCheckedType::NullAllowed);
+	UMovieSceneControlRigParameterSection* ParamSection = nullptr;
 	if (Track)
 	{
 		FFrameNumber  FrameTime = GetTimeForKey();
 		UMovieSceneSection* Section = Track->FindSection(FrameTime);
-		UMovieSceneControlRigParameterSection* ParamSection = Cast<UMovieSceneControlRigParameterSection>(Section);
+		ParamSection = Cast<UMovieSceneControlRigParameterSection>(Section);
 
 		if (ParamSection && ParamSection->GetDoNotKey())
 		{
@@ -2357,9 +2405,14 @@ void FControlRigParameterTrackEditor::AddControlKeys(USceneComponent *InSceneCom
 		}
 	}
 
+	if (!ParamSection)
+	{
+		return;
+	}
+
 	TSharedRef<FGeneratedTrackKeys> GeneratedKeys = MakeShared<FGeneratedTrackKeys>();
 
-	GetControlRigKeys(InControlRig, RigControlName, ChannelsToKey, *GeneratedKeys);
+	GetControlRigKeys(InControlRig, RigControlName, ChannelsToKey, ParamSection, *GeneratedKeys);
 	TGuardValue<bool> Guard(bIsDoingSelection, true);
 
 	auto OnKeyProperty = [=](FFrameNumber Time) -> FKeyPropertyResult
@@ -3100,10 +3153,10 @@ void FControlRigParameterSection::BuildSectionContextMenu(FMenuBuilder& MenuBuil
 			})
 				);
 		};
-		MenuBuilder.BeginSection(NAME_None, LOCTEXT("RigSectionFilterControls", "Filter Controls"));
+		MenuBuilder.BeginSection(NAME_None, LOCTEXT("RigSectionActiveChannels", "Active Channels"));
 		{
 			MenuBuilder.AddSubMenu(
-				LOCTEXT("ToggleRigControlsText", "Toggle Rig Controls"), LOCTEXT("ToggleRigControlsText_Tooltip", "Toggle Rig Controls"),
+				LOCTEXT("ToggleRigControlsText", "Rig Controls"), LOCTEXT("ToggleRigControlsText_Tooltip", "Causes this section to affect all rig controls"),
 				FNewMenuDelegate::CreateLambda([=](FMenuBuilder& SubMenuBuilder) {
 				int32 Index = 0;
 				for (FRigControlElement* ControlElement : Controls)
@@ -3111,7 +3164,7 @@ void FControlRigParameterSection::BuildSectionContextMenu(FMenuBuilder& MenuBuil
 					const FName RigName = ControlElement->GetName();
 					FText Name = FText::FromName(RigName);
 					FText Text = FText::Format(LOCTEXT("RigControlToggle", "{0}"), Name);
-					FText TooltipText = FText::Format(LOCTEXT("RigControlToggleTooltip", "Toggle Rig Control {0}"), Name);
+					FText TooltipText = FText::Format(LOCTEXT("RigControlToggleTooltip", "Causes this section to affect rig control {0}"), Name);
 					SubMenuBuilder.AddMenuEntry(
 						Text, TooltipText,
 						FSlateIcon(), ToggleControls(Index++), NAME_None, EUserInterfaceActionType::ToggleButton);
diff --git a/Engine/Plugins/Experimental/ControlRig/Source/ControlRigEditor/Private/Sequencer/ControlRigParameterTrackEditor.h b/Engine/Plugins/Experimental/ControlRig/Source/ControlRigEditor/Private/Sequencer/ControlRigParameterTrackEditor.h
--- a/Engine/Plugins/Experimental/ControlRig/Source/ControlRigEditor/Private/Sequencer/ControlRigParameterTrackEditor.h
+++ b/Engine/Plugins/Experimental/ControlRig/Source/ControlRigEditor/Private/Sequencer/ControlRigParameterTrackEditor.h
@@ -153,7 +153,7 @@ private:
 public:
 
 	void AddControlKeys(USceneComponent *InSceneComp, UControlRig* InControlRig, FName PropertyName, FName ParameterName, EMovieSceneTransformChannel ChannelsToKey, ESequencerKeyMode KeyMode, float InLocalTime);
-	void GetControlRigKeys(UControlRig* InControlRig, FName ParameterName, EMovieSceneTransformChannel ChannelsToKey, FGeneratedTrackKeys& OutGeneratedKeys);
+	void GetControlRigKeys(UControlRig* InControlRig, FName ParameterName, EMovieSceneTransformChannel ChannelsToKey, UMovieSceneControlRigParameterSection* SectionToKey, FGeneratedTrackKeys& OutGeneratedKeys);
 	FKeyPropertyResult AddKeysToControlRig(
 		USceneComponent *InSceneComp, UControlRig* InControlRig, FFrameNumber KeyTime, FGeneratedTrackKeys& GeneratedKeys,
 		ESequencerKeyMode KeyMode, TSubclassOf<UMovieSceneTrack> TrackClass, FName ControlRigName, FName RigControlName);
diff --git a/Engine/Source/Editor/MovieSceneTools/Private/TrackEditors/TransformTrackEditor.cpp b/Engine/Source/Editor/MovieSceneTools/Private/TrackEditors/TransformTrackEditor.cpp
--- a/Engine/Source/Editor/MovieSceneTools/Private/TrackEditors/TransformTrackEditor.cpp
+++ b/Engine/Source/Editor/MovieSceneTools/Private/TrackEditors/TransformTrackEditor.cpp
@@ -685,6 +685,9 @@ FRotator UnwindRotator(const FRotator& InOld, const FRotator& InNew)
 
 void F3DTransformTrackEditor::GetTransformKeys( const TOptional<FTransformData>& LastTransform, const FTransformData& CurrentTransform, EMovieSceneTransformChannel ChannelsToKey, UObject* Object, UMovieSceneSection* Section, FGeneratedTrackKeys& OutGeneratedKeys)
 {
+	UMovieScene3DTransformSection* TransformSection = Cast<UMovieScene3DTransformSection>(Section);
+	EMovieSceneTransformChannel TransformMask = TransformSection->GetMask().GetChannels();
+
 	using namespace UE::MovieScene;
 
 	bool bLastVectorIsValid = LastTransform.IsSet();
@@ -718,6 +721,19 @@ void F3DTransformTrackEditor::GetTransformKeys( const TOptional<FTransformData>&
 			bKeyX = bKeyY = bKeyZ = true;
 		}
 
+		if (!EnumHasAnyFlags(TransformMask, EMovieSceneTransformChannel::TranslationX))
+		{
+			bKeyX = false;
+		}
+		if (!EnumHasAnyFlags(TransformMask, EMovieSceneTransformChannel::TranslationY))
+		{
+			bKeyY = false;
+		}
+		if (!EnumHasAnyFlags(TransformMask, EMovieSceneTransformChannel::TranslationZ))
+		{
+			bKeyZ = false;
+		}
+
 		FVector3f KeyVector = RecomposedTransform.Translation;
 
 		OutGeneratedKeys.Add(FMovieSceneChannelValueSetter::Create<FMovieSceneFloatChannel>(0, KeyVector.X, bKeyX));
@@ -746,6 +762,19 @@ void F3DTransformTrackEditor::GetTransformKeys( const TOptional<FTransformData>&
 			bKeyX = bKeyY = bKeyZ = true;
 		}
 
+		if (!EnumHasAnyFlags(TransformMask, EMovieSceneTransformChannel::RotationX))
+		{
+			bKeyX = false;
+		}
+		if (!EnumHasAnyFlags(TransformMask, EMovieSceneTransformChannel::RotationY))
+		{
+			bKeyY = false;
+		}
+		if (!EnumHasAnyFlags(TransformMask, EMovieSceneTransformChannel::RotationZ))
+		{
+			bKeyZ = false;
+		}
+
 		// Do we need to unwind re-composed rotations?
 		KeyRotator = UnwindRotator(CurrentTransform.Rotation, RecomposedTransform.Rotation);
 		OutGeneratedKeys.Add(FMovieSceneChannelValueSetter::Create<FMovieSceneFloatChannel>(3, KeyRotator.Roll, bKeyX));
@@ -772,6 +801,19 @@ void F3DTransformTrackEditor::GetTransformKeys( const TOptional<FTransformData>&
 			bKeyX = bKeyY = bKeyZ = true;
 		}
 
+		if (!EnumHasAnyFlags(TransformMask, EMovieSceneTransformChannel::ScaleX))
+		{
+			bKeyX = false;
+		}
+		if (!EnumHasAnyFlags(TransformMask, EMovieSceneTransformChannel::ScaleY))
+		{
+			bKeyY = false;
+		}
+		if (!EnumHasAnyFlags(TransformMask, EMovieSceneTransformChannel::ScaleZ))
+		{
+			bKeyZ = false;
+		}
+
 		FVector3f KeyVector = RecomposedTransform.Scale;
 		OutGeneratedKeys.Add(FMovieSceneChannelValueSetter::Create<FMovieSceneFloatChannel>(6, KeyVector.X, bKeyX));
 		OutGeneratedKeys.Add(FMovieSceneChannelValueSetter::Create<FMovieSceneFloatChannel>(7, KeyVector.Y, bKeyY));
