Sequencer: Fix typo

diff --git a/Engine/Plugins/Experimental/ControlRig/Source/ControlRigEditor/Private/Sequencer/ControlRigParameterTrackEditor.cpp b/Engine/Plugins/Experimental/ControlRig/Source/ControlRigEditor/Private/Sequencer/ControlRigParameterTrackEditor.cpp
--- a/Engine/Plugins/Experimental/ControlRig/Source/ControlRigEditor/Private/Sequencer/ControlRigParameterTrackEditor.cpp
+++ b/Engine/Plugins/Experimental/ControlRig/Source/ControlRigEditor/Private/Sequencer/ControlRigParameterTrackEditor.cpp
@@ -2106,7 +2106,7 @@ void FControlRigParameterTrackEditor::GetControlRigKeys(UControlRig* InControlRi
 	int32 IntChannelIndex = 0;
 	for (int32 ControlIndex = 0; ControlIndex < Controls.Num(); ++ControlIndex)
 	{
-		FRigControlElement* ControlElement = Controls[ControolIndex];
+		FRigControlElement* ControlElement = Controls[ControlIndex];
 
 		if (!ControlElement->Settings.bAnimatable)
 		{
