Control Rig: Use outermost for package path and notify the asset registry that a new level sequence asset has been created

diff --git a/Engine/Plugins/Experimental/ControlRig/Source/ControlRigEditor/Private/ControlRigEditorModule.cpp b/Engine/Plugins/Experimental/ControlRig/Source/ControlRigEditor/Private/ControlRigEditorModule.cpp
--- a/Engine/Plugins/Experimental/ControlRig/Source/ControlRigEditor/Private/ControlRigEditorModule.cpp
+++ b/Engine/Plugins/Experimental/ControlRig/Source/ControlRigEditor/Private/ControlRigEditorModule.cpp
@@ -719,7 +719,7 @@ void FControlRigEditorModule::BakeToControlRig(UClass* ControlRigClass, UAnimSeq
 		FControlRigEditorModule::UnLinkLevelSequence(AnimSequence);
 
 		FString SequenceName = FString::Printf(TEXT("Driving_%s"), *AnimSequence->GetName());
-		FString PackagePath = AnimSequence->GetPathName();
+		FString PackagePath = AnimSequence->GetOutermost()->GetName();
 		
 		FAssetToolsModule& AssetToolsModule = FModuleManager::LoadModuleChecked<FAssetToolsModule>("AssetTools");
 		FString UniquePackageName;
@@ -728,6 +728,8 @@ void FControlRigEditorModule::BakeToControlRig(UClass* ControlRigClass, UAnimSeq
 
 		UPackage* Package = CreatePackage(*UniquePackageName);
 		ULevelSequence* LevelSequence = NewObject<ULevelSequence>(Package, *UniqueAssetName, RF_Public | RF_Standalone);
+					
+		FAssetRegistryModule::AssetCreated(LevelSequence);
 
 		LevelSequence->Initialize(); //creates movie scene
 		LevelSequence->MarkPackageDirty();
