Control Rig: Fix compilation issues SRigSpacePickerWidget.cpp not finding FModuleManager

diff --git a/Engine/Plugins/Animation/ControlRig/Source/ControlRigEditor/Public/ControlRigSpaceChannelCurveModel.h b/Engine/Plugins/Animation/ControlRig/Source/ControlRigEditor/Public/ControlRigSpaceChannelCurveModel.h
--- a/Engine/Plugins/Animation/ControlRig/Source/ControlRigEditor/Public/ControlRigSpaceChannelCurveModel.h
+++ b/Engine/Plugins/Animation/ControlRig/Source/ControlRigEditor/Public/ControlRigSpaceChannelCurveModel.h
@@ -8,6 +8,7 @@
 #include "KeyBarCurveModel.h"
 #include "Sequencer/MovieSceneControlRigSpaceChannel.h"
 #include "Channels/MovieSceneChannelHandle.h"
+#include "Framework/MultiBox/MultiBoxBuilder.h"
 
 class UMovieSceneSection;
 class ISequencer;
