changed "Gizmo" label to "Control Shape"

diff --git a/Engine/Plugins/Animation/ControlRig/Source/ControlRigEditor/Private/ControlRigElementDetails.cpp b/Engine/Plugins/Animation/ControlRig/Source/ControlRigEditor/Private/ControlRigElementDetails.cpp
--- a/Engine/Plugins/Animation/ControlRig/Source/ControlRigEditor/Private/ControlRigElementDetails.cpp
+++ b/Engine/Plugins/Animation/ControlRig/Source/ControlRigEditor/Private/ControlRigElementDetails.cpp
@@ -1743,8 +1743,9 @@ void FRigControlElementDetails::CustomizeChildren(TSharedRef<class IPropertyHand
 	ControlGroup->AddPropertyRow(ControlTypeHandle.ToSharedRef());
 	ControlGroup->AddPropertyRow(SettingsHandle->GetChildHandle(TEXT("bAnimatable")).ToSharedRef());
 
-	// any but bool controls show the offset + gizmo
-	if(IsAnyControlNotOfType(ERigControlType::Bool))
+	// any but bool controls show the offset + gizmo + limits
+	const bool bNeedsGizmoProperties = IsAnyControlNotOfType(ERigControlType::Bool);
+	if (bNeedsGizmoProperties)
 	{
 		// setup offset
 		{
@@ -1755,22 +1756,10 @@ void FRigControlElementDetails::CustomizeChildren(TSharedRef<class IPropertyHand
 			ControlGroup->AddPropertyRow(TransformHandle.ToSharedRef()).DisplayName(FText::FromString(TEXT("Offset Transform")));
 		}
 
-		GizmoGroup = &StructBuilder.AddGroup(TEXT("Gizmo"), LOCTEXT("Gizmo", "Gizmo"));
-	}
-
-	if(IsAnyControlOfType(ERigControlType::Float) ||
-		IsAnyControlOfType(ERigControlType::Integer) ||
-		IsAnyControlOfType(ERigControlType::Vector2D) ||
-		IsAnyControlOfType(ERigControlType::Position) ||
-		IsAnyControlOfType(ERigControlType::Rotator) ||
-		IsAnyControlOfType(ERigControlType::Scale) ||
-		IsAnyControlOfType(ERigControlType::Transform) ||
-		IsAnyControlOfType(ERigControlType::TransformNoScale) ||
-		IsAnyControlOfType(ERigControlType::EulerTransform))
-	{
+		GizmoGroup = &StructBuilder.AddGroup(TEXT("Control Shape"), LOCTEXT("ControlShape", "Control Shape"));
 		LimitsGroup = &StructBuilder.AddGroup(TEXT("Limits"), LOCTEXT("Limits", "Limits"));
 	}
-	
+
 	if(IsAnyControlOfType(ERigControlType::Float) ||
 		IsAnyControlOfType(ERigControlType::Integer) ||
 		IsAnyControlOfType(ERigControlType::Vector2D) ||
@@ -1846,15 +1835,7 @@ void FRigControlElementDetails::CustomizeChildren(TSharedRef<class IPropertyHand
 		FRigControlElementDetails_SetupValueWidget(*ControlGroup, StructBuilder, ERigControlValueType::Maximum, ControlElements[0], HierarchyBeingCustomized);
 	}
 
-	if(IsAnyControlOfType(ERigControlType::Float) ||
-		IsAnyControlOfType(ERigControlType::Integer) ||
-		IsAnyControlOfType(ERigControlType::Vector2D) ||
-		IsAnyControlOfType(ERigControlType::Position) ||
-		IsAnyControlOfType(ERigControlType::Scale) ||
-		IsAnyControlOfType(ERigControlType::Rotator) ||
-		IsAnyControlOfType(ERigControlType::Transform) ||
-		IsAnyControlOfType(ERigControlType::TransformNoScale) ||
-		IsAnyControlOfType(ERigControlType::EulerTransform))
+	if (bNeedsGizmoProperties)
 	{
 		GizmoGroup->AddPropertyRow(SettingsHandle->GetChildHandle(TEXT("bGizmoEnabled")).ToSharedRef())
 		.DisplayName(FText::FromString(TEXT("Enabled")));
