Control Rig - Pose Window: Pose window shouldn't be able to tween negatively.

Current logic for tweening in the Pose Window is:
Input a value, create a key at the frame where lerped from the first pose to the second.
For example, if the input value is 0.3f, we should create a key at the frame where it's 30% from the first pose to the second, which is a lerp operation.
Inputting a negative value here doesn't make sense.

diff --git a/Engine/Plugins/Animation/ControlRig/Source/ControlRigEditor/Private/EditMode/SControlRigControlViews.cpp b/Engine/Plugins/Animation/ControlRig/Source/ControlRigEditor/Private/EditMode/SControlRigControlViews.cpp
--- a/Engine/Plugins/Animation/ControlRig/Source/ControlRigEditor/Private/EditMode/SControlRigControlViews.cpp
+++ b/Engine/Plugins/Animation/ControlRig/Source/ControlRigEditor/Private/EditMode/SControlRigControlViews.cpp
@@ -341,9 +341,9 @@ void SControlRigPoseView::Construct(const FArguments& InArgs)
 								// Only allow spinning if we have a single value
 								.Value(this, &SControlRigPoseView::OnGetPoseBlendValue)
 								.ToolTipText(LOCTEXT("BlendTooltip", "Blend between current pose and pose asset. Use Ctrl drag for under and over shoot."))
-								.MinValue(-2.0f)
-								.MaxValue(2.0f)
-								.MinSliderValue(-1.0f)
+								.MinValue(0.0f)
+								.MaxValue(1.0f)
+								.MinSliderValue(0.0f)
 								.MaxSliderValue(1.0f)
 								.SliderExponent(1)
 								.Delta(0.005f)
