Control Rig: Fix cannot parse any custom default values for Transform variables in Control Rig

diff --git a/Engine/Plugins/Animation/ControlRig/Source/ControlRigEditor/Private/Editor/ControlRigEditor.cpp b/Engine/Plugins/Animation/ControlRig/Source/ControlRigEditor/Private/Editor/ControlRigEditor.cpp
--- a/Engine/Plugins/Animation/ControlRig/Source/ControlRigEditor/Private/Editor/ControlRigEditor.cpp
+++ b/Engine/Plugins/Animation/ControlRig/Source/ControlRigEditor/Private/Editor/ControlRigEditor.cpp
@@ -4634,14 +4634,10 @@ void FControlRigEditor::UpdateDefaultValueForVariable(FBPVariableDescription& In
 
 			if (TargetProperty)
 			{
-				// Grab the address of where the property is actually stored (UObject* base, plus the offset defined in the property)
-				void* OldPropertyAddr = TargetProperty->ContainerPtrToValuePtr<void>(ObjectContainer);
-				if (OldPropertyAddr)
-				{
-					FString NewDefaultValue;
-					TargetProperty->ExportTextItem(NewDefaultValue, OldPropertyAddr, OldPropertyAddr, nullptr, PPF_None);
-					InVariable.DefaultValue = NewDefaultValue;
-				}
+				FString NewDefaultValue;
+				const uint8* Container = (const uint8*)ObjectContainer;
+				FBlueprintEditorUtils::PropertyValueToString(TargetProperty, Container, NewDefaultValue, nullptr);
+				InVariable.DefaultValue = NewDefaultValue;
 			}
 		}
 	}
