Control Rig: compilation causes right click menu gets slower over time. We were constantly adding nodes to the template graph making the process of generating a unique name gets slower over time

diff --git a/Engine/Plugins/Animation/ControlRig/Source/ControlRigEditor/Private/Graph/NodeSpawners/ControlRigVariableNodeSpawner.cpp b/Engine/Plugins/Animation/ControlRig/Source/ControlRigEditor/Private/Graph/NodeSpawners/ControlRigVariableNodeSpawner.cpp
--- a/Engine/Plugins/Animation/ControlRig/Source/ControlRigEditor/Private/Graph/NodeSpawners/ControlRigVariableNodeSpawner.cpp
+++ b/Engine/Plugins/Animation/ControlRig/Source/ControlRigEditor/Private/Graph/NodeSpawners/ControlRigVariableNodeSpawner.cpp
@@ -168,6 +168,12 @@ UEdGraphNode* UControlRigVariableNodeSpawner::Invoke(UEdGraph* ParentGraph, FBin
 			}
 			Controller->CloseUndoBracket();
 		}
+		else
+		{
+			// similar to UBlueprintNodeSpawner::Invoke -> UBlueprintNodeSpawner::SpawnEdGraphNode
+			// we simply want the node, but not actually adding it to a graph
+			Controller->RemoveNode(ModelNode, false);
+		}
 	}
 	else
 	{
