Control Rig: Copy pasting a control rig into a new project with functions breaks function references

diff --git a/Engine/Plugins/Animation/ControlRig/Source/ControlRigDeveloper/Private/ControlRigBlueprint.cpp b/Engine/Plugins/Animation/ControlRig/Source/ControlRigDeveloper/Private/ControlRigBlueprint.cpp
--- a/Engine/Plugins/Animation/ControlRig/Source/ControlRigDeveloper/Private/ControlRigBlueprint.cpp
+++ b/Engine/Plugins/Animation/ControlRig/Source/ControlRigDeveloper/Private/ControlRigBlueprint.cpp
@@ -422,6 +422,7 @@ void UControlRigBlueprint::PostLoad()
 
 		InitializeModelIfRequired(false /* recompile vm */);
 
+		PatchFunctionReferencesOnLoad();
 		PatchVariableNodesOnLoad();
 
 #if WITH_EDITOR
@@ -3303,6 +3304,46 @@ FName UControlRigBlueprint::AddCRMemberVariableFromExternal(FRigVMExternalVariab
 	return NAME_None;
 }
 
+void UControlRigBlueprint::PatchFunctionReferencesOnLoad()
+{
+	// If the asset was copied from one project to another, the function referenced might have a different
+	// path, even if the function is internal to the contorl rig. In that case, let's try to find the function
+	// in the local function library.
+	
+	TArray<URigVMNode*> Nodes = Model->GetNodes();
+	for (URigVMLibraryNode* Library : FunctionLibrary->GetFunctions())
+	{
+		Nodes.Append(Library->GetContainedNodes());
+	}
+	
+	for (URigVMNode* Node : Nodes)
+	{
+		if (URigVMFunctionReferenceNode* FunctionReferenceNode = Cast<URigVMFunctionReferenceNode>(Node))
+		{
+			if (!FunctionReferenceNode->GetReferencedNode())
+			{
+				if(FunctionLibrary)
+				{
+					FString FunctionPath = FunctionReferenceNode->ReferencedNodePtr.ToSoftObjectPath().GetSubPathString();
+					
+					FString Left, Right;
+					if(FunctionPath.Split(TEXT("."), &Left, &Right))
+					{
+						FString LibraryNodePath = FunctionLibrary->GetNodePath();
+						if(Left == FunctionLibrary->GetName())
+						{
+							if (URigVMLibraryNode* LibraryNode = Cast<URigVMLibraryNode>(FunctionLibrary->FindNode(Right)))
+							{
+								FunctionReferenceNode->SetReferencedNode(LibraryNode);
+							}
+						}
+					}
+				}				
+			}
+		}
+	}
+}
+
 #endif
 
 void UControlRigBlueprint::PatchVariableNodesOnLoad()
diff --git a/Engine/Plugins/Animation/ControlRig/Source/ControlRigDeveloper/Public/ControlRigBlueprint.h b/Engine/Plugins/Animation/ControlRig/Source/ControlRigDeveloper/Public/ControlRigBlueprint.h
--- a/Engine/Plugins/Animation/ControlRig/Source/ControlRigDeveloper/Public/ControlRigBlueprint.h
+++ b/Engine/Plugins/Animation/ControlRig/Source/ControlRigDeveloper/Public/ControlRigBlueprint.h
@@ -528,6 +528,7 @@ private:
 	static int32 AddCRMemberVariable(UControlRigBlueprint* InBlueprint, const FName& InVarName, FEdGraphPinType InVarType, bool bIsPublic, bool bIsReadOnly, FString InDefaultValue);
 	FName AddCRMemberVariableFromExternal(FRigVMExternalVariable InVariableToCreate, FString InDefaultValue = FString());
 #endif
+	void PatchFunctionReferencesOnLoad();
 	void PatchVariableNodesOnLoad();
 
 	TMap<FName, int32> AddedMemberVariableMap;
diff --git a/Engine/Source/Developer/RigVMDeveloper/Public/RigVMModel/Nodes/RigVMFunctionReferenceNode.h b/Engine/Source/Developer/RigVMDeveloper/Public/RigVMModel/Nodes/RigVMFunctionReferenceNode.h
--- a/Engine/Source/Developer/RigVMDeveloper/Public/RigVMModel/Nodes/RigVMFunctionReferenceNode.h
+++ b/Engine/Source/Developer/RigVMDeveloper/Public/RigVMModel/Nodes/RigVMFunctionReferenceNode.h
@@ -58,5 +58,6 @@ private:
 
 	friend class URigVMController;
 	friend class FRigVMParserAST;
+	friend class UControlRigBlueprint;
 };
 
