Control Rig: Fix a crash caused by uncleared delegate

diff --git a/Engine/Plugins/Animation/ControlRig/Source/ControlRig/Private/ControlRigControlActor.cpp b/Engine/Plugins/Animation/ControlRig/Source/ControlRig/Private/ControlRigControlActor.cpp
--- a/Engine/Plugins/Animation/ControlRig/Source/ControlRig/Private/ControlRigControlActor.cpp
+++ b/Engine/Plugins/Animation/ControlRig/Source/ControlRig/Private/ControlRigControlActor.cpp
@@ -119,6 +119,15 @@ void AControlRigControlActor::EndPlay(const EEndPlayReason::Type EndPlayReason)
 	Super::EndPlay(EndPlayReason);
 }
 
+void AControlRigControlActor::BeginDestroy()
+{
+	// since end play is not always called, we have to clear the delegate here
+	// clearing it at destructor might be too late as in some cases, the control rig was already GCed
+	RemoveUnbindDelegate();
+	ControlRig = nullptr;
+	Super::BeginDestroy();
+}
+
 void AControlRigControlActor::Refresh()
 {
 	if (ActorToTrack == nullptr)
diff --git a/Engine/Plugins/Animation/ControlRig/Source/ControlRig/Public/ControlRigControlActor.h b/Engine/Plugins/Animation/ControlRig/Source/ControlRig/Public/ControlRigControlActor.h
--- a/Engine/Plugins/Animation/ControlRig/Source/ControlRig/Public/ControlRigControlActor.h
+++ b/Engine/Plugins/Animation/ControlRig/Source/ControlRig/Public/ControlRigControlActor.h
@@ -27,7 +27,7 @@ public:
 	virtual void Tick(float DeltaSeconds) override;
 	virtual void BeginPlay() override { Clear(); }
 	virtual void EndPlay(const EEndPlayReason::Type EndPlayReason) override;
-
+	virtual void BeginDestroy() override;
 
 #if WITH_EDITOR
 	virtual void PostEditChangeProperty(FPropertyChangedEvent& PropertyChangedEvent) override;
