Control Rig: Revert back to new SpringInterp node

diff --git a/Engine/Plugins/Animation/ControlRig/Source/ControlRig/Private/Units/Simulation/RigUnit_SpringInterp.cpp b/Engine/Plugins/Animation/ControlRig/Source/ControlRig/Private/Units/Simulation/RigUnit_SpringInterp.cpp
--- a/Engine/Plugins/Animation/ControlRig/Source/ControlRig/Private/Units/Simulation/RigUnit_SpringInterp.cpp
+++ b/Engine/Plugins/Animation/ControlRig/Source/ControlRig/Private/Units/Simulation/RigUnit_SpringInterp.cpp
@@ -5,10 +5,60 @@
 
 namespace RigUnitSpringInterpConstants
 {
+	static const float FixedTimeStep = 1.0f / 60.0f;
+	static const float MaxTimeStep = 0.1f;
 	static const float Mass = 1.0f;
 }
 
 FRigUnit_SpringInterp_Execute()
+{
+	DECLARE_SCOPE_HIERARCHICAL_COUNTER_RIGUNIT()
+ 
+	if (Context.State == EControlRigState::Init)
+	{
+		SpringState.Reset();
+	}
+	else
+	{
+		// Clamp to avoid large time deltas.
+		float RemainingTime = FMath::Min(Context.DeltaTime, RigUnitSpringInterpConstants::MaxTimeStep);
+ 
+		Result = Current;
+		while (RemainingTime >= RigUnitSpringInterpConstants::FixedTimeStep)
+		{
+			Result = UKismetMathLibrary::FloatSpringInterp(Result, Target, SpringState, Stiffness, CriticalDamping, RigUnitSpringInterpConstants::FixedTimeStep, Mass);
+			RemainingTime -= RigUnitSpringInterpConstants::FixedTimeStep;
+		}
+ 
+		Result = UKismetMathLibrary::FloatSpringInterp(Result, Target, SpringState, Stiffness, CriticalDamping, RemainingTime, Mass);
+	}
+}
+ 
+FRigUnit_SpringInterpVector_Execute()
+{
+	DECLARE_SCOPE_HIERARCHICAL_COUNTER_RIGUNIT()
+ 
+	if (Context.State == EControlRigState::Init)
+	{
+		SpringState.Reset();
+	}
+	else
+	{
+		// Clamp to avoid large time deltas.
+		float RemainingTime = FMath::Min(Context.DeltaTime, RigUnitSpringInterpConstants::MaxTimeStep);
+ 
+		Result = Current;
+		while (RemainingTime >= RigUnitSpringInterpConstants::FixedTimeStep)
+		{
+			Result = UKismetMathLibrary::VectorSpringInterp(Result, Target, SpringState, Stiffness, CriticalDamping, RigUnitSpringInterpConstants::FixedTimeStep, Mass);
+			RemainingTime -= RigUnitSpringInterpConstants::FixedTimeStep;
+		}
+ 
+		Result = UKismetMathLibrary::VectorSpringInterp(Result, Target, SpringState, Stiffness, CriticalDamping, RemainingTime, Mass);
+	}
+}
+
+FRigUnit_SpringInterpV2_Execute()
 {
 	DECLARE_SCOPE_HIERARCHICAL_COUNTER_RIGUNIT()
 
@@ -41,7 +91,7 @@ FRigUnit_SpringInterp_Execute()
 	}
 }
 
-FRigUnit_SpringInterpVector_Execute()
+FRigUnit_SpringInterpVectorV2_Execute()
 {
     DECLARE_SCOPE_HIERARCHICAL_COUNTER_RIGUNIT()
 
@@ -73,7 +123,7 @@ FRigUnit_SpringInterpVector_Execute()
 	Velocity = SpringState.Velocity;
 }
 
-FRigUnit_SpringInterpQuaternion_Execute()
+FRigUnit_SpringInterpQuaternionV2_Execute()
 {
     DECLARE_SCOPE_HIERARCHICAL_COUNTER_RIGUNIT()
 
diff --git a/Engine/Plugins/Animation/ControlRig/Source/ControlRig/Private/Units/Simulation/RigUnit_SpringInterp.h b/Engine/Plugins/Animation/ControlRig/Source/ControlRig/Private/Units/Simulation/RigUnit_SpringInterp.h
--- a/Engine/Plugins/Animation/ControlRig/Source/ControlRig/Private/Units/Simulation/RigUnit_SpringInterp.h
+++ b/Engine/Plugins/Animation/ControlRig/Source/ControlRig/Private/Units/Simulation/RigUnit_SpringInterp.h
@@ -9,12 +9,96 @@
 /**
  * Uses a simple spring model to interpolate a float from Current to Target.
  */
-USTRUCT(meta=(DisplayName="Spring Interpolate", Keywords="Alpha,SpringInterpolate,Verlet", PrototypeName = "SpringInterp"))
+USTRUCT(meta=(DisplayName="Spring Interpolate", Keywords="Alpha,SpringInterpolate", PrototypeName = "SpringInterp", Deprecated="5.0"))
 struct CONTROLRIG_API FRigUnit_SpringInterp : public FRigUnit_SimBase
 {
 	GENERATED_BODY()
-
+ 
 	FRigUnit_SpringInterp()
+	{
+		Current = Target = Result = 0.0f;
+		Stiffness = 10.0f;
+		CriticalDamping = 2.0f;
+		Mass = 10.0f;
+		SpringState = FFloatSpringState();
+	}
+ 
+	RIGVM_METHOD()
+	virtual void Execute(const FRigUnitContext& Context) override;
+ 
+	UPROPERTY(meta = (Input))
+	float Current;
+ 
+	UPROPERTY(meta=(Input))
+	float Target;
+ 
+	UPROPERTY(meta=(Input))
+	float Stiffness;
+ 
+	UPROPERTY(meta=(Input))
+	float CriticalDamping;
+ 
+	UPROPERTY(meta=(Input))
+	float Mass;
+ 
+	UPROPERTY(meta=(Output))
+	float Result;
+ 
+	UPROPERTY()
+	FFloatSpringState SpringState;
+};
+ 
+/**
+ * Uses a simple spring model to interpolate a vector from Current to Target.
+ */
+USTRUCT(meta=(DisplayName="Spring Interpolate", Keywords="Alpha,SpringInterpolate", PrototypeName = "SpringInterp", MenuDescSuffix = "(Vector)", Deprecated="5.0"))
+struct CONTROLRIG_API FRigUnit_SpringInterpVector : public FRigUnit_SimBase
+{
+	GENERATED_BODY()
+	
+	FRigUnit_SpringInterpVector()
+	{
+		Current = Target = Result = FVector::ZeroVector;
+		Stiffness = 10.0f;
+		CriticalDamping = 2.0f;
+		Mass = 10.0f;
+		SpringState = FVectorSpringState();
+	}
+ 
+	RIGVM_METHOD()
+	virtual void Execute(const FRigUnitContext& Context) override;
+ 
+	UPROPERTY(meta = (Input))
+	FVector Current;
+ 
+	UPROPERTY(meta=(Input))
+	FVector Target;
+ 
+	UPROPERTY(meta=(Input))
+	float Stiffness;
+ 
+	UPROPERTY(meta=(Input))
+	float CriticalDamping;
+ 
+	UPROPERTY(meta=(Input))
+	float Mass;
+ 
+	UPROPERTY(meta=(Output))
+	FVector Result;
+ 
+	UPROPERTY()
+	FVectorSpringState SpringState;
+};
+
+/**
+ * Uses a simple spring model to interpolate a float from Current to Target.
+ */
+USTRUCT(meta=(DisplayName="Spring Interpolate", Keywords="Alpha,SpringInterpolate,Verlet", PrototypeName = "SpringInterp"))
+struct CONTROLRIG_API FRigUnit_SpringInterpV2 : public FRigUnit_SimBase
+{
+	GENERATED_BODY()
+
+	FRigUnit_SpringInterpV2()
 	{
 		bUseCurrentInput = false;
 		Current = Target = Velocity = Result = SimulatedResult = 0.0f;
@@ -100,11 +184,11 @@ struct CONTROLRIG_API FRigUnit_SpringInterp : public FRigUnit_SimBase
  * Uses a simple spring model to interpolate a vector from Current to Target.
  */
 USTRUCT(meta=(DisplayName="Spring Interpolate", Keywords="Alpha,SpringInterpolate,Verlet", PrototypeName = "SpringInterp", MenuDescSuffix = "(Vector)"))
-struct CONTROLRIG_API FRigUnit_SpringInterpVector : public FRigUnit_SimBase
+struct CONTROLRIG_API FRigUnit_SpringInterpVectorV2 : public FRigUnit_SimBase
 {
 	GENERATED_BODY()
 	
-	FRigUnit_SpringInterpVector()
+	FRigUnit_SpringInterpVectorV2()
 	{
 		bUseCurrentInput = false;
 		Current = Target = Velocity = Result = SimulatedResult = FVector::ZeroVector;
@@ -190,11 +274,11 @@ struct CONTROLRIG_API FRigUnit_SpringInterpVector : public FRigUnit_SimBase
  * Uses a simple spring model to interpolate a quaternion from Current to Target.
  */
 USTRUCT(meta=(DisplayName="Spring Interpolate", Keywords="Alpha,SpringInterpolate,Verlet", PrototypeName = "SpringInterp", MenuDescSuffix = "(Quaternion)"))
-struct CONTROLRIG_API FRigUnit_SpringInterpQuaternion : public FRigUnit_SimBase
+struct CONTROLRIG_API FRigUnit_SpringInterpQuaternionV2 : public FRigUnit_SimBase
 {
 	GENERATED_BODY()
 	
-	FRigUnit_SpringInterpQuaternion()
+	FRigUnit_SpringInterpQuaternionV2()
 	{
 		bUseCurrentInput = false;
 		Current = Target = Result = SimulatedResult = FQuat::Identity;
