Control Rig: Fix function's Return Node doesn't support Array variable Drag and Drop

diff --git a/Engine/Plugins/Animation/ControlRig/Source/ControlRigDeveloper/Private/Graph/ControlRigGraphSchema.cpp b/Engine/Plugins/Animation/ControlRig/Source/ControlRigDeveloper/Private/Graph/ControlRigGraphSchema.cpp
--- a/Engine/Plugins/Animation/ControlRig/Source/ControlRigDeveloper/Private/Graph/ControlRigGraphSchema.cpp
+++ b/Engine/Plugins/Animation/ControlRig/Source/ControlRigDeveloper/Private/Graph/ControlRigGraphSchema.cpp
@@ -1280,6 +1280,10 @@ UEdGraphPin* UControlRigGraphSchema::DropPinOnNode(UEdGraphNode* InTargetNode, c
 						if (URigVMController* Controller = RigBlueprint->GetController(Model))
 						{
 							FString TypeName = ExternalVar.TypeName.ToString();
+							if (ExternalVar.bIsArray)
+							{
+								TypeName = FString::Printf(URigVMController::TArrayTemplate, *TypeName);
+							}
 							FName TypeObjectPathName = NAME_None;
 							if (ExternalVar.TypeObject)
 							{
diff --git a/Engine/Source/Developer/RigVMDeveloper/Public/RigVMModel/RigVMController.h b/Engine/Source/Developer/RigVMDeveloper/Public/RigVMModel/RigVMController.h
--- a/Engine/Source/Developer/RigVMDeveloper/Public/RigVMModel/RigVMController.h
+++ b/Engine/Source/Developer/RigVMDeveloper/Public/RigVMModel/RigVMController.h
@@ -890,11 +890,14 @@ private:
 	static FString GetSanitizedPinName(const FString& InName);
 	static FString GetSanitizedPinPath(const FString& InName);
 	static void SanitizeName(FString& InOutName, bool bAllowPeriod, bool bAllowSpace);
+
+public:
 	static const TCHAR TArrayPrefix[];
 	static const TCHAR TObjectPtrPrefix[];
 	static const TCHAR TArrayTemplate[];
 	static const TCHAR TObjectPtrTemplate[];
 
+private: 
 	UPROPERTY(transient)
 	TArray<TObjectPtr<URigVMGraph>> Graphs;
 
