Control Rig: Rename the term Gizmo to Shape

diff --git a/Engine/Plugins/Animation/ControlRig/Config/BaseControlRig.ini b/Engine/Plugins/Animation/ControlRig/Config/BaseControlRig.ini
--- a/Engine/Plugins/Animation/ControlRig/Config/BaseControlRig.ini
+++ b/Engine/Plugins/Animation/ControlRig/Config/BaseControlRig.ini
@@ -4,8 +4,6 @@
 +ClassRedirects=(OldName="/Script/ControlRigEditor.ControlRigGraphNode", NewName="/Script/ControlRigDeveloper.ControlRigGraphNode")
 +ClassRedirects=(OldName="/Script/ControlRigEditor.ControlRigGraphSchema", NewName="/Script/ControlRigDeveloper.ControlRigGraphSchema")
 +ClassRedirects=(OldName="/Script/ControlRigEditor.AnimGraphNode_ControlRig", NewName="/Script/ControlRigDeveloper.AnimGraphNode_ControlRig")
-+StructRedirects=(OldName="/Script/ControlRigEditor.ControlRigBlueprintPropertyLink", NewName="/Script/ControlRigDeveloper.ControlRigBlueprintPropertyLink")
-
 
 +StructRedirects=(OldName="RigJoint",NewName="/Script/ControlRig.RigBone")
 +PropertyRedirects=(OldName="RigBoneHierarchy.Joints",NewName="RigBoneHierarchy.Bones")
@@ -13,4 +11,19 @@
 
 +ClassRedirects=(OldName="RigVMGraph",NewName="RigVMGraph",ValueChanges=(("Model","RigVMModel")))
 
-+EnumRedirects=(OldName="ERigElementType",ValueChanges=(("Space","Null"))
\ No newline at end of file
++EnumRedirects=(OldName="ERigElementType",ValueChanges=(("Space","Null"))
+
+; renaming gizmo to shape
++ClassRedirects=(OldName="/Script/ControlRig.ControlRigGizmoLibrary", NewName="/Script/ControlRig.ControlRigShapeLibrary")
++ClassRedirects=(OldName="/Script/ControlRig.ControlRigGizmoActor", NewName="/Script/ControlRig.ControlRigShapeActor")
++StructRedirects=(OldName="ControlRigGizmoDefinition", NewName="/Script/ControlRig.ControlRigShapeDefinition")
++StructRedirects=(OldName="GizmoActorCreationParam", NewName="/Script/ControlRig.ControlShapeActorCreationParam")
++PropertyRedirects=(OldName="/Script/ControlRig.ControlRigShapeLibrary.DefaultGizmo",NewName="DefaultShape")
++PropertyRedirects=(OldName="/Script/ControlRig.ControlRigShapeLibrary.Gizmos",NewName="Shapes")
++PropertyRedirects=(OldName="/Script/ControlRig.ControlRigShapeDefinition.GizmoName",NewName="ShapeName")
++PropertyRedirects=(OldName="/Script/ControlRig.RigControlSettings.bGizmoEnabled",NewName="bShapeEnabled")
++PropertyRedirects=(OldName="/Script/ControlRig.RigControlSettings.bGizmoVisible",NewName="bShapeVisible")
++PropertyRedirects=(OldName="/Script/ControlRig.RigControlSettings.GizmoName",NewName="ShapeName")
++PropertyRedirects=(OldName="/Script/ControlRig.RigControlSettings.GizmoColor",NewName="ShapeColor")
++PropertyRedirects=(OldName="/Script/ControlRig.RigControlElement.Gizmo",NewName="Shape")
++PropertyRedirects=(OldName="/Script/ControlRigEditor.ControlRigEditModeSettings.GizmoScale",NewName="ShapeScale")
diff --git a/Engine/Plugins/Animation/ControlRig/Source/ControlRig/Private/ControlRig.cpp b/Engine/Plugins/Animation/ControlRig/Source/ControlRig/Private/ControlRig.cpp
--- a/Engine/Plugins/Animation/ControlRig/Source/ControlRig/Private/ControlRig.cpp
+++ b/Engine/Plugins/Animation/ControlRig/Source/ControlRig/Private/ControlRig.cpp
@@ -885,8 +885,8 @@ void UControlRig::Execute(const EControlRigState InState, const FName& InEventNa
 		{
 			const FRigControlSettings& Settings = ControlElement->Settings;
 			
-			if (Settings.bGizmoEnabled &&
-				Settings.bGizmoVisible &&
+			if (Settings.bShapeEnabled &&
+				Settings.bShapeVisible &&
 				!Settings.bIsTransientControl &&
 				Settings.bDrawLimits &&
 				(Settings.bLimitTranslation
@@ -900,7 +900,7 @@ void UControlRig::Execute(const EControlRigState InState, const FName& InEventNa
 				}
 
 				FTransform Transform = GetHierarchy()->GetGlobalControlOffsetTransformByIndex(ControlElement->GetIndex());
-				FControlRigDrawInstruction Instruction(EControlRigDrawSettings::Lines, Settings.GizmoColor, 0.f, Transform);
+				FControlRigDrawInstruction Instruction(EControlRigDrawSettings::Lines, Settings.ShapeColor, 0.f, Transform);
 
 				switch (Settings.ControlType)
 				{
@@ -1405,9 +1405,9 @@ void UControlRig::CreateRigControlsForCurveContainer()
 				Settings.bIsCurve = true;
 				Settings.bAnimatable = true;
 				Settings.bDrawLimits = false;
-				Settings.bGizmoEnabled = false;
-				Settings.bGizmoVisible = false;
-				Settings.GizmoColor = FLinearColor::Red;
+				Settings.bShapeEnabled = false;
+				Settings.bShapeVisible = false;
+				Settings.ShapeColor = FLinearColor::Red;
 
 				FRigControlValue Value;
 				Value.Set<float>(CurveElement->Value);
@@ -1533,20 +1533,22 @@ FTransform UControlRig::GetControlLocalTransform(const FName& InControlName)
 	return DynamicHierarchy->GetLocalTransform(FRigElementKey(InControlName, ERigElementType::Control));
 }
 
-UControlRigGizmoLibrary* UControlRig::GetGizmoLibrary() const
+const TArray<TSoftObjectPtr<UControlRigShapeLibrary>>& UControlRig::GetShapeLibraries() const
 {
 	if (UControlRig* CDO = Cast<UControlRig>(GetClass()->GetDefaultObject()))
 	{
-		if (!CDO->GizmoLibrary.IsValid())
+		for(TSoftObjectPtr<UControlRigShapeLibrary>& ShapeLibrary : CDO->ShapeLibraries)
 		{
-			CDO->GizmoLibrary.LoadSynchronous();
-		}
-		if (CDO->GizmoLibrary.IsValid())
-		{
-			return CDO->GizmoLibrary.Get();
+			if (!ShapeLibrary.IsValid())
+			{
+				ShapeLibrary.LoadSynchronous();
+			}
 		}
+		return CDO->ShapeLibraries;
 	}
-	return nullptr;
+
+	static TArray<TSoftObjectPtr<UControlRigShapeLibrary>> EmptyShapeLibraries;
+	return EmptyShapeLibraries;
 }
 
 void UControlRig::SelectControl(const FName& InControlName, bool bSelect)
@@ -1615,7 +1617,7 @@ void UControlRig::HandleHierarchyModified(ERigHierarchyNotification InNotificati
 			break;
 		}
 		case ERigHierarchyNotification::ControlSettingChanged:
-		case ERigHierarchyNotification::ControlGizmoTransformChanged:
+		case ERigHierarchyNotification::ControlShapeTransformChanged:
 		{
 			if(FRigControlElement* ControlElement = Cast<FRigControlElement>((FRigBaseElement*)InElement))
 			{
@@ -1657,8 +1659,8 @@ FName UControlRig::AddTransientControl(URigVMPin* InPin, FRigElementKey SpaceKey
 	URigVMPin* PinForLink = InPin->GetPinForLink();
 
 	const FName ControlName = GetNameForTransientControl(InPin);
-	FTransform GizmoTransform = FTransform::Identity;
-	GizmoTransform.SetScale3D(FVector::ZeroVector);
+	FTransform ShapeTransform = FTransform::Identity;
+	ShapeTransform.SetScale3D(FVector::ZeroVector);
 
 	FRigControlSettings Settings;
 	if (URigVMPin* ColorPin = PinForLink->GetNode()->FindPin(TEXT("Color")))
@@ -1666,7 +1668,7 @@ FName UControlRig::AddTransientControl(URigVMPin* InPin, FRigElementKey SpaceKey
 		if (ColorPin->GetCPPType() == TEXT("FLinearColor"))
 		{
 			FRigControlValue Value;
-			Settings.GizmoColor = Value.SetFromString<FLinearColor>(ColorPin->GetDefaultValue());
+			Settings.ShapeColor = Value.SetFromString<FLinearColor>(ColorPin->GetDefaultValue());
 		}
 	}
 	Settings.bIsTransientControl = true;
@@ -1686,7 +1688,7 @@ FName UControlRig::AddTransientControl(URigVMPin* InPin, FRigElementKey SpaceKey
     	Settings,
     	FRigControlValue::Make(FTransform::Identity),
     	OffsetTransform,
-    	GizmoTransform, false);
+    	ShapeTransform, false);
 
 	SetTransientControlValue(InPin);
 
@@ -1790,8 +1792,8 @@ FName UControlRig::AddTransientControl(const FRigElementKey& InElement)
 		return NAME_None;
 	}
 
-	FTransform GizmoTransform = FTransform::Identity;
-	GizmoTransform.SetScale3D(FVector::ZeroVector);
+	FTransform ShapeTransform = FTransform::Identity;
+	ShapeTransform.SetScale3D(FVector::ZeroVector);
 
 	FRigControlSettings Settings;
 	Settings.bIsTransientControl = true;
@@ -1828,7 +1830,7 @@ FName UControlRig::AddTransientControl(const FRigElementKey& InElement)
         Settings,
         FRigControlValue::Make(FTransform::Identity),
         FTransform::Identity,
-        GizmoTransform, false);
+        ShapeTransform, false);
 
 	if (InElement.Type == ERigElementType::Bone)
 	{
@@ -2899,7 +2901,7 @@ UControlRig::FTransientControlScope::FTransientControlScope(TObjectPtr<URigHiera
 		// preserve whatever value that was produced by this transient control at the moment
 		Info.Value = Hierarchy->GetControlValue(Control->GetKey(),ERigControlValueType::Current);
 		Info.OffsetTransform = Hierarchy->GetControlOffsetTransform(Control, ERigTransformType::CurrentLocal);
-		Info.GizmoTransform = Hierarchy->GetControlGizmoTransform(Control, ERigTransformType::CurrentLocal);
+		Info.ShapeTransform = Hierarchy->GetControlShapeTransform(Control, ERigTransformType::CurrentLocal);
 				
 		SavedTransientControls.Add(Info);
 	}
@@ -2916,7 +2918,7 @@ UControlRig::FTransientControlScope::~FTransientControlScope()
             Info.Settings,
             Info.Value,
             Info.OffsetTransform,
-            Info.GizmoTransform,
+            Info.ShapeTransform,
             false,
             false
         );
diff --git a/Engine/Plugins/Animation/ControlRig/Source/ControlRig/Private/ControlRigControlActor.cpp b/Engine/Plugins/Animation/ControlRig/Source/ControlRig/Private/ControlRigControlActor.cpp
--- a/Engine/Plugins/Animation/ControlRig/Source/ControlRig/Private/ControlRigControlActor.cpp
+++ b/Engine/Plugins/Animation/ControlRig/Source/ControlRig/Private/ControlRigControlActor.cpp
@@ -107,7 +107,7 @@ void AControlRigControlActor::Clear()
 	}
 
 	ControlNames.Reset();
-	GizmoTransforms.Reset();
+	ShapeTransforms.Reset();
 	Components.Reset();
 	Materials.Reset();
 }
@@ -154,32 +154,24 @@ void AControlRigControlActor::Refresh()
 			OnUnbindDelegate = Binding->OnControlRigUnbind().AddLambda([ this ]( ) { this->Clear(); this->Refresh(); });
 		}
 
-		UControlRigGizmoLibrary* GizmoLibrary = ControlRig->GetGizmoLibrary();
-		if (GizmoLibrary == nullptr)
+		const TArray<TSoftObjectPtr<UControlRigShapeLibrary>>& ShapeLibraries = ControlRig->GetShapeLibraries();
+		if(ShapeLibraries.IsEmpty())
 		{
 			return;
 		}
-
+		
 		// disable collision again
 		SetActorEnableCollision(false);
 
-		// prepare the material + color param
-		UMaterialInterface* BaseMaterial = nullptr;
-		if (MaterialOverride && !ColorParameter.IsEmpty())
-		{
-			BaseMaterial = MaterialOverride;
-			ColorParameterName = *ColorParameter;
-		}
-		else
+		for(const TSoftObjectPtr<UControlRigShapeLibrary>& ShapeLibrary : ShapeLibraries)
 		{
-			BaseMaterial = GizmoLibrary->DefaultMaterial.LoadSynchronous();
-			ColorParameterName = GizmoLibrary->MaterialColorParameter;
+			ShapeLibrary->DefaultMaterial.LoadSynchronous();
 		}
 
 		URigHierarchy* Hierarchy = ControlRig->GetHierarchy();
-		Hierarchy->ForEach<FRigControlElement>([this, GizmoLibrary, BaseMaterial, Hierarchy](FRigControlElement* ControlElement) -> bool
+		Hierarchy->ForEach<FRigControlElement>([this, ShapeLibraries, Hierarchy](FRigControlElement* ControlElement) -> bool
         {
-			if (!ControlElement->Settings.bGizmoEnabled)
+			if (!ControlElement->Settings.bShapeEnabled)
 			{
 				return true;
 			}
@@ -196,9 +188,33 @@ void AControlRigControlActor::Refresh()
 				case ERigControlType::TransformNoScale:
 				case ERigControlType::EulerTransform:
 				{
-					if (const FControlRigGizmoDefinition* GizmoDef = GizmoLibrary->GetGizmoByName(ControlElement->Settings.GizmoName))
+					if (const FControlRigShapeDefinition* ShapeDef = UControlRigShapeLibrary::GetShapeByName(ControlElement->Settings.ShapeName, ShapeLibraries))
 					{
-						UStaticMesh* StaticMesh = GizmoDef->StaticMesh.LoadSynchronous();
+						UMaterialInterface* BaseMaterial = nullptr;
+						if (MaterialOverride && !ColorParameter.IsEmpty())
+						{
+							BaseMaterial = MaterialOverride;
+							ColorParameterName = *ColorParameter;
+						}
+						else
+						{
+							if(!ShapeDef->Library.IsValid())
+							{
+								return true;
+							}
+
+							if(ShapeDef->Library->DefaultMaterial.IsValid())
+							{
+								BaseMaterial = ShapeDef->Library->DefaultMaterial.Get();
+							}
+							else
+							{
+								BaseMaterial = ShapeDef->Library->DefaultMaterial.LoadSynchronous();
+							}
+							ColorParameterName = ShapeDef->Library->MaterialColorParameter;
+						}
+
+						UStaticMesh* StaticMesh = ShapeDef->StaticMesh.LoadSynchronous();
 						UStaticMeshComponent* Component = NewObject< UStaticMeshComponent>(ActorRootComponent);
 						Component->SetStaticMesh(StaticMesh);
 						Component->SetupAttachment(ActorRootComponent);
@@ -211,7 +227,7 @@ void AControlRigControlActor::Refresh()
 						Component->SetMaterial(0, MaterialInstance);
 
 						ControlNames.Add(ControlElement->GetName());
-						GizmoTransforms.Add(Hierarchy->GetControlGizmoTransform(ControlElement, ERigTransformType::CurrentLocal) * GizmoDef->Transform);
+						ShapeTransforms.Add(Hierarchy->GetControlShapeTransform(ControlElement, ERigTransformType::CurrentLocal) * ShapeDef->Transform);
 						Components.Add(Component);
 						Materials.Add(MaterialInstance);
 					}
@@ -238,8 +254,8 @@ void AControlRigControlActor::Refresh()
 		}
 
 		FTransform ControlTransform = ControlRig->GetControlGlobalTransform(ControlNames[GizmoIndex]);
-		Components[GizmoIndex]->SetRelativeTransform(GizmoTransforms[GizmoIndex] * ControlTransform);
-		Materials[GizmoIndex]->SetVectorParameterValue(ColorParameterName, FVector(ControlElement->Settings.GizmoColor));
+		Components[GizmoIndex]->SetRelativeTransform(ShapeTransforms[GizmoIndex] * ControlTransform);
+		Materials[GizmoIndex]->SetVectorParameterValue(ColorParameterName, FVector(ControlElement->Settings.ShapeColor));
 	}
 }
 
diff --git a/Engine/Plugins/Animation/ControlRig/Source/ControlRig/Private/ControlRigGizmoActor.cpp b/Engine/Plugins/Animation/ControlRig/Source/ControlRig/Private/ControlRigGizmoActor.cpp
--- a/Engine/Plugins/Animation/ControlRig/Source/ControlRig/Private/ControlRigGizmoActor.cpp
+++ b/Engine/Plugins/Animation/ControlRig/Source/ControlRig/Private/ControlRigGizmoActor.cpp
@@ -7,7 +7,7 @@
 #include "Materials/MaterialInstanceDynamic.h"
 #include "Engine/CollisionProfile.h"
 
-AControlRigGizmoActor::AControlRigGizmoActor(const FObjectInitializer& ObjectInitializer)
+AControlRigShapeActor::AControlRigShapeActor(const FObjectInitializer& ObjectInitializer)
 	: Super(ObjectInitializer)
 	, ControlRigIndex(INDEX_NONE)
 	, ControlName(NAME_None)
@@ -34,7 +34,7 @@ AControlRigGizmoActor::AControlRigGizmoActor(const FObjectInitializer& ObjectIni
 	StaticMeshComponent->bSelectable = bSelectable && bEnabled;
 }
 
-void AControlRigGizmoActor::SetEnabled(bool bInEnabled)
+void AControlRigShapeActor::SetEnabled(bool bInEnabled)
 {
 	if(bEnabled != bInEnabled)
 	{
@@ -45,12 +45,12 @@ void AControlRigGizmoActor::SetEnabled(bool bInEnabled)
 	}
 }
 
-bool AControlRigGizmoActor::IsEnabled() const
+bool AControlRigShapeActor::IsEnabled() const
 {
 	return bEnabled;
 }
 
-void AControlRigGizmoActor::SetSelected(bool bInSelected)
+void AControlRigShapeActor::SetSelected(bool bInSelected)
 {
 	if(bSelected != bInSelected)
 	{
@@ -60,12 +60,12 @@ void AControlRigGizmoActor::SetSelected(bool bInSelected)
 	}
 }
 
-bool AControlRigGizmoActor::IsSelectedInEditor() const
+bool AControlRigShapeActor::IsSelectedInEditor() const
 {
 	return bSelected;
 }
 
-void AControlRigGizmoActor::SetSelectable(bool bInSelectable)
+void AControlRigShapeActor::SetSelectable(bool bInSelectable)
 {
 	if (bSelectable != bInSelectable)
 	{
@@ -78,7 +78,7 @@ void AControlRigGizmoActor::SetSelectable(bool bInSelectable)
 	}
 }
 
-void AControlRigGizmoActor::SetHovered(bool bInHovered)
+void AControlRigShapeActor::SetHovered(bool bInHovered)
 {
 	bool bOldHovered = bHovered;
 
@@ -91,13 +91,13 @@ void AControlRigGizmoActor::SetHovered(bool bInHovered)
 	}
 }
 
-bool AControlRigGizmoActor::IsHovered() const
+bool AControlRigShapeActor::IsHovered() const
 {
 	return bHovered;
 }
 
 
-void AControlRigGizmoActor::SetGizmoColor(const FLinearColor& InColor)
+void AControlRigShapeActor::SetShapeColor(const FLinearColor& InColor)
 {
 	if (StaticMeshComponent && !ColorParameterName.IsNone())
 	{
@@ -108,9 +108,9 @@ void AControlRigGizmoActor::SetGizmoColor(const FLinearColor& InColor)
 	}
 }
 
-// FControlRigGizmoHelper START
+// FControlRigShapeHelper START
 
-namespace FControlRigGizmoHelper
+namespace FControlRigShapeHelper
 {
 	FActorSpawnParameters GetDefaultSpawnParameter()
 	{
@@ -124,51 +124,51 @@ namespace FControlRigGizmoHelper
 		return ActorSpawnParameters;
 	}
 
-	// create gizmo from custom staticmesh, may deprecate this unless we come up with better usage
-	AControlRigGizmoActor* CreateGizmoActor(UWorld* InWorld, UStaticMesh* InStaticMesh, const FGizmoActorCreationParam& CreationParam)
+	// create shape from custom staticmesh, may deprecate this unless we come up with better usage
+	AControlRigShapeActor* CreateShapeActor(UWorld* InWorld, UStaticMesh* InStaticMesh, const FControlShapeActorCreationParam& CreationParam)
 	{
 		if (InWorld)
 		{
-			AControlRigGizmoActor* GizmoActor = CreateDefaultGizmoActor(InWorld, CreationParam);
+			AControlRigShapeActor* ShapeActor = CreateDefaultShapeActor(InWorld, CreationParam);
 
-			if (GizmoActor)
+			if (ShapeActor)
 			{
 				if (InStaticMesh)
 				{
-					GizmoActor->StaticMeshComponent->SetStaticMesh(InStaticMesh);
+					ShapeActor->StaticMeshComponent->SetStaticMesh(InStaticMesh);
 				}
 
-				return GizmoActor;
+				return ShapeActor;
 			}
 		}
 
 		return nullptr;
 	}
 
-	AControlRigGizmoActor* CreateGizmoActor(UWorld* InWorld, TSubclassOf<AControlRigGizmoActor> InClass, const FGizmoActorCreationParam& CreationParam)
+	AControlRigShapeActor* CreateShapeActor(UWorld* InWorld, TSubclassOf<AControlRigShapeActor> InClass, const FControlShapeActorCreationParam& CreationParam)
 	{
-		AControlRigGizmoActor* GizmoActor = InWorld->SpawnActor<AControlRigGizmoActor>(InClass, GetDefaultSpawnParameter());
-		if (GizmoActor)
+		AControlRigShapeActor* ShapeActor = InWorld->SpawnActor<AControlRigShapeActor>(InClass, GetDefaultSpawnParameter());
+		if (ShapeActor)
 		{
 			// set transform
-			GizmoActor->SetActorTransform(CreationParam.SpawnTransform);
-			return GizmoActor;
+			ShapeActor->SetActorTransform(CreationParam.SpawnTransform);
+			return ShapeActor;
 		}
 
 		return nullptr;
 	}
 
-	AControlRigGizmoActor* CreateDefaultGizmoActor(UWorld* InWorld, const FGizmoActorCreationParam& CreationParam)
+	AControlRigShapeActor* CreateDefaultShapeActor(UWorld* InWorld, const FControlShapeActorCreationParam& CreationParam)
 	{
-		AControlRigGizmoActor* GizmoActor = InWorld->SpawnActor<AControlRigGizmoActor>(AControlRigGizmoActor::StaticClass(), GetDefaultSpawnParameter());
-		if (GizmoActor)
+		AControlRigShapeActor* ShapeActor = InWorld->SpawnActor<AControlRigShapeActor>(AControlRigShapeActor::StaticClass(), GetDefaultSpawnParameter());
+		if (ShapeActor)
 		{
-			GizmoActor->ControlRigIndex = CreationParam.ControlRigIndex;
-			GizmoActor->ControlName = CreationParam.ControlName;
-			GizmoActor->SetSelectable(CreationParam.bSelectable);
-			GizmoActor->SetActorTransform(CreationParam.SpawnTransform);
+			ShapeActor->ControlRigIndex = CreationParam.ControlRigIndex;
+			ShapeActor->ControlName = CreationParam.ControlName;
+			ShapeActor->SetSelectable(CreationParam.bSelectable);
+			ShapeActor->SetActorTransform(CreationParam.SpawnTransform);
 
-			UStaticMeshComponent* MeshComponent = GizmoActor->StaticMeshComponent;
+			UStaticMeshComponent* MeshComponent = ShapeActor->StaticMeshComponent;
 
 			if (!CreationParam.StaticMesh.IsValid())
 			{
@@ -177,7 +177,7 @@ namespace FControlRigGizmoHelper
 			if (CreationParam.StaticMesh.IsValid())
 			{
 				MeshComponent->SetStaticMesh(CreationParam.StaticMesh.Get());
-				MeshComponent->SetRelativeTransform(CreationParam.MeshTransform * CreationParam.GizmoTransform);
+				MeshComponent->SetRelativeTransform(CreationParam.MeshTransform * CreationParam.ShapeTransform);
 			}
 
 			if (!CreationParam.Material.IsValid())
@@ -186,19 +186,19 @@ namespace FControlRigGizmoHelper
 			}
 			if (CreationParam.StaticMesh.IsValid())
 			{
-				GizmoActor->ColorParameterName = CreationParam.ColorParameterName;
-				UMaterialInstanceDynamic* MaterialInstance = UMaterialInstanceDynamic::Create(CreationParam.Material.Get(), GizmoActor);
+				ShapeActor->ColorParameterName = CreationParam.ColorParameterName;
+				UMaterialInstanceDynamic* MaterialInstance = UMaterialInstanceDynamic::Create(CreationParam.Material.Get(), ShapeActor);
 				MaterialInstance->SetVectorParameterValue(CreationParam.ColorParameterName, FVector(CreationParam.Color));
 				MeshComponent->SetMaterial(0, MaterialInstance);
 			}
-			return GizmoActor;
+			return ShapeActor;
 		}
 
 		return nullptr;
 	}
 }
 
-void AControlRigGizmoActor::SetGlobalTransform(const FTransform& InTransform)
+void AControlRigShapeActor::SetGlobalTransform(const FTransform& InTransform)
 {
 	if (RootComponent)
 	{
@@ -206,7 +206,7 @@ void AControlRigGizmoActor::SetGlobalTransform(const FTransform& InTransform)
 	}
 }
 
-FTransform AControlRigGizmoActor::GetGlobalTransform() const
+FTransform AControlRigShapeActor::GetGlobalTransform() const
 {
 	if (RootComponent)
 	{
@@ -215,4 +215,4 @@ FTransform AControlRigGizmoActor::GetGlobalTransform() const
 
 	return FTransform::Identity;
 }
-// FControlRigGizmoHelper END
\ No newline at end of file
+// FControlRigShapeHelper END
\ No newline at end of file
diff --git a/Engine/Plugins/Animation/ControlRig/Source/ControlRig/Private/ControlRigGizmoLibrary.cpp b/Engine/Plugins/Animation/ControlRig/Source/ControlRig/Private/ControlRigGizmoLibrary.cpp
--- a/Engine/Plugins/Animation/ControlRig/Source/ControlRig/Private/ControlRigGizmoLibrary.cpp
+++ b/Engine/Plugins/Animation/ControlRig/Source/ControlRig/Private/ControlRigGizmoLibrary.cpp
@@ -4,44 +4,44 @@
 
 #define LOCTEXT_NAMESPACE "ControlRigGizmoLibrary"
 
-UControlRigGizmoLibrary::UControlRigGizmoLibrary()
+UControlRigShapeLibrary::UControlRigShapeLibrary()
 {
 }
 
 #if WITH_EDITOR
 
 // UObject interface
-void UControlRigGizmoLibrary::PostEditChangeChainProperty(struct FPropertyChangedChainEvent& PropertyChangedEvent)
+void UControlRigShapeLibrary::PostEditChangeChainProperty(struct FPropertyChangedChainEvent& PropertyChangedEvent)
 {
-	if (PropertyChangedEvent.Property->GetName() == TEXT("GizmoName"))
+	if (PropertyChangedEvent.Property->GetName() == TEXT("ShapeName"))
 	{
 		FProperty* MemberProperty = PropertyChangedEvent.PropertyChain.GetHead()->GetValue();
-		if (MemberProperty->GetName() == TEXT("DefaultGizmo"))
+		if (MemberProperty->GetName() == TEXT("DefaultShape"))
 		{
-			DefaultGizmo.GizmoName = TEXT("Gizmo");
+			DefaultShape.ShapeName = FControlRigShapeDefinition().ShapeName;
 			GetUpdatedNameList(true);
 		}
-		else if(MemberProperty->GetName() == TEXT("Gizmos"))
+		else if(MemberProperty->GetName() == TEXT("Shapes"))
 		{
-			if (Gizmos.Num() == 0)
+			if (Shapes.Num() == 0)
 			{
 				return;
 			}
 
-			int32 GizmoIndexEdited = PropertyChangedEvent.GetArrayIndex(TEXT("Gizmos"));
-			if (Gizmos.IsValidIndex(GizmoIndexEdited))
+			int32 ShapeIndexEdited = PropertyChangedEvent.GetArrayIndex(TEXT("Shapes"));
+			if (Shapes.IsValidIndex(ShapeIndexEdited))
 			{
 				TArray<FName> Names;
-				Names.Add(DefaultGizmo.GizmoName);
-				for (int32 GizmoIndex = 0; GizmoIndex < Gizmos.Num(); GizmoIndex++)
+				Names.Add(DefaultShape.ShapeName);
+				for (int32 ShapeIndex = 0; ShapeIndex < Shapes.Num(); ShapeIndex++)
 				{
-					if (GizmoIndex != GizmoIndexEdited)
+					if (ShapeIndex != ShapeIndexEdited)
 					{
-						Names.Add(Gizmos[GizmoIndex].GizmoName);
+						Names.Add(Shapes[ShapeIndex].ShapeName);
 					}
 				}
 
-				FName DesiredName = Gizmos[GizmoIndexEdited].GizmoName;
+				FName DesiredName = Shapes[ShapeIndexEdited].ShapeName;
 				FString Name = DesiredName.ToString();
 				int32 Suffix = 0;
 				while (Names.Contains(*Name))
@@ -49,18 +49,18 @@ void UControlRigGizmoLibrary::PostEditChangeChainProperty(struct FPropertyChange
 					Suffix++;
 					Name = FString::Printf(TEXT("%s_%d"), *DesiredName.ToString(), Suffix);
 				}
-				Gizmos[GizmoIndexEdited].GizmoName = *Name;
+				Shapes[ShapeIndexEdited].ShapeName = *Name;
 			}
 			GetUpdatedNameList(true);
 		}
 	}
-	else if (PropertyChangedEvent.Property->GetName() == TEXT("Gizmos"))
+	else if (PropertyChangedEvent.Property->GetName() == TEXT("Shapes"))
 	{
 		TArray<FName> Names;
-		Names.Add(DefaultGizmo.GizmoName);
-		for (int32 GizmoIndex = 0; GizmoIndex < Gizmos.Num(); GizmoIndex++)
+		Names.Add(DefaultShape.ShapeName);
+		for (int32 ShapeIndex = 0; ShapeIndex < Shapes.Num(); ShapeIndex++)
 		{
-			FName DesiredName = Gizmos[GizmoIndex].GizmoName;
+			FName DesiredName = Shapes[ShapeIndex].ShapeName;
 			FString Name = DesiredName.ToString();
 			int32 Suffix = 0;
 			while (Names.Contains(*Name))
@@ -68,9 +68,9 @@ void UControlRigGizmoLibrary::PostEditChangeChainProperty(struct FPropertyChange
 				Suffix++;
 				Name = FString::Printf(TEXT("%s_%d"), *DesiredName.ToString(), Suffix);
 			}
-			Gizmos[GizmoIndex].GizmoName = *Name;
+			Shapes[ShapeIndex].ShapeName = *Name;
 
-			Names.Add(Gizmos[GizmoIndex].GizmoName);
+			Names.Add(Shapes[ShapeIndex].ShapeName);
 		}
 		GetUpdatedNameList(true);
 	}
@@ -78,42 +78,81 @@ void UControlRigGizmoLibrary::PostEditChangeChainProperty(struct FPropertyChange
 
 #endif
 
-const FControlRigGizmoDefinition* UControlRigGizmoLibrary::GetGizmoByName(const FName& InName, bool bUseDefaultIfNotFound) const
+const FControlRigShapeDefinition* UControlRigShapeLibrary::GetShapeByName(const FName& InName, bool bUseDefaultIfNotFound) const
 {
-	if (InName == DefaultGizmo.GizmoName)
+	UControlRigShapeLibrary* MutableThis = (UControlRigShapeLibrary*)this;
+	if (InName == MutableThis->DefaultShape.ShapeName)
 	{
-		return &DefaultGizmo;
+		MutableThis->DefaultShape.Library = MutableThis;
+		return &MutableThis->DefaultShape;
 	}
 
-	for (int32 GizmoIndex = 0; GizmoIndex < Gizmos.Num(); GizmoIndex++)
+	for (int32 ShapeIndex = 0; ShapeIndex < MutableThis->Shapes.Num(); ShapeIndex++)
 	{
-		if (Gizmos[GizmoIndex].GizmoName == InName)
+		if (MutableThis->Shapes[ShapeIndex].ShapeName == InName)
 		{
-			return &Gizmos[GizmoIndex];
+			MutableThis->Shapes[ShapeIndex].Library = MutableThis;
+			return &MutableThis->Shapes[ShapeIndex];
 		}
 	}
 
 	if (bUseDefaultIfNotFound)
 	{
-		return &DefaultGizmo;
+		MutableThis->DefaultShape.Library = MutableThis;
+		return &DefaultShape;
 	}
 
 	return nullptr;
 }
 
-const TArray<FName> UControlRigGizmoLibrary::GetUpdatedNameList(bool bReset)
+const FControlRigShapeDefinition* UControlRigShapeLibrary::GetShapeByName(const FName& InName, const TArray<TSoftObjectPtr<UControlRigShapeLibrary>>& InShapeLibraries)
+{
+	const FString InString = InName.ToString();
+	FString Left, Right;
+	FName RightName;
+	if(!InString.Split(TEXT("."), &Left, &Right))
+	{
+		Left = FString();
+		Right = InString;
+		RightName = InName;
+	}
+	else
+	{
+		RightName = *Right;
+	}
+
+	for(const TSoftObjectPtr<UControlRigShapeLibrary>& ShapeLibrary : InShapeLibraries)
+	{
+		if(!ShapeLibrary.IsValid())
+		{
+			continue;
+		}
+
+		if(ShapeLibrary->GetName().Equals(Left) || Left.IsEmpty())
+		{
+			if(const FControlRigShapeDefinition* Shape = ShapeLibrary->GetShapeByName(RightName))
+			{
+				return Shape;
+			}
+		}
+	}
+
+	return nullptr;
+}
+
+const TArray<FName> UControlRigShapeLibrary::GetUpdatedNameList(bool bReset)
 {
 	if (bReset)
 	{
 		NameList.Reset();
 	}
 
-	if (NameList.Num() != Gizmos.Num())
+	if (NameList.Num() != Shapes.Num())
 	{
 		NameList.Reset();
-		for (const FControlRigGizmoDefinition& Gizmo : Gizmos)
+		for (const FControlRigShapeDefinition& Shape : Shapes)
 		{
-			NameList.Add(Gizmo.GizmoName);
+			NameList.Add(Shape.ShapeName);
 		}
 		NameList.Sort(FNameLexicalLess());
 	}
diff --git a/Engine/Plugins/Animation/ControlRig/Source/ControlRig/Private/Rigs/RigHierarchy.cpp b/Engine/Plugins/Animation/ControlRig/Source/ControlRig/Private/Rigs/RigHierarchy.cpp
--- a/Engine/Plugins/Animation/ControlRig/Source/ControlRig/Private/Rigs/RigHierarchy.cpp
+++ b/Engine/Plugins/Animation/ControlRig/Source/ControlRig/Private/Rigs/RigHierarchy.cpp
@@ -587,7 +587,7 @@ void URigHierarchy::ResetPoseToInitial(ERigElementType InTypeFilter)
 		if(FRigControlElement* ControlElement = Cast<FRigControlElement>(Elements[ElementIndex]))
 		{
 			ControlElement->Offset.Current = ControlElement->Offset.Initial;
-			ControlElement->Gizmo.Current = ControlElement->Gizmo.Initial;
+			ControlElement->Shape.Current = ControlElement->Shape.Initial;
 		}
 
 		if(FRigTransformElement* TransformElement = Cast<FRigTransformElement>(Elements[ElementIndex]))
@@ -2085,7 +2085,7 @@ void URigHierarchy::SetTransform(FRigTransformElement* InTransformElement, const
 
 	if(FRigControlElement* ControlElement = Cast<FRigControlElement>(InTransformElement))
 	{
-		ControlElement->Gizmo.MarkDirty(MakeGlobal(InTransformType));
+		ControlElement->Shape.MarkDirty(MakeGlobal(InTransformType));
 	}
 
 	EnsureCacheValidity();
@@ -2262,7 +2262,7 @@ void URigHierarchy::SetControlOffsetTransform(FRigControlElement* InControlEleme
 	const ERigTransformType::Type OpposedType = SwapLocalAndGlobal(InTransformType);
 	InControlElement->Offset.Set(InTransformType, InTransform);
 	InControlElement->Offset.MarkDirty(OpposedType);
-	InControlElement->Gizmo.MarkDirty(MakeGlobal(InTransformType));
+	InControlElement->Shape.MarkDirty(MakeGlobal(InTransformType));
 
 	EnsureCacheValidity();
 
@@ -2333,7 +2333,7 @@ void URigHierarchy::SetControlOffsetTransform(FRigControlElement* InControlEleme
 #endif
 }
 
-FTransform URigHierarchy::GetControlGizmoTransform(FRigControlElement* InControlElement,
+FTransform URigHierarchy::GetControlShapeTransform(FRigControlElement* InControlElement,
 	const ERigTransformType::Type InTransformType) const
 {
 	if(InControlElement == nullptr)
@@ -2341,28 +2341,28 @@ FTransform URigHierarchy::GetControlGizmoTransform(FRigControlElement* InControl
 		return FTransform::Identity;
 	}
 	
-	if(InControlElement->Gizmo.IsDirty(InTransformType))
+	if(InControlElement->Shape.IsDirty(InTransformType))
 	{
 		const ERigTransformType::Type OpposedType = SwapLocalAndGlobal(InTransformType);
 		const ERigTransformType::Type GlobalType = MakeGlobal(InTransformType);
-		ensure(!InControlElement->Gizmo.IsDirty(OpposedType));
+		ensure(!InControlElement->Shape.IsDirty(OpposedType));
 
 		const FTransform ParentTransform = GetTransform(InControlElement, GlobalType);
 		if(IsLocal(InTransformType))
 		{
-			InControlElement->Gizmo.Set(InTransformType, InControlElement->Gizmo.Get(OpposedType).GetRelativeTransform(ParentTransform));
+			InControlElement->Shape.Set(InTransformType, InControlElement->Shape.Get(OpposedType).GetRelativeTransform(ParentTransform));
 		}
 		else
 		{
-			InControlElement->Gizmo.Set(InTransformType, InControlElement->Gizmo.Get(OpposedType) * ParentTransform);
+			InControlElement->Shape.Set(InTransformType, InControlElement->Shape.Get(OpposedType) * ParentTransform);
 		}
 
 		EnsureCacheValidity();
 	}
-	return InControlElement->Gizmo.Get(InTransformType);
+	return InControlElement->Shape.Get(InTransformType);
 }
 
-void URigHierarchy::SetControlGizmoTransform(FRigControlElement* InControlElement, const FTransform& InTransform,
+void URigHierarchy::SetControlShapeTransform(FRigControlElement* InControlElement, const FTransform& InTransform,
 	const ERigTransformType::Type InTransformType, bool bSetupUndo, bool bForce, bool bPrintPythonCommands)
 {
 	if(InControlElement == nullptr)
@@ -2370,26 +2370,26 @@ void URigHierarchy::SetControlGizmoTransform(FRigControlElement* InControlElemen
 		return;
 	}
 
-	if(!InControlElement->Gizmo.IsDirty(InTransformType))
+	if(!InControlElement->Shape.IsDirty(InTransformType))
 	{
-		const FTransform PreviousTransform = InControlElement->Gizmo.Get(InTransformType);
+		const FTransform PreviousTransform = InControlElement->Shape.Get(InTransformType);
 		if(!bForce && FRigComputedTransform::Equals(PreviousTransform, InTransform))
 		{
 			return;
 		}
 	}
 
-	const FTransform PreviousTransform = GetControlGizmoTransform(InControlElement, InTransformType);
+	const FTransform PreviousTransform = GetControlShapeTransform(InControlElement, InTransformType);
 	const ERigTransformType::Type OpposedType = SwapLocalAndGlobal(InTransformType);
-	InControlElement->Gizmo.Set(InTransformType, InTransform);
-	InControlElement->Gizmo.MarkDirty(OpposedType);
+	InControlElement->Shape.Set(InTransformType, InTransform);
+	InControlElement->Shape.MarkDirty(OpposedType);
 
 	if (IsInitial(InTransformType))
 	{
-		// control's gizmo transform, similar to offset transform, is considered a special type of transform
+		// control's shape transform, similar to offset transform, is considered a special type of transform
 		// whenever its initial value is changed, we want to make sure the current is kept in sync
 		// such that the viewport can reflect this change
-		SetControlGizmoTransform(InControlElement, InTransform, ERigTransformType::MakeCurrent(InTransformType), false, bForce);
+		SetControlShapeTransform(InControlElement, InTransform, ERigTransformType::MakeCurrent(InTransformType), false, bForce);
 	}
 	
 	EnsureCacheValidity();
@@ -2399,10 +2399,10 @@ void URigHierarchy::SetControlGizmoTransform(FRigControlElement* InControlElemen
 	{
 		PushTransformToStack(
             InControlElement->GetKey(),
-            ERigTransformStackEntryType::ControlGizmo,
+            ERigTransformStackEntryType::ControlShape,
             InTransformType,
             PreviousTransform,
-            InControlElement->Gizmo.Get(InTransformType),
+            InControlElement->Shape.Get(InTransformType),
             false,
             bSetupUndo);
 	}
@@ -2410,7 +2410,7 @@ void URigHierarchy::SetControlGizmoTransform(FRigControlElement* InControlElemen
 
 	if(IsLocal(InTransformType))
 	{
-		Notify(ERigHierarchyNotification::ControlGizmoTransformChanged, InControlElement);
+		Notify(ERigHierarchyNotification::ControlShapeTransformChanged, InControlElement);
 	}
 
 #if WITH_EDITOR
@@ -2426,7 +2426,7 @@ void URigHierarchy::SetControlGizmoTransform(FRigControlElement* InControlElemen
 				if(FRigControlElement* ListeningElement = Cast<FRigControlElement>(ListeningHierarchy->Find(InControlElement->GetKey())))
 				{
 					// bSetupUndo = false such that all listening hierarchies performs undo at the same time the root hierachy undos
-					ListeningHierarchy->SetControlGizmoTransform(ListeningElement, InTransform, InTransformType, false, bForce);
+					ListeningHierarchy->SetControlShapeTransform(ListeningElement, InTransform, InTransformType, false, bForce);
 				}
 			}
 		}
@@ -2449,7 +2449,7 @@ void URigHierarchy::SetControlGizmoTransform(FRigControlElement* InControlElemen
 		if (!BlueprintName.IsEmpty())
 		{
 			RigVMPythonUtils::Print(BlueprintName,
-				FString::Printf(TEXT("hierarchy.set_control_gizmo_transform(%s, %s, %s)"),
+				FString::Printf(TEXT("hierarchy.set_control_shape_transform(%s, %s, %s)"),
 				*InControlElement->GetKey().ToPythonString(),
 				*RigVMPythonUtils::TransformToPythonString(InTransform),
 				ERigTransformType::IsInitial(InTransformType) ? TEXT("True") : TEXT("False")));
@@ -2724,7 +2724,7 @@ void URigHierarchy::SetControlVisibility(FRigControlElement* InControlElement, b
 		return;
 	}
 
-	InControlElement->Settings.bGizmoVisible = bVisibility;
+	InControlElement->Settings.bShapeVisible = bVisibility;
 	Notify(ERigHierarchyNotification::ControlVisibilityChanged, InControlElement);
 
 #if WITH_EDITOR
@@ -3475,7 +3475,7 @@ void URigHierarchy::PropagateDirtyFlags(FRigTransformElement* InTransformElement
 			if(FRigControlElement* ControlElement = Cast<FRigControlElement>(ElementToDirty.Element))
 			{
 				ControlElement->Offset.MarkDirty(GlobalType);
-				ControlElement->Gizmo.MarkDirty(GlobalType);
+				ControlElement->Shape.MarkDirty(GlobalType);
 			}
 
 #if URIGHIERARCHY_RECURSIVE_DIRTY_PROPAGATION
@@ -3557,9 +3557,9 @@ void URigHierarchy::EnsureCacheValidityImpl()
 
                         if(ControlElementToDirty->Pose.IsDirty(GlobalType))
                         {
-                            checkf(ControlElementToDirty->Gizmo.IsDirty(GlobalType) ||
-                                    ControlElementToDirty->Gizmo.IsDirty(LocalType),
-                                    TEXT("Control '%s' %s Pose Cache is dirty, but the Gizmo is not."),
+                            checkf(ControlElementToDirty->Shape.IsDirty(GlobalType) ||
+                                    ControlElementToDirty->Shape.IsDirty(LocalType),
+                                    TEXT("Control '%s' %s Pose Cache is dirty, but the Shape is not."),
 									*ControlElementToDirty->GetKey().ToString(),
 									*TransformTypeString);
 						}
@@ -3684,13 +3684,13 @@ void URigHierarchy::EnsureCacheValidityImpl()
 				const ERigTransformType::Type OpposedType = ERigTransformType::SwapLocalAndGlobal(TransformType);
 				const FString& TransformTypeString = TransformTypeStrings[TransformTypeIndex];
 
-				if(!ControlElement->Gizmo.IsDirty(TransformType) && !ControlElement->Gizmo.IsDirty(OpposedType))
+				if(!ControlElement->Shape.IsDirty(TransformType) && !ControlElement->Shape.IsDirty(OpposedType))
 				{
-					const FTransform CachedTransform = HierarchyForLambda->GetControlGizmoTransform(ControlElement, TransformType);
-					ControlElement->Gizmo.MarkDirty(TransformType);
-					const FTransform ComputedTransform = HierarchyForLambda->GetControlGizmoTransform(ControlElement, TransformType);
+					const FTransform CachedTransform = HierarchyForLambda->GetControlShapeTransform(ControlElement, TransformType);
+					ControlElement->Shape.MarkDirty(TransformType);
+					const FTransform ComputedTransform = HierarchyForLambda->GetControlShapeTransform(ControlElement, TransformType);
 					checkf(FRigComputedTransform::Equals(CachedTransform, ComputedTransform),
-						TEXT("Element '%s' Gizmo %s Cached vs Computed doesn't match."),
+						TEXT("Element '%s' Shape %s Cached vs Computed doesn't match."),
 						*Element->GetName().ToString(),
 						*TransformTypeString);
 				}
@@ -3841,7 +3841,7 @@ void URigHierarchy::PushTransformToStack(const FRigElementKey& InKey, ERigTransf
 
 	static const FText TransformPoseTitle = NSLOCTEXT("RigHierarchy", "Set Pose Transform", "Set Pose Transform");
 	static const FText ControlOffsetTitle = NSLOCTEXT("RigHierarchy", "Set Control Offset", "Set Control Offset");
-	static const FText ControlGizmoTitle = NSLOCTEXT("RigHierarchy", "Set Control Gizo", "Set Control Gizo");
+	static const FText ControlShapeTitle = NSLOCTEXT("RigHierarchy", "Set Control Gizo", "Set Control Gizo");
 	static const FText CurveValueTitle = NSLOCTEXT("RigHierarchy", "Set Curve Value", "Set Curve Value");
 	
 	FText Title;
@@ -3857,7 +3857,7 @@ void URigHierarchy::PushTransformToStack(const FRigElementKey& InKey, ERigTransf
 			Title = TransformPoseTitle;
 			break;
 		}
-		case ERigTransformStackEntryType::ControlGizmo:
+		case ERigTransformStackEntryType::ControlShape:
 		{
 			Title = TransformPoseTitle;
 			break;
@@ -4004,9 +4004,9 @@ bool URigHierarchy::ApplyTransformFromStack(const FRigTransformStackEntry& InEnt
 			SetControlOffsetTransform(Cast<FRigControlElement>(Element), Transform, InEntry.TransformType, InEntry.bAffectChildren, false); 
 			break;
 		}
-		case ERigTransformStackEntryType::ControlGizmo:
+		case ERigTransformStackEntryType::ControlShape:
 		{
-			SetControlGizmoTransform(Cast<FRigControlElement>(Element), Transform, InEntry.TransformType, false); 
+			SetControlShapeTransform(Cast<FRigControlElement>(Element), Transform, InEntry.TransformType, false); 
 			break;
 		}
 		case ERigTransformStackEntryType::CurveValue:
@@ -4038,7 +4038,7 @@ void URigHierarchy::ComputeAllTransforms()
 			if(FRigControlElement* ControlElement = Get<FRigControlElement>(ElementIndex))
 			{
 				GetControlOffsetTransform(ControlElement, TransformType);
-				GetControlGizmoTransform(ControlElement, TransformType);
+				GetControlShapeTransform(ControlElement, TransformType);
 			}
 		}
 	}
@@ -4700,18 +4700,18 @@ TArray<FString> URigHierarchy::ControlSettingsToPythonCommands(const FRigControl
 	Commands.Add(FString::Printf(TEXT("%s.draw_limits = %s"),
 		*NameSettings,
 		Settings.bDrawLimits ? TEXT("True") : TEXT("False")));
-	Commands.Add(FString::Printf(TEXT("%s.gizmo_color = %s"),
+	Commands.Add(FString::Printf(TEXT("%s.shape_color = %s"),
 		*NameSettings,
-		*RigVMPythonUtils::LinearColorToPythonString(Settings.GizmoColor)));
-	Commands.Add(FString::Printf(TEXT("%s.gizmo_enabled = %s"),
+		*RigVMPythonUtils::LinearColorToPythonString(Settings.ShapeColor)));
+	Commands.Add(FString::Printf(TEXT("%s.shape_enabled = %s"),
 		*NameSettings,
-		Settings.bGizmoEnabled ? TEXT("True") : TEXT("False")));
-	Commands.Add(FString::Printf(TEXT("%s.gizmo_name = '%s'"),
+		Settings.bShapeEnabled ? TEXT("True") : TEXT("False")));
+	Commands.Add(FString::Printf(TEXT("%s.shape_name = '%s'"),
 		*NameSettings,
-		*Settings.GizmoName.ToString()));
-	Commands.Add(FString::Printf(TEXT("%s.gizmo_visible = %s"),
+		*Settings.ShapeName.ToString()));
+	Commands.Add(FString::Printf(TEXT("%s.shape_visible = %s"),
 		*NameSettings,
-		Settings.bGizmoVisible ? TEXT("True") : TEXT("False")));
+		Settings.bShapeVisible ? TEXT("True") : TEXT("False")));
 	Commands.Add(FString::Printf(TEXT("%s.is_transient_control = %s"),
 		*NameSettings,
 		Settings.bIsTransientControl ? TEXT("True") : TEXT("False")));
diff --git a/Engine/Plugins/Animation/ControlRig/Source/ControlRig/Private/Rigs/RigHierarchyController.cpp b/Engine/Plugins/Animation/ControlRig/Source/ControlRig/Private/Rigs/RigHierarchyController.cpp
--- a/Engine/Plugins/Animation/ControlRig/Source/ControlRig/Private/Rigs/RigHierarchyController.cpp
+++ b/Engine/Plugins/Animation/ControlRig/Source/ControlRig/Private/Rigs/RigHierarchyController.cpp
@@ -292,7 +292,7 @@ FRigElementKey URigHierarchyController::AddControl(
 	FRigControlSettings InSettings,
 	FRigControlValue InValue,
 	FTransform InOffsetTransform,
-	FTransform InGizmoTransform,
+	FTransform InShapeTransform,
 	bool bSetupUndo, bool bPrintPythonCommand)
 {
 	if(!IsValid())
@@ -318,17 +318,17 @@ FRigElementKey URigHierarchyController::AddControl(
 		AddElement(NewElement, Hierarchy->Get(Hierarchy->GetIndex(InParent)), false);
 		
 		NewElement->Offset.Set(ERigTransformType::InitialLocal, InOffsetTransform);  
-		NewElement->Gizmo.Set(ERigTransformType::InitialLocal, InGizmoTransform);  
+		NewElement->Shape.Set(ERigTransformType::InitialLocal, InShapeTransform);  
 		Hierarchy->SetControlValue(NewElement, InValue, ERigControlValueType::Initial, false);
 
 		NewElement->Parent.MarkDirty(ERigTransformType::InitialGlobal);
 		NewElement->Offset.MarkDirty(ERigTransformType::InitialGlobal);
 		NewElement->Pose.MarkDirty(ERigTransformType::InitialGlobal);
-		NewElement->Gizmo.MarkDirty(ERigTransformType::InitialGlobal);
+		NewElement->Shape.MarkDirty(ERigTransformType::InitialGlobal);
 		NewElement->Parent.Current = NewElement->Parent.Initial;
 		NewElement->Offset.Current = NewElement->Offset.Initial;
 		NewElement->Pose.Current = NewElement->Pose.Initial;
-		NewElement->Gizmo.Current = NewElement->Gizmo.Initial;
+		NewElement->Shape.Current = NewElement->Shape.Initial;
 	}
 
 #if WITH_EDITOR
@@ -1245,10 +1245,10 @@ TArray<FRigElementKey> URigHierarchyController::ImportFromHierarchyContainer(con
 		Settings.bDrawLimits = Control.bDrawLimits;
 		Settings.MinimumValue = Control.MinimumValue;
 		Settings.MaximumValue = Control.MaximumValue;
-		Settings.bGizmoEnabled = Control.bGizmoEnabled;
-		Settings.bGizmoVisible = Control.bGizmoVisible;
-		Settings.GizmoName = Control.GizmoName;
-		Settings.GizmoColor = Control.GizmoColor;
+		Settings.bShapeEnabled = Control.bGizmoEnabled;
+		Settings.bShapeVisible = Control.bGizmoVisible;
+		Settings.ShapeName = Control.GizmoName;
+		Settings.ShapeColor = Control.GizmoColor;
 		Settings.ControlEnum = Control.ControlEnum;
 
 		FRigControlValue InitialValue = Control.InitialValue;
@@ -1449,7 +1449,7 @@ TArray<FString> URigHierarchyController::GetAddControlPythonCommands(FRigControl
 	Commands.Append(GetSetControlValuePythonCommands(Control, Settings.MinimumValue, ERigControlValueType::Minimum));
 	Commands.Append(GetSetControlValuePythonCommands(Control, Settings.MaximumValue, ERigControlValueType::Maximum));
 	Commands.Append(GetSetControlOffsetTransformPythonCommands(Control, Control->Offset.Initial.Local.Transform, true, true));
-	Commands.Append(GetSetControlGizmoTransformPythonCommands(Control, Control->Gizmo.Initial.Local.Transform, true));
+	Commands.Append(GetSetControlShapeTransformPythonCommands(Control, Control->Shape.Initial.Local.Transform, true));
 
 	return Commands;
 }
@@ -1527,11 +1527,11 @@ TArray<FString> URigHierarchyController::GetSetControlOffsetTransformPythonComma
 		bAffectChildren ? TEXT("True") : TEXT("False"))};
 }
 
-TArray<FString> URigHierarchyController::GetSetControlGizmoTransformPythonCommands(const FRigControlElement* Control,
+TArray<FString> URigHierarchyController::GetSetControlShapeTransformPythonCommands(const FRigControlElement* Control,
 	const FTransform& InTransform, bool bInitial) const
 {
-	//SetControlGizmoTransform(FRigElementKey InKey, FTransform InTransform, bool bInitial = false, bool bSetupUndo = false)
-	return {FString::Printf(TEXT("hierarchy.set_control_gizmo_transform(%s, %s, %s)"),
+	//SetControlShapeTransform(FRigElementKey InKey, FTransform InTransform, bool bInitial = false, bool bSetupUndo = false)
+	return {FString::Printf(TEXT("hierarchy.set_control_shape_transform(%s, %s, %s)"),
 		*Control->GetKey().ToPythonString(),
 		*RigVMPythonUtils::TransformToPythonString(InTransform),
 		bInitial ? TEXT("True") : TEXT("False"))};
@@ -1973,8 +1973,8 @@ bool URigHierarchyController::AddParent(FRigBaseElement* InChild, FRigBaseElemen
 		{
 			Hierarchy->GetControlOffsetTransform(ControlElement, ERigTransformType::CurrentLocal);
 			Hierarchy->GetControlOffsetTransform(ControlElement, ERigTransformType::InitialLocal);
-			Hierarchy->GetControlGizmoTransform(ControlElement, ERigTransformType::CurrentLocal);
-			Hierarchy->GetControlGizmoTransform(ControlElement, ERigTransformType::InitialLocal);
+			Hierarchy->GetControlShapeTransform(ControlElement, ERigTransformType::CurrentLocal);
+			Hierarchy->GetControlShapeTransform(ControlElement, ERigTransformType::InitialLocal);
 		}
 	}
 
@@ -2022,8 +2022,8 @@ bool URigHierarchyController::AddParent(FRigBaseElement* InChild, FRigBaseElemen
 			{
 				ControlElement->Offset.MarkDirty(ERigTransformType::CurrentGlobal);  
 				ControlElement->Offset.MarkDirty(ERigTransformType::InitialGlobal);
-				ControlElement->Gizmo.MarkDirty(ERigTransformType::CurrentGlobal);  
-				ControlElement->Gizmo.MarkDirty(ERigTransformType::InitialGlobal);
+				ControlElement->Shape.MarkDirty(ERigTransformType::CurrentGlobal);  
+				ControlElement->Shape.MarkDirty(ERigTransformType::InitialGlobal);
 			}
 		}
 
@@ -2212,8 +2212,8 @@ bool URigHierarchyController::RemoveParent(FRigBaseElement* InChild, FRigBaseEle
 			{
 				ControlElement->Offset.MarkDirty(ERigTransformType::CurrentGlobal);  
 				ControlElement->Offset.MarkDirty(ERigTransformType::InitialGlobal);
-				ControlElement->Gizmo.MarkDirty(ERigTransformType::CurrentGlobal);  
-				ControlElement->Gizmo.MarkDirty(ERigTransformType::InitialGlobal);
+				ControlElement->Shape.MarkDirty(ERigTransformType::CurrentGlobal);  
+				ControlElement->Shape.MarkDirty(ERigTransformType::InitialGlobal);
 			}
 
 			Hierarchy->TopologyVersion++;
@@ -2422,7 +2422,7 @@ TArray<FRigElementKey> URigHierarchyController::MirrorElements(TArray<FRigElemen
 		FTransform GlobalTransform = Hierarchy->GetGlobalTransform(OriginalKeys[Index]);
 		FTransform InitialTransform = Hierarchy->GetInitialGlobalTransform(OriginalKeys[Index]);
 
-		// also mirror the offset, limits and gizmo transform
+		// also mirror the offset, limits and shape transform
 		if (OriginalKeys[Index].Type == ERigElementType::Control)
 		{
 			if(FRigControlElement* DuplicatedControlElement = Hierarchy->Find<FRigControlElement>(DuplicatedKeys[Index]))
@@ -2459,10 +2459,10 @@ TArray<FRigElementKey> URigHierarchyController::MirrorElements(TArray<FRigElemen
 				Hierarchy->SetGlobalTransform(DuplicatedKeys[Index], InSettings.MirrorTransform(GlobalTransform), true, false, true);
 				Hierarchy->SetGlobalTransform(DuplicatedKeys[Index], InSettings.MirrorTransform(GlobalTransform), false, false, true);
 
-				// mirror gizmo transform
-				FTransform GlobalGizmoTransform = Hierarchy->GetControlGizmoTransform(DuplicatedControlElement, ERigTransformType::InitialLocal) * OriginalGlobalOffsetTransform;
-				Hierarchy->SetControlGizmoTransform(DuplicatedControlElement, InSettings.MirrorTransform(GlobalGizmoTransform).GetRelativeTransform(DuplicatedGlobalOffsetTransform), ERigTransformType::InitialLocal, true);
-				Hierarchy->SetControlGizmoTransform(DuplicatedControlElement, InSettings.MirrorTransform(GlobalGizmoTransform).GetRelativeTransform(DuplicatedGlobalOffsetTransform), ERigTransformType::CurrentLocal, true);
+				// mirror shape transform
+				FTransform GlobalShapeTransform = Hierarchy->GetControlShapeTransform(DuplicatedControlElement, ERigTransformType::InitialLocal) * OriginalGlobalOffsetTransform;
+				Hierarchy->SetControlShapeTransform(DuplicatedControlElement, InSettings.MirrorTransform(GlobalShapeTransform).GetRelativeTransform(DuplicatedGlobalOffsetTransform), ERigTransformType::InitialLocal, true);
+				Hierarchy->SetControlShapeTransform(DuplicatedControlElement, InSettings.MirrorTransform(GlobalShapeTransform).GetRelativeTransform(DuplicatedGlobalOffsetTransform), ERigTransformType::CurrentLocal, true);
 			}
 		}
 		else
diff --git a/Engine/Plugins/Animation/ControlRig/Source/ControlRig/Private/Rigs/RigHierarchyElements.cpp b/Engine/Plugins/Animation/ControlRig/Source/ControlRig/Private/Rigs/RigHierarchyElements.cpp
--- a/Engine/Plugins/Animation/ControlRig/Source/ControlRig/Private/Rigs/RigHierarchyElements.cpp
+++ b/Engine/Plugins/Animation/ControlRig/Source/ControlRig/Private/Rigs/RigHierarchyElements.cpp
@@ -416,10 +416,10 @@ FRigControlSettings::FRigControlSettings()
 , bDrawLimits(true)
 , MinimumValue()
 , MaximumValue()
-, bGizmoEnabled(true)
-, bGizmoVisible(true)
-, GizmoName(TEXT("Gizmo"))
-, GizmoColor(FLinearColor::Red)
+, bShapeEnabled(true)
+, bShapeVisible(true)
+, ShapeName(NAME_None)
+, ShapeColor(FLinearColor::Red)
 , bIsTransientControl(false)
 , ControlEnum(nullptr)
 , Customization()
@@ -453,10 +453,10 @@ void FRigControlSettings::Save(FArchive& Ar)
 	Ar << bDrawLimits;
 	Ar << MinimumValue;
 	Ar << MaximumValue;
-	Ar << bGizmoEnabled;
-	Ar << bGizmoVisible;
-	Ar << GizmoName;
-	Ar << GizmoColor;
+	Ar << bShapeEnabled;
+	Ar << bShapeVisible;
+	Ar << ShapeName;
+	Ar << ShapeColor;
 	Ar << bIsTransientControl;
 	Ar << ControlEnumPathName;
 	Ar << Customization.AvailableSpaces;
@@ -494,10 +494,10 @@ void FRigControlSettings::Load(FArchive& Ar)
 		Ar << MaximumTransform;
 	}
 
-	Ar << bGizmoEnabled;
-	Ar << bGizmoVisible;
-	Ar << GizmoName;
-	Ar << GizmoColor;
+	Ar << bShapeEnabled;
+	Ar << bShapeVisible;
+	Ar << ShapeName;
+	Ar << ShapeColor;
 	Ar << bIsTransientControl;
 	Ar << ControlEnumPathName;
 
@@ -568,15 +568,15 @@ bool FRigControlSettings::operator==(const FRigControlSettings& InOther) const
 	{
 		return false;
 	}
-	if(bGizmoEnabled != InOther.bGizmoEnabled)
+	if(bShapeEnabled != InOther.bShapeEnabled)
 	{
 		return false;
 	}
-	if(bGizmoVisible != InOther.bGizmoVisible)
+	if(bShapeVisible != InOther.bShapeVisible)
 	{
 		return false;
 	}
-	if(GizmoName != InOther.GizmoName)
+	if(ShapeName != InOther.ShapeName)
 	{
 		return false;
 	}
@@ -588,7 +588,7 @@ bool FRigControlSettings::operator==(const FRigControlSettings& InOther) const
 	{
 		return false;
 	}
-	if(!GizmoColor.Equals(InOther.GizmoColor, 0.001))
+	if(!ShapeColor.Equals(InOther.ShapeColor, 0.001))
 	{
 		return false;
 	}
@@ -622,7 +622,7 @@ void FRigControlElement::Save(FArchive& Ar, URigHierarchy* Hierarchy, ESerializa
 	{
 		Settings.Save(Ar);
 		Offset.Save(Ar);
-		Gizmo.Save(Ar);
+		Shape.Save(Ar);
 	}
 }
 
@@ -634,7 +634,7 @@ void FRigControlElement::Load(FArchive& Ar, URigHierarchy* Hierarchy, ESerializa
 	{
 		Settings.Load(Ar);
 		Offset.Load(Ar);
-		Gizmo.Load(Ar);
+		Shape.Load(Ar);
 	}
 }
 
@@ -645,7 +645,7 @@ void FRigControlElement::CopyFrom(URigHierarchy* InHierarchy, FRigBaseElement* I
 	const FRigControlElement* Source = CastChecked<FRigControlElement>(InOther);
 	Settings = Source->Settings;
 	Offset = Source->Offset;
-	Gizmo = Source->Gizmo;
+	Shape = Source->Shape;
 }
 
 void FRigControlElement::CopyPose(FRigBaseElement* InOther, bool bCurrent, bool bInitial)
@@ -657,12 +657,12 @@ void FRigControlElement::CopyPose(FRigBaseElement* InOther, bool bCurrent, bool
 		if(bCurrent)
 		{
 			Offset.Current = Other->Offset.Current;
-			Gizmo.Current = Other->Gizmo.Current;
+			Shape.Current = Other->Shape.Current;
 		}
 		if(bInitial)
 		{
 			Offset.Initial = Other->Offset.Initial;
-			Gizmo.Initial = Other->Gizmo.Initial;
+			Shape.Initial = Other->Shape.Initial;
 		}
 	}
 }
diff --git a/Engine/Plugins/Animation/ControlRig/Source/ControlRig/Private/Settings/ControlRigSettings.cpp b/Engine/Plugins/Animation/ControlRig/Source/ControlRig/Private/Settings/ControlRigSettings.cpp
--- a/Engine/Plugins/Animation/ControlRig/Source/ControlRig/Private/Settings/ControlRigSettings.cpp
+++ b/Engine/Plugins/Animation/ControlRig/Source/ControlRig/Private/Settings/ControlRigSettings.cpp
@@ -11,7 +11,7 @@ UControlRigSettings::UControlRigSettings(const FObjectInitializer& ObjectInitial
 	: Super(ObjectInitializer)
 {
 #if WITH_EDITORONLY_DATA
-	DefaultGizmoLibrary = LoadObject<UControlRigGizmoLibrary>(nullptr, TEXT("/ControlRig/Controls/DefaultGizmoLibrary.DefaultGizmoLibrary"));
+	DefaultShapeLibrary = LoadObject<UControlRigShapeLibrary>(nullptr, TEXT("/ControlRig/Controls/DefaultGizmoLibraryNormalized.DefaultGizmoLibraryNormalized"));
 #endif
 }
 
diff --git a/Engine/Plugins/Animation/ControlRig/Source/ControlRig/Private/Units/Hierarchy/RigUnit_SetControlColor.cpp b/Engine/Plugins/Animation/ControlRig/Source/ControlRig/Private/Units/Hierarchy/RigUnit_SetControlColor.cpp
--- a/Engine/Plugins/Animation/ControlRig/Source/ControlRig/Private/Units/Hierarchy/RigUnit_SetControlColor.cpp
+++ b/Engine/Plugins/Animation/ControlRig/Source/ControlRig/Private/Units/Hierarchy/RigUnit_SetControlColor.cpp
@@ -21,7 +21,7 @@ FRigUnit_SetControlColor_Execute()
 				if (CachedControlIndex.UpdateCache(FRigElementKey(Control, ERigElementType::Control), Hierarchy))
 				{
 					FRigControlElement* ControlElement = Hierarchy->GetChecked<FRigControlElement>(CachedControlIndex);
-					ControlElement->Settings.GizmoColor = Color;
+					ControlElement->Settings.ShapeColor = Color;
 					Hierarchy->Notify(ERigHierarchyNotification::ControlSettingChanged, ControlElement);
 				}
 				break;
diff --git a/Engine/Plugins/Animation/ControlRig/Source/ControlRig/Public/ControlRig.h b/Engine/Plugins/Animation/ControlRig/Source/ControlRig/Public/ControlRig.h
--- a/Engine/Plugins/Animation/ControlRig/Source/ControlRig/Public/ControlRig.h
+++ b/Engine/Plugins/Animation/ControlRig/Source/ControlRig/Public/ControlRig.h
@@ -333,7 +333,7 @@ public:
 	virtual void SetControlLocalTransform(const FName& InControlName, const FTransform& InLocalTransform, bool bNotify = true, const FRigControlModifiedContext& Context = FRigControlModifiedContext(), bool bSetupUndo = true);
 	virtual FTransform GetControlLocalTransform(const FName& InControlName) ;
 
-	virtual UControlRigGizmoLibrary* GetGizmoLibrary() const;
+	virtual const TArray<TSoftObjectPtr<UControlRigShapeLibrary>>& GetShapeLibraries() const;
 	virtual void CreateRigControlsForCurveContainer();
 	virtual void GetControlsInOrder(TArray<FRigControlElement*>& SortedControls) const;
 
@@ -438,7 +438,10 @@ private:
 	TObjectPtr<URigHierarchy> DynamicHierarchy;
 
 	UPROPERTY()
-	TSoftObjectPtr<UControlRigGizmoLibrary> GizmoLibrary;
+	TSoftObjectPtr<UControlRigShapeLibrary> GizmoLibrary_DEPRECATED;
+
+	UPROPERTY()
+	TArray<TSoftObjectPtr<UControlRigShapeLibrary>> ShapeLibraries;
 
 	/** Runtime object binding */
 	TSharedPtr<IControlRigObjectBinding> ObjectBinding;
@@ -771,7 +774,7 @@ private:
 			FRigControlSettings Settings;
 			FRigControlValue Value;
 			FTransform OffsetTransform;
-			FTransform GizmoTransform;
+			FTransform ShapeTransform;
 		};
 		
 		TArray<FTransientControlInfo> SavedTransientControls;
diff --git a/Engine/Plugins/Animation/ControlRig/Source/ControlRig/Public/ControlRigControlActor.h b/Engine/Plugins/Animation/ControlRig/Source/ControlRig/Public/ControlRigControlActor.h
--- a/Engine/Plugins/Animation/ControlRig/Source/ControlRig/Public/ControlRigControlActor.h
+++ b/Engine/Plugins/Animation/ControlRig/Source/ControlRig/Public/ControlRigControlActor.h
@@ -74,7 +74,7 @@ private:
 	TArray<FName> ControlNames;
 
 	UPROPERTY(transient)
-	TArray<FTransform> GizmoTransforms;
+	TArray<FTransform> ShapeTransforms;
 
 	UPROPERTY(transient)
 	TArray<TObjectPtr<UStaticMeshComponent>> Components;
diff --git a/Engine/Plugins/Animation/ControlRig/Source/ControlRig/Public/ControlRigGizmoActor.h b/Engine/Plugins/Animation/ControlRig/Source/ControlRig/Public/ControlRigGizmoActor.h
--- a/Engine/Plugins/Animation/ControlRig/Source/ControlRig/Public/ControlRigGizmoActor.h
+++ b/Engine/Plugins/Animation/ControlRig/Source/ControlRig/Public/ControlRigGizmoActor.h
@@ -12,16 +12,16 @@
 struct FActorSpawnParameters;
 
 USTRUCT()
-struct FGizmoActorCreationParam
+struct FControlShapeActorCreationParam
 {
 	GENERATED_BODY()
 
-	FGizmoActorCreationParam()
+	FControlShapeActorCreationParam()
 		: ManipObj(nullptr)
 		, ControlRigIndex(INDEX_NONE)
 		, ControlName(NAME_None)
 		, SpawnTransform(FTransform::Identity)
-		, GizmoTransform(FTransform::Identity)
+		, ShapeTransform(FTransform::Identity)
 		, MeshTransform(FTransform::Identity)
 		, StaticMesh(nullptr)
 		, Material(nullptr)
@@ -35,7 +35,7 @@ struct FGizmoActorCreationParam
 	int32		ControlRigIndex;
 	FName		ControlName;
 	FTransform	SpawnTransform;
-	FTransform  GizmoTransform;
+	FTransform  ShapeTransform;
 	FTransform  MeshTransform;
 	TSoftObjectPtr<UStaticMesh> StaticMesh;
 	TSoftObjectPtr<UMaterial> Material;
@@ -46,12 +46,12 @@ struct FGizmoActorCreationParam
 
 /** An actor used to represent a rig control */
 UCLASS(NotPlaceable, Transient)
-class CONTROLRIG_API AControlRigGizmoActor : public AActor
+class CONTROLRIG_API AControlRigShapeActor : public AActor
 {
 	GENERATED_BODY()
 
 public:
-	AControlRigGizmoActor(const FObjectInitializer& ObjectInitializer);
+	AControlRigShapeActor(const FObjectInitializer& ObjectInitializer);
 
 	// this is the one holding transform for the controls
 	UPROPERTY()
@@ -107,8 +107,8 @@ public:
 	/** Called from the edit mode each tick */
 	virtual void TickControl() {};
 
-	/** changes the gizmo color */
-	virtual void SetGizmoColor(const FLinearColor& InColor);
+	/** changes the shape color */
+	virtual void SetShapeColor(const FLinearColor& InColor);
 
 	/** Event called when the transform of this control has changed */
 	UFUNCTION(BlueprintImplementableEvent)
@@ -132,38 +132,38 @@ public:
 
 	// this returns root component transform based on attach
 	// when there is no attach, it is based on 0
-	UFUNCTION(BlueprintCallable, Category = "ControlRig|Gizmo")
+	UFUNCTION(BlueprintCallable, Category = "ControlRig|Shape")
 	void SetGlobalTransform(const FTransform& InTransform);
 
-	UFUNCTION(BlueprintPure, Category = "ControlRig|Gizmo")
+	UFUNCTION(BlueprintPure, Category = "ControlRig|Shape")
 	FTransform GetGlobalTransform() const;
 private:
 	/** Whether this control is enabled */
-	UPROPERTY(BlueprintGetter = IsEnabled, BlueprintSetter= SetEnabled, Category = "ControlRig|Gizmo")
+	UPROPERTY(BlueprintGetter = IsEnabled, BlueprintSetter= SetEnabled, Category = "ControlRig|Shape")
 	uint8 bEnabled : 1;
 
 	/** Whether this control is selected */
-	UPROPERTY(BlueprintGetter = IsSelectedInEditor, BlueprintSetter = SetSelected, Category = "ControlRig|Gizmo")
+	UPROPERTY(BlueprintGetter = IsSelectedInEditor, BlueprintSetter = SetSelected, Category = "ControlRig|Shape")
 	uint8 bSelected : 1;
 
 	/** Whether this control can be selected */
-	UPROPERTY(BlueprintGetter = IsSelectable, BlueprintSetter = SetSelectable, Category = "ControlRig|Gizmo")
+	UPROPERTY(BlueprintGetter = IsSelectable, BlueprintSetter = SetSelectable, Category = "ControlRig|Shape")
 	uint8 bSelectable : 1;
 
 	/** Whether this control is hovered */
-	UPROPERTY(BlueprintGetter = IsHovered, BlueprintSetter = SetHovered, Category = "ControlRig|Gizmo")
+	UPROPERTY(BlueprintGetter = IsHovered, BlueprintSetter = SetHovered, Category = "ControlRig|Shape")
 	uint8 bHovered : 1;
 
 };
 
 /**
- * Creating Gizmo Param helper functions
+ * Creating Shape Param helper functions
  */
-namespace FControlRigGizmoHelper
+namespace FControlRigShapeHelper
 {
-	extern CONTROLRIG_API AControlRigGizmoActor* CreateGizmoActor(UWorld* InWorld, UStaticMesh* InStaticMesh, const FGizmoActorCreationParam& CreationParam);
-	AControlRigGizmoActor* CreateGizmoActor(UWorld* InWorld, TSubclassOf<AControlRigGizmoActor> InClass, const FGizmoActorCreationParam& CreationParam);
-	extern CONTROLRIG_API AControlRigGizmoActor* CreateDefaultGizmoActor(UWorld* InWorld, const FGizmoActorCreationParam& CreationParam);
+	extern CONTROLRIG_API AControlRigShapeActor* CreateShapeActor(UWorld* InWorld, UStaticMesh* InStaticMesh, const FControlShapeActorCreationParam& CreationParam);
+	AControlRigShapeActor* CreateShapeActor(UWorld* InWorld, TSubclassOf<AControlRigShapeActor> InClass, const FControlShapeActorCreationParam& CreationParam);
+	extern CONTROLRIG_API AControlRigShapeActor* CreateDefaultShapeActor(UWorld* InWorld, const FControlShapeActorCreationParam& CreationParam);
 
 	FActorSpawnParameters GetDefaultSpawnParameter();
 }
\ No newline at end of file
diff --git a/Engine/Plugins/Animation/ControlRig/Source/ControlRig/Public/ControlRigGizmoLibrary.h b/Engine/Plugins/Animation/ControlRig/Source/ControlRig/Public/ControlRigGizmoLibrary.h
--- a/Engine/Plugins/Animation/ControlRig/Source/ControlRig/Public/ControlRigGizmoLibrary.h
+++ b/Engine/Plugins/Animation/ControlRig/Source/ControlRig/Public/ControlRigGizmoLibrary.h
@@ -10,50 +10,55 @@
 
 #include "ControlRigGizmoLibrary.generated.h"
 
-USTRUCT(BlueprintType, meta = (DisplayName = "Gizmo"))
-struct CONTROLRIG_API FControlRigGizmoDefinition
+class UControlRigShapeLibrary;
+
+USTRUCT(BlueprintType, meta = (DisplayName = "Shape"))
+struct CONTROLRIG_API FControlRigShapeDefinition
 {
 	GENERATED_USTRUCT_BODY()
 
-	FControlRigGizmoDefinition()
+	FControlRigShapeDefinition()
 	{
-		GizmoName = TEXT("Gizmo");
+		ShapeName = TEXT("Default");
 		StaticMesh = nullptr;
 		Transform = FTransform::Identity;
 	}
 
-	UPROPERTY(BlueprintReadOnly, EditAnywhere, Category = "Gizmo")
-	FName GizmoName;
+	UPROPERTY(BlueprintReadOnly, EditAnywhere, Category = "Shape")
+	FName ShapeName;
 
-	UPROPERTY(BlueprintReadOnly, EditAnywhere, Category = "Gizmo")
+	UPROPERTY(BlueprintReadOnly, EditAnywhere, Category = "Shape")
 	TSoftObjectPtr<UStaticMesh> StaticMesh;
 
-	UPROPERTY(BlueprintReadOnly, EditAnywhere, Category = "Gizmo")
+	UPROPERTY(BlueprintReadOnly, EditAnywhere, Category = "Shape")
 	FTransform Transform;
+
+	mutable TWeakObjectPtr<UControlRigShapeLibrary> Library;
 };
 
-UCLASS(BlueprintType, meta = (DisplayName = "GizmoLibrary"))
-class CONTROLRIG_API UControlRigGizmoLibrary : public UObject
+UCLASS(BlueprintType, meta = (DisplayName = "Control Rig Shape Library"))
+class CONTROLRIG_API UControlRigShapeLibrary : public UObject
 {
 	GENERATED_BODY()
 
 public:
 
-	UControlRigGizmoLibrary();
+	UControlRigShapeLibrary();
 
-	UPROPERTY(EditAnywhere, Category = "GizmoLibrary")
-	FControlRigGizmoDefinition DefaultGizmo;
+	UPROPERTY(EditAnywhere, Category = "ShapeLibrary")
+	FControlRigShapeDefinition DefaultShape;
 
-	UPROPERTY(EditAnywhere, Category = "GizmoLibrary")
+	UPROPERTY(EditAnywhere, Category = "ShapeLibrary")
 	TSoftObjectPtr<UMaterial> DefaultMaterial;
 
-	UPROPERTY(EditAnywhere, Category = "GizmoLibrary")
+	UPROPERTY(EditAnywhere, Category = "ShapeLibrary")
 	FName MaterialColorParameter;
 
-	UPROPERTY(EditAnywhere, Category = "GizmoLibrary")
-	TArray<FControlRigGizmoDefinition> Gizmos;
+	UPROPERTY(EditAnywhere, Category = "ShapeLibrary")
+	TArray<FControlRigShapeDefinition> Shapes;
 
-	const FControlRigGizmoDefinition* GetGizmoByName(const FName& InName, bool bUseDefaultIfNotFound = false) const;
+	const FControlRigShapeDefinition* GetShapeByName(const FName& InName, bool bUseDefaultIfNotFound = false) const;
+	static const FControlRigShapeDefinition* GetShapeByName(const FName& InName, const TArray<TSoftObjectPtr<UControlRigShapeLibrary>>& InShapeLibraries);
 
 #if WITH_EDITOR
 
diff --git a/Engine/Plugins/Animation/ControlRig/Source/ControlRig/Public/ControlRigObjectVersion.h b/Engine/Plugins/Animation/ControlRig/Source/ControlRig/Public/ControlRigObjectVersion.h
--- a/Engine/Plugins/Animation/ControlRig/Source/ControlRig/Public/ControlRigObjectVersion.h
+++ b/Engine/Plugins/Animation/ControlRig/Source/ControlRig/Public/ControlRigObjectVersion.h
@@ -42,6 +42,9 @@ struct CONTROLRIG_API FControlRigObjectVersion
 		// RigHierarchy now stores min and max values as float storages 
 		StorageMinMaxValuesAsFloatStorage,
 
+		// RenameGizmoToShape 
+		RenameGizmoToShape,
+
 		// -----<new versions can be added above this line>-------------------------------------------------
 		VersionPlusOne,
 		LatestVersion = VersionPlusOne - 1
diff --git a/Engine/Plugins/Animation/ControlRig/Source/ControlRig/Public/Rigs/RigHierarchy.h b/Engine/Plugins/Animation/ControlRig/Source/ControlRig/Public/Rigs/RigHierarchy.h
--- a/Engine/Plugins/Animation/ControlRig/Source/ControlRig/Public/Rigs/RigHierarchy.h
+++ b/Engine/Plugins/Animation/ControlRig/Source/ControlRig/Public/Rigs/RigHierarchy.h
@@ -34,7 +34,7 @@ enum ERigTransformStackEntryType
 {
 	TransformPose,
 	ControlOffset,
-	ControlGizmo,
+	ControlShape,
 	CurveValue
 };
 
@@ -1077,31 +1077,31 @@ public:
 	}
 
 	/**
-	 * Returns the global gizmo transform for a given control element.
+	 * Returns the global shape transform for a given control element.
 	 * @param InKey The key of the control to retrieve the transform for
 	 * @param bInitial If true the initial transform will be used
-	 * @return The global gizmo transform
+	 * @return The global shape transform
 	 */
 	UFUNCTION(BlueprintCallable, Category = URigHierarchy)
-    FORCEINLINE_DEBUGGABLE FTransform GetGlobalControlGizmoTransform(FRigElementKey InKey, bool bInitial = false) const
+    FORCEINLINE_DEBUGGABLE FTransform GetGlobalControlShapeTransform(FRigElementKey InKey, bool bInitial = false) const
 	{
-		return GetGlobalControlGizmoTransformByIndex(GetIndex(InKey), bInitial);
+		return GetGlobalControlShapeTransformByIndex(GetIndex(InKey), bInitial);
 	}
 
 	/**
-	 * Returns the global gizmo transform for a given control element.
+	 * Returns the global shape transform for a given control element.
 	 * @param InElementIndex The index of the control to retrieve the transform for
 	 * @param bInitial If true the initial transform will be used
-	 * @return The global gizmo transform
+	 * @return The global shape transform
 	 */
 	UFUNCTION(BlueprintCallable, Category = URigHierarchy)
-    FORCEINLINE_DEBUGGABLE FTransform GetGlobalControlGizmoTransformByIndex(int32 InElementIndex, bool bInitial = false) const
+    FORCEINLINE_DEBUGGABLE FTransform GetGlobalControlShapeTransformByIndex(int32 InElementIndex, bool bInitial = false) const
 	{
 		if(Elements.IsValidIndex(InElementIndex))
 		{
 			if(FRigControlElement* ControlElement = Cast<FRigControlElement>(Elements[InElementIndex]))
 			{
-				return GetControlGizmoTransform(ControlElement, bInitial ? ERigTransformType::InitialGlobal : ERigTransformType::CurrentGlobal);
+				return GetControlShapeTransform(ControlElement, bInitial ? ERigTransformType::InitialGlobal : ERigTransformType::CurrentGlobal);
 			}
 		}
 		return FTransform::Identity;
@@ -1440,33 +1440,33 @@ public:
 	}
 
 	/**
-	 * Sets the gizmo transform for a given control element by key
-	 * @param InKey The key of the control element to set the gizmo transform for
-	 * @param InTransform The new gizmo transform value to set
+	 * Sets the shape transform for a given control element by key
+	 * @param InKey The key of the control element to set the shape transform for
+	 * @param InTransform The new shape transform value to set
 	 * @param bInitial If true the initial value will be used
 	 * @param bSetupUndo If true the transform stack will be setup for undo / redo
 	 */
 	UFUNCTION(BlueprintCallable, Category = URigHierarchy)
-    FORCEINLINE_DEBUGGABLE void SetControlGizmoTransform(FRigElementKey InKey, FTransform InTransform, bool bInitial = false, bool bSetupUndo = false)
+    FORCEINLINE_DEBUGGABLE void SetControlShapeTransform(FRigElementKey InKey, FTransform InTransform, bool bInitial = false, bool bSetupUndo = false)
 	{
-		return SetControlGizmoTransformByIndex(GetIndex(InKey), InTransform, bInitial, bSetupUndo);
+		return SetControlShapeTransformByIndex(GetIndex(InKey), InTransform, bInitial, bSetupUndo);
 	}
 
 	/**
-	 * Sets the local gizmo transform for a given control element by index
-	 * @param InElementIndex The index of the control element to set the gizmo transform for
-	 * @param InTransform The new local gizmo transform value to set
+	 * Sets the local shape transform for a given control element by index
+	 * @param InElementIndex The index of the control element to set the shape transform for
+	 * @param InTransform The new local shape transform value to set
 	 * @param bInitial If true the initial value will be used
 	 * @param bSetupUndo If true the transform stack will be setup for undo / redo
 	 */
 	UFUNCTION(BlueprintCallable, Category = URigHierarchy)
-    FORCEINLINE_DEBUGGABLE void SetControlGizmoTransformByIndex(int32 InElementIndex, FTransform InTransform, bool bInitial = false, bool bSetupUndo = false)
+    FORCEINLINE_DEBUGGABLE void SetControlShapeTransformByIndex(int32 InElementIndex, FTransform InTransform, bool bInitial = false, bool bSetupUndo = false)
 	{
 		if(Elements.IsValidIndex(InElementIndex))
 		{
 			if(FRigControlElement* ControlElement = Cast<FRigControlElement>(Elements[InElementIndex]))
 			{
-				SetControlGizmoTransform(ControlElement, InTransform, bInitial ? ERigTransformType::InitialLocal : ERigTransformType::CurrentLocal, bSetupUndo);
+				SetControlShapeTransform(ControlElement, InTransform, bInitial ? ERigTransformType::InitialLocal : ERigTransformType::CurrentLocal, bSetupUndo);
 			}
 		}
 	}
@@ -2317,22 +2317,22 @@ public:
 	void SetControlOffsetTransform(FRigControlElement* InControlElement, const FTransform& InTransform, const ERigTransformType::Type InTransformType, bool bAffectChildren, bool bSetupUndo = false, bool bForce = false, bool bPrintPythonCommands = false);
 
 	/**
-	 * Returns the global gizmo transform for a given control element.
-	 * @param InControlElement The control element to retrieve the gizmo transform for
+	 * Returns the global shape transform for a given control element.
+	 * @param InControlElement The control element to retrieve the shape transform for
 	 * @param InTransformType The type of transform to set
-	 * @return The global gizmo transform
+	 * @return The global shape transform
 	 */
-	FTransform GetControlGizmoTransform(FRigControlElement* InControlElement, const ERigTransformType::Type InTransformType) const;
+	FTransform GetControlShapeTransform(FRigControlElement* InControlElement, const ERigTransformType::Type InTransformType) const;
 
 	/**
-	 * Sets the gizmo transform for a given control element
+	 * Sets the shape transform for a given control element
 	 * @param InControlElement The element to set the transform for
-	 * @param InTransform The gizmo transform to set
-	 * @param InTransformType The type of transform to set. Note: for gizmo transform, setting the initial value also updates the current value
+	 * @param InTransform The shape transform to set
+	 * @param InTransformType The type of transform to set. Note: for shape transform, setting the initial value also updates the current value
 	 * @param bSetupUndo If true the transform stack will be setup for undo / redo
 	 * @param bForce Set the transform even if it is the same as the previously set one
 	 */
-	void SetControlGizmoTransform(FRigControlElement* InControlElement, const FTransform& InTransform, const ERigTransformType::Type InTransformType, bool bSetupUndo = false, bool bForce = false, bool bPrintPythonCommands = false);
+	void SetControlShapeTransform(FRigControlElement* InControlElement, const FTransform& InTransform, const ERigTransformType::Type InTransformType, bool bSetupUndo = false, bool bForce = false, bool bPrintPythonCommands = false);
 
 	/**
 	 * Sets the control settings for a given control element
diff --git a/Engine/Plugins/Animation/ControlRig/Source/ControlRig/Public/Rigs/RigHierarchyController.h b/Engine/Plugins/Animation/ControlRig/Source/ControlRig/Public/Rigs/RigHierarchyController.h
--- a/Engine/Plugins/Animation/ControlRig/Source/ControlRig/Public/Rigs/RigHierarchyController.h
+++ b/Engine/Plugins/Animation/ControlRig/Source/ControlRig/Public/Rigs/RigHierarchyController.h
@@ -108,7 +108,7 @@ public:
 	 * @param InSettings All of the control's settings
 	 * @param InValue The value to use for the control
 	 * @param InOffsetTransform The transform to use for the offset
-	 * @param InGizmoTransform The transform to use for the gizmo
+	 * @param InShapeTransform The transform to use for the shape
 	 * @param bSetupUndo If set to true the stack will record the change for undo / redo
 	 * @param bPrintPythonCommand If set to true a python command equivalent to this call will be printed out
 	 * @return The key for the newly created control.
@@ -119,7 +119,7 @@ public:
     	FRigControlSettings InSettings,
     	FRigControlValue InValue,
     	FTransform InOffsetTransform = FTransform::Identity,
-        FTransform InGizmoTransform = FTransform::Identity,
+        FTransform InShapeTransform = FTransform::Identity,
         bool bSetupUndo = true,
         bool bPrintPythonCommand = false
         );
@@ -515,7 +515,7 @@ public:
 	
 	TArray<FString> GetSetControlOffsetTransformPythonCommands(const FRigControlElement* Control, const FTransform& Offset, bool bInitial = false, bool bAffectChildren = true) const;
 	
-	TArray<FString> GetSetControlGizmoTransformPythonCommands(const FRigControlElement* Control, const FTransform& Transform, bool bInitial = false) const;
+	TArray<FString> GetSetControlShapeTransformPythonCommands(const FRigControlElement* Control, const FTransform& Transform, bool bInitial = false) const;
 #endif
 	
 private:
diff --git a/Engine/Plugins/Animation/ControlRig/Source/ControlRig/Public/Rigs/RigHierarchyDefines.h b/Engine/Plugins/Animation/ControlRig/Source/ControlRig/Public/Rigs/RigHierarchyDefines.h
--- a/Engine/Plugins/Animation/ControlRig/Source/ControlRig/Public/Rigs/RigHierarchyDefines.h
+++ b/Engine/Plugins/Animation/ControlRig/Source/ControlRig/Public/Rigs/RigHierarchyDefines.h
@@ -48,7 +48,7 @@ enum class ERigHierarchyNotification : uint8
 	HierarchyReset,
 	ControlSettingChanged,
 	ControlVisibilityChanged,
-	ControlGizmoTransformChanged,
+	ControlShapeTransformChanged,
 	ParentWeightsChanged,
 
 	/** MAX - invalid */
diff --git a/Engine/Plugins/Animation/ControlRig/Source/ControlRig/Public/Rigs/RigHierarchyElements.h b/Engine/Plugins/Animation/ControlRig/Source/ControlRig/Public/Rigs/RigHierarchyElements.h
--- a/Engine/Plugins/Animation/ControlRig/Source/ControlRig/Public/Rigs/RigHierarchyElements.h
+++ b/Engine/Plugins/Animation/ControlRig/Source/ControlRig/Public/Rigs/RigHierarchyElements.h
@@ -899,20 +899,20 @@ struct CONTROLRIG_API FRigControlSettings
 	UPROPERTY(BlueprintReadWrite, Category = Limits, meta = (EditCondition = "bLimitTranslation || bLimitRotation || bLimitScale"))
 	FRigControlValue MaximumValue;
 
-	/** Set to true if the gizmo is enabled in 3d */
+	/** Set to true if the shape is enabled in 3d */
 	UPROPERTY(BlueprintReadWrite, EditAnywhere, Category = Gizmo)
-	bool bGizmoEnabled;
+	bool bShapeEnabled;
 
-	/** Set to true if the gizmo is currently visible in 3d */
-	UPROPERTY(BlueprintReadWrite, EditAnywhere, Category = Gizmo, meta = (EditCondition = "bGizmoEnabled"))
-	bool bGizmoVisible;
+	/** Set to true if the shape is currently visible in 3d */
+	UPROPERTY(BlueprintReadWrite, EditAnywhere, Category = Shape, meta = (EditCondition = "bShapeEnabled"))
+	bool bShapeVisible;
 
 	/* This is optional UI setting - this doesn't mean this is always used, but it is optional for manipulation layer to use this*/
-	UPROPERTY(BlueprintReadWrite, EditAnywhere, Category = Gizmo, meta = (EditCondition = "bGizmoEnabled"))
-	FName GizmoName;
+	UPROPERTY(BlueprintReadWrite, EditAnywhere, Category = Shape, meta = (EditCondition = "bShapeEnabled"))
+	FName ShapeName;
 
-	UPROPERTY(BlueprintReadWrite, EditAnywhere, Category = Gizmo, meta = (EditCondition = "bGizmoEnabled"))
-	FLinearColor GizmoColor;
+	UPROPERTY(BlueprintReadWrite, EditAnywhere, Category = Shape, meta = (EditCondition = "bShapeEnabled"))
+	FLinearColor ShapeColor;
 
 	/** If the control is transient and only visible in the control rig editor */
 	UPROPERTY(BlueprintReadWrite, Category = Control)
@@ -1009,7 +1009,7 @@ public:
 	FRigCurrentAndInitialTransform Offset;
 
 	UPROPERTY(BlueprintReadOnly, EditAnywhere, Category = RigElement)
-	FRigCurrentAndInitialTransform Gizmo;
+	FRigCurrentAndInitialTransform Shape;
 
 protected:
 	
diff --git a/Engine/Plugins/Animation/ControlRig/Source/ControlRig/Public/Settings/ControlRigSettings.h b/Engine/Plugins/Animation/ControlRig/Source/ControlRig/Public/Settings/ControlRigSettings.h
--- a/Engine/Plugins/Animation/ControlRig/Source/ControlRig/Public/Settings/ControlRigSettings.h
+++ b/Engine/Plugins/Animation/ControlRig/Source/ControlRig/Public/Settings/ControlRigSettings.h
@@ -38,8 +38,8 @@ class CONTROLRIG_API UControlRigSettings : public UDeveloperSettings
 public:
 
 #if WITH_EDITORONLY_DATA
-	UPROPERTY(EditAnywhere, config, Category = DefaultGizmo)
-	TSoftObjectPtr<UControlRigGizmoLibrary> DefaultGizmoLibrary;
+	UPROPERTY(EditAnywhere, config, Category = Shapes)
+	TSoftObjectPtr<UControlRigShapeLibrary> DefaultShapeLibrary;
 #endif
 	
 	static UControlRigSettings * Get() { return GetMutableDefault<UControlRigSettings>(); }
diff --git a/Engine/Plugins/Animation/ControlRig/Source/ControlRigDeveloper/Private/ControlRigBlueprint.cpp b/Engine/Plugins/Animation/ControlRig/Source/ControlRigDeveloper/Private/ControlRigBlueprint.cpp
--- a/Engine/Plugins/Animation/ControlRig/Source/ControlRigDeveloper/Private/ControlRigBlueprint.cpp
+++ b/Engine/Plugins/Animation/ControlRig/Source/ControlRigDeveloper/Private/ControlRigBlueprint.cpp
@@ -66,7 +66,8 @@ UControlRigBlueprint::UControlRigBlueprint(const FObjectInitializer& ObjectIniti
 	bSuspendAllNotifications = false;
 
 #if WITH_EDITORONLY_DATA
-	GizmoLibrary = UControlRigSettings::Get()->DefaultGizmoLibrary;
+	GizmoLibrary_DEPRECATED = nullptr;
+	ShapeLibraries.Add(UControlRigSettings::Get()->DefaultShapeLibrary);
 #endif
 
 	bRecompileOnLoad = 0;
@@ -508,6 +509,33 @@ void UControlRigBlueprint::PostLoad()
 #endif
 	}
 
+	// upgrade the gizmo libraries to shape libraries
+	if(GizmoLibrary_DEPRECATED.IsValid())
+	{
+		ShapeLibraries.Reset();
+		ShapeLibraries.Add(GizmoLibrary_DEPRECATED);
+		GizmoLibrary_DEPRECATED.Reset();
+	}
+	else if (GetLinkerCustomVersion(FControlRigObjectVersion::GUID) < FControlRigObjectVersion::RenameGizmoToShape)
+	{
+		// if it's an older file and it doesn't have the GizmoLibrary stored,
+		// refer to the previous default.
+		ShapeLibraries.Reset();
+		ShapeLibraries.Add(LoadObject<UControlRigShapeLibrary>(nullptr, TEXT("/ControlRig/Controls/DefaultGizmoLibrary.DefaultGizmoLibrary")));
+
+		// also walk over all controls and check if any of them were using the "default" gizmo
+		Hierarchy->ForEach<FRigControlElement>([](FRigControlElement* ControlElement) -> bool
+		{
+			static FName PreviousDefault = TEXT("Gizmo");
+			static FName NewDefault = FControlRigShapeDefinition().ShapeName;
+			if(ControlElement->Settings.ShapeName == PreviousDefault)
+			{
+				ControlElement->Settings.ShapeName = NewDefault;
+			}
+			return true;
+		});
+	}
+
 #if WITH_EDITOR
 	// delay compilation until the package has been loaded
 	FCoreUObjectDelegates::OnEndLoadPackage.AddUObject(this, &UControlRigBlueprint::HandlePackageDone);
@@ -2130,6 +2158,11 @@ bool UControlRigBlueprint::ChangeMemberVariableType(const FName& InName, const F
 	return true;
 }
 
+const FControlRigShapeDefinition* UControlRigBlueprint::GetControlShapeByName(const FName& InName) const
+{
+	return UControlRigShapeLibrary::GetShapeByName(InName, ShapeLibraries);
+}
+
 FName UControlRigBlueprint::AddTransientControl(URigVMPin* InPin)
 {
 	TUniquePtr<FControlValueScope> ValueScope;
diff --git a/Engine/Plugins/Animation/ControlRig/Source/ControlRigDeveloper/Private/ControlRigVisualGraphUtils.cpp b/Engine/Plugins/Animation/ControlRig/Source/ControlRigDeveloper/Private/ControlRigVisualGraphUtils.cpp
--- a/Engine/Plugins/Animation/ControlRig/Source/ControlRigDeveloper/Private/ControlRigVisualGraphUtils.cpp
+++ b/Engine/Plugins/Animation/ControlRig/Source/ControlRigDeveloper/Private/ControlRigVisualGraphUtils.cpp
@@ -170,7 +170,7 @@ FString FControlRigVisualGraphUtils::DumpRigHierarchyToDotGraph(URigHierarchy* I
 				}
 				case ERigElementType::Control:
 				{
-					Color = Cast<FRigControlElement>(InElement)->Settings.GizmoColor;
+					Color = Cast<FRigControlElement>(InElement)->Settings.ShapeColor;
 					break;
 				}
 				default:
diff --git a/Engine/Plugins/Animation/ControlRig/Source/ControlRigDeveloper/Public/ControlRigBlueprint.h b/Engine/Plugins/Animation/ControlRig/Source/ControlRigDeveloper/Public/ControlRigBlueprint.h
--- a/Engine/Plugins/Animation/ControlRig/Source/ControlRigDeveloper/Public/ControlRigBlueprint.h
+++ b/Engine/Plugins/Animation/ControlRig/Source/ControlRigDeveloper/Public/ControlRigBlueprint.h
@@ -405,8 +405,13 @@ public:
 #endif
 
 #if WITH_EDITORONLY_DATA
-	UPROPERTY(EditAnywhere, Category = DefaultGizmo)
-	TSoftObjectPtr<UControlRigGizmoLibrary> GizmoLibrary;
+	UPROPERTY()
+	TSoftObjectPtr<UControlRigShapeLibrary> GizmoLibrary_DEPRECATED;
+
+	UPROPERTY(EditAnywhere, Category = Shapes)
+	TArray<TSoftObjectPtr<UControlRigShapeLibrary>> ShapeLibraries;
+
+	const FControlRigShapeDefinition* GetControlShapeByName(const FName& InName) const;
 #endif
 
 	UPROPERTY(transient, DuplicateTransient, VisibleAnywhere, Category = "VM", meta = (DisplayName = "VM Statistics", DisplayAfter = "VMCompileSettings"))
diff --git a/Engine/Plugins/Animation/ControlRig/Source/ControlRigDeveloper/Public/ControlRigBlueprintCompiler.h b/Engine/Plugins/Animation/ControlRig/Source/ControlRigDeveloper/Public/ControlRigBlueprintCompiler.h
--- a/Engine/Plugins/Animation/ControlRig/Source/ControlRigDeveloper/Public/ControlRigBlueprintCompiler.h
+++ b/Engine/Plugins/Animation/ControlRig/Source/ControlRigDeveloper/Public/ControlRigBlueprintCompiler.h
@@ -5,7 +5,6 @@
 #include "KismetCompilerModule.h"
 #include "KismetCompiler.h"
 
-struct FControlRigBlueprintPropertyLink;
 class UControlRigGraph;
 class UControlRigGraphNode;
 class UControlRigBlueprint;
diff --git a/Engine/Plugins/Animation/ControlRig/Source/ControlRigEditor/Private/ControlRigEditorModule.cpp b/Engine/Plugins/Animation/ControlRig/Source/ControlRigEditor/Private/ControlRigEditorModule.cpp
--- a/Engine/Plugins/Animation/ControlRig/Source/ControlRigEditor/Private/ControlRigEditorModule.cpp
+++ b/Engine/Plugins/Animation/ControlRig/Source/ControlRigEditor/Private/ControlRigEditorModule.cpp
@@ -257,7 +257,7 @@ void FControlRigEditorModule::StartupModule()
 	};
 
 	RegisterAssetTypeAction(MakeShareable(new FControlRigBlueprintActions()));
-	RegisterAssetTypeAction(MakeShareable(new FControlRigGizmoLibraryActions()));
+	RegisterAssetTypeAction(MakeShareable(new FControlRigShapeLibraryActions()));
 	RegisterAssetTypeAction(MakeShareable(new FAssetTypeActions_ControlRigPose()));
 	
 	// Register sequencer track editor
diff --git a/Engine/Plugins/Animation/ControlRig/Source/ControlRigEditor/Private/ControlRigEditorStyle.h b/Engine/Plugins/Animation/ControlRig/Source/ControlRigEditor/Private/ControlRigEditorStyle.h
--- a/Engine/Plugins/Animation/ControlRig/Source/ControlRigEditor/Private/ControlRigEditorStyle.h
+++ b/Engine/Plugins/Animation/ControlRig/Source/ControlRigEditor/Private/ControlRigEditorStyle.h
@@ -108,7 +108,7 @@ public:
 				// similar style to "LevelViewport.StartingPlayInEditorBorder"
 				Set( "ControlRig.Viewport.Border", new BOX_BRUSH( "Old/Window/ViewportDebugBorder", 0.8f, FLinearColor(1.0f,1.0f,1.0f,1.0f) ) );
 				// similar style to "AnimViewport.Notification.Warning"
-				Set( "ControlRig.Viewport.Notification.ChangeGizmoTransform", new BOX_BRUSH("Common/RoundedSelection_16x", 4.0f/16.0f, FLinearColor(FColor(169, 0, 148))));
+				Set( "ControlRig.Viewport.Notification.ChangeShapeTransform", new BOX_BRUSH("Common/RoundedSelection_16x", 4.0f/16.0f, FLinearColor(FColor(169, 0, 148))));
 				// restore the previous content root
 				SetContentRoot(PreviousContentRoot);
 			}
diff --git a/Engine/Plugins/Animation/ControlRig/Source/ControlRigEditor/Private/ControlRigElementDetails.cpp b/Engine/Plugins/Animation/ControlRig/Source/ControlRigEditor/Private/ControlRigElementDetails.cpp
--- a/Engine/Plugins/Animation/ControlRig/Source/ControlRigEditor/Private/ControlRigElementDetails.cpp
+++ b/Engine/Plugins/Animation/ControlRig/Source/ControlRigEditor/Private/ControlRigElementDetails.cpp
@@ -672,10 +672,10 @@ void FRigComputedTransformDetails::CustomizeChildren(TSharedRef<IPropertyHandle>
 		PropertyPath.RightChopInline(7);
 		PropertyChain.AddTail(FRigControlElement::StaticStruct()->FindPropertyByName(TEXT("Offset")));
 	}
-	else if(PropertyPath.StartsWith(TEXT("Gizmo.")))
+	else if(PropertyPath.StartsWith(TEXT("Shape.")))
 	{
 		PropertyPath.RightChopInline(6);
-		PropertyChain.AddTail(FRigControlElement::StaticStruct()->FindPropertyByName(TEXT("Gizmo")));
+		PropertyChain.AddTail(FRigControlElement::StaticStruct()->FindPropertyByName(TEXT("Shape")));
 	}
 
 	if(PropertyPath.StartsWith(TEXT("Current.")))
@@ -1560,19 +1560,24 @@ void FRigControlElementDetails::CustomizeChildren(TSharedRef<class IPropertyHand
 {
 	FRigTransformElementDetails::CustomizeChildren(InStructPropertyHandle, StructBuilder, StructCustomizationUtils);
 
-	GizmoNameList.Reset();
+	ShapeNameList.Reset();
 	if (BlueprintBeingCustomized)
 	{
-		if (!BlueprintBeingCustomized->GizmoLibrary.IsValid())
+		bool bUseNameSpace = BlueprintBeingCustomized->ShapeLibraries.Num() > 1;
+		for(TSoftObjectPtr<UControlRigShapeLibrary>& ShapeLibrary : BlueprintBeingCustomized->ShapeLibraries)
 		{
-			BlueprintBeingCustomized->GizmoLibrary.LoadSynchronous();
-		}
-		if (BlueprintBeingCustomized->GizmoLibrary.IsValid())
-		{
-			GizmoNameList.Add(MakeShared<FString>(BlueprintBeingCustomized->GizmoLibrary->DefaultGizmo.GizmoName.ToString()));
-			for (const FControlRigGizmoDefinition& Gizmo : BlueprintBeingCustomized->GizmoLibrary->Gizmos)
+			if (!ShapeLibrary.IsValid())
+			{
+				ShapeLibrary.LoadSynchronous();
+			}
+			if (ShapeLibrary.IsValid())
 			{
-				GizmoNameList.Add(MakeShared<FString>(Gizmo.GizmoName.ToString()));
+				const FString NameSpace = bUseNameSpace ? ShapeLibrary->GetName() + TEXT(".") : FString();
+				ShapeNameList.Add(MakeShared<FString>(NameSpace + ShapeLibrary->DefaultShape.ShapeName.ToString()));
+				for (const FControlRigShapeDefinition& Shape : ShapeLibrary->Shapes)
+				{
+					ShapeNameList.Add(MakeShared<FString>(NameSpace + Shape.ShapeName.ToString()));
+				}
 			}
 		}
 	}
@@ -1583,7 +1588,7 @@ void FRigControlElementDetails::CustomizeChildren(TSharedRef<class IPropertyHand
 	}
 
 	IDetailGroup* ControlGroup = &StructBuilder.AddGroup(TEXT("Control"), LOCTEXT("Control", "Control"));
-	IDetailGroup* GizmoGroup = nullptr;
+	IDetailGroup* ShapeGroup = nullptr;
 	IDetailGroup* LimitsGroup = nullptr;
 
 	const TSharedPtr<IPropertyHandle> SettingsHandle = InStructPropertyHandle->GetChildHandle(TEXT("Settings"));
@@ -1743,9 +1748,9 @@ void FRigControlElementDetails::CustomizeChildren(TSharedRef<class IPropertyHand
 	ControlGroup->AddPropertyRow(ControlTypeHandle.ToSharedRef());
 	ControlGroup->AddPropertyRow(SettingsHandle->GetChildHandle(TEXT("bAnimatable")).ToSharedRef());
 
-	// any but bool controls show the offset + gizmo + limits
-	const bool bNeedsGizmoProperties = IsAnyControlNotOfType(ERigControlType::Bool);
-	if (bNeedsGizmoProperties)
+	// any but bool controls show the offset + shape + limits
+	const bool bNeedsShapeProperties = IsAnyControlNotOfType(ERigControlType::Bool);
+	if (bNeedsShapeProperties)
 	{
 		// setup offset
 		{
@@ -1756,7 +1761,7 @@ void FRigControlElementDetails::CustomizeChildren(TSharedRef<class IPropertyHand
 			ControlGroup->AddPropertyRow(TransformHandle.ToSharedRef()).DisplayName(FText::FromString(TEXT("Offset Transform")));
 		}
 
-		GizmoGroup = &StructBuilder.AddGroup(TEXT("Control Shape"), LOCTEXT("ControlShape", "Control Shape"));
+		ShapeGroup = &StructBuilder.AddGroup(TEXT("Control Shape"), LOCTEXT("ControlShape", "Control Shape"));
 		LimitsGroup = &StructBuilder.AddGroup(TEXT("Limits"), LOCTEXT("Limits", "Limits"));
 	}
 
@@ -1835,40 +1840,40 @@ void FRigControlElementDetails::CustomizeChildren(TSharedRef<class IPropertyHand
 		FRigControlElementDetails_SetupValueWidget(*ControlGroup, StructBuilder, ERigControlValueType::Maximum, ControlElements[0], HierarchyBeingCustomized);
 	}
 
-	if (bNeedsGizmoProperties)
+	if (bNeedsShapeProperties)
 	{
-		GizmoGroup->AddPropertyRow(SettingsHandle->GetChildHandle(TEXT("bGizmoEnabled")).ToSharedRef())
+		ShapeGroup->AddPropertyRow(SettingsHandle->GetChildHandle(TEXT("bShapeEnabled")).ToSharedRef())
 		.DisplayName(FText::FromString(TEXT("Enabled")));
-		GizmoGroup->AddPropertyRow(SettingsHandle->GetChildHandle(TEXT("bGizmoVisible")).ToSharedRef())
+		ShapeGroup->AddPropertyRow(SettingsHandle->GetChildHandle(TEXT("bShapeVisible")).ToSharedRef())
 		.DisplayName(FText::FromString(TEXT("Visible")));
 
-		// setup gizmo transform
+		// setup shape transform
 		{
-			const TSharedPtr<IPropertyHandle> GizmoHandle = InStructPropertyHandle->GetChildHandle(TEXT("Gizmo"));
-			const TSharedPtr<IPropertyHandle> InitialHandle = GizmoHandle->GetChildHandle(TEXT("Initial"));
+			const TSharedPtr<IPropertyHandle> ShapeHandle = InStructPropertyHandle->GetChildHandle(TEXT("Shape"));
+			const TSharedPtr<IPropertyHandle> InitialHandle = ShapeHandle->GetChildHandle(TEXT("Initial"));
 			const TSharedPtr<IPropertyHandle> LocalHandle = InitialHandle->GetChildHandle(TEXT("Local"));
 			const TSharedPtr<IPropertyHandle> TransformHandle = LocalHandle->GetChildHandle(TEXT("Transform"));
-			GizmoGroup->AddPropertyRow(TransformHandle.ToSharedRef())
-			.IsEnabled(TAttribute<bool>::CreateSP(this, &FRigControlElementDetails::IsGizmoEnabled));
+			ShapeGroup->AddPropertyRow(TransformHandle.ToSharedRef())
+			.IsEnabled(TAttribute<bool>::CreateSP(this, &FRigControlElementDetails::IsShapeEnabled));
 		}
 
-		const TSharedPtr<IPropertyHandle> GizmoShapeHandle = SettingsHandle->GetChildHandle(TEXT("GizmoName"));
-		GizmoGroup->AddPropertyRow(GizmoShapeHandle.ToSharedRef()).CustomWidget()
+		const TSharedPtr<IPropertyHandle> ShapeShapeHandle = SettingsHandle->GetChildHandle(TEXT("ShapeName"));
+		ShapeGroup->AddPropertyRow(ShapeShapeHandle.ToSharedRef()).CustomWidget()
 		.NameContent()
 		[
 			SNew(STextBlock)
 			.Text(FText::FromString(TEXT("Shape")))
 			.Font(IDetailLayoutBuilder::GetDetailFont())
-			.IsEnabled(this, &FRigControlElementDetails::IsGizmoEnabled)
+			.IsEnabled(this, &FRigControlElementDetails::IsShapeEnabled)
 		]
 		.ValueContent()
 		[
-			SNew(SControlRigGizmoNameList, ControlElements, BlueprintBeingCustomized)
-			.OnGetNameListContent(this, &FRigControlElementDetails::GetGizmoNameList)
-			.IsEnabled(this, &FRigControlElementDetails::IsGizmoEnabled)
+			SNew(SControlRigShapeNameList, ControlElements, BlueprintBeingCustomized)
+			.OnGetNameListContent(this, &FRigControlElementDetails::GetShapeNameList)
+			.IsEnabled(this, &FRigControlElementDetails::IsShapeEnabled)
 		];
 
-		GizmoGroup->AddPropertyRow(SettingsHandle->GetChildHandle(TEXT("GizmoColor")).ToSharedRef())
+		ShapeGroup->AddPropertyRow(SettingsHandle->GetChildHandle(TEXT("ShapeColor")).ToSharedRef())
 		.DisplayName(FText::FromString(TEXT("Color")));
 	}
 
@@ -1927,7 +1932,7 @@ void FRigControlElementDetails::CustomizeChildren(TSharedRef<class IPropertyHand
 	AnimationGroup->AddPropertyRow(SettingsHandle->GetChildHandle(TEXT("Customization"))->GetChildHandle(TEXT("AvailableSpaces")).ToSharedRef());
 }
 
-bool FRigControlElementDetails::IsGizmoEnabled() const
+bool FRigControlElementDetails::IsShapeEnabled() const
 {
 	for(TWeakObjectPtr<UDetailsViewWrapperObject> ObjectBeingCustomized : ObjectsBeingCustomized)
 	{
@@ -1935,7 +1940,7 @@ bool FRigControlElementDetails::IsGizmoEnabled() const
 		{
 			if(FRigControlElement* ControlElement = Cast<FRigControlElement>(ObjectBeingCustomized->GetContent<FRigBaseElement>()))
 			{
-				if(ControlElement->Settings.bGizmoEnabled)
+				if(ControlElement->Settings.bShapeEnabled)
 				{
 					return true;
 				}
@@ -1974,9 +1979,9 @@ bool FRigControlElementDetails::IsEnabled(ERigControlValueType InValueType) cons
 	return true;
 }
 
-const TArray<TSharedPtr<FString>>& FRigControlElementDetails::GetGizmoNameList() const
+const TArray<TSharedPtr<FString>>& FRigControlElementDetails::GetShapeNameList() const
 {
-	return GizmoNameList;
+	return ShapeNameList;
 }
 
 const TArray<TSharedPtr<FString>>& FRigControlElementDetails::GetControlTypeList() const
diff --git a/Engine/Plugins/Animation/ControlRig/Source/ControlRigEditor/Private/ControlRigElementDetails.h b/Engine/Plugins/Animation/ControlRig/Source/ControlRigEditor/Private/ControlRigElementDetails.h
--- a/Engine/Plugins/Animation/ControlRig/Source/ControlRigEditor/Private/ControlRigElementDetails.h
+++ b/Engine/Plugins/Animation/ControlRig/Source/ControlRigEditor/Private/ControlRigElementDetails.h
@@ -223,23 +223,20 @@ public:
 	/** IPropertyTypeCustomization interface */
 	virtual void CustomizeChildren(TSharedRef<class IPropertyHandle> InStructPropertyHandle, class IDetailChildrenBuilder& StructBuilder, IPropertyTypeCustomizationUtils& StructCustomizationUtils) override;
 
-	bool IsGizmoEnabled() const;
+	bool IsShapeEnabled() const;
 	bool IsEnabled(ERigControlValueType InValueType) const;
 
-	const TArray<TSharedPtr<FString>>& GetGizmoNameList() const;
+	const TArray<TSharedPtr<FString>>& GetShapeNameList() const;
 	const TArray<TSharedPtr<FString>>& GetControlTypeList() const;
 
 	FText GetDisplayName() const;
 	void SetDisplayName(const FText& InNewText, ETextCommit::Type InCommitType);
 
 private:
-	TArray<TSharedPtr<FString>> GizmoNameList;
+	TArray<TSharedPtr<FString>> ShapeNameList;
 	static TArray<TSharedPtr<FString>> ControlTypeList;
 	TSharedPtr<FRigInfluenceEntryModifier> InfluenceModifier;
 	TSharedPtr<FStructOnScope> InfluenceModifierStruct;
-	FEditPropertyChain PoseCurrentChain;
-	FEditPropertyChain OffsetPropertyChain;
-	FEditPropertyChain GizmoPropertyChain;
 };
 
 class FRigNullElementDetails : public FRigTransformElementDetails
diff --git a/Engine/Plugins/Animation/ControlRig/Source/ControlRigEditor/Private/ControlRigGizmoLibraryActions.cpp b/Engine/Plugins/Animation/ControlRig/Source/ControlRigEditor/Private/ControlRigGizmoLibraryActions.cpp
--- a/Engine/Plugins/Animation/ControlRig/Source/ControlRigEditor/Private/ControlRigGizmoLibraryActions.cpp
+++ b/Engine/Plugins/Animation/ControlRig/Source/ControlRigEditor/Private/ControlRigGizmoLibraryActions.cpp
@@ -4,7 +4,7 @@
 
 #define LOCTEXT_NAMESPACE "AssetTypeActions"
 
-void FControlRigGizmoLibraryActions::GetActions( const TArray<UObject*>& InObjects, FMenuBuilder& MenuBuilder )
+void FControlRigShapeLibraryActions::GetActions( const TArray<UObject*>& InObjects, FMenuBuilder& MenuBuilder )
 {
 	//FAssetTypeActions_AnimationAsset::GetActions(InObjects, MenuBuilder);
 }
diff --git a/Engine/Plugins/Animation/ControlRig/Source/ControlRigEditor/Private/ControlRigGizmoLibraryActions.h b/Engine/Plugins/Animation/ControlRig/Source/ControlRigEditor/Private/ControlRigGizmoLibraryActions.h
--- a/Engine/Plugins/Animation/ControlRig/Source/ControlRigEditor/Private/ControlRigGizmoLibraryActions.h
+++ b/Engine/Plugins/Animation/ControlRig/Source/ControlRigEditor/Private/ControlRigGizmoLibraryActions.h
@@ -8,7 +8,7 @@
 
 class FMenuBuilder;
 
-class FControlRigGizmoLibraryActions : public FAssetTypeActions_Base
+class FControlRigShapeLibraryActions : public FAssetTypeActions_Base
 {
 public:
 	// IAssetTypeActions Implementation
@@ -17,9 +17,9 @@ public:
 		return true;
 	}
 	virtual void GetActions(const TArray<UObject*>& InObjects, FMenuBuilder& MenuBuilder) override;
-	virtual FText GetName() const override { return NSLOCTEXT("AssetTypeActions", "AssetTypeActions_ControlRigGizmoLibrary", "Gizmo Library"); }
+	virtual FText GetName() const override { return NSLOCTEXT("AssetTypeActions", "AssetTypeActions_ControlRigShapeLibrary", "Control Rig Shape Library"); }
 	virtual FColor GetTypeColor() const override { return FColor(100,100,255); }
-	virtual UClass* GetSupportedClass() const override { return UControlRigGizmoLibrary::StaticClass(); }
+	virtual UClass* GetSupportedClass() const override { return UControlRigShapeLibrary::StaticClass(); }
 	virtual bool CanFilter() override { return true; }
 	virtual uint32 GetCategories() override { return EAssetTypeCategories::Animation; }
 };
diff --git a/Engine/Plugins/Animation/ControlRig/Source/ControlRigEditor/Private/ControlRigGizmoLibraryFactory.cpp b/Engine/Plugins/Animation/ControlRig/Source/ControlRigEditor/Private/ControlRigGizmoLibraryFactory.cpp
--- a/Engine/Plugins/Animation/ControlRig/Source/ControlRigEditor/Private/ControlRigGizmoLibraryFactory.cpp
+++ b/Engine/Plugins/Animation/ControlRig/Source/ControlRigEditor/Private/ControlRigGizmoLibraryFactory.cpp
@@ -6,30 +6,30 @@
 
 #define LOCTEXT_NAMESPACE "ControlRigGizmoLibraryFactory"
 
-UControlRigGizmoLibraryFactory::UControlRigGizmoLibraryFactory()
+UControlRigShapeLibraryFactory::UControlRigShapeLibraryFactory()
 {
 	bCreateNew = true;
 	bEditAfterNew = true;
-	SupportedClass = UControlRigGizmoLibrary::StaticClass();
+	SupportedClass = UControlRigShapeLibrary::StaticClass();
 }
 
-UObject* UControlRigGizmoLibraryFactory::FactoryCreateNew(UClass* Class, UObject* InParent, FName Name, EObjectFlags Flags, UObject* Context, FFeedbackContext* Warn)
+UObject* UControlRigShapeLibraryFactory::FactoryCreateNew(UClass* Class, UObject* InParent, FName Name, EObjectFlags Flags, UObject* Context, FFeedbackContext* Warn)
 {
-	UControlRigGizmoLibrary* GizmoLibrary = NewObject<UControlRigGizmoLibrary>(InParent, Name, Flags);
+	UControlRigShapeLibrary* ShapeLibrary = NewObject<UControlRigShapeLibrary>(InParent, Name, Flags);
 
-	GizmoLibrary->DefaultGizmo.StaticMesh = LoadObject<UStaticMesh>(nullptr, TEXT("/ControlRig/Controls/ControlRig_Sphere_solid.ControlRig_Sphere_solid"));
-	GizmoLibrary->DefaultMaterial = LoadObject<UMaterial>(nullptr, TEXT("/ControlRig/Controls/ControlRigGizmoMaterial.ControlRigGizmoMaterial"));
-	GizmoLibrary->MaterialColorParameter = TEXT("Color");
+	ShapeLibrary->DefaultShape.StaticMesh = LoadObject<UStaticMesh>(nullptr, TEXT("/ControlRig/Controls/ControlRig_Sphere_solid.ControlRig_Sphere_solid"));
+	ShapeLibrary->DefaultMaterial = LoadObject<UMaterial>(nullptr, TEXT("/ControlRig/Controls/ControlRigGizmoMaterial.ControlRigGizmoMaterial"));
+	ShapeLibrary->MaterialColorParameter = TEXT("Color");
 
-	return GizmoLibrary;
+	return ShapeLibrary;
 }
 
-FText UControlRigGizmoLibraryFactory::GetDisplayName() const
+FText UControlRigShapeLibraryFactory::GetDisplayName() const
 {
-	return LOCTEXT("ControlRigGizmoLibraryFactoryName", "Control Rig Gizmo Library");
+	return LOCTEXT("ControlRigShapeLibraryFactoryName", "Control Rig Shape Library");
 }
 
-uint32 UControlRigGizmoLibraryFactory::GetMenuCategories() const
+uint32 UControlRigShapeLibraryFactory::GetMenuCategories() const
 {
 	return EAssetTypeCategories::Animation;
 }
diff --git a/Engine/Plugins/Animation/ControlRig/Source/ControlRigEditor/Private/ControlRigGizmoLibraryFactory.h b/Engine/Plugins/Animation/ControlRig/Source/ControlRigEditor/Private/ControlRigGizmoLibraryFactory.h
--- a/Engine/Plugins/Animation/ControlRig/Source/ControlRigEditor/Private/ControlRigGizmoLibraryFactory.h
+++ b/Engine/Plugins/Animation/ControlRig/Source/ControlRigEditor/Private/ControlRigGizmoLibraryFactory.h
@@ -9,12 +9,12 @@
 #include "ControlRigGizmoLibraryFactory.generated.h"
 
 UCLASS(MinimalAPI, HideCategories=Object)
-class UControlRigGizmoLibraryFactory : public UFactory
+class UControlRigShapeLibraryFactory : public UFactory
 {
 	GENERATED_BODY()
 
 public:
-	UControlRigGizmoLibraryFactory();
+	UControlRigShapeLibraryFactory();
 
 	// UFactory Interface
 	virtual UObject* FactoryCreateNew(UClass* Class, UObject* InParent, FName Name, EObjectFlags Flags, UObject* Context, FFeedbackContext* Warn) override;
diff --git a/Engine/Plugins/Animation/ControlRig/Source/ControlRigEditor/Private/ControlRigHierarchyCommands.cpp b/Engine/Plugins/Animation/ControlRig/Source/ControlRigEditor/Private/ControlRigHierarchyCommands.cpp
--- a/Engine/Plugins/Animation/ControlRig/Source/ControlRigEditor/Private/ControlRigHierarchyCommands.cpp
+++ b/Engine/Plugins/Animation/ControlRig/Source/ControlRigEditor/Private/ControlRigHierarchyCommands.cpp
@@ -19,13 +19,13 @@ void FControlRigHierarchyCommands::RegisterCommands()
 	UI_COMMAND(PasteGlobalTransforms, "Paste Global Transform", "Paste the global transforms.", EUserInterfaceActionType::Button, FInputChord(EKeys::V, EModifierKey::Control | EModifierKey::Shift));
 	UI_COMMAND(SetInitialTransformFromClosestBone, "Set Initial Transform from Closest Bone", "Find the Closest Bone to Initial Transform of the Selected Bones/Spaces or To Offset Transform of Selected Controls", EUserInterfaceActionType::Button, FInputChord());
 	UI_COMMAND(SetInitialTransformFromCurrentTransform, "Set Initial Transform from Current", "Save the Current Transform To Initial Transform of the Selected Bones/Spaces or To Offset Transform of Selected Controls", EUserInterfaceActionType::Button, FInputChord());
-	UI_COMMAND(SetGizmoTransformFromCurrent, "Set Gizmo Transform From Current", "Transfer the Current Local Transform of the Control to its Gizmo Transform, the Control's Transform Resets", EUserInterfaceActionType::Button, FInputChord());
+	UI_COMMAND(SetShapeTransformFromCurrent, "Set Shape Transform From Current", "Transfer the Current Local Transform of the Control to its Shape Transform, the Control's Transform Resets", EUserInterfaceActionType::Button, FInputChord());
 	UI_COMMAND(ResetTransform, "Reset Transform", "Reset the Transform", EUserInterfaceActionType::Button, FInputChord(EKeys::G, EModifierKey::Control));
 	UI_COMMAND(ResetAllTransforms, "Reset All Transforms", "Resets all Transforms", EUserInterfaceActionType::Button, FInputChord(EKeys::G, EModifierKey::Control | EModifierKey::Shift));
 	UI_COMMAND(ResetNull, "Reset Null", "Resets or injects a Null below the Control", EUserInterfaceActionType::Button, FInputChord());
 	UI_COMMAND(FrameSelection, "Frame Selection", "Expands and frames the selection in the tree", EUserInterfaceActionType::Button, FInputChord(EKeys::F));
-	UI_COMMAND(ControlBoneTransform, "Control Bone Transform", "Sets the bone transform using a gizmo", EUserInterfaceActionType::Button, FInputChord());
-	UI_COMMAND(ControlSpaceTransform, "Control Space Transform", "Sets the space transform using a gizmo", EUserInterfaceActionType::Button, FInputChord());
+	UI_COMMAND(ControlBoneTransform, "Control Bone Transform", "Sets the bone transform using a shape", EUserInterfaceActionType::Button, FInputChord());
+	UI_COMMAND(ControlSpaceTransform, "Control Space Transform", "Sets the space transform using a shape", EUserInterfaceActionType::Button, FInputChord());
 	UI_COMMAND(Unparent, "Unparent", "Unparents the selected elements from the hierarchy", EUserInterfaceActionType::Button, FInputChord(EKeys::P, EModifierKey::Shift));
 	UI_COMMAND(FilteringFlattensHierarchy, "Filtering Flattens Hierarchy", "Whether to keep the hierarchy or flatten it when searching for tree items", EUserInterfaceActionType::ToggleButton, FInputChord());
 	UI_COMMAND(HideParentsWhenFiltering, "Hide Parents When Filtering", "Whether to show parent items grayed out, or hide them entirely when filtering", EUserInterfaceActionType::ToggleButton, FInputChord());
@@ -36,7 +36,7 @@ void FControlRigHierarchyCommands::RegisterCommands()
 	UI_COMMAND(ShowRigidBodies, "Show RigidBodies", "Whether to show or hide rigidbodies", EUserInterfaceActionType::ToggleButton, FInputChord());
 	UI_COMMAND(ShowReferences, "Show References", "Whether to show or hide references", EUserInterfaceActionType::ToggleButton, FInputChord());
 	UI_COMMAND(ShowDynamicHierarchy, "Dynamic Hierarchy", "Whether to show the static or dynamic hierarchy", EUserInterfaceActionType::ToggleButton, FInputChord());
-	UI_COMMAND(ToggleGizmoTransformEdit, "Toggle Gizmo Transform Edit", "Toggle Editing Selected Control's Gizmo Transform", EUserInterfaceActionType::ToggleButton, FInputChord(EKeys::Period, EModifierKey::Control)); 
+	UI_COMMAND(ToggleControlShapeTransformEdit, "Toggle Shape Transform Edit", "Toggle Editing Selected Control's Shape Transform", EUserInterfaceActionType::ToggleButton, FInputChord(EKeys::Period, EModifierKey::Control)); 
 	UI_COMMAND(TestSpaceSwitching, "Test Space Switching", "Try out space switching on the control", EUserInterfaceActionType::Button, FInputChord(EKeys::Q, EModifierKey::Control)); 
 }
 
diff --git a/Engine/Plugins/Animation/ControlRig/Source/ControlRigEditor/Private/ControlRigHierarchyCommands.h b/Engine/Plugins/Animation/ControlRig/Source/ControlRigEditor/Private/ControlRigHierarchyCommands.h
--- a/Engine/Plugins/Animation/ControlRig/Source/ControlRigEditor/Private/ControlRigHierarchyCommands.h
+++ b/Engine/Plugins/Animation/ControlRig/Source/ControlRigEditor/Private/ControlRigHierarchyCommands.h
@@ -66,16 +66,16 @@ public:
 	/* set initial transform from closest bone */
 	TSharedPtr<FUICommandInfo> SetInitialTransformFromClosestBone;
 
-	/* set gizmo transform from the control's current local transform and reset the control's transform */
-	TSharedPtr<FUICommandInfo> SetGizmoTransformFromCurrent;
+	/* set shape transform from the control's current local transform and reset the control's transform */
+	TSharedPtr<FUICommandInfo> SetShapeTransformFromCurrent;
 
 	/* frames the selection in the tree */
 	TSharedPtr<FUICommandInfo> FrameSelection;
 
-	/* sets the bone transform using a gizmo */
+	/* sets the bone transform using a shape */
 	TSharedPtr<FUICommandInfo> ControlBoneTransform;
 
-	/* sets the space transform using a gizmo */
+	/* sets the space transform using a shape */
 	TSharedPtr<FUICommandInfo> ControlSpaceTransform;
 
 	/* Unparents the selected elements from the hierarchy */
@@ -108,8 +108,8 @@ public:
 	/* Show dynamic hierarchy */
 	TSharedPtr< FUICommandInfo > ShowDynamicHierarchy;
 
-	/** Toggle Gizmo Transform Edit*/
-	TSharedPtr< FUICommandInfo > ToggleGizmoTransformEdit;
+	/** Toggle Shape Transform Edit*/
+	TSharedPtr< FUICommandInfo > ToggleControlShapeTransformEdit;
 
 	/** Test out space switching as it would look like for animator */
 	TSharedPtr< FUICommandInfo > TestSpaceSwitching;
diff --git a/Engine/Plugins/Animation/ControlRig/Source/ControlRigEditor/Private/ControlRigThumbnailRenderer.cpp b/Engine/Plugins/Animation/ControlRig/Source/ControlRigEditor/Private/ControlRigThumbnailRenderer.cpp
--- a/Engine/Plugins/Animation/ControlRig/Source/ControlRigEditor/Private/ControlRigThumbnailRenderer.cpp
+++ b/Engine/Plugins/Animation/ControlRig/Source/ControlRigEditor/Private/ControlRigThumbnailRenderer.cpp
@@ -23,24 +23,27 @@ bool UControlRigThumbnailRenderer::CanVisualizeAsset(UObject* Object)
 		USkeletalMesh* SkeletalMesh = InRigBlueprint->PreviewSkeletalMesh.Get();
 		if (SkeletalMesh != nullptr)
 		{
-			if (InRigBlueprint->GizmoLibrary.IsValid())
+			for(const TSoftObjectPtr<UControlRigShapeLibrary>& ShapeLibrary : InRigBlueprint->ShapeLibraries)
 			{
-				bool bMissingMeshCount = true;
-				InRigBlueprint->Hierarchy->ForEach<FRigControlElement>([&](FRigControlElement* ControlElement) -> bool
+				if (ShapeLibrary.IsValid())
 				{
-					if (const FControlRigGizmoDefinition* GizmoDef = InRigBlueprint->GizmoLibrary->GetGizmoByName(ControlElement->Settings.GizmoName))
+					bool bMissingMeshCount = true;
+					InRigBlueprint->Hierarchy->ForEach<FRigControlElement>([&](FRigControlElement* ControlElement) -> bool
 					{
-						UStaticMesh* StaticMesh = GizmoDef->StaticMesh.Get();
-						if (StaticMesh == nullptr) // not yet loaded
+						if (const FControlRigShapeDefinition* ShapeDef = ShapeLibrary->GetShapeByName(ControlElement->Settings.ShapeName))
 						{
-							bMissingMeshCount = true;
-							return false;
+							UStaticMesh* StaticMesh = ShapeDef->StaticMesh.Get();
+							if (StaticMesh == nullptr) // not yet loaded
+							{
+								bMissingMeshCount = true;
+								return false;
+							}
 						}
-					}
-					return true;
-				});
+						return true;
+					});
 
-				return !bMissingMeshCount;
+					return !bMissingMeshCount;
+				}
 			}
 		}
 
@@ -60,7 +63,7 @@ void UControlRigThumbnailRenderer::Draw(UObject* Object, int32 X, int32 Y, uint3
 			RigBlueprint = InRigBlueprint;
 			Super::Draw(SkeletalMesh, X, Y, Width, Height, RenderTarget, Canvas, bAdditionalViewFamily);
 
-			for (auto Pair : GizmoActors)
+			for (auto Pair : ShapeActors)
 			{
 				if (Pair.Value && Pair.Value->GetOuter())
 				{
@@ -68,14 +71,15 @@ void UControlRigThumbnailRenderer::Draw(UObject* Object, int32 X, int32 Y, uint3
 					Pair.Value->MarkPendingKill();
 				}
 			}
-			GizmoActors.Reset();
+			ShapeActors.Reset();
 		}
 	}
 }
 
 void UControlRigThumbnailRenderer::AddAdditionalPreviewSceneContent(UObject* Object, UWorld* PreviewWorld)
 {
-	if (ThumbnailScene && ThumbnailScene->GetPreviewActor() && RigBlueprint && RigBlueprint->GizmoLibrary && RigBlueprint->GeneratedClass)
+	TSharedRef<FSkeletalMeshThumbnailScene> ThumbnailScene = ThumbnailSceneCache.EnsureThumbnailScene(Object);
+	if (ThumbnailScene->GetPreviewActor() && RigBlueprint && !RigBlueprint->ShapeLibraries.IsEmpty() && RigBlueprint->GeneratedClass)
 	{
 		UControlRig* ControlRig = nullptr;
 
@@ -115,37 +119,42 @@ void UControlRigThumbnailRenderer::AddAdditionalPreviewSceneContent(UObject* Obj
 				case ERigControlType::TransformNoScale:
 				case ERigControlType::EulerTransform:
 				{
-					if (const FControlRigGizmoDefinition* GizmoDef = RigBlueprint->GizmoLibrary->GetGizmoByName(ControlElement->Settings.GizmoName))
+					if (const FControlRigShapeDefinition* ShapeDef = RigBlueprint->GetControlShapeByName(ControlElement->Settings.ShapeName))
 					{
-						UStaticMesh* StaticMesh = GizmoDef->StaticMesh.Get();
+						UStaticMesh* StaticMesh = ShapeDef->StaticMesh.Get();
 						if (StaticMesh == nullptr) // not yet loaded
 						{
 							return true;
 						}
 
-						const FTransform GizmoGlobalTransform = ControlRig->GetHierarchy()->GetGlobalControlGizmoTransform(ControlElement->GetKey());
+						const FTransform ShapeGlobalTransform = ControlRig->GetHierarchy()->GetGlobalControlShapeTransform(ControlElement->GetKey());
 
 						FActorSpawnParameters SpawnInfo;
 						SpawnInfo.SpawnCollisionHandlingOverride = ESpawnActorCollisionHandlingMethod::AlwaysSpawn;
 						SpawnInfo.bNoFail = true;
 						SpawnInfo.ObjectFlags = RF_Transient;
-						AStaticMeshActor* GizmoActor = PreviewWorld->SpawnActor<AStaticMeshActor>(SpawnInfo);
-						GizmoActor->SetActorEnableCollision(false);
+						AStaticMeshActor* ShapeActor = PreviewWorld->SpawnActor<AStaticMeshActor>(SpawnInfo);
+						ShapeActor->SetActorEnableCollision(false);
 
-						UMaterial* DefaultMaterial = RigBlueprint->GizmoLibrary->DefaultMaterial.Get();
+						if(!ShapeDef->Library.IsValid())
+						{
+							return true;
+						}
+						
+						UMaterial* DefaultMaterial = ShapeDef->Library.Get()->DefaultMaterial.Get();
 						if (DefaultMaterial == nullptr) // not yet loaded
 						{
 							return true;
 						}
 
-						UMaterialInstanceDynamic* MaterialInstance = UMaterialInstanceDynamic::Create(DefaultMaterial, GizmoActor);
-						MaterialInstance->SetVectorParameterValue(RigBlueprint->GizmoLibrary->MaterialColorParameter, FVector(ControlElement->Settings.GizmoColor));
-						GizmoActor->GetStaticMeshComponent()->SetMaterial(0, MaterialInstance);
+						UMaterialInstanceDynamic* MaterialInstance = UMaterialInstanceDynamic::Create(DefaultMaterial, ShapeActor);
+						MaterialInstance->SetVectorParameterValue(ShapeDef->Library.Get()->MaterialColorParameter, FVector(ControlElement->Settings.ShapeColor));
+						ShapeActor->GetStaticMeshComponent()->SetMaterial(0, MaterialInstance);
 
-						GizmoActors.Add(ControlElement->GetName(), GizmoActor);
+						ShapeActors.Add(ControlElement->GetName(), ShapeActor);
 
-						GizmoActor->GetStaticMeshComponent()->SetStaticMesh(StaticMesh);
-						GizmoActor->SetActorTransform(GizmoGlobalTransform);
+						ShapeActor->GetStaticMeshComponent()->SetStaticMesh(StaticMesh);
+						ShapeActor->SetActorTransform(ShapeGlobalTransform);
 					}
 					break;
 				}
diff --git a/Engine/Plugins/Animation/ControlRig/Source/ControlRigEditor/Private/EditMode/ControlRigControlsProxy.cpp b/Engine/Plugins/Animation/ControlRig/Source/ControlRigEditor/Private/EditMode/ControlRigControlsProxy.cpp
--- a/Engine/Plugins/Animation/ControlRig/Source/ControlRigEditor/Private/EditMode/ControlRigControlsProxy.cpp
+++ b/Engine/Plugins/Animation/ControlRig/Source/ControlRigEditor/Private/EditMode/ControlRigControlsProxy.cpp
@@ -776,7 +776,7 @@ void UControlRigDetailPanelControlProxies::RecreateAllProxies(UControlRig* InCon
 	TArray<FRigControlElement*> Controls = InControlRig->AvailableControls();
 	for (FRigControlElement* ControlElement : Controls)
 	{
-		if(ControlElement->Settings.bGizmoEnabled && ControlElement->Settings.bAnimatable)
+		if(ControlElement->Settings.bShapeEnabled && ControlElement->Settings.bAnimatable)
 		{
 			AddProxy(ControlElement->GetName(), InControlRig, ControlElement);
 		}
diff --git a/Engine/Plugins/Animation/ControlRig/Source/ControlRigEditor/Private/EditMode/ControlRigEditMode.cpp b/Engine/Plugins/Animation/ControlRig/Source/ControlRigEditor/Private/EditMode/ControlRigEditMode.cpp
--- a/Engine/Plugins/Animation/ControlRig/Source/ControlRigEditor/Private/EditMode/ControlRigEditMode.cpp
+++ b/Engine/Plugins/Animation/ControlRig/Source/ControlRigEditor/Private/EditMode/ControlRigEditMode.cpp
@@ -131,13 +131,13 @@ namespace ControlRigSelectionConstants
 }
 
 FControlRigEditMode::FControlRigEditMode()
-	: bIsChangingGizmoTransform(false)
+	: bIsChangingControlShapeTransform(false)
 	, InteractionScope(nullptr)
 	, bManipulatorMadeChange(false)
 	, bSelecting(false)
 	, bSelectionChanged(false)
 	, PivotTransform(FTransform::Identity)
-	, bRecreateGizmosRequired(false)
+	, bRecreateControlShapesRequired(false)
 	, CurrentViewportClient(nullptr)
 	, bIsChangingCoordSystem(false)
 {
@@ -156,7 +156,7 @@ FControlRigEditMode::~FControlRigEditMode()
 {	
 	CommandBindings = nullptr;
 
-	DestroyGizmosActors();
+	DestroyShapesActors();
 
 	TArray<TWeakObjectPtr<UControlRig>> PreviousRuntimeRigs = RuntimeControlRigs;
 	for (int32 PreviousRuntimeRigIndex = 0; PreviousRuntimeRigIndex < PreviousRuntimeRigs.Num(); PreviousRuntimeRigIndex++)
@@ -290,13 +290,13 @@ void FControlRigEditMode::SetObjects_Internal()
 
 	if (!RuntimeControlRig)
 	{
-		DestroyGizmosActors();
+		DestroyShapesActors();
 		SetUpDetailPanel();
 	}
 	else
 	{
 		// create default manipulation layer
-		RequestToRecreateGizmoActors();
+		RequestToRecreateControlShapeActors();
 	}
 }
 
@@ -374,7 +374,7 @@ void FControlRigEditMode::Exit()
 		Toolkit.Reset();
 	}
 
-	DestroyGizmosActors();
+	DestroyShapesActors();
 
 	TArray<TWeakObjectPtr<UControlRig>> PreviousRuntimeRigs = RuntimeControlRigs;
 	for (int32 PreviousRuntimeRigIndex = 0; PreviousRuntimeRigIndex < PreviousRuntimeRigs.Num(); PreviousRuntimeRigIndex++)
@@ -402,7 +402,7 @@ void FControlRigEditMode::Exit()
 	ControlProxy->RemoveAllProxies();
 
 	//make sure the widget is reset
-	ResetGizmoSize();
+	ResetControlShapeSize();
 
 	// Call parent implementation
 	FEdMode::Exit();
@@ -431,19 +431,19 @@ void FControlRigEditMode::Tick(FEditorViewportClient* ViewportClient, float Delt
 	}
 	RecalcPivotTransform();
 
-	// Defer creation of gizmos if manipulating the viewport
-	if (bRecreateGizmosRequired && !(FSlateApplication::Get().HasAnyMouseCaptor() || GUnrealEd->IsUserInteracting()))
+	// Defer creation of shapes if manipulating the viewport
+	if (bRecreateControlShapesRequired && !(FSlateApplication::Get().HasAnyMouseCaptor() || GUnrealEd->IsUserInteracting()))
 	{
-		RecreateGizmoActors();
+		RecreateControlShapeActors();
 		TArray<FRigElementKey> SelectedRigElements = GetSelectedRigElements();
 		for (const FRigElementKey& SelectedKey : SelectedRigElements)
 		{
 			if (SelectedKey.Type == ERigElementType::Control)
 			{
-				AControlRigGizmoActor* GizmoActor = GetGizmoFromControlName(SelectedKey.Name);
-				if (GizmoActor)
+				AControlRigShapeActor* ShapeActor = GetControlShapeFromControlName(SelectedKey.Name);
+				if (ShapeActor)
 				{
-					GizmoActor->SetSelected(true);
+					ShapeActor->SetSelected(true);
 				}
 
 				if (IsInLevelEditor())
@@ -464,7 +464,7 @@ void FControlRigEditMode::Tick(FEditorViewportClient* ViewportClient, float Delt
 		}
 		SetUpDetailPanel();
 		HandleSelectionChanged();
-		bRecreateGizmosRequired = false;
+		bRecreateControlShapesRequired = false;
 	}
 
 	// We need to tick here since changing a bone for example
@@ -478,10 +478,10 @@ void FControlRigEditMode::Tick(FEditorViewportClient* ViewportClient, float Delt
 			const UE::Widget::EWidgetMode CurrentWidgetMode = ViewportClient->GetWidgetMode();
 			for (FRigElementKey SelectedRigElement : SelectedRigElements)
 			{
-				//need to loop through the gizmo actors and set widget based upon the first one
-				if (AControlRigGizmoActor* GizmoActor = GetGizmoFromControlName(SelectedRigElement.Name))
+				//need to loop through the shape actors and set widget based upon the first one
+				if (AControlRigShapeActor* ShapeActor = GetControlShapeFromControlName(SelectedRigElement.Name))
 				{
-					if (!ModeSupportedByGizmoActor(GizmoActor, CurrentWidgetMode))
+					if (!ModeSupportedByShapeActor(ShapeActor, CurrentWidgetMode))
 					{
 						if (FRigControlElement* ControlElement = ControlRig->FindControl(SelectedRigElement.Name))
 						{
@@ -633,7 +633,7 @@ void FControlRigEditMode::Render(const FSceneView* View, FViewport* Viewport, FP
 	FTransform ComponentTransform = GetHostingSceneComponentTransform();
 	if (bRender)
 	{
-		for (AControlRigGizmoActor* Actor : GizmoActors)
+		for (AControlRigShapeActor* Actor : ShapeActors)
 		{
 			//Actor->SetActorHiddenInGame(bIsHidden);
 			if (GIsEditor && Actor->GetWorld() != nullptr && !Actor->GetWorld()->IsPlayInEditor())
@@ -810,7 +810,7 @@ void FControlRigEditMode::Render(const FSceneView* View, FViewport* Viewport, FP
 	}
 	else
 	{
-		for (AControlRigGizmoActor* Actor : GizmoActors)
+		for (AControlRigShapeActor* Actor : ShapeActors)
 		{
 			//Actor->SetActorHiddenInGame(bIsHidden);
 			if (GIsEditor && Actor->GetWorld() != nullptr && !Actor->GetWorld()->IsPlayInEditor())
@@ -923,9 +923,9 @@ bool FControlRigEditMode::StartTracking(FEditorViewportClient* InViewportClient,
 
 bool FControlRigEditMode::UsesTransformWidget() const
 {
-	for (const AControlRigGizmoActor* GizmoActor : GizmoActors)
+	for (const AControlRigShapeActor* ShapeActor : ShapeActors)
 	{
-		if (GizmoActor->IsSelected())
+		if (ShapeActor->IsSelected())
 		{
 			return true;
 		}
@@ -941,11 +941,11 @@ bool FControlRigEditMode::UsesTransformWidget() const
 
 bool FControlRigEditMode::UsesTransformWidget(UE::Widget::EWidgetMode CheckMode) const
 {
-	for (const AControlRigGizmoActor* GizmoActor : GizmoActors)
+	for (const AControlRigShapeActor* ShapeActor : ShapeActors)
 	{
-		if (GizmoActor->IsSelected())
+		if (ShapeActor->IsSelected())
 		{
-			return ModeSupportedByGizmoActor(GizmoActor, CheckMode);
+			return ModeSupportedByShapeActor(ShapeActor, CheckMode);
 		}
 	}
 
@@ -990,10 +990,10 @@ bool FControlRigEditMode::HandleClick(FEditorViewportClient* InViewportClient, H
 	{
 		if(ActorHitProxy->Actor)
 		{
-			if (ActorHitProxy->Actor->IsA<AControlRigGizmoActor>())
+			if (ActorHitProxy->Actor->IsA<AControlRigShapeActor>())
 			{
-				AControlRigGizmoActor* GizmoActor = CastChecked<AControlRigGizmoActor>(ActorHitProxy->Actor);
-				if (GizmoActor->IsSelectable())
+				AControlRigShapeActor* ShapeActor = CastChecked<AControlRigShapeActor>(ActorHitProxy->Actor);
+				if (ShapeActor->IsSelectable())
 				{
 					FScopedTransaction ScopedTransaction(LOCTEXT("SelectControlTransaction", "Select Control"), IsInLevelEditor() && !GIsTransacting);
 
@@ -1005,7 +1005,7 @@ bool FControlRigEditMode::HandleClick(FEditorViewportClient* InViewportClient, H
 						InteractionScope = nullptr;
 					}
 					
-					const FName& ControlName = GizmoActor->ControlName;
+					const FName& ControlName = ShapeActor->ControlName;
 					if (Click.IsShiftDown()) //guess we just select
 					{
 						SetRigElementSelection(ERigElementType::Control, ControlName, true);
@@ -1231,22 +1231,22 @@ void FControlRigEditMode::OpenContextMenu(FEditorViewportClient* InViewportClien
 	}
 }
 
-bool FControlRigEditMode::IntersectSelect(bool InSelect, const TFunctionRef<bool(const AControlRigGizmoActor*, const FTransform&)>& Intersects)
+bool FControlRigEditMode::IntersectSelect(bool InSelect, const TFunctionRef<bool(const AControlRigShapeActor*, const FTransform&)>& Intersects)
 {
 	FTransform ComponentTransform = GetHostingSceneComponentTransform();
 
 	bool bSelected = false;
-	for (AControlRigGizmoActor* GizmoActor : GizmoActors)
+	for (AControlRigShapeActor* ShapeActor : ShapeActors)
 	{
-		if (GizmoActor->IsHiddenEd())
+		if (ShapeActor->IsHiddenEd())
 		{
 			continue;
 		}
 
-		const FTransform ControlTransform = GizmoActor->GetGlobalTransform() * ComponentTransform;
-		if (Intersects(GizmoActor, ControlTransform))
+		const FTransform ControlTransform = ShapeActor->GetGlobalTransform() * ComponentTransform;
+		if (Intersects(ShapeActor, ControlTransform))
 		{
-			SetRigElementSelection(ERigElementType::Control, GizmoActor->ControlName, InSelect);
+			SetRigElementSelection(ERigElementType::Control, ShapeActor->ControlName, InSelect);
 			bSelected = true;
 		}
 	}
@@ -1325,13 +1325,13 @@ bool FControlRigEditMode::BoxSelect(FBox& InBox, bool InSelect)
 	{
 		AActor* Actor = *It;
 
-		if (!Actor->IsA<AControlRigGizmoActor>())
+		if (!Actor->IsA<AControlRigShapeActor>())
 		{
 			continue;
 		}
 
-		AControlRigGizmoActor* GizmoActor = CastChecked<AControlRigGizmoActor>(Actor);
-		if (!GizmoActor->IsSelectable())
+		AControlRigShapeActor* ShapeActor = CastChecked<AControlRigShapeActor>(Actor);
+		if (!ShapeActor->IsSelectable())
 		{
 			continue;
 		}
@@ -1339,7 +1339,7 @@ bool FControlRigEditMode::BoxSelect(FBox& InBox, bool InSelect)
 		if (IntersectsBox(*Actor, InBox, LevelViewportClient, bStrictDragSelection))
 		{
 			bSomethingSelected = true;
-			const FName& ControlName = GizmoActor->ControlName;
+			const FName& ControlName = ShapeActor->ControlName;
 			SetRigElementSelection(ERigElementType::Control, ControlName, true);
 
 			if (bShiftDown)
@@ -1370,19 +1370,19 @@ bool FControlRigEditMode::FrustumSelect(const FConvexVolume& InFrustum, FEditorV
 		ClearRigElementSelection(FRigElementTypeHelper::ToMask(ERigElementType::Control));
 	}
 
-	for (AControlRigGizmoActor* GizmoActor : GizmoActors)
+	for (AControlRigShapeActor* ShapeActor : ShapeActors)
 	{
-		for (UActorComponent* Component : GizmoActor->GetComponents())
+		for (UActorComponent* Component : ShapeActor->GetComponents())
 		{
 			UPrimitiveComponent* PrimitiveComponent = Cast<UPrimitiveComponent>(Component);
 			if (PrimitiveComponent && PrimitiveComponent->IsRegistered() && PrimitiveComponent->IsVisibleInEditor())
 			{
 				if (PrimitiveComponent->ComponentIsTouchingSelectionFrustum(InFrustum, InViewportClient->EngineShowFlags, false /*only bsp*/, false/*encompass entire*/))
 				{
-					if (GizmoActor->IsSelectable())
+					if (ShapeActor->IsSelectable())
 					{
 						bSomethingSelected = true;
-						const FName& ControlName = GizmoActor->ControlName;
+						const FName& ControlName = ShapeActor->ControlName;
 						SetRigElementSelection(ERigElementType::Control, ControlName, true);
 					}
 				}
@@ -1430,21 +1430,21 @@ bool FControlRigEditMode::InputDelta(FEditorViewportClient* InViewportClient, FV
 		{
 			FTransform ComponentTransform = GetHostingSceneComponentTransform();
 
-			if (bIsChangingGizmoTransform)
+			if (bIsChangingControlShapeTransform)
 			{
-				for (AControlRigGizmoActor* GizmoActor : GizmoActors)
+				for (AControlRigShapeActor* ShapeActor : ShapeActors)
 				{
-					if (GizmoActor->IsSelected())
+					if (ShapeActor->IsSelected())
 					{
 						if (bManipulatorMadeChange == false)
 						{
-							GEditor->BeginTransaction(LOCTEXT("ChangeControlGizmoTransaction", "Change Control Gizmo Transform"));
+							GEditor->BeginTransaction(LOCTEXT("ChangeControlShapeTransaction", "Change Control Shape Transform"));
 						}
 
-						ChangeControlGizmoTransform(GizmoActor, bDoTranslation, InDrag, bDoRotation, InRot, bDoScale, InScale, ComponentTransform);
+						ChangeControlShapeTransform(ShapeActor, bDoTranslation, InDrag, bDoRotation, InRot, bDoScale, InScale, ComponentTransform);
 						bManipulatorMadeChange = true;
 
-						// break here since we only support changing gizmo transform of a single control at a time
+						// break here since we only support changing shape transform of a single control at a time
 						break;
 					}
 				}
@@ -1457,9 +1457,9 @@ bool FControlRigEditMode::InputDelta(FEditorViewportClient* InViewportClient, FV
 				bool bCalcLocal = bDoLocal;
 				bool bFirstTime = true;
 				FTransform InOutLocal = FTransform::Identity;
-				for (AControlRigGizmoActor* GizmoActor : GizmoActors)
+				for (AControlRigShapeActor* ShapeActor : ShapeActors)
 				{
-					if (GizmoActor->IsSelected())
+					if (ShapeActor->IsSelected())
 					{
 						// test local vs global
 						if (bManipulatorMadeChange == false)
@@ -1479,7 +1479,7 @@ bool FControlRigEditMode::InputDelta(FEditorViewportClient* InViewportClient, FV
 							}
 						}
 
-						MoveGizmo(GizmoActor, bDoTranslation, InDrag, bDoRotation, InRot, bDoScale, InScale, ComponentTransform,
+						MoveControlShape(ShapeActor, bDoTranslation, InDrag, bDoRotation, InRot, bDoScale, InScale, ComponentTransform,
 							bUseLocal, bDoLocal, InOutLocal);
 						bManipulatorMadeChange = true;
 					}
@@ -1571,11 +1571,11 @@ bool FControlRigEditMode::IsCompatibleWith(FEditorModeID OtherModeID) const
 
 void FControlRigEditMode::AddReferencedObjects( FReferenceCollector& Collector )
 {
-	if (GizmoActors.Num() > 0)
+	if (ShapeActors.Num() > 0)
 	{
-		for (AControlRigGizmoActor* GizmoActor : GizmoActors)
+		for (AControlRigShapeActor* ShapeActor : ShapeActors)
 		{
-			Collector.AddReferencedObject(GizmoActor);
+			Collector.AddReferencedObject(ShapeActor);
 		}
 	}
 	if (ControlProxy)
@@ -1735,10 +1735,10 @@ void FControlRigEditMode::RefreshObjects()
 
 bool FControlRigEditMode::CanRemoveFromPreviewScene(const USceneComponent* InComponent)
 {
-	for (AControlRigGizmoActor* GizmoActor : GizmoActors)
+	for (AControlRigShapeActor* ShapeActor : ShapeActors)
 	{
 		TInlineComponentArray<USceneComponent*> SceneComponents;
-		GizmoActor->GetComponents(SceneComponents, true);
+		ShapeActor->GetComponents(SceneComponents, true);
 		if (SceneComponents.Contains(InComponent))
 		{
 			return false;
@@ -1774,20 +1774,20 @@ void FControlRigEditMode::RecalcPivotTransform()
 			}
 		}
 		
-		if (bIsChangingGizmoTransform)
+		if (bIsChangingControlShapeTransform)
 		{
 			if (UControlRig* ControlRig = GetControlRig(true))
 			{ 
-				for (const AControlRigGizmoActor* GizmoActor : GizmoActors)
+				for (const AControlRigShapeActor* ShapeActor : ShapeActors)
 				{
-					if (GizmoActor->IsSelected())
+					if (ShapeActor->IsSelected())
 					{
-						if (FRigControlElement* ControlElement = ControlRig->GetHierarchy()->Find<FRigControlElement>(FRigElementKey(GizmoActor->ControlName, ERigElementType::Control)))
+						if (FRigControlElement* ControlElement = ControlRig->GetHierarchy()->Find<FRigControlElement>(FRigElementKey(ShapeActor->ControlName, ERigElementType::Control)))
 						{
-							PivotTransform = ControlRig->GetHierarchy()->GetControlGizmoTransform(ControlElement, ERigTransformType::CurrentGlobal);
+							PivotTransform = ControlRig->GetHierarchy()->GetControlShapeTransform(ControlElement, ERigTransformType::CurrentGlobal);
 						}
 
-						// break here since we don't want to change the gizmo transform of multiple controls.
+						// break here since we don't want to change the shape transform of multiple controls.
 						break;
 					}
 				}
@@ -1796,11 +1796,11 @@ void FControlRigEditMode::RecalcPivotTransform()
 		else
 		{
 			const UControlRigEditModeSettings* Settings = GetDefault<UControlRigEditModeSettings>();
-			for (const AControlRigGizmoActor* GizmoActor : GizmoActors)
+			for (const AControlRigShapeActor* ShapeActor : ShapeActors)
 			{
-				if (GizmoActor->IsSelected())
+				if (ShapeActor->IsSelected())
 				{
-					LastTransform = GizmoActor->GetActorTransform().GetRelativeTransform(ComponentTransform);
+					LastTransform = ShapeActor->GetActorTransform().GetRelativeTransform(ComponentTransform);
 					PivotLocation += LastTransform.GetLocation();
 					++NumSelectedControls;
 					if (Settings && Settings->bLocalTransformsInEachLocalSpace) //if in local just use first actors transform
@@ -1858,22 +1858,22 @@ void FControlRigEditMode::RecalcPivotTransform()
 void FControlRigEditMode::HandleSelectionChanged()
 {
 
-	for (AControlRigGizmoActor* GizmoActor : GizmoActors)
+	for (AControlRigShapeActor* ShapeActor : ShapeActors)
 	{
 		TInlineComponentArray<UPrimitiveComponent*> PrimitiveComponents;
-		GizmoActor->GetComponents(PrimitiveComponents, true);
+		ShapeActor->GetComponents(PrimitiveComponents, true);
 		for (UPrimitiveComponent* PrimitiveComponent : PrimitiveComponents)
 		{
 			PrimitiveComponent->PushSelectionToProxy();
 		}
 	}
 
-	// automatically exit gizmo transform edit mode if there is no gizmo selected
-	if (bIsChangingGizmoTransform)
+	// automatically exit shape transform edit mode if there is no shape selected
+	if (bIsChangingControlShapeTransform)
 	{
-		if (!CanChangeControlGizmoTransform())
+		if (!CanChangeControlShapeTransform())
 		{
-			bIsChangingGizmoTransform = false;
+			bIsChangingControlShapeTransform = false;
 		}
 	}
 
@@ -1905,20 +1905,20 @@ void FControlRigEditMode::BindCommands()
 	);
 
 	CommandBindings->MapAction(
-		Commands.IncreaseGizmoSize,
-		FExecuteAction::CreateRaw(this, &FControlRigEditMode::IncreaseGizmoSize));
+		Commands.IncreaseControlShapeSize,
+		FExecuteAction::CreateRaw(this, &FControlRigEditMode::IncreaseShapeSize));
 
 	CommandBindings->MapAction(
-		Commands.DecreaseGizmoSize,
-		FExecuteAction::CreateRaw(this, &FControlRigEditMode::DecreaseGizmoSize));
+		Commands.DecreaseControlShapeSize,
+		FExecuteAction::CreateRaw(this, &FControlRigEditMode::DecreaseShapeSize));
 
 	CommandBindings->MapAction(
-		Commands.ResetGizmoSize,
-		FExecuteAction::CreateRaw(this, &FControlRigEditMode::ResetGizmoSize));
+		Commands.ResetControlShapeSize,
+		FExecuteAction::CreateRaw(this, &FControlRigEditMode::ResetControlShapeSize));
 
 	CommandBindings->MapAction(
-		Commands.ToggleGizmoTransformEdit,
-		FExecuteAction::CreateRaw(this, &FControlRigEditMode::ToggleGizmoTransformEdit));
+		Commands.ToggleControlShapeTransformEdit,
+		FExecuteAction::CreateRaw(this, &FControlRigEditMode::ToggleControlShapeTransformEdit));
 
 	CommandBindings->MapAction(
 		Commands.OpenSpacePickerWidget,
@@ -1937,10 +1937,10 @@ bool FControlRigEditMode::GetRigElementGlobalTransform(const FRigElementKey& InE
 	if (InElement.Type == ERigElementType::Control)
 	{
 		// this code is weird. Need to set this info in manipulation layer
-		AControlRigGizmoActor* GizmoActor = GetGizmoFromControlName(InElement.Name);
-		if (GizmoActor && ensure(GizmoActor->IsSelected()))
+		AControlRigShapeActor* ShapeActor = GetControlShapeFromControlName(InElement.Name);
+		if (ShapeActor && ensure(ShapeActor->IsSelected()))
 		{
-			OutGlobalTransform = GetGizmoTransform(GizmoActor);
+			OutGlobalTransform = GetControlShapeTransform(ShapeActor);
 			return true;
 		}
 	}
@@ -1994,10 +1994,10 @@ void FControlRigEditMode::FrameSelection()
 	{
 		if (SelectedKey.Type == ERigElementType::Control)
 		{
-			AControlRigGizmoActor* GizmoActor = GetGizmoFromControlName(SelectedKey.Name);
-			if (GizmoActor)
+			AControlRigShapeActor* ShapeActor = GetControlShapeFromControlName(SelectedKey.Name);
+			if (ShapeActor)
 			{
-				Actors.Add(GizmoActor);
+				Actors.Add(ShapeActor);
 			}
 		}
 	}
@@ -2047,36 +2047,36 @@ void FControlRigEditMode::FrameItems(const TArray<FRigElementKey>& InItems)
 	}
 }
 
-void FControlRigEditMode::IncreaseGizmoSize()
+void FControlRigEditMode::IncreaseShapeSize()
 {
 	UControlRigEditModeSettings* Settings = GetMutableDefault<UControlRigEditModeSettings>();
-	Settings->GizmoScale += 0.1f;
-	GetModeManager()->SetWidgetScale(Settings->GizmoScale);
+	Settings->ShapeScale += 0.1f;
+	GetModeManager()->SetWidgetScale(Settings->ShapeScale);
 }
 
-void FControlRigEditMode::DecreaseGizmoSize()
+void FControlRigEditMode::DecreaseShapeSize()
 {
 	UControlRigEditModeSettings* Settings = GetMutableDefault<UControlRigEditModeSettings>();
-	Settings->GizmoScale -= 0.1f;
-	GetModeManager()->SetWidgetScale(Settings->GizmoScale);
+	Settings->ShapeScale -= 0.1f;
+	GetModeManager()->SetWidgetScale(Settings->ShapeScale);
 }
 
-void FControlRigEditMode::ResetGizmoSize()
+void FControlRigEditMode::ResetControlShapeSize()
 {
 	UControlRigEditModeSettings* Settings = GetMutableDefault<UControlRigEditModeSettings>();
-	Settings->GizmoScale = 1.0f;
-	GetModeManager()->SetWidgetScale(Settings->GizmoScale);
+	Settings->ShapeScale = 1.0f;
+	GetModeManager()->SetWidgetScale(Settings->ShapeScale);
 }
 
-void FControlRigEditMode::ToggleGizmoTransformEdit()
+void FControlRigEditMode::ToggleControlShapeTransformEdit()
 { 
-	if (bIsChangingGizmoTransform)
+	if (bIsChangingControlShapeTransform)
 	{
-		bIsChangingGizmoTransform = false;
+		bIsChangingControlShapeTransform = false;
 	}
-	else if (CanChangeControlGizmoTransform())
+	else if (CanChangeControlShapeTransform())
 	{
-		bIsChangingGizmoTransform = true;
+		bIsChangingControlShapeTransform = true;
 	}
 }
 
@@ -2204,10 +2204,10 @@ void FControlRigEditMode::OpenSpacePickerWidget()
 	PickerWidget->OpenDialog(false);
 }
 
-FText FControlRigEditMode::GetToggleGizmoTransformEditHotKey() const
+FText FControlRigEditMode::GetToggleControlShapeTransformEditHotKey() const
 {
 	const FControlRigEditModeCommands& Commands = FControlRigEditModeCommands::Get();
-	return Commands.ToggleGizmoTransformEdit->GetInputText();
+	return Commands.ToggleControlShapeTransformEdit->GetInputText();
 }
 
 void FControlRigEditMode::ToggleManipulators()
@@ -2258,11 +2258,11 @@ bool FControlRigEditMode::MouseMove(FEditorViewportClient* ViewportClient, FView
 	HActor* ActorHitProxy = HitProxyCast<HActor>(Viewport->GetHitProxy(x, y));
 	if(ActorHitProxy && ActorHitProxy->Actor)
 	{
-		if (ActorHitProxy->Actor->IsA<AControlRigGizmoActor>())
+		if (ActorHitProxy->Actor->IsA<AControlRigShapeActor>())
 		{
-			for (AControlRigGizmoActor* GizmoActor : GizmoActors)
+			for (AControlRigShapeActor* ShapeActor : ShapeActors)
 			{
-				GizmoActor->SetHovered(GizmoActor == ActorHitProxy->Actor);
+				ShapeActor->SetHovered(ShapeActor == ActorHitProxy->Actor);
 			}
 		}
 	}
@@ -2272,9 +2272,9 @@ bool FControlRigEditMode::MouseMove(FEditorViewportClient* ViewportClient, FView
 
 bool FControlRigEditMode::MouseLeave(FEditorViewportClient* ViewportClient, FViewport* Viewport)
 {
-	for (AControlRigGizmoActor* GizmoActor : GizmoActors)
+	for (AControlRigShapeActor* ShapeActor : ShapeActors)
 	{
-		GizmoActor->SetHovered(false);
+		ShapeActor->SetHovered(false);
 	}
 
 	return false;
@@ -2285,11 +2285,11 @@ void FControlRigEditMode::PostUndo()
 	UControlRig* RuntimeControlRig = GetControlRig(false);
 	if (!RuntimeControlRig)
 	{
-		DestroyGizmosActors();
+		DestroyShapesActors();
 	}
 }
 
-void FControlRigEditMode::RecreateGizmoActors(const TArray<FRigElementKey>& InSelectedElements)
+void FControlRigEditMode::RecreateControlShapeActors(const TArray<FRigElementKey>& InSelectedElements)
 {
 	if (ControlProxy)
 	{
@@ -2302,23 +2302,23 @@ void FControlRigEditMode::RecreateGizmoActors(const TArray<FRigElementKey>& InSe
 		InteractionRig = InteractionRig == nullptr ? ControlRig : InteractionRig;
 
 		// create gizmo actors
-		CreateGizmoActors(GetWorld());
+		CreateShapeActors(GetWorld());
 
 		USceneComponent* Component = GetHostingSceneComponent();
 		if (Component)
 		{
 			AActor* PreviewActor = Component->GetOwner();
 
-			for (AControlRigGizmoActor* GizmoActor : GizmoActors)
+			for (AControlRigShapeActor* ShapeActor : ShapeActors)
 			{
 				// attach to preview actor, so that we can communicate via relative transfrom from the previewactor
-				GizmoActor->AttachToActor(PreviewActor, FAttachmentTransformRules::KeepWorldTransform);
+				ShapeActor->AttachToActor(PreviewActor, FAttachmentTransformRules::KeepWorldTransform);
 
 				TInlineComponentArray<UPrimitiveComponent*> PrimitiveComponents;
-				GizmoActor->GetComponents(PrimitiveComponents, true);
+				ShapeActor->GetComponents(PrimitiveComponents, true);
 				for (UPrimitiveComponent* PrimitiveComponent : PrimitiveComponents)
 				{
-					PrimitiveComponent->SelectionOverrideDelegate = UPrimitiveComponent::FSelectionOverride::CreateRaw(this, &FControlRigEditMode::GizmoSelectionOverride);
+					PrimitiveComponent->SelectionOverrideDelegate = UPrimitiveComponent::FSelectionOverride::CreateRaw(this, &FControlRigEditMode::ShapeSelectionOverride);
 					PrimitiveComponent->PushSelectionToProxy();
 				}
 			}
@@ -2346,12 +2346,12 @@ FControlRigEditMode* FControlRigEditMode::GetEditModeFromWorldContext(UWorld* In
 	return nullptr;
 }
 
-bool FControlRigEditMode::GizmoSelectionOverride(const UPrimitiveComponent* InComponent) const
+bool FControlRigEditMode::ShapeSelectionOverride(const UPrimitiveComponent* InComponent) const
 {
     //Think we only want to do this in regular editor, in the level editor we are driving selection
 	if (!IsInLevelEditor())
 	{
-	AControlRigGizmoActor* OwnerActor = Cast<AControlRigGizmoActor>(InComponent->GetOwner());
+	AControlRigShapeActor* OwnerActor = Cast<AControlRigShapeActor>(InComponent->GetOwner());
 	if (OwnerActor)
 	{
 		// See if the actor is in a selected unit proxy
@@ -2423,17 +2423,17 @@ void FControlRigEditMode::OnHierarchyModified(ERigHierarchyNotification InNotif,
 		case ERigHierarchyNotification::ElementRenamed:
 		case ERigHierarchyNotification::HierarchyReset:
 		case ERigHierarchyNotification::ControlSettingChanged:
-		case ERigHierarchyNotification::ControlGizmoTransformChanged:
+		case ERigHierarchyNotification::ControlShapeTransformChanged:
 		{
 			// in case the gizmo is turned off, automatically exit gizmo transform edit mode
-			if (bIsChangingGizmoTransform)
+			if (bIsChangingControlShapeTransform)
 			{
-				if (!CanChangeControlGizmoTransform())
+				if (!CanChangeControlShapeTransform())
 				{
-					bIsChangingGizmoTransform = false;
+					bIsChangingControlShapeTransform = false;
 				} 
 			}
-			RequestToRecreateGizmoActors();
+			RequestToRecreateControlShapeActors();
 			break;
 		}
 		case ERigHierarchyNotification::ElementSelected:
@@ -2462,10 +2462,10 @@ void FControlRigEditMode::OnHierarchyModified(ERigHierarchyNotification InNotif,
 						}
 						// users may select gizmo and control rig units, so we have to let them go through both of them if they do
 						// first go through gizmo actor
-						AControlRigGizmoActor* GizmoActor = GetGizmoFromControlName(Key.Name);
-						if (GizmoActor)
+						AControlRigShapeActor* ShapeActor = GetControlShapeFromControlName(Key.Name);
+						if (ShapeActor)
 						{
-							GizmoActor->SetSelected(bSelected);
+							ShapeActor->SetSelected(bSelected);
 
 						}
 						if (IsInLevelEditor())
@@ -2517,9 +2517,9 @@ void FControlRigEditMode::OnControlModified(UControlRig* Subject, FRigControlEle
 	if (UControlRig* ControlRig = static_cast<UControlRig*>(Subject))
 	{
 		FTransform ComponentTransform = GetHostingSceneComponentTransform();
-		if (AControlRigGizmoActor* const* Actor = GizmoToControlMap.FindKey(InControl.Index))
+		if (AControlRigShapeActor* const* Actor = GizmoToControlMap.FindKey(InControl.Index))
 		{
-			TickGizmo(*Actor, ComponentTransform);
+			TickControlShape(*Actor, ComponentTransform);
 		}
 	}
 	*/
@@ -2554,7 +2554,7 @@ void FControlRigEditMode::OnCoordSystemChanged(ECoordSystem InCoordSystem)
 	}
 }
 
-bool FControlRigEditMode::CanChangeControlGizmoTransform()
+bool FControlRigEditMode::CanChangeControlShapeTransform()
 {
 	if (!IsInLevelEditor())
 	{
@@ -2569,11 +2569,11 @@ bool FControlRigEditMode::CanChangeControlGizmoTransform()
 					// only enable for a Control with Gizmo enabled and visible
 					if (FRigControlElement* ControlElement = ControlRig->GetHierarchy()->Find<FRigControlElement>(SelectedRigElements[0]))
 					{
-						if (ControlElement->Settings.bGizmoEnabled && ControlElement->Settings.bGizmoVisible)
+						if (ControlElement->Settings.bShapeEnabled && ControlElement->Settings.bShapeVisible)
 						{
-							if (AControlRigGizmoActor* GizmoActor = GetGizmoFromControlName(SelectedRigElements[0].Name))
+							if (AControlRigShapeActor* ShapeActor = GetControlShapeFromControlName(SelectedRigElements[0].Name))
 							{
-								if (ensure(GizmoActor->IsSelected()))
+								if (ensure(ShapeActor->IsSelected()))
 								{
 									return true;
 								}
@@ -2588,24 +2588,24 @@ bool FControlRigEditMode::CanChangeControlGizmoTransform()
 	return false;
 }
 
-void FControlRigEditMode::SetGizmoTransform(AControlRigGizmoActor* GizmoActor, const FTransform& InTransform)
+void FControlRigEditMode::SetControlShapeTransform(AControlRigShapeActor* ShapeActor, const FTransform& InTransform)
 {
-	if (UControlRig* ControlRig = GetControlRig(true, GizmoActor->ControlRigIndex))
+	if (UControlRig* ControlRig = GetControlRig(true, ShapeActor->ControlRigIndex))
 	{
-		ControlRig->SetControlGlobalTransform(GizmoActor->ControlName, InTransform);
+		ControlRig->SetControlGlobalTransform(ShapeActor->ControlName, InTransform);
 	}
 }
 
-FTransform FControlRigEditMode::GetGizmoTransform(AControlRigGizmoActor* GizmoActor) const
+FTransform FControlRigEditMode::GetControlShapeTransform(AControlRigShapeActor* ShapeActor) const
 {
-	if (UControlRig* ControlRig = GetControlRig(true, GizmoActor->ControlRigIndex))
+	if (UControlRig* ControlRig = GetControlRig(true, ShapeActor->ControlRigIndex))
 	{
-		return ControlRig->GetControlGlobalTransform(GizmoActor->ControlName);
+		return ControlRig->GetControlGlobalTransform(ShapeActor->ControlName);
 	}
 	return FTransform::Identity;
 }
 
-void FControlRigEditMode::MoveGizmo(AControlRigGizmoActor* GizmoActor, const bool bTranslation, FVector& InDrag,
+void FControlRigEditMode::MoveControlShape(AControlRigShapeActor* ShapeActor, const bool bTranslation, FVector& InDrag,
 	const bool bRotation, FRotator& InRot, const bool bScale, FVector& InScale, const FTransform& ToWorldTransform,
 	bool bUseLocal, bool bCalcLocal, FTransform& InOutLocal)
 {
@@ -2614,13 +2614,13 @@ void FControlRigEditMode::MoveGizmo(AControlRigGizmoActor* GizmoActor, const boo
 	//first case is where we do all controls by the local diff.
 	if (bUseLocal)
 	{
-		if (UControlRig* RuntimeControlRig = GetControlRig(false, GizmoActor->ControlRigIndex))
+		if (UControlRig* RuntimeControlRig = GetControlRig(false, ShapeActor->ControlRigIndex))
 		{
-			UControlRig* InteractionControlRig = GetControlRig(true, GizmoActor->ControlRigIndex);
+			UControlRig* InteractionControlRig = GetControlRig(true, ShapeActor->ControlRigIndex);
 
 			FRigControlModifiedContext Context;
 			Context.EventName = FRigUnit_BeginExecution::EventName;
-			FTransform CurrentLocalTransform = InteractionControlRig->GetControlLocalTransform(GizmoActor->ControlName);
+			FTransform CurrentLocalTransform = InteractionControlRig->GetControlLocalTransform(ShapeActor->ControlName);
 			if (bRotation)
 			{
 
@@ -2640,12 +2640,12 @@ void FControlRigEditMode::MoveGizmo(AControlRigGizmoActor* GizmoActor, const boo
 
 			if (bTransformChanged)
 			{
-				InteractionControlRig->SetControlLocalTransform(GizmoActor->ControlName, CurrentLocalTransform);
+				InteractionControlRig->SetControlLocalTransform(ShapeActor->ControlName, CurrentLocalTransform);
 
-				FTransform CurrentTransform  = InteractionControlRig->GetControlGlobalTransform(GizmoActor->ControlName);			// assumes it's attached to actor
+				FTransform CurrentTransform  = InteractionControlRig->GetControlGlobalTransform(ShapeActor->ControlName);			// assumes it's attached to actor
 				CurrentTransform = ToWorldTransform * CurrentTransform;
 
-				GizmoActor->SetGlobalTransform(CurrentTransform);
+				ShapeActor->SetGlobalTransform(CurrentTransform);
 
 				if (RuntimeControlRig->GetInteractionRig() == InteractionControlRig)
 				{
@@ -2656,7 +2656,7 @@ void FControlRigEditMode::MoveGizmo(AControlRigGizmoActor* GizmoActor, const boo
 	}
 	if(!bTransformChanged) //not local or doing scale.
 	{
-		FTransform CurrentTransform = GetGizmoTransform(GizmoActor)* ToWorldTransform;
+		FTransform CurrentTransform = GetControlShapeTransform(ShapeActor) * ToWorldTransform;
 
 		if (bRotation)
 		{
@@ -2684,16 +2684,16 @@ void FControlRigEditMode::MoveGizmo(AControlRigGizmoActor* GizmoActor, const boo
 
 		if (bTransformChanged)
 		{
-			if (UControlRig* RuntimeControlRig = GetControlRig(false, GizmoActor->ControlRigIndex))
+			if (UControlRig* RuntimeControlRig = GetControlRig(false, ShapeActor->ControlRigIndex))
 			{
-				UControlRig* InteractionControlRig = GetControlRig(true, GizmoActor->ControlRigIndex);
+				UControlRig* InteractionControlRig = GetControlRig(true, ShapeActor->ControlRigIndex);
 
 				FTransform NewTransform = CurrentTransform.GetRelativeTransform(ToWorldTransform);
 				FRigControlModifiedContext Context;
 				Context.EventName = FRigUnit_BeginExecution::EventName;
 				if (bCalcLocal)
 				{
-					InOutLocal = InteractionControlRig->GetControlLocalTransform(GizmoActor->ControlName);
+					InOutLocal = InteractionControlRig->GetControlLocalTransform(ShapeActor->ControlName);
 				}
 
 				bool bPrintPythonCommands = false;
@@ -2701,11 +2701,11 @@ void FControlRigEditMode::MoveGizmo(AControlRigGizmoActor* GizmoActor, const boo
 				{
 					bPrintPythonCommands = World->IsPreviewWorld();
 				}
-				InteractionControlRig->SetControlGlobalTransform(GizmoActor->ControlName, NewTransform, true, Context, true, bPrintPythonCommands);			// assumes it's attached to actor
-				GizmoActor->SetGlobalTransform(CurrentTransform);
+				InteractionControlRig->SetControlGlobalTransform(ShapeActor->ControlName, NewTransform, true, Context, true, bPrintPythonCommands);			// assumes it's attached to actor
+				ShapeActor->SetGlobalTransform(CurrentTransform);
 				if (bCalcLocal)
 				{
-					FTransform NewLocal = InteractionControlRig->GetControlLocalTransform(GizmoActor->ControlName);
+					FTransform NewLocal = InteractionControlRig->GetControlLocalTransform(ShapeActor->ControlName);
 					InOutLocal = NewLocal.GetRelativeTransform(InOutLocal);
 
 				}
@@ -2720,7 +2720,7 @@ void FControlRigEditMode::MoveGizmo(AControlRigGizmoActor* GizmoActor, const boo
 #if WITH_EDITOR
 	if (bTransformChanged)
 	{
-		if (UControlRig* RuntimeControlRig = GetControlRig(false, GizmoActor->ControlRigIndex))
+		if (UControlRig* RuntimeControlRig = GetControlRig(false, ShapeActor->ControlRigIndex))
 		{
 			if (UWorld* World = RuntimeControlRig->GetWorld())
 			{
@@ -2737,7 +2737,7 @@ void FControlRigEditMode::MoveGizmo(AControlRigGizmoActor* GizmoActor, const boo
 #endif
 }
 
-void FControlRigEditMode::ChangeControlGizmoTransform(AControlRigGizmoActor* GizmoActor, const bool bTranslation, FVector& InDrag,
+void FControlRigEditMode::ChangeControlShapeTransform(AControlRigShapeActor* ShapeActor, const bool bTranslation, FVector& InDrag,
 	const bool bRotation, FRotator& InRot, const bool bScale, FVector& InScale, const FTransform& ToWorldTransform)
 {
 	bool bTransformChanged = false; 
@@ -2746,9 +2746,9 @@ void FControlRigEditMode::ChangeControlGizmoTransform(AControlRigGizmoActor* Giz
 
 	if (UControlRig* ControlRig = GetControlRig(true))
 	{
-		if (FRigControlElement* ControlElement = ControlRig->GetHierarchy()->Find<FRigControlElement>(FRigElementKey(GizmoActor->ControlName, ERigElementType::Control)))
+		if (FRigControlElement* ControlElement = ControlRig->GetHierarchy()->Find<FRigControlElement>(FRigElementKey(ShapeActor->ControlName, ERigElementType::Control)))
 		{
-			CurrentTransform = ControlRig->GetHierarchy()->GetControlGizmoTransform(ControlElement, ERigTransformType::CurrentGlobal);
+			CurrentTransform = ControlRig->GetHierarchy()->GetControlShapeTransform(ControlElement, ERigTransformType::CurrentGlobal);
 			CurrentTransform = CurrentTransform * ToWorldTransform;
 		}
 	}
@@ -2779,36 +2779,33 @@ void FControlRigEditMode::ChangeControlGizmoTransform(AControlRigGizmoActor* Giz
 
 	if (bTransformChanged)
 	{
-		if (UControlRig* RuntimeControlRig = GetControlRig(false, GizmoActor->ControlRigIndex))
+		if (UControlRig* RuntimeControlRig = GetControlRig(false, ShapeActor->ControlRigIndex))
 		{
-			UControlRig* InteractionControlRig = GetControlRig(true, GizmoActor->ControlRigIndex);
+			UControlRig* InteractionControlRig = GetControlRig(true, ShapeActor->ControlRigIndex);
 
 			FTransform NewTransform = CurrentTransform.GetRelativeTransform(ToWorldTransform);
 			FRigControlModifiedContext Context;
 
-			if (FRigControlElement* ControlElement = InteractionControlRig->GetHierarchy()->Find<FRigControlElement>(FRigElementKey(GizmoActor->ControlName, ERigElementType::Control)))
+			if (FRigControlElement* ControlElement = InteractionControlRig->GetHierarchy()->Find<FRigControlElement>(FRigElementKey(ShapeActor->ControlName, ERigElementType::Control)))
 			{
 				// do not setup undo for this first step since it is just used to calculate the local transform
-				InteractionControlRig->GetHierarchy()->SetControlGizmoTransform(ControlElement, NewTransform, ERigTransformType::CurrentGlobal, false);
-				FTransform CurrentLocalGizmoTransform = InteractionControlRig->GetHierarchy()->GetControlGizmoTransform(ControlElement, ERigTransformType::CurrentLocal);
+				InteractionControlRig->GetHierarchy()->SetControlShapeTransform(ControlElement, NewTransform, ERigTransformType::CurrentGlobal, false);
+				FTransform CurrentLocalShapeTransform = InteractionControlRig->GetHierarchy()->GetControlShapeTransform(ControlElement, ERigTransformType::CurrentLocal);
 				// this call should trigger an instance-to-BP update in ControlRigEditor
-				InteractionControlRig->GetHierarchy()->SetControlGizmoTransform(ControlElement, CurrentLocalGizmoTransform, ERigTransformType::InitialLocal, true);
+				InteractionControlRig->GetHierarchy()->SetControlShapeTransform(ControlElement, CurrentLocalShapeTransform, ERigTransformType::InitialLocal, true);
 
 				FTransform MeshTransform = FTransform::Identity;
-				FTransform GizmoTransform = CurrentLocalGizmoTransform;
+				FTransform ShapeTransform = CurrentLocalShapeTransform;
 
 				if (UControlRig* ControlRig = GetControlRig(true))
 				{
-					if (const UControlRigGizmoLibrary* GizmoLibrary = GetControlRig(true)->GetGizmoLibrary())
+					if (const FControlRigShapeDefinition* Gizmo = UControlRigShapeLibrary::GetShapeByName(ControlElement->Settings.ShapeName, GetControlRig(true)->GetShapeLibraries()))
 					{
-						if (const FControlRigGizmoDefinition* Gizmo = GizmoLibrary->GetGizmoByName(ControlElement->Settings.GizmoName, true /* use default */))
-						{
-							MeshTransform = Gizmo->Transform;
-						}
+						MeshTransform = Gizmo->Transform;
 					}
 				}
 
-				GizmoActor->StaticMeshComponent->SetRelativeTransform(MeshTransform * GizmoTransform);
+				ShapeActor->StaticMeshComponent->SetRelativeTransform(MeshTransform * ShapeTransform);
 			}
 		}
 	} 
@@ -2840,14 +2837,14 @@ bool IsSupportedControlType(const ERigControlType ControlType)
 	return false;
 }
 
-bool FControlRigEditMode::ModeSupportedByGizmoActor(const AControlRigGizmoActor* GizmoActor, UE::Widget::EWidgetMode InMode) const
+bool FControlRigEditMode::ModeSupportedByShapeActor(const AControlRigShapeActor* ShapeActor, UE::Widget::EWidgetMode InMode) const
 {
-	if (UControlRig* ControlRig = GetControlRig(true, GizmoActor->ControlRigIndex))
+	if (UControlRig* ControlRig = GetControlRig(true, ShapeActor->ControlRigIndex))
 	{
-		const FRigControlElement* ControlElement = ControlRig->FindControl(GizmoActor->ControlName);
+		const FRigControlElement* ControlElement = ControlRig->FindControl(ShapeActor->ControlName);
 		if (ControlElement)
 		{
-			if (bIsChangingGizmoTransform)
+			if (bIsChangingControlShapeTransform)
 			{
 				return true;
 			}
@@ -2938,36 +2935,36 @@ bool FControlRigEditMode::ModeSupportedByGizmoActor(const AControlRigGizmoActor*
 	return false;
 }
 
-void FControlRigEditMode::TickGizmo(AControlRigGizmoActor* GizmoActor, const FTransform& ComponentTransform)
+void FControlRigEditMode::TickControlShape(AControlRigShapeActor* ShapeActor, const FTransform& ComponentTransform)
 {
 	UControlRigEditModeSettings* Settings = GetMutableDefault<UControlRigEditModeSettings>();
-	if (GizmoActor)
+	if (ShapeActor)
 	{
-		if (UControlRig* ControlRig = GetControlRig(true, GizmoActor->ControlRigIndex))
+		if (UControlRig* ControlRig = GetControlRig(true, ShapeActor->ControlRigIndex))
 		{
-			const FTransform Transform = ControlRig->GetControlGlobalTransform(GizmoActor->ControlName);
-			GizmoActor->SetActorTransform(Transform * ComponentTransform);
+			const FTransform Transform = ControlRig->GetControlGlobalTransform(ShapeActor->ControlName);
+			ShapeActor->SetActorTransform(Transform * ComponentTransform);
 
-			if (FRigControlElement* ControlElement = ControlRig->FindControl(GizmoActor->ControlName))
+			if (FRigControlElement* ControlElement = ControlRig->FindControl(ShapeActor->ControlName))
 			{
-				GizmoActor->SetGizmoColor(ControlElement->Settings.GizmoColor);
-				GizmoActor->SetIsTemporarilyHiddenInEditor(!ControlElement->Settings.bGizmoVisible || Settings->bHideManipulators);
+				ShapeActor->SetShapeColor(ControlElement->Settings.ShapeColor);
+				ShapeActor->SetIsTemporarilyHiddenInEditor(!ControlElement->Settings.bShapeVisible || Settings->bHideManipulators);
 				if (!IsInLevelEditor()) //don't change this in level editor otherwise we can never select it
 				{
-					GizmoActor->SetSelectable(ControlElement->Settings.bGizmoVisible && !Settings->bHideManipulators && ControlElement->Settings.bAnimatable);
+					ShapeActor->SetSelectable(ControlElement->Settings.bShapeVisible && !Settings->bHideManipulators && ControlElement->Settings.bAnimatable);
 				}
 			}
 		}
 	}
 }
 
-AControlRigGizmoActor* FControlRigEditMode::GetGizmoFromControlName(const FName& ControlName) const
+AControlRigShapeActor* FControlRigEditMode::GetControlShapeFromControlName(const FName& ControlName) const
 {
-	for (AControlRigGizmoActor* GizmoActor : GizmoActors)
+	for (AControlRigShapeActor* ShapeActor : ShapeActors)
 	{
-		if (GizmoActor->ControlName == ControlName)
+		if (ShapeActor->ControlName == ControlName)
 		{
-			return GizmoActor;
+			return ShapeActor;
 		}
 	}
 
@@ -3037,9 +3034,9 @@ void FControlRigEditMode::TickManipulatableObjects(float DeltaTime)
 	PostPoseUpdate();
 }
 
-bool FControlRigEditMode::CreateGizmoActors(UWorld* World)
+bool FControlRigEditMode::CreateShapeActors(UWorld* World)
 {
-	DestroyGizmosActors();
+	DestroyShapesActors();
 
 	FActorSpawnParameters ActorSpawnParameters;
 	ActorSpawnParameters.bTemporaryEditorActor = true;
@@ -3054,41 +3051,38 @@ bool FControlRigEditMode::CreateGizmoActors(UWorld* World)
 		}
 
 		TArray<FRigControlElement*> Controls = ControlRig->AvailableControls();
-		UControlRigGizmoLibrary* GizmoLibrary = ControlRig->GetGizmoLibrary();
+		const TArray<TSoftObjectPtr<UControlRigShapeLibrary>> ShapeLibraries = ControlRig->GetShapeLibraries();
 
 		for (FRigControlElement* ControlElement : Controls)
 		{
-			if (!ControlElement->Settings.bGizmoEnabled)
+			if (!ControlElement->Settings.bShapeEnabled)
 			{
 				continue;
 			}
 			if (IsSupportedControlType(ControlElement->Settings.ControlType))
 			{
-				FGizmoActorCreationParam Param;
+				FControlShapeActorCreationParam Param;
 				Param.ManipObj = ControlRig;
 				Param.ControlRigIndex = ControlRigIndex;
 				Param.ControlName = ControlElement->GetName();
 				Param.SpawnTransform = ControlRig->GetControlGlobalTransform(ControlElement->GetName());
-				Param.GizmoTransform = ControlRig->GetHierarchy()->GetControlGizmoTransform(ControlElement, ERigTransformType::CurrentLocal);
+				Param.ShapeTransform = ControlRig->GetHierarchy()->GetControlShapeTransform(ControlElement, ERigTransformType::CurrentLocal);
 				Param.bSelectable = ControlElement->Settings.bAnimatable;
 
-				if (GizmoLibrary)
+				if (const FControlRigShapeDefinition* ShapeDef = UControlRigShapeLibrary::GetShapeByName(ControlElement->Settings.ShapeName, ShapeLibraries))
 				{
-					if (const FControlRigGizmoDefinition* Gizmo = GizmoLibrary->GetGizmoByName(ControlElement->Settings.GizmoName, true /* use default */))
-					{
-						Param.MeshTransform = Gizmo->Transform;
-						Param.StaticMesh = Gizmo->StaticMesh;
-						Param.Material = GizmoLibrary->DefaultMaterial;
-						Param.ColorParameterName = GizmoLibrary->MaterialColorParameter;
-					}
+					Param.MeshTransform = ShapeDef->Transform;
+					Param.StaticMesh = ShapeDef->StaticMesh;
+					Param.Material = ShapeDef->Library->DefaultMaterial;
+					Param.ColorParameterName = ShapeDef->Library->MaterialColorParameter;
 				}
 
-				Param.Color = ControlElement->Settings.GizmoColor;
+				Param.Color = ControlElement->Settings.ShapeColor;
 
-				AControlRigGizmoActor* GizmoActor = FControlRigGizmoHelper::CreateDefaultGizmoActor(World, Param);
-				if (GizmoActor)
+				AControlRigShapeActor* ShapeActor = FControlRigShapeHelper::CreateDefaultShapeActor(World, Param);
+				if (ShapeActor)
 				{
-					GizmoActors.Add(GizmoActor);
+					ShapeActors.Add(ShapeActor);
 				}
 			}
 		}
@@ -3096,7 +3090,7 @@ bool FControlRigEditMode::CreateGizmoActors(UWorld* World)
 
 	WorldPtr = World;
 	OnWorldCleanupHandle = FWorldDelegates::OnWorldCleanup.AddSP(this, &FControlRigEditMode::OnWorldCleanup);
-	return (GizmoActors.Num() > 0);
+	return (ShapeActors.Num() > 0);
 }
 
 void FControlRigEditMode::OnWorldCleanup(UWorld* World, bool bSessionEnded, bool bCleanupResources)
@@ -3104,21 +3098,21 @@ void FControlRigEditMode::OnWorldCleanup(UWorld* World, bool bSessionEnded, bool
 	// if world gets cleaned up first, we destroy gizmo actors
 	if (WorldPtr == World)
 	{
-		DestroyGizmosActors();
+		DestroyShapesActors();
 	}
 }
 
-void FControlRigEditMode::DestroyGizmosActors()
+void FControlRigEditMode::DestroyShapesActors()
 {
-	for (AControlRigGizmoActor* GizmoActor : GizmoActors)
+	for (AControlRigShapeActor* ShapeActor : ShapeActors)
 	{
-		UWorld* World = GizmoActor->GetWorld();
+		UWorld* World = ShapeActor->GetWorld();
 		if (World)
 		{
-			World->DestroyActor(GizmoActor);
+			World->DestroyActor(ShapeActor);
 		}
 	}
-	GizmoActors.Reset();
+	ShapeActors.Reset();
 
 	if (OnWorldCleanupHandle.IsValid())
 	{
@@ -3175,9 +3169,9 @@ void FControlRigEditMode::OnPoseInitialized()
 void FControlRigEditMode::PostPoseUpdate()
 {
 	FTransform ComponentTransform = GetHostingSceneComponentTransform();
-	for(AControlRigGizmoActor* GizmoActor : GizmoActors)
+	for(AControlRigShapeActor* ShapeActor : ShapeActors)
 	{
-		TickGizmo(GizmoActor, ComponentTransform);
+		TickControlShape(ShapeActor, ComponentTransform);
 	}
 
 }
diff --git a/Engine/Plugins/Animation/ControlRig/Source/ControlRigEditor/Private/EditMode/ControlRigEditMode.h b/Engine/Plugins/Animation/ControlRig/Source/ControlRigEditor/Private/EditMode/ControlRigEditMode.h
--- a/Engine/Plugins/Animation/ControlRig/Source/ControlRigEditor/Private/EditMode/ControlRigEditMode.h
+++ b/Engine/Plugins/Animation/ControlRig/Source/ControlRigEditor/Private/EditMode/ControlRigEditMode.h
@@ -29,7 +29,7 @@ class FPrimitiveDrawInterface;
 class FToolBarBuilder;
 class FExtender;
 class IMovieScenePlayer;
-class AControlRigGizmoActor;
+class AControlRigShapeActor;
 class UDefaultControlRigManipulationLayer;
 class UControlRigDetailPanelControlProxies;
 class UControlRigControlsProxy;
@@ -140,19 +140,19 @@ public:
 	void OnControlModified(UControlRig* Subject, FRigControlElement* InControlElement, const FRigControlModifiedContext& Context);
 
 	/** return true if it can be removed from preview scene 
-	- this is to ensure preview scene doesn't remove Gizmo actors */
+	- this is to ensure preview scene doesn't remove shape actors */
 	bool CanRemoveFromPreviewScene(const USceneComponent* InComponent);
 
 	FUICommandList* GetCommandBindings() const { return CommandBindings.Get(); }
 
-	/** Requests to recreate the gizmo actors in the next tick */
-	void RequestToRecreateGizmoActors() { bRecreateGizmosRequired = true; }
+	/** Requests to recreate the shape actors in the next tick */
+	void RequestToRecreateControlShapeActors() { bRecreateControlShapesRequired = true; }
 
 protected:
 
-	// Gizmo related functions wrt enable/selection
+	// shape related functions wrt enable/selection
 	/** Get the node name from the property path */
-	AControlRigGizmoActor* GetGizmoFromControlName(const FName& InControlName) const;
+	AControlRigShapeActor* GetControlShapeFromControlName(const FName& InControlName) const;
 
 protected:
 	/** Helper function: set ControlRigs array to the details panel */
@@ -165,7 +165,7 @@ protected:
 	void RecalcPivotTransform();
 
 	/** Helper function for box/frustum intersection */
-	bool IntersectSelect(bool InSelect, const TFunctionRef<bool(const AControlRigGizmoActor*, const FTransform&)>& Intersects);
+	bool IntersectSelect(bool InSelect, const TFunctionRef<bool(const AControlRigShapeActor*, const FTransform&)>& Intersects);
 
 	/** Handle selection internally */
 	void HandleSelectionChanged();
@@ -195,38 +195,37 @@ private:
 	/** Reset Transforms */
 	void ResetTransforms(bool bSelectionOnly);
 
-	/** Increase Gizmo Size */
-	void IncreaseGizmoSize();
+	/** Increase Shape Size */
+	void IncreaseShapeSize();
 
-	/** Decrease Gizmo Size */
-	void DecreaseGizmoSize();
+	/** Decrease Shape Size */
+	void DecreaseShapeSize();
 
-	/** Reset Gizmo Size */
-	void ResetGizmoSize();
+	/** Reset Shape Size */
+	void ResetControlShapeSize();
 
 public:
 	
-	/** Toggle Gizmo Transform Edit*/
-	void ToggleGizmoTransformEdit();
+	/** Toggle Shape Transform Edit*/
+	void ToggleControlShapeTransformEdit();
 
 private:
 	
-	/** The hotkey text is passed to a viewport notification to inform users how to toggle gizmo edit*/
-	FText GetToggleGizmoTransformEditHotKey() const;
+	/** The hotkey text is passed to a viewport notification to inform users how to toggle shape edit*/
+	FText GetToggleControlShapeTransformEditHotKey() const;
 
 	/** Bind our keyboard commands */
 	void BindCommands();
 
 	/** It creates if it doesn't have it */
-	void RecreateGizmoActors(const TArray<FRigElementKey>& InSelectedElements = TArray<FRigElementKey>());
-
+	void RecreateControlShapeActors(const TArray<FRigElementKey>& InSelectedElements = TArray<FRigElementKey>());
 
 
 	/** Let the preview scene know how we want to select components */
-	bool GizmoSelectionOverride(const UPrimitiveComponent* InComponent) const;
+	bool ShapeSelectionOverride(const UPrimitiveComponent* InComponent) const;
 
-	/** Enable editing of control's gizmo transform instead of control's transform*/
-	bool bIsChangingGizmoTransform;
+	/** Enable editing of control's shape transform instead of control's transform*/
+	bool bIsChangingControlShapeTransform;
 
 protected:
 
@@ -266,11 +265,11 @@ protected:
 	/** GetSelectedRigElements */
 	TArray<FRigElementKey> GetSelectedRigElements() const;
 
-	/* Flag to recreate gizmos during tick */
-	bool bRecreateGizmosRequired;
+	/* Flag to recreate shapes during tick */
+	bool bRecreateControlShapesRequired;
 
-	/** Gizmo actors */
-	TArray<AControlRigGizmoActor*> GizmoActors;
+	/** Shape actors */
+	TArray<AControlRigShapeActor*> ShapeActors;
 	UControlRigDetailPanelControlProxies* ControlProxy;
 
 	/** Utility functions for UI/Some other viewport manipulation*/
@@ -316,7 +315,7 @@ private:
 	TArray<ECoordSystem> CoordSystemPerWidgetMode;
 	bool bIsChangingCoordSystem;
 
-	bool CanChangeControlGizmoTransform();
+	bool CanChangeControlShapeTransform();
 public:
 	//Toolbar functions
 	void SetOnlySelectRigControls(bool val);
@@ -324,24 +323,24 @@ public:
 
 private:
 	TSet<FName> GetActiveControlsFromSequencer(UControlRig* ControlRig);
-	bool CreateGizmoActors(UWorld* World);
-	void DestroyGizmosActors();
+	bool CreateShapeActors(UWorld* World);
+	void DestroyShapesActors();
 
 	void AddControlRig(UControlRig* InControlRig);
 	void RemoveControlRig(UControlRig* InControlRig);
 	void TickManipulatableObjects(float DeltaTime);
 
-	void SetGizmoTransform(AControlRigGizmoActor* GizmoActor, const FTransform& InTransform);
-	FTransform GetGizmoTransform(AControlRigGizmoActor* GizmoActor) const;
-	void MoveGizmo(AControlRigGizmoActor* GizmoActor, const bool bTranslation, FVector& InDrag, 
+	void SetControlShapeTransform(AControlRigShapeActor* ShapeActor, const FTransform& InTransform);
+	FTransform GetControlShapeTransform(AControlRigShapeActor* ShapeActor) const;
+	void MoveControlShape(AControlRigShapeActor* ShapeActor, const bool bTranslation, FVector& InDrag, 
 		const bool bRotation, FRotator& InRot, const bool bScale, FVector& InScale, const FTransform& ToWorldTransform,
 		bool bUseLocal, bool bCalcLocal, FTransform& InOutLocal);
 
-	void ChangeControlGizmoTransform(AControlRigGizmoActor* GizmoActor, const bool bTranslation, FVector& InDrag,
+	void ChangeControlShapeTransform(AControlRigShapeActor* ShapeActor, const bool bTranslation, FVector& InDrag,
 		const bool bRotation, FRotator& InRot, const bool bScale, FVector& InScale, const FTransform& ToWorldTransform);
 
-	void TickGizmo(AControlRigGizmoActor* GizmoActor, const FTransform& ComponentTransform);
-	bool ModeSupportedByGizmoActor(const AControlRigGizmoActor* GizmoActor, UE::Widget::EWidgetMode InMode) const;
+	void TickControlShape(AControlRigShapeActor* ShapeActor, const FTransform& ComponentTransform);
+	bool ModeSupportedByShapeActor(const AControlRigShapeActor* ShapeActor, UE::Widget::EWidgetMode InMode) const;
 
 	// Object binding
 	/** Setup bindings to a runtime object (or clear by passing in nullptr). */
diff --git a/Engine/Plugins/Animation/ControlRig/Source/ControlRigEditor/Private/EditMode/ControlRigEditModeCommands.cpp b/Engine/Plugins/Animation/ControlRig/Source/ControlRigEditor/Private/EditMode/ControlRigEditModeCommands.cpp
--- a/Engine/Plugins/Animation/ControlRig/Source/ControlRigEditor/Private/EditMode/ControlRigEditModeCommands.cpp
+++ b/Engine/Plugins/Animation/ControlRig/Source/ControlRigEditor/Private/EditMode/ControlRigEditModeCommands.cpp
@@ -12,11 +12,11 @@ void FControlRigEditModeCommands::RegisterCommands()
 	UI_COMMAND(FrameSelection, "Frame Selection", "Focus the viewport on the current selection", EUserInterfaceActionType::Button, FInputChord(EKeys::F));
 	UI_COMMAND(ClearSelection, "Clear Selection", "Clear Selection", EUserInterfaceActionType::Button, FInputChord(EKeys::Escape));
 
-	UI_COMMAND(IncreaseGizmoSize, "Increase Gizmo Size", "Increase Size of the Gizmos In The Viewport", EUserInterfaceActionType::Button, FInputChord(EKeys::Equals, EModifierKey::Shift));
-	UI_COMMAND(DecreaseGizmoSize, "Decrease Gizmo Size", "Decrease Size of the Gizmos In The Viewport", EUserInterfaceActionType::Button, FInputChord(EKeys::Hyphen, EModifierKey::Shift));
-	UI_COMMAND(ResetGizmoSize, "Reset Gizmo Size", "Resize Gizmo Size To Default", EUserInterfaceActionType::Button, FInputChord(EKeys::Equals));
+	UI_COMMAND(IncreaseControlShapeSize, "Increase Shape Size", "Increase Size of the Shapes In The Viewport", EUserInterfaceActionType::Button, FInputChord(EKeys::Equals, EModifierKey::Shift));
+	UI_COMMAND(DecreaseControlShapeSize, "Decrease Shape Size", "Decrease Size of the Shapes In The Viewport", EUserInterfaceActionType::Button, FInputChord(EKeys::Hyphen, EModifierKey::Shift));
+	UI_COMMAND(ResetControlShapeSize, "Reset Shape Size", "Resize Shape Size To Default", EUserInterfaceActionType::Button, FInputChord(EKeys::Equals));
 
-	UI_COMMAND(ToggleGizmoTransformEdit, "Toggle Gizmo Transform Edit", "Toggle Editing Selected Control's Gizmo Transform", EUserInterfaceActionType::ToggleButton, FInputChord(EKeys::Period, EModifierKey::Control)); 
+	UI_COMMAND(ToggleControlShapeTransformEdit, "Toggle Shape Transform Edit", "Toggle Editing Selected Control's Shape Transform", EUserInterfaceActionType::ToggleButton, FInputChord(EKeys::Period, EModifierKey::Control)); 
 	UI_COMMAND(OpenSpacePickerWidget, "Open the Space Picker", "Allows space switching on the control", EUserInterfaceActionType::Button, FInputChord(EKeys::Q, EModifierKey::Control)); 
 }
 
diff --git a/Engine/Plugins/Animation/ControlRig/Source/ControlRigEditor/Private/EditMode/ControlRigEditModeCommands.h b/Engine/Plugins/Animation/ControlRig/Source/ControlRigEditor/Private/EditMode/ControlRigEditModeCommands.h
--- a/Engine/Plugins/Animation/ControlRig/Source/ControlRigEditor/Private/EditMode/ControlRigEditModeCommands.h
+++ b/Engine/Plugins/Animation/ControlRig/Source/ControlRigEditor/Private/EditMode/ControlRigEditModeCommands.h
@@ -35,17 +35,17 @@ public:
 	/** Frame selected elements */
 	TSharedPtr<FUICommandInfo> FrameSelection;
 
-	/** Increase Gizmo Size */
-	TSharedPtr< FUICommandInfo > IncreaseGizmoSize;
+	/** Increase Shape Size */
+	TSharedPtr< FUICommandInfo > IncreaseControlShapeSize;
 
-	/** Decrease Gizmo Size */
-	TSharedPtr< FUICommandInfo > DecreaseGizmoSize;
+	/** Decrease Shape Size */
+	TSharedPtr< FUICommandInfo > DecreaseControlShapeSize;
 
-	/** Reset Gizmo Size */
-	TSharedPtr< FUICommandInfo > ResetGizmoSize;
+	/** Reset Shape Size */
+	TSharedPtr< FUICommandInfo > ResetControlShapeSize;
 
-	/** Toggle Gizmo Transform Edit*/
-	TSharedPtr< FUICommandInfo > ToggleGizmoTransformEdit;
+	/** Toggle Shape Transform Edit*/
+	TSharedPtr< FUICommandInfo > ToggleControlShapeTransformEdit;
 
 	/** Opens up the space picker widget */
 	TSharedPtr< FUICommandInfo > OpenSpacePickerWidget;
diff --git a/Engine/Plugins/Animation/ControlRig/Source/ControlRigEditor/Private/EditMode/ControlRigEditModeSettings.cpp b/Engine/Plugins/Animation/ControlRig/Source/ControlRigEditor/Private/EditMode/ControlRigEditModeSettings.cpp
--- a/Engine/Plugins/Animation/ControlRig/Source/ControlRigEditor/Private/EditMode/ControlRigEditModeSettings.cpp
+++ b/Engine/Plugins/Animation/ControlRig/Source/ControlRigEditor/Private/EditMode/ControlRigEditModeSettings.cpp
@@ -13,10 +13,10 @@ void UControlRigEditModeSettings::PostEditChangeProperty(struct FPropertyChanged
 {
 	Super::PostEditChangeProperty(PropertyChangedEvent);
 #if WITH_EDITOR
-	if (PropertyChangedEvent.Property && PropertyChangedEvent.Property->GetFName() == GET_MEMBER_NAME_CHECKED(UControlRigEditModeSettings, GizmoScale))
+	if (PropertyChangedEvent.Property && PropertyChangedEvent.Property->GetFName() == GET_MEMBER_NAME_CHECKED(UControlRigEditModeSettings, ShapeScale))
 	{
 		FEditorModeTools& Tools = GLevelEditorModeTools();
-		Tools.SetWidgetScale(GizmoScale);
+		Tools.SetWidgetScale(ShapeScale);
 	}
 #endif
 	if (PropertyChangedEvent.ChangeType != EPropertyChangeType::Interactive)
@@ -30,7 +30,7 @@ void UControlRigEditModeSettings::PostEditChangeProperty(struct FPropertyChanged
 void UControlRigEditModeSettings::PostEditUndo()
 {
 	FEditorModeTools& Tools = GLevelEditorModeTools();
-	Tools.SetWidgetScale(GizmoScale);
+	Tools.SetWidgetScale(ShapeScale);
 }
 #endif
 
diff --git a/Engine/Plugins/Animation/ControlRig/Source/ControlRigEditor/Private/EditMode/ControlRigEditModeSettings.h b/Engine/Plugins/Animation/ControlRig/Source/ControlRigEditor/Private/EditMode/ControlRigEditModeSettings.h
--- a/Engine/Plugins/Animation/ControlRig/Source/ControlRigEditor/Private/EditMode/ControlRigEditModeSettings.h
+++ b/Engine/Plugins/Animation/ControlRig/Source/ControlRigEditor/Private/EditMode/ControlRigEditModeSettings.h
@@ -22,7 +22,7 @@ class UControlRigEditModeSettings : public UObject
 		, bCoordSystemPerWidgetMode(true)
 		, bOnlySelectRigControls(false)
 		, bLocalTransformsInEachLocalSpace(true)
-		, GizmoScale(1.0f)
+		, ShapeScale(1.0f)
 	{}
 
 	// UObject interface
@@ -67,5 +67,5 @@ public:
 	
 	/** The scale for Gizmos */
 	UPROPERTY(config, EditAnywhere, AdvancedDisplay, Category = "Animation")
-	float GizmoScale;
+	float ShapeScale;
 };
\ No newline at end of file
diff --git a/Engine/Plugins/Animation/ControlRig/Source/ControlRigEditor/Private/Editor/ControlRigEditor.cpp b/Engine/Plugins/Animation/ControlRig/Source/ControlRigEditor/Private/Editor/ControlRigEditor.cpp
--- a/Engine/Plugins/Animation/ControlRig/Source/ControlRigEditor/Private/Editor/ControlRigEditor.cpp
+++ b/Engine/Plugins/Animation/ControlRig/Source/ControlRigEditor/Private/Editor/ControlRigEditor.cpp
@@ -1055,7 +1055,7 @@ void FControlRigEditor::ToggleSetupMode()
 	{
 		if (UControlRigBlueprint* RigBlueprint = Cast<UControlRigBlueprint>(GetBlueprintObj()))
 		{
-			EditMode->RecreateGizmoActors(RigBlueprint->Hierarchy->GetSelectedKeys());
+			EditMode->RecreateControlShapeActors(RigBlueprint->Hierarchy->GetSelectedKeys());
 		}
 
 		UControlRigEditModeSettings* Settings = GetMutableDefault<UControlRigEditModeSettings>();
@@ -1833,7 +1833,7 @@ void FControlRigEditor::Compile()
 		}
 
 		// clear transient controls such that we don't leave
-		// a phantom gizmo in the viewport
+		// a phantom shape in the viewport
 		// have to do this before compile() because during compile
 		// a new control rig instance is created without the transient controls
 		// so clear is never called for old transient controls
@@ -1934,7 +1934,7 @@ void FControlRigEditor::Compile()
 
 		if (FControlRigEditMode* EditMode = GetEditMode())
 		{
-			EditMode->RecreateGizmoActors(RigBlueprint->Hierarchy->GetSelectedKeys());
+			EditMode->RecreateControlShapeActors(RigBlueprint->Hierarchy->GetSelectedKeys());
 		}
 	}
 
@@ -2382,7 +2382,7 @@ void FControlRigEditor::PostTransaction(bool bSuccess, const FTransaction* Trans
 
 		if (FControlRigEditMode* EditMode = GetEditMode())
 		{
-			EditMode->RequestToRecreateGizmoActors();
+			EditMode->RequestToRecreateControlShapeActors();
 		}
 	}
 }
@@ -2773,10 +2773,10 @@ void FControlRigEditor::HandleControlRigExecutedEvent(UControlRig* InControlRig,
 							Hierarchy->GetControlOffsetTransform(ControlElement, ERigTransformType::CurrentLocal);
 							Hierarchy->GetControlOffsetTransform(ControlElement, ERigTransformType::InitialGlobal);
 							Hierarchy->GetControlOffsetTransform(ControlElement, ERigTransformType::InitialLocal);
-							Hierarchy->GetControlGizmoTransform(ControlElement, ERigTransformType::CurrentGlobal);
-							Hierarchy->GetControlGizmoTransform(ControlElement, ERigTransformType::CurrentLocal);
-							Hierarchy->GetControlGizmoTransform(ControlElement, ERigTransformType::InitialGlobal);
-							Hierarchy->GetControlGizmoTransform(ControlElement, ERigTransformType::InitialLocal);
+							Hierarchy->GetControlShapeTransform(ControlElement, ERigTransformType::CurrentGlobal);
+							Hierarchy->GetControlShapeTransform(ControlElement, ERigTransformType::CurrentLocal);
+							Hierarchy->GetControlShapeTransform(ControlElement, ERigTransformType::InitialGlobal);
+							Hierarchy->GetControlShapeTransform(ControlElement, ERigTransformType::InitialLocal);
 						}
 
 						Struct->CopyScriptStruct(WrapperObject->GetContent<FRigBaseElement>(), Element);
@@ -3231,28 +3231,28 @@ void FControlRigEditor::HandleViewportCreated(const TSharedRef<class IPersonaVie
 		return FEditorFontGlyphs::Eye;
 	};
 
-	auto GetChangingGizmoTransformText = [this]()
+	auto GetChangingShapeTransformText = [this]()
 	{
 		if (FControlRigEditMode* EditMode = GetEditMode())
 		{
-			FText HotKeyText = EditMode->GetToggleGizmoTransformEditHotKey();
+			FText HotKeyText = EditMode->GetToggleControlShapeTransformEditHotKey();
 
 			if (!HotKeyText.IsEmpty())
 			{
 				FFormatNamedArguments Args;
 				Args.Add(TEXT("HotKey"), HotKeyText);
-				return FText::Format(LOCTEXT("ControlRigBPViewportGizmoTransformEditNotification", "Currently Manipulating Gizmo Transform - Press {HotKey} to Exit"), Args);
+				return FText::Format(LOCTEXT("ControlRigBPViewportShapeTransformEditNotification", "Currently Manipulating Shape Transform - Press {HotKey} to Exit"), Args);
 			}
 		}
 		
-		return LOCTEXT("ControlRigBPViewportGizmoTransformEditNotification", "Currently Manipulating Gizmo Transform - Assign a Hotkey and Use It to Exit");
+		return LOCTEXT("ControlRigBPViewportShapeTransformEditNotification", "Currently Manipulating Shape Transform - Assign a Hotkey and Use It to Exit");
 	};
 
-	auto GetChangingGizmoTransformTextVisibility = [this]()
+	auto GetChangingShapeTransformTextVisibility = [this]()
 	{
 		if (FControlRigEditMode* EditMode = GetEditMode())
 		{
-			return EditMode->bIsChangingGizmoTransform ? EVisibility::Visible : EVisibility::Collapsed;
+			return EditMode->bIsChangingControlShapeTransform ? EVisibility::Visible : EVisibility::Collapsed;
 		}
 		return EVisibility::Collapsed;
 	};
@@ -3321,31 +3321,31 @@ void FControlRigEditor::HandleViewportCreated(const TSharedRef<class IPersonaVie
 		FPersonaViewportNotificationOptions(TAttribute<EVisibility>::Create(GetCompilationStateVisibility))
 	);
 	
-	FPersonaViewportNotificationOptions ChangeGizmoTransformNotificationOptions;
-	ChangeGizmoTransformNotificationOptions.OnGetVisibility = TAttribute<EVisibility>::Create(GetChangingGizmoTransformTextVisibility);
-	ChangeGizmoTransformNotificationOptions.OnGetBrushOverride = TAttribute<const FSlateBrush*>(FControlRigEditorStyle::Get().GetBrush("ControlRig.Viewport.Notification.ChangeGizmoTransform"));
+	FPersonaViewportNotificationOptions ChangeShapeTransformNotificationOptions;
+	ChangeShapeTransformNotificationOptions.OnGetVisibility = TAttribute<EVisibility>::Create(GetChangingShapeTransformTextVisibility);
+	ChangeShapeTransformNotificationOptions.OnGetBrushOverride = TAttribute<const FSlateBrush*>(FControlRigEditorStyle::Get().GetBrush("ControlRig.Viewport.Notification.ChangeShapeTransform"));
 
-	// notification that shows when users enter the mode that allows them to change gizmo transform
+	// notification that shows when users enter the mode that allows them to change shape transform
 	InViewport->AddNotification(EMessageSeverity::Type::Info,
 		false,
 		SNew(SHorizontalBox)
-		.Visibility_Lambda(GetChangingGizmoTransformTextVisibility)
+		.Visibility_Lambda(GetChangingShapeTransformTextVisibility)
 		+SHorizontalBox::Slot()
 		.FillWidth(1.0f)
 		.Padding(4.0f, 4.0f)
 		[
 			SNew(SHorizontalBox)
-			.ToolTipText_Lambda(GetChangingGizmoTransformText)
+			.ToolTipText_Lambda(GetChangingShapeTransformText)
 			+SHorizontalBox::Slot()
 			.FillWidth(1.0f)
 			.VAlign(VAlign_Center)
 			[
 				SNew(STextBlock)
-				.Text_Lambda(GetChangingGizmoTransformText)
+				.Text_Lambda(GetChangingShapeTransformText)
 				.TextStyle(FEditorStyle::Get(), "AnimViewport.MessageText")
 			]
 		],
-		ChangeGizmoTransformNotificationOptions
+		ChangeShapeTransformNotificationOptions
 	);
 
 	InViewport->AddToolbarExtender(TEXT("AnimViewportDefaultCamera"), FMenuExtensionDelegate::CreateLambda(
@@ -3884,7 +3884,7 @@ void FControlRigEditor::HandlePreviewSceneCreated(const TSharedRef<IPersonaPrevi
 			/*
 			if (FControlRigEditMode* EditMode = GetEditMode())
 			{
-				EditMode->RequestToRecreateGizmoActors();
+				EditMode->RequestToRecreateControlShapeActors();
 			}
 			*/
 		}
@@ -3935,7 +3935,7 @@ void FControlRigEditor::UpdateControlRig()
 
 			if (UControlRig* CDO = Cast<UControlRig>(Class->GetDefaultObject()))
 			{
-				CDO->GizmoLibrary = GetControlRigBlueprint()->GizmoLibrary;
+				CDO->ShapeLibraries = GetControlRigBlueprint()->ShapeLibraries;
 			}
 
 			CacheNameLists();
@@ -4348,7 +4348,7 @@ void FControlRigEditor::OnWrappedPropertyChangedChainEvent(UDetailsViewWrapperOb
 
 		const FString PoseString = TEXT("Pose->");
 		const FString OffsetString = TEXT("Offset->");
-		const FString GizmoString = TEXT("Gizmo->");
+		const FString ShapeString = TEXT("Shape->");
 		const FString SettingsString = TEXT("Settings->");
 
 		struct Local
@@ -4437,7 +4437,7 @@ void FControlRigEditor::OnWrappedPropertyChangedChainEvent(UDetailsViewWrapperOb
 			
 			ControlRigBP->Hierarchy->SetControlOffsetTransform(ControlElement, Transform, ERigTransformType::MakeInitial(TransformType), true, true, false, true);
 		}
-		else if(PropertyPath.RemoveFromStart(GizmoString))
+		else if(PropertyPath.RemoveFromStart(ShapeString))
 		{
 			FRigControlElement* ControlElement = ControlRigBP->Hierarchy->Find<FRigControlElement>(WrappedElement->GetKey());
 			if(ControlElement == nullptr)
@@ -4447,9 +4447,9 @@ void FControlRigEditor::OnWrappedPropertyChangedChainEvent(UDetailsViewWrapperOb
 
 			ERigTransformType::Type TransformType = Local::GetTransformTypeFromPath(PropertyPath);
 			bIsInitial = bIsInitial || ERigTransformType::IsInitial(TransformType);
-			const FTransform Transform = CastChecked<FRigControlElement>(FirstWrappedElement)->Gizmo.Get(TransformType);
+			const FTransform Transform = CastChecked<FRigControlElement>(FirstWrappedElement)->Shape.Get(TransformType);
 			
-			ControlRigBP->Hierarchy->SetControlGizmoTransform(ControlElement, Transform, ERigTransformType::MakeInitial(TransformType), true, false, true);
+			ControlRigBP->Hierarchy->SetControlShapeTransform(ControlElement, Transform, ERigTransformType::MakeInitial(TransformType), true, false, true);
 		}
 		else if(PropertyPath.RemoveFromStart(SettingsString))
 		{
@@ -4985,7 +4985,7 @@ void FControlRigEditor::OnHierarchyModified_AnyThread(ERigHierarchyNotification
 				}
 				break;
 			}
-			case ERigHierarchyNotification::ControlGizmoTransformChanged:
+			case ERigHierarchyNotification::ControlShapeTransformChanged:
 			{
 				if(DetailViewShowsRigElement(Key))
 				{
@@ -4995,13 +4995,13 @@ void FControlRigEditor::OnHierarchyModified_AnyThread(ERigHierarchyNotification
 					FRigControlElement* SourceControlElement = Cast<FRigControlElement>(Element);
 					if(SourceControlElement)
 					{
-						FTransform InitialGizmoTransform = WeakHierarchy.Get()->GetControlGizmoTransform(SourceControlElement, ERigTransformType::InitialLocal);
+						FTransform InitialShapeTransform = WeakHierarchy.Get()->GetControlShapeTransform(SourceControlElement, ERigTransformType::InitialLocal);
 
-						// set current gizmo transform = initial gizmo transform so that the viewport reflects this change
-						WeakHierarchy.Get()->SetControlGizmoTransform(SourceControlElement, InitialGizmoTransform, ERigTransformType::CurrentLocal, false); 
+						// set current shape transform = initial shape transform so that the viewport reflects this change
+						WeakHierarchy.Get()->SetControlShapeTransform(SourceControlElement, InitialShapeTransform, ERigTransformType::CurrentLocal, false); 
 
-						RigBlueprint->Hierarchy->SetControlGizmoTransform(Key, WeakHierarchy.Get()->GetControlGizmoTransform(SourceControlElement, ERigTransformType::InitialLocal), true);
-						RigBlueprint->Hierarchy->SetControlGizmoTransform(Key, WeakHierarchy.Get()->GetControlGizmoTransform(SourceControlElement, ERigTransformType::CurrentLocal), false);
+						RigBlueprint->Hierarchy->SetControlShapeTransform(Key, WeakHierarchy.Get()->GetControlShapeTransform(SourceControlElement, ERigTransformType::InitialLocal), true);
+						RigBlueprint->Hierarchy->SetControlShapeTransform(Key, WeakHierarchy.Get()->GetControlShapeTransform(SourceControlElement, ERigTransformType::CurrentLocal), false);
 
 						RigBlueprint->Modify();
 						RigBlueprint->MarkPackageDirty();
diff --git a/Engine/Plugins/Animation/ControlRig/Source/ControlRigEditor/Private/Editor/SControlRigGizmoNameList.cpp b/Engine/Plugins/Animation/ControlRig/Source/ControlRigEditor/Private/Editor/SControlRigGizmoNameList.cpp
--- a/Engine/Plugins/Animation/ControlRig/Source/ControlRigEditor/Private/Editor/SControlRigGizmoNameList.cpp
+++ b/Engine/Plugins/Animation/ControlRig/Source/ControlRigEditor/Private/Editor/SControlRigGizmoNameList.cpp
@@ -10,14 +10,14 @@
 #include "ControlRigBlueprint.h"
 #include "ControlRig.h"
 
-void SControlRigGizmoNameList::Construct(const FArguments& InArgs, FRigControlElement* ControlElement, UControlRigBlueprint* InBlueprint)
+void SControlRigShapeNameList::Construct(const FArguments& InArgs, FRigControlElement* ControlElement, UControlRigBlueprint* InBlueprint)
 {
 	TArray<FRigControlElement*> ControlElements;
 	ControlElements.Add(ControlElement);
 	return Construct(InArgs, ControlElements, InBlueprint);
 }
 
-void SControlRigGizmoNameList::Construct(const FArguments& InArgs, TArray<FRigControlElement*> ControlElements, UControlRigBlueprint* InBlueprint)
+void SControlRigShapeNameList::Construct(const FArguments& InArgs, TArray<FRigControlElement*> ControlElements, UControlRigBlueprint* InBlueprint)
 {
 	this->OnGetNameListContent = InArgs._OnGetNameListContent;
 	this->ControlKeys.Reset();
@@ -46,20 +46,20 @@ void SControlRigGizmoNameList::Construct(const FArguments& InArgs, TArray<FRigCo
 		[
 			SAssignNew(NameListComboBox, SControlRigGraphPinNameListValueWidget)
 			.OptionsSource(&GetNameList())
-			.OnGenerateWidget(this, &SControlRigGizmoNameList::MakeNameListItemWidget)
-			.OnSelectionChanged(this, &SControlRigGizmoNameList::OnNameListChanged)
-			.OnComboBoxOpening(this, &SControlRigGizmoNameList::OnNameListComboBox)
+			.OnGenerateWidget(this, &SControlRigShapeNameList::MakeNameListItemWidget)
+			.OnSelectionChanged(this, &SControlRigShapeNameList::OnNameListChanged)
+			.OnComboBoxOpening(this, &SControlRigShapeNameList::OnNameListComboBox)
 			.InitiallySelectedItem(InitialSelected)
 			.Content()
 			[
 				SNew(STextBlock)
-				.Text(this, &SControlRigGizmoNameList::GetNameListText)
+				.Text(this, &SControlRigShapeNameList::GetNameListText)
 			]
 		]
 	);
 }
 
-const TArray<TSharedPtr<FString>>& SControlRigGizmoNameList::GetNameList() const
+const TArray<TSharedPtr<FString>>& SControlRigShapeNameList::GetNameList() const
 {
 	if (OnGetNameListContent.IsBound())
 	{
@@ -68,7 +68,7 @@ const TArray<TSharedPtr<FString>>& SControlRigGizmoNameList::GetNameList() const
 	return EmptyList;
 }
 
-FText SControlRigGizmoNameList::GetNameListText() const
+FText SControlRigShapeNameList::GetNameListText() const
 {
 	FName FirstName = NAME_None;
 	FText Text;
@@ -77,13 +77,13 @@ FText SControlRigGizmoNameList::GetNameListText() const
 		const int32 ControlIndex = Blueprint->Hierarchy->GetIndex(ControlKeys[KeyIndex]);
 		if (ControlIndex != INDEX_NONE)
 		{
-			const FName GizmoName = Blueprint->Hierarchy->GetChecked<FRigControlElement>(ControlIndex)->Settings.GizmoName; 
+			const FName ShapeName = Blueprint->Hierarchy->GetChecked<FRigControlElement>(ControlIndex)->Settings.ShapeName; 
 			if(KeyIndex == 0)
 			{
-				Text = FText::FromName(GizmoName);
-				FirstName = GizmoName;
+				Text = FText::FromName(ShapeName);
+				FirstName = ShapeName;
 			}
-			else if(FirstName != GizmoName)
+			else if(FirstName != ShapeName)
 			{
 				static const FString MultipleValues = TEXT("Multiple Values");
 				Text = FText::FromString(MultipleValues);
@@ -94,9 +94,9 @@ FText SControlRigGizmoNameList::GetNameListText() const
 	return Text;
 }
 
-void SControlRigGizmoNameList::SetNameListText(const FText& NewTypeInValue, ETextCommit::Type /*CommitInfo*/)
+void SControlRigShapeNameList::SetNameListText(const FText& NewTypeInValue, ETextCommit::Type /*CommitInfo*/)
 {
-	const FScopedTransaction Transaction(NSLOCTEXT("ControlRigEditor", "ChangeGizmoName", "Change Gizmo Name"));
+	const FScopedTransaction Transaction(NSLOCTEXT("ControlRigEditor", "ChangeShapeName", "Change Shape Name"));
 
 	for(int32 KeyIndex = 0; KeyIndex < ControlKeys.Num(); KeyIndex++)
 	{
@@ -107,24 +107,24 @@ void SControlRigGizmoNameList::SetNameListText(const FText& NewTypeInValue, ETex
 			URigHierarchy* Hierarchy = Blueprint->Hierarchy;
 
 			FRigControlElement* ControlElement = Hierarchy->Get<FRigControlElement>(ControlIndex);
-			if ((ControlElement != nullptr) && (ControlElement->Settings.GizmoName != NewName))
+			if ((ControlElement != nullptr) && (ControlElement->Settings.ShapeName != NewName))
 			{
 				Hierarchy->Modify();
 
 				FRigControlSettings Settings = ControlElement->Settings;
-				Settings.GizmoName = NewName;
+				Settings.ShapeName = NewName;
 				Hierarchy->SetControlSettings(ControlElement, Settings, true, true, true);
 			}
 		}
 	}
 }
 
-TSharedRef<SWidget> SControlRigGizmoNameList::MakeNameListItemWidget(TSharedPtr<FString> InItem)
+TSharedRef<SWidget> SControlRigShapeNameList::MakeNameListItemWidget(TSharedPtr<FString> InItem)
 {
 	return 	SNew(STextBlock).Text(FText::FromString(*InItem));// .Font(FEditorStyle::GetFontStyle(TEXT("PropertyWindow.NormalFont")));
 }
 
-void SControlRigGizmoNameList::OnNameListChanged(TSharedPtr<FString> NewSelection, ESelectInfo::Type SelectInfo)
+void SControlRigShapeNameList::OnNameListChanged(TSharedPtr<FString> NewSelection, ESelectInfo::Type SelectInfo)
 {
 	if (SelectInfo != ESelectInfo::Direct)
 	{
@@ -140,7 +140,7 @@ void SControlRigGizmoNameList::OnNameListChanged(TSharedPtr<FString> NewSelectio
 	}
 }
 
-void SControlRigGizmoNameList::OnNameListComboBox()
+void SControlRigShapeNameList::OnNameListComboBox()
 {
 	TSharedPtr<FString> CurrentlySelected;
 	for (TSharedPtr<FString> Item : GetNameList())
diff --git a/Engine/Plugins/Animation/ControlRig/Source/ControlRigEditor/Private/Editor/SControlRigGizmoNameList.h b/Engine/Plugins/Animation/ControlRig/Source/ControlRigEditor/Private/Editor/SControlRigGizmoNameList.h
--- a/Engine/Plugins/Animation/ControlRig/Source/ControlRigEditor/Private/Editor/SControlRigGizmoNameList.h
+++ b/Engine/Plugins/Animation/ControlRig/Source/ControlRigEditor/Private/Editor/SControlRigGizmoNameList.h
@@ -10,13 +10,13 @@
 #include "Rigs/RigHierarchyContainer.h"
 #include "ControlRigBlueprint.h"
 
-class SControlRigGizmoNameList : public SBox
+class SControlRigShapeNameList : public SBox
 {
 public:
 
 	DECLARE_DELEGATE_RetVal( const TArray<TSharedPtr<FString>>&, FOnGetNameListContent );
 
-	SLATE_BEGIN_ARGS(SControlRigGizmoNameList){}
+	SLATE_BEGIN_ARGS(SControlRigShapeNameList){}
 
 		SLATE_EVENT(FOnGetNameListContent, OnGetNameListContent)
 
diff --git a/Engine/Plugins/Animation/ControlRig/Source/ControlRigEditor/Private/Editor/SRigHierarchy.cpp b/Engine/Plugins/Animation/ControlRig/Source/ControlRigEditor/Private/Editor/SRigHierarchy.cpp
--- a/Engine/Plugins/Animation/ControlRig/Source/ControlRigEditor/Private/Editor/SRigHierarchy.cpp
+++ b/Engine/Plugins/Animation/ControlRig/Source/ControlRigEditor/Private/Editor/SRigHierarchy.cpp
@@ -355,8 +355,8 @@ void SRigHierarchy::BindCommands()
 		FCanExecuteAction::CreateSP(this, &SRigHierarchy::IsMultiSelected));
 
 	CommandList->MapAction(
-		Commands.SetGizmoTransformFromCurrent,
-		FExecuteAction::CreateSP(this, &SRigHierarchy::HandleSetGizmoTransformFromCurrent),
+		Commands.SetShapeTransformFromCurrent,
+		FExecuteAction::CreateSP(this, &SRigHierarchy::HandleSetShapeTransformFromCurrent),
 		FCanExecuteAction::CreateSP(this, &SRigHierarchy::IsControlSelected));
 
 	CommandList->MapAction(
@@ -437,10 +437,10 @@ void SRigHierarchy::BindCommands()
 		FIsActionChecked::CreateLambda([this]() { return DisplaySettings.bShowDynamicHierarchy; }));
 
 	CommandList->MapAction(
-		Commands.ToggleGizmoTransformEdit,
+		Commands.ToggleControlShapeTransformEdit,
 		FExecuteAction::CreateLambda([this]()
 		{
-			ControlRigEditor.Pin()->GetEditMode()->ToggleGizmoTransformEdit();
+			ControlRigEditor.Pin()->GetEditMode()->ToggleControlShapeTransformEdit();
 		}));
 
 	CommandList->MapAction(
@@ -1086,7 +1086,7 @@ void SRigHierarchy::CreateContextMenu() const
 					TransformsSection.AddMenuEntry(Commands.ResetAllTransforms);
 					TransformsSection.AddMenuEntry(Commands.SetInitialTransformFromCurrentTransform);
 					TransformsSection.AddMenuEntry(Commands.SetInitialTransformFromClosestBone);
-					TransformsSection.AddMenuEntry(Commands.SetGizmoTransformFromCurrent);
+					TransformsSection.AddMenuEntry(Commands.SetShapeTransformFromCurrent);
 					TransformsSection.AddMenuEntry(Commands.Unparent);
 
 					FToolMenuSection& AssetsSection = InMenu->AddSection(TEXT("Assets"), LOCTEXT("AssetsHeader", "Assets"));
@@ -1183,7 +1183,7 @@ void SRigHierarchy::RefreshHierarchy(const FAssetData& InAssetData)
 		FScopedTransaction Transaction(LOCTEXT("HierarchyRefresh", "Refresh Transform"));
 
 		// don't select bone if we are in setup mode.
-		// we do this to avoid the editmode / viewport gizmos to refresh recursively,
+		// we do this to avoid the editmode / viewport shapes to refresh recursively,
 		// which can add an extreme slowdown depending on the number of bones (n^(n-1))
 		bool bSelectBones = true;
 		if (UControlRig* CurrentRig = ControlRigEditor.Pin()->ControlRig)
@@ -1267,7 +1267,7 @@ void SRigHierarchy::ImportHierarchy(const FAssetData& InAssetData)
 		FScopedTransaction Transaction(LOCTEXT("HierarchyImport", "Import Hierarchy"));
 
 		// don't select bone if we are in setup mode.
-		// we do this to avoid the editmode / viewport gizmos to refresh recursively,
+		// we do this to avoid the editmode / viewport shapes to refresh recursively,
 		// which can add an extreme slowdown depending on the number of bones (n^(n-1))
 		bool bSelectBones = true;
 		if (UControlRig* CurrentRig = ControlRigEditor.Pin()->ControlRig)
@@ -2570,7 +2570,7 @@ void SRigHierarchy::HandleSetInitialTransformFromClosestBone()
 	}
 }
 
-void SRigHierarchy::HandleSetGizmoTransformFromCurrent()
+void SRigHierarchy::HandleSetShapeTransformFromCurrent()
 {
 	if (IsControlSelected())
 	{
@@ -2578,7 +2578,7 @@ void SRigHierarchy::HandleSetGizmoTransformFromCurrent()
 		{
 			if (URigHierarchy* DebuggedHierarchy = GetDebuggedHierarchy())
 			{
-				FScopedTransaction Transaction(LOCTEXT("HierarchySetGizmoTransforms", "Set Gizmo Transforms"));
+				FScopedTransaction Transaction(LOCTEXT("HierarchySetShapeTransforms", "Set Shape Transforms"));
 
 				TArray<TSharedPtr<FRigTreeElement>> SelectedItems = TreeView->GetSelectedItems();
 				for (const TSharedPtr<FRigTreeElement>& SelectedItem : SelectedItems)
@@ -2587,16 +2587,16 @@ void SRigHierarchy::HandleSetGizmoTransformFromCurrent()
 					{
 						const FRigElementKey Key = ControlElement->GetKey();
 						
-						if (ControlElement->Settings.bGizmoEnabled)
+						if (ControlElement->Settings.bShapeEnabled)
 						{
 							const FTransform OffsetGlobalTransform = DebuggedHierarchy->GetGlobalControlOffsetTransform(Key); 
-							const FTransform GizmoGlobalTransform = DebuggedHierarchy->GetGlobalControlGizmoTransform(Key);
-							const FTransform GizmoLocalTransform = GizmoGlobalTransform.GetRelativeTransform(OffsetGlobalTransform);
+							const FTransform ShapeGlobalTransform = DebuggedHierarchy->GetGlobalControlShapeTransform(Key);
+							const FTransform ShapeLocalTransform = ShapeGlobalTransform.GetRelativeTransform(OffsetGlobalTransform);
 							
-							DebuggedHierarchy->SetControlGizmoTransform(Key, GizmoLocalTransform, true, true);
-							DebuggedHierarchy->SetControlGizmoTransform(Key, GizmoLocalTransform, false, true);
-							GetHierarchy()->SetControlGizmoTransform(Key, GizmoLocalTransform, true, true);
-							GetHierarchy()->SetControlGizmoTransform(Key, GizmoLocalTransform, false, true);
+							DebuggedHierarchy->SetControlShapeTransform(Key, ShapeLocalTransform, true, true);
+							DebuggedHierarchy->SetControlShapeTransform(Key, ShapeLocalTransform, false, true);
+							GetHierarchy()->SetControlShapeTransform(Key, ShapeLocalTransform, true, true);
+							GetHierarchy()->SetControlShapeTransform(Key, ShapeLocalTransform, false, true);
 
 							DebuggedHierarchy->SetLocalTransform(Key, FTransform::Identity, false, true, true);
 							DebuggedHierarchy->SetLocalTransform(Key, FTransform::Identity, true, true, true);
@@ -2606,7 +2606,7 @@ void SRigHierarchy::HandleSetGizmoTransformFromCurrent()
 
 						if (FControlRigEditorEditMode* EditMode = ControlRigEditor.Pin()->GetEditMode())
 						{
-							EditMode->RequestToRecreateGizmoActors();
+							EditMode->RequestToRecreateControlShapeActors();
 						}
 					}
 				}
diff --git a/Engine/Plugins/Animation/ControlRig/Source/ControlRigEditor/Private/Editor/SRigHierarchy.h b/Engine/Plugins/Animation/ControlRig/Source/ControlRigEditor/Private/Editor/SRigHierarchy.h
--- a/Engine/Plugins/Animation/ControlRig/Source/ControlRigEditor/Private/Editor/SRigHierarchy.h
+++ b/Engine/Plugins/Animation/ControlRig/Source/ControlRigEditor/Private/Editor/SRigHierarchy.h
@@ -200,7 +200,7 @@ private:
 	void HandleResetInitialTransform();
 	void HandleSetInitialTransformFromCurrentTransform();
 	void HandleSetInitialTransformFromClosestBone();
-	void HandleSetGizmoTransformFromCurrent();
+	void HandleSetShapeTransformFromCurrent();
 	void HandleFrameSelection();
 	void HandleControlBoneOrSpaceTransform();
 	void HandleUnparent();
diff --git a/Engine/Plugins/Animation/ControlRig/Source/ControlRigEditor/Public/ControlRigThumbnailRenderer.h b/Engine/Plugins/Animation/ControlRig/Source/ControlRigEditor/Public/ControlRigThumbnailRenderer.h
--- a/Engine/Plugins/Animation/ControlRig/Source/ControlRigEditor/Public/ControlRigThumbnailRenderer.h
+++ b/Engine/Plugins/Animation/ControlRig/Source/ControlRigEditor/Public/ControlRigThumbnailRenderer.h
@@ -32,6 +32,6 @@ class UControlRigThumbnailRenderer : public USkeletalMeshThumbnailRenderer
 
 protected:
 	class UControlRigBlueprint* RigBlueprint;
-	TMap<FName, AStaticMeshActor*> GizmoActors;
+	TMap<FName, AStaticMeshActor*> ShapeActors;
 };
 
