Control Rig: Fix regression of missing shape when creating new control

diff --git a/Engine/Plugins/Animation/ControlRig/Source/ControlRig/Private/Rigs/RigHierarchyElements.cpp b/Engine/Plugins/Animation/ControlRig/Source/ControlRig/Private/Rigs/RigHierarchyElements.cpp
--- a/Engine/Plugins/Animation/ControlRig/Source/ControlRig/Private/Rigs/RigHierarchyElements.cpp
+++ b/Engine/Plugins/Animation/ControlRig/Source/ControlRig/Private/Rigs/RigHierarchyElements.cpp
@@ -4,6 +4,7 @@
 #include "Rigs/RigHierarchy.h"
 #include "Units/RigUnitContext.h"
 #include "ControlRigObjectVersion.h"
+#include "ControlRigGizmoLibrary.h"
 
 ////////////////////////////////////////////////////////////////////////////////
 // FRigBaseElement
@@ -424,6 +425,8 @@ FRigControlSettings::FRigControlSettings()
 , ControlEnum(nullptr)
 , Customization()
 {
+	// rely on the default provided by the shape definition
+	ShapeName = FControlRigShapeDefinition().ShapeName; 
 }
 
 void FRigControlSettings::Save(FArchive& Ar)
