Fix build break due to bad merge

diff --git a/Engine/Plugins/Animation/ControlRig/Source/ControlRigEditor/Private/ControlRigThumbnailRenderer.cpp b/Engine/Plugins/Animation/ControlRig/Source/ControlRigEditor/Private/ControlRigThumbnailRenderer.cpp
--- a/Engine/Plugins/Animation/ControlRig/Source/ControlRigEditor/Private/ControlRigThumbnailRenderer.cpp
+++ b/Engine/Plugins/Animation/ControlRig/Source/ControlRigEditor/Private/ControlRigThumbnailRenderer.cpp
@@ -78,8 +78,7 @@ void UControlRigThumbnailRenderer::Draw(UObject* Object, int32 X, int32 Y, uint3
 
 void UControlRigThumbnailRenderer::AddAdditionalPreviewSceneContent(UObject* Object, UWorld* PreviewWorld)
 {
-	TSharedRef<FSkeletalMeshThumbnailScene> ThumbnailScene = ThumbnailSceneCache.EnsureThumbnailScene(Object);
-	if (ThumbnailScene->GetPreviewActor() && RigBlueprint && !RigBlueprint->ShapeLibraries.IsEmpty() && RigBlueprint->GeneratedClass)
+	if (ThumbnailScene && ThumbnailScene->GetPreviewActor() && RigBlueprint && !RigBlueprint->ShapeLibraries.IsEmpty() && RigBlueprint->GeneratedClass)
 	{
 		UControlRig* ControlRig = nullptr;
 
