ControlRig: Fix a crash when click "Bake" under Spaces in the Animation Tab when a level sequence is opened

diff --git a/Engine/Plugins/Animation/ControlRig/Source/ControlRigEditor/Private/EditMode/SControlRigEditModeTools.cpp b/Engine/Plugins/Animation/ControlRig/Source/ControlRigEditor/Private/EditMode/SControlRigEditModeTools.cpp
--- a/Engine/Plugins/Animation/ControlRig/Source/ControlRigEditor/Private/EditMode/SControlRigEditModeTools.cpp
+++ b/Engine/Plugins/Animation/ControlRig/Source/ControlRigEditor/Private/EditMode/SControlRigEditModeTools.cpp
@@ -763,6 +763,10 @@ FReply SControlRigEditModeTools::OnBakeControlsToNewSpaceButtonClicked()
 	FRigSpacePickerBakeSettings Settings;
 	//Find default target space, just use first control and find space at current sequencer time
 	//Then Find range
+
+	// FindSpaceChannelAndSectionForControl() will trigger RecreateCurveEditor(), which will deselect the controls
+	// but in theory the selection will be recovered in the next tick, so here we just cache the selected controls
+	// and use it throughout this function. If this deselection is causing other problems, this part could use a revisit.
 	TArray<FRigElementKey> ControlKeys = SpacePickerWidget->GetControls();
 
 	FSpaceChannelAndSection SpaceChannelAndSection = FControlRigSpaceChannelHelpers::FindSpaceChannelAndSectionForControl(ControlRig, ControlKeys[0].Name, Sequencer, true /*bCreateIfNeeded*/);
@@ -797,7 +801,7 @@ FReply SControlRigEditModeTools::OnBakeControlsToNewSpaceButtonClicked()
 			SNew(SRigSpacePickerBakeWidget)
 			.Settings(Settings)
 			.Hierarchy(SpacePickerWidget->GetHierarchy())
-			.Controls(SpacePickerWidget->GetControls())
+			.Controls(ControlKeys) // use the cached controls here since the selection is not recovered until next tick.
 			.Sequencer(Sequencer)
 			.GetControlCustomization(this, &SControlRigEditModeTools::HandleGetControlElementCustomization)
 			.OnBake_Lambda([Sequencer, ControlRig,TickResolution](URigHierarchy* InHierarchy, TArray<FRigElementKey> InControls, FRigSpacePickerBakeSettings InSettings)
