Control Rig: Added additional guard to avoid premature control rig initialization

diff --git a/Engine/Plugins/Animation/ControlRig/Source/ControlRig/Private/ControlRig.cpp b/Engine/Plugins/Animation/ControlRig/Source/ControlRig/Private/ControlRig.cpp
--- a/Engine/Plugins/Animation/ControlRig/Source/ControlRig/Private/ControlRig.cpp
+++ b/Engine/Plugins/Animation/ControlRig/Source/ControlRig/Private/ControlRig.cpp
@@ -103,7 +103,10 @@ void UControlRig::BeginDestroy()
 	{
 		if(UControlRig* CDO = GetClass()->GetDefaultObject<UControlRig>())
 		{
-			CDO->GetHierarchy()->UnregisterListeningHierarchy(GetHierarchy());
+			if (!CDO->HasAnyFlags(RF_BeginDestroyed))
+			{
+				CDO->GetHierarchy()->UnregisterListeningHierarchy(GetHierarchy());
+			}
 		}
 	}
 #endif
diff --git a/Engine/Plugins/Animation/ControlRig/Source/ControlRig/Private/Sequencer/MovieSceneControlRigParameterTrack.cpp b/Engine/Plugins/Animation/ControlRig/Source/ControlRig/Private/Sequencer/MovieSceneControlRigParameterTrack.cpp
--- a/Engine/Plugins/Animation/ControlRig/Source/ControlRig/Private/Sequencer/MovieSceneControlRigParameterTrack.cpp
+++ b/Engine/Plugins/Animation/ControlRig/Source/ControlRig/Private/Sequencer/MovieSceneControlRigParameterTrack.cpp
@@ -377,11 +377,23 @@ void UMovieSceneControlRigParameterTrack::PostLoad()
 #if WITH_EDITOR
 void UMovieSceneControlRigParameterTrack::HandlePackageDone(TConstArrayView<UPackage*> InPackages)
 {
-	if (!InPackages.Contains(GetPackage()))
+	if (!GetPackage()->GetHasBeenEndLoaded())
 	{
 		return;
 	}
 
+	// ensure both packages are fully end-loaded
+	if (ControlRig)
+	{
+		if (const UPackage* ControlRigPackage = Cast<UPackage>(ControlRig->GetClass()->GetOutermost()))
+		{
+			if (!ControlRigPackage->GetHasBeenEndLoaded())
+			{
+				return;
+			}
+		}
+	}
+
 	FCoreUObjectDelegates::OnEndLoadPackage.RemoveAll(this);
 
 	ReconstructControlRig();
