Control Rig: Fix allowing function argument names starting with numbers

diff --git a/Engine/Plugins/Animation/ControlRig/Source/ControlRigDeveloper/Private/Graph/ControlRigGraphSchema.cpp b/Engine/Plugins/Animation/ControlRig/Source/ControlRigDeveloper/Private/Graph/ControlRigGraphSchema.cpp
--- a/Engine/Plugins/Animation/ControlRig/Source/ControlRigDeveloper/Private/Graph/ControlRigGraphSchema.cpp
+++ b/Engine/Plugins/Animation/ControlRig/Source/ControlRigDeveloper/Private/Graph/ControlRigGraphSchema.cpp
@@ -70,6 +70,26 @@ FControlRigLocalVariableNameValidator::FControlRigLocalVariableNameValidator(con
 	}
 }
 
+EValidatorResult FControlRigLocalVariableNameValidator::IsValid(const FString& Name, bool bOriginal)
+{
+	EValidatorResult Result = FStringSetNameValidator::IsValid(Name, bOriginal);
+	if (Result == EValidatorResult::Ok)
+	{
+		if (URigVMController::GetSanitizedName(Name, false, true) == Name)
+		{
+			return Result;
+		}
+
+		return EValidatorResult::ContainsInvalidCharacters;
+	}
+	return Result;
+}
+
+EValidatorResult FControlRigLocalVariableNameValidator::IsValid(const FName& Name, bool bOriginal)
+{
+	return IsValid(Name.ToString(), bOriginal);
+}
+
 FEdGraphPinType FControlRigGraphSchemaAction_LocalVar::GetPinType() const
 {
 	if (UControlRigGraph* Graph = Cast<UControlRigGraph>(GetVariableScope()))
diff --git a/Engine/Plugins/Animation/ControlRig/Source/ControlRigDeveloper/Public/Graph/ControlRigGraphSchema.h b/Engine/Plugins/Animation/ControlRig/Source/ControlRigDeveloper/Public/Graph/ControlRigGraphSchema.h
--- a/Engine/Plugins/Animation/ControlRig/Source/ControlRigDeveloper/Public/Graph/ControlRigGraphSchema.h
+++ b/Engine/Plugins/Animation/ControlRig/Source/ControlRigDeveloper/Public/Graph/ControlRigGraphSchema.h
@@ -51,6 +51,11 @@ class CONTROLRIGDEVELOPER_API FControlRigLocalVariableNameValidator : public FSt
 
 public:
 	FControlRigLocalVariableNameValidator(const class UBlueprint* Blueprint, const URigVMGraph* Graph, FName InExistingName = NAME_None);
+
+	// Begin FNameValidatorInterface
+	virtual EValidatorResult IsValid(const FString& Name, bool bOriginal) override;
+	virtual EValidatorResult IsValid(const FName& Name, bool bOriginal) override;
+	// End FNameValidatorInterface
 };
 
 USTRUCT()
diff --git a/Engine/Source/Developer/RigVMDeveloper/Public/RigVMModel/RigVMController.h b/Engine/Source/Developer/RigVMDeveloper/Public/RigVMModel/RigVMController.h
--- a/Engine/Source/Developer/RigVMDeveloper/Public/RigVMModel/RigVMController.h
+++ b/Engine/Source/Developer/RigVMDeveloper/Public/RigVMModel/RigVMController.h
@@ -881,8 +881,10 @@ private:
 
 	void RefreshFunctionReferences(URigVMLibraryNode* InFunctionDefinition, bool bSetupUndoRedo);
 
-	static int32 GetMaxNameLength() { return 100; }
 	FString GetGraphOuterName() const;
+
+public:
+	static int32 GetMaxNameLength() { return 100; }
 	static FString GetSanitizedName(const FString& InName, bool bAllowPeriod, bool bAllowSpace);
 	static FString GetSanitizedGraphName(const FString& InName);
 	static FString GetSanitizedNodeName(const FString& InName);
@@ -891,7 +893,7 @@ private:
 	static FString GetSanitizedPinPath(const FString& InName);
 	static void SanitizeName(FString& InOutName, bool bAllowPeriod, bool bAllowSpace);
 
-public:
+
 	static const TCHAR TArrayPrefix[];
 	static const TCHAR TObjectPtrPrefix[];
 	static const TCHAR TArrayTemplate[];
