Refresh editor after creating a new item

diff --git a/Engine/Plugins/Animation/ControlRig/Source/ControlRigEditor/Private/Editor/SRigHierarchy.cpp b/Engine/Plugins/Animation/ControlRig/Source/ControlRigEditor/Private/Editor/SRigHierarchy.cpp
--- a/Engine/Plugins/Animation/ControlRig/Source/ControlRigEditor/Private/Editor/SRigHierarchy.cpp
+++ b/Engine/Plugins/Animation/ControlRig/Source/ControlRigEditor/Private/Editor/SRigHierarchy.cpp
@@ -1522,6 +1522,11 @@ void SRigHierarchy::HandleNewItem(ERigElementType InElementType)
 					return;
 				}
 			}
+
+			if (ControlRigBlueprint.IsValid())
+			{
+				ControlRigBlueprint->BroadcastRefreshEditor();
+			}
 		}
 
 		Controller->ClearSelection();
