Control Rig: Add Exponential function, E as a constant, and Tau as a keyword for 2 Pi

diff --git a/Engine/Plugins/Animation/ControlRig/Source/ControlRig/Private/Units/Math/RigUnit_MathFloat.cpp b/Engine/Plugins/Animation/ControlRig/Source/ControlRig/Private/Units/Math/RigUnit_MathFloat.cpp
--- a/Engine/Plugins/Animation/ControlRig/Source/ControlRig/Private/Units/Math/RigUnit_MathFloat.cpp
+++ b/Engine/Plugins/Animation/ControlRig/Source/ControlRig/Private/Units/Math/RigUnit_MathFloat.cpp
@@ -18,6 +18,11 @@ FRigUnit_MathFloatConstTwoPi_Execute()
 	Value = PI * 2.f;
 }
 
+FRigUnit_MathFloatConstE_Execute()
+{
+	Value = EULERS_NUMBER;
+}
+
 FRigUnit_MathFloatAdd_Execute()
 {
     DECLARE_SCOPE_HIERARCHICAL_COUNTER_RIGUNIT()
@@ -303,3 +308,9 @@ FRigUnit_MathFloatLawOfCosine_Execute()
 	AlphaAngle = PI - GammaAngle - BetaAngle;
 	bValid = true;
 }
+
+FRigUnit_MathFloatExponential_Execute()
+{
+	DECLARE_SCOPE_HIERARCHICAL_COUNTER_RIGUNIT()
+	Result = FMath::Exp(Value);
+}
diff --git a/Engine/Plugins/Animation/ControlRig/Source/ControlRig/Private/Units/Math/RigUnit_MathFloat.h b/Engine/Plugins/Animation/ControlRig/Source/ControlRig/Private/Units/Math/RigUnit_MathFloat.h
--- a/Engine/Plugins/Animation/ControlRig/Source/ControlRig/Private/Units/Math/RigUnit_MathFloat.h
+++ b/Engine/Plugins/Animation/ControlRig/Source/ControlRig/Private/Units/Math/RigUnit_MathFloat.h
@@ -89,7 +89,7 @@ struct CONTROLRIG_API FRigUnit_MathFloatConstHalfPi : public FRigUnit_MathFloatC
 /**
  * Returns PI * 2.0
  */
-USTRUCT(meta=(DisplayName="Two Pi"))
+USTRUCT(meta=(DisplayName="Two Pi", Keywords="Tau"))
 struct CONTROLRIG_API FRigUnit_MathFloatConstTwoPi : public FRigUnit_MathFloatConstant
 {
 	GENERATED_BODY()
@@ -98,6 +98,18 @@ struct CONTROLRIG_API FRigUnit_MathFloatConstTwoPi : public FRigUnit_MathFloatCo
 	virtual void Execute(const FRigUnitContext& Context) override;
 };
 
+/**
+ * Returns E
+ */
+USTRUCT(meta=(DisplayName="E", Keywords="Euler"))
+struct CONTROLRIG_API FRigUnit_MathFloatConstE : public FRigUnit_MathFloatConstant
+{
+	GENERATED_BODY()
+	
+	RIGVM_METHOD()
+	virtual void Execute(const FRigUnitContext& Context) override;
+};
+
 /**
  * Returns the sum of the two values
  */
@@ -850,4 +862,16 @@ struct CONTROLRIG_API FRigUnit_MathFloatLawOfCosine : public FRigUnit_MathFloatB
 
 	UPROPERTY(meta = (Output))
 	bool bValid;
+};
+
+/**
+ * Computes the base-e exponential of the given value 
+ */
+USTRUCT(meta = (DisplayName = "Exponential"))
+struct CONTROLRIG_API FRigUnit_MathFloatExponential : public FRigUnit_MathFloatUnaryOp
+{
+	GENERATED_BODY()
+
+	RIGVM_METHOD()
+	virtual void Execute(const FRigUnitContext& Context) override;
 };
\ No newline at end of file
