Control Rig: Added a missing null check

diff --git a/Engine/Plugins/Animation/ControlRig/Source/ControlRig/Private/ControlRig.cpp b/Engine/Plugins/Animation/ControlRig/Source/ControlRig/Private/ControlRig.cpp
--- a/Engine/Plugins/Animation/ControlRig/Source/ControlRig/Private/ControlRig.cpp
+++ b/Engine/Plugins/Animation/ControlRig/Source/ControlRig/Private/ControlRig.cpp
@@ -106,7 +106,10 @@ void UControlRig::BeginDestroy()
 		{
 			if (!CDO->HasAnyFlags(RF_BeginDestroyed))
 			{
-				CDO->GetHierarchy()->UnregisterListeningHierarchy(GetHierarchy());
+				if (CDO->GetHierarchy())
+				{
+					CDO->GetHierarchy()->UnregisterListeningHierarchy(GetHierarchy());
+				}
 			}
 		}
 	}
