Control Rig: Update default objects for gizmo library changes in Postload

diff --git a/Engine/Plugins/Animation/ControlRig/Source/ControlRigDeveloper/Private/ControlRigBlueprint.cpp b/Engine/Plugins/Animation/ControlRig/Source/ControlRigDeveloper/Private/ControlRigBlueprint.cpp
--- a/Engine/Plugins/Animation/ControlRig/Source/ControlRigDeveloper/Private/ControlRigBlueprint.cpp
+++ b/Engine/Plugins/Animation/ControlRig/Source/ControlRigDeveloper/Private/ControlRigBlueprint.cpp
@@ -510,18 +510,21 @@ void UControlRigBlueprint::PostLoad()
 	}
 
 	// upgrade the gizmo libraries to shape libraries
-	if(GizmoLibrary_DEPRECATED.IsValid())
-	{
-		ShapeLibraries.Reset();
-		ShapeLibraries.Add(GizmoLibrary_DEPRECATED);
-		GizmoLibrary_DEPRECATED.Reset();
-	}
-	else if (GetLinkerCustomVersion(FControlRigObjectVersion::GUID) < FControlRigObjectVersion::RenameGizmoToShape)
+	if(GizmoLibrary_DEPRECATED.IsValid() || GetLinkerCustomVersion(FControlRigObjectVersion::GUID) < FControlRigObjectVersion::RenameGizmoToShape)
 	{
 		// if it's an older file and it doesn't have the GizmoLibrary stored,
 		// refer to the previous default.
 		ShapeLibraries.Reset();
-		ShapeLibraries.Add(LoadObject<UControlRigShapeLibrary>(nullptr, TEXT("/ControlRig/Controls/DefaultGizmoLibrary.DefaultGizmoLibrary")));
+
+		if(GizmoLibrary_DEPRECATED.IsValid())
+		{
+			ShapeLibraries.Add(GizmoLibrary_DEPRECATED.LoadSynchronous());
+			GizmoLibrary_DEPRECATED.Reset();
+		}
+		else
+		{
+			ShapeLibraries.Add(LoadObject<UControlRigShapeLibrary>(nullptr, TEXT("/ControlRig/Controls/DefaultGizmoLibrary.DefaultGizmoLibrary")));
+		}
 
 		// also walk over all controls and check if any of them were using the "default" gizmo
 		Hierarchy->ForEach<FRigControlElement>([](FRigControlElement* ControlElement) -> bool
@@ -534,6 +537,28 @@ void UControlRigBlueprint::PostLoad()
 			}
 			return true;
 		});
+
+		UControlRigBlueprintGeneratedClass* RigClass = GetControlRigBlueprintGeneratedClass();
+		UControlRig* CDO = Cast<UControlRig>(RigClass->GetDefaultObject(false /* create if needed */));
+
+		CDO->ShapeLibraries = ShapeLibraries;
+		CDO->GizmoLibrary_DEPRECATED.Reset();
+		CDO->GetHierarchy()->CopyHierarchy(Hierarchy);
+		CDO->Initialize(true);
+
+		TArray<UObject*> ArchetypeInstances;
+		CDO->GetArchetypeInstances(ArchetypeInstances);
+		for (UObject* Instance : ArchetypeInstances)
+		{
+			if (UControlRig* InstanceRig = Cast<UControlRig>(Instance))
+			{
+				InstanceRig->ShapeLibraries = ShapeLibraries;
+				InstanceRig->GizmoLibrary_DEPRECATED.Reset();
+				InstanceRig->GetHierarchy()->CopyHierarchy(Hierarchy);
+				InstanceRig->Initialize(true);
+			}
+		}
+
 	}
 
 #if WITH_EDITOR
