Added optional custom parent space to spherical pose reader.

diff --git a/Engine/Plugins/Animation/ControlRig/Source/ControlRig/Private/Units/Highlevel/Hierarchy/RigUnit_SphericalPoseReader.cpp b/Engine/Plugins/Animation/ControlRig/Source/ControlRig/Private/Units/Highlevel/Hierarchy/RigUnit_SphericalPoseReader.cpp
--- a/Engine/Plugins/Animation/ControlRig/Source/ControlRig/Private/Units/Highlevel/Hierarchy/RigUnit_SphericalPoseReader.cpp
+++ b/Engine/Plugins/Animation/ControlRig/Source/ControlRig/Private/Units/Highlevel/Hierarchy/RigUnit_SphericalPoseReader.cpp
@@ -34,15 +34,31 @@ FRigUnit_SphericalPoseReader_Execute()
 		FalloffRegionScaleFactors,
 		FlipWidthScaling,
 		FlipHeightScaling);
+
+	// get parent global space
+	FTransform GlobalDriverParentTransform;
+	FTransform GlobalDriverParentTransformInit;
+	if (OptionalParentItem.IsValid())
+	{
+		GlobalDriverParentTransform = Hierarchy->GetGlobalTransform(OptionalParentItem);
+		GlobalDriverParentTransformInit = Hierarchy->GetInitialGlobalTransform(OptionalParentItem);
+	}else
+	{
+		// if user does not specify a custom parent space, then simply use the driver's parent
+		GlobalDriverParentTransform = Hierarchy->GetParentTransform(DriverItem);
+		GlobalDriverParentTransformInit = Hierarchy->GetInitialGlobalTransform(DriverItem);
+	}
 	
 	// get local rotation space of driver
-	FTransform LocalDriverTransformInit = Hierarchy->GetInitialLocalTransform(DriverItem);
+	FTransform GlobalDriverTransformInit = Hierarchy->GetInitialGlobalTransform(DriverItem);
+	FTransform LocalDriverTransformInit = GlobalDriverTransformInit * GlobalDriverParentTransformInit.Inverse();
 	// apply static offset in local space
 	FQuat RotationOffsetQuat = FQuat::MakeFromEuler(RotationOffset);
 	LocalDriverTransformInit.SetRotation(LocalDriverTransformInit.GetRotation() * RotationOffsetQuat);
-	// put into global space
-	const FTransform GlobalDriverParentTransform = Hierarchy->GetParentTransform(DriverItem);
+	
+	// calculate world transform of sphere
 	FTransform WorldOffset = LocalDriverTransformInit * GlobalDriverParentTransform;
+	WorldOffset.SetLocation(Hierarchy->GetGlobalTransform(DriverItem).GetLocation());
 
 	// get driver axis
 	const FTransform GlobalDriverTransform = Hierarchy->GetGlobalTransform(DriverItem);
@@ -192,7 +208,7 @@ void FRigUnit_SphericalPoseReader::DistanceToEllipse(
     float TX = 0.70710678118f;
     float TY = 0.70710678118f;
 
-    const int Iterations = 2; // this could be higher for greatly quality
+    const int Iterations = 2; // this could be higher for accuracy, but 2 seems good enough for this use case
     for (int i=0; i<Iterations; ++i)
     {
         const float ScaledX = SizeX * TX;
diff --git a/Engine/Plugins/Animation/ControlRig/Source/ControlRig/Private/Units/Highlevel/Hierarchy/RigUnit_SphericalPoseReader.h b/Engine/Plugins/Animation/ControlRig/Source/ControlRig/Private/Units/Highlevel/Hierarchy/RigUnit_SphericalPoseReader.h
--- a/Engine/Plugins/Animation/ControlRig/Source/ControlRig/Private/Units/Highlevel/Hierarchy/RigUnit_SphericalPoseReader.h
+++ b/Engine/Plugins/Animation/ControlRig/Source/ControlRig/Private/Units/Highlevel/Hierarchy/RigUnit_SphericalPoseReader.h
@@ -323,7 +323,7 @@ struct CONTROLRIG_API FRigUnit_SphericalPoseReader: public FRigUnit_HighlevelBas
 	UPROPERTY(meta = (Output))
 	float OutputParam;
 	
-	// Rotate the entire falloff region to align with the desired area of effect.
+	// The bone that will drive the output parameter when rotated into the active regions of this pose reader.
 	UPROPERTY(meta = (Input, ExpandByDefault))
     FRigElementKey DriverItem;
 
@@ -357,6 +357,10 @@ struct CONTROLRIG_API FRigUnit_SphericalPoseReader: public FRigUnit_HighlevelBas
 	UPROPERTY(meta = (Input))
 	bool FlipHeightScaling = false;
 
+	// An optional parent to use as a stable frame of reference for the active regions (defaults to DriverItem's parent if unset).
+	UPROPERTY(meta = (Input))
+	FRigElementKey OptionalParentItem;
+
 	UPROPERTY(meta = (Input))
 	FSphericalPoseReaderDebugSettings Debug;
 
