Control Rig: Compiling after changing a transform array variable's type to bool with existing get nodes causes crash

diff --git a/Engine/Plugins/Animation/ControlRig/Source/ControlRigDeveloper/Private/ControlRigBlueprint.cpp b/Engine/Plugins/Animation/ControlRig/Source/ControlRigDeveloper/Private/ControlRigBlueprint.cpp
--- a/Engine/Plugins/Animation/ControlRig/Source/ControlRigDeveloper/Private/ControlRigBlueprint.cpp
+++ b/Engine/Plugins/Animation/ControlRig/Source/ControlRigDeveloper/Private/ControlRigBlueprint.cpp
@@ -3885,6 +3885,10 @@ void UControlRigBlueprint::OnVariableRenamed(const FName& InOldVarName, const FN
 
 void UControlRigBlueprint::OnVariableTypeChanged(const FName& InVarName, FEdGraphPinType InOldPinType, FEdGraphPinType InNewPinType)
 {
+	FString CPPType;
+	UObject* CPPTypeObject;
+	FRigVMGraphVariableDescription::CPPTypeFromPinType(InNewPinType, CPPType, CPPTypeObject);
+	
 	TArray<URigVMGraph*> AllGraphs = GetAllModels();
 	for (URigVMGraph* Graph : AllGraphs)
 	{
@@ -3896,10 +3900,9 @@ void UControlRigBlueprint::OnVariableTypeChanged(const FName& InVarName, FEdGrap
 			const bool bSetupUndoRedo = false;
 #endif
 
-			FRigVMExternalVariable NewVariable = UControlRig::GetExternalVariableFromPinType(InVarName, InNewPinType);
-			if (NewVariable.IsValid(true)) // allow nullptr
+			if (!CPPType.IsEmpty())
 			{
-				Controller->OnExternalVariableTypeChanged(InVarName, NewVariable.TypeName.ToString(), NewVariable.TypeObject, bSetupUndoRedo);
+				Controller->OnExternalVariableTypeChanged(InVarName, CPPType, CPPTypeObject, bSetupUndoRedo);
 			}
 			else
 			{
@@ -3908,25 +3911,23 @@ void UControlRigBlueprint::OnVariableTypeChanged(const FName& InVarName, FEdGrap
 		}
 	}
 
-	FRigVMExternalVariable NewVariable = UControlRig::GetExternalVariableFromPinType(InVarName, InNewPinType);
-	FString PinType = NewVariable.TypeName.ToString();
-	if(UScriptStruct* ScriptStruct = Cast<UScriptStruct>(NewVariable.TypeObject))
+	if(UScriptStruct* ScriptStruct = Cast<UScriptStruct>(CPPTypeObject))
 	{
 		for (auto Var : NewVariables)
 		{
 			if (Var.VarName == InVarName)
 			{
-				PinType = ScriptStruct->GetName();
+				CPPType = ScriptStruct->GetName();
 			}
 		}
 	}
-	else if (UEnum* Enum = Cast<UEnum>(NewVariable.TypeObject))
+	else if (UEnum* Enum = Cast<UEnum>(CPPTypeObject))
 	{
 		for (auto Var : NewVariables)
 		{
 			if (Var.VarName == InVarName)
 			{
-				PinType = Enum->GetName();
+				CPPType = Enum->GetName();
 			}
 		}
 	}
@@ -3934,7 +3935,7 @@ void UControlRigBlueprint::OnVariableTypeChanged(const FName& InVarName, FEdGrap
 	RigVMPythonUtils::Print(GetFName().ToString(),
 		FString::Printf(TEXT("blueprint.change_member_variable_type('%s', '%s')"),
 		*InVarName.ToString(),
-		NewVariable.bIsArray ? *RigVMUtilities::ArrayTypeFromBaseType(PinType) : *PinType));
+		*CPPType));
 
 	BroadcastExternalVariablesChangedEvent();
 }
