Control Rig: Fix cook issue: wait for package done to perform actual postload work. hierarchy is not ready until after package is done

diff --git a/Engine/Plugins/Animation/ControlRig/Source/ControlRigDeveloper/Private/ControlRigBlueprint.cpp b/Engine/Plugins/Animation/ControlRig/Source/ControlRigDeveloper/Private/ControlRigBlueprint.cpp
--- a/Engine/Plugins/Animation/ControlRig/Source/ControlRigDeveloper/Private/ControlRigBlueprint.cpp
+++ b/Engine/Plugins/Animation/ControlRig/Source/ControlRigDeveloper/Private/ControlRigBlueprint.cpp
@@ -544,8 +544,6 @@ void UControlRigBlueprint::PostLoad()
 
 		CDO->ShapeLibraries = ShapeLibraries;
 		CDO->GizmoLibrary_DEPRECATED.Reset();
-		CDO->GetHierarchy()->CopyHierarchy(Hierarchy);
-		CDO->Initialize(true);
 
 		TArray<UObject*> ArchetypeInstances;
 		CDO->GetArchetypeInstances(ArchetypeInstances);
@@ -555,8 +553,6 @@ void UControlRigBlueprint::PostLoad()
 			{
 				InstanceRig->ShapeLibraries = ShapeLibraries;
 				InstanceRig->GizmoLibrary_DEPRECATED.Reset();
-				InstanceRig->GetHierarchy()->CopyHierarchy(Hierarchy);
-				InstanceRig->Initialize(true);
 			}
 		}
 
