Control Rig: Clean up use of space once the last key has been deleted

diff --git a/Engine/Plugins/Animation/ControlRig/Source/ControlRig/Private/Sequencer/MovieSceneControlRigParameterSection.cpp b/Engine/Plugins/Animation/ControlRig/Source/ControlRig/Private/Sequencer/MovieSceneControlRigParameterSection.cpp
--- a/Engine/Plugins/Animation/ControlRig/Source/ControlRig/Private/Sequencer/MovieSceneControlRigParameterSection.cpp
+++ b/Engine/Plugins/Animation/ControlRig/Source/ControlRig/Private/Sequencer/MovieSceneControlRigParameterSection.cpp
@@ -1160,7 +1160,7 @@ void UMovieSceneControlRigParameterSection::AddSpaceChannel(FName InControlName,
 				if (OnSpaceChannelAdded.IsBound())
 				{
 					FSpaceControlNameAndChannel& NameAndChannel = SpaceChannels[SpaceChannels.Num() - 1];
-					OnSpaceChannelAdded.Broadcast(this, &NameAndChannel.SpaceCurve);
+					OnSpaceChannelAdded.Broadcast(this, InControlName, &NameAndChannel.SpaceCurve);
 				}
 			}
 			if (bReconstructChannel)
diff --git a/Engine/Plugins/Animation/ControlRig/Source/ControlRig/Private/Sequencer/MovieSceneControlRigParameterTrack.cpp b/Engine/Plugins/Animation/ControlRig/Source/ControlRig/Private/Sequencer/MovieSceneControlRigParameterTrack.cpp
--- a/Engine/Plugins/Animation/ControlRig/Source/ControlRig/Private/Sequencer/MovieSceneControlRigParameterTrack.cpp
+++ b/Engine/Plugins/Animation/ControlRig/Source/ControlRig/Private/Sequencer/MovieSceneControlRigParameterTrack.cpp
@@ -8,6 +8,7 @@
 #include "MovieScene.h"
 #include "MovieSceneTimeHelpers.h"
 #include "Channels/MovieSceneChannelProxy.h"
+#include "Rigs/RigHierarchyController.h"
 #define LOCTEXT_NAMESPACE "MovieSceneParameterControlRigTrack"
 
 
@@ -56,9 +57,37 @@ UMovieSceneSection* UMovieSceneControlRigParameterTrack::CreateNewSection()
 	return  NewSection;
 }
 
-void UMovieSceneControlRigParameterTrack::HandleOnSpaceAdded(UMovieSceneControlRigParameterSection* Section, FMovieSceneControlRigSpaceChannel* Channel)
+void UMovieSceneControlRigParameterTrack::HandleOnSpaceAdded(UMovieSceneControlRigParameterSection* Section, const FName& InControlName, FMovieSceneControlRigSpaceChannel* Channel)
 {
-	OnSpaceChannelAdded.Broadcast(Section, Channel);
+	OnSpaceChannelAdded.Broadcast(Section, InControlName, Channel);
+
+	Channel->OnSpaceNoLongerUsed().RemoveAll(this);
+	Channel->OnSpaceNoLongerUsed().AddUObject(this, &UMovieSceneControlRigParameterTrack::HandleOnSpaceNoLongerUsed, InControlName);
+}
+
+void UMovieSceneControlRigParameterTrack::HandleOnSpaceNoLongerUsed(FMovieSceneControlRigSpaceChannel* InChannel, const TArray<FRigElementKey>& InSpaces, FName InControlName)
+{
+	if (InChannel && GetControlRig())
+	{
+		if(URigHierarchy* Hierarchy = GetControlRig()->GetHierarchy())
+		{
+			if(URigHierarchyController* Controller = Hierarchy->GetController())
+			{
+				const FRigElementKey ControlKey(InControlName, ERigElementType::Control);
+				if(Hierarchy->Find<FRigControlElement>(ControlKey))
+				{
+					FRigElementKey DefaultParent = Hierarchy->GetFirstParent(ControlKey);
+					for(const FRigElementKey& ParentToRemove : InSpaces)
+					{
+						if(DefaultParent != ParentToRemove)
+						{
+							Controller->RemoveParent(ControlKey, ParentToRemove, false, false, false);
+						}
+					}
+				}
+			}
+		}
+	}
 }
 
 void UMovieSceneControlRigParameterTrack::RemoveAllAnimationData()
diff --git a/Engine/Plugins/Animation/ControlRig/Source/ControlRig/Private/Sequencer/MovieSceneControlRigSpaceChannel.cpp b/Engine/Plugins/Animation/ControlRig/Source/ControlRig/Private/Sequencer/MovieSceneControlRigSpaceChannel.cpp
--- a/Engine/Plugins/Animation/ControlRig/Source/ControlRig/Private/Sequencer/MovieSceneControlRigSpaceChannel.cpp
+++ b/Engine/Plugins/Animation/ControlRig/Source/ControlRig/Private/Sequencer/MovieSceneControlRigSpaceChannel.cpp
@@ -42,11 +42,20 @@ void FMovieSceneControlRigSpaceChannel::DuplicateKeys(TArrayView<const FKeyHandl
 
 void FMovieSceneControlRigSpaceChannel::DeleteKeys(TArrayView<const FKeyHandle> InHandles)
 {
+	TArray<FRigElementKey> BeforeKeys, AfterKeys;
+	GetUniqueSpaceList(&BeforeKeys);
+	
 	GetData().DeleteKeys(InHandles);
+
+	GetUniqueSpaceList(&AfterKeys);
+	BroadcastSpaceNoLongerUsed(BeforeKeys, AfterKeys);
 }
 
 void FMovieSceneControlRigSpaceChannel::DeleteKeysFrom(FFrameNumber InTime, bool bDeleteKeysBefore)
 {
+	TArray<FRigElementKey> BeforeKeys, AfterKeys;
+	GetUniqueSpaceList(&BeforeKeys);
+	
 	// Insert a key at the current time to maintain evaluation
 	if (GetData().GetTimes().Num() > 0)
 	{
@@ -56,6 +65,9 @@ void FMovieSceneControlRigSpaceChannel::DeleteKeysFrom(FFrameNumber InTime, bool
 	}
 
 	GetData().DeleteKeysFrom(InTime, bDeleteKeysBefore);
+
+	GetUniqueSpaceList(&AfterKeys);
+	BroadcastSpaceNoLongerUsed(BeforeKeys, AfterKeys);
 }
 
 void FMovieSceneControlRigSpaceChannel::ChangeFrameResolution(FFrameRate SourceRate, FFrameRate DestinationRate)
@@ -85,6 +97,19 @@ void FMovieSceneControlRigSpaceChannel::Offset(FFrameNumber DeltaPosition)
 	GetData().Offset(DeltaPosition);
 }
 
+void FMovieSceneControlRigSpaceChannel::GetUniqueSpaceList(TArray<FRigElementKey>* OutList)
+{
+	if(OutList)
+	{
+		OutList->Reset();
+		
+		for(const FMovieSceneControlRigSpaceBaseKey& KeyValue : KeyValues)
+		{
+			OutList->AddUnique(KeyValue.ControlRigElement);
+		}
+	}
+}
+
 //this will also delete any duplicated keys.
 TArray <FSpaceRange> FMovieSceneControlRigSpaceChannel::FindSpaceIntervals()
 {
@@ -137,6 +162,30 @@ TArray <FSpaceRange> FMovieSceneControlRigSpaceChannel::FindSpaceIntervals()
 	return Ranges;
 }
 
+void FMovieSceneControlRigSpaceChannel::BroadcastSpaceNoLongerUsed(const TArray<FRigElementKey>& BeforeKeys, const TArray<FRigElementKey>& AfterKeys)
+{
+	if(BeforeKeys == AfterKeys)
+	{
+		return;
+	}
+
+	if(!SpaceNoLongerUsedEvent.IsBound())
+	{
+		return;
+	}
+
+	TArray<FRigElementKey> SpacesNoLongerUsed;
+	for(const FRigElementKey& BeforeKey : BeforeKeys)
+	{
+		if(!AfterKeys.Contains(BeforeKey))
+		{
+			SpacesNoLongerUsed.Add(BeforeKey);
+		}
+	}
+
+	SpaceNoLongerUsedEvent.Broadcast(this, SpacesNoLongerUsed);
+}
+
 FName FMovieSceneControlRigSpaceBaseKey::GetName() const
 {
 	switch (SpaceType)
diff --git a/Engine/Plugins/Animation/ControlRig/Source/ControlRig/Public/Sequencer/MovieSceneControlRigParameterSection.h b/Engine/Plugins/Animation/ControlRig/Source/ControlRig/Public/Sequencer/MovieSceneControlRigParameterSection.h
--- a/Engine/Plugins/Animation/ControlRig/Source/ControlRig/Public/Sequencer/MovieSceneControlRigParameterSection.h
+++ b/Engine/Plugins/Animation/ControlRig/Source/ControlRig/Public/Sequencer/MovieSceneControlRigParameterSection.h
@@ -173,7 +173,7 @@ class CONTROLRIG_API UMovieSceneControlRigParameterSection : public UMovieSceneP
 public:
 
 	/** Bindable event for when we add a space channel*/
-	DECLARE_MULTICAST_DELEGATE_TwoParams(FSpaceChannelAddedEvent, UMovieSceneControlRigParameterSection*, FMovieSceneControlRigSpaceChannel*);
+	DECLARE_MULTICAST_DELEGATE_ThreeParams(FSpaceChannelAddedEvent, UMovieSceneControlRigParameterSection*, const FName&, FMovieSceneControlRigSpaceChannel*);
 
 	void AddEnumParameterKey(FName InParameterName, FFrameNumber InTime, uint8 InValue);
 	void AddIntegerParameterKey(FName InParameterName, FFrameNumber InTime, int32 InValue);
diff --git a/Engine/Plugins/Animation/ControlRig/Source/ControlRig/Public/Sequencer/MovieSceneControlRigParameterTrack.h b/Engine/Plugins/Animation/ControlRig/Source/ControlRig/Public/Sequencer/MovieSceneControlRigParameterTrack.h
--- a/Engine/Plugins/Animation/ControlRig/Source/ControlRig/Public/Sequencer/MovieSceneControlRigParameterTrack.h
+++ b/Engine/Plugins/Animation/ControlRig/Source/ControlRig/Public/Sequencer/MovieSceneControlRigParameterTrack.h
@@ -119,7 +119,9 @@ public:
 private:
 
 	//we register this with sections.
-	void HandleOnSpaceAdded(UMovieSceneControlRigParameterSection* Section, FMovieSceneControlRigSpaceChannel* Channel);
+	void HandleOnSpaceAdded(UMovieSceneControlRigParameterSection* Section, const FName& InControlName, FMovieSceneControlRigSpaceChannel* Channel);
+	//we'll register this against all space channels
+	void HandleOnSpaceNoLongerUsed(FMovieSceneControlRigSpaceChannel* InChannel, const TArray<FRigElementKey>& InSpaces, FName InControlName);
 	//then send this event out to the track editor.
 	UMovieSceneControlRigParameterSection::FSpaceChannelAddedEvent OnSpaceChannelAdded;
 
diff --git a/Engine/Plugins/Animation/ControlRig/Source/ControlRig/Public/Sequencer/MovieSceneControlRigSpaceChannel.h b/Engine/Plugins/Animation/ControlRig/Source/ControlRig/Public/Sequencer/MovieSceneControlRigSpaceChannel.h
--- a/Engine/Plugins/Animation/ControlRig/Source/ControlRig/Public/Sequencer/MovieSceneControlRigSpaceChannel.h
+++ b/Engine/Plugins/Animation/ControlRig/Source/ControlRig/Public/Sequencer/MovieSceneControlRigSpaceChannel.h
@@ -17,6 +17,9 @@
 #include "MovieSceneControlRigSpaceChannel.generated.h"
 
 class UControlRig;
+struct FMovieSceneControlRigSpaceChannel;
+
+DECLARE_MULTICAST_DELEGATE_TwoParams(FMovieSceneControlRigSpaceChannelSpaceNoLongerUsedEvent, FMovieSceneControlRigSpaceChannel*, const TArray<FRigElementKey>&);
 
 UENUM(Blueprintable)
 enum class EMovieSceneControlRigSpaceType : uint8
@@ -109,9 +112,15 @@ public:
 	virtual void Reset() override;
 	virtual void Offset(FFrameNumber DeltaPosition) override;
 
+	void GetUniqueSpaceList(TArray<FRigElementKey>* OutList);
+	FMovieSceneControlRigSpaceChannelSpaceNoLongerUsedEvent& OnSpaceNoLongerUsed() { return SpaceNoLongerUsedEvent; }
+
 	TArray <FSpaceRange> FindSpaceIntervals();
+
 private:
 
+	void BroadcastSpaceNoLongerUsed(const TArray<FRigElementKey>& BeforeKeys, const TArray<FRigElementKey>& AfterKeys);
+
 	/** Sorted array of key times */
 	UPROPERTY(meta = (KeyTimes))
 	TArray<FFrameNumber> KeyTimes;
@@ -121,6 +130,8 @@ private:
 	TArray<FMovieSceneControlRigSpaceBaseKey> KeyValues;
 
 	FMovieSceneKeyHandleMap KeyHandles;
+	
+	FMovieSceneControlRigSpaceChannelSpaceNoLongerUsedEvent SpaceNoLongerUsedEvent;
 };
 
 
diff --git a/Engine/Plugins/Animation/ControlRig/Source/ControlRigEditor/Private/Sequencer/ControlRigParameterTrackEditor.cpp b/Engine/Plugins/Animation/ControlRig/Source/ControlRigEditor/Private/Sequencer/ControlRigParameterTrackEditor.cpp
--- a/Engine/Plugins/Animation/ControlRig/Source/ControlRigEditor/Private/Sequencer/ControlRigParameterTrackEditor.cpp
+++ b/Engine/Plugins/Animation/ControlRig/Source/ControlRigEditor/Private/Sequencer/ControlRigParameterTrackEditor.cpp
@@ -1966,7 +1966,7 @@ UMovieSceneControlRigParameterTrack* FControlRigParameterTrackEditor::FindTrack(
 	return nullptr;
 }
 
-void FControlRigParameterTrackEditor::HandleOnSpaceAdded(UMovieSceneControlRigParameterSection* Section, FMovieSceneControlRigSpaceChannel* SpaceChannel)
+void FControlRigParameterTrackEditor::HandleOnSpaceAdded(UMovieSceneControlRigParameterSection* Section, const FName& ControlName, FMovieSceneControlRigSpaceChannel* SpaceChannel)
 {
 	if (SpaceChannel)
 	{
@@ -2021,7 +2021,7 @@ void FControlRigParameterTrackEditor::PostUndo(bool bSuccess)
 			TArray<FSpaceControlNameAndChannel>& SpaceChannels = Section->GetSpaceChannels();
 			for(FSpaceControlNameAndChannel& Channel: SpaceChannels)
 			{ 
-				HandleOnSpaceAdded(Section, &(Channel.SpaceCurve));
+				HandleOnSpaceAdded(Section, Channel.ControlName, &(Channel.SpaceCurve));
 			}
 		}
 	}
diff --git a/Engine/Plugins/Animation/ControlRig/Source/ControlRigEditor/Private/Sequencer/ControlRigParameterTrackEditor.h b/Engine/Plugins/Animation/ControlRig/Source/ControlRigEditor/Private/Sequencer/ControlRigParameterTrackEditor.h
--- a/Engine/Plugins/Animation/ControlRig/Source/ControlRigEditor/Private/Sequencer/ControlRigParameterTrackEditor.h
+++ b/Engine/Plugins/Animation/ControlRig/Source/ControlRigEditor/Private/Sequencer/ControlRigParameterTrackEditor.h
@@ -130,7 +130,7 @@ private:
 	void HandleOnInitialized(UControlRig* Subject, const EControlRigState InState, const FName& InEventName);
 
 	/** SpaceChannel Delegates*/
-	void HandleOnSpaceAdded(UMovieSceneControlRigParameterSection* Section, FMovieSceneControlRigSpaceChannel* Channel);
+	void HandleOnSpaceAdded(UMovieSceneControlRigParameterSection* Section, const FName& ControlName, FMovieSceneControlRigSpaceChannel* Channel);
 	void HandleSpaceKeyDeleted(UMovieSceneControlRigParameterSection* Section, FMovieSceneControlRigSpaceChannel* Channel, const TArray<FKeyAddOrDeleteEventItem>& DeletedItems);
 	void HandleSpaceKeyMoved(UMovieSceneControlRigParameterSection* Section, FMovieSceneControlRigSpaceChannel* Channel, const  TArray<FKeyMoveEventItem>& MovedItems);
 
