Control Rig: Fix can't set default values of function

diff --git a/Engine/Plugins/Animation/ControlRig/Source/ControlRigDeveloper/Private/Graph/ControlRigGraph.cpp b/Engine/Plugins/Animation/ControlRig/Source/ControlRigDeveloper/Private/Graph/ControlRigGraph.cpp
--- a/Engine/Plugins/Animation/ControlRig/Source/ControlRigDeveloper/Private/Graph/ControlRigGraph.cpp
+++ b/Engine/Plugins/Animation/ControlRig/Source/ControlRigDeveloper/Private/Graph/ControlRigGraph.cpp
@@ -220,12 +220,6 @@ void UControlRigGraph::HandleModifiedEvent(ERigVMGraphNotifType InNotifType, URi
 		return;
 	}
 
-	// We don't need a visual representation for the low level function library graph
-	if(InGraph->IsA<URigVMFunctionLibrary>())
-	{
-		return;
-	}
-
 	if(UControlRigGraphSchema* ControlRigSchema = (UControlRigGraphSchema*)GetControlRigGraphSchema())
 	{
 		ControlRigSchema->HandleModifiedEvent(InNotifType, InGraph, InSubject);
diff --git a/Engine/Plugins/Animation/ControlRig/Source/ControlRigEditor/Private/ControlRigGraphDetails.cpp b/Engine/Plugins/Animation/ControlRig/Source/ControlRigEditor/Private/ControlRigGraphDetails.cpp
--- a/Engine/Plugins/Animation/ControlRig/Source/ControlRigEditor/Private/ControlRigGraphDetails.cpp
+++ b/Engine/Plugins/Animation/ControlRig/Source/ControlRigEditor/Private/ControlRigGraphDetails.cpp
@@ -514,6 +514,15 @@ void FControlRigArgumentDefaultNode::GenerateChildContent(IDetailChildrenBuilder
 	[
 		SAssignNew(OwnedNodeWidget, SControlRigGraphNode).GraphNodeObj(ControlRigGraphNode)
 	];
+
+	OwnedNodeWidget->SetIsEditable(true);
+	TArray< TSharedRef<SWidget> > Pins;
+	OwnedNodeWidget->GetPins(Pins);
+	for (TSharedRef<SWidget> Pin : Pins)
+	{
+		TSharedRef<SGraphPin> SPin = StaticCastSharedRef<SGraphPin>(Pin);
+		SPin->GetPinObj()->bNotConnectable = true;
+	}
 }
 
 void FControlRigArgumentDefaultNode::OnGraphChanged(const FEdGraphEditAction& InAction)
