Control Rig: Fix Control Pin Value not showing in the viewport

diff --git a/Engine/Plugins/Animation/ControlRig/Source/ControlRigEditor/Private/EditMode/ControlRigEditMode.cpp b/Engine/Plugins/Animation/ControlRig/Source/ControlRigEditor/Private/EditMode/ControlRigEditMode.cpp
--- a/Engine/Plugins/Animation/ControlRig/Source/ControlRigEditor/Private/EditMode/ControlRigEditMode.cpp
+++ b/Engine/Plugins/Animation/ControlRig/Source/ControlRigEditor/Private/EditMode/ControlRigEditMode.cpp
@@ -1684,6 +1684,20 @@ TArray<FRigElementKey> FControlRigEditMode::GetSelectedRigElements() const
 
 			}
 		}
+		else
+		{
+			// check if there is a pin value transient control active
+			// when a pin control is active, all existing selection should have been cleared
+			TArray<FRigControlElement*> TransientControls = ControlRig->GetHierarchy()->GetTransientControls();
+
+			if (TransientControls.Num() > 0)
+			{
+				if (ensure(SelectedKeys.Num() == 0))
+				{
+					SelectedKeys.Add(TransientControls[0]->GetKey());
+				}
+			}
+		}
 	}
 
 	return SelectedKeys;
