Control Rig: Notify hierarchy of unused spaces also when adding / updating keys

diff --git a/Engine/Plugins/Animation/ControlRig/Source/ControlRig/Private/Sequencer/MovieSceneControlRigSpaceChannel.cpp b/Engine/Plugins/Animation/ControlRig/Source/ControlRig/Private/Sequencer/MovieSceneControlRigSpaceChannel.cpp
--- a/Engine/Plugins/Animation/ControlRig/Source/ControlRig/Private/Sequencer/MovieSceneControlRigSpaceChannel.cpp
+++ b/Engine/Plugins/Animation/ControlRig/Source/ControlRig/Private/Sequencer/MovieSceneControlRigSpaceChannel.cpp
@@ -183,7 +183,10 @@ void FMovieSceneControlRigSpaceChannel::BroadcastSpaceNoLongerUsed(const TArray<
 		}
 	}
 
-	SpaceNoLongerUsedEvent.Broadcast(this, SpacesNoLongerUsed);
+	if(!SpacesNoLongerUsed.IsEmpty())
+	{
+		SpaceNoLongerUsedEvent.Broadcast(this, SpacesNoLongerUsed);
+	}
 }
 
 FName FMovieSceneControlRigSpaceBaseKey::GetName() const
diff --git a/Engine/Plugins/Animation/ControlRig/Source/ControlRig/Public/Sequencer/MovieSceneControlRigSpaceChannel.h b/Engine/Plugins/Animation/ControlRig/Source/ControlRig/Public/Sequencer/MovieSceneControlRigSpaceChannel.h
--- a/Engine/Plugins/Animation/ControlRig/Source/ControlRig/Public/Sequencer/MovieSceneControlRigSpaceChannel.h
+++ b/Engine/Plugins/Animation/ControlRig/Source/ControlRig/Public/Sequencer/MovieSceneControlRigSpaceChannel.h
@@ -132,6 +132,8 @@ private:
 	FMovieSceneKeyHandleMap KeyHandles;
 	
 	FMovieSceneControlRigSpaceChannelSpaceNoLongerUsedEvent SpaceNoLongerUsedEvent;
+
+	friend struct FControlRigSpaceChannelHelpers;
 };
 
 
diff --git a/Engine/Plugins/Animation/ControlRig/Source/ControlRigEditor/Private/EditMode/ControlRigEditMode.cpp b/Engine/Plugins/Animation/ControlRig/Source/ControlRigEditor/Private/EditMode/ControlRigEditMode.cpp
--- a/Engine/Plugins/Animation/ControlRig/Source/ControlRigEditor/Private/EditMode/ControlRigEditMode.cpp
+++ b/Engine/Plugins/Animation/ControlRig/Source/ControlRigEditor/Private/EditMode/ControlRigEditMode.cpp
@@ -138,6 +138,7 @@ FControlRigEditMode::FControlRigEditMode()
 	, bSelectionChanged(false)
 	, PivotTransform(FTransform::Identity)
 	, bRecreateControlShapesRequired(false)
+	, bSuspendHierarchyNotifs(false)
 	, CurrentViewportClient(nullptr)
 	, bIsChangingCoordSystem(false)
 {
@@ -2430,6 +2431,11 @@ bool FControlRigEditMode::AreRigElementSelectedAndMovable() const
 
 void FControlRigEditMode::OnHierarchyModified(ERigHierarchyNotification InNotif, URigHierarchy* InHierarchy, const FRigBaseElement* InElement)
 {
+	if(bSuspendHierarchyNotifs)
+	{
+		return;
+	}
+	
 	switch(InNotif)
 	{
 		case ERigHierarchyNotification::ElementAdded:
diff --git a/Engine/Plugins/Animation/ControlRig/Source/ControlRigEditor/Private/EditMode/ControlRigEditMode.h b/Engine/Plugins/Animation/ControlRig/Source/ControlRigEditor/Private/EditMode/ControlRigEditMode.h
--- a/Engine/Plugins/Animation/ControlRig/Source/ControlRigEditor/Private/EditMode/ControlRigEditMode.h
+++ b/Engine/Plugins/Animation/ControlRig/Source/ControlRigEditor/Private/EditMode/ControlRigEditMode.h
@@ -268,6 +268,9 @@ protected:
 	/* Flag to recreate shapes during tick */
 	bool bRecreateControlShapesRequired;
 
+	/* Flag to temporarily disable handling notifs from the hierarchy */
+	bool bSuspendHierarchyNotifs;
+
 	/** Shape actors */
 	TArray<AControlRigShapeActor*> ShapeActors;
 	UControlRigDetailPanelControlProxies* ControlProxy;
diff --git a/Engine/Plugins/Animation/ControlRig/Source/ControlRigEditor/Private/EditMode/SControlRigEditModeTools.cpp b/Engine/Plugins/Animation/ControlRig/Source/ControlRigEditor/Private/EditMode/SControlRigEditModeTools.cpp
--- a/Engine/Plugins/Animation/ControlRig/Source/ControlRigEditor/Private/EditMode/SControlRigEditModeTools.cpp
+++ b/Engine/Plugins/Animation/ControlRig/Source/ControlRigEditor/Private/EditMode/SControlRigEditModeTools.cpp
@@ -729,7 +729,18 @@ void SControlRigEditModeTools::HandleSpaceListChanged(URigHierarchy* InHierarchy
 			}
 
 			ControlRig->SetControlCustomization(InControlKey, ControlCustomization);
-			InHierarchy->Notify(ERigHierarchyNotification::ControlSettingChanged, ControlElement);
+
+			if (FControlRigEditMode* EditMode = static_cast<FControlRigEditMode*>(ModeTools->GetActiveMode(FControlRigEditMode::ModeName)))
+			{
+				const TGuardValue<bool> SuspendGuard(EditMode->bSuspendHierarchyNotifs, true);
+				InHierarchy->Notify(ERigHierarchyNotification::ControlSettingChanged, ControlElement);
+			}
+			else
+			{
+				InHierarchy->Notify(ERigHierarchyNotification::ControlSettingChanged, ControlElement);
+			}
+
+			SpacePickerWidget->RefreshContents();
 		}
 	}
 }
diff --git a/Engine/Plugins/Animation/ControlRig/Source/ControlRigEditor/Private/Sequencer/ControlRigSpaceChannelEditors.cpp b/Engine/Plugins/Animation/ControlRig/Source/ControlRigEditor/Private/Sequencer/ControlRigSpaceChannelEditors.cpp
--- a/Engine/Plugins/Animation/ControlRig/Source/ControlRigEditor/Private/Sequencer/ControlRigSpaceChannelEditors.cpp
+++ b/Engine/Plugins/Animation/ControlRig/Source/ControlRigEditor/Private/Sequencer/ControlRigSpaceChannelEditors.cpp
@@ -276,15 +276,20 @@ FKeyHandle FControlRigSpaceChannelHelpers::SequencerKeyControlRigSpaceChannel(UC
 	FMovieSceneControlRigSpaceBaseKey ExistingValue;
 	using namespace UE::MovieScene;
 	EvaluateChannel(Channel, Time, ExistingValue);
+	TArray<FRigElementKey> BeforeKeys, AfterKeys;
+	Channel->GetUniqueSpaceList(&BeforeKeys);
+
 	FMovieSceneControlRigSpaceBaseKey Value = ExistingValue;
 
 	if (SpaceKey == RigHierarchy->GetWorldSpaceReferenceKey())
 	{
 		Value.SpaceType = EMovieSceneControlRigSpaceType::World;
+		Value.ControlRigElement = URigHierarchy::GetWorldSpaceReferenceKey();
 	}
 	else if (SpaceKey == RigHierarchy->GetDefaultParentKey())
 	{
 		Value.SpaceType = EMovieSceneControlRigSpaceType::Parent;
+		Value.ControlRigElement = RigHierarchy->GetFirstParent(ControlKey);
 	}
 	else
 	{
@@ -292,7 +297,7 @@ FKeyHandle FControlRigSpaceChannelHelpers::SequencerKeyControlRigSpaceChannel(UC
 		if (DefaultParent == SpaceKey)
 		{
 			Value.SpaceType = EMovieSceneControlRigSpaceType::Parent;
-
+			Value.ControlRigElement = DefaultParent;
 		}
 		else  //support all types
 		{
@@ -382,6 +387,8 @@ FKeyHandle FControlRigSpaceChannelHelpers::SequencerKeyControlRigSpaceChannel(UC
 		}
 		Sequencer->NotifyMovieSceneDataChanged(EMovieSceneDataChangeType::MovieSceneStructureItemAdded);
 
+		Channel->GetUniqueSpaceList(&AfterKeys);
+		Channel->BroadcastSpaceNoLongerUsed(BeforeKeys, AfterKeys);
 	}
 	return Handle;
 
