Control Rig: Show Variable Binding UI in details panel

diff --git a/Engine/Plugins/Animation/ControlRig/Source/ControlRigEditor/Private/ControlRigGraphDetails.cpp b/Engine/Plugins/Animation/ControlRig/Source/ControlRigEditor/Private/ControlRigGraphDetails.cpp
--- a/Engine/Plugins/Animation/ControlRig/Source/ControlRigEditor/Private/ControlRigGraphDetails.cpp
+++ b/Engine/Plugins/Animation/ControlRig/Source/ControlRigEditor/Private/ControlRigGraphDetails.cpp
@@ -23,6 +23,7 @@
 #include "Graph/ControlRigGraphSchema.h"
 #include "EditorCategoryUtils.h"
 #include "IPropertyUtilities.h"
+#include "Graph/SControlRigGraphPinVariableBinding.h"
 
 #define LOCTEXT_NAMESPACE "ControlRigGraphDetails"
 
@@ -1181,6 +1182,12 @@ void FControlRigWrappedNodeDetails::CustomizeDetails(IDetailLayoutBuilder& Detai
 	
 	IDetailCategoryBuilder& DefaultsCategory = DetailLayout.EditCategory(Categories[0]);
 
+	UControlRigBlueprint* Blueprint = FirstNode->GetTypedOuter<UControlRigBlueprint>();
+	if(Blueprint == nullptr)
+	{
+		return;
+	}
+
 	TSharedPtr<IPropertyHandle> StoredStructHandle = DetailLayout.GetProperty(TEXT("StoredStruct"));
 	for(URigVMPin* Pin : FirstNode->GetPins())
 	{
@@ -1202,21 +1209,32 @@ void FControlRigWrappedNodeDetails::CustomizeDetails(IDetailLayoutBuilder& Detai
 		{
 			static const TCHAR DisabledFormat[] = TEXT("%s\n\nNote: Editing disabled since pin has an input link.");
 
-			DefaultsCategory.AddProperty(PinHandle)
-			.DisplayName(FText::FromName(Pin->GetDisplayName()))
-			.IsEnabled(!bHasAnyInputLink)
-			.ToolTip(FText::FromString(FString::Printf(DisabledFormat, *Pin->GetToolTipText().ToString())));
+			if (Pin->IsBoundToVariable())
+			{
+				DefaultsCategory.AddCustomRow(FText::FromName(Pin->GetDisplayName()))
+				.NameContent()
+				[
+					PinHandle->CreatePropertyNameWidget()
+				]
+				.ValueContent()
+				[
+					SNew(SControlRigVariableBinding)
+						.ModelPin(Pin)
+						.Blueprint(Blueprint)
+				];
+			}
+			else
+			{
+				DefaultsCategory.AddProperty(PinHandle)
+				.DisplayName(FText::FromName(Pin->GetDisplayName()))
+				.IsEnabled(!bHasAnyInputLink)
+				.ToolTip(FText::FromString(FString::Printf(DisabledFormat, *Pin->GetToolTipText().ToString())));				
+			}
 		}
 	}
 
 	DetailLayout.HideProperty(StoredStructHandle);
 
-	UControlRigBlueprint* Blueprint = FirstNode->GetTypedOuter<UControlRigBlueprint>();
-	if(Blueprint == nullptr)
-	{
-		return;
-	}
-
 	if(Objects.Num() > 1)
 	{
 		return;
