Control Rig: Fix cannot create function ref nodes when pulling out of an execute pin

diff --git a/Engine/Plugins/Animation/ControlRig/Source/ControlRigEditor/Private/Graph/NodeSpawners/ControlRigFunctionRefNodeSpawner.cpp b/Engine/Plugins/Animation/ControlRig/Source/ControlRigEditor/Private/Graph/NodeSpawners/ControlRigFunctionRefNodeSpawner.cpp
--- a/Engine/Plugins/Animation/ControlRig/Source/ControlRigEditor/Private/Graph/NodeSpawners/ControlRigFunctionRefNodeSpawner.cpp
+++ b/Engine/Plugins/Animation/ControlRig/Source/ControlRigEditor/Private/Graph/NodeSpawners/ControlRigFunctionRefNodeSpawner.cpp
@@ -287,6 +287,14 @@ UControlRigGraphNode* UControlRigFunctionRefNodeSpawner::SpawnNode(UEdGraph* Par
 					}
 				}
 			}
+			else
+			{
+				// If the package is a template, do not remove the ControlRigGraphNode
+				// We might be spawning a node to populate the PROTO_ context menu for function declarations.
+				Controller->SuspendNotifications(true);
+				Controller->RemoveNode(ModelNode, false);
+				Controller->SuspendNotifications(false);
+			}
 
 			if (bUndo)
 			{
diff --git a/Engine/Source/Developer/RigVMDeveloper/Private/RigVMModel/RigVMController.cpp b/Engine/Source/Developer/RigVMDeveloper/Private/RigVMModel/RigVMController.cpp
--- a/Engine/Source/Developer/RigVMDeveloper/Private/RigVMModel/RigVMController.cpp
+++ b/Engine/Source/Developer/RigVMDeveloper/Private/RigVMModel/RigVMController.cpp
@@ -8245,12 +8245,6 @@ URigVMFunctionReferenceNode* URigVMController::AddFunctionReferenceNode(URigVMLi
 		return nullptr;
 	}
 
-	if (Graph->GetOutermost() == GetTransientPackage() && !bIsRunningUnitTest)
-	{
-		// don't allow linking to transient graphs, which usually mean template nodes
-		return nullptr;
-	}
-
 	if(!CanAddFunctionRefForDefinition(InFunctionDefinition, true))
 	{
 		return nullptr;
