Control Rig: REnaming options in Mirror dialog

diff --git a/Engine/Plugins/Animation/ControlRig/Source/ControlRig/Private/Rigs/RigHierarchyController.cpp b/Engine/Plugins/Animation/ControlRig/Source/ControlRig/Private/Rigs/RigHierarchyController.cpp
--- a/Engine/Plugins/Animation/ControlRig/Source/ControlRig/Private/Rigs/RigHierarchyController.cpp
+++ b/Engine/Plugins/Animation/ControlRig/Source/ControlRig/Private/Rigs/RigHierarchyController.cpp
@@ -2475,7 +2475,7 @@ TArray<FRigElementKey> URigHierarchyController::MirrorElements(TArray<FRigElemen
 	}
 
 	// correct the names
-	if (!InSettings.OldName.IsEmpty() && !InSettings.NewName.IsEmpty())
+	if (!InSettings.SearchString.IsEmpty() && !InSettings.ReplaceString.IsEmpty())
 	{
 		URigHierarchyController* Controller = Hierarchy->GetController(true);
 		check(Controller);
@@ -2484,7 +2484,7 @@ TArray<FRigElementKey> URigHierarchyController::MirrorElements(TArray<FRigElemen
 		{
 			FName OldName = OriginalKeys[Index].Name;
 			FString OldNameStr = OldName.ToString();
-			FString NewNameStr = OldNameStr.Replace(*InSettings.OldName, *InSettings.NewName, ESearchCase::CaseSensitive);
+			FString NewNameStr = OldNameStr.Replace(*InSettings.SearchString, *InSettings.ReplaceString, ESearchCase::CaseSensitive);
 			if (NewNameStr != OldNameStr)
 			{
 				Controller->RenameElement(DuplicatedKeys[Index], *NewNameStr, true);
@@ -2514,8 +2514,8 @@ TArray<FRigElementKey> URigHierarchyController::MirrorElements(TArray<FRigElemen
 				*ArrayStr,
 				(InSettings.MirrorAxis.GetValue() == EAxis::X) ? TEXT("X") : (InSettings.MirrorAxis.GetValue() == EAxis::Y) ? TEXT("Y") : TEXT("Z"),
 				(InSettings.AxisToFlip.GetValue() == EAxis::X) ? TEXT("X") : (InSettings.AxisToFlip.GetValue() == EAxis::Y) ? TEXT("Y") : TEXT("Z"),
-				*InSettings.OldName,
-				*InSettings.NewName,
+				*InSettings.SearchString,
+				*InSettings.ReplaceString,
 				(bSelectNewElements) ? TEXT("True") : TEXT("False")));
 		}
 	}
diff --git a/Engine/Plugins/Animation/ControlRig/Source/ControlRig/Public/Rigs/RigHierarchyDefines.h b/Engine/Plugins/Animation/ControlRig/Source/ControlRig/Public/Rigs/RigHierarchyDefines.h
--- a/Engine/Plugins/Animation/ControlRig/Source/ControlRig/Public/Rigs/RigHierarchyDefines.h
+++ b/Engine/Plugins/Animation/ControlRig/Source/ControlRig/Public/Rigs/RigHierarchyDefines.h
@@ -1611,13 +1611,13 @@ struct CONTROLRIG_API FRigMirrorSettings
 	UPROPERTY(BlueprintReadWrite, EditAnywhere, Category = Settings)
 	TEnumAsByte<EAxis::Type> AxisToFlip;
 
-	// the string to replace all occurences of with New Name
-	UPROPERTY(BlueprintReadWrite, EditAnywhere, Category = Settings)
-	FString OldName;
+	// the string to search for
+	UPROPERTY(BlueprintReadWrite, EditAnywhere, Category = Settings, DisplayName = "Search")
+	FString SearchString;
 
-	// the string to replace all occurences of Old Name with
-	UPROPERTY(BlueprintReadWrite, EditAnywhere, Category = Settings)
-	FString NewName;
+	// the string to replace the search occurrences with
+	UPROPERTY(BlueprintReadWrite, EditAnywhere, Category = Settings, DisplayName = "Replace")
+	FString ReplaceString;
 
 	FTransform MirrorTransform(const FTransform& InTransform) const;
 	FVector MirrorVector(const FVector& InVector) const;
diff --git a/Engine/Plugins/Animation/ControlRig/Source/ControlRigEditor/Private/ControlRigEditorModule.cpp b/Engine/Plugins/Animation/ControlRig/Source/ControlRigEditor/Private/ControlRigEditorModule.cpp
--- a/Engine/Plugins/Animation/ControlRig/Source/ControlRigEditor/Private/ControlRigEditorModule.cpp
+++ b/Engine/Plugins/Animation/ControlRig/Source/ControlRigEditor/Private/ControlRigEditorModule.cpp
@@ -1941,7 +1941,7 @@ void FControlRigEditorModule::GetContextMenuActions(const UControlRigGraphSchema
 									}
 
 									FString OldNameStr = Key.Name.ToString();
-									FString NewNameStr = OldNameStr.Replace(*Settings.OldName, *Settings.NewName, ESearchCase::CaseSensitive);
+									FString NewNameStr = OldNameStr.Replace(*Settings.SearchString, *Settings.ReplaceString, ESearchCase::CaseSensitive);
 									if(NewNameStr != OldNameStr)
 									{
 										Key.Name = *NewNameStr;
diff --git a/Engine/Plugins/Animation/ControlRig/Source/ControlRigEditor/Private/Editor/SRigHierarchy.cpp b/Engine/Plugins/Animation/ControlRig/Source/ControlRigEditor/Private/Editor/SRigHierarchy.cpp
--- a/Engine/Plugins/Animation/ControlRig/Source/ControlRigEditor/Private/Editor/SRigHierarchy.cpp
+++ b/Engine/Plugins/Animation/ControlRig/Source/ControlRigEditor/Private/Editor/SRigHierarchy.cpp
@@ -1601,7 +1601,7 @@ void SRigHierarchy::HandleMirrorItem()
 		KismetInspector->ShowSingleStruct(StructToDisplay);
 
 		TSharedRef<SCustomDialog> MirrorDialog = SNew(SCustomDialog)
-			.Title(FText(LOCTEXT("ControlRigHierarchyMirror", "Mirror Hierarchy")))
+			.Title(FText(LOCTEXT("ControlRigHierarchyMirror", "Mirror Selected Rig Elements")))
 			.DialogContent( KismetInspector)
 			.Buttons({
 				SCustomDialog::FButton(LOCTEXT("OK", "OK")),
