Control Rig: Drag to any blank area in Rig Hierarchy to unparent Rig Element

diff --git a/Engine/Plugins/Animation/ControlRig/Source/ControlRigEditor/Private/Editor/SRigHierarchy.cpp b/Engine/Plugins/Animation/ControlRig/Source/ControlRigEditor/Private/Editor/SRigHierarchy.cpp
--- a/Engine/Plugins/Animation/ControlRig/Source/ControlRigEditor/Private/Editor/SRigHierarchy.cpp
+++ b/Engine/Plugins/Animation/ControlRig/Source/ControlRigEditor/Private/Editor/SRigHierarchy.cpp
@@ -2005,7 +2005,12 @@ FReply SRigHierarchy::OnAcceptDrop(const FDragDropEvent& DragDropEvent, EItemDro
 		}
 		else
 		{
-			return ReparentOrMatchTransform(RigDragDropOp->GetElements(), TargetItem->Key, bReparentItems);
+			FRigElementKey TargetKey;
+			if (TargetItem.IsValid())
+			{
+				TargetKey = TargetItem->Key;
+			}
+			return ReparentOrMatchTransform(RigDragDropOp->GetElements(), TargetKey, bReparentItems);			
 		}
 
 	}
@@ -2072,6 +2077,11 @@ bool SRigHierarchy::HandleVerifyNameChanged(const FRigElementKey& OldKey, const
 	return true;
 }
 
+FReply SRigHierarchy::OnDrop(const FGeometry& MyGeometry, const FDragDropEvent& DragDropEvent)
+{
+	return OnAcceptDrop(DragDropEvent, EItemDropZone::OntoItem, nullptr);
+}
+
 void SRigHierarchy::HandleResetTransform(bool bSelectionOnly)
 {
 	if ((IsMultiSelected() || !bSelectionOnly) && ControlRigEditor.IsValid())
diff --git a/Engine/Plugins/Animation/ControlRig/Source/ControlRigEditor/Private/Editor/SRigHierarchy.h b/Engine/Plugins/Animation/ControlRig/Source/ControlRigEditor/Private/Editor/SRigHierarchy.h
--- a/Engine/Plugins/Animation/ControlRig/Source/ControlRigEditor/Private/Editor/SRigHierarchy.h
+++ b/Engine/Plugins/Animation/ControlRig/Source/ControlRigEditor/Private/Editor/SRigHierarchy.h
@@ -225,6 +225,9 @@ public:
 	FName HandleRenameElement(const FRigElementKey& OldKey, const FString& NewName);
 	bool HandleVerifyNameChanged(const FRigElementKey& OldKey, const FString& NewName, FText& OutErrorMessage);
 
+	// SWidget Overrides
+	virtual FReply OnDrop(const FGeometry& MyGeometry, const FDragDropEvent& DragDropEvent) override;
+
 	friend class FRigTreeElement;
 	friend class SRigHierarchyItem;
 };
\ No newline at end of file
diff --git a/Engine/Plugins/Animation/ControlRig/Source/ControlRigEditor/Private/Editor/SRigHierarchyTreeView.cpp b/Engine/Plugins/Animation/ControlRig/Source/ControlRigEditor/Private/Editor/SRigHierarchyTreeView.cpp
--- a/Engine/Plugins/Animation/ControlRig/Source/ControlRigEditor/Private/Editor/SRigHierarchyTreeView.cpp
+++ b/Engine/Plugins/Animation/ControlRig/Source/ControlRigEditor/Private/Editor/SRigHierarchyTreeView.cpp
@@ -342,11 +342,6 @@ bool SRigHierarchyTreeView::AddElement(const FRigBaseElement* InElement)
 	return true;
 }
 
-void SRigHierarchyTreeView::AddSpacerElement()
-{
-	AddElement(FRigElementKey(), FRigElementKey());
-}
-
 bool SRigHierarchyTreeView::ReparentElement(FRigElementKey InKey, FRigElementKey InParentKey)
 {
 	if (!InKey.IsValid() || InKey == InParentKey)
@@ -472,11 +467,6 @@ void SRigHierarchyTreeView::RefreshTreeView(bool bRebuildContent)
 					return Delegates.OnCompareKeys.Execute(A->Key, B->Key);
 				});
 			}
-
-			if (RootElements.Num() > 0)
-			{
-				AddSpacerElement();
-			}
 		}
 	}
 	else
@@ -489,7 +479,6 @@ void SRigHierarchyTreeView::RefreshTreeView(bool bRebuildContent)
 			{
 				return InElement.Get()->Key == FRigElementKey();
 			});
-			AddSpacerElement();
 		}
 	}
 
diff --git a/Engine/Plugins/Animation/ControlRig/Source/ControlRigEditor/Private/Editor/SRigHierarchyTreeView.h b/Engine/Plugins/Animation/ControlRig/Source/ControlRigEditor/Private/Editor/SRigHierarchyTreeView.h
--- a/Engine/Plugins/Animation/ControlRig/Source/ControlRigEditor/Private/Editor/SRigHierarchyTreeView.h
+++ b/Engine/Plugins/Animation/ControlRig/Source/ControlRigEditor/Private/Editor/SRigHierarchyTreeView.h
@@ -229,7 +229,6 @@ public:
 	static TSharedPtr<FRigTreeElement> FindElement(const FRigElementKey& InElementKey, TSharedPtr<FRigTreeElement> CurrentItem);
 	bool AddElement(FRigElementKey InKey, FRigElementKey InParentKey = FRigElementKey());
 	bool AddElement(const FRigBaseElement* InElement);
-	void AddSpacerElement();
 	bool ReparentElement(FRigElementKey InKey, FRigElementKey InParentKey);
 	bool RemoveElement(FRigElementKey InKey);
 	void RefreshTreeView(bool bRebuildContent = true);
