Control Rig: Ensure name uniqueness for template nodes

diff --git a/Engine/Plugins/Animation/ControlRig/Source/ControlRigEditor/Private/Graph/NodeSpawners/ControlRigVariableNodeSpawner.cpp b/Engine/Plugins/Animation/ControlRig/Source/ControlRigEditor/Private/Graph/NodeSpawners/ControlRigVariableNodeSpawner.cpp
--- a/Engine/Plugins/Animation/ControlRig/Source/ControlRigEditor/Private/Graph/NodeSpawners/ControlRigVariableNodeSpawner.cpp
+++ b/Engine/Plugins/Animation/ControlRig/Source/ControlRigEditor/Private/Graph/NodeSpawners/ControlRigVariableNodeSpawner.cpp
@@ -145,7 +145,18 @@ UEdGraphNode* UControlRigVariableNodeSpawner::Invoke(UEdGraph* ParentGraph, FBin
 		TypeName = FString::Printf(TEXT("TArray<%s>"), *TypeName);
 	}
 
-	if (URigVMNode* ModelNode = Controller->AddVariableNodeFromObjectPath(ExternalVariable.Name, TypeName, ObjectPath, bIsGetter, FString(), Location, FString(), !bIsTemplateNode, !bIsTemplateNode))
+	FString NodeName;
+	if(bIsTemplateNode)
+	{
+		// since we are removing the node at the end of this function
+		// we need to create a unique here.
+		static constexpr TCHAR VariableNodeNameFormat[] = TEXT("VariableNode_%s_%s");
+		static const FString GetterPrefix = TEXT("Getter");
+		static const FString SetterPrefix = TEXT("Setter");
+		NodeName = FString::Printf(VariableNodeNameFormat, bIsGetter ? *GetterPrefix : *SetterPrefix, *ExternalVariable.Name.ToString());
+	}
+
+	if (URigVMNode* ModelNode = Controller->AddVariableNodeFromObjectPath(ExternalVariable.Name, TypeName, ObjectPath, bIsGetter, FString(), Location, NodeName, !bIsTemplateNode, !bIsTemplateNode))
 	{
 		for (UEdGraphNode* Node : ParentGraph->Nodes)
 		{
