Control Rig: Update node tooltips

diff --git a/Engine/Plugins/Animation/ControlRig/Source/ControlRig/Private/Units/Debug/RigUnit_DebugPrimitives.h b/Engine/Plugins/Animation/ControlRig/Source/ControlRig/Private/Units/Debug/RigUnit_DebugPrimitives.h
--- a/Engine/Plugins/Animation/ControlRig/Source/ControlRig/Private/Units/Debug/RigUnit_DebugPrimitives.h
+++ b/Engine/Plugins/Animation/ControlRig/Source/ControlRig/Private/Units/Debug/RigUnit_DebugPrimitives.h
@@ -46,7 +46,7 @@ struct CONTROLRIG_API FRigUnit_DebugRectangle : public FRigUnit_DebugBaseMutable
 };
 
 /**
- * Draws vectors on each bone in the viewport across the entire hierarchy
+ * Draws a rectangle in the viewport given a transform
  */
 USTRUCT(meta=(DisplayName="Draw Rectangle", Keywords="Draw Square"))
 struct CONTROLRIG_API FRigUnit_DebugRectangleItemSpace : public FRigUnit_DebugBaseMutable
diff --git a/Engine/Plugins/Animation/ControlRig/Source/ControlRig/Private/Units/Execution/RigUnit_Collection.h b/Engine/Plugins/Animation/ControlRig/Source/ControlRig/Private/Units/Execution/RigUnit_Collection.h
--- a/Engine/Plugins/Animation/ControlRig/Source/ControlRig/Private/Units/Execution/RigUnit_Collection.h
+++ b/Engine/Plugins/Animation/ControlRig/Source/ControlRig/Private/Units/Execution/RigUnit_Collection.h
@@ -174,7 +174,7 @@ struct CONTROLRIG_API FRigUnit_CollectionChildren : public FRigUnit_CollectionBa
 
 /**
 * Creates an item array based on the direct or recursive children
-* of a provided parent item. Returns an empty collection for an invalid parent item.
+* of a provided parent item. Returns an empty array for an invalid parent item.
 */
 USTRUCT(meta = (DisplayName = "Children", Keywords = "Bone,Joint,Collection,Filter,Parent", Varying))
 struct CONTROLRIG_API FRigUnit_CollectionChildrenArray : public FRigUnit_CollectionBase
@@ -283,7 +283,7 @@ struct CONTROLRIG_API FRigUnit_CollectionReplaceItemsArray : public FRigUnit_Col
 };
 
 /**
- * Returns a collection provided a specific list of items.
+ * Returns a collection provided a specific array of items.
  */
 USTRUCT(meta = (DisplayName = "Collection from Items", Keywords = "Collection,Array", Varying))
 struct CONTROLRIG_API FRigUnit_CollectionItems : public FRigUnit_CollectionBase
@@ -310,7 +310,7 @@ struct CONTROLRIG_API FRigUnit_CollectionItems : public FRigUnit_CollectionBase
 };
 
 /**
-* Returns a collection provided a specific list of items.
+* Returns an array of items provided a collection
 */
 USTRUCT(meta = (DisplayName = "Get Items from Collection", Keywords = "Collection,Array", Varying))
 struct CONTROLRIG_API FRigUnit_CollectionGetItems : public FRigUnit_CollectionBase
diff --git a/Engine/Plugins/Animation/ControlRig/Source/ControlRig/Private/Units/Execution/RigUnit_Hierarchy.h b/Engine/Plugins/Animation/ControlRig/Source/ControlRig/Private/Units/Execution/RigUnit_Hierarchy.h
--- a/Engine/Plugins/Animation/ControlRig/Source/ControlRig/Private/Units/Execution/RigUnit_Hierarchy.h
+++ b/Engine/Plugins/Animation/ControlRig/Source/ControlRig/Private/Units/Execution/RigUnit_Hierarchy.h
@@ -600,7 +600,7 @@ struct CONTROLRIG_API FRigUnit_PoseGetTransform : public FRigUnit_HierarchyBase
 };
 
 /**
-* Returns the hierarchy's pose transform
+* Returns an array of transforms from a given hierarchy pose
 */
 USTRUCT(meta=(DisplayName="Get Pose Cache Transform Array", Keywords="Hierarchy,Pose,State", Varying))
 struct CONTROLRIG_API FRigUnit_PoseGetTransformArray : public FRigUnit_HierarchyBase
diff --git a/Engine/Plugins/Animation/ControlRig/Source/ControlRig/Private/Units/Highlevel/Hierarchy/RigUnit_AimBone.h b/Engine/Plugins/Animation/ControlRig/Source/ControlRig/Private/Units/Highlevel/Hierarchy/RigUnit_AimBone.h
--- a/Engine/Plugins/Animation/ControlRig/Source/ControlRig/Private/Units/Highlevel/Hierarchy/RigUnit_AimBone.h
+++ b/Engine/Plugins/Animation/ControlRig/Source/ControlRig/Private/Units/Highlevel/Hierarchy/RigUnit_AimBone.h
@@ -127,7 +127,7 @@ struct CONTROLRIG_API FRigUnit_AimBone_DebugSettings
 };
 
 /**
- * Aligns the rotation of a primary and secondary axis of an input transform to a world target.
+ * Outputs an aligned transform of a primary and secondary axis of an input transform to a world target.
  * Note: This node operates in world space!
  */
 USTRUCT(meta = (DisplayName = "Aim Math", Category = "Hierarchy", Keywords = "Lookat"))
diff --git a/Engine/Plugins/Animation/ControlRig/Source/ControlRig/Private/Units/Highlevel/Hierarchy/RigUnit_DistributeRotation.h b/Engine/Plugins/Animation/ControlRig/Source/ControlRig/Private/Units/Highlevel/Hierarchy/RigUnit_DistributeRotation.h
--- a/Engine/Plugins/Animation/ControlRig/Source/ControlRig/Private/Units/Highlevel/Hierarchy/RigUnit_DistributeRotation.h
+++ b/Engine/Plugins/Animation/ControlRig/Source/ControlRig/Private/Units/Highlevel/Hierarchy/RigUnit_DistributeRotation.h
@@ -171,7 +171,7 @@ struct CONTROLRIG_API FRigUnit_DistributeRotationForCollection : public FRigUnit
 };
 
 /**
- * Distributes rotations provided across a collection of items.
+ * Distributes rotations provided across a array of items.
  * Each rotation is expressed by a quaternion and a ratio, where the ratio is between 0.0 and 1.0
  * Note: This node adds rotation in local space of each item!
  */
diff --git a/Engine/Plugins/Animation/ControlRig/Source/ControlRig/Private/Units/Highlevel/Hierarchy/RigUnit_SphericalPoseReader.h b/Engine/Plugins/Animation/ControlRig/Source/ControlRig/Private/Units/Highlevel/Hierarchy/RigUnit_SphericalPoseReader.h
--- a/Engine/Plugins/Animation/ControlRig/Source/ControlRig/Private/Units/Highlevel/Hierarchy/RigUnit_SphericalPoseReader.h
+++ b/Engine/Plugins/Animation/ControlRig/Source/ControlRig/Private/Units/Highlevel/Hierarchy/RigUnit_SphericalPoseReader.h
@@ -265,6 +265,9 @@ struct CONTROLRIG_API FSphericalPoseReaderDebugSettings
 	}
 };
 
+/*
+ * Outputs a float value between 0-1 based off of the driver item's rotation in a specified region.
+ */
 USTRUCT(meta=(DisplayName="Spherical Pose Reader", Category="Hierarchy", Keywords="Pose Reader"))
 struct CONTROLRIG_API FRigUnit_SphericalPoseReader: public FRigUnit_HighlevelBaseMutable
 {
diff --git a/Engine/Plugins/Animation/ControlRig/Source/ControlRig/Private/Units/Math/RigUnit_MathTransform.h b/Engine/Plugins/Animation/ControlRig/Source/ControlRig/Private/Units/Math/RigUnit_MathTransform.h
--- a/Engine/Plugins/Animation/ControlRig/Source/ControlRig/Private/Units/Math/RigUnit_MathTransform.h
+++ b/Engine/Plugins/Animation/ControlRig/Source/ControlRig/Private/Units/Math/RigUnit_MathTransform.h
@@ -81,7 +81,7 @@ struct CONTROLRIG_API FRigUnit_MathTransformFromEulerTransform : public FRigUnit
 };
 
 /**
- * Retrieves the euler transform
+ * Retrieves a euler based transform from a quaternion based transform
  */
 USTRUCT(meta=(DisplayName="To Euler Transform", PrototypeName="ToEulerTransform", Keywords="Make,Construct"))
 struct CONTROLRIG_API FRigUnit_MathTransformToEulerTransform : public FRigUnit_MathTransformBase
