Control Rig: Adapt node display names

diff --git a/Engine/Plugins/Animation/ControlRig/Source/ControlRig/Private/Units/Execution/RigUnit_Collection.h b/Engine/Plugins/Animation/ControlRig/Source/ControlRig/Private/Units/Execution/RigUnit_Collection.h
--- a/Engine/Plugins/Animation/ControlRig/Source/ControlRig/Private/Units/Execution/RigUnit_Collection.h
+++ b/Engine/Plugins/Animation/ControlRig/Source/ControlRig/Private/Units/Execution/RigUnit_Collection.h
@@ -140,7 +140,7 @@ struct CONTROLRIG_API FRigUnit_CollectionNameSearchArray : public FRigUnit_Colle
  * Creates a collection based on the direct or recursive children
  * of a provided parent item. Returns an empty collection for an invalid parent item.
  */
-USTRUCT(meta = (DisplayName = "Children", Keywords = "Bone,Joint,Collection,Filter,Parent", Varying, Deprecated = "5.0"))
+USTRUCT(meta = (DisplayName = "Get Children", Keywords = "Bone,Joint,Collection,Filter,Parent", Varying, Deprecated = "5.0"))
 struct CONTROLRIG_API FRigUnit_CollectionChildren : public FRigUnit_CollectionBase
 {
 	GENERATED_BODY()
@@ -176,7 +176,7 @@ struct CONTROLRIG_API FRigUnit_CollectionChildren : public FRigUnit_CollectionBa
 * Creates an item array based on the direct or recursive children
 * of a provided parent item. Returns an empty array for an invalid parent item.
 */
-USTRUCT(meta = (DisplayName = "Children", Keywords = "Bone,Joint,Collection,Filter,Parent", Varying))
+USTRUCT(meta = (DisplayName = "Get Children", Keywords = "Bone,Joint,Collection,Filter,Parent", Varying))
 struct CONTROLRIG_API FRigUnit_CollectionChildrenArray : public FRigUnit_CollectionBase
 {
 	GENERATED_BODY()
diff --git a/Engine/Plugins/Animation/ControlRig/Source/ControlRig/Private/Units/Math/RigUnit_MathRBFInterpolate.h b/Engine/Plugins/Animation/ControlRig/Source/ControlRig/Private/Units/Math/RigUnit_MathRBFInterpolate.h
--- a/Engine/Plugins/Animation/ControlRig/Source/ControlRig/Private/Units/Math/RigUnit_MathRBFInterpolate.h
+++ b/Engine/Plugins/Animation/ControlRig/Source/ControlRig/Private/Units/Math/RigUnit_MathRBFInterpolate.h
@@ -182,7 +182,7 @@ struct FMathRBFInterpolateQuatFloat_Target
 	float Value = 0.0f;
 };
 
-USTRUCT(meta = (Abstract, DisplayName="Quaternion to Float", PrototypeName="RBFInterpolateQuatToFloat"))
+USTRUCT(meta = (Abstract, DisplayName = "RBF Quaternion to Float", PrototypeName="RBFInterpolateQuatToFloat"))
 struct CONTROLRIG_API FRigUnit_MathRBFInterpolateQuatFloat : 
 	public FRigUnit_MathRBFInterpolateQuatBase
 {
@@ -211,7 +211,7 @@ struct FMathRBFInterpolateQuatVector_Target
 	FVector Value = FVector::ZeroVector;
 };
 
-USTRUCT(meta = (Abstract, DisplayName = "Quaternion to Vector", PrototypeName = "RBFInterpolateQuatToVector"))
+USTRUCT(meta = (Abstract, DisplayName = "RBF Quaternion to Vector", PrototypeName = "RBFInterpolateQuatToVector"))
 struct CONTROLRIG_API FRigUnit_MathRBFInterpolateQuatVector :
 	public FRigUnit_MathRBFInterpolateQuatBase
 {
@@ -240,7 +240,7 @@ struct FMathRBFInterpolateQuatColor_Target
 	FLinearColor Value = FLinearColor::Transparent;
 };
 
-USTRUCT(meta = (Abstract, DisplayName = "Quaternion to Color", PrototypeName = "RBFInterpolateQuatToColor"))
+USTRUCT(meta = (Abstract, DisplayName = "RBF Quaternion to Color", PrototypeName = "RBFInterpolateQuatToColor"))
 struct CONTROLRIG_API FRigUnit_MathRBFInterpolateQuatColor :
 	public FRigUnit_MathRBFInterpolateQuatBase
 {
@@ -269,7 +269,7 @@ struct FMathRBFInterpolateQuatQuat_Target
 	FQuat Value = FQuat::Identity;
 };
 
-USTRUCT(meta = (Abstract, DisplayName = "Quaternion to Quaternion", PrototypeName = "RBFInterpolateQuatToQuat"))
+USTRUCT(meta = (Abstract, DisplayName = "RBF Quaternion to Quaternion", PrototypeName = "RBFInterpolateQuatToQuat"))
 struct CONTROLRIG_API FRigUnit_MathRBFInterpolateQuatQuat :
 	public FRigUnit_MathRBFInterpolateQuatBase
 {
@@ -298,7 +298,7 @@ struct FMathRBFInterpolateQuatXform_Target
 	FTransform Value = FTransform::Identity;
 };
 
-USTRUCT(meta = (Abstract, DisplayName = "Quaternion to Transform", PrototypeName = "RBFInterpolateQuatToXform"))
+USTRUCT(meta = (Abstract, DisplayName = "RBF Quaternion to Transform", PrototypeName = "RBFInterpolateQuatToXform"))
 struct CONTROLRIG_API FRigUnit_MathRBFInterpolateQuatXform :
 	public FRigUnit_MathRBFInterpolateQuatBase
 {
@@ -329,7 +329,7 @@ struct FMathRBFInterpolateVectorFloat_Target
 	float Value = 0.0f;
 };
 
-USTRUCT(meta = (Abstract, DisplayName = "Vector to Float", PrototypeName = "RBFInterpolateVectorToFloat", Keywords = "RBF,Interpolate,Vector"))
+USTRUCT(meta = (Abstract, DisplayName = "RBF Vector to Float", PrototypeName = "RBFInterpolateVectorToFloat", Keywords = "RBF,Interpolate,Vector"))
 struct CONTROLRIG_API FRigUnit_MathRBFInterpolateVectorFloat :
 	public FRigUnit_MathRBFInterpolateVectorBase
 {
@@ -358,7 +358,7 @@ struct FMathRBFInterpolateVectorVector_Target
 	FVector Value = FVector::ZeroVector;
 };
 
-USTRUCT(meta = (Abstract, DisplayName = "Vector to Vector", PrototypeName = "RBFInterpolateVectorToVector", Keywords = "RBF,Interpolate,Vector"))
+USTRUCT(meta = (Abstract, DisplayName = "RBF Vector to Vector", PrototypeName = "RBFInterpolateVectorToVector", Keywords = "RBF,Interpolate,Vector"))
 struct CONTROLRIG_API FRigUnit_MathRBFInterpolateVectorVector :
 	public FRigUnit_MathRBFInterpolateVectorBase
 {
@@ -387,7 +387,7 @@ struct FMathRBFInterpolateVectorColor_Target
 	FLinearColor Value = FLinearColor::Transparent;
 };
 
-USTRUCT(meta = (Abstract, DisplayName = "Vector to Color", PrototypeName = "RBFInterpolateVectorToColor", Keywords = "RBF,Interpolate,Vector"))
+USTRUCT(meta = (Abstract, DisplayName = "RBF Vector to Color", PrototypeName = "RBFInterpolateVectorToColor", Keywords = "RBF,Interpolate,Vector"))
 struct CONTROLRIG_API FRigUnit_MathRBFInterpolateVectorColor :
 	public FRigUnit_MathRBFInterpolateVectorBase
 {
@@ -417,7 +417,7 @@ struct FMathRBFInterpolateVectorQuat_Target
 	FQuat Value = FQuat::Identity;
 };
 
-USTRUCT(meta = (Abstract, DisplayName = "Vector to Quat", PrototypeName = "RBFInterpolateVectorToQuat", Keywords = "RBF,Interpolate,Vector"))
+USTRUCT(meta = (Abstract, DisplayName = "RBF Vector to Quat", PrototypeName = "RBFInterpolateVectorToQuat", Keywords = "RBF,Interpolate,Vector"))
 struct CONTROLRIG_API FRigUnit_MathRBFInterpolateVectorQuat :
 	public FRigUnit_MathRBFInterpolateVectorBase
 {
@@ -445,7 +445,7 @@ struct FMathRBFInterpolateVectorXform_Target
 	FTransform Value = FTransform::Identity;
 };
 
-USTRUCT(meta = (Abstract, DisplayName = "Vector to Transform", PrototypeName = "RBFInterpolateVectorToXform", Keywords = "RBF,Interpolate,Vector"))
+USTRUCT(meta = (Abstract, DisplayName = "RBF Vector to Transform", PrototypeName = "RBFInterpolateVectorToXform", Keywords = "RBF,Interpolate,Vector"))
 struct CONTROLRIG_API FRigUnit_MathRBFInterpolateVectorXform :
 	public FRigUnit_MathRBFInterpolateVectorBase
 {
diff --git a/Engine/Plugins/Animation/ControlRig/Source/ControlRig/Private/Units/Simulation/RigUnit_AlphaInterp.h b/Engine/Plugins/Animation/ControlRig/Source/ControlRig/Private/Units/Simulation/RigUnit_AlphaInterp.h
--- a/Engine/Plugins/Animation/ControlRig/Source/ControlRig/Private/Units/Simulation/RigUnit_AlphaInterp.h
+++ b/Engine/Plugins/Animation/ControlRig/Source/ControlRig/Private/Units/Simulation/RigUnit_AlphaInterp.h
@@ -9,7 +9,7 @@
 /**
  * Adds a float value over time over and over again
  */
-USTRUCT(meta=(DisplayName="Interpolate", Keywords="Alpha,Lerp,LinearInterpolate", PrototypeName = "AlphaInterp"))
+USTRUCT(meta=(DisplayName="Alpha Interpolate (float)", Keywords="Alpha,Lerp,LinearInterpolate", PrototypeName = "AlphaInterp"))
 struct CONTROLRIG_API FRigUnit_AlphaInterp : public FRigUnit_SimBase
 {
 	GENERATED_BODY()
@@ -82,7 +82,7 @@ struct CONTROLRIG_API FRigUnit_AlphaInterp : public FRigUnit_SimBase
 /**
  * Adds a vector value over time over and over again
  */
-USTRUCT(meta=(DisplayName="Interpolate", Keywords="Alpha,Lerp,LinearInterpolate", PrototypeName = "AlphaInterp", MenuDescSuffix = "(Vector)"))
+USTRUCT(meta=(DisplayName="Alpha Interpolate (Vector)", Keywords="Alpha,Lerp,LinearInterpolate", PrototypeName = "AlphaInterp", MenuDescSuffix = "(Vector)"))
 struct CONTROLRIG_API FRigUnit_AlphaInterpVector : public FRigUnit_SimBase
 {
 	GENERATED_BODY()
