Control Rig: Fix crash occurs when a user tries to expand public function in Rig Graph that were added from other Control Rig

diff --git a/Engine/Plugins/Animation/ControlRig/Source/ControlRigDeveloper/Private/ControlRigBlueprint.cpp b/Engine/Plugins/Animation/ControlRig/Source/ControlRigDeveloper/Private/ControlRigBlueprint.cpp
--- a/Engine/Plugins/Animation/ControlRig/Source/ControlRigDeveloper/Private/ControlRigBlueprint.cpp
+++ b/Engine/Plugins/Animation/ControlRig/Source/ControlRigDeveloper/Private/ControlRigBlueprint.cpp
@@ -1390,6 +1390,29 @@ URigVMController* UControlRigBlueprint::GetOrCreateController(URigVMGraph* InGra
 		}
 		return FRigVMController_BulkEditResult();
 	});
+
+	Controller->RequestNewExternalVariableDelegate.BindLambda([WeakThis](FRigVMGraphVariableDescription InVariable, bool bInIsPublic, bool bInIsReadOnly) -> FName
+	{
+		if (WeakThis.IsValid())
+		{
+			for (FBPVariableDescription& ExistingVariable : WeakThis->NewVariables)
+			{
+				if (ExistingVariable.VarName == InVariable.Name)
+				{
+					return FName();
+				}
+			}
+
+			FRigVMExternalVariable ExternalVariable = InVariable.ToExternalVariable();
+			return WeakThis->AddMemberVariable(InVariable.Name,
+				ExternalVariable.TypeObject ? ExternalVariable.TypeObject->GetPathName() : ExternalVariable.TypeName.ToString(),
+				bInIsPublic,
+				bInIsReadOnly,
+				InVariable.DefaultValue);
+		}
+		
+		return FName();
+	});
 	
 #endif
 
