Control Rig: Color function icon in context menu + improve tooltips

diff --git a/Engine/Plugins/Animation/ControlRig/Source/ControlRigDeveloper/Private/ControlRigBlueprint.cpp b/Engine/Plugins/Animation/ControlRig/Source/ControlRigDeveloper/Private/ControlRigBlueprint.cpp
--- a/Engine/Plugins/Animation/ControlRig/Source/ControlRigDeveloper/Private/ControlRigBlueprint.cpp
+++ b/Engine/Plugins/Animation/ControlRig/Source/ControlRigDeveloper/Private/ControlRigBlueprint.cpp
@@ -58,6 +58,25 @@ FEdGraphPinType FControlRigPublicFunctionArg::GetPinType() const
 	return UControlRig::GetPinTypeFromExternalVariable(Variable);
 }
 
+bool FControlRigPublicFunctionData::IsMutable() const
+{
+	for(const FControlRigPublicFunctionArg& Arg : Arguments)
+	{
+		if(!Arg.CPPTypeObjectPath.IsNone())
+		{
+			if(UScriptStruct* Struct = Cast<UScriptStruct>(
+				URigVMPin::FindObjectFromCPPTypeObjectPath(Arg.CPPTypeObjectPath.ToString())))
+			{
+				if(Struct->IsChildOf(FRigVMExecuteContext::StaticStruct()))
+				{
+					return true;
+				}
+			}
+		}
+	}
+	return false;
+}
+
 TArray<UControlRigBlueprint*> UControlRigBlueprint::sCurrentlyOpenedRigBlueprints;
 
 UControlRigBlueprint::UControlRigBlueprint(const FObjectInitializer& ObjectInitializer)
diff --git a/Engine/Plugins/Animation/ControlRig/Source/ControlRigDeveloper/Public/ControlRigBlueprint.h b/Engine/Plugins/Animation/ControlRig/Source/ControlRigDeveloper/Public/ControlRigBlueprint.h
--- a/Engine/Plugins/Animation/ControlRig/Source/ControlRigDeveloper/Public/ControlRigBlueprint.h
+++ b/Engine/Plugins/Animation/ControlRig/Source/ControlRigDeveloper/Public/ControlRigBlueprint.h
@@ -101,6 +101,8 @@ struct CONTROLRIGDEVELOPER_API FControlRigPublicFunctionData
 
 	UPROPERTY()
 	TArray<FControlRigPublicFunctionArg> Arguments;
+
+	bool IsMutable() const;
 };
 
 USTRUCT()
diff --git a/Engine/Plugins/Animation/ControlRig/Source/ControlRigEditor/Private/Graph/NodeSpawners/ControlRigFunctionRefNodeSpawner.cpp b/Engine/Plugins/Animation/ControlRig/Source/ControlRigEditor/Private/Graph/NodeSpawners/ControlRigFunctionRefNodeSpawner.cpp
--- a/Engine/Plugins/Animation/ControlRig/Source/ControlRigEditor/Private/Graph/NodeSpawners/ControlRigFunctionRefNodeSpawner.cpp
+++ b/Engine/Plugins/Animation/ControlRig/Source/ControlRigEditor/Private/Graph/NodeSpawners/ControlRigFunctionRefNodeSpawner.cpp
@@ -22,6 +22,7 @@
 #if WITH_EDITOR
 #include "Editor.h"
 #include "SGraphActionMenu.h"
+#include "GraphEditorSettings.h" 
 #endif
 
 #define LOCTEXT_NAMESPACE "ControlRigFunctionRefNodeSpawner"
@@ -38,16 +39,22 @@ UControlRigFunctionRefNodeSpawner* UControlRigFunctionRefNodeSpawner::CreateFrom
 	FBlueprintActionUiSpec& MenuSignature = NodeSpawner->DefaultMenuSignature;
 
 	const FString Category = InFunction->GetNodeCategory();
-	FString CategoryPrefix = TEXT("Local Functions");
-
-	if (!Category.IsEmpty() && !CategoryPrefix.IsEmpty())
-	{
-		CategoryPrefix += TEXT("|");
-	}
 
 	MenuSignature.MenuName = FText::FromString(InFunction->GetName());
 	MenuSignature.Tooltip = InFunction->GetToolTipText();
-	MenuSignature.Category = FText::FromString(CategoryPrefix + Category);
+
+	static const FString LocalFunctionString = TEXT("Local Function");
+	if(MenuSignature.Tooltip.IsEmpty())
+	{
+		MenuSignature.Tooltip = FText::FromString(LocalFunctionString);
+	}
+	else
+	{
+		static constexpr TCHAR Format[] = TEXT("%s\n\n%s");
+		MenuSignature.Tooltip = FText::FromString(FString::Printf(Format, *MenuSignature.Tooltip.ToString(), *LocalFunctionString));
+	}
+	
+	MenuSignature.Category = FText::FromString(Category);
 	MenuSignature.Keywords = FText::FromString(InFunction->GetNodeKeywords());
 
 	// add at least one character, so that PrimeDefaultUiSpec() doesn't 
@@ -60,6 +67,17 @@ UControlRigFunctionRefNodeSpawner* UControlRigFunctionRefNodeSpawner::CreateFrom
 
 	MenuSignature.Icon = FSlateIcon("EditorStyle", "Kismet.AllClasses.FunctionIcon");
 
+#if WITH_EDITOR
+	if (InFunction->IsMutable())
+	{
+		MenuSignature.IconTint = GetDefault<UGraphEditorSettings>()->FunctionCallNodeTitleColor;
+	}
+	else
+	{
+		MenuSignature.IconTint = GetDefault<UGraphEditorSettings>()->PureFunctionCallNodeTitleColor;
+	}
+#endif
+	
 	return NodeSpawner;
 }
 
@@ -74,16 +92,32 @@ UControlRigFunctionRefNodeSpawner* UControlRigFunctionRefNodeSpawner::CreateFrom
 	FBlueprintActionUiSpec& MenuSignature = NodeSpawner->DefaultMenuSignature;
 
 	const FString Category = InPublicFunction.Category;
-	FString CategoryPrefix =  InAssetData.ToSoftObjectPath().GetAssetName();
 
-	if (!Category.IsEmpty() && !CategoryPrefix.IsEmpty())
+	MenuSignature.MenuName = FText::FromName(InPublicFunction.Name);
+	MenuSignature.Category = FText::FromString(Category);
+	MenuSignature.Keywords = FText::FromString(InPublicFunction.Keywords);
+
+	if(const UControlRigBlueprint* ReferencedBlueprint = Cast<UControlRigBlueprint>(InAssetData.FastGetAsset(false)))
 	{
-		CategoryPrefix += TEXT("|");
+		if(const URigVMFunctionLibrary* FunctionLibrary = ReferencedBlueprint->GetLocalFunctionLibrary())
+		{
+			if(const URigVMLibraryNode* FunctionNode = FunctionLibrary->FindFunction(InPublicFunction.Name))
+			{
+				MenuSignature.Tooltip = FunctionNode->GetToolTipText();
+			}
+		}
 	}
 
-	MenuSignature.MenuName = FText::FromName(InPublicFunction.Name);
-	MenuSignature.Category = FText::FromString(CategoryPrefix + Category);
-	MenuSignature.Keywords = FText::FromString(InPublicFunction.Keywords);
+	const FString ObjectPathString = InAssetData.ObjectPath.ToString();
+	if(MenuSignature.Tooltip.IsEmpty())
+	{
+		MenuSignature.Tooltip = FText::FromString(ObjectPathString);
+	}
+	else
+	{
+		static constexpr TCHAR Format[] = TEXT("%s\n\n%s");
+		MenuSignature.Tooltip = FText::FromString(FString::Printf(Format, *MenuSignature.Tooltip.ToString(), *ObjectPathString));
+	}
 
 	// add at least one character, so that PrimeDefaultUiSpec() doesn't 
 	// attempt to query the template node
@@ -95,6 +129,17 @@ UControlRigFunctionRefNodeSpawner* UControlRigFunctionRefNodeSpawner::CreateFrom
 
 	MenuSignature.Icon = FSlateIcon("EditorStyle", "Kismet.AllClasses.FunctionIcon");
 
+#if WITH_EDITOR
+	if (InPublicFunction.IsMutable())
+	{
+		MenuSignature.IconTint = GetDefault<UGraphEditorSettings>()->FunctionCallNodeTitleColor;
+	}
+	else
+	{
+		MenuSignature.IconTint = GetDefault<UGraphEditorSettings>()->PureFunctionCallNodeTitleColor;
+	}
+#endif
+	
 	return NodeSpawner;
 }
 
