Control Rig: Adapt node colors

diff --git a/Engine/Plugins/Animation/ControlRig/Source/ControlRig/Private/Units/Debug/RigUnit_DebugBase.h b/Engine/Plugins/Animation/ControlRig/Source/ControlRig/Private/Units/Debug/RigUnit_DebugBase.h
--- a/Engine/Plugins/Animation/ControlRig/Source/ControlRig/Private/Units/Debug/RigUnit_DebugBase.h
+++ b/Engine/Plugins/Animation/ControlRig/Source/ControlRig/Private/Units/Debug/RigUnit_DebugBase.h
@@ -5,13 +5,13 @@
 #include "Units/RigUnit.h"
 #include "RigUnit_DebugBase.generated.h"
 
-USTRUCT(meta=(Abstract, Category="Debug", NodeColor = "0.25 0.25 0.05"))
+USTRUCT(meta=(Abstract, Category="Debug", NodeColor = "0.83077 0.846873 0.049707"))
 struct CONTROLRIG_API FRigUnit_DebugBase : public FRigUnit
 {
 	GENERATED_BODY()
 };
 
-USTRUCT(meta=(Abstract, Category="Debug", NodeColor = "0.25 0.25 0.05"))
+USTRUCT(meta=(Abstract, Category="Debug", NodeColor = "0.83077 0.846873 0.049707"))
 struct CONTROLRIG_API FRigUnit_DebugBaseMutable : public FRigUnitMutable
 {
 	GENERATED_BODY()
diff --git a/Engine/Plugins/Animation/ControlRig/Source/ControlRig/Private/Units/Debug/RigUnit_ProfilingBracket.h b/Engine/Plugins/Animation/ControlRig/Source/ControlRig/Private/Units/Debug/RigUnit_ProfilingBracket.h
--- a/Engine/Plugins/Animation/ControlRig/Source/ControlRig/Private/Units/Debug/RigUnit_ProfilingBracket.h
+++ b/Engine/Plugins/Animation/ControlRig/Source/ControlRig/Private/Units/Debug/RigUnit_ProfilingBracket.h
@@ -8,7 +8,7 @@
 /**
  * Starts a profiling timer for debugging, used in conjunction with End Profiling Timer
  */
-USTRUCT(meta=(DisplayName="Start Profiling Timer", Keywords="Measure,BeginProfiling,Profile"))
+USTRUCT(meta=(DisplayName="Start Profiling Timer", Keywords="Measure,BeginProfiling,Profile", NodeColor="0.25, 0.25, 0.05000000074505806"))
 struct CONTROLRIG_API FRigUnit_StartProfilingTimer : public FRigUnit_DebugBaseMutable
 {
 	GENERATED_BODY()
@@ -24,7 +24,7 @@ struct CONTROLRIG_API FRigUnit_StartProfilingTimer : public FRigUnit_DebugBaseMu
 /**
  * Ends an existing profiling timer for debugging, used in conjunction with Start Profiling Timer
  */
-USTRUCT(meta = (DisplayName = "End Profiling Timer", Keywords = "Measure,StopProfiling,Meter,Profile"))
+USTRUCT(meta = (DisplayName = "End Profiling Timer", Keywords = "Measure,StopProfiling,Meter,Profile", NodeColor="0.25, 0.25, 0.05000000074505806"))
 struct CONTROLRIG_API FRigUnit_EndProfilingTimer : public FRigUnit_DebugBaseMutable
 {
 	GENERATED_BODY()
diff --git a/Engine/Plugins/Animation/ControlRig/Source/ControlRig/Private/Units/Debug/RigUnit_VisualDebug.h b/Engine/Plugins/Animation/ControlRig/Source/ControlRig/Private/Units/Debug/RigUnit_VisualDebug.h
--- a/Engine/Plugins/Animation/ControlRig/Source/ControlRig/Private/Units/Debug/RigUnit_VisualDebug.h
+++ b/Engine/Plugins/Animation/ControlRig/Source/ControlRig/Private/Units/Debug/RigUnit_VisualDebug.h
@@ -62,7 +62,7 @@ struct CONTROLRIG_API FRigUnit_VisualDebugVector : public FRigUnit_DebugBase
 /**
  * Debug draw parameters for a Point or Vector given a vector
  */
-USTRUCT(meta=(DisplayName = "Visual Debug Vector", PrototypeName = "VisualDebug", Keywords = "Draw,Point", Varying))
+USTRUCT(meta=(DisplayName = "Visual Debug Vector", PrototypeName = "VisualDebug", Keywords = "Draw,Point", Varying, NodeColor = "0.25 0.25 0.050000000745058"))
 struct CONTROLRIG_API FRigUnit_VisualDebugVectorItemSpace : public FRigUnit_DebugBase
 {
 	GENERATED_BODY()
@@ -144,7 +144,7 @@ struct CONTROLRIG_API FRigUnit_VisualDebugQuat : public FRigUnit_DebugBase
 /**
  * Debug draw parameters for an Axis given a quaternion
  */
-USTRUCT(meta = (DisplayName = "Visual Debug Quat", PrototypeName = "VisualDebug", Keywords = "Draw,Rotation", Varying))
+USTRUCT(meta = (DisplayName = "Visual Debug Quat", PrototypeName = "VisualDebug", Keywords = "Draw,Rotation", Varying, NodeColor = "0.25 0.25 0.050000000745058"))
 struct CONTROLRIG_API FRigUnit_VisualDebugQuatItemSpace : public FRigUnit_DebugBase
 {
 	GENERATED_BODY()
@@ -218,7 +218,7 @@ struct CONTROLRIG_API FRigUnit_VisualDebugTransform : public FRigUnit_DebugBase
 /**
  * Debug draw parameters for an Axis given a transform
  */
-USTRUCT(meta=(DisplayName="Visual Debug Transform", PrototypeName = "VisualDebug", Keywords = "Draw,Axes", Varying))
+USTRUCT(meta=(DisplayName="Visual Debug Transform", PrototypeName = "VisualDebug", Keywords = "Draw,Axes", Varying, NodeColor = "0.25 0.25 0.050000000745058"))
 struct CONTROLRIG_API FRigUnit_VisualDebugTransformItemSpace : public FRigUnit_DebugBase
 {
 	GENERATED_BODY()
diff --git a/Engine/Plugins/Animation/ControlRig/Source/ControlRig/Private/Units/Drawing/RigUnit_DrawContainer.h b/Engine/Plugins/Animation/ControlRig/Source/ControlRig/Private/Units/Drawing/RigUnit_DrawContainer.h
--- a/Engine/Plugins/Animation/ControlRig/Source/ControlRig/Private/Units/Drawing/RigUnit_DrawContainer.h
+++ b/Engine/Plugins/Animation/ControlRig/Source/ControlRig/Private/Units/Drawing/RigUnit_DrawContainer.h
@@ -8,7 +8,7 @@
 /**
  * Get Imported Draw Container curve transform and color
  */
-USTRUCT(meta=(DisplayName="Get Draw Instruction", Category="Drawing", NodeColor = "0.25 0.25 0.05", Keywords = "Curve,Shape"))
+USTRUCT(meta=(DisplayName="Get Draw Instruction", Category="Drawing", NodeColor = "0.83077 0.846873 0.049707", Keywords = "Curve,Shape"))
 struct CONTROLRIG_API FRigUnit_DrawContainerGetInstruction : public FRigUnit
 {
 	GENERATED_BODY()
@@ -36,7 +36,7 @@ struct CONTROLRIG_API FRigUnit_DrawContainerGetInstruction : public FRigUnit
 /**
  * Set Imported Draw Container curve color
  */
-USTRUCT(meta = (DisplayName = "Set Draw Color", Category = "Drawing", NodeColor = "0.25 0.25 0.05", Keywords = "Curve,Shape"))
+USTRUCT(meta = (DisplayName = "Set Draw Color", Category = "Drawing", NodeColor = "0.83077 0.846873 0.049707", Keywords = "Curve,Shape"))
 struct CONTROLRIG_API FRigUnit_DrawContainerSetColor: public FRigUnitMutable
 {
 	GENERATED_BODY()
@@ -60,7 +60,7 @@ struct CONTROLRIG_API FRigUnit_DrawContainerSetColor: public FRigUnitMutable
 /**
  * Set Imported Draw Container curve thickness
  */
-USTRUCT(meta = (DisplayName = "Set Draw Thickness", Category = "Drawing", NodeColor = "0.25 0.25 0.05", Keywords = "Curve,Shape"))
+USTRUCT(meta = (DisplayName = "Set Draw Thickness", Category = "Drawing", NodeColor = "0.83077 0.846873 0.049707", Keywords = "Curve,Shape"))
 struct CONTROLRIG_API FRigUnit_DrawContainerSetThickness : public FRigUnitMutable
 {
 	GENERATED_BODY()
@@ -84,7 +84,7 @@ struct CONTROLRIG_API FRigUnit_DrawContainerSetThickness : public FRigUnitMutabl
 /**
  * Set Imported Draw Container curve transform
  */
-USTRUCT(meta = (DisplayName = "Set Draw Transform", Category = "Drawing", NodeColor = "0.25 0.25 0.05", Keywords = "Curve,Shape"))
+USTRUCT(meta = (DisplayName = "Set Draw Transform", Category = "Drawing", NodeColor = "0.83077 0.846873 0.049707", Keywords = "Curve,Shape"))
 struct CONTROLRIG_API FRigUnit_DrawContainerSetTransform : public FRigUnitMutable
 {
 	GENERATED_BODY()
diff --git a/Engine/Plugins/Animation/ControlRig/Source/ControlRig/Private/Units/Execution/RigUnit_Collection.h b/Engine/Plugins/Animation/ControlRig/Source/ControlRig/Private/Units/Execution/RigUnit_Collection.h
--- a/Engine/Plugins/Animation/ControlRig/Source/ControlRig/Private/Units/Execution/RigUnit_Collection.h
+++ b/Engine/Plugins/Animation/ControlRig/Source/ControlRig/Private/Units/Execution/RigUnit_Collection.h
@@ -6,13 +6,13 @@
 #include "ControlRigDefines.h" 
 #include "RigUnit_Collection.generated.h"
 
-USTRUCT(meta = (Abstract, NodeColor = "0.7 0.05 0.5", Category = "Items"))
+USTRUCT(meta = (Abstract, NodeColor = "0.4627450108528137 1.0 0.3294120132923126", Category = "Items"))
 struct CONTROLRIG_API FRigUnit_CollectionBase : public FRigUnit
 {
 	GENERATED_BODY()
 };
 
-USTRUCT(meta = (Abstract, NodeColor = "0.7 0.05 0.5", Category = "Items"))
+USTRUCT(meta = (Abstract, NodeColor = "0.4627450108528137 1.0 0.3294120132923126", Category = "Items"))
 struct CONTROLRIG_API FRigUnit_CollectionBaseMutable : public FRigUnitMutable
 {
 	GENERATED_BODY()
diff --git a/Engine/Plugins/Animation/ControlRig/Source/ControlRig/Private/Units/Execution/RigUnit_ForLoop.h b/Engine/Plugins/Animation/ControlRig/Source/ControlRig/Private/Units/Execution/RigUnit_ForLoop.h
--- a/Engine/Plugins/Animation/ControlRig/Source/ControlRig/Private/Units/Execution/RigUnit_ForLoop.h
+++ b/Engine/Plugins/Animation/ControlRig/Source/ControlRig/Private/Units/Execution/RigUnit_ForLoop.h
@@ -9,7 +9,7 @@
 /**
  * Given a count, execute iteratively until the count is up
  */
-USTRUCT(meta=(DisplayName="For Loop", Category="Execution", TitleColor="1 0 0", NodeColor="0.1 0.1 0.1", Keywords="Iterate", Icon="EditorStyle|GraphEditor.Macro.Loop_16x"))
+USTRUCT(meta=(DisplayName="For Loop", Category="Execution", TitleColor="1 0 0", NodeColor="1 1 1", Keywords="Iterate", Icon="EditorStyle|GraphEditor.Macro.Loop_16x"))
 struct CONTROLRIG_API FRigUnit_ForLoopCount : public FRigUnitMutable
 {
 	GENERATED_BODY()
diff --git a/Engine/Plugins/Animation/ControlRig/Source/ControlRig/Private/Units/Execution/RigUnit_IsInteracting.h b/Engine/Plugins/Animation/ControlRig/Source/ControlRig/Private/Units/Execution/RigUnit_IsInteracting.h
--- a/Engine/Plugins/Animation/ControlRig/Source/ControlRig/Private/Units/Execution/RigUnit_IsInteracting.h
+++ b/Engine/Plugins/Animation/ControlRig/Source/ControlRig/Private/Units/Execution/RigUnit_IsInteracting.h
@@ -9,7 +9,7 @@
 /**
  * Returns true if the Control Rig is being interacted
  */
-USTRUCT(meta=(DisplayName="Is Interacting", Category="Execution", TitleColor="1 0 0", NodeColor="0.1 0.1 0.1", Keywords="Gizmo,Manipulation,Interaction", Varying))
+USTRUCT(meta=(DisplayName="Is Interacting", Category="Execution", TitleColor="1 0 0", NodeColor="1 1 1", Keywords="Gizmo,Manipulation,Interaction", Varying))
 struct CONTROLRIG_API FRigUnit_IsInteracting : public FRigUnit
 {
 	GENERATED_BODY()
diff --git a/Engine/Plugins/Animation/ControlRig/Source/ControlRig/Private/Units/Execution/RigUnit_SequenceExecution.h b/Engine/Plugins/Animation/ControlRig/Source/ControlRig/Private/Units/Execution/RigUnit_SequenceExecution.h
--- a/Engine/Plugins/Animation/ControlRig/Source/ControlRig/Private/Units/Execution/RigUnit_SequenceExecution.h
+++ b/Engine/Plugins/Animation/ControlRig/Source/ControlRig/Private/Units/Execution/RigUnit_SequenceExecution.h
@@ -9,7 +9,7 @@
 /**
  * Allows for a single execution pulse to trigger a series of events in order.
  */
-USTRUCT(meta=(DisplayName="Sequence", Category="Execution", TitleColor="1 0 0", NodeColor="0.1 0.1 0.1", Icon="EditorStyle|GraphEditor.Sequence_16x"))
+USTRUCT(meta=(DisplayName="Sequence", Category="Execution", TitleColor="1 0 0", NodeColor="1 1 1", Icon="EditorStyle|GraphEditor.Sequence_16x"))
 struct CONTROLRIG_API FRigUnit_SequenceExecution : public FRigUnit
 {
 	GENERATED_BODY()
diff --git a/Engine/Plugins/Animation/ControlRig/Source/ControlRig/Private/Units/Hierarchy/RigUnit_BoneName.h b/Engine/Plugins/Animation/ControlRig/Source/ControlRig/Private/Units/Hierarchy/RigUnit_BoneName.h
--- a/Engine/Plugins/Animation/ControlRig/Source/ControlRig/Private/Units/Hierarchy/RigUnit_BoneName.h
+++ b/Engine/Plugins/Animation/ControlRig/Source/ControlRig/Private/Units/Hierarchy/RigUnit_BoneName.h
@@ -8,7 +8,7 @@
 /**
  * The Item node is used to share a specific item across the graph
  */
-USTRUCT(meta=(DisplayName="Item", Category="Hierarchy", DocumentationPolicy = "Strict", Constant))
+USTRUCT(meta=(DisplayName="Item", Category="Hierarchy", NodeColor = "0.4627450108528137 1.0 0.3294120132923126", DocumentationPolicy = "Strict", Constant))
 struct CONTROLRIG_API FRigUnit_Item : public FRigUnit
 {
 	GENERATED_BODY()
diff --git a/Engine/Plugins/Animation/ControlRig/Source/ControlRig/Private/Units/Hierarchy/RigUnit_SetControlTransform.h b/Engine/Plugins/Animation/ControlRig/Source/ControlRig/Private/Units/Hierarchy/RigUnit_SetControlTransform.h
--- a/Engine/Plugins/Animation/ControlRig/Source/ControlRig/Private/Units/Hierarchy/RigUnit_SetControlTransform.h
+++ b/Engine/Plugins/Animation/ControlRig/Source/ControlRig/Private/Units/Hierarchy/RigUnit_SetControlTransform.h
@@ -8,7 +8,7 @@
 /**
  * SetControlBool is used to perform a change in the hierarchy by setting a single control's bool value.
  */
-USTRUCT(meta=(DisplayName="Set Control Bool", Category="Hierarchy", DocumentationPolicy="Strict", Keywords = "SetControlBool,SetGizmoBool", PrototypeName = "SetControlValue"))
+USTRUCT(meta=(DisplayName="Set Control Bool", Category="Hierarchy", DocumentationPolicy="Strict", Keywords = "SetControlBool,SetGizmoBool", PrototypeName = "SetControlValue", NodeColor = "0.0 0.36470600962638855 1.0"))
 struct CONTROLRIG_API FRigUnit_SetControlBool : public FRigUnitMutable
 {
 	GENERATED_BODY()
@@ -61,7 +61,7 @@ struct CONTROLRIG_API FRigUnit_SetMultiControlBool_Entry
 /**
  * SetMultiControlBool is used to perform a change in the hierarchy by setting multiple controls' bool value.
  */
-USTRUCT(meta = (DisplayName = "Set Multiple Controls Bool", Category = "Hierarchy", DocumentationPolicy = "Strict", Keywords = "SetMultipleControlsBool,SetControlBool,SetGizmoBool", PrototypeName = "SetMultiControlValue"))
+USTRUCT(meta = (DisplayName = "Set Multiple Controls Bool", Category = "Hierarchy", DocumentationPolicy = "Strict", Keywords = "SetMultipleControlsBool,SetControlBool,SetGizmoBool", PrototypeName = "SetMultiControlValue", NodeColor = "0.0 0.36470600962638855 1.0"))
 struct CONTROLRIG_API FRigUnit_SetMultiControlBool : public FRigUnitMutable
 {
 	GENERATED_BODY()
@@ -88,7 +88,7 @@ struct CONTROLRIG_API FRigUnit_SetMultiControlBool : public FRigUnitMutable
 /**
  * SetControlFloat is used to perform a change in the hierarchy by setting a single control's float value.
  */
-USTRUCT(meta=(DisplayName="Set Control Float", Category="Hierarchy", DocumentationPolicy="Strict", Keywords = "SetControlFloat,SetGizmoFloat", PrototypeName = "SetControlValue"))
+USTRUCT(meta=(DisplayName="Set Control Float", Category="Hierarchy", DocumentationPolicy="Strict", Keywords = "SetControlFloat,SetGizmoFloat", PrototypeName = "SetControlValue", NodeColor = "0.0 0.36470600962638855 1.0"))
 struct CONTROLRIG_API FRigUnit_SetControlFloat : public FRigUnitMutable
 {
 	GENERATED_BODY()
@@ -149,7 +149,7 @@ struct CONTROLRIG_API FRigUnit_SetMultiControlFloat_Entry
 /**
  * SetMultiControlFloat is used to perform a change in the hierarchy by setting multiple controls' float value.
  */
-USTRUCT(meta = (DisplayName = "Set Multiple Controls Float", Category = "Hierarchy", DocumentationPolicy = "Strict", Keywords = "SetMultipleControlsFloat,SetControlFloat,SetGizmoFloat", PrototypeName = "SetMultiControlValue"))
+USTRUCT(meta = (DisplayName = "Set Multiple Controls Float", Category = "Hierarchy", DocumentationPolicy = "Strict", Keywords = "SetMultipleControlsFloat,SetControlFloat,SetGizmoFloat", PrototypeName = "SetMultiControlValue", NodeColor = "0.0 0.36470600962638855 1.0"))
 struct CONTROLRIG_API FRigUnit_SetMultiControlFloat : public FRigUnitMutable
 {
 	GENERATED_BODY()
@@ -185,7 +185,7 @@ struct CONTROLRIG_API FRigUnit_SetMultiControlFloat : public FRigUnitMutable
 /**
  * SetControlInteger is used to perform a change in the hierarchy by setting a single control's int32 value.
  */
-USTRUCT(meta=(DisplayName="Set Control Integer", Category="Hierarchy", DocumentationPolicy="Strict", Keywords = "SetControlInteger,SetGizmoInteger", PrototypeName = "SetControlValue"))
+USTRUCT(meta=(DisplayName="Set Control Integer", Category="Hierarchy", DocumentationPolicy="Strict", Keywords = "SetControlInteger,SetGizmoInteger", PrototypeName = "SetControlValue", NodeColor = "0.0 0.36470600962638855 1.0"))
 struct CONTROLRIG_API FRigUnit_SetControlInteger : public FRigUnitMutable
 {
 	GENERATED_BODY()
@@ -246,7 +246,7 @@ struct CONTROLRIG_API FRigUnit_SetMultiControlInteger_Entry
 /**
  * SetMultiControlInteger is used to perform a change in the hierarchy by setting multiple controls' integer value.
  */
-USTRUCT(meta = (DisplayName = "Set Multiple Controls Integer", Category = "Hierarchy", DocumentationPolicy = "Strict", Keywords = "SetMultipleControlsInteger,SetControlInteger,SetGizmoInteger", PrototypeName = "SetMultiControlValue"))
+USTRUCT(meta = (DisplayName = "Set Multiple Controls Integer", Category = "Hierarchy", DocumentationPolicy = "Strict", Keywords = "SetMultipleControlsInteger,SetControlInteger,SetGizmoInteger", PrototypeName = "SetMultiControlValue", NodeColor = "0.0 0.36470600962638855 1.0"))
 struct CONTROLRIG_API FRigUnit_SetMultiControlInteger : public FRigUnitMutable
 {
 	GENERATED_BODY()
@@ -280,7 +280,7 @@ struct CONTROLRIG_API FRigUnit_SetMultiControlInteger : public FRigUnitMutable
 /**
  * SetControlVector2D is used to perform a change in the hierarchy by setting a single control's Vector2D value.
  */
-USTRUCT(meta=(DisplayName="Set Control Vector2D", Category="Hierarchy", DocumentationPolicy="Strict", Keywords = "SetControlVector2D,SetGizmoVector2D", PrototypeName = "SetControlValue"))
+USTRUCT(meta=(DisplayName="Set Control Vector2D", Category="Hierarchy", DocumentationPolicy="Strict", Keywords = "SetControlVector2D,SetGizmoVector2D", PrototypeName = "SetControlValue", NodeColor = "0.0 0.36470600962638855 1.0"))
 struct CONTROLRIG_API FRigUnit_SetControlVector2D : public FRigUnitMutable
 {
 	GENERATED_BODY()
@@ -342,7 +342,7 @@ struct CONTROLRIG_API FRigUnit_SetMultiControlVector2D_Entry
 /**
  * SetMultiControlVector2D is used to perform a change in the hierarchy by setting multiple controls' vector2D value.
  */
-USTRUCT(meta = (DisplayName = "Set Multiple Controls Vector2D", Category = "Hierarchy", DocumentationPolicy = "Strict", Keywords = "SetMultipleControlsVector2D,SetControlVector2D,SetGizmoVector2D", PrototypeName = "SetMultiControlValue"))
+USTRUCT(meta = (DisplayName = "Set Multiple Controls Vector2D", Category = "Hierarchy", DocumentationPolicy = "Strict", Keywords = "SetMultipleControlsVector2D,SetControlVector2D,SetGizmoVector2D", PrototypeName = "SetMultiControlValue", NodeColor = "0.0 0.36470600962638855 1.0"))
 struct CONTROLRIG_API FRigUnit_SetMultiControlVector2D : public FRigUnitMutable
 {
 	GENERATED_BODY()
@@ -376,7 +376,7 @@ struct CONTROLRIG_API FRigUnit_SetMultiControlVector2D : public FRigUnitMutable
 /**
  * SetControlVector is used to perform a change in the hierarchy by setting a single control's Vector value.
  */
-USTRUCT(meta=(DisplayName="Set Control Vector", Category="Hierarchy", DocumentationPolicy="Strict", Keywords = "SetControlVector,SetGizmoVector", PrototypeName = "SetControlValue"))
+USTRUCT(meta=(DisplayName="Set Control Vector", Category="Hierarchy", DocumentationPolicy="Strict", Keywords = "SetControlVector,SetGizmoVector", PrototypeName = "SetControlValue", NodeColor = "0.0 0.36470600962638855 1.0"))
 struct CONTROLRIG_API FRigUnit_SetControlVector : public FRigUnitMutable
 {
 	GENERATED_BODY()
@@ -424,7 +424,7 @@ struct CONTROLRIG_API FRigUnit_SetControlVector : public FRigUnitMutable
 /**
  * SetControlRotator is used to perform a change in the hierarchy by setting a single control's Rotator value.
  */
-USTRUCT(meta=(DisplayName="Set Control Rotator", Category="Hierarchy", DocumentationPolicy="Strict", Keywords = "SetControlRotator,SetGizmoRotator", PrototypeName = "SetControlValue"))
+USTRUCT(meta=(DisplayName="Set Control Rotator", Category="Hierarchy", DocumentationPolicy="Strict", Keywords = "SetControlRotator,SetGizmoRotator", PrototypeName = "SetControlValue", NodeColor = "0.0 0.36470600962638855 1.0"))
 struct CONTROLRIG_API FRigUnit_SetControlRotator : public FRigUnitMutable
 {
 	GENERATED_BODY()
@@ -503,7 +503,7 @@ struct CONTROLRIG_API FRigUnit_SetMultiControlRotator_Entry
 /**
  * SetMultiControlRotator is used to perform a change in the hierarchy by setting multiple controls' rotator value.
  */
-USTRUCT(meta = (DisplayName = "Set Multiple Controls Rotator", Category = "Hierarchy", DocumentationPolicy = "Strict", Keywords = "SetMultipleControlsRotator,SetControlRotator,SetGizmoRotator", PrototypeName = "SetMultiControlValue"))
+USTRUCT(meta = (DisplayName = "Set Multiple Controls Rotator", Category = "Hierarchy", DocumentationPolicy = "Strict", Keywords = "SetMultipleControlsRotator,SetControlRotator,SetGizmoRotator", PrototypeName = "SetMultiControlValue", NodeColor = "0.0 0.36470600962638855 1.0"))
 struct CONTROLRIG_API FRigUnit_SetMultiControlRotator : public FRigUnitMutable
 {
 	GENERATED_BODY()
diff --git a/Engine/Plugins/Animation/ControlRig/Source/ControlRig/Private/Units/Highlevel/Hierarchy/RigUnit_CCDIK.h b/Engine/Plugins/Animation/ControlRig/Source/ControlRig/Private/Units/Highlevel/Hierarchy/RigUnit_CCDIK.h
--- a/Engine/Plugins/Animation/ControlRig/Source/ControlRig/Private/Units/Highlevel/Hierarchy/RigUnit_CCDIK.h
+++ b/Engine/Plugins/Animation/ControlRig/Source/ControlRig/Private/Units/Highlevel/Hierarchy/RigUnit_CCDIK.h
@@ -259,7 +259,7 @@ struct CONTROLRIG_API FRigUnit_CCDIKPerItem : public FRigUnit_HighlevelBaseMutab
  * the Cyclic Coordinate Descent Inverse Kinematics algorithm.
  * For now this node supports single effector chains only.
  */
-USTRUCT(meta=(DisplayName="CCDIK", Category="Hierarchy", Keywords="N-Bone,IK"))
+USTRUCT(meta=(DisplayName="CCDIK", Category="Hierarchy", Keywords="N-Bone,IK", NodeColor = "0 1 1"))
 struct CONTROLRIG_API FRigUnit_CCDIKItemArray : public FRigUnit_HighlevelBaseMutable
 {
 	GENERATED_BODY()
diff --git a/Engine/Plugins/Animation/ControlRig/Source/ControlRig/Private/Units/Highlevel/Hierarchy/RigUnit_FABRIK.h b/Engine/Plugins/Animation/ControlRig/Source/ControlRig/Private/Units/Highlevel/Hierarchy/RigUnit_FABRIK.h
--- a/Engine/Plugins/Animation/ControlRig/Source/ControlRig/Private/Units/Highlevel/Hierarchy/RigUnit_FABRIK.h
+++ b/Engine/Plugins/Animation/ControlRig/Source/ControlRig/Private/Units/Highlevel/Hierarchy/RigUnit_FABRIK.h
@@ -179,7 +179,7 @@ struct CONTROLRIG_API FRigUnit_FABRIKPerItem : public FRigUnit_HighlevelBaseMuta
  * the Forward and Backward Reaching Inverse Kinematics algorithm.
  * For now this node supports single effector chains only.
  */
-USTRUCT(meta=(DisplayName="Basic FABRIK", Category="Hierarchy", Keywords="N-Bone,IK"))
+USTRUCT(meta=(DisplayName="Basic FABRIK", Category="Hierarchy", Keywords="N-Bone,IK", NodeColor="0 1 1"))
 struct CONTROLRIG_API FRigUnit_FABRIKItemArray : public FRigUnit_HighlevelBaseMutable
 {
 	GENERATED_BODY()
diff --git a/Engine/Plugins/Animation/ControlRig/Source/ControlRig/Private/Units/Highlevel/Hierarchy/RigUnit_MultiFABRIK.h b/Engine/Plugins/Animation/ControlRig/Source/ControlRig/Private/Units/Highlevel/Hierarchy/RigUnit_MultiFABRIK.h
--- a/Engine/Plugins/Animation/ControlRig/Source/ControlRig/Private/Units/Highlevel/Hierarchy/RigUnit_MultiFABRIK.h
+++ b/Engine/Plugins/Animation/ControlRig/Source/ControlRig/Private/Units/Highlevel/Hierarchy/RigUnit_MultiFABRIK.h
@@ -118,7 +118,7 @@ struct CONTROLRIG_API FRigUnit_MultiFABRIK_EndEffector
  * the Forward and Backward Reaching Inverse Kinematics algorithm.
  * For now this node supports single effector chains only.
  */
-USTRUCT(meta=(DisplayName="Multi Effector FABRIK", Category="Hierarchy", Keywords="Multi, Effector, N-Chain,IK"))
+USTRUCT(meta=(DisplayName="Multi Effector FABRIK", Category="Hierarchy", Keywords="Multi, Effector, N-Chain,IK", NodeColor="0 1 1"))
 struct CONTROLRIG_API FRigUnit_MultiFABRIK : public FRigUnit_HighlevelBaseMutable
 {
 	GENERATED_BODY()
diff --git a/Engine/Plugins/Animation/ControlRig/Source/ControlRig/Private/Units/Highlevel/Hierarchy/RigUnit_SpringIK.h b/Engine/Plugins/Animation/ControlRig/Source/ControlRig/Private/Units/Highlevel/Hierarchy/RigUnit_SpringIK.h
--- a/Engine/Plugins/Animation/ControlRig/Source/ControlRig/Private/Units/Highlevel/Hierarchy/RigUnit_SpringIK.h
+++ b/Engine/Plugins/Animation/ControlRig/Source/ControlRig/Private/Units/Highlevel/Hierarchy/RigUnit_SpringIK.h
@@ -72,7 +72,7 @@ struct CONTROLRIG_API FRigUnit_SpringIK_WorkData
  * It support custom constraints including distance, length etc.
  * Note: This node operates in world space!
  */
-USTRUCT(meta=(DisplayName="Spring IK", Category="Hierarchy", Keywords="IK"))
+USTRUCT(meta=(DisplayName="Spring IK", Category="Hierarchy", Keywords="IK", NodeColor="0 1 1"))
 struct CONTROLRIG_API FRigUnit_SpringIK : public FRigUnit_HighlevelBaseMutable
 {
 	GENERATED_BODY()
diff --git a/Engine/Plugins/Animation/ControlRig/Source/ControlRig/Private/Units/Highlevel/Hierarchy/RigUnit_TwoBoneIKSimple.h b/Engine/Plugins/Animation/ControlRig/Source/ControlRig/Private/Units/Highlevel/Hierarchy/RigUnit_TwoBoneIKSimple.h
--- a/Engine/Plugins/Animation/ControlRig/Source/ControlRig/Private/Units/Highlevel/Hierarchy/RigUnit_TwoBoneIKSimple.h
+++ b/Engine/Plugins/Animation/ControlRig/Source/ControlRig/Private/Units/Highlevel/Hierarchy/RigUnit_TwoBoneIKSimple.h
@@ -207,7 +207,7 @@ struct CONTROLRIG_API FRigUnit_TwoBoneIKSimple : public FRigUnit_HighlevelBaseMu
  * Solves the two bone IK given two bones.
  * Note: This node operates in world space!
  */
-USTRUCT(meta=(DisplayName="Basic IK", Category="Hierarchy", Keywords="TwoBone,IK"))
+USTRUCT(meta=(DisplayName="Basic IK", Category="Hierarchy", Keywords="TwoBone,IK", NodeColor="0 1 1"))
 struct CONTROLRIG_API FRigUnit_TwoBoneIKSimplePerItem : public FRigUnit_HighlevelBaseMutable
 {
 	GENERATED_BODY()
@@ -374,7 +374,7 @@ struct CONTROLRIG_API FRigUnit_TwoBoneIKSimplePerItem : public FRigUnit_Highleve
  * Solves the two bone IK given positions
  * Note: This node operates in world space!
  */
-USTRUCT(meta = (DisplayName = "Basic IK Positions", Category = "Hierarchy", Keywords = "TwoBone,IK"))
+USTRUCT(meta = (DisplayName = "Basic IK Positions", Category = "Hierarchy", Keywords = "TwoBone,IK", NodeColor = "0 1 1"))
 struct CONTROLRIG_API FRigUnit_TwoBoneIKSimpleVectors : public FRigUnit_HighlevelBase
 {
 	GENERATED_BODY()
@@ -452,7 +452,7 @@ struct CONTROLRIG_API FRigUnit_TwoBoneIKSimpleVectors : public FRigUnit_Highleve
  * Solves the two bone IK given transforms
  * Note: This node operates in world space!
  */
-USTRUCT(meta = (DisplayName = "Basic IK Transforms", Category = "Hierarchy", Keywords = "TwoBone,IK"))
+USTRUCT(meta = (DisplayName = "Basic IK Transforms", Category = "Hierarchy", Keywords = "TwoBone,IK", NodeColor = "0 1 1"))
 struct CONTROLRIG_API FRigUnit_TwoBoneIKSimpleTransforms : public FRigUnit_HighlevelBase
 {
 	GENERATED_BODY()
diff --git a/Engine/Plugins/Animation/ControlRig/Source/ControlRig/Public/Units/Hierarchy/RigUnit_SetCurveValue.h b/Engine/Plugins/Animation/ControlRig/Source/ControlRig/Public/Units/Hierarchy/RigUnit_SetCurveValue.h
--- a/Engine/Plugins/Animation/ControlRig/Source/ControlRig/Public/Units/Hierarchy/RigUnit_SetCurveValue.h
+++ b/Engine/Plugins/Animation/ControlRig/Source/ControlRig/Public/Units/Hierarchy/RigUnit_SetCurveValue.h
@@ -9,7 +9,7 @@
 /**
  * SetCurveValue is used to perform a change in the curve container by setting a single Curve value.
  */
-USTRUCT(meta=(DisplayName="Set Curve Value", Category="Curve", Keywords = "SetCurveValue"))
+USTRUCT(meta=(DisplayName="Set Curve Value", Category="Curve", Keywords = "SetCurveValue", NodeColor = "0.0 0.36470600962638855 1.0"))
 struct FRigUnit_SetCurveValue : public FRigUnitMutable
 {
 	GENERATED_BODY()
diff --git a/Engine/Plugins/Animation/ControlRigSpline/Source/ControlRigSpline/Public/ControlRigSplineUnits.h b/Engine/Plugins/Animation/ControlRigSpline/Source/ControlRigSpline/Public/ControlRigSplineUnits.h
--- a/Engine/Plugins/Animation/ControlRigSpline/Source/ControlRigSpline/Public/ControlRigSplineUnits.h
+++ b/Engine/Plugins/Animation/ControlRigSpline/Source/ControlRigSpline/Public/ControlRigSplineUnits.h
@@ -14,7 +14,7 @@ struct CONTROLRIGSPLINE_API FRigUnit_ControlRigSplineBase : public FRigUnit
 	GENERATED_BODY()
 };
 
-USTRUCT(meta = (DisplayName = "Spline From Points", Keywords="Spline From Positions", Category = "Spline", Varying))
+USTRUCT(meta = (DisplayName = "Spline From Points", Keywords="Spline From Positions", Category = "Spline", Varying, NodeColor="0.737911 0.099899 0.099899"))
 struct CONTROLRIGSPLINE_API FRigUnit_ControlRigSplineFromPoints : public FRigUnit_ControlRigSplineBase
 {
 	GENERATED_BODY()
@@ -50,7 +50,7 @@ struct CONTROLRIGSPLINE_API FRigUnit_ControlRigSplineFromPoints : public FRigUni
 	FControlRigSpline Spline;
 };
 
-USTRUCT(meta = (DisplayName = "Set Spline Points", Category = "Spline", Varying))
+USTRUCT(meta = (DisplayName = "Set Spline Points", Category = "Spline", Varying, NodeColor = "0.0 0.36470600962638855 1.0"))
 struct CONTROLRIGSPLINE_API FRigUnit_SetSplinePoints : public FRigUnitMutable
 {
 	GENERATED_BODY()
@@ -71,7 +71,7 @@ struct CONTROLRIGSPLINE_API FRigUnit_SetSplinePoints : public FRigUnitMutable
 	FControlRigSpline Spline;
 };
 
-USTRUCT(meta = (DisplayName = "Position From Spline", Keywords="Point From Spline", Category = "Spline"))
+USTRUCT(meta = (DisplayName = "Position From Spline", Keywords="Point From Spline", Category = "Spline", NodeColor="0.737911 0.099899 0.099899"))
 struct CONTROLRIGSPLINE_API FRigUnit_PositionFromControlRigSpline : public FRigUnit_ControlRigSplineBase
 {
 	GENERATED_BODY()
@@ -96,7 +96,7 @@ struct CONTROLRIGSPLINE_API FRigUnit_PositionFromControlRigSpline : public FRigU
 	FVector Position;
 };
 
-USTRUCT(meta = (DisplayName = "Transform From Spline", Category = "Spline"))
+USTRUCT(meta = (DisplayName = "Transform From Spline", Category = "Spline", NodeColor="0.737911 0.099899 0.099899"))
 struct CONTROLRIGSPLINE_API FRigUnit_TransformFromControlRigSpline : public FRigUnit_ControlRigSplineBase
 {
 	GENERATED_BODY()
@@ -129,7 +129,7 @@ struct CONTROLRIGSPLINE_API FRigUnit_TransformFromControlRigSpline : public FRig
 	FTransform Transform;
 };
 
-USTRUCT(meta = (DisplayName = "Tangent From Spline", Category = "Spline"))
+USTRUCT(meta = (DisplayName = "Tangent From Spline", Category = "Spline", NodeColor="0.737911 0.099899 0.099899"))
 struct CONTROLRIGSPLINE_API FRigUnit_TangentFromControlRigSpline : public FRigUnit_ControlRigSplineBase
 {
 	GENERATED_BODY()
@@ -154,7 +154,7 @@ struct CONTROLRIGSPLINE_API FRigUnit_TangentFromControlRigSpline : public FRigUn
 	FVector Tangent;
 };
 
-USTRUCT(meta = (DisplayName = "Draw Spline", Category = "Spline"))
+USTRUCT(meta = (DisplayName = "Draw Spline", Category = "Spline", NodeColor = "0.83077 0.846873 0.049707"))
 struct CONTROLRIGSPLINE_API FRigUnit_DrawControlRigSpline : public FRigUnitMutable
 {
 	GENERATED_BODY()
@@ -183,7 +183,7 @@ struct CONTROLRIGSPLINE_API FRigUnit_DrawControlRigSpline : public FRigUnitMutab
 	int32 Detail;
 };
 
-USTRUCT(meta = (DisplayName = "Get Length Of Spline", Category = "Spline"))
+USTRUCT(meta = (DisplayName = "Get Length Of Spline", Category = "Spline", NodeColor="0.737911 0.099899 0.099899"))
 struct CONTROLRIGSPLINE_API FRigUnit_GetLengthControlRigSpline : public FRigUnit
 {
 	GENERATED_BODY()
@@ -494,7 +494,7 @@ struct CONTROLRIGSPLINE_API FRigUnit_FitSplineCurveToChainItemArray : public FRi
 	FControlRigSpline Spline;	
 };
 
-USTRUCT(meta = (DisplayName = "Closest Parameter From Spline", Category = "Spline"))
+USTRUCT(meta = (DisplayName = "Closest Parameter From Spline", Category = "Spline", NodeColor="0.737911 0.099899 0.099899"))
 struct CONTROLRIGSPLINE_API FRigUnit_ClosestParameterFromControlRigSpline : public FRigUnit_ControlRigSplineBase
 {
 	GENERATED_BODY()
diff --git a/Engine/Plugins/Experimental/FullBodyIK/Source/PBIK/Public/RigUnit_PBIK.h b/Engine/Plugins/Experimental/FullBodyIK/Source/PBIK/Public/RigUnit_PBIK.h
--- a/Engine/Plugins/Experimental/FullBodyIK/Source/PBIK/Public/RigUnit_PBIK.h
+++ b/Engine/Plugins/Experimental/FullBodyIK/Source/PBIK/Public/RigUnit_PBIK.h
@@ -89,7 +89,10 @@ struct FPBIKEffector
 	float PinRotation = 1.0f;
 };
 
-USTRUCT(meta=(DisplayName="Full Body IK", Category="Hierarchy", Keywords="Position Based, PBIK, IK, Full Body, Multi, Effector, N-Chain, FB"))
+/*
+ * Based on a Position Based solver at core, this node can solve multi chains within a root using multi effectors
+ */
+USTRUCT(meta=(DisplayName="Full Body IK", Category="Hierarchy", Keywords="Position Based, PBIK, IK, Full Body, Multi, Effector, N-Chain, FB", NodeColor="0 1 1"))
 struct FRigUnit_PBIK : public FRigUnit_HighlevelBaseMutable
 {
 	GENERATED_BODY()
