Hide plus icon for URigVMFunctionReturnNode's array pins

diff --git a/Engine/Plugins/Animation/ControlRig/Source/ControlRigEditor/Private/Graph/SControlRigGraphNode.cpp b/Engine/Plugins/Animation/ControlRig/Source/ControlRigEditor/Private/Graph/SControlRigGraphNode.cpp
--- a/Engine/Plugins/Animation/ControlRig/Source/ControlRigEditor/Private/Graph/SControlRigGraphNode.cpp
+++ b/Engine/Plugins/Animation/ControlRig/Source/ControlRigEditor/Private/Graph/SControlRigGraphNode.cpp
@@ -733,6 +733,12 @@ EVisibility SControlRigGraphNode::GetArrayPlusButtonVisibility(URigVMPin* InMode
 {
 	if(InModelPin)
 	{
+		URigVMNode* RigVMNode = InModelPin->GetNode();
+		if (RigVMNode->IsA<URigVMFunctionReturnNode>())
+		{
+			return EVisibility::Hidden;
+		}
+		
 		if(InModelPin->GetSourceLinks().Num() == 0)
 		{
 			return EVisibility::Visible;
