ControlRig: Fix pin control issues, jittering, wrong parents, does not work in setup mode, does not work when there is no skeletal mesh, crash during auto save after pin control moves.

diff --git a/Engine/Plugins/Animation/ControlRig/Source/ControlRig/Private/ControlRig.cpp b/Engine/Plugins/Animation/ControlRig/Source/ControlRig/Private/ControlRig.cpp
--- a/Engine/Plugins/Animation/ControlRig/Source/ControlRig/Private/ControlRig.cpp
+++ b/Engine/Plugins/Animation/ControlRig/Source/ControlRig/Private/ControlRig.cpp
@@ -712,72 +712,77 @@ void UControlRig::Execute(const EControlRigState InState, const FName& InEventNa
 		if(!IsRunningPreSetup() && !IsRunningPostSetup())
 		{
 			bRequiresSetupEvent = bSetupModeEnabled;
-
-			if (UControlRig* CDO = Cast<UControlRig>(GetClass()->GetDefaultObject()))
 			{
-				if(bCopyHierarchyBeforeSetup && !bSetupModeEnabled)
+				// save the current state of all pose elements to preserve user intention, since setup event can
+				// run in between forward events
+				// the saved pose is reapplied to the rig after setup event as the pose scope goes out of scope
+				TUniquePtr<UControlRig::FPoseScope> PoseScope;
+				if (!bSetupModeEnabled)
 				{
-					if(CDO->GetHierarchy()->GetTopologyVersion()!= GetHierarchy()->GetTopologyVersion())
-					{
-#if WITH_EDITOR
-						FTransientControlScope TransientControlScope(GetHierarchy());
-#endif
-						GetHierarchy()->CopyHierarchy(CDO->GetHierarchy());
-					}
+					// only do this in non-setup mode because 
+					// when setup mode is enabled, the control values are cleared before reaching here (too late to save them)
+					PoseScope = MakeUnique<UControlRig::FPoseScope>(this, ERigElementType::ToResetAfterSetupEvent);
 				}
 				
-				if (bResetInitialTransformsBeforeSetup && !bSetupModeEnabled)
+				if (UControlRig* CDO = Cast<UControlRig>(GetClass()->GetDefaultObject()))
 				{
-					GetHierarchy()->CopyPose(CDO->GetHierarchy(), false, true);
-				}
-			}
-
-			{
+					if(bCopyHierarchyBeforeSetup && !bSetupModeEnabled)
+					{
+						if(CDO->GetHierarchy()->GetTopologyVersion()!= GetHierarchy()->GetTopologyVersion())
+						{
 #if WITH_EDITOR
-				TUniquePtr<FTransientControlPoseScope> TransientControlPoseScope;
-				if (bSetupModeEnabled)
-				{
-					// save the transient control value, it should not be constantly reset in setup mode
-					TransientControlPoseScope = MakeUnique<FTransientControlPoseScope>(this);
-				}
+							FTransientControlScope TransientControlScope(GetHierarchy());
 #endif
-				{
-					// save the current state of all pose elements to preserve user intention
-					// the saved pose is reapplied to the rig after setup event as the pose scope goes out of scope
-					TUniquePtr<UControlRig::FPoseScope> PoseScope;
-					if (!bSetupModeEnabled)
-					{
-						// only do this in non-setup mode because 
-						// when setup mode is enabled, the control values are cleared before reaching here (too late to save them)
-						PoseScope = MakeUnique<UControlRig::FPoseScope>(this, ERigElementType::ToResetAfterSetupEvent);
+							GetHierarchy()->CopyHierarchy(CDO->GetHierarchy());
+						}
 					}
-
+					
+					if (bResetInitialTransformsBeforeSetup && !bSetupModeEnabled)
 					{
-
-						// reset the pose to initial such that setup event can run from a deterministic initial state
-						GetHierarchy()->ResetPoseToInitial(ERigElementType::All);
+						GetHierarchy()->CopyPose(CDO->GetHierarchy(), false, true);
 					}
+				}
 
-					if (PreSetupEvent.IsBound())
+				{
+#if WITH_EDITOR
+					TUniquePtr<FTransientControlPoseScope> TransientControlPoseScope;
+					if (bSetupModeEnabled)
 					{
-						FControlRigBracketScope BracketScope(PreSetupBracket);
-						PreSetupEvent.Broadcast(this, EControlRigState::Update, FRigUnit_PrepareForExecution::EventName);
+						// save the transient control value, it should not be constantly reset in setup mode
+						TransientControlPoseScope = MakeUnique<FTransientControlPoseScope>(this);
 					}
+#endif
+					// reset the pose to initial such that setup event can run from a deterministic initial state
+					GetHierarchy()->ResetPoseToInitial(ERigElementType::All);
+				}
+
+				if (PreSetupEvent.IsBound())
+				{
+					FControlRigBracketScope BracketScope(PreSetupBracket);
+					PreSetupEvent.Broadcast(this, EControlRigState::Update, FRigUnit_PrepareForExecution::EventName);
+				}
 
-					ExecuteUnits(Context, FRigUnit_PrepareForExecution::EventName);
+				ExecuteUnits(Context, FRigUnit_PrepareForExecution::EventName);
 
-					if (PostSetupEvent.IsBound())
-					{
-						FControlRigBracketScope BracketScope(PostSetupBracket);
-						PostSetupEvent.Broadcast(this, EControlRigState::Update, FRigUnit_PrepareForExecution::EventName);
-					}
+				if (PostSetupEvent.IsBound())
+				{
+					FControlRigBracketScope BracketScope(PostSetupBracket);
+					PostSetupEvent.Broadcast(this, EControlRigState::Update, FRigUnit_PrepareForExecution::EventName);
 				}
+			}
 
+			if (bSetupModeEnabled)
+			{
+#if WITH_EDITOR
+				TUniquePtr<FTransientControlPoseScope> TransientControlPoseScope;
 				if (bSetupModeEnabled)
 				{
-					GetHierarchy()->ResetPoseToInitial(ERigElementType::Bone);
+					// save the transient control value, it should not be constantly reset in setup mode
+					TransientControlPoseScope = MakeUnique<FTransientControlPoseScope>(this);
 				}
-			}
+#endif
+				GetHierarchy()->ResetPoseToInitial(ERigElementType::Bone);
+			}			
 		}
 		else
 		{
@@ -1135,7 +1140,8 @@ void UControlRig::ExecuteUnits(FRigUnitContext& InOutContext, const FName& InEve
 				{
 					VM->Initialize(LocalMemory, AdditionalArguments);
 
-					URigVM* InitializedVMSnapshot = NewObject<URigVM>(CDO, NAME_None, RF_Public);
+					// objects assigned to transient properties need transient flag, otherwise it might get exported during save
+					URigVM* InitializedVMSnapshot = NewObject<URigVM>(CDO, NAME_None, RF_Public | RF_Transient);
 					InitializedVMSnapshot->WorkMemoryStorageObject = NewObject<URigVMMemoryStorage>(InitializedVMSnapshot, VM->GetWorkMemory()->GetClass());
 					InitializedVMSnapshot->WorkMemoryStorageObject->CopyFrom(VM->WorkMemoryStorageObject);
 
@@ -1720,17 +1726,11 @@ FName UControlRig::AddTransientControl(URigVMPin* InPin, FRigElementKey SpaceKey
 	Settings.bIsTransientControl = true;
 	Settings.DisplayName = TEXT("Temporary Control");
 
-	FRigElementKey Parent;
-	if(SpaceKey.IsValid() && SpaceKey.Type == ERigElementType::Bone)
-	{
-		Parent = DynamicHierarchy->GetFirstParent(SpaceKey);
-	}
-
 	Controller->ClearSelection();
 
     const FRigElementKey ControlKey = Controller->AddControl(
     	ControlName,
-    	Parent,
+    	SpaceKey,
     	Settings,
     	FRigControlValue::Make(FTransform::Identity),
     	OffsetTransform,
@@ -2987,27 +2987,6 @@ UControlRig::FTransientControlScope::~FTransientControlScope()
 	}
 }
 
-UControlRig::FTransientControlPoseScope::FTransientControlPoseScope(TObjectPtr<UControlRig> InControlRig)
-{
-	ControlRig = InControlRig;
-
-	TArray<FRigControlElement*> TransientControls = ControlRig->GetHierarchy()->GetTransientControls();
-	TArray<FRigElementKey> Keys;
-	for(FRigControlElement* TransientControl : TransientControls)
-	{
-		Keys.Add(TransientControl->GetKey());
-	}
-	
-	CachedPose = ControlRig->GetHierarchy()->GetPose(false, ERigElementType::Control, TArrayView<FRigElementKey>(Keys));
-}
-
-UControlRig::FTransientControlPoseScope::~FTransientControlPoseScope()
-{
-	check(ControlRig);
-
-	ControlRig->GetHierarchy()->SetPose(CachedPose);
-}
-
 #endif
  
 #undef LOCTEXT_NAMESPACE
diff --git a/Engine/Plugins/Animation/ControlRig/Source/ControlRig/Private/Units/Hierarchy/RigUnit_OffsetTransform.h b/Engine/Plugins/Animation/ControlRig/Source/ControlRig/Private/Units/Hierarchy/RigUnit_OffsetTransform.h
--- a/Engine/Plugins/Animation/ControlRig/Source/ControlRig/Private/Units/Hierarchy/RigUnit_OffsetTransform.h
+++ b/Engine/Plugins/Animation/ControlRig/Source/ControlRig/Private/Units/Hierarchy/RigUnit_OffsetTransform.h
@@ -23,7 +23,22 @@ struct CONTROLRIG_API FRigUnit_OffsetTransformForItem : public FRigUnitMutable
 
 	virtual FRigElementKey DetermineSpaceForPin(const FString& InPinPath, void* InUserContext) const override
 	{
-		return Item;
+		if (const URigHierarchy* Hierarchy = (const URigHierarchy*)InUserContext)
+		{
+			return Hierarchy->GetFirstParent(Item);
+		}
+		return FRigElementKey();
+	}
+
+	virtual FTransform DetermineOffsetTransformForPin(const FString& InPinPath, void* InUserContext) const override
+	{
+		if (const URigHierarchy* Hierarchy = (const URigHierarchy*)InUserContext)
+		{
+			// this is similar to RigUnit_ModifyTransform
+			return OffsetTransform.Inverse() * Hierarchy->GetLocalTransform(Item);
+		}
+		
+		return FTransform::Identity;
 	}
 
 	RIGVM_METHOD()
diff --git a/Engine/Plugins/Animation/ControlRig/Source/ControlRig/Private/Units/Highlevel/Hierarchy/RigUnit_ModifyTransforms.h b/Engine/Plugins/Animation/ControlRig/Source/ControlRig/Private/Units/Highlevel/Hierarchy/RigUnit_ModifyTransforms.h
--- a/Engine/Plugins/Animation/ControlRig/Source/ControlRig/Private/Units/Highlevel/Hierarchy/RigUnit_ModifyTransforms.h
+++ b/Engine/Plugins/Animation/ControlRig/Source/ControlRig/Private/Units/Highlevel/Hierarchy/RigUnit_ModifyTransforms.h
@@ -101,8 +101,7 @@ struct CONTROLRIG_API FRigUnit_ModifyTransforms : public FRigUnit_HighlevelBaseM
 				{
 					if (const URigHierarchy* Hierarchy = (const URigHierarchy*)InUserContext)
 					{
-						// in UControlRig::AddTransientControl, the transient control will be added to the parent of this item
-						return ItemToModify[Index].Item;
+						return Hierarchy->GetFirstParent(ItemToModify[Index].Item);
 					}
 				}
 			}
@@ -126,20 +125,20 @@ struct CONTROLRIG_API FRigUnit_ModifyTransforms : public FRigUnit_HighlevelBaseM
 
 			if (ItemToModify.IsValidIndex(Index))
 			{
-				if (Mode == EControlRigModifyBoneMode::AdditiveLocal)
+				if (ItemToModify[Index].Item.IsValid())
 				{
 					if (const URigHierarchy* Hierarchy = (const URigHierarchy*)InUserContext)
 					{
-						if (ItemToModify[Index].Item.IsValid())
+						if (Mode == EControlRigModifyBoneMode::AdditiveLocal)
 						{
 							// in the case of AdditiveLocal, we want the transient control's global transform to match the item's global transform
-							// we know: transient control's global = transform value * offset transfom * parent transform
+							// we know: transient control's global = transform value * offset transform * parent transform
 							// we know: item's global transform = item's local transform from rig graph evaluation * its parent transform
 							// we ensure: the item and the transient control share the same parent
 							// so from our goal, transient control's global == item's global transform
-							// we get: transform value * offset transfom * parent transform == item's local transform from rig graph evaluation * its parent transform
-							// we get: transform value * offset transfom == item's local transform from rig graph evaluation
-							// we get: offset transfom == transform value.inverse * item's local transform from rig graph evaluation
+							// we get: transform value * offset transform * parent transform == item's local transform from rig graph evaluation * its parent transform
+							// we get: transform value * offset transform == item's local transform from rig graph evaluation
+							// we get: offset transform == transform value.inverse * item's local transform from rig graph evaluation
 							return ItemToModify[Index].Transform.Inverse() * Hierarchy->GetLocalTransform(ItemToModify[Index].Item);
 						}
 					}
diff --git a/Engine/Plugins/Animation/ControlRig/Source/ControlRig/Public/ControlRig.h b/Engine/Plugins/Animation/ControlRig/Source/ControlRig/Public/ControlRig.h
--- a/Engine/Plugins/Animation/ControlRig/Source/ControlRig/Public/ControlRig.h
+++ b/Engine/Plugins/Animation/ControlRig/Source/ControlRig/Public/ControlRig.h
@@ -796,8 +796,25 @@ public:
 	class FTransientControlPoseScope
 	{
 	public:
-		FTransientControlPoseScope(TObjectPtr<UControlRig> InControlRig);
-		~FTransientControlPoseScope();
+		FTransientControlPoseScope(TObjectPtr<UControlRig> InControlRig)
+		{
+			ControlRig = InControlRig;
+
+			TArray<FRigControlElement*> TransientControls = ControlRig->GetHierarchy()->GetTransientControls();
+			TArray<FRigElementKey> Keys;
+			for(FRigControlElement* TransientControl : TransientControls)
+			{
+				Keys.Add(TransientControl->GetKey());
+			}
+	
+			CachedPose = ControlRig->GetHierarchy()->GetPose(false, ERigElementType::Control, TArrayView<FRigElementKey>(Keys));
+		}
+		~FTransientControlPoseScope()
+		{
+			check(ControlRig);
+
+			ControlRig->GetHierarchy()->SetPose(CachedPose);
+		}
 	
 	private:
 		
diff --git a/Engine/Plugins/Animation/ControlRig/Source/ControlRigDeveloper/Private/ControlRigBlueprint.cpp b/Engine/Plugins/Animation/ControlRig/Source/ControlRigDeveloper/Private/ControlRigBlueprint.cpp
--- a/Engine/Plugins/Animation/ControlRig/Source/ControlRigDeveloper/Private/ControlRigBlueprint.cpp
+++ b/Engine/Plugins/Animation/ControlRig/Source/ControlRigDeveloper/Private/ControlRigBlueprint.cpp
@@ -3508,7 +3508,10 @@ void UControlRigBlueprint::PatchVariableNodesOnLoad()
 void UControlRigBlueprint::PropagatePoseFromInstanceToBP(UControlRig* InControlRig)
 {
 	check(InControlRig);
-	Hierarchy->CopyPose(InControlRig->GetHierarchy(), true, true);
+	// current transforms in BP and CDO are meaningless, no need to copy them
+	// we use BP hierarchy to initialize CDO and instances' hierarchy, 
+	// so it should always be in the initial state.
+	Hierarchy->CopyPose(InControlRig->GetHierarchy(), false, true);
 }
 
 void UControlRigBlueprint::PropagatePoseFromBPToInstances()
diff --git a/Engine/Plugins/Animation/ControlRig/Source/ControlRigEditor/Private/Editor/ControlRigEditor.cpp b/Engine/Plugins/Animation/ControlRig/Source/ControlRigEditor/Private/Editor/ControlRigEditor.cpp
--- a/Engine/Plugins/Animation/ControlRig/Source/ControlRigEditor/Private/Editor/ControlRigEditor.cpp
+++ b/Engine/Plugins/Animation/ControlRig/Source/ControlRigEditor/Private/Editor/ControlRigEditor.cpp
@@ -2971,8 +2971,12 @@ void FControlRigEditor::Tick(float DeltaTime)
 			ControlRig != nullptr && 
 			bExecutionControlRig)
 		{
-			// reset transforms here to prevent additive transforms from accumulating to INF
-			ControlRig->GetHierarchy()->ResetPoseToInitial(ERigElementType::Bone);
+			{
+				// prevent transient controls from getting reset
+				UControlRig::FTransientControlPoseScope	PoseScope(ControlRig);
+				// reset transforms here to prevent additive transforms from accumulating to INF
+				ControlRig->GetHierarchy()->ResetPoseToInitial(ERigElementType::Bone);
+			}
 			
 			ControlRig->SetDeltaTime(DeltaTime);
 			ControlRig->Evaluate_AnyThread();
@@ -3635,9 +3639,14 @@ bool FControlRigEditor::IsPinControlNameListEnabled() const
 {
 	if (ControlRig)
 	{
-		if (ControlRig->GetHierarchy()->GetTransientControls().Num() > 0)
+		TArray<FRigControlElement*>	TransientControls = ControlRig->GetHierarchy()->GetTransientControls();
+		if (TransientControls.Num() > 0)
 		{
-			return true;
+			// if the transient control is not for a rig element, it is for a pin
+			if (UControlRig::GetElementKeyFromTransientControl(TransientControls[0]->GetKey()) == FRigElementKey())
+			{
+				return true;
+			}
 		}
 	}
 	return false;
@@ -3653,11 +3662,13 @@ FText FControlRigEditor::GetPinControlNameListText() const
 	if (ControlRig)
 	{
 		FText Result;
-		ControlRig->GetHierarchy()->ForEach<FRigControlElement>([&Result](FRigControlElement* ControlElement) -> bool
+		ControlRig->GetHierarchy()->ForEach<FRigControlElement>([this, &Result](FRigControlElement* ControlElement) -> bool
         {
 			if (ControlElement->Settings.bIsTransientControl)
 			{
-				Result = FText::FromName(ControlElement->GetName());
+				FRigElementKey Parent = ControlRig->GetHierarchy()->GetFirstParent(ControlElement->GetKey());
+				Result = FText::FromName(Parent.Name);
+				
 				return false;
 			}
 			return true;
@@ -3697,24 +3708,41 @@ void FControlRigEditor::SetPinControlNameListText(const FText& NewTypeInValue, E
 				if (NewParentIndex == INDEX_NONE)
 				{
 					NewParentName = NAME_None;
-					GetControlRigBlueprint()->GetHierarchyController()->RemoveAllParents(ControlElement->GetKey(), true, false);
+					ControlRig->GetHierarchy()->GetController()->RemoveAllParents(ControlElement->GetKey(), true, false);
 				}
 				else
 				{
-					GetControlRigBlueprint()->GetHierarchyController()->SetParent(ControlElement->GetKey(), FRigElementKey(NewParentName, ERigElementType::Bone), true, false);
+					ControlRig->GetHierarchy()->GetController()->SetParent(ControlElement->GetKey(), FRigElementKey(NewParentName, ERigElementType::Bone), true, false);
 				}
 
 				// find out if the controlled pin is part of a visual debug node
 				if (UControlRigBlueprint* ControlRigBlueprint = CastChecked<UControlRigBlueprint>(GetBlueprintObj()))
 				{
-					if (URigVMPin* ControlledPin = GetFocusedModel()->FindPin(ControlElement->GetName().ToString()))
+					FString PinName = UControlRig::GetPinNameFromTransientControl(ControlElement->GetKey());
+					if (URigVMPin* ControlledPin = GetFocusedModel()->FindPin(PinName))
 					{
 						URigVMNode* ControlledNode = ControlledPin->GetPinForLink()->GetNode();
-						if (URigVMPin* BoneSpacePin = ControlledNode->FindPin(TEXT("BoneSpace")))
+						if (URigVMUnitNode* UnitNode = Cast<URigVMUnitNode>(ControlledNode))
 						{
-							if (BoneSpacePin->GetCPPType() == TEXT("FName") && BoneSpacePin->GetCustomWidgetName() == TEXT("BoneName"))
+							if (const FString* Value = UnitNode->GetScriptStruct()->FindMetaData(TEXT("PrototypeName")))
 							{
-								GetFocusedController()->SetPinDefaultValue(BoneSpacePin->GetPinPath(), NewParentName.ToString(), false, false, false);
+								if (Value->Equals("VisualDebug"))
+								{
+									if (URigVMPin* SpacePin = ControlledNode->FindPin(TEXT("Space")))
+									{
+										FString DefaultValue;
+										const FRigElementKey NewSpaceKey(NewParentName, ERigElementType::Bone); 
+										FRigElementKey::StaticStruct()->ExportText(DefaultValue, &NewSpaceKey, nullptr, nullptr, PPF_None, nullptr);
+										ensure(GetFocusedController()->SetPinDefaultValue(SpacePin->GetPinPath(), DefaultValue, false, false, false));
+									}
+									else if (URigVMPin* BoneSpacePin = ControlledNode->FindPin(TEXT("BoneSpace")))
+									{
+										if (BoneSpacePin->GetCPPType() == TEXT("FName") && BoneSpacePin->GetCustomWidgetName() == TEXT("BoneName"))
+										{
+											GetFocusedController()->SetPinDefaultValue(BoneSpacePin->GetPinPath(), NewParentName.ToString(), false, false, false);
+										}
+									}
+								}
 							}
 						}
 					}
