Control Rig: Determine instruction index for UI entry / return nodes

diff --git a/Engine/Plugins/Animation/ControlRig/Source/ControlRigDeveloper/Private/Graph/ControlRigGraph.cpp b/Engine/Plugins/Animation/ControlRig/Source/ControlRigDeveloper/Private/Graph/ControlRigGraph.cpp
--- a/Engine/Plugins/Animation/ControlRig/Source/ControlRigDeveloper/Private/Graph/ControlRigGraph.cpp
+++ b/Engine/Plugins/Animation/ControlRig/Source/ControlRigDeveloper/Private/Graph/ControlRigGraph.cpp
@@ -730,6 +730,42 @@ int32 UControlRigGraph::GetInstructionIndex(const UControlRigGraphNode* InNode,
 
 				return InstructionIndex;
 			}
+			else if(URigVMFunctionEntryNode* EntryNode = Cast<URigVMFunctionEntryNode>(InModelNode))
+			{
+				int32 InstructionIndex = INDEX_NONE;
+				if(!bAsInput)
+				{
+					TArray<URigVMNode*> TargetNodes = EntryNode->GetLinkedTargetNodes();
+					for(URigVMNode* TargetNode : TargetNodes)
+					{
+						InstructionIndex = GetInstructionIndex(TargetNode, InByteCode, Indices, bAsInput);
+						if(InstructionIndex != INDEX_NONE)
+						{
+							break;
+						}
+					}
+					Indices.FindOrAdd(InModelNode).Key = InstructionIndex;
+				}
+				return InstructionIndex;
+			}
+			else if(URigVMFunctionReturnNode* ReturnNode = Cast<URigVMFunctionReturnNode>(InModelNode))
+			{
+				int32 InstructionIndex = INDEX_NONE;
+				if(bAsInput)
+				{
+					TArray<URigVMNode*> SourceNodes = ReturnNode->GetLinkedSourceNodes();
+					for(URigVMNode* SourceNode : SourceNodes)
+					{
+						InstructionIndex = GetInstructionIndex(SourceNode, InByteCode, Indices, bAsInput);
+						if(InstructionIndex != INDEX_NONE)
+						{
+							break;
+						}
+					}
+					Indices.FindOrAdd(InModelNode).Key = InstructionIndex;
+				}
+				return InstructionIndex;
+			}
 
 			Indices.FindOrAdd(InModelNode, TPair<int32, int32>(INDEX_NONE, INDEX_NONE));
 
