Control Rig: Further adapt node colors

diff --git a/Engine/Plugins/Animation/ControlRig/Source/ControlRig/Private/Units/Debug/RigUnit_VisualDebug.h b/Engine/Plugins/Animation/ControlRig/Source/ControlRig/Private/Units/Debug/RigUnit_VisualDebug.h
--- a/Engine/Plugins/Animation/ControlRig/Source/ControlRig/Private/Units/Debug/RigUnit_VisualDebug.h
+++ b/Engine/Plugins/Animation/ControlRig/Source/ControlRig/Private/Units/Debug/RigUnit_VisualDebug.h
@@ -62,7 +62,7 @@ struct CONTROLRIG_API FRigUnit_VisualDebugVector : public FRigUnit_DebugBase
 /**
  * Debug draw parameters for a Point or Vector given a vector
  */
-USTRUCT(meta=(DisplayName = "Visual Debug Vector", PrototypeName = "VisualDebug", Keywords = "Draw,Point", Varying, NodeColor = "0.25 0.25 0.050000000745058"))
+USTRUCT(meta=(DisplayName = "Visual Debug Vector", PrototypeName = "VisualDebug", Keywords = "Draw,Point", Varying))
 struct CONTROLRIG_API FRigUnit_VisualDebugVectorItemSpace : public FRigUnit_DebugBase
 {
 	GENERATED_BODY()
@@ -144,7 +144,7 @@ struct CONTROLRIG_API FRigUnit_VisualDebugQuat : public FRigUnit_DebugBase
 /**
  * Debug draw parameters for an Axis given a quaternion
  */
-USTRUCT(meta = (DisplayName = "Visual Debug Quat", PrototypeName = "VisualDebug", Keywords = "Draw,Rotation", Varying, NodeColor = "0.25 0.25 0.050000000745058"))
+USTRUCT(meta = (DisplayName = "Visual Debug Quat", PrototypeName = "VisualDebug", Keywords = "Draw,Rotation", Varying))
 struct CONTROLRIG_API FRigUnit_VisualDebugQuatItemSpace : public FRigUnit_DebugBase
 {
 	GENERATED_BODY()
@@ -218,7 +218,7 @@ struct CONTROLRIG_API FRigUnit_VisualDebugTransform : public FRigUnit_DebugBase
 /**
  * Debug draw parameters for an Axis given a transform
  */
-USTRUCT(meta=(DisplayName="Visual Debug Transform", PrototypeName = "VisualDebug", Keywords = "Draw,Axes", Varying, NodeColor = "0.25 0.25 0.050000000745058"))
+USTRUCT(meta=(DisplayName="Visual Debug Transform", PrototypeName = "VisualDebug", Keywords = "Draw,Axes", Varying))
 struct CONTROLRIG_API FRigUnit_VisualDebugTransformItemSpace : public FRigUnit_DebugBase
 {
 	GENERATED_BODY()
diff --git a/Engine/Plugins/Animation/ControlRig/Source/ControlRig/Private/Units/Hierarchy/RigUnit_SetControlColor.h b/Engine/Plugins/Animation/ControlRig/Source/ControlRig/Private/Units/Hierarchy/RigUnit_SetControlColor.h
--- a/Engine/Plugins/Animation/ControlRig/Source/ControlRig/Private/Units/Hierarchy/RigUnit_SetControlColor.h
+++ b/Engine/Plugins/Animation/ControlRig/Source/ControlRig/Private/Units/Hierarchy/RigUnit_SetControlColor.h
@@ -8,7 +8,7 @@
 /**
  * SetControlColor is used to change the gizmo's color on a control at runtime
  */
-USTRUCT(meta=(DisplayName="Set Control Color", Category="Hierarchy", DocumentationPolicy="Strict", Keywords = "SetControlColor,SetGizmoColor", PrototypeName = "SetControlColor"))
+USTRUCT(meta=(DisplayName="Set Control Color", Category="Hierarchy", DocumentationPolicy="Strict", Keywords = "SetControlColor,SetGizmoColor", PrototypeName = "SetControlColor", NodeColor = "0.0 0.36470600962638855 1.0"))
 struct CONTROLRIG_API FRigUnit_SetControlColor : public FRigUnitMutable
 {
 	GENERATED_BODY()
