Control Rig: Further adapt node display names

diff --git a/Engine/Plugins/Animation/ControlRig/Source/ControlRig/Private/Units/Simulation/RigUnit_AlphaInterp.h b/Engine/Plugins/Animation/ControlRig/Source/ControlRig/Private/Units/Simulation/RigUnit_AlphaInterp.h
--- a/Engine/Plugins/Animation/ControlRig/Source/ControlRig/Private/Units/Simulation/RigUnit_AlphaInterp.h
+++ b/Engine/Plugins/Animation/ControlRig/Source/ControlRig/Private/Units/Simulation/RigUnit_AlphaInterp.h
@@ -9,7 +9,7 @@
 /**
  * Takes in a float value and outputs an accumulated value with a customized scale and clamp
  */
-USTRUCT(meta=(DisplayName="Alpha Interpolate (float)", Keywords="Alpha,Lerp,LinearInterpolate", PrototypeName = "AlphaInterp"))
+USTRUCT(meta=(DisplayName="Alpha Interpolate", Keywords="Alpha,Lerp,LinearInterpolate", PrototypeName = "AlphaInterp", MenuDescSuffix = "(Float)"))
 struct CONTROLRIG_API FRigUnit_AlphaInterp : public FRigUnit_SimBase
 {
 	GENERATED_BODY()
@@ -82,7 +82,7 @@ struct CONTROLRIG_API FRigUnit_AlphaInterp : public FRigUnit_SimBase
 /**
  * Takes in a vector value and outputs an accumulated value with a customized scale and clamp
  */
-USTRUCT(meta=(DisplayName="Alpha Interpolate (Vector)", Keywords="Alpha,Lerp,LinearInterpolate", PrototypeName = "AlphaInterp", MenuDescSuffix = "(Vector)"))
+USTRUCT(meta=(DisplayName="Alpha Interpolate", Keywords="Alpha,Lerp,LinearInterpolate", PrototypeName = "AlphaInterp", MenuDescSuffix = "(Vector)"))
 struct CONTROLRIG_API FRigUnit_AlphaInterpVector : public FRigUnit_SimBase
 {
 	GENERATED_BODY()
