Control Rig: Fix collapse node renaming inside function

diff --git a/Engine/Plugins/Animation/ControlRig/Source/ControlRigDeveloper/Private/ControlRigBlueprint.cpp b/Engine/Plugins/Animation/ControlRig/Source/ControlRigDeveloper/Private/ControlRigBlueprint.cpp
--- a/Engine/Plugins/Animation/ControlRig/Source/ControlRigDeveloper/Private/ControlRigBlueprint.cpp
+++ b/Engine/Plugins/Animation/ControlRig/Source/ControlRigDeveloper/Private/ControlRigBlueprint.cpp
@@ -3154,7 +3154,7 @@ void UControlRigBlueprint::HandleModifiedEvent(ERigVMGraphNotifType InNotifType,
 							}
 							else if (RigGraph->ModelNodePath.StartsWith(OldNodePathPrefix))
 							{
-								RigGraph->ModelNodePath = NewNodePathPrefix + RigGraph->ModelNodePath.LeftChop(OldNodePathPrefix.Len());
+								RigGraph->ModelNodePath = NewNodePathPrefix + RigGraph->ModelNodePath.RightChop(OldNodePathPrefix.Len());
 							}
 						}
 					}
diff --git a/Engine/Plugins/Animation/ControlRig/Source/ControlRigDeveloper/Private/Graph/ControlRigGraphSchema.cpp b/Engine/Plugins/Animation/ControlRig/Source/ControlRigDeveloper/Private/Graph/ControlRigGraphSchema.cpp
--- a/Engine/Plugins/Animation/ControlRig/Source/ControlRigDeveloper/Private/Graph/ControlRigGraphSchema.cpp
+++ b/Engine/Plugins/Animation/ControlRig/Source/ControlRigDeveloper/Private/Graph/ControlRigGraphSchema.cpp
@@ -1244,20 +1244,18 @@ bool UControlRigGraphSchema::TryRenameGraph(UEdGraph* GraphToRename, const FName
 	{
 		if (UControlRigBlueprint* RigBlueprint = Cast<UControlRigBlueprint>(FBlueprintEditorUtils::FindBlueprintForGraph(RigGraph)))
 		{
-			if (URigVMGraph* Model = RigBlueprint->GetModel())
+			if (URigVMGraph* Model = RigGraph->GetModel())
 			{
-				URigVMLibraryNode* LibraryNode = Cast<URigVMLibraryNode>(Model->FindNode(RigGraph->ModelNodePath));
-				if (LibraryNode == nullptr && RigBlueprint->GetLocalFunctionLibrary())
-				{
-					LibraryNode = Cast<URigVMLibraryNode>(RigBlueprint->GetLocalFunctionLibrary()->FindFunction(*RigGraph->ModelNodePath));
-				}
-
-				if (LibraryNode)
+				if (URigVMGraph* RootModel = Model->GetRootGraph())
 				{
-					if (URigVMController* Controller = RigBlueprint->GetOrCreateController(LibraryNode->GetGraph()))
+					URigVMLibraryNode* LibraryNode = Cast<URigVMLibraryNode>(RootModel->FindNode(RigGraph->ModelNodePath));
+					if (LibraryNode)
 					{
-						Controller->RenameNode(LibraryNode, InNewName, true, true);
-						return true;
+						if (URigVMController* Controller = RigBlueprint->GetOrCreateController(LibraryNode->GetGraph()))
+						{
+							Controller->RenameNode(LibraryNode, InNewName, true, true);
+							return true;
+						}
 					}
 				}
 			}
diff --git a/Engine/Source/Developer/RigVMDeveloper/Private/RigVMModel/RigVMGraph.cpp b/Engine/Source/Developer/RigVMDeveloper/Private/RigVMModel/RigVMGraph.cpp
--- a/Engine/Source/Developer/RigVMDeveloper/Private/RigVMModel/RigVMGraph.cpp
+++ b/Engine/Source/Developer/RigVMDeveloper/Private/RigVMModel/RigVMGraph.cpp
@@ -156,8 +156,21 @@ URigVMNode* URigVMGraph::FindNode(const FString& InNodePath) const
 		return nullptr;
 	}
 
-	FString Left = InNodePath, Right;
-	URigVMNode::SplitNodePathAtStart(InNodePath, Left, Right);
+	FString Path = InNodePath;
+	if (Path.StartsWith(TEXT("FunctionLibrary::|")))
+	{
+		if (GetRootGraph()->IsA<URigVMFunctionLibrary>())
+		{
+			Path.RightChopInline(18);
+		}
+		else
+		{
+			return nullptr;
+		}
+	}
+
+	FString Left = Path, Right;
+	URigVMNode::SplitNodePathAtStart(Path, Left, Right);
 
 	if (Right.IsEmpty())
 	{
