Control Rig: Fix tool tip referencing deprecated name for nuill

diff --git a/Engine/Plugins/Animation/ControlRig/Source/ControlRig/Public/ControlRigComponent.h b/Engine/Plugins/Animation/ControlRig/Source/ControlRig/Public/ControlRigComponent.h
--- a/Engine/Plugins/Animation/ControlRig/Source/ControlRig/Public/ControlRigComponent.h
+++ b/Engine/Plugins/Animation/ControlRig/Source/ControlRig/Public/ControlRigComponent.h
@@ -315,7 +315,7 @@ public:
 	UPROPERTY(EditAnywhere, Category = "Animation")
 	bool bResetTransformBeforeTick;
 
-	/** When checked the initial transforms on bones, spaces and controls are reset prior to a setup event */
+	/** When checked the initial transforms on bones, nulls and controls are reset prior to a setup event */
 	UPROPERTY(EditAnywhere, Category = "Animation")
 	bool bResetInitialsBeforeSetup;
 
