Control Rig: Adjust node categories

diff --git a/Engine/Plugins/Animation/ControlRig/Source/ControlRig/Private/Units/Drawing/RigUnit_DrawContainer.h b/Engine/Plugins/Animation/ControlRig/Source/ControlRig/Private/Units/Drawing/RigUnit_DrawContainer.h
--- a/Engine/Plugins/Animation/ControlRig/Source/ControlRig/Private/Units/Drawing/RigUnit_DrawContainer.h
+++ b/Engine/Plugins/Animation/ControlRig/Source/ControlRig/Private/Units/Drawing/RigUnit_DrawContainer.h
@@ -8,7 +8,7 @@
 /**
  * Get Imported Draw Container curve transform and color
  */
-USTRUCT(meta=(DisplayName="Get Draw Instruction", Category="Drawing", NodeColor = "0.83077 0.846873 0.049707", Keywords = "Curve,Shape"))
+USTRUCT(meta=(DisplayName="Get Draw Instruction", Category="Drawing", NodeColor = "0.83077 0.846873 0.049707", Keywords = "Curve,Shape", Deprecated = "5.0"))
 struct CONTROLRIG_API FRigUnit_DrawContainerGetInstruction : public FRigUnit
 {
 	GENERATED_BODY()
@@ -36,7 +36,7 @@ struct CONTROLRIG_API FRigUnit_DrawContainerGetInstruction : public FRigUnit
 /**
  * Set Imported Draw Container curve color
  */
-USTRUCT(meta = (DisplayName = "Set Draw Color", Category = "Drawing", NodeColor = "0.83077 0.846873 0.049707", Keywords = "Curve,Shape"))
+USTRUCT(meta = (DisplayName = "Set Draw Color", Category = "Drawing", NodeColor = "0.83077 0.846873 0.049707", Keywords = "Curve,Shape", Deprecated = "5.0"))
 struct CONTROLRIG_API FRigUnit_DrawContainerSetColor: public FRigUnitMutable
 {
 	GENERATED_BODY()
@@ -60,7 +60,7 @@ struct CONTROLRIG_API FRigUnit_DrawContainerSetColor: public FRigUnitMutable
 /**
  * Set Imported Draw Container curve thickness
  */
-USTRUCT(meta = (DisplayName = "Set Draw Thickness", Category = "Drawing", NodeColor = "0.83077 0.846873 0.049707", Keywords = "Curve,Shape"))
+USTRUCT(meta = (DisplayName = "Set Draw Thickness", Category = "Drawing", NodeColor = "0.83077 0.846873 0.049707", Keywords = "Curve,Shape", Deprecated = "5.0"))
 struct CONTROLRIG_API FRigUnit_DrawContainerSetThickness : public FRigUnitMutable
 {
 	GENERATED_BODY()
@@ -84,7 +84,7 @@ struct CONTROLRIG_API FRigUnit_DrawContainerSetThickness : public FRigUnitMutabl
 /**
  * Set Imported Draw Container curve transform
  */
-USTRUCT(meta = (DisplayName = "Set Draw Transform", Category = "Drawing", NodeColor = "0.83077 0.846873 0.049707", Keywords = "Curve,Shape"))
+USTRUCT(meta = (DisplayName = "Set Draw Transform", Category = "Drawing", NodeColor = "0.83077 0.846873 0.049707", Keywords = "Curve,Shape", Deprecated = "5.0"))
 struct CONTROLRIG_API FRigUnit_DrawContainerSetTransform : public FRigUnitMutable
 {
 	GENERATED_BODY()
diff --git a/Engine/Plugins/Animation/ControlRig/Source/ControlRig/Private/Units/Execution/RigUnit_BeginExecution.h b/Engine/Plugins/Animation/ControlRig/Source/ControlRig/Private/Units/Execution/RigUnit_BeginExecution.h
--- a/Engine/Plugins/Animation/ControlRig/Source/ControlRig/Private/Units/Execution/RigUnit_BeginExecution.h
+++ b/Engine/Plugins/Animation/ControlRig/Source/ControlRig/Private/Units/Execution/RigUnit_BeginExecution.h
@@ -9,7 +9,7 @@
 /**
  * Event for driving the skeleton hierarchy with variables and rig elements
  */
-USTRUCT(meta=(DisplayName="Forwards Solve", Category="Execution", NodeColor="1, 0, 0", Keywords="Begin,Update,Tick,Forward,Event"))
+USTRUCT(meta=(DisplayName="Forwards Solve", Category="Events", NodeColor="1, 0, 0", Keywords="Begin,Update,Tick,Forward,Event"))
 struct CONTROLRIG_API FRigUnit_BeginExecution : public FRigUnit
 {
 	GENERATED_BODY()
diff --git a/Engine/Plugins/Animation/ControlRig/Source/ControlRig/Private/Units/Execution/RigUnit_Collection.h b/Engine/Plugins/Animation/ControlRig/Source/ControlRig/Private/Units/Execution/RigUnit_Collection.h
--- a/Engine/Plugins/Animation/ControlRig/Source/ControlRig/Private/Units/Execution/RigUnit_Collection.h
+++ b/Engine/Plugins/Animation/ControlRig/Source/ControlRig/Private/Units/Execution/RigUnit_Collection.h
@@ -176,7 +176,7 @@ struct CONTROLRIG_API FRigUnit_CollectionChildren : public FRigUnit_CollectionBa
 * Creates an item array based on the direct or recursive children
 * of a provided parent item. Returns an empty array for an invalid parent item.
 */
-USTRUCT(meta = (DisplayName = "Get Children", Keywords = "Bone,Joint,Collection,Filter,Parent", Varying))
+USTRUCT(meta = (DisplayName = "Get Children", Category = "Hierarchy", Keywords = "Bone,Joint,Collection,Filter,Parent", Varying))
 struct CONTROLRIG_API FRigUnit_CollectionChildrenArray : public FRigUnit_CollectionBase
 {
 	GENERATED_BODY()
@@ -285,7 +285,7 @@ struct CONTROLRIG_API FRigUnit_CollectionReplaceItemsArray : public FRigUnit_Col
 /**
  * Returns a collection provided a specific array of items.
  */
-USTRUCT(meta = (DisplayName = "Collection from Items", Keywords = "Collection,Array", Varying))
+USTRUCT(meta = (DisplayName = "Collection from Items", Category = "Items|Collections", Keywords = "Collection,Array", Varying))
 struct CONTROLRIG_API FRigUnit_CollectionItems : public FRigUnit_CollectionBase
 {
 	GENERATED_BODY()
@@ -312,7 +312,7 @@ struct CONTROLRIG_API FRigUnit_CollectionItems : public FRigUnit_CollectionBase
 /**
 * Returns an array of items provided a collection
 */
-USTRUCT(meta = (DisplayName = "Get Items from Collection", Keywords = "Collection,Array", Varying))
+USTRUCT(meta = (DisplayName = "Get Items from Collection", Category = "Items|Collections", Keywords = "Collection,Array", Varying))
 struct CONTROLRIG_API FRigUnit_CollectionGetItems : public FRigUnit_CollectionBase
 {
 	GENERATED_BODY()
diff --git a/Engine/Plugins/Animation/ControlRig/Source/ControlRig/Private/Units/Execution/RigUnit_Hierarchy.h b/Engine/Plugins/Animation/ControlRig/Source/ControlRig/Private/Units/Execution/RigUnit_Hierarchy.h
--- a/Engine/Plugins/Animation/ControlRig/Source/ControlRig/Private/Units/Execution/RigUnit_Hierarchy.h
+++ b/Engine/Plugins/Animation/ControlRig/Source/ControlRig/Private/Units/Execution/RigUnit_Hierarchy.h
@@ -284,7 +284,7 @@ struct CONTROLRIG_API FRigUnit_HierarchyGetPose : public FRigUnit_HierarchyBase
 /**
  * Returns the hierarchy's pose
  */
-USTRUCT(meta=(DisplayName="Get Pose Cache", Keywords="Hierarchy,Pose,State", Varying))
+USTRUCT(meta=(DisplayName="Get Pose Cache", Keywords="Hierarchy,Pose,State", Varying, Category = "Pose Cache"))
 struct CONTROLRIG_API FRigUnit_HierarchyGetPoseItemArray : public FRigUnit_HierarchyBase
 {
 	GENERATED_BODY()
@@ -354,7 +354,7 @@ struct CONTROLRIG_API FRigUnit_HierarchySetPose : public FRigUnit_HierarchyBaseM
 /**
  * Sets the hierarchy's pose
  */
-USTRUCT(meta=(DisplayName="Apply Pose Cache", Keywords="Hierarchy,Pose,State", Varying))
+USTRUCT(meta=(DisplayName="Apply Pose Cache", Keywords="Hierarchy,Pose,State", Varying, Category = "Pose Cache"))
 struct CONTROLRIG_API FRigUnit_HierarchySetPoseItemArray : public FRigUnit_HierarchyBaseMutable
 {
 	GENERATED_BODY()
@@ -390,7 +390,7 @@ struct CONTROLRIG_API FRigUnit_HierarchySetPoseItemArray : public FRigUnit_Hiera
 /**
 * Returns true if the hierarchy pose is empty (has no items)
 */
-USTRUCT(meta=(DisplayName="Is Pose Cache Empty", Keywords="Hierarchy,Pose,State", Varying))
+USTRUCT(meta=(DisplayName="Is Pose Cache Empty", Keywords="Hierarchy,Pose,State", Varying, Category = "Pose Cache"))
 struct CONTROLRIG_API FRigUnit_PoseIsEmpty : public FRigUnit_HierarchyBase
 {
 	GENERATED_BODY()
@@ -442,7 +442,7 @@ struct CONTROLRIG_API FRigUnit_PoseGetItems : public FRigUnit_HierarchyBase
 /**
 * Returns the items in the hierarchy pose
 */
-USTRUCT(meta=(DisplayName="Get Pose Cache Items", Keywords="Hierarchy,Pose,State", Varying))
+USTRUCT(meta=(DisplayName="Get Pose Cache Items", Keywords="Hierarchy,Pose,State", Varying, Category = "Pose Cache"))
 struct CONTROLRIG_API FRigUnit_PoseGetItemsItemArray : public FRigUnit_HierarchyBase
 {
 	GENERATED_BODY()
@@ -469,7 +469,7 @@ struct CONTROLRIG_API FRigUnit_PoseGetItemsItemArray : public FRigUnit_Hierarchy
 /**
 * Compares two pose caches and compares their values.
 */
-USTRUCT(meta=(DisplayName="Get Pose Cache Delta", Keywords="Hierarchy,Pose,State", Varying))
+USTRUCT(meta=(DisplayName="Get Pose Cache Delta", Keywords="Hierarchy,Pose,State", Varying, Category = "Pose Cache"))
 struct CONTROLRIG_API FRigUnit_PoseGetDelta : public FRigUnit_HierarchyBase
 {
 	GENERATED_BODY()
@@ -552,7 +552,7 @@ struct CONTROLRIG_API FRigUnit_PoseGetDelta : public FRigUnit_HierarchyBase
 /**
 * Returns the hierarchy's pose transform
 */
-USTRUCT(meta=(DisplayName="Get Pose Cache Transform", Keywords="Hierarchy,Pose,State", Varying))
+USTRUCT(meta=(DisplayName="Get Pose Cache Transform", Keywords="Hierarchy,Pose,State", Varying, Category = "Pose Cache"))
 struct CONTROLRIG_API FRigUnit_PoseGetTransform : public FRigUnit_HierarchyBase
 {
 	GENERATED_BODY()
@@ -602,7 +602,7 @@ struct CONTROLRIG_API FRigUnit_PoseGetTransform : public FRigUnit_HierarchyBase
 /**
 * Returns an array of transforms from a given hierarchy pose
 */
-USTRUCT(meta=(DisplayName="Get Pose Cache Transform Array", Keywords="Hierarchy,Pose,State", Varying))
+USTRUCT(meta=(DisplayName="Get Pose Cache Transform Array", Keywords="Hierarchy,Pose,State", Varying, Category = "Pose Cache"))
 struct CONTROLRIG_API FRigUnit_PoseGetTransformArray : public FRigUnit_HierarchyBase
 {
 	GENERATED_BODY()
@@ -637,7 +637,7 @@ struct CONTROLRIG_API FRigUnit_PoseGetTransformArray : public FRigUnit_Hierarchy
 /**
 * Returns the hierarchy's pose curve value
 */
-USTRUCT(meta=(DisplayName="Get Pose Cache Curve", Keywords="Hierarchy,Pose,State", Varying))
+USTRUCT(meta=(DisplayName="Get Pose Cache Curve", Keywords="Hierarchy,Pose,State", Varying, Category = "Pose Cache"))
 struct CONTROLRIG_API FRigUnit_PoseGetCurve : public FRigUnit_HierarchyBase
 {
 	GENERATED_BODY()
@@ -676,7 +676,7 @@ struct CONTROLRIG_API FRigUnit_PoseGetCurve : public FRigUnit_HierarchyBase
 /**
 * Given a pose, execute iteratively across all items in the pose
 */
-USTRUCT(meta=(DisplayName="For Each Pose Cache Element", Keywords="Collection,Loop,Iterate", Icon="EditorStyle|GraphEditor.Macro.ForEach_16x"))
+USTRUCT(meta=(DisplayName="For Each Pose Cache Element", Keywords="Collection,Loop,Iterate", Icon="EditorStyle|GraphEditor.Macro.ForEach_16x", Category = "Pose Cache"))
 struct CONTROLRIG_API FRigUnit_PoseLoop : public FRigUnit_HierarchyBaseMutable
 {
 	GENERATED_BODY()
diff --git a/Engine/Plugins/Animation/ControlRig/Source/ControlRig/Private/Units/Execution/RigUnit_InverseExecution.h b/Engine/Plugins/Animation/ControlRig/Source/ControlRig/Private/Units/Execution/RigUnit_InverseExecution.h
--- a/Engine/Plugins/Animation/ControlRig/Source/ControlRig/Private/Units/Execution/RigUnit_InverseExecution.h
+++ b/Engine/Plugins/Animation/ControlRig/Source/ControlRig/Private/Units/Execution/RigUnit_InverseExecution.h
@@ -9,7 +9,7 @@
 /**
  * Event for driving elements based off the skeleton hierarchy
  */
-USTRUCT(meta=(DisplayName="Backwards Solve", Category="Inverse", NodeColor="1, 0, 0", Keywords="Inverse,Reverse,Backwards,Event"))
+USTRUCT(meta=(DisplayName="Backwards Solve", Category="Events", NodeColor="1, 0, 0", Keywords="Inverse,Reverse,Backwards,Event"))
 struct CONTROLRIG_API FRigUnit_InverseExecution : public FRigUnit
 {
 	GENERATED_BODY()
diff --git a/Engine/Plugins/Animation/ControlRig/Source/ControlRig/Private/Units/Execution/RigUnit_PrepareForExecution.h b/Engine/Plugins/Animation/ControlRig/Source/ControlRig/Private/Units/Execution/RigUnit_PrepareForExecution.h
--- a/Engine/Plugins/Animation/ControlRig/Source/ControlRig/Private/Units/Execution/RigUnit_PrepareForExecution.h
+++ b/Engine/Plugins/Animation/ControlRig/Source/ControlRig/Private/Units/Execution/RigUnit_PrepareForExecution.h
@@ -9,7 +9,7 @@
 /**
  * Event to setup elements before all solves
  */
-USTRUCT(meta=(DisplayName="Setup Event", Category="Setup", NodeColor="1, 0, 0", Keywords="Setup,Init,Fit"))
+USTRUCT(meta=(DisplayName="Setup Event", Category="Events", NodeColor="1, 0, 0", Keywords="Setup,Init,Fit"))
 struct CONTROLRIG_API FRigUnit_PrepareForExecution : public FRigUnit
 {
 	GENERATED_BODY()
diff --git a/Engine/Plugins/Animation/ControlRig/Source/ControlRig/Private/Units/Hierarchy/RigUnit_GetRelativeTransform.h b/Engine/Plugins/Animation/ControlRig/Source/ControlRig/Private/Units/Hierarchy/RigUnit_GetRelativeTransform.h
--- a/Engine/Plugins/Animation/ControlRig/Source/ControlRig/Private/Units/Hierarchy/RigUnit_GetRelativeTransform.h
+++ b/Engine/Plugins/Animation/ControlRig/Source/ControlRig/Private/Units/Hierarchy/RigUnit_GetRelativeTransform.h
@@ -8,7 +8,7 @@
 /**
  * GetRelativeTransform is used to retrieve a single transform from a hierarchy in the space of another transform
  */
-USTRUCT(meta=(DisplayName="Get Relative Transform", Category="Hierarchy", DocumentationPolicy = "Strict", Keywords = "Offset,Local", NodeColor="0.462745, 1,0, 0.329412", Varying))
+USTRUCT(meta=(DisplayName="Get Relative Transform", Category="Transforms", DocumentationPolicy = "Strict", Keywords = "Offset,Local", NodeColor="0.462745, 1,0, 0.329412", Varying))
 struct CONTROLRIG_API FRigUnit_GetRelativeTransformForItem : public FRigUnit
 {
 	GENERATED_BODY()
diff --git a/Engine/Plugins/Animation/ControlRig/Source/ControlRig/Private/Units/Hierarchy/RigUnit_GetTransform.h b/Engine/Plugins/Animation/ControlRig/Source/ControlRig/Private/Units/Hierarchy/RigUnit_GetTransform.h
--- a/Engine/Plugins/Animation/ControlRig/Source/ControlRig/Private/Units/Hierarchy/RigUnit_GetTransform.h
+++ b/Engine/Plugins/Animation/ControlRig/Source/ControlRig/Private/Units/Hierarchy/RigUnit_GetTransform.h
@@ -8,7 +8,7 @@
 /**
  * GetTransform is used to retrieve a single transform from a hierarchy.
  */
-USTRUCT(meta=(DisplayName="Get Transform", Category="Hierarchy", DocumentationPolicy = "Strict", Keywords="GetBoneTransform,GetControlTransform,GetInitialTransform,GetSpaceTransform,GetTransform", NodeColor="0.462745, 1,0, 0.329412",Varying))
+USTRUCT(meta=(DisplayName="Get Transform", Category="Transforms", DocumentationPolicy = "Strict", Keywords="GetBoneTransform,GetControlTransform,GetInitialTransform,GetSpaceTransform,GetTransform", NodeColor="0.462745, 1,0, 0.329412",Varying))
 struct CONTROLRIG_API FRigUnit_GetTransform : public FRigUnit
 {
 	GENERATED_BODY()
@@ -116,7 +116,7 @@ struct CONTROLRIG_API FRigUnit_GetTransformArray : public FRigUnit
 /**
 * GetTransformArray is used to retrieve an array of transforms from the hierarchy.
 */
-USTRUCT(meta=(DisplayName="Get Transform Array", Category="Hierarchy", DocumentationPolicy = "Strict", Keywords="GetBoneTransform,GetControlTransform,GetInitialTransform,GetSpaceTransform,GetTransform", NodeColor="0.462745, 1,0, 0.329412",Varying))
+USTRUCT(meta=(DisplayName="Get Transform Array", Category="Transforms", DocumentationPolicy = "Strict", Keywords="GetBoneTransform,GetControlTransform,GetInitialTransform,GetSpaceTransform,GetTransform", NodeColor="0.462745, 1,0, 0.329412",Varying))
 struct CONTROLRIG_API FRigUnit_GetTransformItemArray : public FRigUnit
 {
 	GENERATED_BODY()
diff --git a/Engine/Plugins/Animation/ControlRig/Source/ControlRig/Private/Units/Hierarchy/RigUnit_OffsetTransform.h b/Engine/Plugins/Animation/ControlRig/Source/ControlRig/Private/Units/Hierarchy/RigUnit_OffsetTransform.h
--- a/Engine/Plugins/Animation/ControlRig/Source/ControlRig/Private/Units/Hierarchy/RigUnit_OffsetTransform.h
+++ b/Engine/Plugins/Animation/ControlRig/Source/ControlRig/Private/Units/Hierarchy/RigUnit_OffsetTransform.h
@@ -8,7 +8,7 @@
 /**
  * Offset Transform is used to add an offset to an existing transform in the hierarchy. The offset is post multiplied.
  */
-USTRUCT(meta=(DisplayName="Offset Transform", Category="Hierarchy", DocumentationPolicy = "Strict", Keywords = "Offset,Relative,AddBoneTransform", NodeColor="0, 0.364706, 1.0", Varying))
+USTRUCT(meta=(DisplayName="Offset Transform", Category="Transforms", DocumentationPolicy = "Strict", Keywords = "Offset,Relative,AddBoneTransform", NodeColor="0, 0.364706, 1.0", Varying))
 struct CONTROLRIG_API FRigUnit_OffsetTransformForItem : public FRigUnitMutable
 {
 	GENERATED_BODY()
diff --git a/Engine/Plugins/Animation/ControlRig/Source/ControlRig/Private/Units/Hierarchy/RigUnit_ParentSwitchConstraint.h b/Engine/Plugins/Animation/ControlRig/Source/ControlRig/Private/Units/Hierarchy/RigUnit_ParentSwitchConstraint.h
--- a/Engine/Plugins/Animation/ControlRig/Source/ControlRig/Private/Units/Hierarchy/RigUnit_ParentSwitchConstraint.h
+++ b/Engine/Plugins/Animation/ControlRig/Source/ControlRig/Private/Units/Hierarchy/RigUnit_ParentSwitchConstraint.h
@@ -86,7 +86,7 @@ struct CONTROLRIG_API FRigUnit_ParentSwitchConstraint : public FRigUnitMutable
  * The Parent Switch Constraint is used to have an item follow one of multiple parents,
  * and allowing to switch between the parent in question.
  */
-USTRUCT(meta=(DisplayName="Parent Switch Constraint", Category="Hierarchy", DocumentationPolicy="Strict", Keywords = "SpaceSwitch"))
+USTRUCT(meta=(DisplayName="Parent Switch Constraint", Category="Constraints", DocumentationPolicy="Strict", Keywords = "SpaceSwitch"))
 struct CONTROLRIG_API FRigUnit_ParentSwitchConstraintArray : public FRigUnitMutable
 {
 	GENERATED_BODY()
diff --git a/Engine/Plugins/Animation/ControlRig/Source/ControlRig/Private/Units/Hierarchy/RigUnit_SetControlColor.h b/Engine/Plugins/Animation/ControlRig/Source/ControlRig/Private/Units/Hierarchy/RigUnit_SetControlColor.h
--- a/Engine/Plugins/Animation/ControlRig/Source/ControlRig/Private/Units/Hierarchy/RigUnit_SetControlColor.h
+++ b/Engine/Plugins/Animation/ControlRig/Source/ControlRig/Private/Units/Hierarchy/RigUnit_SetControlColor.h
@@ -8,7 +8,7 @@
 /**
  * SetControlColor is used to change the gizmo's color on a control at runtime
  */
-USTRUCT(meta=(DisplayName="Set Control Color", Category="Hierarchy", DocumentationPolicy="Strict", Keywords = "SetControlColor,SetGizmoColor", PrototypeName = "SetControlColor", NodeColor = "0.0 0.36470600962638855 1.0"))
+USTRUCT(meta=(DisplayName="Set Control Color", Category="Controls", DocumentationPolicy="Strict", Keywords = "SetControlColor,SetGizmoColor", PrototypeName = "SetControlColor", NodeColor = "0.0 0.36470600962638855 1.0"))
 struct CONTROLRIG_API FRigUnit_SetControlColor : public FRigUnitMutable
 {
 	GENERATED_BODY()
diff --git a/Engine/Plugins/Animation/ControlRig/Source/ControlRig/Private/Units/Hierarchy/RigUnit_SetControlOffset.h b/Engine/Plugins/Animation/ControlRig/Source/ControlRig/Private/Units/Hierarchy/RigUnit_SetControlOffset.h
--- a/Engine/Plugins/Animation/ControlRig/Source/ControlRig/Private/Units/Hierarchy/RigUnit_SetControlOffset.h
+++ b/Engine/Plugins/Animation/ControlRig/Source/ControlRig/Private/Units/Hierarchy/RigUnit_SetControlOffset.h
@@ -7,8 +7,9 @@
 
 /**
  * SetControlOffset is used to perform a change in the hierarchy by setting a single control's transform.
+ * This is typically only used during the Setup Event.
  */
-USTRUCT(meta=(DisplayName="Set Control Offset", Category="Setup", DocumentationPolicy="Strict", Keywords = "SetControlOffset,Initial,InitialTransform,SetInitialTransform,SetInitialControlTransform", PrototypeName = "SetControlValue", NodeColor="0, 0.364706, 1.0"))
+USTRUCT(meta=(DisplayName="Set Control Offset", Category="Controls", DocumentationPolicy="Strict", Keywords = "SetControlOffset,Initial,InitialTransform,SetInitialTransform,SetInitialControlTransform", PrototypeName = "SetControlValue", NodeColor="0, 0.364706, 1.0"))
 struct CONTROLRIG_API FRigUnit_SetControlOffset : public FRigUnitMutable
 {
 	GENERATED_BODY()
diff --git a/Engine/Plugins/Animation/ControlRig/Source/ControlRig/Private/Units/Hierarchy/RigUnit_SetControlTransform.h b/Engine/Plugins/Animation/ControlRig/Source/ControlRig/Private/Units/Hierarchy/RigUnit_SetControlTransform.h
--- a/Engine/Plugins/Animation/ControlRig/Source/ControlRig/Private/Units/Hierarchy/RigUnit_SetControlTransform.h
+++ b/Engine/Plugins/Animation/ControlRig/Source/ControlRig/Private/Units/Hierarchy/RigUnit_SetControlTransform.h
@@ -8,7 +8,7 @@
 /**
  * SetControlBool is used to perform a change in the hierarchy by setting a single control's bool value.
  */
-USTRUCT(meta=(DisplayName="Set Control Bool", Category="Hierarchy", DocumentationPolicy="Strict", Keywords = "SetControlBool,SetGizmoBool", PrototypeName = "SetControlValue", NodeColor = "0.0 0.36470600962638855 1.0"))
+USTRUCT(meta=(DisplayName="Set Control Bool", Category="Controls", DocumentationPolicy="Strict", Keywords = "SetControlBool,SetGizmoBool", PrototypeName = "SetControlValue", NodeColor = "0.0 0.36470600962638855 1.0"))
 struct CONTROLRIG_API FRigUnit_SetControlBool : public FRigUnitMutable
 {
 	GENERATED_BODY()
@@ -61,7 +61,7 @@ struct CONTROLRIG_API FRigUnit_SetMultiControlBool_Entry
 /**
  * SetMultiControlBool is used to perform a change in the hierarchy by setting multiple controls' bool value.
  */
-USTRUCT(meta = (DisplayName = "Set Multiple Controls Bool", Category = "Hierarchy", DocumentationPolicy = "Strict", Keywords = "SetMultipleControlsBool,SetControlBool,SetGizmoBool", PrototypeName = "SetMultiControlValue", NodeColor = "0.0 0.36470600962638855 1.0"))
+USTRUCT(meta = (DisplayName = "Set Multiple Controls Bool", Category = "Controls", DocumentationPolicy = "Strict", Keywords = "SetMultipleControlsBool,SetControlBool,SetGizmoBool", PrototypeName = "SetMultiControlValue", NodeColor = "0.0 0.36470600962638855 1.0"))
 struct CONTROLRIG_API FRigUnit_SetMultiControlBool : public FRigUnitMutable
 {
 	GENERATED_BODY()
@@ -88,7 +88,7 @@ struct CONTROLRIG_API FRigUnit_SetMultiControlBool : public FRigUnitMutable
 /**
  * SetControlFloat is used to perform a change in the hierarchy by setting a single control's float value.
  */
-USTRUCT(meta=(DisplayName="Set Control Float", Category="Hierarchy", DocumentationPolicy="Strict", Keywords = "SetControlFloat,SetGizmoFloat", PrototypeName = "SetControlValue", NodeColor = "0.0 0.36470600962638855 1.0"))
+USTRUCT(meta=(DisplayName="Set Control Float", Category="Controls", DocumentationPolicy="Strict", Keywords = "SetControlFloat,SetGizmoFloat", PrototypeName = "SetControlValue", NodeColor = "0.0 0.36470600962638855 1.0"))
 struct CONTROLRIG_API FRigUnit_SetControlFloat : public FRigUnitMutable
 {
 	GENERATED_BODY()
@@ -149,7 +149,7 @@ struct CONTROLRIG_API FRigUnit_SetMultiControlFloat_Entry
 /**
  * SetMultiControlFloat is used to perform a change in the hierarchy by setting multiple controls' float value.
  */
-USTRUCT(meta = (DisplayName = "Set Multiple Controls Float", Category = "Hierarchy", DocumentationPolicy = "Strict", Keywords = "SetMultipleControlsFloat,SetControlFloat,SetGizmoFloat", PrototypeName = "SetMultiControlValue", NodeColor = "0.0 0.36470600962638855 1.0"))
+USTRUCT(meta = (DisplayName = "Set Multiple Controls Float", Category = "Controls", DocumentationPolicy = "Strict", Keywords = "SetMultipleControlsFloat,SetControlFloat,SetGizmoFloat", PrototypeName = "SetMultiControlValue", NodeColor = "0.0 0.36470600962638855 1.0"))
 struct CONTROLRIG_API FRigUnit_SetMultiControlFloat : public FRigUnitMutable
 {
 	GENERATED_BODY()
@@ -185,7 +185,7 @@ struct CONTROLRIG_API FRigUnit_SetMultiControlFloat : public FRigUnitMutable
 /**
  * SetControlInteger is used to perform a change in the hierarchy by setting a single control's int32 value.
  */
-USTRUCT(meta=(DisplayName="Set Control Integer", Category="Hierarchy", DocumentationPolicy="Strict", Keywords = "SetControlInteger,SetGizmoInteger", PrototypeName = "SetControlValue", NodeColor = "0.0 0.36470600962638855 1.0"))
+USTRUCT(meta=(DisplayName="Set Control Integer", Category="Controls", DocumentationPolicy="Strict", Keywords = "SetControlInteger,SetGizmoInteger", PrototypeName = "SetControlValue", NodeColor = "0.0 0.36470600962638855 1.0"))
 struct CONTROLRIG_API FRigUnit_SetControlInteger : public FRigUnitMutable
 {
 	GENERATED_BODY()
@@ -246,7 +246,7 @@ struct CONTROLRIG_API FRigUnit_SetMultiControlInteger_Entry
 /**
  * SetMultiControlInteger is used to perform a change in the hierarchy by setting multiple controls' integer value.
  */
-USTRUCT(meta = (DisplayName = "Set Multiple Controls Integer", Category = "Hierarchy", DocumentationPolicy = "Strict", Keywords = "SetMultipleControlsInteger,SetControlInteger,SetGizmoInteger", PrototypeName = "SetMultiControlValue", NodeColor = "0.0 0.36470600962638855 1.0"))
+USTRUCT(meta = (DisplayName = "Set Multiple Controls Integer", Category = "Controls", DocumentationPolicy = "Strict", Keywords = "SetMultipleControlsInteger,SetControlInteger,SetGizmoInteger", PrototypeName = "SetMultiControlValue", NodeColor = "0.0 0.36470600962638855 1.0"))
 struct CONTROLRIG_API FRigUnit_SetMultiControlInteger : public FRigUnitMutable
 {
 	GENERATED_BODY()
@@ -280,7 +280,7 @@ struct CONTROLRIG_API FRigUnit_SetMultiControlInteger : public FRigUnitMutable
 /**
  * SetControlVector2D is used to perform a change in the hierarchy by setting a single control's Vector2D value.
  */
-USTRUCT(meta=(DisplayName="Set Control Vector2D", Category="Hierarchy", DocumentationPolicy="Strict", Keywords = "SetControlVector2D,SetGizmoVector2D", PrototypeName = "SetControlValue", NodeColor = "0.0 0.36470600962638855 1.0"))
+USTRUCT(meta=(DisplayName="Set Control Vector2D", Category="Controls", DocumentationPolicy="Strict", Keywords = "SetControlVector2D,SetGizmoVector2D", PrototypeName = "SetControlValue", NodeColor = "0.0 0.36470600962638855 1.0"))
 struct CONTROLRIG_API FRigUnit_SetControlVector2D : public FRigUnitMutable
 {
 	GENERATED_BODY()
@@ -342,7 +342,7 @@ struct CONTROLRIG_API FRigUnit_SetMultiControlVector2D_Entry
 /**
  * SetMultiControlVector2D is used to perform a change in the hierarchy by setting multiple controls' vector2D value.
  */
-USTRUCT(meta = (DisplayName = "Set Multiple Controls Vector2D", Category = "Hierarchy", DocumentationPolicy = "Strict", Keywords = "SetMultipleControlsVector2D,SetControlVector2D,SetGizmoVector2D", PrototypeName = "SetMultiControlValue", NodeColor = "0.0 0.36470600962638855 1.0"))
+USTRUCT(meta = (DisplayName = "Set Multiple Controls Vector2D", Category = "Controls", DocumentationPolicy = "Strict", Keywords = "SetMultipleControlsVector2D,SetControlVector2D,SetGizmoVector2D", PrototypeName = "SetMultiControlValue", NodeColor = "0.0 0.36470600962638855 1.0"))
 struct CONTROLRIG_API FRigUnit_SetMultiControlVector2D : public FRigUnitMutable
 {
 	GENERATED_BODY()
@@ -376,7 +376,7 @@ struct CONTROLRIG_API FRigUnit_SetMultiControlVector2D : public FRigUnitMutable
 /**
  * SetControlVector is used to perform a change in the hierarchy by setting a single control's Vector value.
  */
-USTRUCT(meta=(DisplayName="Set Control Vector", Category="Hierarchy", DocumentationPolicy="Strict", Keywords = "SetControlVector,SetGizmoVector", PrototypeName = "SetControlValue", NodeColor = "0.0 0.36470600962638855 1.0"))
+USTRUCT(meta=(DisplayName="Set Control Vector", Category="Controls", DocumentationPolicy="Strict", Keywords = "SetControlVector,SetGizmoVector", PrototypeName = "SetControlValue", NodeColor = "0.0 0.36470600962638855 1.0"))
 struct CONTROLRIG_API FRigUnit_SetControlVector : public FRigUnitMutable
 {
 	GENERATED_BODY()
@@ -424,7 +424,7 @@ struct CONTROLRIG_API FRigUnit_SetControlVector : public FRigUnitMutable
 /**
  * SetControlRotator is used to perform a change in the hierarchy by setting a single control's Rotator value.
  */
-USTRUCT(meta=(DisplayName="Set Control Rotator", Category="Hierarchy", DocumentationPolicy="Strict", Keywords = "SetControlRotator,SetGizmoRotator", PrototypeName = "SetControlValue", NodeColor = "0.0 0.36470600962638855 1.0"))
+USTRUCT(meta=(DisplayName="Set Control Rotator", Category="Controls", DocumentationPolicy="Strict", Keywords = "SetControlRotator,SetGizmoRotator", PrototypeName = "SetControlValue", NodeColor = "0.0 0.36470600962638855 1.0"))
 struct CONTROLRIG_API FRigUnit_SetControlRotator : public FRigUnitMutable
 {
 	GENERATED_BODY()
@@ -503,7 +503,7 @@ struct CONTROLRIG_API FRigUnit_SetMultiControlRotator_Entry
 /**
  * SetMultiControlRotator is used to perform a change in the hierarchy by setting multiple controls' rotator value.
  */
-USTRUCT(meta = (DisplayName = "Set Multiple Controls Rotator", Category = "Hierarchy", DocumentationPolicy = "Strict", Keywords = "SetMultipleControlsRotator,SetControlRotator,SetGizmoRotator", PrototypeName = "SetMultiControlValue", NodeColor = "0.0 0.36470600962638855 1.0"))
+USTRUCT(meta = (DisplayName = "Set Multiple Controls Rotator", Category = "Controls", DocumentationPolicy = "Strict", Keywords = "SetMultipleControlsRotator,SetControlRotator,SetGizmoRotator", PrototypeName = "SetMultiControlValue", NodeColor = "0.0 0.36470600962638855 1.0"))
 struct CONTROLRIG_API FRigUnit_SetMultiControlRotator : public FRigUnitMutable
 {
 	GENERATED_BODY()
@@ -537,7 +537,7 @@ struct CONTROLRIG_API FRigUnit_SetMultiControlRotator : public FRigUnitMutable
 /**
  * SetControlTransform is used to perform a change in the hierarchy by setting a single control's transform.
  */
-USTRUCT(meta=(DisplayName="Set Control Transform", Category="Hierarchy", DocumentationPolicy="Strict", Keywords = "SetControlTransform,SetGizmoTransform", PrototypeName = "SetControlValue", Deprecated = "4.25"))
+USTRUCT(meta=(DisplayName="Set Control Transform", Category="Controls", DocumentationPolicy="Strict", Keywords = "SetControlTransform,SetGizmoTransform", PrototypeName = "SetControlValue", Deprecated = "4.25"))
 struct CONTROLRIG_API FRigUnit_SetControlTransform : public FRigUnitMutable
 {
 	GENERATED_BODY()
diff --git a/Engine/Plugins/Animation/ControlRig/Source/ControlRig/Private/Units/Hierarchy/RigUnit_SetControlVisibility.h b/Engine/Plugins/Animation/ControlRig/Source/ControlRig/Private/Units/Hierarchy/RigUnit_SetControlVisibility.h
--- a/Engine/Plugins/Animation/ControlRig/Source/ControlRig/Private/Units/Hierarchy/RigUnit_SetControlVisibility.h
+++ b/Engine/Plugins/Animation/ControlRig/Source/ControlRig/Private/Units/Hierarchy/RigUnit_SetControlVisibility.h
@@ -8,7 +8,7 @@
 /**
  * SetControlVisibility is used to change the gizmo visibility on a control at runtime
  */
-USTRUCT(meta=(DisplayName="Set Control Visibility", Category="Hierarchy", DocumentationPolicy="Strict", Keywords = "SetControlVisibility,Visibility,Hide,Show,Hidden,Visible,SetGizmoVisibility", PrototypeName = "SetControlVisibility", NodeColor="0, 0.364706, 1.0"))
+USTRUCT(meta=(DisplayName="Set Control Visibility", Category="Controls", DocumentationPolicy="Strict", Keywords = "SetControlVisibility,Visibility,Hide,Show,Hidden,Visible,SetGizmoVisibility", PrototypeName = "SetControlVisibility", NodeColor="0, 0.364706, 1.0"))
 struct CONTROLRIG_API FRigUnit_SetControlVisibility : public FRigUnitMutable
 {
 	GENERATED_BODY()
diff --git a/Engine/Plugins/Animation/ControlRig/Source/ControlRig/Private/Units/Hierarchy/RigUnit_SetRelativeTransform.h b/Engine/Plugins/Animation/ControlRig/Source/ControlRig/Private/Units/Hierarchy/RigUnit_SetRelativeTransform.h
--- a/Engine/Plugins/Animation/ControlRig/Source/ControlRig/Private/Units/Hierarchy/RigUnit_SetRelativeTransform.h
+++ b/Engine/Plugins/Animation/ControlRig/Source/ControlRig/Private/Units/Hierarchy/RigUnit_SetRelativeTransform.h
@@ -8,7 +8,7 @@
 /**
  * SetRelativeTransform is used to set a single transform from a hierarchy in the space of another item
  */
-USTRUCT(meta=(DisplayName="Set Relative Transform", Category="Hierarchy", DocumentationPolicy = "Strict", Keywords = "Offset,Local", NodeColor="0, 0.364706, 1.0", Varying))
+USTRUCT(meta=(DisplayName="Set Relative Transform", Category="Transforms", DocumentationPolicy = "Strict", Keywords = "Offset,Local", NodeColor="0, 0.364706, 1.0", Varying))
 struct CONTROLRIG_API FRigUnit_SetRelativeTransformForItem : public FRigUnitMutable
 {
 	GENERATED_BODY()
diff --git a/Engine/Plugins/Animation/ControlRig/Source/ControlRig/Private/Units/Highlevel/Harmonics/RigUnit_BoneHarmonics.h b/Engine/Plugins/Animation/ControlRig/Source/ControlRig/Private/Units/Highlevel/Harmonics/RigUnit_BoneHarmonics.h
--- a/Engine/Plugins/Animation/ControlRig/Source/ControlRig/Private/Units/Highlevel/Harmonics/RigUnit_BoneHarmonics.h
+++ b/Engine/Plugins/Animation/ControlRig/Source/ControlRig/Private/Units/Highlevel/Harmonics/RigUnit_BoneHarmonics.h
@@ -145,7 +145,7 @@ struct CONTROLRIG_API FRigUnit_BoneHarmonics : public FRigUnit_HighlevelBaseMuta
 /**
  * Drives an array of items through a harmonics spectrum
  */
-USTRUCT(meta=(DisplayName="Harmonics", Keywords="Sin,Wave"))
+USTRUCT(meta=(DisplayName="Harmonics", Keywords="Sin,Wave", Category = "Simulation"))
 struct CONTROLRIG_API FRigUnit_ItemHarmonics : public FRigUnit_HighlevelBaseMutable
 {
 	GENERATED_BODY()
diff --git a/Engine/Plugins/Animation/ControlRig/Source/ControlRig/Private/Units/Highlevel/Harmonics/RigUnit_ChainHarmonics.h b/Engine/Plugins/Animation/ControlRig/Source/ControlRig/Private/Units/Highlevel/Harmonics/RigUnit_ChainHarmonics.h
--- a/Engine/Plugins/Animation/ControlRig/Source/ControlRig/Private/Units/Highlevel/Harmonics/RigUnit_ChainHarmonics.h
+++ b/Engine/Plugins/Animation/ControlRig/Source/ControlRig/Private/Units/Highlevel/Harmonics/RigUnit_ChainHarmonics.h
@@ -231,7 +231,7 @@ struct CONTROLRIG_API FRigUnit_ChainHarmonics : public FRigUnit_HighlevelBaseMut
 /**
  * Given a root will drive all items underneath in a chain base harmonics spectrum
  */
-USTRUCT(meta=(DisplayName="Chain Harmonics"))
+USTRUCT(meta=(DisplayName="Chain Harmonics", Category = "Simulation"))
 struct CONTROLRIG_API FRigUnit_ChainHarmonicsPerItem : public FRigUnit_HighlevelBaseMutable
 {
 	GENERATED_BODY()
diff --git a/Engine/Plugins/Animation/ControlRig/Source/ControlRig/Private/Units/Highlevel/Hierarchy/RigUnit_ModifyTransforms.h b/Engine/Plugins/Animation/ControlRig/Source/ControlRig/Private/Units/Highlevel/Hierarchy/RigUnit_ModifyTransforms.h
--- a/Engine/Plugins/Animation/ControlRig/Source/ControlRig/Private/Units/Highlevel/Hierarchy/RigUnit_ModifyTransforms.h
+++ b/Engine/Plugins/Animation/ControlRig/Source/ControlRig/Private/Units/Highlevel/Hierarchy/RigUnit_ModifyTransforms.h
@@ -66,7 +66,7 @@ struct CONTROLRIG_API FRigUnit_ModifyTransforms_WorkData
 /**
  * Modify Transforms is used to perform a change in the hierarchy by setting one or more bones' transforms
  */
-USTRUCT(meta=(DisplayName="Modify Transforms", Category="Hierarchy", DocumentationPolicy="Strict", Keywords = "ModifyBone"))
+USTRUCT(meta=(DisplayName="Modify Transforms", Category="Transforms", DocumentationPolicy="Strict", Keywords = "ModifyBone"))
 struct CONTROLRIG_API FRigUnit_ModifyTransforms : public FRigUnit_HighlevelBaseMutable
 {
 	GENERATED_BODY()
diff --git a/Engine/Plugins/Animation/ControlRig/Source/ControlRig/Private/Units/Highlevel/Hierarchy/RigUnit_TransformConstraint.h b/Engine/Plugins/Animation/ControlRig/Source/ControlRig/Private/Units/Highlevel/Hierarchy/RigUnit_TransformConstraint.h
--- a/Engine/Plugins/Animation/ControlRig/Source/ControlRig/Private/Units/Highlevel/Hierarchy/RigUnit_TransformConstraint.h
+++ b/Engine/Plugins/Animation/ControlRig/Source/ControlRig/Private/Units/Highlevel/Hierarchy/RigUnit_TransformConstraint.h
@@ -239,7 +239,7 @@ struct FRigUnit_ParentConstraint_AdvancedSettings
 /**
 * Constrains an item's transform to multiple items' transforms
 */
-USTRUCT(meta=(DisplayName="Parent Constraint", Category="Transforms", Keywords = "Parent,Orient,Scale"))
+USTRUCT(meta=(DisplayName="Parent Constraint", Category="Constraints", Keywords = "Parent,Orient,Scale"))
 struct FRigUnit_ParentConstraint : public FRigUnit_HighlevelBaseMutable
 {
 	GENERATED_BODY()
@@ -277,7 +277,7 @@ struct FRigUnit_ParentConstraint : public FRigUnit_HighlevelBaseMutable
 /**
 * Constrains an item's position to multiple items' positions 
 */
-USTRUCT(meta=(DisplayName="Position Constraint", Category="Transforms", Keywords = "Parent,Translation"))
+USTRUCT(meta=(DisplayName="Position Constraint", Category="Constraints", Keywords = "Parent,Translation"))
 struct FRigUnit_PositionConstraint : public FRigUnit_HighlevelBaseMutable
 {
 	GENERATED_BODY()
@@ -335,7 +335,7 @@ struct FRigUnit_RotationConstraint_AdvancedSettings
 /**
 * Constrains an item's rotation to multiple items' rotations 
 */
-USTRUCT(meta=(DisplayName="Rotation Constraint", Category="Transforms", Keywords = "Parent,Orientation,Orient,Rotate"))
+USTRUCT(meta=(DisplayName="Rotation Constraint", Category="Constraints", Keywords = "Parent,Orientation,Orient,Rotate"))
 struct FRigUnit_RotationConstraint : public FRigUnit_HighlevelBaseMutable
 {
 	GENERATED_BODY()
@@ -373,7 +373,7 @@ struct FRigUnit_RotationConstraint : public FRigUnit_HighlevelBaseMutable
 /**
 * Constrains an item's scale to multiple items' scales
 */
-USTRUCT(meta=(DisplayName="Scale Constraint", Category="Transforms", Keywords = "Parent, Size"))
+USTRUCT(meta=(DisplayName="Scale Constraint", Category="Constraints", Keywords = "Parent, Size"))
 struct FRigUnit_ScaleConstraint : public FRigUnit_HighlevelBaseMutable
 {
 	GENERATED_BODY()
diff --git a/Engine/Plugins/Animation/ControlRig/Source/ControlRig/Private/Units/Math/RigUnit_MathQuaternion.h b/Engine/Plugins/Animation/ControlRig/Source/ControlRig/Private/Units/Math/RigUnit_MathQuaternion.h
--- a/Engine/Plugins/Animation/ControlRig/Source/ControlRig/Private/Units/Math/RigUnit_MathQuaternion.h
+++ b/Engine/Plugins/Animation/ControlRig/Source/ControlRig/Private/Units/Math/RigUnit_MathQuaternion.h
@@ -523,7 +523,7 @@ struct CONTROLRIG_API FRigUnit_MathQuaternionSwingTwist : public FRigUnit_MathQu
 /**
  * Enum of possible rotation orders
  */
-USTRUCT(meta = (DisplayName = "Rotation Order", Category = "Math|Rotation", Constant))
+USTRUCT(meta = (DisplayName = "Rotation Order", Category = "Math|Quaternion", Constant))
 struct CONTROLRIG_API FRigUnit_MathQuaternionRotationOrder : public FRigUnit_MathBase
 {
 	GENERATED_BODY()
diff --git a/Engine/Plugins/Animation/ControlRig/Source/ControlRig/Private/Units/Simulation/RigUnit_Accumulate.h b/Engine/Plugins/Animation/ControlRig/Source/ControlRig/Private/Units/Simulation/RigUnit_Accumulate.h
--- a/Engine/Plugins/Animation/ControlRig/Source/ControlRig/Private/Units/Simulation/RigUnit_Accumulate.h
+++ b/Engine/Plugins/Animation/ControlRig/Source/ControlRig/Private/Units/Simulation/RigUnit_Accumulate.h
@@ -5,11 +5,17 @@
 #include "Units/Simulation/RigUnit_SimBase.h"
 #include "RigUnit_Accumulate.generated.h"
 
+USTRUCT(meta=(Category="Simulation|Accumulate"))
+struct CONTROLRIG_API FRigUnit_AccumulateBase : public FRigUnit_SimBase
+{
+	GENERATED_BODY()
+};
+
 /**
  * Adds a value over time over and over again
  */
 USTRUCT(meta=(DisplayName="Accumulate Add (Float)", PrototypeName="AccumulateAdd", Keywords="Simulate,++"))
-struct CONTROLRIG_API FRigUnit_AccumulateFloatAdd : public FRigUnit_SimBase
+struct CONTROLRIG_API FRigUnit_AccumulateFloatAdd : public FRigUnit_AccumulateBase
 {
 	GENERATED_BODY()
 	
@@ -42,7 +48,7 @@ struct CONTROLRIG_API FRigUnit_AccumulateFloatAdd : public FRigUnit_SimBase
  * Adds a vector over time over and over again
  */
 USTRUCT(meta = (DisplayName="Accumulate Add (Vector)", PrototypeName="AccumulateAdd", Keywords="Simulate,++"))
-struct CONTROLRIG_API FRigUnit_AccumulateVectorAdd : public FRigUnit_SimBase
+struct CONTROLRIG_API FRigUnit_AccumulateVectorAdd : public FRigUnit_AccumulateBase
 {
 	GENERATED_BODY()
 
@@ -75,7 +81,7 @@ struct CONTROLRIG_API FRigUnit_AccumulateVectorAdd : public FRigUnit_SimBase
  * Multiplies a value over time over and over again
  */
 USTRUCT(meta=(DisplayName="Accumulate Mul (Float)", PrototypeName="AccumulateMul", Keywords="Simulate,**"))
-struct CONTROLRIG_API FRigUnit_AccumulateFloatMul : public FRigUnit_SimBase
+struct CONTROLRIG_API FRigUnit_AccumulateFloatMul : public FRigUnit_AccumulateBase
 {
 	GENERATED_BODY()
 	
@@ -108,7 +114,7 @@ struct CONTROLRIG_API FRigUnit_AccumulateFloatMul : public FRigUnit_SimBase
  * Multiplies a vector over time over and over again
  */
 USTRUCT(meta = (DisplayName="Accumulate Mul (Vector)", PrototypeName="AccumulateMul", Keywords="Simulate,**"))
-struct CONTROLRIG_API FRigUnit_AccumulateVectorMul : public FRigUnit_SimBase
+struct CONTROLRIG_API FRigUnit_AccumulateVectorMul : public FRigUnit_AccumulateBase
 {
 	GENERATED_BODY()
 
@@ -141,7 +147,7 @@ struct CONTROLRIG_API FRigUnit_AccumulateVectorMul : public FRigUnit_SimBase
  * Multiplies a quaternion over time over and over again
  */
 USTRUCT(meta=(DisplayName="Accumulate Mul (Quaternion)", PrototypeName="AccumulateMul", Keywords="Simulate,**"))
-struct CONTROLRIG_API FRigUnit_AccumulateQuatMul : public FRigUnit_SimBase
+struct CONTROLRIG_API FRigUnit_AccumulateQuatMul : public FRigUnit_AccumulateBase
 {
 	GENERATED_BODY()
 	
@@ -177,7 +183,7 @@ struct CONTROLRIG_API FRigUnit_AccumulateQuatMul : public FRigUnit_SimBase
  * Multiplies a transform over time over and over again
  */
 USTRUCT(meta = (DisplayName="Accumulate Mul (Transform)", PrototypeName="AccumulateMul", Keywords="Simulate,**"))
-struct CONTROLRIG_API FRigUnit_AccumulateTransformMul : public FRigUnit_SimBase
+struct CONTROLRIG_API FRigUnit_AccumulateTransformMul : public FRigUnit_AccumulateBase
 {
 	GENERATED_BODY()
 
@@ -213,7 +219,7 @@ struct CONTROLRIG_API FRigUnit_AccumulateTransformMul : public FRigUnit_SimBase
  * Interpolates two values over time over and over again
  */
 USTRUCT(meta=(DisplayName="Accumulate Lerp (Float)", PrototypeName="AccumulateLerp", Keywords="Simulate,Ramp"))
-struct CONTROLRIG_API FRigUnit_AccumulateFloatLerp : public FRigUnit_SimBase
+struct CONTROLRIG_API FRigUnit_AccumulateFloatLerp : public FRigUnit_AccumulateBase
 {
 	GENERATED_BODY()
 	
@@ -249,7 +255,7 @@ struct CONTROLRIG_API FRigUnit_AccumulateFloatLerp : public FRigUnit_SimBase
  * Interpolates two vectors over time over and over again
  */
 USTRUCT(meta = (DisplayName="Accumulate Lerp (Vector)", PrototypeName="AccumulateLerp", Keywords="Simulate,Ramp"))
-struct CONTROLRIG_API FRigUnit_AccumulateVectorLerp : public FRigUnit_SimBase
+struct CONTROLRIG_API FRigUnit_AccumulateVectorLerp : public FRigUnit_AccumulateBase
 {
 	GENERATED_BODY()
 
@@ -286,7 +292,7 @@ struct CONTROLRIG_API FRigUnit_AccumulateVectorLerp : public FRigUnit_SimBase
  * Interpolates two quaternions over time over and over again
  */
 USTRUCT(meta=(DisplayName="Accumulate Lerp (Quaternion)", PrototypeName="AccumulateLerp", Keywords="Simulate,Ramp"))
-struct CONTROLRIG_API FRigUnit_AccumulateQuatLerp : public FRigUnit_SimBase
+struct CONTROLRIG_API FRigUnit_AccumulateQuatLerp : public FRigUnit_AccumulateBase
 {
 	GENERATED_BODY()
 	
@@ -323,7 +329,7 @@ struct CONTROLRIG_API FRigUnit_AccumulateQuatLerp : public FRigUnit_SimBase
  * Interpolates two transforms over time over and over again
  */
 USTRUCT(meta = (DisplayName="Accumulate Lerp (Transform)", PrototypeName="AccumulateLerp", Keywords="Simulate,Ramp"))
-struct CONTROLRIG_API FRigUnit_AccumulateTransformLerp : public FRigUnit_SimBase
+struct CONTROLRIG_API FRigUnit_AccumulateTransformLerp : public FRigUnit_AccumulateBase
 {
 	GENERATED_BODY()
 
@@ -360,7 +366,7 @@ struct CONTROLRIG_API FRigUnit_AccumulateTransformLerp : public FRigUnit_SimBase
  * Accumulates the min and max values over time
  */
 USTRUCT(meta=(DisplayName="Accumulate Range (Float)", PrototypeName="AccumulateRange", Keywords="Range"))
-struct CONTROLRIG_API FRigUnit_AccumulateFloatRange : public FRigUnit_SimBase
+struct CONTROLRIG_API FRigUnit_AccumulateFloatRange : public FRigUnit_AccumulateBase
 {
 	GENERATED_BODY()
 	
@@ -392,7 +398,7 @@ struct CONTROLRIG_API FRigUnit_AccumulateFloatRange : public FRigUnit_SimBase
  * Accumulates the min and max values over time
  */
 USTRUCT(meta = (DisplayName="Accumulate Range (Vector)", PrototypeName="AccumulateRange", Keywords="Range"))
-struct CONTROLRIG_API FRigUnit_AccumulateVectorRange : public FRigUnit_SimBase
+struct CONTROLRIG_API FRigUnit_AccumulateVectorRange : public FRigUnit_AccumulateBase
 {
 	GENERATED_BODY()
 
diff --git a/Engine/Plugins/Animation/ControlRig/Source/ControlRig/Private/Units/Simulation/RigUnit_AlphaInterp.h b/Engine/Plugins/Animation/ControlRig/Source/ControlRig/Private/Units/Simulation/RigUnit_AlphaInterp.h
--- a/Engine/Plugins/Animation/ControlRig/Source/ControlRig/Private/Units/Simulation/RigUnit_AlphaInterp.h
+++ b/Engine/Plugins/Animation/ControlRig/Source/ControlRig/Private/Units/Simulation/RigUnit_AlphaInterp.h
@@ -9,7 +9,7 @@
 /**
  * Takes in a float value and outputs an accumulated value with a customized scale and clamp
  */
-USTRUCT(meta=(DisplayName="Alpha Interpolate", Keywords="Alpha,Lerp,LinearInterpolate", PrototypeName = "AlphaInterp", MenuDescSuffix = "(Float)"))
+USTRUCT(meta=(DisplayName="Alpha Interpolate", Keywords="Alpha,Lerp,LinearInterpolate", Category = "Simulation|Time", PrototypeName = "AlphaInterp", MenuDescSuffix = "(Float)"))
 struct CONTROLRIG_API FRigUnit_AlphaInterp : public FRigUnit_SimBase
 {
 	GENERATED_BODY()
@@ -82,7 +82,7 @@ struct CONTROLRIG_API FRigUnit_AlphaInterp : public FRigUnit_SimBase
 /**
  * Takes in a vector value and outputs an accumulated value with a customized scale and clamp
  */
-USTRUCT(meta=(DisplayName="Alpha Interpolate", Keywords="Alpha,Lerp,LinearInterpolate", PrototypeName = "AlphaInterp", MenuDescSuffix = "(Vector)"))
+USTRUCT(meta=(DisplayName="Alpha Interpolate", Keywords="Alpha,Lerp,LinearInterpolate", Category = "Simulation|Time", PrototypeName = "AlphaInterp", MenuDescSuffix = "(Vector)"))
 struct CONTROLRIG_API FRigUnit_AlphaInterpVector : public FRigUnit_SimBase
 {
 	GENERATED_BODY()
diff --git a/Engine/Plugins/Animation/ControlRig/Source/ControlRig/Private/Units/Simulation/RigUnit_DeltaFromPrevious.h b/Engine/Plugins/Animation/ControlRig/Source/ControlRig/Private/Units/Simulation/RigUnit_DeltaFromPrevious.h
--- a/Engine/Plugins/Animation/ControlRig/Source/ControlRig/Private/Units/Simulation/RigUnit_DeltaFromPrevious.h
+++ b/Engine/Plugins/Animation/ControlRig/Source/ControlRig/Private/Units/Simulation/RigUnit_DeltaFromPrevious.h
@@ -8,7 +8,7 @@
 /**
  * Computes the difference from the previous value going through the node
  */
-USTRUCT(meta=(DisplayName="DeltaFromPrevious (Float)", PrototypeName="DeltaFromPrevious", Keywords="Difference,Velocity,Acceleration"))
+USTRUCT(meta=(DisplayName="DeltaFromPrevious (Float)", Category = "Simulation|Time", PrototypeName="DeltaFromPrevious", Keywords="Difference,Velocity,Acceleration"))
 struct CONTROLRIG_API FRigUnit_DeltaFromPreviousFloat : public FRigUnit_SimBase
 {
 	GENERATED_BODY()
@@ -37,7 +37,7 @@ struct CONTROLRIG_API FRigUnit_DeltaFromPreviousFloat : public FRigUnit_SimBase
 /**
  * Computes the difference from the previous value going through the node
  */
-USTRUCT(meta=(DisplayName="DeltaFromPrevious (Vector)", PrototypeName="DeltaFromPrevious", Keywords="Difference,Velocity,Acceleration"))
+USTRUCT(meta=(DisplayName="DeltaFromPrevious (Vector)", Category = "Simulation|Time", PrototypeName="DeltaFromPrevious", Keywords="Difference,Velocity,Acceleration"))
 struct CONTROLRIG_API FRigUnit_DeltaFromPreviousVector : public FRigUnit_SimBase
 {
 	GENERATED_BODY()
@@ -66,7 +66,7 @@ struct CONTROLRIG_API FRigUnit_DeltaFromPreviousVector : public FRigUnit_SimBase
 /**
  * Computes the difference from the previous value going through the node
  */
-USTRUCT(meta=(DisplayName="DeltaFromPrevious (Quaternion)", PrototypeName="DeltaFromPrevious", Keywords="Difference,Velocity,Acceleration"))
+USTRUCT(meta=(DisplayName="DeltaFromPrevious (Quaternion)", Category = "Simulation|Time", PrototypeName="DeltaFromPrevious", Keywords="Difference,Velocity,Acceleration"))
 struct CONTROLRIG_API FRigUnit_DeltaFromPreviousQuat : public FRigUnit_SimBase
 {
 	GENERATED_BODY()
@@ -95,7 +95,7 @@ struct CONTROLRIG_API FRigUnit_DeltaFromPreviousQuat : public FRigUnit_SimBase
 /**
  * Computes the difference from the previous value going through the node
  */
-USTRUCT(meta=(DisplayName="DeltaFromPrevious (Transform)", PrototypeName="DeltaFromPrevious", Keywords="Difference,Velocity,Acceleration"))
+USTRUCT(meta=(DisplayName="DeltaFromPrevious (Transform)", Category = "Simulation|Time", PrototypeName="DeltaFromPrevious", Keywords="Difference,Velocity,Acceleration"))
 struct CONTROLRIG_API FRigUnit_DeltaFromPreviousTransform : public FRigUnit_SimBase
 {
 	GENERATED_BODY()
diff --git a/Engine/Plugins/Animation/ControlRig/Source/ControlRig/Private/Units/Simulation/RigUnit_Kalman.h b/Engine/Plugins/Animation/ControlRig/Source/ControlRig/Private/Units/Simulation/RigUnit_Kalman.h
--- a/Engine/Plugins/Animation/ControlRig/Source/ControlRig/Private/Units/Simulation/RigUnit_Kalman.h
+++ b/Engine/Plugins/Animation/ControlRig/Source/ControlRig/Private/Units/Simulation/RigUnit_Kalman.h
@@ -9,7 +9,7 @@
  * Averages a value over time.
  * This uses a Kalman Filter internally.
  */
-USTRUCT(meta=(DisplayName="Average Over Time (Float)", PrototypeName="KalmanFilter", Keywords="Average,Smooth,KalmanFilter"))
+USTRUCT(meta=(DisplayName="Average Over Time (Float)", Category = "Simulation|Time", PrototypeName="KalmanFilter", Keywords="Average,Smooth,KalmanFilter"))
 struct CONTROLRIG_API FRigUnit_KalmanFloat : public FRigUnit_SimBase
 {
 	GENERATED_BODY()
@@ -44,7 +44,7 @@ struct CONTROLRIG_API FRigUnit_KalmanFloat : public FRigUnit_SimBase
  * Averages a value over time.
  * This uses a Kalman Filter internally.
  */
-USTRUCT(meta=(DisplayName="Average Over Time (Vector)", PrototypeName="KalmanFilter", Keywords="Average,Smooth,KalmanFilter"))
+USTRUCT(meta=(DisplayName="Average Over Time (Vector)", Category = "Simulation|Time", PrototypeName="KalmanFilter", Keywords="Average,Smooth,KalmanFilter"))
 struct CONTROLRIG_API FRigUnit_KalmanVector : public FRigUnit_SimBase
 {
 	GENERATED_BODY()
@@ -79,7 +79,7 @@ struct CONTROLRIG_API FRigUnit_KalmanVector : public FRigUnit_SimBase
  * Averages a transform over time.
  * This uses a Kalman Filter internally.
  */
-USTRUCT(meta=(DisplayName="Average Over Time (Transform)", PrototypeName="KalmanFilter", Keywords="Average,Smooth,KalmanFilter"))
+USTRUCT(meta=(DisplayName="Average Over Time (Transform)", Category = "Simulation|Time", PrototypeName="KalmanFilter", Keywords="Average,Smooth,KalmanFilter"))
 struct CONTROLRIG_API FRigUnit_KalmanTransform : public FRigUnit_SimBase
 {
 	GENERATED_BODY()
diff --git a/Engine/Plugins/Animation/ControlRig/Source/ControlRig/Private/Units/Simulation/RigUnit_SpringInterp.h b/Engine/Plugins/Animation/ControlRig/Source/ControlRig/Private/Units/Simulation/RigUnit_SpringInterp.h
--- a/Engine/Plugins/Animation/ControlRig/Source/ControlRig/Private/Units/Simulation/RigUnit_SpringInterp.h
+++ b/Engine/Plugins/Animation/ControlRig/Source/ControlRig/Private/Units/Simulation/RigUnit_SpringInterp.h
@@ -93,7 +93,7 @@ struct CONTROLRIG_API FRigUnit_SpringInterpVector : public FRigUnit_SimBase
 /**
  * Uses a simple spring model to interpolate a float from Current to Target.
  */
-USTRUCT(meta=(DisplayName="Spring Interpolate", Keywords="Alpha,SpringInterpolate,Verlet", PrototypeName = "SpringInterp"))
+USTRUCT(meta=(DisplayName="Spring Interpolate", Keywords="Alpha,SpringInterpolate,Verlet", Category = "Simulation|Springs", PrototypeName = "SpringInterp"))
 struct CONTROLRIG_API FRigUnit_SpringInterpV2 : public FRigUnit_SimBase
 {
 	GENERATED_BODY()
@@ -183,7 +183,7 @@ struct CONTROLRIG_API FRigUnit_SpringInterpV2 : public FRigUnit_SimBase
 /**
  * Uses a simple spring model to interpolate a vector from Current to Target.
  */
-USTRUCT(meta=(DisplayName="Spring Interpolate", Keywords="Alpha,SpringInterpolate,Verlet", PrototypeName = "SpringInterp", MenuDescSuffix = "(Vector)"))
+USTRUCT(meta=(DisplayName="Spring Interpolate", Keywords="Alpha,SpringInterpolate,Verlet", Category = "Simulation|Springs", PrototypeName = "SpringInterp", MenuDescSuffix = "(Vector)"))
 struct CONTROLRIG_API FRigUnit_SpringInterpVectorV2 : public FRigUnit_SimBase
 {
 	GENERATED_BODY()
@@ -273,7 +273,7 @@ struct CONTROLRIG_API FRigUnit_SpringInterpVectorV2 : public FRigUnit_SimBase
 /**
  * Uses a simple spring model to interpolate a quaternion from Current to Target.
  */
-USTRUCT(meta=(DisplayName="Spring Interpolate", Keywords="Alpha,SpringInterpolate,Verlet", PrototypeName = "SpringInterp", MenuDescSuffix = "(Quaternion)"))
+USTRUCT(meta=(DisplayName="Spring Interpolate", Keywords="Alpha,SpringInterpolate,Verlet", Category = "Simulation|Springs", PrototypeName = "SpringInterp", MenuDescSuffix = "(Quaternion)"))
 struct CONTROLRIG_API FRigUnit_SpringInterpQuaternionV2 : public FRigUnit_SimBase
 {
 	GENERATED_BODY()
diff --git a/Engine/Plugins/Animation/ControlRig/Source/ControlRig/Private/Units/Simulation/RigUnit_TimeOffset.h b/Engine/Plugins/Animation/ControlRig/Source/ControlRig/Private/Units/Simulation/RigUnit_TimeOffset.h
--- a/Engine/Plugins/Animation/ControlRig/Source/ControlRig/Private/Units/Simulation/RigUnit_TimeOffset.h
+++ b/Engine/Plugins/Animation/ControlRig/Source/ControlRig/Private/Units/Simulation/RigUnit_TimeOffset.h
@@ -8,7 +8,7 @@
 /**
  * Records a value over time and can access the value from the past
  */
-USTRUCT(meta=(DisplayName="Value Over Time (Float)", PrototypeName="TimeOffset", Keywords="Buffer,Delta,History,Previous,TimeOffset,Delay"))
+USTRUCT(meta=(DisplayName="Value Over Time (Float)", Category = "Simulation|Time", PrototypeName="TimeOffset", Keywords="Buffer,Delta,History,Previous,TimeOffset,Delay"))
 struct CONTROLRIG_API FRigUnit_TimeOffsetFloat : public FRigUnit_SimBase
 {
 	GENERATED_BODY()
@@ -61,7 +61,7 @@ struct CONTROLRIG_API FRigUnit_TimeOffsetFloat : public FRigUnit_SimBase
 /**
  * Records a value over time and can access the value from the past
  */
-USTRUCT(meta=(DisplayName="Value Over Time (Vector)", PrototypeName="TimeOffset", Keywords="Buffer,Delta,History,Previous,TimeOffset,Delay"))
+USTRUCT(meta=(DisplayName="Value Over Time (Vector)", Category = "Simulation|Time", PrototypeName="TimeOffset", Keywords="Buffer,Delta,History,Previous,TimeOffset,Delay"))
 struct CONTROLRIG_API FRigUnit_TimeOffsetVector : public FRigUnit_SimBase
 {
 	GENERATED_BODY()
@@ -114,7 +114,7 @@ struct CONTROLRIG_API FRigUnit_TimeOffsetVector : public FRigUnit_SimBase
 /**
  * Records a value over time and can access the value from the past
  */
-USTRUCT(meta=(DisplayName="Value Over Time (Transform)", PrototypeName="TimeOffset", Keywords="Buffer,Delta,History,Previous,TimeOffset,Delay"))
+USTRUCT(meta=(DisplayName="Value Over Time (Transform)", Category = "Simulation|Time", PrototypeName="TimeOffset", Keywords="Buffer,Delta,History,Previous,TimeOffset,Delay"))
 struct CONTROLRIG_API FRigUnit_TimeOffsetTransform : public FRigUnit_SimBase
 {
 	GENERATED_BODY()
diff --git a/Engine/Plugins/Animation/ControlRig/Source/ControlRig/Private/Units/Simulation/RigUnit_Timeline.h b/Engine/Plugins/Animation/ControlRig/Source/ControlRig/Private/Units/Simulation/RigUnit_Timeline.h
--- a/Engine/Plugins/Animation/ControlRig/Source/ControlRig/Private/Units/Simulation/RigUnit_Timeline.h
+++ b/Engine/Plugins/Animation/ControlRig/Source/ControlRig/Private/Units/Simulation/RigUnit_Timeline.h
@@ -8,7 +8,7 @@
 /**
  * Simulates a time value - can act as a timeline playing back
  */
-USTRUCT(meta=(DisplayName="Accumulated Time", Keywords="Playback,Pause,Timeline"))
+USTRUCT(meta=(DisplayName="Accumulated Time", Category = "Simulation|Accumulate", Keywords="Playback,Pause,Timeline"))
 struct CONTROLRIG_API FRigUnit_Timeline : public FRigUnit_SimBase
 {
 	GENERATED_BODY()
diff --git a/Engine/Plugins/Animation/ControlRig/Source/ControlRig/Private/Units/Simulation/RigUnit_Verlet.h b/Engine/Plugins/Animation/ControlRig/Source/ControlRig/Private/Units/Simulation/RigUnit_Verlet.h
--- a/Engine/Plugins/Animation/ControlRig/Source/ControlRig/Private/Units/Simulation/RigUnit_Verlet.h
+++ b/Engine/Plugins/Animation/ControlRig/Source/ControlRig/Private/Units/Simulation/RigUnit_Verlet.h
@@ -10,7 +10,7 @@
  * Simulates a single position over time using Verlet integration. It is recommended to use SpringInterp instead as it
  * is more accurate and stable, and has more meaningful parameters.
  */
-USTRUCT(meta=(DisplayName="Verlet (Vector)", PrototypeName="Verlet", Keywords="Simulate,Integrate"))
+USTRUCT(meta=(DisplayName="Verlet (Vector)", Category = "Simulation|Springs", PrototypeName="Verlet", Keywords="Simulate,Integrate"))
 struct CONTROLRIG_API FRigUnit_VerletIntegrateVector : public FRigUnit_SimBase
 {
 	GENERATED_BODY()
diff --git a/Engine/Plugins/Animation/ControlRig/Source/ControlRig/Public/Units/Hierarchy/RigUnit_SetTransform.h b/Engine/Plugins/Animation/ControlRig/Source/ControlRig/Public/Units/Hierarchy/RigUnit_SetTransform.h
--- a/Engine/Plugins/Animation/ControlRig/Source/ControlRig/Public/Units/Hierarchy/RigUnit_SetTransform.h
+++ b/Engine/Plugins/Animation/ControlRig/Source/ControlRig/Public/Units/Hierarchy/RigUnit_SetTransform.h
@@ -12,7 +12,7 @@
  * actually sets the Control's offset transform and resets the local
  * values to (0, 0, 0).
  */
-USTRUCT(meta=(DisplayName="Set Transform", Category="Hierarchy", DocumentationPolicy = "Strict", Keywords="SetBoneTransform,SetControlTransform,SetInitialTransform,SetSpaceTransform", NodeColor="0, 0.364706, 1.0", Varying))
+USTRUCT(meta=(DisplayName="Set Transform", Category="Transforms", DocumentationPolicy = "Strict", Keywords="SetBoneTransform,SetControlTransform,SetInitialTransform,SetSpaceTransform", NodeColor="0, 0.364706, 1.0", Varying))
 struct CONTROLRIG_API FRigUnit_SetTransform : public FRigUnitMutable
 {
 	GENERATED_BODY()
@@ -85,7 +85,7 @@ struct CONTROLRIG_API FRigUnit_SetTransform : public FRigUnitMutable
 /**
  * SetTranslation is used to set a single translation on hierarchy.
  */
-USTRUCT(meta=(DisplayName="Set Translation", Category="Hierarchy", DocumentationPolicy = "Strict", Keywords="SetBoneTranslation,SetControlTranslation,SetInitialTranslation,SetSpaceTranslation,SetBoneLocation,SetControlLocation,SetInitialLocation,SetSpaceLocation,SetBonePosition,SetControlPosition,SetInitialPosition,SetSpacePosition,SetTranslation,SetLocation,SetPosition", NodeColor="0, 0.364706, 1.0", Varying))
+USTRUCT(meta=(DisplayName="Set Translation", Category="Transforms", DocumentationPolicy = "Strict", Keywords="SetBoneTranslation,SetControlTranslation,SetInitialTranslation,SetSpaceTranslation,SetBoneLocation,SetControlLocation,SetInitialLocation,SetSpaceLocation,SetBonePosition,SetControlPosition,SetInitialPosition,SetSpacePosition,SetTranslation,SetLocation,SetPosition", NodeColor="0, 0.364706, 1.0", Varying))
 struct CONTROLRIG_API FRigUnit_SetTranslation : public FRigUnitMutable
 {
 	GENERATED_BODY()
@@ -149,7 +149,7 @@ struct CONTROLRIG_API FRigUnit_SetTranslation : public FRigUnitMutable
 /**
  * SetRotation is used to set a single rotation on hierarchy.
  */
-USTRUCT(meta=(DisplayName="Set Rotation", Category="Hierarchy", DocumentationPolicy = "Strict", Keywords="SetBoneRotation,SetControlRotation,SetInitialRotation,SetSpaceRotation,SetBoneOrientation,SetControlOrientation,SetInitialOrientation,SetSpaceOrientation,SetRotation,SetOrientation", NodeColor="0, 0.364706, 1.0", Varying))
+USTRUCT(meta=(DisplayName="Set Rotation", Category="Transforms", DocumentationPolicy = "Strict", Keywords="SetBoneRotation,SetControlRotation,SetInitialRotation,SetSpaceRotation,SetBoneOrientation,SetControlOrientation,SetInitialOrientation,SetSpaceOrientation,SetRotation,SetOrientation", NodeColor="0, 0.364706, 1.0", Varying))
 struct CONTROLRIG_API FRigUnit_SetRotation : public FRigUnitMutable
 {
 	GENERATED_BODY()
@@ -213,7 +213,7 @@ struct CONTROLRIG_API FRigUnit_SetRotation : public FRigUnitMutable
 /**
  * SetScale is used to set a single scale on hierarchy.
  */
-USTRUCT(meta=(DisplayName="Set Scale", Category="Hierarchy", DocumentationPolicy = "Strict", Keywords="SetBoneScale,SetControlScale,SetInitialScale,SetSpaceScale,SetScale", NodeColor="0, 0.364706, 1.0", Varying))
+USTRUCT(meta=(DisplayName="Set Scale", Category="Transforms", DocumentationPolicy = "Strict", Keywords="SetBoneScale,SetControlScale,SetInitialScale,SetSpaceScale,SetScale", NodeColor="0, 0.364706, 1.0", Varying))
 struct FRigUnit_SetScale : public FRigUnitMutable
 {
 	GENERATED_BODY()
@@ -331,7 +331,7 @@ struct CONTROLRIG_API FRigUnit_SetTransformArray : public FRigUnitMutable
  * actually sets the Control's offset transform and resets the local
  * values to (0, 0, 0).
  */
-USTRUCT(meta=(DisplayName="Set Transform Array", Category="Hierarchy", DocumentationPolicy = "Strict", Keywords="SetBoneTransform,SetControlTransform,SetInitialTransform,SetSpaceTransform", NodeColor="0, 0.364706, 1.0", Varying))
+USTRUCT(meta=(DisplayName="Set Transform Array", Category="Transforms", DocumentationPolicy = "Strict", Keywords="SetBoneTransform,SetControlTransform,SetInitialTransform,SetSpaceTransform", NodeColor="0, 0.364706, 1.0", Varying))
 struct CONTROLRIG_API FRigUnit_SetTransformItemArray : public FRigUnitMutable
 {
 	GENERATED_BODY()
diff --git a/Engine/Plugins/Animation/ControlRig/Source/ControlRig/Public/Units/Highlevel/RigUnit_HighlevelBase.h b/Engine/Plugins/Animation/ControlRig/Source/ControlRig/Public/Units/Highlevel/RigUnit_HighlevelBase.h
--- a/Engine/Plugins/Animation/ControlRig/Source/ControlRig/Public/Units/Highlevel/RigUnit_HighlevelBase.h
+++ b/Engine/Plugins/Animation/ControlRig/Source/ControlRig/Public/Units/Highlevel/RigUnit_HighlevelBase.h
@@ -12,13 +12,13 @@ enum class EControlRigVectorKind : uint8
 	Location
 };
 
-USTRUCT(meta=(Abstract, Category="Debug", NodeColor = "0.462745, 1,0, 0.329412"))
+USTRUCT(meta=(Abstract, Category="Debug", NodeColor = "0.462745 1.0 0.329412"))
 struct FRigUnit_HighlevelBase : public FRigUnit
 {
 	GENERATED_BODY()
 };
 
-USTRUCT(meta=(Abstract, Category="Debug", NodeColor = "0, 0.364706, 1.0"))
+USTRUCT(meta=(Abstract, Category="Debug", NodeColor = "0 0.364706 1.0"))
 struct FRigUnit_HighlevelBaseMutable : public FRigUnitMutable
 {
 	GENERATED_BODY()
diff --git a/Engine/Plugins/Animation/ControlRigSpline/Source/ControlRigSpline/Public/ControlRigSplineUnits.h b/Engine/Plugins/Animation/ControlRigSpline/Source/ControlRigSpline/Public/ControlRigSplineUnits.h
--- a/Engine/Plugins/Animation/ControlRigSpline/Source/ControlRigSpline/Public/ControlRigSplineUnits.h
+++ b/Engine/Plugins/Animation/ControlRigSpline/Source/ControlRigSpline/Public/ControlRigSplineUnits.h
@@ -17,7 +17,7 @@ struct CONTROLRIGSPLINE_API FRigUnit_ControlRigSplineBase : public FRigUnit
 /*
  * Creates a Spline curve from an array of positions
  */
-USTRUCT(meta = (DisplayName = "Spline From Points", Keywords="Spline From Positions", Category = "Spline", Varying, NodeColor="0.737911 0.099899 0.099899"))
+USTRUCT(meta = (DisplayName = "Spline From Points", Keywords="Spline From Positions", Category = "Splines", Varying, NodeColor="0.737911 0.099899 0.099899"))
 struct CONTROLRIGSPLINE_API FRigUnit_ControlRigSplineFromPoints : public FRigUnit_ControlRigSplineBase
 {
 	GENERATED_BODY()
@@ -56,7 +56,7 @@ struct CONTROLRIGSPLINE_API FRigUnit_ControlRigSplineFromPoints : public FRigUni
 /*
  * Set the points of a spline, given a spline and an array of positions
  */
-USTRUCT(meta = (DisplayName = "Set Spline Points", Category = "Spline", Varying, NodeColor = "0.0 0.36470600962638855 1.0"))
+USTRUCT(meta = (DisplayName = "Set Spline Points", Category = "Splines", Varying, NodeColor = "0.0 0.36470600962638855 1.0"))
 struct CONTROLRIGSPLINE_API FRigUnit_SetSplinePoints : public FRigUnitMutable
 {
 	GENERATED_BODY()
@@ -80,7 +80,7 @@ struct CONTROLRIGSPLINE_API FRigUnit_SetSplinePoints : public FRigUnitMutable
 /*
  * Retrieves the position from a given Spline and U value
  */
-USTRUCT(meta = (DisplayName = "Position From Spline", Keywords="Point From Spline", Category = "Spline", NodeColor="0.737911 0.099899 0.099899"))
+USTRUCT(meta = (DisplayName = "Position From Spline", Keywords="Point From Spline", Category = "Splines", NodeColor="0.737911 0.099899 0.099899"))
 struct CONTROLRIGSPLINE_API FRigUnit_PositionFromControlRigSpline : public FRigUnit_ControlRigSplineBase
 {
 	GENERATED_BODY()
@@ -108,7 +108,7 @@ struct CONTROLRIGSPLINE_API FRigUnit_PositionFromControlRigSpline : public FRigU
 /*
  * Retrieves the transform from a given Spline and U value based on the given Up Vector and Roll
  */
-USTRUCT(meta = (DisplayName = "Transform From Spline", Category = "Spline", NodeColor="0.737911 0.099899 0.099899"))
+USTRUCT(meta = (DisplayName = "Transform From Spline", Category = "Splines", NodeColor="0.737911 0.099899 0.099899"))
 struct CONTROLRIGSPLINE_API FRigUnit_TransformFromControlRigSpline : public FRigUnit_ControlRigSplineBase
 {
 	GENERATED_BODY()
@@ -144,7 +144,7 @@ struct CONTROLRIGSPLINE_API FRigUnit_TransformFromControlRigSpline : public FRig
 /*
  * Retrieves the tangent from a given Spline and U value
  */
-USTRUCT(meta = (DisplayName = "Tangent From Spline", Category = "Spline", NodeColor="0.737911 0.099899 0.099899"))
+USTRUCT(meta = (DisplayName = "Tangent From Spline", Category = "Splines", NodeColor="0.737911 0.099899 0.099899"))
 struct CONTROLRIGSPLINE_API FRigUnit_TangentFromControlRigSpline : public FRigUnit_ControlRigSplineBase
 {
 	GENERATED_BODY()
@@ -172,7 +172,7 @@ struct CONTROLRIGSPLINE_API FRigUnit_TangentFromControlRigSpline : public FRigUn
 /*
  * Draws the given spline in the viewport
  */
-USTRUCT(meta = (DisplayName = "Draw Spline", Category = "Spline", NodeColor = "0.83077 0.846873 0.049707"))
+USTRUCT(meta = (DisplayName = "Draw Spline", Category = "Splines", NodeColor = "0.83077 0.846873 0.049707"))
 struct CONTROLRIGSPLINE_API FRigUnit_DrawControlRigSpline : public FRigUnitMutable
 {
 	GENERATED_BODY()
@@ -204,7 +204,7 @@ struct CONTROLRIGSPLINE_API FRigUnit_DrawControlRigSpline : public FRigUnitMutab
 /*
  * Retrieves the length from a given Splin
  */
-USTRUCT(meta = (DisplayName = "Get Length Of Spline", Category = "Spline", NodeColor="0.737911 0.099899 0.099899"))
+USTRUCT(meta = (DisplayName = "Get Length Of Spline", Category = "Splines", NodeColor="0.737911 0.099899 0.099899"))
 struct CONTROLRIGSPLINE_API FRigUnit_GetLengthControlRigSpline : public FRigUnit
 {
 	GENERATED_BODY()
@@ -345,7 +345,7 @@ struct CONTROLRIGSPLINE_API FRigUnit_FitChainToSplineCurve : public FRigUnit_Hig
  * Fits a given chain to a spline curve.
  * Additionally provides rotational control matching the features of the Distribute Rotation node.
  */
-USTRUCT(meta=(DisplayName="Fit Chain on Spline Curve", Category="Hierarchy", Keywords="Fit,Resample,Spline"))
+USTRUCT(meta=(DisplayName="Fit Chain on Spline Curve", Category="Splines", Keywords="Fit,Resample,Spline"))
 struct CONTROLRIGSPLINE_API FRigUnit_FitChainToSplineCurveItemArray : public FRigUnit_HighlevelBaseMutable
 {
 	GENERATED_BODY()
@@ -489,7 +489,7 @@ struct CONTROLRIGSPLINE_API FRigUnit_FitSplineCurveToChain : public FRigUnit_Hig
 /**
  * Fits a given spline curve to a chain.
  */
-USTRUCT(meta=(DisplayName="Fit Spline Curve on Chain", Category="Hierarchy", Keywords="Fit,Resample,Spline"))
+USTRUCT(meta=(DisplayName="Fit Spline Curve on Chain", Category="Splines", Keywords="Fit,Resample,Spline"))
 struct CONTROLRIGSPLINE_API FRigUnit_FitSplineCurveToChainItemArray : public FRigUnit_HighlevelBaseMutable
 {
 	GENERATED_BODY()
@@ -518,7 +518,7 @@ struct CONTROLRIGSPLINE_API FRigUnit_FitSplineCurveToChainItemArray : public FRi
 /*
  * Retrieves the closest U value from a given Spline and a position
  */
-USTRUCT(meta = (DisplayName = "Closest Parameter From Spline", Category = "Spline", NodeColor="0.737911 0.099899 0.099899"))
+USTRUCT(meta = (DisplayName = "Closest Parameter From Spline", Category = "Splines", NodeColor="0.737911 0.099899 0.099899"))
 struct CONTROLRIGSPLINE_API FRigUnit_ClosestParameterFromControlRigSpline : public FRigUnit_ControlRigSplineBase
 {
 	GENERATED_BODY()
@@ -547,7 +547,7 @@ struct CONTROLRIGSPLINE_API FRigUnit_ClosestParameterFromControlRigSpline : publ
 /*
  * Returns the U parameter of a spline given a length percentage (0.0 - 1.0)
  */
-USTRUCT(meta = (DisplayName = "Parameter At Length Percentage", Category = "Spline"))
+USTRUCT(meta = (DisplayName = "Parameter At Length Percentage", Category = "Splines"))
 struct CONTROLRIGSPLINE_API FRigUnit_ParameterAtPercentage : public FRigUnit_ControlRigSplineBase
 {
 	GENERATED_BODY()
