Control Rig: Align debug toolbar with Blueprint

diff --git a/Engine/Plugins/Animation/ControlRig/Source/ControlRigEditor/Private/ControlRigBlueprintCommands.cpp b/Engine/Plugins/Animation/ControlRig/Source/ControlRigEditor/Private/ControlRigBlueprintCommands.cpp
--- a/Engine/Plugins/Animation/ControlRig/Source/ControlRigEditor/Private/ControlRigBlueprintCommands.cpp
+++ b/Engine/Plugins/Animation/ControlRig/Source/ControlRigEditor/Private/ControlRigBlueprintCommands.cpp
@@ -18,6 +18,7 @@ void FControlRigBlueprintCommands::RegisterCommands()
 	UI_COMMAND(ReleaseMode, "Release Mode", "Compiles and Executes the rig, ignoring debug data.", EUserInterfaceActionType::Button, FInputChord());
 	UI_COMMAND(DebugMode, "Debug Mode", "Compiles and Executes the unoptimized rig, stopping at breakpoints.", EUserInterfaceActionType::Button, FInputChord());
 	UI_COMMAND(ResumeExecution, "Resume", "Resumes execution after being halted at a breakpoint.", EUserInterfaceActionType::Button, FInputChord(EKeys::F5));
+	UI_COMMAND(ShowCurrentStatement, "Show Current Statement", "Focuses on the node currently being debugged.", EUserInterfaceActionType::Button, FInputChord());
 	UI_COMMAND(StepInto, "Step Into", "Steps into the collapsed/function node, when halted at a breakpoint.", EUserInterfaceActionType::Button, FInputChord(EKeys::F11));
 	UI_COMMAND(StepOut, "Step Out", "Steps out of the collapsed/function node, when halted at a breakpoint.", EUserInterfaceActionType::Button, FInputChord(EKeys::F11, EModifierKey::Shift));
 	UI_COMMAND(StepOver, "Step Over", "Steps over the node, when halted at a breakpoint.", EUserInterfaceActionType::Button, FInputChord(EKeys::F10));
diff --git a/Engine/Plugins/Animation/ControlRig/Source/ControlRigEditor/Private/ControlRigBlueprintCommands.h b/Engine/Plugins/Animation/ControlRig/Source/ControlRigEditor/Private/ControlRigBlueprintCommands.h
--- a/Engine/Plugins/Animation/ControlRig/Source/ControlRigEditor/Private/ControlRigBlueprintCommands.h
+++ b/Engine/Plugins/Animation/ControlRig/Source/ControlRigEditor/Private/ControlRigBlueprintCommands.h
@@ -54,6 +54,9 @@ public:
 	/** Resume the execution of the graph when halted at a breakpoint */
 	TSharedPtr< FUICommandInfo > ResumeExecution;
 
+	/** Focuses on the node currently being debugged */
+	TSharedPtr< FUICommandInfo > ShowCurrentStatement;
+
 	/** Steps to the next node in the execution graph (at the same graph level) when halted at a breakpoint */
 	TSharedPtr< FUICommandInfo > StepOver;
 
diff --git a/Engine/Plugins/Animation/ControlRig/Source/ControlRigEditor/Private/Editor/ControlRigEditor.cpp b/Engine/Plugins/Animation/ControlRig/Source/ControlRigEditor/Private/Editor/ControlRigEditor.cpp
--- a/Engine/Plugins/Animation/ControlRig/Source/ControlRigEditor/Private/Editor/ControlRigEditor.cpp
+++ b/Engine/Plugins/Animation/ControlRig/Source/ControlRigEditor/Private/Editor/ControlRigEditor.cpp
@@ -95,6 +95,7 @@
 #include "Kismet2/KismetDebugUtilities.h"
 #include "Kismet2/WatchedPin.h"
 #include "ToolMenus.h"
+#include "Editor/EditorStyle/Private/SlateEditorStyle.h"
 
 #define LOCTEXT_NAMESPACE "ControlRigEditor"
 
@@ -603,6 +604,11 @@ void FControlRigEditor::BindCommands()
 		FExecuteAction::CreateSP(this, &FControlRigEditor::HandleBreakpointActionRequested, ERigVMBreakpointAction::Resume),			
 		FIsActionChecked::CreateSP(this, &FControlRigEditor::IsHaltedAtBreakpoint));
 
+	GetToolkitCommands()->MapAction(
+		FControlRigBlueprintCommands::Get().ShowCurrentStatement,
+		FExecuteAction::CreateSP(this, &FControlRigEditor::HandleShowCurrentStatement),
+		FIsActionChecked::CreateSP(this, &FControlRigEditor::IsHaltedAtBreakpoint));
+
 	GetToolkitCommands()->MapAction(
 		FControlRigBlueprintCommands::Get().StepOver,
 		FExecuteAction::CreateSP(this, &FControlRigEditor::HandleBreakpointActionRequested, ERigVMBreakpointAction::StepOver),
@@ -789,19 +795,25 @@ void FControlRigEditor::FillToolbar(FToolBarBuilder& ToolbarBuilder)
 			FSlateIcon(FEditorStyle::GetStyleSetName(), "LevelEditor.Recompile"),
 			true);
 
-
+		ToolbarBuilder.BeginStyleOverride(FName("Toolbar.BackplateLeftPlay"));
 		ToolbarBuilder.AddToolBarButton(FControlRigBlueprintCommands::Get().ResumeExecution,
-			NAME_None, TAttribute<FText>(), TAttribute<FText>(), FSlateIcon(FControlRigEditorStyle::Get().GetStyleSetName(), "ControlRig.ResumeExecution"));
+			NAME_None, TAttribute<FText>(), TAttribute<FText>(), FSlateIcon(FSlateEditorStyle::Get().GetStyleSetName(), "PlayWorld.ResumePlaySession"));
 
-		ToolbarBuilder.AddToolBarButton(FControlRigBlueprintCommands::Get().StepOver,
-			NAME_None, TAttribute<FText>(), TAttribute<FText>(), FSlateIcon(FControlRigEditorStyle::Get().GetStyleSetName(), "ControlRig.StepOver"));
+		ToolbarBuilder.BeginStyleOverride(FName("Toolbar.BackplateLeft"));
+		ToolbarBuilder.AddToolBarButton(FControlRigBlueprintCommands::Get().ShowCurrentStatement,
+			NAME_None, TAttribute<FText>(), TAttribute<FText>(), FSlateIcon(FSlateEditorStyle::Get().GetStyleSetName(), "PlayWorld.ShowCurrentStatement"));
 
+		ToolbarBuilder.BeginStyleOverride(FName("Toolbar.BackplateCenter"));
+		ToolbarBuilder.AddToolBarButton(FControlRigBlueprintCommands::Get().StepOver,
+			NAME_None, TAttribute<FText>(), TAttribute<FText>(), FSlateIcon(FSlateEditorStyle::Get().GetStyleSetName(), "PlayWorld.StepOver"));
 		ToolbarBuilder.AddToolBarButton(FControlRigBlueprintCommands::Get().StepInto,
-			NAME_None, TAttribute<FText>(), TAttribute<FText>(), FSlateIcon(FControlRigEditorStyle::Get().GetStyleSetName(), "ControlRig.StepInto"));
+			NAME_None, TAttribute<FText>(), TAttribute<FText>(), FSlateIcon(FSlateEditorStyle::Get().GetStyleSetName(), "PlayWorld.StepInto"));
 
+		ToolbarBuilder.BeginStyleOverride(FName("Toolbar.BackplateRight"));
 		ToolbarBuilder.AddToolBarButton(FControlRigBlueprintCommands::Get().StepOut,
-			NAME_None, TAttribute<FText>(), TAttribute<FText>(), FSlateIcon(FControlRigEditorStyle::Get().GetStyleSetName(), "ControlRig.StepOut"));
-		
+			NAME_None, TAttribute<FText>(), TAttribute<FText>(), FSlateIcon(FSlateEditorStyle::Get().GetStyleSetName(), "PlayWorld.StepOut"));
+
+		ToolbarBuilder.EndStyleOverride();
 
 	}
 	ToolbarBuilder.EndSection();
@@ -6082,6 +6094,23 @@ bool FControlRigEditor::OnActionMatchesName(FEdGraphSchemaAction* InAction, cons
 	return false;
 }
 
+void FControlRigEditor::HandleShowCurrentStatement()
+{
+	if(HaltedAtNode)
+	{
+		if(UControlRigBlueprint* Blueprint = Cast<UControlRigBlueprint>(GetBlueprintObj()))
+		{
+			if(UControlRigGraph* EdGraph = Cast<UControlRigGraph>(Blueprint->GetEdGraph(HaltedAtNode->GetGraph())))
+			{
+				if(UEdGraphNode* EdNode = EdGraph->FindNodeForModelNodeName(HaltedAtNode->GetFName()))
+				{
+					JumpToHyperlink(EdNode, false);
+				}
+			}
+		}
+	}
+}
+
 void FControlRigEditor::HandleBreakpointActionRequested(const ERigVMBreakpointAction BreakpointAction)
 {
 	if (ControlRig)
diff --git a/Engine/Plugins/Animation/ControlRig/Source/ControlRigEditor/Private/Editor/ControlRigEditor.h b/Engine/Plugins/Animation/ControlRig/Source/ControlRigEditor/Private/Editor/ControlRigEditor.h
--- a/Engine/Plugins/Animation/ControlRig/Source/ControlRigEditor/Private/Editor/ControlRigEditor.h
+++ b/Engine/Plugins/Animation/ControlRig/Source/ControlRigEditor/Private/Editor/ControlRigEditor.h
@@ -377,6 +377,8 @@ private:
 
 	 virtual bool OnActionMatchesName(FEdGraphSchemaAction* InAction, const FName& InName) const override;
 
+	void HandleShowCurrentStatement();
+	
 	void HandleBreakpointActionRequested(const ERigVMBreakpointAction BreakpointAction);
 
 	bool IsHaltedAtBreakpoint() const;
