Control Rig: Refactor details panel groups to sections

diff --git a/Engine/Plugins/Animation/ControlRig/Source/ControlRig/Public/Rigs/RigHierarchyElements.h b/Engine/Plugins/Animation/ControlRig/Source/ControlRig/Public/Rigs/RigHierarchyElements.h
--- a/Engine/Plugins/Animation/ControlRig/Source/ControlRig/Public/Rigs/RigHierarchyElements.h
+++ b/Engine/Plugins/Animation/ControlRig/Source/ControlRig/Public/Rigs/RigHierarchyElements.h
@@ -1007,7 +1007,7 @@ private:
 
 public:
 
-	UPROPERTY(BlueprintReadOnly, EditAnywhere, Category = Control, meta=(ShowOnlyInnerProperties))
+	UPROPERTY(BlueprintReadOnly, EditAnywhere, Category = Control)
 	FRigControlSettings Settings;
 
 	UPROPERTY(BlueprintReadOnly, EditAnywhere, Category = RigElement)
diff --git a/Engine/Plugins/Animation/ControlRig/Source/ControlRigEditor/Private/ControlRigEditorModule.cpp b/Engine/Plugins/Animation/ControlRig/Source/ControlRigEditor/Private/ControlRigEditorModule.cpp
--- a/Engine/Plugins/Animation/ControlRig/Source/ControlRigEditor/Private/ControlRigEditorModule.cpp
+++ b/Engine/Plugins/Animation/ControlRig/Source/ControlRigEditor/Private/ControlRigEditorModule.cpp
@@ -235,18 +235,7 @@ void FControlRigEditorModule::StartupModule()
 	PropertiesToUnregisterOnShutdown.Add(FRigComputedTransform::StaticStruct()->GetFName());
 	PropertyEditorModule.RegisterCustomPropertyTypeLayout(PropertiesToUnregisterOnShutdown.Last(), FOnGetPropertyTypeCustomizationInstance::CreateStatic(&FRigComputedTransformDetails::MakeInstance));
 
-	PropertiesToUnregisterOnShutdown.Add(FRigBoneElement::StaticStruct()->GetFName());
-	PropertyEditorModule.RegisterCustomPropertyTypeLayout(PropertiesToUnregisterOnShutdown.Last(), FOnGetPropertyTypeCustomizationInstance::CreateStatic(&FRigBoneElementDetails::MakeInstance));
-
-	PropertiesToUnregisterOnShutdown.Add(FRigControlElement::StaticStruct()->GetFName());
-	PropertyEditorModule.RegisterCustomPropertyTypeLayout(PropertiesToUnregisterOnShutdown.Last(), FOnGetPropertyTypeCustomizationInstance::CreateStatic(&FRigControlElementDetails::MakeInstance));
-
-	PropertiesToUnregisterOnShutdown.Add(FRigNullElement::StaticStruct()->GetFName());
-	PropertyEditorModule.RegisterCustomPropertyTypeLayout(PropertiesToUnregisterOnShutdown.Last(), FOnGetPropertyTypeCustomizationInstance::CreateStatic(&FRigNullElementDetails::MakeInstance));
-
-	PropertiesToUnregisterOnShutdown.Add(FRigVMGraphVariableDescription::StaticStruct()->GetFName());
-	PropertyEditorModule.RegisterCustomPropertyTypeLayout(PropertiesToUnregisterOnShutdown.Last(), FOnGetPropertyTypeCustomizationInstance::CreateStatic(&FRigVMLocalVariableDetails::MakeInstance));
-
+	FRigBaseElementDetails::RegisterSectionMappings(PropertyEditorModule);
 
 	// Register asset tools
 	auto RegisterAssetTypeAction = [this](const TSharedRef<IAssetTypeActions>& InAssetTypeAction)
diff --git a/Engine/Plugins/Animation/ControlRig/Source/ControlRigEditor/Private/ControlRigElementDetails.cpp b/Engine/Plugins/Animation/ControlRig/Source/ControlRigEditor/Private/ControlRigElementDetails.cpp
--- a/Engine/Plugins/Animation/ControlRig/Source/ControlRigEditor/Private/ControlRigElementDetails.cpp
+++ b/Engine/Plugins/Animation/ControlRig/Source/ControlRigEditor/Private/ControlRigElementDetails.cpp
@@ -722,17 +722,17 @@ void FRigComputedTransformDetails::OnTransformChanged(FEditPropertyChain* InProp
 	}
 }
 
-void FRigBaseElementDetails::CustomizeHeader(TSharedRef<IPropertyHandle> InStructPropertyHandle, FDetailWidgetRow& HeaderRow, IPropertyTypeCustomizationUtils& StructCustomizationUtils)
+void FRigBaseElementDetails::CustomizeDetails(IDetailLayoutBuilder& DetailBuilder)
 {
 	BlueprintBeingCustomized = nullptr;
 	HierarchyBeingCustomized = nullptr;
 	ObjectsBeingCustomized.Reset();
-	
-	TArray<UObject*> Objects;
-	InStructPropertyHandle->GetOuterObjects(Objects);
-	for (UObject* Object : Objects)
+
+	TArray<TWeakObjectPtr<UObject>> DetailObjects;
+	DetailBuilder.GetObjectsBeingCustomized(DetailObjects);
+	for(TWeakObjectPtr<UObject> DetailObject : DetailObjects)
 	{
-		UDetailsViewWrapperObject* WrapperObject = CastChecked<UDetailsViewWrapperObject>(Object);
+		UDetailsViewWrapperObject* WrapperObject = CastChecked<UDetailsViewWrapperObject>(DetailObject.Get());
 
 		if(HierarchyBeingCustomized == nullptr)
 		{
@@ -756,13 +756,12 @@ void FRigBaseElementDetails::CustomizeHeader(TSharedRef<IPropertyHandle> InStruc
 
 	if(BlueprintBeingCustomized == nullptr)
 	{	
-		RigElementDetails_GetCustomizedInfo(InStructPropertyHandle, BlueprintBeingCustomized);
+		RigElementDetails_GetCustomizedInfo(DetailBuilder.GetProperty(TEXT("Key")), BlueprintBeingCustomized);
 	}
-}
 
-void FRigBaseElementDetails::CustomizeChildren(TSharedRef<class IPropertyHandle> InStructPropertyHandle, class IDetailChildrenBuilder& StructBuilder, IPropertyTypeCustomizationUtils& StructCustomizationUtils)
-{
-	StructBuilder.AddCustomRow(FText::FromString(TEXT("Name")))
+	IDetailCategoryBuilder& GeneralCategory = DetailBuilder.EditCategory(TEXT("General"), LOCTEXT("General", "General"));
+
+	GeneralCategory.AddCustomRow(FText::FromString(TEXT("Name")))
 	.NameContent()
 	[
 		SNew(STextBlock)
@@ -778,6 +777,8 @@ void FRigBaseElementDetails::CustomizeChildren(TSharedRef<class IPropertyHandle>
 		.OnVerifyTextChanged(this, &FRigBaseElementDetails::OnVerifyNameChanged)
 		.IsEnabled(ObjectsBeingCustomized.Num() == 1)
 	];
+
+	DetailBuilder.HideCategory(TEXT("RigElement"));
 }
 
 FRigElementKey FRigBaseElementDetails::GetElementKey() const
@@ -785,7 +786,7 @@ FRigElementKey FRigBaseElementDetails::GetElementKey() const
 	check(ObjectsBeingCustomized.Num() == 1);
 	if(ObjectsBeingCustomized[0].IsValid())
 	{
-		return ObjectsBeingCustomized[0]->GetContent<FRigBaseElement>()->GetKey(); 
+		return ObjectsBeingCustomized[0]->GetContent<FRigBaseElement>().GetKey(); 
 	}
 	return FRigElementKey();
 }
@@ -891,7 +892,7 @@ TArray<FRigElementKey> FRigBaseElementDetails::GetElementKeys() const
 	{
 		if(ObjectBeingCustomized.IsValid())
 		{
-			Keys.Add(ObjectBeingCustomized->GetContent<FRigBaseElement>()->GetKey());
+			Keys.Add(ObjectBeingCustomized->GetContent<FRigBaseElement>().GetKey());
 		}
 	}
 	return Keys;
@@ -903,9 +904,10 @@ bool FRigBaseElementDetails::IsAnyControlOfType(ERigControlType InType) const
 	{
 		if(ObjectBeingCustomized.IsValid())
 		{
-			if(FRigControlElement* ControlElement = Cast<FRigControlElement>(ObjectBeingCustomized->GetContent<FRigBaseElement>()))
+			if(ObjectBeingCustomized->IsChildOf<FRigControlElement>())
 			{
-				if(ControlElement->Settings.ControlType == InType)
+				const FRigControlElement ControlElement = ObjectBeingCustomized->GetContent<FRigControlElement>();
+				if(ControlElement.Settings.ControlType == InType)
 				{
 					return true;
 				}
@@ -921,9 +923,10 @@ bool FRigBaseElementDetails::IsAnyControlNotOfType(ERigControlType InType) const
 	{
 		if(ObjectBeingCustomized.IsValid())
 		{
-			if(FRigControlElement* ControlElement = Cast<FRigControlElement>(ObjectBeingCustomized->GetContent<FRigBaseElement>()))
+			if(ObjectBeingCustomized->IsChildOf<FRigControlElement>())
 			{
-				if(ControlElement->Settings.ControlType != InType)
+				const FRigControlElement ControlElement = ObjectBeingCustomized->GetContent<FRigControlElement>();
+				if(ControlElement.Settings.ControlType != InType)
 				{
 					return true;
 				}
@@ -933,15 +936,36 @@ bool FRigBaseElementDetails::IsAnyControlNotOfType(ERigControlType InType) const
 	return false;
 }
 
-void FRigTransformElementDetails::CustomizeChildren(TSharedRef<class IPropertyHandle> InStructPropertyHandle, class IDetailChildrenBuilder& StructBuilder, IPropertyTypeCustomizationUtils& StructCustomizationUtils)
+void FRigBaseElementDetails::RegisterSectionMappings(FPropertyEditorModule& PropertyEditorModule)
+{
+	FRigBoneElementDetails().RegisterSectionMappings(PropertyEditorModule, UDetailsViewWrapperObject::GetClassForStruct(FRigBoneElement::StaticStruct()));
+	FRigNullElementDetails().RegisterSectionMappings(PropertyEditorModule, UDetailsViewWrapperObject::GetClassForStruct(FRigNullElement::StaticStruct()));
+	FRigControlElementDetails().RegisterSectionMappings(PropertyEditorModule, UDetailsViewWrapperObject::GetClassForStruct(FRigControlElement::StaticStruct()));
+}
+
+void FRigBaseElementDetails::RegisterSectionMappings(FPropertyEditorModule& PropertyEditorModule, UClass* InClass)
 {
-	FRigBaseElementDetails::CustomizeChildren(InStructPropertyHandle, StructBuilder, StructCustomizationUtils);
+}
 
-	IDetailGroup& PoseGroupDefault = StructBuilder.AddGroup(TEXT("Pose"), LOCTEXT("Pose", "Pose")); 
-	IDetailGroup& PoseGroupAdvanced = StructBuilder.GetParentCategory().AddGroup(TEXT("Pose"), LOCTEXT("Pose", "Pose"), true, false); 
+void FRigTransformElementDetails::CustomizeDetails(IDetailLayoutBuilder& DetailBuilder)
+{
+	FRigBaseElementDetails::CustomizeDetails(DetailBuilder);
+}
+
+void FRigTransformElementDetails::RegisterSectionMappings(FPropertyEditorModule& PropertyEditorModule, UClass* InClass)
+{
+	FRigBaseElementDetails::RegisterSectionMappings(PropertyEditorModule, InClass);
 	
+	TSharedRef<FPropertySection> TransformSection = PropertyEditorModule.FindOrCreateSection(InClass->GetFName(), "Transform", LOCTEXT("Transform", "Transform"));
+	TransformSection->AddCategory("General");
+	TransformSection->AddCategory("Transform");
+}
+
+void FRigTransformElementDetails::CustomizeTransform(IDetailLayoutBuilder& DetailBuilder)
+{
+	IDetailCategoryBuilder& TransformCategory = DetailBuilder.EditCategory(TEXT("Transform"), LOCTEXT("Transform", "Transform"));
 	{
-		const TSharedPtr<IPropertyHandle> PoseHandle = InStructPropertyHandle->GetChildHandle(TEXT("Pose"));
+		const TSharedPtr<IPropertyHandle> PoseHandle = DetailBuilder.GetProperty(TEXT("Pose"));
 		const TSharedPtr<IPropertyHandle> InitialHandle = PoseHandle->GetChildHandle(TEXT("Initial"));
 		const TSharedPtr<IPropertyHandle> CurrentHandle = PoseHandle->GetChildHandle(TEXT("Current"));
 
@@ -950,8 +974,9 @@ void FRigTransformElementDetails::CustomizeChildren(TSharedRef<class IPropertyHa
 			const TSharedPtr<IPropertyHandle> GlobalHandle = InitialHandle->GetChildHandle(TEXT("Global"));
 			const TSharedPtr<IPropertyHandle> TransformHandle = GlobalHandle->GetChildHandle(TEXT("Transform"));
 
-			PoseGroupAdvanced.AddPropertyRow(TransformHandle.ToSharedRef()).DisplayName(FText::FromString(TEXT("Initial Global")))
-			.IsEnabled(TAttribute<bool>::Create(TAttribute<bool>::FGetter::CreateSP(this, &FRigBaseElementDetails::IsSetupModeEnabled)));
+			TransformCategory.AddProperty(TransformHandle.ToSharedRef(), EPropertyLocation::Advanced)
+				.DisplayName(FText::FromString(TEXT("Initial Global")))
+				.IsEnabled(TAttribute<bool>::Create(TAttribute<bool>::FGetter::CreateSP(this, &FRigBaseElementDetails::IsSetupModeEnabled)));
 		}
 
 		// setup initial local
@@ -959,8 +984,9 @@ void FRigTransformElementDetails::CustomizeChildren(TSharedRef<class IPropertyHa
 			const TSharedPtr<IPropertyHandle> LocalHandle = InitialHandle->GetChildHandle(TEXT("Local"));
 			const TSharedPtr<IPropertyHandle> TransformHandle = LocalHandle->GetChildHandle(TEXT("Transform"));
 
-			PoseGroupAdvanced.AddPropertyRow(TransformHandle.ToSharedRef()).DisplayName(FText::FromString(TEXT("Initial Local")))
-			.IsEnabled(TAttribute<bool>::Create(TAttribute<bool>::FGetter::CreateSP(this, &FRigBaseElementDetails::IsSetupModeEnabled)));
+			TransformCategory.AddProperty(TransformHandle.ToSharedRef(), EPropertyLocation::Advanced)
+				.DisplayName(FText::FromString(TEXT("Initial Local")))
+				.IsEnabled(TAttribute<bool>::Create(TAttribute<bool>::FGetter::CreateSP(this, &FRigBaseElementDetails::IsSetupModeEnabled)));
 		}
 
 		// setup current global
@@ -968,8 +994,9 @@ void FRigTransformElementDetails::CustomizeChildren(TSharedRef<class IPropertyHa
 			const TSharedPtr<IPropertyHandle> GlobalHandle = CurrentHandle->GetChildHandle(TEXT("Global"));
 			const TSharedPtr<IPropertyHandle> TransformHandle = GlobalHandle->GetChildHandle(TEXT("Transform"));
 
-			PoseGroupAdvanced.AddPropertyRow(TransformHandle.ToSharedRef()).DisplayName(FText::FromString(TEXT("Current Global")))
-			.IsEnabled(TAttribute<bool>::Create(TAttribute<bool>::FGetter::CreateSP(this, &FRigBaseElementDetails::IsSetupModeEnabled)));
+			TransformCategory.AddProperty(TransformHandle.ToSharedRef(), EPropertyLocation::Advanced)
+				.DisplayName(FText::FromString(TEXT("Current Global")))
+				.IsEnabled(TAttribute<bool>::Create(TAttribute<bool>::FGetter::CreateSP(this, &FRigBaseElementDetails::IsSetupModeEnabled)));
 		}
 
 		// setup current local
@@ -977,8 +1004,9 @@ void FRigTransformElementDetails::CustomizeChildren(TSharedRef<class IPropertyHa
 			const TSharedPtr<IPropertyHandle> LocalHandle = CurrentHandle->GetChildHandle(TEXT("Local"));
 			const TSharedPtr<IPropertyHandle> TransformHandle = LocalHandle->GetChildHandle(TEXT("Transform"));
 
-			PoseGroupDefault.AddPropertyRow(TransformHandle.ToSharedRef()).DisplayName(FText::FromString(TEXT("Current Local")))
-			.IsEnabled(TAttribute<bool>::Create(TAttribute<bool>::FGetter::CreateSP(this, &FRigTransformElementDetails::IsCurrentLocalEnabled)));
+			TransformCategory.AddProperty(TransformHandle.ToSharedRef())
+				.DisplayName(FText::FromString(TEXT("Current Local")))
+				.IsEnabled(TAttribute<bool>::Create(TAttribute<bool>::FGetter::CreateSP(this, &FRigTransformElementDetails::IsCurrentLocalEnabled)));
 		}
 	}
 }
@@ -989,7 +1017,7 @@ bool FRigTransformElementDetails::IsCurrentLocalEnabled() const
 	{
 		if(ObjectBeingCustomized.IsValid())
 		{
-			if(ObjectBeingCustomized->GetContent<FRigBaseElement>()->GetType() == ERigElementType::Control)
+			if(ObjectBeingCustomized->GetContent<FRigBaseElement>().GetType() == ERigElementType::Control)
 			{
 				return false;
 			}
@@ -998,7 +1026,7 @@ bool FRigTransformElementDetails::IsCurrentLocalEnabled() const
 	return true;
 }
 
-void FRigControlElementDetails_SetupBoolValueWidget(IDetailGroup& InGroup, IDetailChildrenBuilder& InStructBuilder, ERigControlValueType InValueType, FRigControlElement* InControlElement, URigHierarchy* InHierarchy)
+void FRigControlElementDetails_SetupBoolValueWidget(IDetailCategoryBuilder& InCategory, const TSharedRef<IPropertyUtilities>& InPropertyUtilities, ERigControlValueType InValueType, const FRigControlElement* InControlElement, URigHierarchy* InHierarchy)
 {
 	UEnum* ControlTypeEnum = StaticEnum<ERigControlType>();
 	UEnum* ValueTypeEnum = StaticEnum<ERigControlValueType>();
@@ -1008,7 +1036,7 @@ void FRigControlElementDetails_SetupBoolValueWidget(IDetailGroup& InGroup, IDeta
 	TWeakObjectPtr<URigHierarchy> HierarchyPtr = InHierarchy;
 	const FRigElementKey Key = InControlElement->GetKey();
 
-	InStructBuilder.AddCustomRow(PropertyLabel)
+	InCategory.AddCustomRow(PropertyLabel)
 	.NameContent()
 	[
 	    SNew(STextBlock)
@@ -1066,7 +1094,7 @@ void FRigControlElementDetails_SetupBoolValueWidget(IDetailGroup& InGroup, IDeta
     })));
 }
 
-void FRigControlElementDetails_SetupIntegerValueWidget(IDetailGroup& InGroup, IDetailChildrenBuilder& InStructBuilder, ERigControlValueType InValueType, FRigControlElement* InControlElement, URigHierarchy* InHierarchy)
+void FRigControlElementDetails_SetupIntegerValueWidget(IDetailCategoryBuilder& InCategory, const TSharedRef<IPropertyUtilities>& InPropertyUtilities, ERigControlValueType InValueType, const FRigControlElement* InControlElement, URigHierarchy* InHierarchy)
 {
 	UEnum* ControlTypeEnum = StaticEnum<ERigControlType>();
 	UEnum* ValueTypeEnum = StaticEnum<ERigControlValueType>();
@@ -1095,7 +1123,7 @@ void FRigControlElementDetails_SetupIntegerValueWidget(IDetailGroup& InGroup, ID
 
 	if (InControlElement->Settings.ControlEnum)
 	{
-		InStructBuilder.AddCustomRow(PropertyLabel)
+		InCategory.AddCustomRow(PropertyLabel)
 		.Visibility(VisibilityAttribute)
         .NameContent()
         [
@@ -1147,7 +1175,7 @@ void FRigControlElementDetails_SetupIntegerValueWidget(IDetailGroup& InGroup, ID
 	}
 	else
 	{
-		InStructBuilder.AddCustomRow(PropertyLabel)
+		InCategory.AddCustomRow(PropertyLabel)
 		.Visibility(VisibilityAttribute)
         .NameContent()
         [
@@ -1231,7 +1259,7 @@ void FRigControlElementDetails_SetupIntegerValueWidget(IDetailGroup& InGroup, ID
 	}
 }
 
-void FRigControlElementDetails_SetupFloatValueWidget(IDetailGroup& InGroup, IDetailChildrenBuilder& InStructBuilder, ERigControlValueType InValueType, FRigControlElement* InControlElement, URigHierarchy* InHierarchy)
+void FRigControlElementDetails_SetupFloatValueWidget(IDetailCategoryBuilder& InCategory, const TSharedRef<IPropertyUtilities>& InPropertyUtilities, ERigControlValueType InValueType, const FRigControlElement* InControlElement, URigHierarchy* InHierarchy)
 {
 	UEnum* ControlTypeEnum = StaticEnum<ERigControlType>();
 	UEnum* ValueTypeEnum = StaticEnum<ERigControlValueType>();
@@ -1258,7 +1286,7 @@ void FRigControlElementDetails_SetupFloatValueWidget(IDetailGroup& InGroup, IDet
         return EnabledAttribute.Get() ? EVisibility::Visible : EVisibility::Hidden;
     }));
 
-	InStructBuilder.AddCustomRow(PropertyLabel)
+	InCategory.AddCustomRow(PropertyLabel)
 	.Visibility(VisibilityAttribute)
     .NameContent()
     [
@@ -1428,7 +1456,7 @@ FRigControlValue FRigControlElementDetails_PackageValue(const FEulerTransform& I
 }
 
 template<typename T>
-void FRigControlElementDetails_SetupStructValueWidget(IDetailGroup& InGroup, IDetailChildrenBuilder& InStructBuilder, ERigControlValueType InValueType, FRigControlElement* InControlElement, URigHierarchy* InHierarchy)
+void FRigControlElementDetails_SetupStructValueWidget(IDetailCategoryBuilder& InCategory, const TSharedRef<IPropertyUtilities>& InPropertyUtilities, ERigControlValueType InValueType, const FRigControlElement* InControlElement, URigHierarchy* InHierarchy)
 {
 	UEnum* ControlTypeEnum = StaticEnum<ERigControlType>();
 	UEnum* ValueTypeEnum = StaticEnum<ERigControlValueType>();
@@ -1463,7 +1491,7 @@ void FRigControlElementDetails_SetupStructValueWidget(IDetailGroup& InGroup, IDe
 		return EVisibility::Hidden;
     }));
 
-	IDetailPropertyRow* Row = InStructBuilder.AddExternalStructure(StructToDisplay.ToSharedRef());
+	IDetailPropertyRow* Row = InCategory.AddExternalStructure(StructToDisplay.ToSharedRef());
 	Row->DisplayName(PropertyLabel);
 	Row->ShouldAutoExpand(true);
 	Row->Visibility(VisibilityAttribute);
@@ -1492,7 +1520,7 @@ void FRigControlElementDetails_SetupStructValueWidget(IDetailGroup& InGroup, IDe
 	Handle->SetOnChildPropertyValueChanged(OnStructContentsChangedDelegate);
 }
 
-void FRigControlElementDetails_SetupValueWidget(IDetailGroup& InGroup, IDetailChildrenBuilder& InStructBuilder, ERigControlValueType InValueType, FRigControlElement* InControlElement, URigHierarchy* InHierarchy)
+void FRigControlElementDetails_SetupValueWidget(IDetailCategoryBuilder& InCategory, const TSharedRef<IPropertyUtilities>& InPropertyUtilities, ERigControlValueType InValueType, const FRigControlElement* InControlElement, URigHierarchy* InHierarchy)
 {
 	switch(InControlElement->Settings.ControlType)
 	{
@@ -1502,48 +1530,48 @@ void FRigControlElementDetails_SetupValueWidget(IDetailGroup& InGroup, IDetailCh
 			{
 				return;
 			}
-			FRigControlElementDetails_SetupBoolValueWidget(InGroup, InStructBuilder, InValueType, InControlElement, InHierarchy);
+			FRigControlElementDetails_SetupBoolValueWidget(InCategory, InPropertyUtilities, InValueType, InControlElement, InHierarchy);
 			break;
 		}
 		case ERigControlType::Integer:
 		{
-			FRigControlElementDetails_SetupIntegerValueWidget(InGroup, InStructBuilder, InValueType, InControlElement, InHierarchy);
+			FRigControlElementDetails_SetupIntegerValueWidget(InCategory, InPropertyUtilities, InValueType, InControlElement, InHierarchy);
 			break;
 		}
 		case ERigControlType::Float:
 		{
-			FRigControlElementDetails_SetupFloatValueWidget(InGroup, InStructBuilder, InValueType, InControlElement, InHierarchy);
+			FRigControlElementDetails_SetupFloatValueWidget(InCategory, InPropertyUtilities, InValueType, InControlElement, InHierarchy);
 			break;
 		}
 		case ERigControlType::Vector2D:
 		{
-			FRigControlElementDetails_SetupStructValueWidget<FVector2D>(InGroup, InStructBuilder, InValueType, InControlElement, InHierarchy);
+			FRigControlElementDetails_SetupStructValueWidget<FVector2D>(InCategory, InPropertyUtilities, InValueType, InControlElement, InHierarchy);
 			break;
 		}
 		case ERigControlType::Position:
 		case ERigControlType::Scale:
 		{
-			FRigControlElementDetails_SetupStructValueWidget<FVector>(InGroup, InStructBuilder, InValueType, InControlElement, InHierarchy);
+			FRigControlElementDetails_SetupStructValueWidget<FVector>(InCategory, InPropertyUtilities, InValueType, InControlElement, InHierarchy);
 			break;
 		}
 		case ERigControlType::Rotator:
 		{
-			FRigControlElementDetails_SetupStructValueWidget<FRotator>(InGroup, InStructBuilder, InValueType, InControlElement, InHierarchy);
+			FRigControlElementDetails_SetupStructValueWidget<FRotator>(InCategory, InPropertyUtilities, InValueType, InControlElement, InHierarchy);
 			break;
 		}
 		case ERigControlType::TransformNoScale:
 		{
-			FRigControlElementDetails_SetupStructValueWidget<FTransformNoScale>(InGroup, InStructBuilder, InValueType, InControlElement, InHierarchy);
+			FRigControlElementDetails_SetupStructValueWidget<FTransformNoScale>(InCategory, InPropertyUtilities, InValueType, InControlElement, InHierarchy);
 			break;
 		}
 		case ERigControlType::EulerTransform:
 		{
-			FRigControlElementDetails_SetupStructValueWidget<FEulerTransform>(InGroup, InStructBuilder, InValueType, InControlElement, InHierarchy);
+			FRigControlElementDetails_SetupStructValueWidget<FEulerTransform>(InCategory, InPropertyUtilities, InValueType, InControlElement, InHierarchy);
 			break;
 		}
 		case ERigControlType::Transform:
 		{
-			FRigControlElementDetails_SetupStructValueWidget<FTransform>(InGroup, InStructBuilder, InValueType, InControlElement, InHierarchy);
+			FRigControlElementDetails_SetupStructValueWidget<FTransform>(InCategory, InPropertyUtilities, InValueType, InControlElement, InHierarchy);
 			break;
 		}
 		default:
@@ -1554,11 +1582,17 @@ void FRigControlElementDetails_SetupValueWidget(IDetailGroup& InGroup, IDetailCh
 	}
 }
 
+void FRigBoneElementDetails::CustomizeDetails(IDetailLayoutBuilder& DetailBuilder)
+{
+	FRigTransformElementDetails::CustomizeDetails(DetailBuilder);
+	CustomizeTransform(DetailBuilder);
+}
+
 TArray<TSharedPtr<FString>> FRigControlElementDetails::ControlTypeList;
 
-void FRigControlElementDetails::CustomizeChildren(TSharedRef<class IPropertyHandle> InStructPropertyHandle, class IDetailChildrenBuilder& StructBuilder, IPropertyTypeCustomizationUtils& StructCustomizationUtils)
+void FRigControlElementDetails::CustomizeDetails(IDetailLayoutBuilder& DetailBuilder)
 {
-	FRigTransformElementDetails::CustomizeChildren(InStructPropertyHandle, StructBuilder, StructCustomizationUtils);
+	FRigTransformElementDetails::CustomizeDetails(DetailBuilder);
 
 	ShapeNameList.Reset();
 	if (BlueprintBeingCustomized)
@@ -1587,14 +1621,19 @@ void FRigControlElementDetails::CustomizeChildren(TSharedRef<class IPropertyHand
 		return;
 	}
 
-	IDetailGroup* ControlGroup = &StructBuilder.AddGroup(TEXT("Control"), LOCTEXT("Control", "Control"));
-	IDetailGroup* ShapeGroup = nullptr;
-	IDetailGroup* LimitsGroup = nullptr;
+	IDetailCategoryBuilder& ControlCategory = DetailBuilder.EditCategory(TEXT("Control"), LOCTEXT("Control", "Control"));
+	IDetailCategoryBuilder& ShapeCategory = DetailBuilder.EditCategory(TEXT("Shape"), LOCTEXT("Shape", "Shape"));
+	IDetailCategoryBuilder& LimitsCategory = DetailBuilder.EditCategory(TEXT("Limits"), LOCTEXT("Limits", "Limits"));
+
+	// the transform category should show up after the control, shape and limits categories
+	CustomizeTransform(DetailBuilder);
 
-	const TSharedPtr<IPropertyHandle> SettingsHandle = InStructPropertyHandle->GetChildHandle(TEXT("Settings"));
+	const TSharedPtr<IPropertyHandle> SettingsHandle = DetailBuilder.GetProperty(TEXT("Settings"));
 	const TSharedPtr<IPropertyHandle> DisplayNameHandle = SettingsHandle->GetChildHandle(TEXT("DisplayName"));
 
-	ControlGroup->AddWidgetRow()
+	DetailBuilder.HideProperty(SettingsHandle);
+
+	ControlCategory.AddCustomRow(LOCTEXT("DisplayName", "Display Name"))
 	.NameContent()
 	[
 		DisplayNameHandle->CreatePropertyNameWidget()
@@ -1617,12 +1656,14 @@ void FRigControlElementDetails::CustomizeChildren(TSharedRef<class IPropertyHand
 		}
 	}
 
+	const TSharedRef<IPropertyUtilities> PropertyUtilities = DetailBuilder.GetPropertyUtilities();
+
 	// when control type changes, we have to refresh detail panel
 	const TSharedPtr<IPropertyHandle> ControlTypeHandle = SettingsHandle->GetChildHandle(TEXT("ControlType"));
 	ControlTypeHandle->SetOnPropertyValueChanged(FSimpleDelegate::CreateLambda(
-		[this, &StructCustomizationUtils]()
+		[this, &PropertyUtilities]()
 		{
-			TArray<FRigControlElement*> ControlElementsInView = GetElementsInDetailsView<FRigControlElement>();
+			TArray<FRigControlElement> ControlElementsInView = GetElementsInDetailsView<FRigControlElement>();
 			TArray<FRigControlElement*> ControlElementsInHierarchy = GetElementsInHierarchy<FRigControlElement>();
 			check(ControlElementsInView.Num() == ControlElementsInHierarchy.Num());
 
@@ -1630,12 +1671,12 @@ void FRigControlElementDetails::CustomizeChildren(TSharedRef<class IPropertyHand
 			{
 				for(int32 ControlIndex = 0; ControlIndex< ControlElementsInView.Num(); ControlIndex++)
 				{
-					FRigControlElement* ViewElement = ControlElementsInView[ControlIndex];
+					const FRigControlElement& ViewElement = ControlElementsInView[ControlIndex];
 					FRigControlElement* ControlElement = ControlElementsInHierarchy[ControlIndex];
 					
 					FRigControlValue ValueToSet;
 
-					ControlElement->Settings.ControlType = ViewElement->Settings.ControlType;
+					ControlElement->Settings.ControlType = ViewElement.Settings.ControlType;
 					ControlElement->Settings.bLimitTranslation = false;
 					ControlElement->Settings.bLimitRotation = false;
 					ControlElement->Settings.bLimitScale = false;
@@ -1726,7 +1767,7 @@ void FRigControlElementDetails::CustomizeChildren(TSharedRef<class IPropertyHand
 					this->HierarchyBeingCustomized->SetControlValue(ControlElement, ValueToSet, ERigControlValueType::Initial, true, false, true);
 					this->HierarchyBeingCustomized->SetControlValue(ControlElement, ValueToSet, ERigControlValueType::Current, true, false, true);
 
-					FRigControlElement::StaticStruct()->CopyScriptStruct(ViewElement, ControlElement);
+					ObjectsBeingCustomized[ControlIndex]->SetContent<FRigControlElement>(*ControlElement);
 
 					if (this->HierarchyBeingCustomized != this->BlueprintBeingCustomized->Hierarchy)
 					{
@@ -1740,13 +1781,13 @@ void FRigControlElementDetails::CustomizeChildren(TSharedRef<class IPropertyHand
 					}
 				}
 				
-				StructCustomizationUtils.GetPropertyUtilities()->ForceRefresh();
+				PropertyUtilities->ForceRefresh();
 			}
 		}
 	));
 
-	ControlGroup->AddPropertyRow(ControlTypeHandle.ToSharedRef());
-	ControlGroup->AddPropertyRow(SettingsHandle->GetChildHandle(TEXT("bAnimatable")).ToSharedRef());
+	ControlCategory.AddProperty(ControlTypeHandle.ToSharedRef());
+	ControlCategory.AddProperty(SettingsHandle->GetChildHandle(TEXT("bAnimatable")).ToSharedRef());
 
 	// any but bool controls show the offset + shape + limits
 	const bool bNeedsShapeProperties = IsAnyControlNotOfType(ERigControlType::Bool);
@@ -1754,15 +1795,12 @@ void FRigControlElementDetails::CustomizeChildren(TSharedRef<class IPropertyHand
 	{
 		// setup offset
 		{
-			const TSharedPtr<IPropertyHandle> OffsetHandle = InStructPropertyHandle->GetChildHandle(TEXT("Offset"));
+			const TSharedPtr<IPropertyHandle> OffsetHandle = DetailBuilder.GetProperty(TEXT("Offset"));
 			const TSharedPtr<IPropertyHandle> InitialHandle = OffsetHandle->GetChildHandle(TEXT("Initial"));
 			const TSharedPtr<IPropertyHandle> LocalHandle = InitialHandle->GetChildHandle(TEXT("Local"));
 			const TSharedPtr<IPropertyHandle> TransformHandle = LocalHandle->GetChildHandle(TEXT("Transform"));
-			ControlGroup->AddPropertyRow(TransformHandle.ToSharedRef()).DisplayName(FText::FromString(TEXT("Offset Transform")));
+			ControlCategory.AddProperty(TransformHandle.ToSharedRef()).DisplayName(FText::FromString(TEXT("Offset Transform")));
 		}
-
-		ShapeGroup = &StructBuilder.AddGroup(TEXT("Control Shape"), LOCTEXT("ControlShape", "Control Shape"));
-		LimitsGroup = &StructBuilder.AddGroup(TEXT("Limits"), LOCTEXT("Limits", "Limits"));
 	}
 
 	if(IsAnyControlOfType(ERigControlType::Float) ||
@@ -1774,7 +1812,7 @@ void FRigControlElementDetails::CustomizeChildren(TSharedRef<class IPropertyHand
 		IsAnyControlOfType(ERigControlType::EulerTransform))
 	{
 		const TSharedPtr<IPropertyHandle> LimitHandle = SettingsHandle->GetChildHandle(TEXT("bLimitTranslation"));
-		LimitsGroup->AddPropertyRow(LimitHandle.ToSharedRef()).DisplayName(FText::FromString(TEXT("Limit Translation")));
+		LimitsCategory.AddProperty(LimitHandle.ToSharedRef()).DisplayName(FText::FromString(TEXT("Limit Translation")));
 	}
 
 	if(IsAnyControlOfType(ERigControlType::Rotator) ||
@@ -1783,7 +1821,7 @@ void FRigControlElementDetails::CustomizeChildren(TSharedRef<class IPropertyHand
 		IsAnyControlOfType(ERigControlType::EulerTransform))
 	{
 		const TSharedPtr<IPropertyHandle> LimitHandle = SettingsHandle->GetChildHandle(TEXT("bLimitRotation"));
-		LimitsGroup->AddPropertyRow(LimitHandle.ToSharedRef()).DisplayName(FText::FromString(TEXT("Limit Rotation")));
+		LimitsCategory.AddProperty(LimitHandle.ToSharedRef()).DisplayName(FText::FromString(TEXT("Limit Rotation")));
 	}
 
 	if(IsAnyControlOfType(ERigControlType::Scale) ||
@@ -1791,7 +1829,7 @@ void FRigControlElementDetails::CustomizeChildren(TSharedRef<class IPropertyHand
 		IsAnyControlOfType(ERigControlType::EulerTransform))
 	{
 		const TSharedPtr<IPropertyHandle> LimitHandle = SettingsHandle->GetChildHandle(TEXT("bLimitScale"));
-		LimitsGroup->AddPropertyRow(LimitHandle.ToSharedRef()).DisplayName(FText::FromString(TEXT("Limit Scale")));
+		LimitsCategory.AddProperty(LimitHandle.ToSharedRef()).DisplayName(FText::FromString(TEXT("Limit Scale")));
 	}
 
 	if(!(IsAnyControlNotOfType(ERigControlType::Integer) &&
@@ -1799,17 +1837,19 @@ void FRigControlElementDetails::CustomizeChildren(TSharedRef<class IPropertyHand
 		IsAnyControlNotOfType(ERigControlType::Vector2D)))
 	{
 		const TSharedPtr<IPropertyHandle> PrimaryAxisHandle = SettingsHandle->GetChildHandle(TEXT("PrimaryAxis"));
-		ControlGroup->AddPropertyRow(PrimaryAxisHandle.ToSharedRef()).DisplayName(FText::FromString(TEXT("Primary Axis")));
+		ControlCategory.AddProperty(PrimaryAxisHandle.ToSharedRef()).DisplayName(FText::FromString(TEXT("Primary Axis")));
 	}
 
-	TArray<FRigControlElement*> ControlElements;
+	TArray<FRigControlElement> ControlElements;
+	TArray<UDetailsViewWrapperObject*> ObjectPerControl;
 	for(TWeakObjectPtr<UDetailsViewWrapperObject> ObjectBeingCustomized : ObjectsBeingCustomized)
 	{
 		if(ObjectBeingCustomized.IsValid())
 		{
-			if(FRigControlElement* ControlElement = Cast<FRigControlElement>(ObjectBeingCustomized->GetContent<FRigBaseElement>()))
+			if(ObjectBeingCustomized->IsChildOf<FRigControlElement>())
 			{
-				ControlElements.Add(ControlElement);
+				ControlElements.Add(ObjectBeingCustomized->GetContent<FRigControlElement>());
+				ObjectPerControl.Add(ObjectBeingCustomized.Get());
 			}
 		}
 	}
@@ -1817,16 +1857,16 @@ void FRigControlElementDetails::CustomizeChildren(TSharedRef<class IPropertyHand
 	// only setup value widgets if there is only ony control selected
 	if(ControlElements.Num() == 1)
 	{
-		FRigControlElementDetails_SetupValueWidget(*ControlGroup, StructBuilder, ERigControlValueType::Current, ControlElements[0], HierarchyBeingCustomized);
+		FRigControlElementDetails_SetupValueWidget(ControlCategory, PropertyUtilities, ERigControlValueType::Current, &ControlElements[0], HierarchyBeingCustomized);
 
-		switch (ControlElements[0]->Settings.ControlType)
+		switch (ControlElements[0].Settings.ControlType)
 		{
 			case ERigControlType::Bool:
 			case ERigControlType::Float:
 			case ERigControlType::Integer:
 			case ERigControlType::Vector2D:
 			{
-				FRigControlElementDetails_SetupValueWidget(*ControlGroup, StructBuilder, ERigControlValueType::Initial, ControlElements[0], HierarchyBeingCustomized);
+				FRigControlElementDetails_SetupValueWidget(ControlCategory, PropertyUtilities, ERigControlValueType::Initial,&ControlElements[0], HierarchyBeingCustomized);
 				break;
 			}
 			default:
@@ -1836,29 +1876,29 @@ void FRigControlElementDetails::CustomizeChildren(TSharedRef<class IPropertyHand
 		}
 		
 		
-		FRigControlElementDetails_SetupValueWidget(*ControlGroup, StructBuilder, ERigControlValueType::Minimum, ControlElements[0], HierarchyBeingCustomized);
-		FRigControlElementDetails_SetupValueWidget(*ControlGroup, StructBuilder, ERigControlValueType::Maximum, ControlElements[0], HierarchyBeingCustomized);
+		FRigControlElementDetails_SetupValueWidget(LimitsCategory, PropertyUtilities, ERigControlValueType::Minimum, &ControlElements[0], HierarchyBeingCustomized);
+		FRigControlElementDetails_SetupValueWidget(LimitsCategory, PropertyUtilities, ERigControlValueType::Maximum, &ControlElements[0], HierarchyBeingCustomized);
 	}
 
 	if (bNeedsShapeProperties)
 	{
-		ShapeGroup->AddPropertyRow(SettingsHandle->GetChildHandle(TEXT("bShapeEnabled")).ToSharedRef())
+		ShapeCategory.AddProperty(SettingsHandle->GetChildHandle(TEXT("bShapeEnabled")).ToSharedRef())
 		.DisplayName(FText::FromString(TEXT("Enabled")));
-		ShapeGroup->AddPropertyRow(SettingsHandle->GetChildHandle(TEXT("bShapeVisible")).ToSharedRef())
+		ShapeCategory.AddProperty(SettingsHandle->GetChildHandle(TEXT("bShapeVisible")).ToSharedRef())
 		.DisplayName(FText::FromString(TEXT("Visible")));
 
 		// setup shape transform
 		{
-			const TSharedPtr<IPropertyHandle> ShapeHandle = InStructPropertyHandle->GetChildHandle(TEXT("Shape"));
+			const TSharedPtr<IPropertyHandle> ShapeHandle = DetailBuilder.GetProperty(TEXT("Shape"));
 			const TSharedPtr<IPropertyHandle> InitialHandle = ShapeHandle->GetChildHandle(TEXT("Initial"));
 			const TSharedPtr<IPropertyHandle> LocalHandle = InitialHandle->GetChildHandle(TEXT("Local"));
 			const TSharedPtr<IPropertyHandle> TransformHandle = LocalHandle->GetChildHandle(TEXT("Transform"));
-			ShapeGroup->AddPropertyRow(TransformHandle.ToSharedRef())
+			ShapeCategory.AddProperty(TransformHandle.ToSharedRef())
 			.IsEnabled(TAttribute<bool>::CreateSP(this, &FRigControlElementDetails::IsShapeEnabled));
 		}
 
 		const TSharedPtr<IPropertyHandle> ShapeShapeHandle = SettingsHandle->GetChildHandle(TEXT("ShapeName"));
-		ShapeGroup->AddPropertyRow(ShapeShapeHandle.ToSharedRef()).CustomWidget()
+		ShapeCategory.AddProperty(ShapeShapeHandle.ToSharedRef()).CustomWidget()
 		.NameContent()
 		[
 			SNew(STextBlock)
@@ -1873,24 +1913,27 @@ void FRigControlElementDetails::CustomizeChildren(TSharedRef<class IPropertyHand
 			.IsEnabled(this, &FRigControlElementDetails::IsShapeEnabled)
 		];
 
-		ShapeGroup->AddPropertyRow(SettingsHandle->GetChildHandle(TEXT("ShapeColor")).ToSharedRef())
+		ShapeCategory.AddProperty(SettingsHandle->GetChildHandle(TEXT("ShapeColor")).ToSharedRef())
 		.DisplayName(FText::FromString(TEXT("Color")));
 	}
 
 	if(IsAnyControlOfType(ERigControlType::Integer))
 	{
 		const TSharedPtr<IPropertyHandle> ControlEnumHandle = SettingsHandle->GetChildHandle(TEXT("ControlEnum"));
-		ControlGroup->AddPropertyRow(ControlEnumHandle.ToSharedRef()).DisplayName(FText::FromString(TEXT("Control Enum")));
+		ControlCategory.AddProperty(ControlEnumHandle.ToSharedRef()).DisplayName(FText::FromString(TEXT("Control Enum")));
 
 		ControlEnumHandle->SetOnPropertyValueChanged(FSimpleDelegate::CreateLambda(
-			[this, &StructCustomizationUtils, ControlElements]()
+			[this, &PropertyUtilities, ControlElements, ObjectPerControl]()
 			{
-				StructCustomizationUtils.GetPropertyUtilities()->ForceRefresh();
+				PropertyUtilities->ForceRefresh();
 
 				if (this->HierarchyBeingCustomized != nullptr)
 				{
-					for(FRigControlElement* ControlBeingCustomized : ControlElements)
+					for(int32 ControlIndex = 0; ControlIndex < ControlElements.Num(); ControlIndex++)
 					{
+						const FRigControlElement ControlInView = ControlElements[ControlIndex];
+						FRigControlElement* ControlBeingCustomized = HierarchyBeingCustomized->Find<FRigControlElement>(ControlInView.GetKey());
+						
 						const UEnum* ControlEnum = ControlBeingCustomized->Settings.ControlEnum;
 						if (ControlEnum != nullptr)
 						{
@@ -1921,15 +1964,35 @@ void FRigControlElementDetails::CustomizeChildren(TSharedRef<class IPropertyHand
 								}
 							}
 						}
+
+						ObjectPerControl[ControlIndex]->SetContent<FRigControlElement>(*ControlBeingCustomized);
 					}
 				}
 			}
 		));
 		
 	}
+	
+	IDetailCategoryBuilder& AnimationCategory = DetailBuilder.EditCategory(TEXT("Animation"), LOCTEXT("Animation", "Animation"));
+	AnimationCategory.AddProperty(SettingsHandle->GetChildHandle(TEXT("Customization"))->GetChildHandle(TEXT("AvailableSpaces")).ToSharedRef());
+}
+
+void FRigControlElementDetails::RegisterSectionMappings(FPropertyEditorModule& PropertyEditorModule, UClass* InClass)
+{
+	FRigTransformElementDetails::RegisterSectionMappings(PropertyEditorModule, InClass);
 
-	IDetailGroup* AnimationGroup = &StructBuilder.AddGroup(TEXT("Animation"), LOCTEXT("Animation", "Animation"));
-	AnimationGroup->AddPropertyRow(SettingsHandle->GetChildHandle(TEXT("Customization"))->GetChildHandle(TEXT("AvailableSpaces")).ToSharedRef());
+	TSharedRef<FPropertySection> ControlSection = PropertyEditorModule.FindOrCreateSection(InClass->GetFName(), "Control", LOCTEXT("Control", "Control"));
+	ControlSection->AddCategory("General");
+	ControlSection->AddCategory("Control");
+	ControlSection->AddCategory("Animation");
+
+	TSharedRef<FPropertySection> ShapeSection = PropertyEditorModule.FindOrCreateSection(InClass->GetFName(), "Shape", LOCTEXT("Shape", "Shape"));
+	ShapeSection->AddCategory("General");
+	ShapeSection->AddCategory("Shape");
+
+	TSharedRef<FPropertySection> LimitsSection = PropertyEditorModule.FindOrCreateSection(InClass->GetFName(), "Limits", LOCTEXT("Limits", "Limits"));
+	LimitsSection->AddCategory("General");
+	LimitsSection->AddCategory("Limits");
 }
 
 bool FRigControlElementDetails::IsShapeEnabled() const
@@ -1938,9 +2001,10 @@ bool FRigControlElementDetails::IsShapeEnabled() const
 	{
 		if(ObjectBeingCustomized.IsValid())
 		{
-			if(FRigControlElement* ControlElement = Cast<FRigControlElement>(ObjectBeingCustomized->GetContent<FRigBaseElement>()))
+			if(ObjectBeingCustomized->IsChildOf<FRigControlElement>())
 			{
-				if(ControlElement->Settings.bShapeEnabled)
+				const FRigControlElement ControlElement = ObjectBeingCustomized->GetContent<FRigControlElement>();
+				if(ControlElement.Settings.bShapeEnabled)
 				{
 					return true;
 				}
@@ -1956,14 +2020,16 @@ bool FRigControlElementDetails::IsEnabled(ERigControlValueType InValueType) cons
 	{
 		if(ObjectBeingCustomized.IsValid())
 		{
-			if(FRigControlElement* ControlElement = Cast<FRigControlElement>(ObjectBeingCustomized->GetContent<FRigBaseElement>()))
+			if(ObjectBeingCustomized->IsChildOf<FRigControlElement>())
 			{
+				const FRigControlElement ControlElement = ObjectBeingCustomized->GetContent<FRigControlElement>();
+				
 				switch (InValueType)
 				{
 					case ERigControlValueType::Minimum:
 					case ERigControlValueType::Maximum:
 					{
-						if(ControlElement->Settings.bLimitTranslation || ControlElement->Settings.bLimitRotation || ControlElement->Settings.bLimitScale)
+						if(ControlElement.Settings.bLimitTranslation || ControlElement.Settings.bLimitRotation || ControlElement.Settings.bLimitScale)
 						{
 							return true;
 						}
@@ -1998,14 +2064,17 @@ FText FRigControlElementDetails::GetDisplayName() const
 		TWeakObjectPtr<UDetailsViewWrapperObject> ObjectBeingCustomized = ObjectsBeingCustomized[ObjectIndex];
 		if(ObjectBeingCustomized.IsValid())
 		{
-			const FRigControlElement* ControlElement = ObjectBeingCustomized->GetContent<FRigControlElement>();
-			if(ObjectIndex == 0)
-			{
-				DisplayName = ControlElement->Settings.DisplayName;
-			}
-			else if(DisplayName != ControlElement->Settings.DisplayName)
+			if(ObjectBeingCustomized->IsChildOf<FRigControlElement>())
 			{
-				return ControlRigDetailsMultipleValues;
+				const FRigControlElement ControlElement = ObjectBeingCustomized->GetContent<FRigControlElement>();
+				if(ObjectIndex == 0)
+				{
+					DisplayName = ControlElement.Settings.DisplayName;
+				}
+				else if(DisplayName != ControlElement.Settings.DisplayName)
+				{
+					return ControlRigDetailsMultipleValues;
+				}
 			}
 		}
 	}
@@ -2031,17 +2100,25 @@ void FRigControlElementDetails::SetDisplayName(const FText& InNewText, ETextComm
 		TWeakObjectPtr<UDetailsViewWrapperObject> ObjectBeingCustomized = ObjectsBeingCustomized[ObjectIndex];
 		if(ObjectBeingCustomized.IsValid())
 		{
-			FRigControlElement* ControlElement = ObjectBeingCustomized->GetContent<FRigControlElement>();
-
-			FRigControlSettings Settings = ControlElement->Settings;
-			Settings.DisplayName = DisplayName;
-
-			if(GetHierarchy())
+			if(ObjectBeingCustomized->IsChildOf<FRigControlElement>())
 			{
-				GetHierarchy()->SetControlSettings(ControlElement->GetKey(), Settings, false, false, true);
+				FRigControlElement ControlElement = ObjectBeingCustomized->GetContent<FRigControlElement>();
+				ControlElement.Settings.DisplayName = DisplayName;
+				ObjectBeingCustomized->SetContent<FRigControlElement>(ControlElement);
+
+				if(GetHierarchy())
+				{
+					GetHierarchy()->SetControlSettings(ControlElement.GetKey(), ControlElement.Settings, false, false, true);
+				}
 			}
 		}
 	}
 }
 
-#undef LOCTEXT_NAMESPACE
\ No newline at end of file
+void FRigNullElementDetails::CustomizeDetails(IDetailLayoutBuilder& DetailBuilder)
+{
+	FRigTransformElementDetails::CustomizeDetails(DetailBuilder);
+	CustomizeTransform(DetailBuilder);
+}
+
+#undef LOCTEXT_NAMESPACE
diff --git a/Engine/Plugins/Animation/ControlRig/Source/ControlRigEditor/Private/ControlRigElementDetails.h b/Engine/Plugins/Animation/ControlRig/Source/ControlRigEditor/Private/ControlRigElementDetails.h
--- a/Engine/Plugins/Animation/ControlRig/Source/ControlRigEditor/Private/ControlRigElementDetails.h
+++ b/Engine/Plugins/Animation/ControlRig/Source/ControlRigEditor/Private/ControlRigElementDetails.h
@@ -120,13 +120,12 @@ protected:
 	void OnTransformChanged(FEditPropertyChain* InPropertyChain);
 };
 
-class FRigBaseElementDetails : public IPropertyTypeCustomization
+class FRigBaseElementDetails : public IDetailCustomization
 {
 public:
 
-	/** IPropertyTypeCustomization interface */
-	virtual void CustomizeHeader(TSharedRef<class IPropertyHandle> InStructPropertyHandle, class FDetailWidgetRow& HeaderRow, IPropertyTypeCustomizationUtils& StructCustomizationUtils) override;
-	virtual void CustomizeChildren(TSharedRef<class IPropertyHandle> InStructPropertyHandle, class IDetailChildrenBuilder& StructBuilder, IPropertyTypeCustomizationUtils& StructCustomizationUtils) override;
+	/** IDetailCustomization interface */
+	virtual void CustomizeDetails(IDetailLayoutBuilder& DetailBuilder) override;
 
 	URigHierarchy* GetHierarchy() const { return HierarchyBeingCustomized; }
 	FRigElementKey GetElementKey() const;
@@ -140,14 +139,14 @@ public:
 	TArray<FRigElementKey> GetElementKeys() const;
 
 	template<typename T>
-	TArray<T*> GetElementsInDetailsView() const
+	TArray<T> GetElementsInDetailsView() const
 	{
-		TArray<T*> Elements;
+		TArray<T> Elements;
 		for(TWeakObjectPtr<UDetailsViewWrapperObject> ObjectBeingCustomized : ObjectsBeingCustomized)
 		{
 			if(ObjectBeingCustomized.IsValid())
 			{
-				Elements.Add(Cast<T>(ObjectBeingCustomized->GetContent<FRigBaseElement>()));
+				Elements.Add(ObjectBeingCustomized->GetContent<T>());
 			}
 		}
 		return Elements;
@@ -163,7 +162,7 @@ public:
 			{
 				if(ObjectBeingCustomized.IsValid())
 				{
-					Elements.Add(HierarchyBeingCustomized->Find<T>(ObjectBeingCustomized->GetContent<FRigBaseElement>()->GetKey()));
+					Elements.Add(HierarchyBeingCustomized->Find<T>(ObjectBeingCustomized->GetContent<FRigBaseElement>().GetKey()));
 				}
 			}
 		}
@@ -175,6 +174,9 @@ public:
 	bool IsAnyControlOfType(ERigControlType InType) const;
 	bool IsAnyControlNotOfType(ERigControlType InType) const;
 
+	static void RegisterSectionMappings(FPropertyEditorModule& PropertyEditorModule);
+	virtual void RegisterSectionMappings(FPropertyEditorModule& PropertyEditorModule, UClass* InClass);
+
 protected:
 
 	UControlRigBlueprint* BlueprintBeingCustomized;
@@ -186,8 +188,13 @@ class FRigTransformElementDetails : public FRigBaseElementDetails
 {
 public:
 
-	/** IPropertyTypeCustomization interface */
-	virtual void CustomizeChildren(TSharedRef<class IPropertyHandle> InStructPropertyHandle, class IDetailChildrenBuilder& StructBuilder, IPropertyTypeCustomizationUtils& StructCustomizationUtils) override;
+	/** IDetailCustomization interface */
+	virtual void CustomizeDetails(IDetailLayoutBuilder& DetailBuilder) override;
+
+	/** FRigBaseElementDetails interface */
+	virtual void RegisterSectionMappings(FPropertyEditorModule& PropertyEditorModule, UClass* InClass) override;
+
+	virtual void CustomizeTransform(IDetailLayoutBuilder& DetailBuilder);
 
 protected:
 
@@ -204,10 +211,13 @@ class FRigBoneElementDetails : public FRigTransformElementDetails
 public:
 
 	/** Makes a new instance of this detail layout class for a specific detail view requesting it */
-	static TSharedRef<IPropertyTypeCustomization> MakeInstance()
+	static TSharedRef<IDetailCustomization> MakeInstance()
 	{
 		return MakeShareable(new FRigBoneElementDetails);
 	}
+
+	/** IDetailCustomization interface */
+	virtual void CustomizeDetails(IDetailLayoutBuilder& DetailBuilder) override;
 };
 
 class FRigControlElementDetails : public FRigTransformElementDetails
@@ -215,13 +225,16 @@ class FRigControlElementDetails : public FRigTransformElementDetails
 public:
 
 	// Makes a new instance of this detail layout class for a specific detail view requesting it
-	static TSharedRef<IPropertyTypeCustomization> MakeInstance()
+	static TSharedRef<IDetailCustomization> MakeInstance()
 	{
 		return MakeShareable(new FRigControlElementDetails);
 	}
 
-	/** IPropertyTypeCustomization interface */
-	virtual void CustomizeChildren(TSharedRef<class IPropertyHandle> InStructPropertyHandle, class IDetailChildrenBuilder& StructBuilder, IPropertyTypeCustomizationUtils& StructCustomizationUtils) override;
+	/** IDetailCustomization interface */
+	virtual void CustomizeDetails(IDetailLayoutBuilder& DetailBuilder) override;
+
+	/** FRigBaseElementDetails interface */
+	virtual void RegisterSectionMappings(FPropertyEditorModule& PropertyEditorModule, UClass* InClass) override;
 
 	bool IsShapeEnabled() const;
 	bool IsEnabled(ERigControlValueType InValueType) const;
@@ -244,8 +257,11 @@ class FRigNullElementDetails : public FRigTransformElementDetails
 public:
 
 	// Makes a new instance of this detail layout class for a specific detail view requesting it
-	static TSharedRef<IPropertyTypeCustomization> MakeInstance()
+	static TSharedRef<IDetailCustomization> MakeInstance()
 	{
 		return MakeShareable(new FRigNullElementDetails);
 	}
+
+	/** IDetailCustomization interface */
+	virtual void CustomizeDetails(IDetailLayoutBuilder& DetailBuilder) override;
 };
diff --git a/Engine/Plugins/Animation/ControlRig/Source/ControlRigEditor/Private/ControlRigGraphDetails.cpp b/Engine/Plugins/Animation/ControlRig/Source/ControlRigEditor/Private/ControlRigGraphDetails.cpp
--- a/Engine/Plugins/Animation/ControlRig/Source/ControlRigEditor/Private/ControlRigGraphDetails.cpp
+++ b/Engine/Plugins/Animation/ControlRig/Source/ControlRigEditor/Private/ControlRigGraphDetails.cpp
@@ -1179,6 +1179,11 @@ void FControlRigWrappedNodeDetails::CustomizeDetails(IDetailLayoutBuilder& Detai
 
 	TArray<FName> Categories;
 	DetailLayout.GetCategoryNames(Categories);
+
+	if(Categories.IsEmpty())
+	{
+		return;
+	}
 	
 	IDetailCategoryBuilder& DefaultsCategory = DetailLayout.EditCategory(Categories[0]);
 
@@ -1188,7 +1193,6 @@ void FControlRigWrappedNodeDetails::CustomizeDetails(IDetailLayoutBuilder& Detai
 		return;
 	}
 
-	TSharedPtr<IPropertyHandle> StoredStructHandle = DetailLayout.GetProperty(TEXT("StoredStruct"));
 	for(URigVMPin* Pin : FirstNode->GetPins())
 	{
 		bool bHasAnyInputLink = false;
@@ -1204,7 +1208,7 @@ void FControlRigWrappedNodeDetails::CustomizeDetails(IDetailLayoutBuilder& Detai
 			}
 		}
 
-		TSharedPtr<IPropertyHandle> PinHandle = StoredStructHandle->GetChildHandle(Pin->GetFName());
+		TSharedPtr<IPropertyHandle> PinHandle = DetailLayout.GetProperty(Pin->GetFName());
 		if(PinHandle.IsValid())
 		{
 			static const TCHAR DisabledFormat[] = TEXT("%s\n\nNote: Editing disabled since pin has an input link.");
@@ -1233,8 +1237,6 @@ void FControlRigWrappedNodeDetails::CustomizeDetails(IDetailLayoutBuilder& Detai
 		}
 	}
 
-	DetailLayout.HideProperty(StoredStructHandle);
-
 	if(Objects.Num() > 1)
 	{
 		return;
@@ -1433,7 +1435,7 @@ void FControlRigWrappedNodeDetails::CustomizeDetails(IDetailLayoutBuilder& Detai
 
 		if(!bAddedProperty)
 		{
-			TSharedPtr<IPropertyHandle> PinHandle = StoredStructHandle->GetChildHandle(Pin->GetFName());
+			TSharedPtr<IPropertyHandle> PinHandle = DetailLayout.GetProperty(Pin->GetFName());
 			if(PinHandle.IsValid())
 			{
 				DebugCategory.AddProperty(PinHandle)
diff --git a/Engine/Plugins/Animation/ControlRig/Source/ControlRigEditor/Private/ControlRigLocalVariableDetails.cpp b/Engine/Plugins/Animation/ControlRig/Source/ControlRigEditor/Private/ControlRigLocalVariableDetails.cpp
--- a/Engine/Plugins/Animation/ControlRig/Source/ControlRigEditor/Private/ControlRigLocalVariableDetails.cpp
+++ b/Engine/Plugins/Animation/ControlRig/Source/ControlRigEditor/Private/ControlRigLocalVariableDetails.cpp
@@ -16,37 +16,34 @@
 
 #define LOCTEXT_NAMESPACE "LocalVariableDetails"
 
-void FRigVMLocalVariableDetails::CustomizeHeader(TSharedRef<IPropertyHandle> InStructPropertyHandle, FDetailWidgetRow& HeaderRow, IPropertyTypeCustomizationUtils& StructCustomizationUtils)
+void FRigVMLocalVariableDetails::CustomizeDetails(IDetailLayoutBuilder& DetailBuilder)
 {
 	ObjectsBeingCustomized.Reset();
-	
-	TArray<UObject*> Objects;
-	InStructPropertyHandle->GetOuterObjects(Objects);
-	for (UObject* Object : Objects)
+
+	TArray<TWeakObjectPtr<UObject>> DetailObjects;
+	DetailBuilder.GetObjectsBeingCustomized(DetailObjects);
+	for(TWeakObjectPtr<UObject> DetailObject : DetailObjects)
 	{
-		UDetailsViewWrapperObject* WrapperObject = CastChecked<UDetailsViewWrapperObject>(Object);
+		UDetailsViewWrapperObject* WrapperObject = CastChecked<UDetailsViewWrapperObject>(DetailObject.Get());
 		ObjectsBeingCustomized.Add(WrapperObject);
 	}
 
 	if(ObjectsBeingCustomized[0].IsValid())
 	{
-		VariableDescription = *ObjectsBeingCustomized[0]->GetContent<FRigVMGraphVariableDescription>();
+		VariableDescription = ObjectsBeingCustomized[0]->GetContent<FRigVMGraphVariableDescription>();
 		GraphBeingCustomized = ObjectsBeingCustomized[0]->GetTypedOuter<URigVMGraph>();
 		BlueprintBeingCustomized = GraphBeingCustomized->GetTypedOuter<UControlRigBlueprint>();
 
 		NameValidator = FControlRigLocalVariableNameValidator(BlueprintBeingCustomized, GraphBeingCustomized, VariableDescription.Name);
 	}
-}
 
-void FRigVMLocalVariableDetails::CustomizeChildren(TSharedRef<IPropertyHandle> InStructPropertyHandle, IDetailChildrenBuilder& StructBuilder, IPropertyTypeCustomizationUtils& StructCustomizationUtils)
-{
-	StructBuilder.GetParentCategory().GetParentLayout().HideCategory("RigVMGraphVariableDescription");
-	IDetailCategoryBuilder& Category = StructBuilder.GetParentCategory().GetParentLayout().EditCategory("Local Variable");
+	DetailBuilder.HideCategory(TEXT("RigVMGraphVariableDescription"));
+	IDetailCategoryBuilder& Category = DetailBuilder.EditCategory(TEXT("Local Variable"));
 		
-	NameHandle = InStructPropertyHandle->GetChildHandle(TEXT("Name"));
-	TypeHandle = InStructPropertyHandle->GetChildHandle(TEXT("CPPType"));
-	TypeObjectHandle = InStructPropertyHandle->GetChildHandle(TEXT("CPPTypeObject"));
-	DefaultValueHandle = InStructPropertyHandle->GetChildHandle(TEXT("DefaultValue"));
+	NameHandle = DetailBuilder.GetProperty(TEXT("Name"));
+	TypeHandle = DetailBuilder.GetProperty(TEXT("CPPType"));
+	TypeObjectHandle = DetailBuilder.GetProperty(TEXT("CPPTypeObject"));
+	DefaultValueHandle = DetailBuilder.GetProperty(TEXT("DefaultValue"));
 	
 	const UEdGraphSchema* Schema = GetDefault<UControlRigGraphSchema>();
 
@@ -103,7 +100,7 @@ void FRigVMLocalVariableDetails::CustomizeChildren(TSharedRef<IPropertyHandle> I
 			FProperty* Property = LiteralMemory->FindPropertyByName(*SourcePath);
 			if (Property)
 			{
-				IDetailCategoryBuilder& DefaultValueCategory = StructBuilder.GetParentCategory().GetParentLayout().EditCategory(TEXT("DefaultValueCategory"), LOCTEXT("DefaultValueCategoryHeading", "Default Value"));
+				IDetailCategoryBuilder& DefaultValueCategory = DetailBuilder.EditCategory(TEXT("DefaultValueCategory"), LOCTEXT("DefaultValueCategoryHeading", "Default Value"));
 				Property->ClearPropertyFlags(CPF_EditConst);
 			
 				const FName SanitizedName = FRigVMPropertyDescription::SanitizeName(*SourcePath);
diff --git a/Engine/Plugins/Animation/ControlRig/Source/ControlRigEditor/Private/ControlRigLocalVariableDetails.h b/Engine/Plugins/Animation/ControlRig/Source/ControlRigEditor/Private/ControlRigLocalVariableDetails.h
--- a/Engine/Plugins/Animation/ControlRig/Source/ControlRigEditor/Private/ControlRigLocalVariableDetails.h
+++ b/Engine/Plugins/Animation/ControlRig/Source/ControlRigEditor/Private/ControlRigLocalVariableDetails.h
@@ -17,7 +17,7 @@
 
 class IPropertyHandle;
 
-class FRigVMLocalVariableDetails : public IPropertyTypeCustomization
+class FRigVMLocalVariableDetails : public IDetailCustomization
 {
 	FRigVMLocalVariableDetails()
 	: GraphBeingCustomized(nullptr)
@@ -29,15 +29,13 @@ class FRigVMLocalVariableDetails : public IPropertyTypeCustomization
 public:
 
 	// Makes a new instance of this detail layout class for a specific detail view requesting it
-	static TSharedRef<IPropertyTypeCustomization> MakeInstance()
+	static TSharedRef<IDetailCustomization> MakeInstance()
 	{
 		return MakeShareable(new FRigVMLocalVariableDetails);
 	}
 
-	/** IPropertyTypeCustomization interface */
-	virtual void CustomizeHeader(TSharedRef<class IPropertyHandle> InStructPropertyHandle, class FDetailWidgetRow& HeaderRow, IPropertyTypeCustomizationUtils& StructCustomizationUtils) override;
-	virtual void CustomizeChildren(TSharedRef<class IPropertyHandle> InStructPropertyHandle, class IDetailChildrenBuilder& StructBuilder, IPropertyTypeCustomizationUtils& StructCustomizationUtils) override;
-
+	/** IDetailCustomization interface */
+	virtual void CustomizeDetails(IDetailLayoutBuilder& DetailBuilder) override;
 
 private:
 
diff --git a/Engine/Plugins/Animation/ControlRig/Source/ControlRigEditor/Private/Editor/ControlRigEditor.cpp b/Engine/Plugins/Animation/ControlRig/Source/ControlRigEditor/Private/Editor/ControlRigEditor.cpp
--- a/Engine/Plugins/Animation/ControlRig/Source/ControlRigEditor/Private/Editor/ControlRigEditor.cpp
+++ b/Engine/Plugins/Animation/ControlRig/Source/ControlRigEditor/Private/Editor/ControlRigEditor.cpp
@@ -1599,7 +1599,7 @@ void FControlRigEditor::SetDetailViewForLocalVariable()
 		{
 			if(UDetailsViewWrapperObject* WrapperObject = Cast<UDetailsViewWrapperObject>(SelectedObject.Get()))
 			{
-				VariableName = WrapperObject->GetContent<FRigVMGraphVariableDescription>()->Name;
+				VariableName = WrapperObject->GetContent<FRigVMGraphVariableDescription>().Name;
 				break;
 			}
 		}
@@ -1669,7 +1669,7 @@ bool FControlRigEditor::DetailViewShowsRigElement(FRigElementKey InKey) const
 			{
 				if(WrapperObject->GetWrappedStruct()->IsChildOf(FRigBaseElement::StaticStruct()))
 				{
-					if(WrapperObject->GetContent<FRigBaseElement>()->GetKey() == InKey)
+					if(WrapperObject->GetContent<FRigBaseElement>().GetKey() == InKey)
 					{
 						return true;
 					}
@@ -2436,7 +2436,7 @@ void FControlRigEditor::EnsureValidRigElementsInDetailPanel()
 			{
 				if(WrapperObject->GetWrappedStruct()->IsChildOf(FRigBaseElement::StaticStruct()))
 				{
-					FRigElementKey Key = WrapperObject->GetContent<FRigBaseElement>()->GetKey();
+					FRigElementKey Key = WrapperObject->GetContent<FRigBaseElement>().GetKey();
 					if(!Hierarchy->Contains(Key))
 					{
 						ClearDetailObject();
@@ -2671,10 +2671,8 @@ void FControlRigEditor::HandleModifiedEvent(ERigVMGraphNotifType InNotifType, UR
 							{
 								if(WrapperObject->GetOuter() == Pin->GetNode())
 								{
-									uint8* PropertyStorage = WrapperObject->GetStructProperty()->ContainerPtrToValuePtr<uint8>(WrapperObject);
-									UScriptStruct* UnitStruct = WrapperObject->GetWrappedStruct();
-									FProperty* TargetProperty = UnitStruct->FindPropertyByName(RootPin->GetFName());
-									PropertyStorage += TargetProperty->GetOffset_ReplaceWith_ContainerPtrToValuePtr();
+									const FProperty* TargetProperty = WrapperObject->GetClass()->FindPropertyByName(RootPin->GetFName());
+									uint8* PropertyStorage = TargetProperty->ContainerPtrToValuePtr<uint8>(WrapperObject);
 									TargetProperty->ImportText(*DefaultValue, PropertyStorage, PPF_None, nullptr);
 								}
 							}
@@ -2789,7 +2787,7 @@ void FControlRigEditor::HandleControlRigExecutedEvent(UControlRig* InControlRig,
 					UScriptStruct* Struct = WrapperObject->GetWrappedStruct(); 
 					if(Struct->IsChildOf(FRigBaseElement::StaticStruct()))
 					{
-						FRigElementKey Key = WrapperObject->GetContent<FRigBaseElement>()->GetKey();
+						const FRigElementKey Key = WrapperObject->GetContent<FRigBaseElement>().GetKey();
 
 						FRigBaseElement* Element = Hierarchy->Find(Key);
 						if(FRigTransformElement* TransformElement = Cast<FRigTransformElement>(Element))
@@ -2814,7 +2812,7 @@ void FControlRigEditor::HandleControlRigExecutedEvent(UControlRig* InControlRig,
 							Hierarchy->GetControlShapeTransform(ControlElement, ERigTransformType::InitialLocal);
 						}
 
-						Struct->CopyScriptStruct(WrapperObject->GetContent<FRigBaseElement>(), Element);
+						WrapperObject->SetContent<FRigBaseElement>(*Element);
 					}
 				}
 			}
@@ -4389,8 +4387,6 @@ void FControlRigEditor::OnWrappedPropertyChangedChainEvent(UDetailsViewWrapperOb
 
 	TGuardValue<bool> SuspendDetailsPanelRefresh(bSuspendDetailsPanelRefresh, true);
 
-	const FString StoredStructString = TEXT("StoredStruct->");
-
 	UControlRigBlueprint* ControlRigBP = GetControlRigBlueprint();
 
 	FString PropertyPath = InPropertyPath;
@@ -4399,15 +4395,13 @@ void FControlRigEditor::OnWrappedPropertyChangedChainEvent(UDetailsViewWrapperOb
 		check(InWrapperObject->GetWrappedStruct() == WrapperObjects[0]->GetWrappedStruct());
 
 		URigHierarchy* Hierarchy = CastChecked<URigHierarchy>(InWrapperObject->GetOuter());
-		FRigBaseElement* WrappedElement = InWrapperObject->GetContent<FRigBaseElement>();
-		FRigBaseElement* FirstWrappedElement = WrapperObjects[0]->GetContent<FRigBaseElement>();
-		FRigElementKey Key = WrappedElement->GetKey();
+		const FRigBaseElement WrappedElement = InWrapperObject->GetContent<FRigBaseElement>();
+		const FRigBaseElement FirstWrappedElement = WrapperObjects[0]->GetContent<FRigBaseElement>();
+		const FRigElementKey& Key = WrappedElement.GetKey();
 		if(!Hierarchy->Contains(Key))
 		{
 			return;
 		}
-		
-		check(PropertyPath.RemoveFromStart(StoredStructString));
 
 		const FString PoseString = TEXT("Pose->");
 		const FString OffsetString = TEXT("Offset->");
@@ -4460,13 +4454,13 @@ void FControlRigEditor::OnWrappedPropertyChangedChainEvent(UDetailsViewWrapperOb
 				Hierarchy = ControlRigBP->Hierarchy;
 			}
 
-			FRigTransformElement* TransformElement = Hierarchy->Find<FRigTransformElement>(WrappedElement->GetKey());
+			FRigTransformElement* TransformElement = Hierarchy->Find<FRigTransformElement>(WrappedElement.GetKey());
 			if(TransformElement == nullptr)
 			{
 				return;
 			}
-			
-			const FTransform Transform = CastChecked<FRigTransformElement>(FirstWrappedElement)->Pose.Get(TransformType);
+
+			const FTransform Transform = InWrapperObject->GetContent<FRigTransformElement>().Pose.Get(TransformType);
 
 			if(ERigTransformType::IsLocal(TransformType) && TransformElement->IsA<FRigControlElement>())
 			{
@@ -4488,7 +4482,7 @@ void FControlRigEditor::OnWrappedPropertyChangedChainEvent(UDetailsViewWrapperOb
 		}
 		else if(PropertyPath.RemoveFromStart(OffsetString))
 		{
-			FRigControlElement* ControlElement = ControlRigBP->Hierarchy->Find<FRigControlElement>(WrappedElement->GetKey());
+			FRigControlElement* ControlElement = ControlRigBP->Hierarchy->Find<FRigControlElement>(WrappedElement.GetKey());
 			if(ControlElement == nullptr)
 			{
 				return;
@@ -4496,13 +4490,14 @@ void FControlRigEditor::OnWrappedPropertyChangedChainEvent(UDetailsViewWrapperOb
 			
 			ERigTransformType::Type TransformType = Local::GetTransformTypeFromPath(PropertyPath);
 			bIsInitial = bIsInitial || ERigTransformType::IsInitial(TransformType);
-			const FTransform Transform = CastChecked<FRigControlElement>(FirstWrappedElement)->Offset.Get(TransformType);
+
+			const FTransform Transform = WrapperObjects[0]->GetContent<FRigControlElement>().Offset.Get(TransformType);
 			
 			ControlRigBP->Hierarchy->SetControlOffsetTransform(ControlElement, Transform, ERigTransformType::MakeInitial(TransformType), true, true, false, true);
 		}
 		else if(PropertyPath.RemoveFromStart(ShapeString))
 		{
-			FRigControlElement* ControlElement = ControlRigBP->Hierarchy->Find<FRigControlElement>(WrappedElement->GetKey());
+			FRigControlElement* ControlElement = ControlRigBP->Hierarchy->Find<FRigControlElement>(WrappedElement.GetKey());
 			if(ControlElement == nullptr)
 			{
 				return;
@@ -4510,15 +4505,16 @@ void FControlRigEditor::OnWrappedPropertyChangedChainEvent(UDetailsViewWrapperOb
 
 			ERigTransformType::Type TransformType = Local::GetTransformTypeFromPath(PropertyPath);
 			bIsInitial = bIsInitial || ERigTransformType::IsInitial(TransformType);
-			const FTransform Transform = CastChecked<FRigControlElement>(FirstWrappedElement)->Shape.Get(TransformType);
+
+			const FTransform Transform = WrapperObjects[0]->GetContent<FRigControlElement>().Shape.Get(TransformType);
 			
 			ControlRigBP->Hierarchy->SetControlShapeTransform(ControlElement, Transform, ERigTransformType::MakeInitial(TransformType), true, false, true);
 		}
 		else if(PropertyPath.RemoveFromStart(SettingsString))
 		{
-			const FRigControlSettings Settings = CastChecked<FRigControlElement>(WrappedElement)->Settings;
+			const FRigControlSettings Settings  = InWrapperObject->GetContent<FRigControlElement>().Settings;
 
-			FRigControlElement* ControlElement = ControlRigBP->Hierarchy->Find<FRigControlElement>(WrappedElement->GetKey());
+			FRigControlElement* ControlElement = ControlRigBP->Hierarchy->Find<FRigControlElement>(WrappedElement.GetKey());
 			if(ControlElement == nullptr)
 			{
 				return;
@@ -4538,14 +4534,12 @@ void FControlRigEditor::OnWrappedPropertyChangedChainEvent(UDetailsViewWrapperOb
 	{
 		check(InWrapperObject->GetWrappedStruct() == WrapperObjects[0]->GetWrappedStruct());
 		
-		uint8* FirstPropertyStorage = WrapperObjects[0]->GetStructProperty()->ContainerPtrToValuePtr<uint8>(WrapperObjects[0].Get());
-		UScriptStruct* UnitStruct = InWrapperObject->GetWrappedStruct();
-		FName RootPinName = InPropertyChangedChainEvent.PropertyChain.GetHead()->GetNextNode()->GetValue()->GetFName();
-		FProperty* TargetProperty = UnitStruct->FindPropertyByName(RootPinName);
+		FName RootPinName = InPropertyChangedChainEvent.PropertyChain.GetHead()->GetValue()->GetFName();
+		FProperty* TargetProperty = WrapperObjects[0]->GetClass()->FindPropertyByName(RootPinName);
+		uint8* FirstPropertyStorage = TargetProperty->ContainerPtrToValuePtr<uint8>(WrapperObjects[0].Get());
 
 		URigVMNode* Node = CastChecked<URigVMNode>(InWrapperObject->GetOuter());
 
-		FirstPropertyStorage += TargetProperty->GetOffset_ReplaceWith_ContainerPtrToValuePtr();
 		FString DefaultValue = FRigVMStruct::ExportToFullyQualifiedText(TargetProperty, FirstPropertyStorage);
 
 		if (!DefaultValue.IsEmpty())
@@ -4557,9 +4551,8 @@ void FControlRigEditor::OnWrappedPropertyChangedChainEvent(UDetailsViewWrapperOb
 	else if(InWrapperObject->GetWrappedStruct()->IsChildOf(FRigVMGraphVariableDescription::StaticStruct()))
 	{
 		check(InWrapperObject->GetWrappedStruct() == WrapperObjects[0]->GetWrappedStruct());
-		check(PropertyPath.RemoveFromStart(StoredStructString));
 		
-		FRigVMGraphVariableDescription* VariableDescription = InWrapperObject->GetContent<FRigVMGraphVariableDescription>();
+		const FRigVMGraphVariableDescription VariableDescription = InWrapperObject->GetContent<FRigVMGraphVariableDescription>();
 		URigVMGraph* Graph = CastChecked<URigVMGraph>(InWrapperObject->GetOuter());
 		URigVMController* Controller = ControlRigBP->GetController(Graph);
 		if (PropertyPath == TEXT("Name") && MyBlueprintWidget.IsValid())
@@ -4573,7 +4566,7 @@ void FControlRigEditor::OnWrappedPropertyChangedChainEvent(UDetailsViewWrapperOb
 					{
 						if (Variable.Name == OldVariableName)
 						{
-							Controller->RenameLocalVariable(OldVariableName, VariableDescription->Name);
+							Controller->RenameLocalVariable(OldVariableName, VariableDescription.Name);
 							break;
 						}
 					}
@@ -4586,9 +4579,9 @@ void FControlRigEditor::OnWrappedPropertyChangedChainEvent(UDetailsViewWrapperOb
 		{			
 			for (FRigVMGraphVariableDescription& Variable : Graph->GetLocalVariables())
 			{
-				if (Variable.Name == VariableDescription->Name)
+				if (Variable.Name == VariableDescription.Name)
 				{
-					Controller->SetLocalVariableType(Variable.Name, VariableDescription->CPPType, VariableDescription->CPPTypeObject);
+					Controller->SetLocalVariableType(Variable.Name, VariableDescription.CPPType, VariableDescription.CPPTypeObject);
 					break;
 				}
 			}
@@ -4598,9 +4591,9 @@ void FControlRigEditor::OnWrappedPropertyChangedChainEvent(UDetailsViewWrapperOb
 		{			
 			for (FRigVMGraphVariableDescription& Variable : Graph->GetLocalVariables())
 			{
-				if (Variable.Name == VariableDescription->Name)
+				if (Variable.Name == VariableDescription.Name)
 				{
-					Controller->SetLocalVariableDefaultValue(Variable.Name, VariableDescription->DefaultValue, true, true, false);
+					Controller->SetLocalVariableDefaultValue(Variable.Name, VariableDescription.DefaultValue, true, true, false);
 					break;
 				}
 			}
diff --git a/Engine/Plugins/Animation/ControlRig/Source/ControlRigEditor/Private/Editor/DetailsViewWrapperObject.cpp b/Engine/Plugins/Animation/ControlRig/Source/ControlRigEditor/Private/Editor/DetailsViewWrapperObject.cpp
--- a/Engine/Plugins/Animation/ControlRig/Source/ControlRigEditor/Private/Editor/DetailsViewWrapperObject.cpp
+++ b/Engine/Plugins/Animation/ControlRig/Source/ControlRigEditor/Private/Editor/DetailsViewWrapperObject.cpp
@@ -3,6 +3,7 @@
 #include "DetailsViewWrapperObject.h"
 #include "Units/RigUnit.h"
 #include "ControlRigElementDetails.h"
+#include "ControlRigLocalVariableDetails.h"
 #include "Modules/ModuleManager.h"
 
 #if WITH_EDITOR
@@ -27,7 +28,7 @@ UClass* UDetailsViewWrapperObject::GetClassForStruct(UScriptStruct* InStruct, bo
 		return nullptr;
 	}
 
-	UClass *SuperClass = UDetailsViewWrapperObject::StaticClass();
+	UClass* SuperClass = UDetailsViewWrapperObject::StaticClass();
 	const FName WrapperClassName(FString::Printf(TEXT("%s_WrapperObject"), *InStruct->GetStructCPPName()));
 
 	UClass* WrapperClass = NewObject<UClass>(
@@ -75,62 +76,76 @@ UClass* UDetailsViewWrapperObject::GetClassForStruct(UScriptStruct* InStruct, bo
 		}
 	};
 
-	FStructProperty* Property = new FStructProperty(WrapperClass, TEXT("StoredStruct"), RF_Public);
-	Property->Struct = InStruct;
-	if (Local::IsStructHashable(InStruct))
-	{
-		Property->SetPropertyFlags(CPF_HasGetValueTypeHash);
-	}
-	
-	// Make sure the variables show up in the details panel
-	Property->SetPropertyFlags(CPF_Edit);
-#if UE_RIGVM_UCLASS_BASED_STORAGE_DISABLED
-	Property->SetMetaData(TEXT("ShowOnlyInnerProperties"), TEXT("true"));
-#else
-	// for non-nodes - show the properties directly under the details panel main category 
-	if(!InStruct->IsChildOf(FRigUnit::StaticStruct()))
-	{
-		Property->SetMetaData(TEXT("ShowOnlyInnerProperties"), TEXT("true"));
-	}
-#endif
+	// duplicate all properties from the struct to the wrapper object
+	FField** LinkToProperty = &WrapperClass->ChildProperties;
 
-	// For rig units
-	if(InStruct->IsChildOf(FRigUnit::StaticStruct()))
+	for (TFieldIterator<FProperty> PropertyIt(InStruct); PropertyIt; ++PropertyIt)
 	{
-		// mark all input properties with editanywhere
-		for (TFieldIterator<FProperty> PropertyIt(InStruct); PropertyIt; ++PropertyIt)
-		{
-			FProperty* ChildProperty = *PropertyIt;
-			if (!ChildProperty->HasMetaData(TEXT("Input")) && !ChildProperty->HasMetaData(TEXT("Visible")))
-			{
-				continue;
-			}
+		const FProperty* InProperty = *PropertyIt;
+		FProperty* NewProperty = CastFieldChecked<FProperty>(FField::Duplicate(InProperty, WrapperClass, InProperty->GetFName()));
+		check(NewProperty);
+		FField::CopyMetaData(InProperty, NewProperty);
 
+		if (NewProperty->HasMetaData(TEXT("Input")) || NewProperty->HasMetaData(TEXT("Visible")))
+		{
 			// filter out execute pins
-			if (FStructProperty* StructProperty = CastField<FStructProperty>(ChildProperty))
+			bool bIsEditable = true;
+			if (FStructProperty* StructProperty = CastField<FStructProperty>(NewProperty))
 			{
 				if (StructProperty->Struct->IsChildOf(FRigVMExecuteContext::StaticStruct()))
 				{
-					continue;
+					bIsEditable = false;
 				}
 			}
-			
-			ChildProperty->SetPropertyFlags(ChildProperty->GetPropertyFlags() | CPF_Edit);
+
+			if(bIsEditable)
+			{
+				NewProperty->SetPropertyFlags(NewProperty->GetPropertyFlags() | CPF_Edit);
+			}
 		}
 
+		*LinkToProperty = NewProperty;
+		LinkToProperty = &(*LinkToProperty)->Next;
+	}
+
 #if WITH_EDITOR
+	if(InStruct->IsChildOf(FRigUnit::StaticStruct()))
+	{
 		FPropertyEditorModule& PropertyEditorModule = FModuleManager::LoadModuleChecked<FPropertyEditorModule>("PropertyEditor");
-		if (!PropertyEditorModule.GetClassNameToDetailLayoutNameMap().Contains(InStruct->GetFName()))
+		if (!PropertyEditorModule.GetClassNameToDetailLayoutNameMap().Contains(WrapperClassName))
 		{
-			PropertyEditorModule.RegisterCustomClassLayout(InStruct->GetFName(), FOnGetDetailCustomizationInstance::CreateStatic(&FRigUnitDetails::MakeInstance));
+			PropertyEditorModule.RegisterCustomClassLayout(WrapperClassName, FOnGetDetailCustomizationInstance::CreateStatic(&FRigUnitDetails::MakeInstance));
 		}
-#endif
 	}
-
-	// Given it's only one property we are adding
-	// we can set the head of the linked list directly
-	WrapperClass->ChildProperties = Property;
-
+	else if(InStruct->IsChildOf(FRigBaseElement::StaticStruct()))
+	{
+		FPropertyEditorModule& PropertyEditorModule = FModuleManager::LoadModuleChecked<FPropertyEditorModule>("PropertyEditor");
+		if (!PropertyEditorModule.GetClassNameToDetailLayoutNameMap().Contains(WrapperClassName))
+		{
+			if(InStruct == FRigBoneElement::StaticStruct())
+			{
+				PropertyEditorModule.RegisterCustomClassLayout(WrapperClassName, FOnGetDetailCustomizationInstance::CreateStatic(&FRigBoneElementDetails::MakeInstance));
+			}
+			else if(InStruct == FRigNullElement::StaticStruct())
+			{
+				PropertyEditorModule.RegisterCustomClassLayout(WrapperClassName, FOnGetDetailCustomizationInstance::CreateStatic(&FRigNullElementDetails::MakeInstance));
+			}
+			else if(InStruct == FRigControlElement::StaticStruct())
+			{
+				PropertyEditorModule.RegisterCustomClassLayout(WrapperClassName, FOnGetDetailCustomizationInstance::CreateStatic(&FRigControlElementDetails::MakeInstance));
+			}
+		}
+	}
+	else if(InStruct->IsChildOf(FRigVMGraphVariableDescription::StaticStruct()))
+	{
+		FPropertyEditorModule& PropertyEditorModule = FModuleManager::LoadModuleChecked<FPropertyEditorModule>("PropertyEditor");
+		if (!PropertyEditorModule.GetClassNameToDetailLayoutNameMap().Contains(WrapperClassName))
+		{
+			PropertyEditorModule.RegisterCustomClassLayout(WrapperClassName, FOnGetDetailCustomizationInstance::CreateStatic(&FRigVMLocalVariableDetails::MakeInstance));
+		}
+	}
+#endif
+	
 	// Update the class
 	WrapperClass->Bind();
 	WrapperClass->StaticLink(true);
@@ -144,6 +159,10 @@ UClass* UDetailsViewWrapperObject::GetClassForStruct(UScriptStruct* InStruct, bo
 	UObject* CDO = WrapperClass->GetDefaultObject(true);
 	CDO->AddToRoot();
 
+	// import the defaults from the struct onto the class
+	TSharedPtr<FStructOnScope> DefaultStruct = MakeShareable(new FStructOnScope(InStruct));
+	CopyPropertiesForUnrelatedStructs((uint8*)CDO, WrapperClass, DefaultStruct->GetStructMemory(), DefaultStruct->GetStruct());
+
 	return WrapperClass;
 }
 
@@ -160,43 +179,23 @@ UDetailsViewWrapperObject* UDetailsViewWrapperObject::MakeInstance(UScriptStruct
 	}
 
 	UDetailsViewWrapperObject* Instance = NewObject<UDetailsViewWrapperObject>(InOuter, WrapperClass, NAME_None, RF_Public | RF_Transient | RF_TextExportTransient | RF_DuplicateTransient);
-	Instance->SetContent(InStructMemory, InStruct->GetStructureSize());
+	Instance->SetContent(InStructMemory, InStruct);
 	return Instance;
 }
 
 UScriptStruct* UDetailsViewWrapperObject::GetWrappedStruct() const
 {
-	return GetStructProperty()->Struct;
+	return ClassToStruct.FindChecked(GetClass());
 }
 
-FStructProperty* UDetailsViewWrapperObject::GetStructProperty() const
+void UDetailsViewWrapperObject::SetContent(const uint8* InStructMemory, const UStruct* InStruct)
 {
-	return CastFieldChecked<FStructProperty>(GetClass()->ChildProperties);
+	CopyPropertiesForUnrelatedStructs((uint8*)this, GetClass(), InStructMemory, InStruct);
 }
 
-void UDetailsViewWrapperObject::SetContent(const uint8* InStructMemory, int32 InStructSize)
+void UDetailsViewWrapperObject::GetContent(uint8* OutStructMemory, const UStruct* InStruct) const
 {
-	UScriptStruct* Struct = GetWrappedStruct();
-	check(Struct);
-	check(InStructSize == Struct->GetStructureSize());
-
-	uint8* Content = GetStoredStructPtr();
-	Struct->CopyScriptStruct(Content, InStructMemory);
-}
-
-const uint8* UDetailsViewWrapperObject::GetContent()
-{
-	return GetStoredStructPtr();
-}
-
-void UDetailsViewWrapperObject::GetContent(uint8* OutStructMemory, int32 InStructSize)
-{
-	UScriptStruct* Struct = GetWrappedStruct();
-	check(Struct);
-	check(InStructSize == Struct->GetStructureSize());
-
-	const uint8* Content = GetContent();
-	Struct->CopyScriptStruct(OutStructMemory, Content);
+	CopyPropertiesForUnrelatedStructs(OutStructMemory, InStruct, (const uint8*)this, GetClass());
 }
 
 void UDetailsViewWrapperObject::PostEditChangeChainProperty(FPropertyChangedChainEvent& PropertyChangedEvent)
@@ -224,12 +223,26 @@ void UDetailsViewWrapperObject::PostEditChangeChainProperty(FPropertyChangedChai
 	WrappedPropertyChangedChainEvent.Broadcast(this, PropertyPath, PropertyChangedEvent);
 }
 
-uint8* UDetailsViewWrapperObject::GetStoredStructPtr()
+void UDetailsViewWrapperObject::CopyPropertiesForUnrelatedStructs(uint8* InTargetMemory,
+	const UStruct* InTargetStruct, const uint8* InSourceMemory, const UStruct* InSourceStruct)
 {
-	if(FStructProperty* StructProperty = CastField<FStructProperty>(GetClass()->ChildProperties))
+	check(InTargetMemory);
+	check(InTargetStruct);
+	check(InSourceMemory);
+	check(InSourceStruct);
+	
+	for (TFieldIterator<FProperty> PropertyIt(InTargetStruct); PropertyIt; ++PropertyIt)
 	{
-		return StructProperty->ContainerPtrToValuePtr<uint8>(this);
+		const FProperty* TargetProperty = *PropertyIt;
+		const FProperty* SourceProperty = InSourceStruct->FindPropertyByName(TargetProperty->GetFName());
+		if(SourceProperty == nullptr)
+		{
+			continue;
+		}
+		check(TargetProperty->SameType(SourceProperty));
+
+		uint8* TargetPropertyMemory = TargetProperty->ContainerPtrToValuePtr<uint8>(InTargetMemory);
+		const uint8* SourcePropertyMemory = SourceProperty->ContainerPtrToValuePtr<uint8>(InSourceMemory);
+		TargetProperty->CopyCompleteValue(TargetPropertyMemory, SourcePropertyMemory);
 	}
-	checkNoEntry();
-	return nullptr;
 }
diff --git a/Engine/Plugins/Animation/ControlRig/Source/ControlRigEditor/Private/Editor/DetailsViewWrapperObject.h b/Engine/Plugins/Animation/ControlRig/Source/ControlRigEditor/Private/Editor/DetailsViewWrapperObject.h
--- a/Engine/Plugins/Animation/ControlRig/Source/ControlRigEditor/Private/Editor/DetailsViewWrapperObject.h
+++ b/Engine/Plugins/Animation/ControlRig/Source/ControlRigEditor/Private/Editor/DetailsViewWrapperObject.h
@@ -18,22 +18,37 @@ public:
 	static UClass* GetClassForStruct(UScriptStruct* InStruct, bool bCreateIfNeeded = true);
 	static UDetailsViewWrapperObject* MakeInstance(UScriptStruct* InStruct, uint8* InStructMemory, UObject* InOuter = nullptr);
 	UScriptStruct* GetWrappedStruct() const;
-	FStructProperty* GetStructProperty() const;
 
-	void SetContent(const uint8* InStructMemory, int32 InStructSize);
-	const uint8* GetContent();
-	void GetContent(uint8* OutStructMemory, int32 InStructSize);
+	bool IsChildOf(const UStruct* InStruct) const
+	{
+		return GetWrappedStruct()->IsChildOf(InStruct);
+	}
+
+	template<typename T>
+	bool IsChildOf() const
+	{
+		return IsChildOf(T::StaticStruct());
+	}
+
+	void SetContent(const uint8* InStructMemory, const UStruct* InStruct);
+	void GetContent(uint8* OutStructMemory, const UStruct* InStruct) const;
 
 	template<typename T>
-	T* GetContent()
+	T GetContent() const
 	{
-		return (T*)GetStoredStructPtr();
+		check(IsChildOf<T>());
+		
+		T Result;
+		GetContent((uint8*)&Result, T::StaticStruct());
+		return Result;
 	}
 
 	template<typename T>
-	void SetContent(const T* InStructMemory)
+	void SetContent(const T& InValue)
 	{
-		SetContent((const uint8*)InStructMemory, T::StaticStruct()->GetStructureSize());
+		check(IsChildOf<T>());
+
+		SetContent((const uint8*)&InValue, T::StaticStruct());
 	}
 
 	FWrappedPropertyChangedChainEvent& GetWrappedPropertyChangedChainEvent() { return WrappedPropertyChangedChainEvent; }
@@ -41,7 +56,7 @@ public:
 
 private:
 
-	uint8* GetStoredStructPtr();
+	static void CopyPropertiesForUnrelatedStructs(uint8* InTargetMemory, const UStruct* InTargetStruct, const uint8* InSourceMemory, const UStruct* InSourceStruct);
 	
 	static TMap<UScriptStruct*, UClass*> StructToClass;
 	static TMap<UClass*, UScriptStruct*> ClassToStruct;
diff --git a/Engine/Plugins/Animation/ControlRig/Source/ControlRigEditor/Private/Editor/SControlRigGizmoNameList.cpp b/Engine/Plugins/Animation/ControlRig/Source/ControlRigEditor/Private/Editor/SControlRigGizmoNameList.cpp
--- a/Engine/Plugins/Animation/ControlRig/Source/ControlRigEditor/Private/Editor/SControlRigGizmoNameList.cpp
+++ b/Engine/Plugins/Animation/ControlRig/Source/ControlRigEditor/Private/Editor/SControlRigGizmoNameList.cpp
@@ -28,6 +28,25 @@ void SControlRigShapeNameList::Construct(const FArguments& InArgs, TArray<FRigCo
 	}
 	this->Blueprint = InBlueprint;
 
+	ConstructCommon();
+}
+
+void SControlRigShapeNameList::Construct(const FArguments& InArgs, TArray<FRigControlElement> ControlElements, UControlRigBlueprint* InBlueprint)
+{
+	this->OnGetNameListContent = InArgs._OnGetNameListContent;
+	this->ControlKeys.Reset();
+
+	for(const FRigControlElement& ControlElement : ControlElements)
+	{
+		this->ControlKeys.Add(ControlElement.GetKey());
+	}
+	this->Blueprint = InBlueprint;
+
+	ConstructCommon();
+}
+
+void SControlRigShapeNameList::ConstructCommon()
+{
 	SBox::Construct(SBox::FArguments());
 
 	TSharedPtr<FString> InitialSelected;
diff --git a/Engine/Plugins/Animation/ControlRig/Source/ControlRigEditor/Private/Editor/SControlRigGizmoNameList.h b/Engine/Plugins/Animation/ControlRig/Source/ControlRigEditor/Private/Editor/SControlRigGizmoNameList.h
--- a/Engine/Plugins/Animation/ControlRig/Source/ControlRigEditor/Private/Editor/SControlRigGizmoNameList.h
+++ b/Engine/Plugins/Animation/ControlRig/Source/ControlRigEditor/Private/Editor/SControlRigGizmoNameList.h
@@ -24,9 +24,12 @@ public:
 
 	void Construct(const FArguments& InArgs, FRigControlElement* ControlElement, UControlRigBlueprint* InBlueprint);
 	void Construct(const FArguments& InArgs, TArray<FRigControlElement*> ControlElements, UControlRigBlueprint* InBlueprint);
+	void Construct(const FArguments& InArgs, TArray<FRigControlElement> ControlElements, UControlRigBlueprint* InBlueprint);
 
 protected:
 
+	void ConstructCommon();
+
 	const TArray<TSharedPtr<FString>>& GetNameList() const;
 	FText GetNameListText() const;
 	virtual void SetNameListText(const FText& NewTypeInValue, ETextCommit::Type CommitInfo);
diff --git a/Engine/Source/Editor/BlueprintGraph/Private/BlueprintActionDatabase.cpp b/Engine/Source/Editor/BlueprintGraph/Private/BlueprintActionDatabase.cpp
--- a/Engine/Source/Editor/BlueprintGraph/Private/BlueprintActionDatabase.cpp
+++ b/Engine/Source/Editor/BlueprintGraph/Private/BlueprintActionDatabase.cpp
@@ -1396,10 +1396,11 @@ void FBlueprintActionDatabase::RefreshClassActions(UClass* const Class)
 	check(Class != nullptr);
 
 	bool const bOutOfDateClass   = Class->HasAnyClassFlags(CLASS_NewerVersionExists);
+	bool const bHiddenClass		 = Class->HasAnyClassFlags(CLASS_Hidden);
 	bool const bIsBlueprintClass = (Cast<UBlueprintGeneratedClass>(Class) != nullptr);
 	bool const bIsLevelScript	 = Class->ClassGeneratedBy && Cast<UBlueprint>(Class->ClassGeneratedBy)->BlueprintType == EBlueprintType::BPTYPE_LevelScript;
 
-	if (bOutOfDateClass || bIsLevelScript)
+	if (bOutOfDateClass || bIsLevelScript || bHiddenClass)
 	{
 		ActionRegistry.Remove(Class);
 		return;
