Control Rig: Fix crash when changing control type

diff --git a/Engine/Plugins/Animation/ControlRig/Source/ControlRigEditor/Private/ControlRigElementDetails.cpp b/Engine/Plugins/Animation/ControlRig/Source/ControlRigEditor/Private/ControlRigElementDetails.cpp
--- a/Engine/Plugins/Animation/ControlRig/Source/ControlRigEditor/Private/ControlRigElementDetails.cpp
+++ b/Engine/Plugins/Animation/ControlRig/Source/ControlRigEditor/Private/ControlRigElementDetails.cpp
@@ -1661,7 +1661,7 @@ void FRigControlElementDetails::CustomizeDetails(IDetailLayoutBuilder& DetailBui
 	// when control type changes, we have to refresh detail panel
 	const TSharedPtr<IPropertyHandle> ControlTypeHandle = SettingsHandle->GetChildHandle(TEXT("ControlType"));
 	ControlTypeHandle->SetOnPropertyValueChanged(FSimpleDelegate::CreateLambda(
-		[this, &PropertyUtilities]()
+		[this, PropertyUtilities]()
 		{
 			TArray<FRigControlElement> ControlElementsInView = GetElementsInDetailsView<FRigControlElement>();
 			TArray<FRigControlElement*> ControlElementsInHierarchy = GetElementsInHierarchy<FRigControlElement>();
@@ -1923,7 +1923,7 @@ void FRigControlElementDetails::CustomizeDetails(IDetailLayoutBuilder& DetailBui
 		ControlCategory.AddProperty(ControlEnumHandle.ToSharedRef()).DisplayName(FText::FromString(TEXT("Control Enum")));
 
 		ControlEnumHandle->SetOnPropertyValueChanged(FSimpleDelegate::CreateLambda(
-			[this, &PropertyUtilities, ControlElements, ObjectPerControl]()
+			[this, PropertyUtilities, ControlElements, ObjectPerControl]()
 			{
 				PropertyUtilities->ForceRefresh();
 
