Control Rig: Fixed cannot bind for each pin to variable

diff --git a/Engine/Plugins/Animation/ControlRig/Source/ControlRigEditor/Private/ControlRigEditorModule.cpp b/Engine/Plugins/Animation/ControlRig/Source/ControlRigEditor/Private/ControlRigEditorModule.cpp
--- a/Engine/Plugins/Animation/ControlRig/Source/ControlRigEditor/Private/ControlRigEditorModule.cpp
+++ b/Engine/Plugins/Animation/ControlRig/Source/ControlRigEditor/Private/ControlRigEditorModule.cpp
@@ -1415,7 +1415,8 @@ void FControlRigEditorModule::GetContextMenuActions(const UControlRigGraphSchema
 					}
 
 					if (Cast<URigVMUnitNode>(ModelPin->GetNode()) != nullptr || 
-						Cast<URigVMLibraryNode>(ModelPin->GetNode()) != nullptr)
+						Cast<URigVMLibraryNode>(ModelPin->GetNode()) != nullptr ||
+						Cast<URigVMArrayNode>(ModelPin->GetNode()) != nullptr)
 					{
 						if (ModelPin->GetDirection() == ERigVMPinDirection::Input && 
 							!ModelPin->IsExecuteContext())
