display curve widget properly

diff --git a/Engine/Plugins/Animation/ControlRig/Source/ControlRigEditor/Private/Graph/SControlRigGraphNode.cpp b/Engine/Plugins/Animation/ControlRig/Source/ControlRigEditor/Private/Graph/SControlRigGraphNode.cpp
--- a/Engine/Plugins/Animation/ControlRig/Source/ControlRigEditor/Private/Graph/SControlRigGraphNode.cpp
+++ b/Engine/Plugins/Animation/ControlRig/Source/ControlRigEditor/Private/Graph/SControlRigGraphNode.cpp
@@ -164,6 +164,8 @@ void SControlRigGraphNode::Construct( const FArguments& InArgs )
 		.Size(FVector2D(1.f, 2.f))
 	];
 
+	const UControlRigGraphSchema* RigSchema = Cast<UControlRigGraphSchema>(EdGraphNode->GetSchema());
+	
 	TMap<URigVMPin*, int32> ModelPinToInfoIndex;
 	for(URigVMPin* ModelPin : ModelPins)
 	{
@@ -175,16 +177,17 @@ void SControlRigGraphNode::Construct( const FArguments& InArgs )
 		PinInfo.Depth = 0;
 		PinInfo.bExpanded = ModelPin->IsExpanded();
 		PinInfo.ModelPinPath = ModelPin->GetPinPath();
-
-		PinInfo.bHideInputWidget = (!ModelPin->IsBoundToVariable()) && (ModelPin->IsStruct() || PinInfo.bIsContainer);
+		PinInfo.bAutoHeight = false;
+		
+		const bool bAskSchemaForEdition = RigSchema && ModelPin->IsStruct();
+		PinInfo.bHideInputWidget = (!ModelPin->IsBoundToVariable()) && PinInfo.bIsContainer;
 		if (!PinInfo.bHideInputWidget)
 		{
-			if (ModelPin->GetSubPins().Num() == 0)
+			if (bAskSchemaForEdition && !PinInfo.bHasChildren)
 			{
-				if (const UControlRigGraphSchema* RigSchema = Cast<UControlRigGraphSchema>(EdGraphNode->GetSchema()))
-				{
-					PinInfo.bHideInputWidget = RigSchema->IsStructEditable(ModelPin->GetScriptStruct());
-				}
+				const bool bIsStructEditable = RigSchema->IsStructEditable(ModelPin->GetScriptStruct()); 
+				PinInfo.bHideInputWidget = !bIsStructEditable;
+				PinInfo.bAutoHeight = bIsStructEditable;
 			}
 		}
 		
@@ -347,7 +350,7 @@ void SControlRigGraphNode::Construct( const FArguments& InArgs )
                 .HAlign(HAlign_Fill)
 				.VAlign(VAlign_Center)
 				.AutoHeight()
-				.MaxHeight(MaxHeight)
+				.MaxHeight(PinInfo.bAutoHeight ? TAttribute<float>() : MaxHeight)
                 [
                     SAssignNew(SlotLayout, SHorizontalBox)
                     .Visibility(this, &SControlRigGraphNode::GetPinVisibility, PinInfo.Index)
@@ -385,7 +388,7 @@ void SControlRigGraphNode::Construct( const FArguments& InArgs )
                 .HAlign(HAlign_Fill)
 				.VAlign(VAlign_Center)
                 .AutoHeight()
-				.MaxHeight(MaxHeight)
+				.MaxHeight(PinInfo.bAutoHeight ? TAttribute<float>() : MaxHeight)
                 [
                     SAssignNew(SlotLayout, SHorizontalBox)
                     .Visibility(this, &SControlRigGraphNode::GetPinVisibility, PinInfo.Index)
@@ -429,7 +432,7 @@ void SControlRigGraphNode::Construct( const FArguments& InArgs )
             .HAlign(HAlign_Fill)
 			.VAlign(VAlign_Center)
             .AutoHeight()
-			.MaxHeight(MaxHeight)
+			.MaxHeight(PinInfo.bAutoHeight ? TAttribute<float>() : MaxHeight)
             [
             	SNew(SHorizontalBox)
 	            .Visibility(this, &SControlRigGraphNode::GetPinVisibility, PinInfo.Index)
diff --git a/Engine/Plugins/Animation/ControlRig/Source/ControlRigEditor/Private/Graph/SControlRigGraphNode.h b/Engine/Plugins/Animation/ControlRig/Source/ControlRigEditor/Private/Graph/SControlRigGraphNode.h
--- a/Engine/Plugins/Animation/ControlRig/Source/ControlRigEditor/Private/Graph/SControlRigGraphNode.h
+++ b/Engine/Plugins/Animation/ControlRig/Source/ControlRigEditor/Private/Graph/SControlRigGraphNode.h
@@ -114,6 +114,7 @@ private:
 		TSharedPtr<SGraphPin> InputPinWidget;
 		TSharedPtr<SGraphPin> OutputPinWidget;
 		bool bExpanded;
+		bool bAutoHeight;
 	};
 	
 	TArray<FPinInfo> PinInfos;
diff --git a/Engine/Plugins/Animation/ControlRig/Source/ControlRigEditor/Private/Graph/SControlRigGraphPinCurveFloat.cpp b/Engine/Plugins/Animation/ControlRig/Source/ControlRigEditor/Private/Graph/SControlRigGraphPinCurveFloat.cpp
--- a/Engine/Plugins/Animation/ControlRig/Source/ControlRigEditor/Private/Graph/SControlRigGraphPinCurveFloat.cpp
+++ b/Engine/Plugins/Animation/ControlRig/Source/ControlRigEditor/Private/Graph/SControlRigGraphPinCurveFloat.cpp
@@ -25,8 +25,9 @@ TSharedRef<SWidget>	SControlRigGraphPinCurveFloat::GetDefaultValueWidget()
 {
 	UControlRigGraph* RigGraph = Cast<UControlRigGraph>(GraphPinObj->GetOwningNode()->GetGraph());
 
+	// 360 is the minimum width required to display keys' values 
 	TSharedRef<SWidget> Widget = SNew(SBox)
-		.MinDesiredWidth(200)
+		.MinDesiredWidth(360)
 		.MaxDesiredWidth(400)
 		.MinDesiredHeight(175)
 		.MaxDesiredHeight(300)
@@ -38,7 +39,7 @@ TSharedRef<SWidget>	SControlRigGraphPinCurveFloat::GetDefaultValueWidget()
 			.ViewMinOutput(0.f)
 			.ViewMaxOutput(1.f)
 			.TimelineLength(1.f)
-			.DesiredSize(FVector2D(300, 200))
+			.DesiredSize(FVector2D(360, 200))
 			.HideUI(true)
 		];
 
