Control Rig: Add Hide Parents when Filtering option in RigHierarchy

diff --git a/Engine/Plugins/Animation/ControlRig/Source/ControlRigEditor/Private/Editor/SRigHierarchy.cpp b/Engine/Plugins/Animation/ControlRig/Source/ControlRigEditor/Private/Editor/SRigHierarchy.cpp
--- a/Engine/Plugins/Animation/ControlRig/Source/ControlRigEditor/Private/Editor/SRigHierarchy.cpp
+++ b/Engine/Plugins/Animation/ControlRig/Source/ControlRigEditor/Private/Editor/SRigHierarchy.cpp
@@ -845,7 +845,7 @@ TSharedRef< SWidget > SRigHierarchy::CreateFilterMenu()
 	{
 		MenuBuilder.AddMenuEntry(Actions.FilteringFlattensHierarchy);
 		MenuBuilder.AddMenuEntry(Actions.ShowDynamicHierarchy);
-		//MenuBuilder.AddMenuEntry(Actions.HideParentsWhenFiltering);
+		MenuBuilder.AddMenuEntry(Actions.HideParentsWhenFiltering);
 	}
 	MenuBuilder.EndSection();
 
diff --git a/Engine/Plugins/Animation/ControlRig/Source/ControlRigEditor/Private/Editor/SRigHierarchyTreeView.cpp b/Engine/Plugins/Animation/ControlRig/Source/ControlRigEditor/Private/Editor/SRigHierarchyTreeView.cpp
--- a/Engine/Plugins/Animation/ControlRig/Source/ControlRigEditor/Private/Editor/SRigHierarchyTreeView.cpp
+++ b/Engine/Plugins/Animation/ControlRig/Source/ControlRigEditor/Private/Editor/SRigHierarchyTreeView.cpp
@@ -25,11 +25,12 @@ FRigTreeDisplaySettings FRigTreeDelegates::DefaultDisplaySettings;
 //////////////////////////////////////////////////////////////
 /// FRigTreeElement
 ///////////////////////////////////////////////////////////
-FRigTreeElement::FRigTreeElement(const FRigElementKey& InKey, TWeakPtr<SRigHierarchyTreeView> InTreeView, bool InSupportsRename)
+FRigTreeElement::FRigTreeElement(const FRigElementKey& InKey, TWeakPtr<SRigHierarchyTreeView> InTreeView, bool InSupportsRename, ERigTreeFilterResult InFilterResult)
 {
 	Key = InKey;
 	bIsTransient = false;
 	bSupportsRename = InSupportsRename;
+	FilterResult = InFilterResult;
 
 	if(InTreeView.IsValid())
 	{
@@ -87,6 +88,7 @@ void SRigHierarchyItem::Construct(const FArguments& InArgs, const TSharedRef<STa
 	TSharedPtr< SInlineEditableTextBlock > InlineWidget;
 
 	const FSlateBrush* Brush = GetBrushForElementType(Delegates.GetHierarchy(), InRigTreeElement->Key);
+	FSlateColor Color = InRigTreeElement->FilterResult == ERigTreeFilterResult::Shown ? FSlateColor::UseForeground() : FSlateColor(FLinearColor::Gray * 0.5f); 
 
 	STableRow<TSharedPtr<FRigTreeElement>>::Construct(
 		STableRow<TSharedPtr<FRigTreeElement>>::FArguments()
@@ -105,6 +107,7 @@ void SRigHierarchyItem::Construct(const FArguments& InArgs, const TSharedRef<STa
 			[
 				SNew(SImage)
 				.Image(Brush)
+				.ColorAndOpacity(Color)
 			]
 			+ SHorizontalBox::Slot()
 			.AutoWidth()
@@ -115,6 +118,7 @@ void SRigHierarchyItem::Construct(const FArguments& InArgs, const TSharedRef<STa
 				.OnVerifyTextChanged(this, &SRigHierarchyItem::OnVerifyNameChanged)
 				.OnTextCommitted(this, &SRigHierarchyItem::OnNameCommitted)
 				.MultiLine(false)
+				.ColorAndOpacity(Color)
 			]
 		], OwnerTable);
 
@@ -187,7 +191,7 @@ bool SRigHierarchyTreeView::AddElement(FRigElementKey InKey, FRigElementKey InPa
 	const FString FilteredString = Settings.FilterText.ToString();
 	if (FilteredString.IsEmpty() || !InKey.IsValid())
 	{
-		TSharedPtr<FRigTreeElement> NewItem = MakeShared<FRigTreeElement>(InKey, SharedThis(this), bSupportsRename);
+		TSharedPtr<FRigTreeElement> NewItem = MakeShared<FRigTreeElement>(InKey, SharedThis(this), bSupportsRename, ERigTreeFilterResult::Shown);
 
 		if (InKey.IsValid())
 		{
@@ -218,9 +222,37 @@ bool SRigHierarchyTreeView::AddElement(FRigElementKey InKey, FRigElementKey InPa
 		FString FilteredStringUnderScores = FilteredString.Replace(TEXT(" "), TEXT("_"));
 		if (InKey.Name.ToString().Contains(FilteredString) || InKey.Name.ToString().Contains(FilteredStringUnderScores))	
 		{
-			TSharedPtr<FRigTreeElement> NewItem = MakeShared<FRigTreeElement>(InKey, SharedThis(this), bSupportsRename);
+			TSharedPtr<FRigTreeElement> NewItem = MakeShared<FRigTreeElement>(InKey, SharedThis(this), bSupportsRename, ERigTreeFilterResult::Shown);
 			ElementMap.Add(InKey, NewItem);
 			RootElements.Add(NewItem);
+
+			if (!Settings.bFlattenHierarchyOnFilter && !Settings.bHideParentsOnFilter)
+			{
+				if(const URigHierarchy* Hierarchy = Delegates.GetHierarchy())
+				{
+					TSharedPtr<FRigTreeElement> ChildItem = NewItem;
+					FRigElementKey ParentKey = Hierarchy->GetFirstParent(InKey);
+					while (ParentKey.IsValid())
+					{
+						if (!ElementMap.Contains(ParentKey))
+						{
+							TSharedPtr<FRigTreeElement> ParentItem = MakeShared<FRigTreeElement>(ParentKey, SharedThis(this), bSupportsRename, ERigTreeFilterResult::ShownDescendant);							
+							ElementMap.Add(ParentKey, ParentItem);
+							RootElements.Add(ParentItem);
+
+							ReparentElement(ChildItem->Key, ParentKey);
+
+							ChildItem = ParentItem;
+							ParentKey = Hierarchy->GetFirstParent(ParentKey);
+						}
+						else
+						{
+							ReparentElement(ChildItem->Key, ParentKey);
+							break;
+						}						
+					}
+				}
+			}
 		}
 	}
 
@@ -454,6 +486,16 @@ void SRigHierarchyTreeView::RefreshTreeView(bool bRebuildContent)
 					SetExpansionRecursive(RootElement, false, true);
 				}
 			}
+			else if (ExpansionState.Num() < ElementMap.Num())
+			{
+				for (const TPair<FRigElementKey, TSharedPtr<FRigTreeElement>>& Element : ElementMap)
+				{
+					if (!ExpansionState.Contains(Element.Key))
+					{
+						SetItemExpansion(Element.Value, true);
+					}
+				}
+			}
 
 			for (const auto& Pair : ElementMap)
 			{
diff --git a/Engine/Plugins/Animation/ControlRig/Source/ControlRigEditor/Private/Editor/SRigHierarchyTreeView.h b/Engine/Plugins/Animation/ControlRig/Source/ControlRigEditor/Private/Editor/SRigHierarchyTreeView.h
--- a/Engine/Plugins/Animation/ControlRig/Source/ControlRigEditor/Private/Editor/SRigHierarchyTreeView.h
+++ b/Engine/Plugins/Animation/ControlRig/Source/ControlRigEditor/Private/Editor/SRigHierarchyTreeView.h
@@ -141,12 +141,28 @@ struct CONTROLRIGEDITOR_API FRigTreeDelegates
 	bool bIsChangingRigHierarchy;
 };
 
+/** 
+ * Order is important here! 
+ * This enum is used internally to the filtering logic and represents an ordering of most filtered (hidden) to 
+ * least filtered (highlighted).
+ */
+enum class ERigTreeFilterResult : int32
+{
+	/** Hide the item */
+	Hidden,
+
+	/** Show the item because child items were shown */
+	ShownDescendant,
+
+	/** Show the item */
+	Shown,
+};
 
 /** An item in the tree */
 class FRigTreeElement : public TSharedFromThis<FRigTreeElement>
 {
 public:
-	FRigTreeElement(const FRigElementKey& InKey, TWeakPtr<SRigHierarchyTreeView> InTreeView, bool InSupportsRename);
+	FRigTreeElement(const FRigElementKey& InKey, TWeakPtr<SRigHierarchyTreeView> InTreeView, bool InSupportsRename, ERigTreeFilterResult InFilterResult);
 public:
 	/** Element Data to display */
 	FRigElementKey Key;
@@ -161,6 +177,9 @@ public:
 	/** Delegate for when the context menu requests a rename */
 	DECLARE_DELEGATE(FOnRenameRequested);
 	FOnRenameRequested OnRenameRequested;
+
+	/** The current filter result */
+	ERigTreeFilterResult FilterResult;
 };
 
 class SRigHierarchyItem : public STableRow<TSharedPtr<FRigTreeElement>>
