Control Rig: Move euler / rotation order conversion to AnimationCore

preflight https://horde.devtools.epicgames.com/job/6193dc759137c0a5bedd8645

diff --git a/Engine/Plugins/Animation/ControlRig/Config/BaseControlRig.ini b/Engine/Plugins/Animation/ControlRig/Config/BaseControlRig.ini
--- a/Engine/Plugins/Animation/ControlRig/Config/BaseControlRig.ini
+++ b/Engine/Plugins/Animation/ControlRig/Config/BaseControlRig.ini
@@ -12,6 +12,7 @@
 +ClassRedirects=(OldName="RigVMGraph",NewName="RigVMGraph",ValueChanges=(("Model","RigVMModel")))
 
 +EnumRedirects=(OldName="ERigElementType",ValueChanges=(("Space","Null"))
++EnumRedirects=(OldName="EControlRigRotationOrder",NewName="EEulerRotationOrder")
 
 ; renaming gizmo to shape
 +ClassRedirects=(OldName="/Script/ControlRig.ControlRigGizmoLibrary", NewName="/Script/ControlRig.ControlRigShapeLibrary")
diff --git a/Engine/Plugins/Animation/ControlRig/Source/ControlRig/Private/Math/ControlRigMathLibrary.cpp b/Engine/Plugins/Animation/ControlRig/Source/ControlRig/Private/Math/ControlRigMathLibrary.cpp
--- a/Engine/Plugins/Animation/ControlRig/Source/ControlRig/Private/Math/ControlRigMathLibrary.cpp
+++ b/Engine/Plugins/Animation/ControlRig/Source/ControlRig/Private/Math/ControlRigMathLibrary.cpp
@@ -14,202 +14,6 @@ float FControlRigMathLibrary::AngleBetween(const FVector& A, const FVector& B)
 	return FMath::Acos(FVector::DotProduct(A, B) / (A.Size() * B.Size()));
 }
 
-FQuat FControlRigMathLibrary::QuatFromEuler(const FVector& XYZAnglesInDegrees, EControlRigRotationOrder RotationOrder)
-{
-	float X = FMath::DegreesToRadians(XYZAnglesInDegrees.X);
-	float Y = FMath::DegreesToRadians(XYZAnglesInDegrees.Y);
-	float Z = FMath::DegreesToRadians(XYZAnglesInDegrees.Z);
-
-	float CosX = FMath::Cos( X * 0.5f );
-	float CosY = FMath::Cos( Y * 0.5f );
-	float CosZ = FMath::Cos( Z * 0.5f );
-
-	float SinX = FMath::Sin( X * 0.5f );
-	float SinY = FMath::Sin( Y * 0.5f );
-	float SinZ = FMath::Sin( Z * 0.5f );
-
-	if ( RotationOrder == EControlRigRotationOrder::XYZ )
-	{
-		return FQuat(SinX * CosY * CosZ - CosX * SinY * SinZ, 
-					 CosX * SinY * CosZ + SinX * CosY * SinZ,
-					 CosX * CosY * SinZ - SinX * SinY * CosZ,
-					 CosX * CosY * CosZ + SinX * SinY * SinZ);
-
-	}
-	else if ( RotationOrder == EControlRigRotationOrder::XZY )
-	{
-		return FQuat(SinX * CosY * CosZ + CosX * SinY * SinZ, 
-					 CosX * SinY * CosZ + SinX * CosY * SinZ,
-					 CosX * CosY * SinZ - SinX * SinY * CosZ,
-					 CosX * CosY * CosZ - SinX * SinY * SinZ);
-
-	}
-	else if ( RotationOrder == EControlRigRotationOrder::YXZ )
-	{
-		return FQuat(SinX * CosY * CosZ - CosX * SinY * SinZ, 
-					 CosX * SinY * CosZ + SinX * CosY * SinZ,
-					 CosX * CosY * SinZ + SinX * SinY * CosZ,
-					 CosX * CosY * CosZ - SinX * SinY * SinZ);
-
-	}
-	else if ( RotationOrder == EControlRigRotationOrder::YZX )
-	{
-		return FQuat(SinX * CosY * CosZ - CosX * SinY * SinZ, 
-					 CosX * SinY * CosZ - SinX * CosY * SinZ,
-					 CosX * CosY * SinZ + SinX * SinY * CosZ,
-					 CosX * CosY * CosZ + SinX * SinY * SinZ);
-	}
-	else if ( RotationOrder == EControlRigRotationOrder::ZXY )
-	{
-		return FQuat(SinX * CosY * CosZ + CosX * SinY * SinZ, 
-					 CosX * SinY * CosZ - SinX * CosY * SinZ,
-					 CosX * CosY * SinZ - SinX * SinY * CosZ,
-					 CosX * CosY * CosZ + SinX * SinY * SinZ);
-
-	}
-	else if ( RotationOrder == EControlRigRotationOrder::ZYX )
-	{
-		return FQuat(SinX * CosY * CosZ + CosX * SinY * SinZ, 
-					 CosX * SinY * CosZ - SinX * CosY * SinZ,
-					 CosX * CosY * SinZ + SinX * SinY * CosZ,
-					 CosX * CosY * CosZ - SinX * SinY * SinZ);
-
-	}
-
-	// should not happen
-	return FQuat::Identity;
-}
-
-FVector FControlRigMathLibrary::EulerFromQuat(const FQuat& Rotation, EControlRigRotationOrder RotationOrder)
-{
-	float X = Rotation.X;
-	float Y = Rotation.Y;
-	float Z = Rotation.Z;
-	float W = Rotation.W;
-	float X2 = X * 2.f;
-	float Y2 = Y * 2.f;
-	float Z2 = Z * 2.f;
-	float XX2 = X * X2;
-	float XY2 = X * Y2;
-	float XZ2 = X * Z2;
-	float YX2 = Y * X2;
-	float YY2 = Y * Y2;
-	float YZ2 = Y * Z2;
-	float ZX2 = Z * X2;
-	float ZY2 = Z * Y2;
-	float ZZ2 = Z * Z2;
-	float WX2 = W * X2;
-	float WY2 = W * Y2;
-	float WZ2 = W * Z2;
-
-	FVector AxisX, AxisY, AxisZ;
-	AxisX.X = (1.f - (YY2 + ZZ2));
-	AxisY.X = (XY2 + WZ2);
-	AxisZ.X = (XZ2 - WY2);
-	AxisX.Y = (XY2 - WZ2);
-	AxisY.Y = (1.f - (XX2 + ZZ2));
-	AxisZ.Y = (YZ2 + WX2);
-	AxisX.Z = (XZ2 + WY2);
-	AxisY.Z = (YZ2 - WX2);
-	AxisZ.Z = (1.f - (XX2 + YY2));
-
-	FVector Result = FVector::ZeroVector;
-
-	if ( RotationOrder == EControlRigRotationOrder::XYZ )
-	{
-		Result.Y = FMath::Asin( - FMath::Clamp<float>( AxisZ.X, -1.f, 1.f ) );
-
-		if ( FMath::Abs( AxisZ.X ) < 1.f - SMALL_NUMBER )
-		{
-			Result.X = FMath::Atan2( AxisZ.Y, AxisZ.Z );
-			Result.Z = FMath::Atan2( AxisY.X, AxisX.X );
-		}
-		else
-		{
-			Result.X = 0.f;
-			Result.Z = FMath::Atan2( -AxisX.Y, AxisY.Y );
-		}
-	}
-	else if ( RotationOrder == EControlRigRotationOrder::XZY )
-	{
-
-		Result.Z = FMath::Asin( FMath::Clamp<float>( AxisY.X, -1.f, 1.f ) );
-
-		if ( FMath::Abs( AxisY.X ) < 1.f - SMALL_NUMBER )
-		{
-			Result.X = FMath::Atan2( -AxisY.Z, AxisY.Y );
-			Result.Y = FMath::Atan2( -AxisZ.X, AxisX.X );
-		}
-		else
-		{
-			Result.X = 0.f;
-			Result.Y = FMath::Atan2( AxisX.Z, AxisZ.Z );
-		}
-	}
-	else if ( RotationOrder == EControlRigRotationOrder::YXZ )
-	{
-		Result.X = FMath::Asin( FMath::Clamp<float>( AxisZ.Y, -1.f, 1.f ) );
-
-		if ( FMath::Abs( AxisZ.Y ) < 1.f - SMALL_NUMBER )
-		{
-			Result.Y = FMath::Atan2( -AxisZ.X, AxisZ.Z );
-			Result.Z = FMath::Atan2( -AxisX.Y, AxisY.Y );
-		}
-		else
-		{
-			Result.Y = 0.f;
-			Result.Z = FMath::Atan2( AxisY.X, AxisX.X );
-		}
-	}
-	else if ( RotationOrder == EControlRigRotationOrder::YZX )
-	{
-		Result.Z = FMath::Asin( - FMath::Clamp<float>( AxisX.Y, -1.f, 1.f ) );
-
-		if ( FMath::Abs( AxisX.Y ) < 1.f - SMALL_NUMBER )
-		{
-			Result.X = FMath::Atan2( AxisZ.Y, AxisY.Y );
-			Result.Y = FMath::Atan2( AxisX.Z, AxisX.X );
-		}
-		else
-		{
-			Result.X = FMath::Atan2( -AxisY.Z, AxisZ.Z );
-			Result.Y = 0.f;
-		}
-	}
-	else if ( RotationOrder == EControlRigRotationOrder::ZXY )
-	{
-		Result.X = FMath::Asin( - FMath::Clamp<float>( AxisY.Z, -1.f, 1.f ) );
-
-		if ( FMath::Abs( AxisY.Z ) < 1.f - SMALL_NUMBER )
-		{
-			Result.Y = FMath::Atan2( AxisX.Z, AxisZ.Z );
-			Result.Z = FMath::Atan2( AxisY.X, AxisY.Y );
-		}
-		else
-		{
-			Result.Y = FMath::Atan2( -AxisZ.X, AxisX.X );
-			Result.Z = 0.f;
-		}
-	}
-	else if ( RotationOrder == EControlRigRotationOrder::ZYX )
-	{
-		Result.Y = FMath::Asin( FMath::Clamp<float>( AxisX.Z, -1.f, 1.f ) );
-
-		if ( FMath::Abs( AxisX.Z ) < 1.f - SMALL_NUMBER )
-		{
-			Result.X = FMath::Atan2( -AxisY.Z, AxisZ.Z );
-			Result.Z = FMath::Atan2( -AxisX.Y, AxisX.X );
-		}
-		else
-		{
-			Result.X = FMath::Atan2( AxisZ.Y, AxisY.Y );
-			Result.Z = 0.f;
-		}
-	}
-
-	return Result * 180.f / PI;
-}
-
 void FControlRigMathLibrary::FourPointBezier(const FCRFourPointBezier& Bezier, float T, FVector& OutPosition, FVector& OutTangent)
 {
 	FourPointBezier(Bezier.A, Bezier.B, Bezier.C, Bezier.D, T, OutPosition, OutTangent);
diff --git a/Engine/Plugins/Animation/ControlRig/Source/ControlRig/Private/Units/Highlevel/Harmonics/RigUnit_BoneHarmonics.cpp b/Engine/Plugins/Animation/ControlRig/Source/ControlRig/Private/Units/Highlevel/Harmonics/RigUnit_BoneHarmonics.cpp
--- a/Engine/Plugins/Animation/ControlRig/Source/ControlRig/Private/Units/Highlevel/Harmonics/RigUnit_BoneHarmonics.cpp
+++ b/Engine/Plugins/Animation/ControlRig/Source/ControlRig/Private/Units/Highlevel/Harmonics/RigUnit_BoneHarmonics.cpp
@@ -2,6 +2,7 @@
 
 #include "Units/Highlevel/Harmonics/RigUnit_BoneHarmonics.h"
 #include "Units/RigUnitContext.h"
+#include "AnimationCoreLibrary.h"
 
 FRigUnit_BoneHarmonics_Execute()
 {
@@ -85,7 +86,7 @@ FRigUnit_ItemHarmonics_Execute()
 			Angles.Z = FMath::Sin(U.Z + WaveOffset.Z);
 			Angles = Angles * WaveAmplitude * Ease;
 
-			FQuat Rotation = FControlRigMathLibrary::QuatFromEuler(Angles, RotationOrder);
+			FQuat Rotation = AnimationCore::QuatFromEuler(Angles, RotationOrder);
 
 			FTransform Transform = Hierarchy->GetGlobalTransform(CachedItems[ItemIndex]);
 			Transform.SetRotation(Transform.GetRotation() * Rotation);
diff --git a/Engine/Plugins/Animation/ControlRig/Source/ControlRig/Private/Units/Highlevel/Harmonics/RigUnit_BoneHarmonics.h b/Engine/Plugins/Animation/ControlRig/Source/ControlRig/Private/Units/Highlevel/Harmonics/RigUnit_BoneHarmonics.h
--- a/Engine/Plugins/Animation/ControlRig/Source/ControlRig/Private/Units/Highlevel/Harmonics/RigUnit_BoneHarmonics.h
+++ b/Engine/Plugins/Animation/ControlRig/Source/ControlRig/Private/Units/Highlevel/Harmonics/RigUnit_BoneHarmonics.h
@@ -91,7 +91,7 @@ struct CONTROLRIG_API FRigUnit_BoneHarmonics : public FRigUnit_HighlevelBaseMuta
 		WaveEase = EControlRigAnimEasingType::Linear;
 		WaveMinimum = 0.5f;
 		WaveMaximum = 1.f;
-		RotationOrder = EControlRigRotationOrder::YZX;
+		RotationOrder = EEulerRotationOrder::YZX;
 		bPropagateToChildren = true;
 	}
 
@@ -128,7 +128,7 @@ struct CONTROLRIG_API FRigUnit_BoneHarmonics : public FRigUnit_HighlevelBaseMuta
 	float WaveMaximum;
 
 	UPROPERTY(meta = (Input))
-	EControlRigRotationOrder RotationOrder;
+	EEulerRotationOrder RotationOrder;
 
 	/**
 	 * If set to true all of the global transforms of the children
@@ -160,7 +160,7 @@ struct CONTROLRIG_API FRigUnit_ItemHarmonics : public FRigUnit_HighlevelBaseMuta
 		WaveEase = EControlRigAnimEasingType::Linear;
 		WaveMinimum = 0.5f;
 		WaveMaximum = 1.f;
-		RotationOrder = EControlRigRotationOrder::YZX;
+		RotationOrder = EEulerRotationOrder::YZX;
 	}
 
 	RIGVM_METHOD()
@@ -196,7 +196,7 @@ struct CONTROLRIG_API FRigUnit_ItemHarmonics : public FRigUnit_HighlevelBaseMuta
 	float WaveMaximum;
 
 	UPROPERTY(meta = (Input))
-	EControlRigRotationOrder RotationOrder;
+	EEulerRotationOrder RotationOrder;
 
 	UPROPERTY(transient)
 	FRigUnit_BoneHarmonics_WorkData WorkData;
diff --git a/Engine/Plugins/Animation/ControlRig/Source/ControlRig/Private/Units/Highlevel/Hierarchy/RigUnit_DistributeRotation.cpp b/Engine/Plugins/Animation/ControlRig/Source/ControlRig/Private/Units/Highlevel/Hierarchy/RigUnit_DistributeRotation.cpp
--- a/Engine/Plugins/Animation/ControlRig/Source/ControlRig/Private/Units/Highlevel/Hierarchy/RigUnit_DistributeRotation.cpp
+++ b/Engine/Plugins/Animation/ControlRig/Source/ControlRig/Private/Units/Highlevel/Hierarchy/RigUnit_DistributeRotation.cpp
@@ -2,6 +2,7 @@
 
 #include "Units/Highlevel/Hierarchy/RigUnit_DistributeRotation.h"
 #include "Units/RigUnitContext.h"
+#include "AnimationCoreLibrary.h"
 
 FRigUnit_DistributeRotation_Execute()
 {
@@ -273,7 +274,7 @@ IMPLEMENT_RIGUNIT_AUTOMATION_TEST(FRigUnit_DistributeRotation)
 	Rotation.Rotation = FQuat::Identity;
 	Rotation.Ratio = 1.f;
 	Unit.Rotations.Add(Rotation);
-	Rotation.Rotation = FControlRigMathLibrary::QuatFromEuler(FVector(0.f, 90.f, 0.f), EControlRigRotationOrder::XYZ);
+	Rotation.Rotation = AnimationCore::QuatFromEuler(FVector(0.f, 90.f, 0.f), EEulerRotationOrder::XYZ);
 	Rotation.Ratio = 0.5f;
 	Unit.Rotations.Add(Rotation);
 
@@ -297,15 +298,15 @@ IMPLEMENT_RIGUNIT_AUTOMATION_TEST(FRigUnit_DistributeRotation)
 	AddErrorIfFalse(FMath::IsNearlyEqual(Unit.WorkData.ItemRotationT[4], 0.0f, 0.001f), TEXT("unexpected bone t"));
 
 	FVector Euler = FVector::ZeroVector;
-	Euler = FControlRigMathLibrary::EulerFromQuat(Hierarchy->GetLocalTransform(0).GetRotation(), EControlRigRotationOrder::XYZ);
+	Euler = AnimationCore::EulerFromQuat(Hierarchy->GetLocalTransform(0).GetRotation(), EEulerRotationOrder::XYZ);
 	AddErrorIfFalse(FMath::IsNearlyEqual((double)Euler.Y, 0.0, 0.1), TEXT("unexpected rotation Y"));
-	Euler = FControlRigMathLibrary::EulerFromQuat(Hierarchy->GetLocalTransform(1).GetRotation(), EControlRigRotationOrder::XYZ);
+	Euler = AnimationCore::EulerFromQuat(Hierarchy->GetLocalTransform(1).GetRotation(), EEulerRotationOrder::XYZ);
 	AddErrorIfFalse(FMath::IsNearlyEqual((double)Euler.Y, 45.0, 0.1), TEXT("unexpected rotation Y"));
-	Euler = FControlRigMathLibrary::EulerFromQuat(Hierarchy->GetLocalTransform(2).GetRotation(), EControlRigRotationOrder::XYZ);
+	Euler = AnimationCore::EulerFromQuat(Hierarchy->GetLocalTransform(2).GetRotation(), EEulerRotationOrder::XYZ);
 	AddErrorIfFalse(FMath::IsNearlyEqual((double)Euler.Y, 90.0, 0.1), TEXT("unexpected rotation Y"));
-	Euler = FControlRigMathLibrary::EulerFromQuat(Hierarchy->GetLocalTransform(3).GetRotation(), EControlRigRotationOrder::XYZ);
+	Euler = AnimationCore::EulerFromQuat(Hierarchy->GetLocalTransform(3).GetRotation(), EEulerRotationOrder::XYZ);
 	AddErrorIfFalse(FMath::IsNearlyEqual((double)Euler.Y, 45.0, 0.1), TEXT("unexpected rotation Y"));
-	Euler = FControlRigMathLibrary::EulerFromQuat(Hierarchy->GetLocalTransform(4).GetRotation(), EControlRigRotationOrder::XYZ);
+	Euler = AnimationCore::EulerFromQuat(Hierarchy->GetLocalTransform(4).GetRotation(), EEulerRotationOrder::XYZ);
 	AddErrorIfFalse(FMath::IsNearlyEqual((double)Euler.Y, 0.0, 0.1), TEXT("unexpected rotation Y"));
 
 	return true;
diff --git a/Engine/Plugins/Animation/ControlRig/Source/ControlRig/Private/Units/Highlevel/Hierarchy/RigUnit_TransformConstraint.cpp b/Engine/Plugins/Animation/ControlRig/Source/ControlRig/Private/Units/Highlevel/Hierarchy/RigUnit_TransformConstraint.cpp
--- a/Engine/Plugins/Animation/ControlRig/Source/ControlRig/Private/Units/Highlevel/Hierarchy/RigUnit_TransformConstraint.cpp
+++ b/Engine/Plugins/Animation/ControlRig/Source/ControlRig/Private/Units/Highlevel/Hierarchy/RigUnit_TransformConstraint.cpp
@@ -287,7 +287,7 @@ FRigUnit_ParentConstraint_Execute()
 			MixedLocalTransform.NormalizeRotation();
 			FVector MixedTranslation = MixedLocalTransform.GetTranslation();
 			FQuat MixedRotation = MixedLocalTransform.GetRotation();
-			FVector MixedEulerRotation = FControlRigMathLibrary::EulerFromQuat(MixedRotation, AdvancedSettings.RotationOrderForFilter);
+			FVector MixedEulerRotation = AnimationCore::EulerFromQuat(MixedRotation, AdvancedSettings.RotationOrderForFilter);
 			FVector MixedScale = MixedLocalTransform.GetScale3D();
 
 			FTransform ChildCurrentLocalTransform = Hierarchy->GetLocalTransform(Child, false);
@@ -307,7 +307,7 @@ FRigUnit_ParentConstraint_Execute()
 			
 			FVector ChildTranslation = ChildCurrentLocalTransform.GetTranslation();
 			FQuat ChildRotation = ChildCurrentLocalTransform.GetRotation();
-			FVector ChildEulerRotation = FControlRigMathLibrary::EulerFromQuat(ChildRotation, AdvancedSettings.RotationOrderForFilter);
+			FVector ChildEulerRotation = AnimationCore::EulerFromQuat(ChildRotation, AdvancedSettings.RotationOrderForFilter);
 			FVector ChildScale = ChildCurrentLocalTransform.GetScale3D();
 
 			FVector FilteredTranslation;
@@ -325,7 +325,7 @@ FRigUnit_ParentConstraint_Execute()
 			FilteredScale.Y = Filter.ScaleFilter.bY ? MixedScale.Y : ChildScale.Y;
 			FilteredScale.Z = Filter.ScaleFilter.bZ ? MixedScale.Z : ChildScale.Z;
 
-			FTransform FilteredMixedLocalTransform(FControlRigMathLibrary::QuatFromEuler(FilteredEulerRotation, AdvancedSettings.RotationOrderForFilter), FilteredTranslation, FilteredScale);
+			FTransform FilteredMixedLocalTransform(AnimationCore::QuatFromEuler(FilteredEulerRotation, AdvancedSettings.RotationOrderForFilter), FilteredTranslation, FilteredScale);
 
 			FTransform FinalLocalTransform = FilteredMixedLocalTransform;
 
@@ -352,12 +352,12 @@ FRigUnit_ParentConstraint_Execute()
  IMPLEMENT_RIGUNIT_AUTOMATION_TEST(FRigUnit_ParentConstraint)
  {
 	// use euler rotation here to match other software's rotation representation more easily
-	EControlRigRotationOrder Order = EControlRigRotationOrder::XZY;
-	const FRigElementKey Child = Controller->AddBone(TEXT("Child"), FRigElementKey(), FTransform( FControlRigMathLibrary::QuatFromEuler( FVector(-10, -10, -10), Order), FVector(0.f, 0.f, 0.f)), true, ERigBoneType::User);
-	const FRigElementKey Parent1 = Controller->AddBone(TEXT("Parent1"), FRigElementKey(), FTransform( FControlRigMathLibrary::QuatFromEuler( FVector(30, -30, -30), Order), FVector(20.f, 20.f, 20.f)), true, ERigBoneType::User);
-	const FRigElementKey Parent2 = Controller->AddBone(TEXT("Parent2"), FRigElementKey(), FTransform( FControlRigMathLibrary::QuatFromEuler( FVector(-40, -40, 40), Order), FVector(40.f, 40.f, 40.f)), true, ERigBoneType::User);
-	const FRigElementKey Parent3 = Controller->AddBone(TEXT("Parent3"), FRigElementKey(), FTransform( FControlRigMathLibrary::QuatFromEuler( FVector(-50, 50, -50), Order), FVector(60.f, 60.f, 60.f)), true, ERigBoneType::User);
-	const FRigElementKey Parent4 = Controller->AddBone(TEXT("Parent4"), FRigElementKey(), FTransform( FControlRigMathLibrary::QuatFromEuler( FVector(60, 60, 60), Order), FVector(80.f, 80.f, 80.f)), true, ERigBoneType::User);
+	EEulerRotationOrder Order = EEulerRotationOrder::XZY;
+	const FRigElementKey Child = Controller->AddBone(TEXT("Child"), FRigElementKey(), FTransform( AnimationCore::QuatFromEuler( FVector(-10, -10, -10), Order), FVector(0.f, 0.f, 0.f)), true, ERigBoneType::User);
+	const FRigElementKey Parent1 = Controller->AddBone(TEXT("Parent1"), FRigElementKey(), FTransform( AnimationCore::QuatFromEuler( FVector(30, -30, -30), Order), FVector(20.f, 20.f, 20.f)), true, ERigBoneType::User);
+	const FRigElementKey Parent2 = Controller->AddBone(TEXT("Parent2"), FRigElementKey(), FTransform( AnimationCore::QuatFromEuler( FVector(-40, -40, 40), Order), FVector(40.f, 40.f, 40.f)), true, ERigBoneType::User);
+	const FRigElementKey Parent3 = Controller->AddBone(TEXT("Parent3"), FRigElementKey(), FTransform( AnimationCore::QuatFromEuler( FVector(-50, 50, -50), Order), FVector(60.f, 60.f, 60.f)), true, ERigBoneType::User);
+	const FRigElementKey Parent4 = Controller->AddBone(TEXT("Parent4"), FRigElementKey(), FTransform( AnimationCore::QuatFromEuler( FVector(60, 60, 60), Order), FVector(80.f, 80.f, 80.f)), true, ERigBoneType::User);
 	
 	Unit.ExecuteContext.Hierarchy = Hierarchy;
 	Unit.Child = Child;
@@ -375,7 +375,7 @@ FRigUnit_ParentConstraint_Execute()
 	AddErrorIfFalse(Hierarchy->GetGlobalTransform(0).GetTranslation().Equals(FVector(50.f, 50.f, 50.f)), TEXT("unexpected translation for average interpolation type"));
 
 	FQuat Result = Hierarchy->GetGlobalTransform(0).GetRotation();
-	FQuat Expected = FControlRigMathLibrary::QuatFromEuler( FVector(-0.852f, 15.189f, -0.572f), Order);
+	FQuat Expected = AnimationCore::QuatFromEuler( FVector(-0.852f, 15.189f, -0.572f), Order);
 	AddErrorIfFalse(Result.Equals(Expected, 0.001f), TEXT("unexpected rotation for average interpolation type"));
 	
 	Hierarchy->ResetPoseToInitial(ERigElementType::Bone);
@@ -384,12 +384,12 @@ FRigUnit_ParentConstraint_Execute()
 	Execute();
 	
 	Result = Hierarchy->GetGlobalTransform(0).GetRotation();
-	Expected = FControlRigMathLibrary::QuatFromEuler( FVector(16.74f, 8.865f, -5.562f), Order);
+	Expected = AnimationCore::QuatFromEuler( FVector(16.74f, 8.865f, -5.562f), Order);
 	AddErrorIfFalse(Result.Equals(Expected, 0.001f), TEXT("unexpected rotation for shortest interpolation type"));
 
 	
 	Hierarchy->ResetPoseToInitial(ERigElementType::Bone);
-	Hierarchy->SetGlobalTransform(2, FTransform(FControlRigMathLibrary::QuatFromEuler( FVector(100, 100, -100), Order), FVector(-40.f, -40.f, -40.f)));
+	Hierarchy->SetGlobalTransform(2, FTransform(AnimationCore::QuatFromEuler( FVector(100, 100, -100), Order), FVector(-40.f, -40.f, -40.f)));
 	Unit.bMaintainOffset = true;
 	Unit.AdvancedSettings.InterpolationType = EConstraintInterpType::Average;
 	
@@ -398,18 +398,18 @@ FRigUnit_ParentConstraint_Execute()
                     TEXT("unexpected translation for maintain offset and average interpolation type"));
 	
 	Result = Hierarchy->GetGlobalTransform(0).GetRotation();
-	Expected = FControlRigMathLibrary::QuatFromEuler( FVector(5.408f, -5.679f, -34.44f), Order);
+	Expected = AnimationCore::QuatFromEuler( FVector(5.408f, -5.679f, -34.44f), Order);
 	AddErrorIfFalse(Result.Equals(Expected, 0.001f), TEXT("unexpected rotation for maintain offset and average interpolation type"));
 
 	
 	Hierarchy->ResetPoseToInitial(ERigElementType::Bone);
-	Hierarchy->SetGlobalTransform(2, FTransform(FControlRigMathLibrary::QuatFromEuler( FVector(100.0f, 100.0f, -100.0f), Order), FVector(-40.f, -40.f, -40.f)));
+	Hierarchy->SetGlobalTransform(2, FTransform(AnimationCore::QuatFromEuler( FVector(100.0f, 100.0f, -100.0f), Order), FVector(-40.f, -40.f, -40.f)));
 	Unit.bMaintainOffset = true;
 	Unit.AdvancedSettings.InterpolationType = EConstraintInterpType::Shortest;
 	
 	Execute();
 	Result = Hierarchy->GetGlobalTransform(0).GetRotation();
-	Expected = FControlRigMathLibrary::QuatFromEuler( FVector(-1.209f, -8.332f, -25.022f), Order);
+	Expected = AnimationCore::QuatFromEuler( FVector(-1.209f, -8.332f, -25.022f), Order);
 	AddErrorIfFalse(Result.Equals(Expected, 0.001f), TEXT("unexpected rotation for maintain offset and shortest interpolation type"));
 	
 	return true;
@@ -543,7 +543,7 @@ FRigUnit_PositionConstraint_Execute()
  IMPLEMENT_RIGUNIT_AUTOMATION_TEST(FRigUnit_PositionConstraint)
  {
 	// use euler rotation here to match other software's rotation representation more easily
-	EControlRigRotationOrder Order = EControlRigRotationOrder::XZY;
+	EEulerRotationOrder Order = EEulerRotationOrder::XZY;
 	const FRigElementKey Child = Controller->AddBone(TEXT("Child"), FRigElementKey(), FTransform(FVector(0.f, 0.f, 0.f)), true, ERigBoneType::User);
 	const FRigElementKey Parent1 = Controller->AddBone(TEXT("Parent1"), FRigElementKey(), FTransform(FVector(20.f, 20.f, 20.f)), true, ERigBoneType::User);
 	const FRigElementKey Parent2 = Controller->AddBone(TEXT("Parent2"), FRigElementKey(), FTransform(FVector(40.f, 40.f, 40.f)), true, ERigBoneType::User);
@@ -687,7 +687,7 @@ FRigUnit_RotationConstraint_Execute()
 			// handle filtering, performed in local space
 			FTransform ChildParentGlobalTransform = Hierarchy->GetParentTransform(Child, false);
 			FQuat MixedLocalRotation = ChildParentGlobalTransform.GetRotation().Inverse() * MixedGlobalRotation;
-			FVector MixedEulerRotation = FControlRigMathLibrary::EulerFromQuat(MixedLocalRotation, AdvancedSettings.RotationOrderForFilter);
+			FVector MixedEulerRotation = AnimationCore::EulerFromQuat(MixedLocalRotation, AdvancedSettings.RotationOrderForFilter);
 			
 			FTransform ChildCurrentLocalTransform = Hierarchy->GetLocalTransform(Child, false);
 			
@@ -706,7 +706,7 @@ FRigUnit_RotationConstraint_Execute()
 
 			FQuat ChildRotation = ChildCurrentLocalTransform.GetRotation();
 			
-			FVector ChildEulerRotation = FControlRigMathLibrary::EulerFromQuat(ChildRotation, AdvancedSettings.RotationOrderForFilter);	
+			FVector ChildEulerRotation = AnimationCore::EulerFromQuat(ChildRotation, AdvancedSettings.RotationOrderForFilter);	
 			
 			FVector FilteredEulerRotation;
 			FilteredEulerRotation.X = Filter.bX ? MixedEulerRotation.X : ChildEulerRotation.X;
@@ -715,7 +715,7 @@ FRigUnit_RotationConstraint_Execute()
 
 			FTransform FilteredMixedLocalTransform = ChildCurrentLocalTransform;
 
-			FilteredMixedLocalTransform.SetRotation(FControlRigMathLibrary::QuatFromEuler(FilteredEulerRotation, AdvancedSettings.RotationOrderForFilter));
+			FilteredMixedLocalTransform.SetRotation(AnimationCore::QuatFromEuler(FilteredEulerRotation, AdvancedSettings.RotationOrderForFilter));
 
 			FTransform FinalLocalTransform = FilteredMixedLocalTransform;
 
@@ -744,12 +744,12 @@ IMPLEMENT_RIGUNIT_AUTOMATION_TEST(FRigUnit_RotationConstraint)
 	// the rotation constraint is expected to behave similarly to parent constraint with translation filter turned off.
 
 	// use euler rotation here to match other software's rotation representation more easily
-	EControlRigRotationOrder Order = EControlRigRotationOrder::XZY;
-	const FRigElementKey Child = Controller->AddBone(TEXT("Child"), FRigElementKey(), FTransform( FControlRigMathLibrary::QuatFromEuler( FVector(-10, -10, -10), Order), FVector(0.f, 0.f, 0.f)), true, ERigBoneType::User);
-	const FRigElementKey Parent1 = Controller->AddBone(TEXT("Parent1"), FRigElementKey(), FTransform( FControlRigMathLibrary::QuatFromEuler( FVector(30, -30, -30), Order), FVector(20.f, 20.f, 20.f)), true, ERigBoneType::User);
-	const FRigElementKey Parent2 = Controller->AddBone(TEXT("Parent2"), FRigElementKey(), FTransform( FControlRigMathLibrary::QuatFromEuler( FVector(-40, -40, 40), Order), FVector(40.f, 40.f, 40.f)), true, ERigBoneType::User);
-	const FRigElementKey Parent3 = Controller->AddBone(TEXT("Parent3"), FRigElementKey(), FTransform( FControlRigMathLibrary::QuatFromEuler( FVector(-50, 50, -50), Order), FVector(60.f, 60.f, 60.f)), true, ERigBoneType::User);
-	const FRigElementKey Parent4 = Controller->AddBone(TEXT("Parent4"), FRigElementKey(), FTransform( FControlRigMathLibrary::QuatFromEuler( FVector(60, 60, 60), Order), FVector(80.f, 80.f, 80.f)), true, ERigBoneType::User);
+	EEulerRotationOrder Order = EEulerRotationOrder::XZY;
+	const FRigElementKey Child = Controller->AddBone(TEXT("Child"), FRigElementKey(), FTransform( AnimationCore::QuatFromEuler( FVector(-10, -10, -10), Order), FVector(0.f, 0.f, 0.f)), true, ERigBoneType::User);
+	const FRigElementKey Parent1 = Controller->AddBone(TEXT("Parent1"), FRigElementKey(), FTransform( AnimationCore::QuatFromEuler( FVector(30, -30, -30), Order), FVector(20.f, 20.f, 20.f)), true, ERigBoneType::User);
+	const FRigElementKey Parent2 = Controller->AddBone(TEXT("Parent2"), FRigElementKey(), FTransform( AnimationCore::QuatFromEuler( FVector(-40, -40, 40), Order), FVector(40.f, 40.f, 40.f)), true, ERigBoneType::User);
+	const FRigElementKey Parent3 = Controller->AddBone(TEXT("Parent3"), FRigElementKey(), FTransform( AnimationCore::QuatFromEuler( FVector(-50, 50, -50), Order), FVector(60.f, 60.f, 60.f)), true, ERigBoneType::User);
+	const FRigElementKey Parent4 = Controller->AddBone(TEXT("Parent4"), FRigElementKey(), FTransform( AnimationCore::QuatFromEuler( FVector(60, 60, 60), Order), FVector(80.f, 80.f, 80.f)), true, ERigBoneType::User);
 	
 	Unit.ExecuteContext.Hierarchy = Hierarchy;
 	Unit.Child = Child;
@@ -766,7 +766,7 @@ IMPLEMENT_RIGUNIT_AUTOMATION_TEST(FRigUnit_RotationConstraint)
 	Execute();
 
 	FQuat Result = Hierarchy->GetGlobalTransform(0).GetRotation();
-	FQuat Expected = FControlRigMathLibrary::QuatFromEuler( FVector(-0.853f, 15.189f, -0.572f), Order);
+	FQuat Expected = AnimationCore::QuatFromEuler( FVector(-0.853f, 15.189f, -0.572f), Order);
 	AddErrorIfFalse(Result.Equals(Expected, 0.001f), TEXT("unexpected rotation for average interpolation type"));
 	
 	Hierarchy->ResetPoseToInitial(ERigElementType::Bone);
@@ -776,31 +776,31 @@ IMPLEMENT_RIGUNIT_AUTOMATION_TEST(FRigUnit_RotationConstraint)
 	Execute();
 	
 	Result = Hierarchy->GetGlobalTransform(0).GetRotation();
-	Expected = FControlRigMathLibrary::QuatFromEuler( FVector(16.74f, 8.865f, -5.562f), Order);
+	Expected = AnimationCore::QuatFromEuler( FVector(16.74f, 8.865f, -5.562f), Order);
 	AddErrorIfFalse(Result.Equals(Expected, 0.001f), TEXT("unexpected rotation for shortest interpolation type"));
 
 	
 	Hierarchy->ResetPoseToInitial(ERigElementType::Bone);
-	Hierarchy->SetGlobalTransform(2, FTransform(FControlRigMathLibrary::QuatFromEuler( FVector(100, 100, -100), Order), FVector(-40.f, -40.f, -40.f)));
+	Hierarchy->SetGlobalTransform(2, FTransform(AnimationCore::QuatFromEuler( FVector(100, 100, -100), Order), FVector(-40.f, -40.f, -40.f)));
 	Unit.bMaintainOffset = true;
 	Unit.AdvancedSettings.InterpolationType = EConstraintInterpType::Average;
 	
 	Execute();
 	
 	Result = Hierarchy->GetGlobalTransform(0).GetRotation();
-	Expected = FControlRigMathLibrary::QuatFromEuler( FVector(5.408f, -5.679f, -34.44f), Order);
+	Expected = AnimationCore::QuatFromEuler( FVector(5.408f, -5.679f, -34.44f), Order);
 	AddErrorIfFalse(Result.Equals(Expected, 0.001f), TEXT("unexpected rotation for maintain offset and average interpolation type"));
 
 	
 	Hierarchy->ResetPoseToInitial(ERigElementType::Bone);
-	Hierarchy->SetGlobalTransform(2, FTransform(FControlRigMathLibrary::QuatFromEuler( FVector(100.0f, 100.0f, -100.0f), Order), FVector(-40.f, -40.f, -40.f)));
+	Hierarchy->SetGlobalTransform(2, FTransform(AnimationCore::QuatFromEuler( FVector(100.0f, 100.0f, -100.0f), Order), FVector(-40.f, -40.f, -40.f)));
 	Unit.bMaintainOffset = true;
 	Unit.AdvancedSettings.InterpolationType = EConstraintInterpType::Shortest;
 	
 	Execute();
 	
 	Result = Hierarchy->GetGlobalTransform(0).GetRotation();
-	Expected = FControlRigMathLibrary::QuatFromEuler( FVector(-1.209f, -8.332f, -25.022f), Order);
+	Expected = AnimationCore::QuatFromEuler( FVector(-1.209f, -8.332f, -25.022f), Order);
 	AddErrorIfFalse(Result.Equals(Expected, 0.001f), TEXT("unexpected rotation for maintain offset and shortest interpolation type"));
 	
 	return true;
@@ -971,7 +971,7 @@ FRigUnit_ScaleConstraint_Execute()
 IMPLEMENT_RIGUNIT_AUTOMATION_TEST(FRigUnit_ScaleConstraint)
 {
 	// use euler rotation here to match other software's rotation representation more easily
-	EControlRigRotationOrder Order = EControlRigRotationOrder::XZY;
+	EEulerRotationOrder Order = EEulerRotationOrder::XZY;
 	const FRigElementKey Child = Controller->AddBone(TEXT("Child"), FRigElementKey(), FTransform(FQuat::Identity, FVector::ZeroVector, FVector::OneVector), true, ERigBoneType::User);
 	const FRigElementKey Parent1 = Controller->AddBone(TEXT("Parent1"), FRigElementKey(), FTransform(FQuat::Identity, FVector::ZeroVector, FVector(4,4,4)), true, ERigBoneType::User);
 	const FRigElementKey Parent2 = Controller->AddBone(TEXT("Parent2"), FRigElementKey(), FTransform(FQuat::Identity, FVector::ZeroVector, FVector(1,1,1)), true, ERigBoneType::User);
diff --git a/Engine/Plugins/Animation/ControlRig/Source/ControlRig/Private/Units/Highlevel/Hierarchy/RigUnit_TransformConstraint.h b/Engine/Plugins/Animation/ControlRig/Source/ControlRig/Private/Units/Highlevel/Hierarchy/RigUnit_TransformConstraint.h
--- a/Engine/Plugins/Animation/ControlRig/Source/ControlRig/Private/Units/Highlevel/Hierarchy/RigUnit_TransformConstraint.h
+++ b/Engine/Plugins/Animation/ControlRig/Source/ControlRig/Private/Units/Highlevel/Hierarchy/RigUnit_TransformConstraint.h
@@ -220,7 +220,7 @@ struct FRigUnit_ParentConstraint_AdvancedSettings
 
 	FRigUnit_ParentConstraint_AdvancedSettings()
 		: InterpolationType(EConstraintInterpType::Average)
-		, RotationOrderForFilter(EControlRigRotationOrder::XZY)
+		, RotationOrderForFilter(EEulerRotationOrder::XZY)
 	{}
 	
 	/**
@@ -233,7 +233,7 @@ struct FRigUnit_ParentConstraint_AdvancedSettings
 	*	Rotation is converted to euler angles using the specified order such that individual axes can be filtered.
 	*/
 	UPROPERTY(EditAnywhere, Category = "Constraint", meta = (Input))
-	EControlRigRotationOrder RotationOrderForFilter;
+	EEulerRotationOrder RotationOrderForFilter;
 };
 
 /**
@@ -316,7 +316,7 @@ struct FRigUnit_RotationConstraint_AdvancedSettings
 
 	FRigUnit_RotationConstraint_AdvancedSettings()
         : InterpolationType(EConstraintInterpType::Average)
-        , RotationOrderForFilter(EControlRigRotationOrder::XZY)
+        , RotationOrderForFilter(EEulerRotationOrder::XZY)
 	{}
 	
 	/**
@@ -329,7 +329,7 @@ struct FRigUnit_RotationConstraint_AdvancedSettings
 	*	Rotation is converted to euler angles using the specified order such that individual axes can be filtered.
 	*/
 	UPROPERTY(EditAnywhere, Category = "Constraint", meta = (Input))
-	EControlRigRotationOrder RotationOrderForFilter;
+	EEulerRotationOrder RotationOrderForFilter;
 };
 
 /**
diff --git a/Engine/Plugins/Animation/ControlRig/Source/ControlRig/Private/Units/Math/RigUnit_MathQuaternion.cpp b/Engine/Plugins/Animation/ControlRig/Source/ControlRig/Private/Units/Math/RigUnit_MathQuaternion.cpp
--- a/Engine/Plugins/Animation/ControlRig/Source/ControlRig/Private/Units/Math/RigUnit_MathQuaternion.cpp
+++ b/Engine/Plugins/Animation/ControlRig/Source/ControlRig/Private/Units/Math/RigUnit_MathQuaternion.cpp
@@ -2,6 +2,7 @@
 
 #include "Units/Math/RigUnit_MathQuaternion.h"
 #include "Units/RigUnitContext.h"
+#include "AnimationCoreLibrary.h"
 
 FRigUnit_MathQuaternionFromAxisAndAngle_Execute()
 {
@@ -18,7 +19,7 @@ FRigUnit_MathQuaternionFromAxisAndAngle_Execute()
 FRigUnit_MathQuaternionFromEuler_Execute()
 {
     DECLARE_SCOPE_HIERARCHICAL_COUNTER_RIGUNIT()
-	Result = FControlRigMathLibrary::QuatFromEuler(Euler, RotationOrder);
+	Result = AnimationCore::QuatFromEuler(Euler, RotationOrder);
 }
 
 FRigUnit_MathQuaternionFromRotator_Execute()
@@ -77,7 +78,7 @@ FRigUnit_MathQuaternionScale_Execute()
 FRigUnit_MathQuaternionToEuler_Execute()
 {
     DECLARE_SCOPE_HIERARCHICAL_COUNTER_RIGUNIT()
-	Result = FControlRigMathLibrary::EulerFromQuat(Value, RotationOrder);
+	Result = AnimationCore::EulerFromQuat(Value, RotationOrder);
 }
 
 FRigUnit_MathQuaternionToRotator_Execute()
diff --git a/Engine/Plugins/Animation/ControlRig/Source/ControlRig/Private/Units/Math/RigUnit_MathQuaternion.h b/Engine/Plugins/Animation/ControlRig/Source/ControlRig/Private/Units/Math/RigUnit_MathQuaternion.h
--- a/Engine/Plugins/Animation/ControlRig/Source/ControlRig/Private/Units/Math/RigUnit_MathQuaternion.h
+++ b/Engine/Plugins/Animation/ControlRig/Source/ControlRig/Private/Units/Math/RigUnit_MathQuaternion.h
@@ -88,7 +88,7 @@ struct CONTROLRIG_API FRigUnit_MathQuaternionFromEuler : public FRigUnit_MathQua
 	FRigUnit_MathQuaternionFromEuler()
 	{
 		Euler = FVector::ZeroVector;
-		RotationOrder = EControlRigRotationOrder::ZYX;
+		RotationOrder = EEulerRotationOrder::ZYX;
 		Result = FQuat::Identity;
 	}
 
@@ -99,7 +99,7 @@ struct CONTROLRIG_API FRigUnit_MathQuaternionFromEuler : public FRigUnit_MathQua
 	FVector Euler;
 
 	UPROPERTY(meta = (Input))
-	EControlRigRotationOrder RotationOrder;
+	EEulerRotationOrder RotationOrder;
 
 	UPROPERTY(meta=(Output))
 	FQuat Result;
@@ -219,7 +219,7 @@ struct CONTROLRIG_API FRigUnit_MathQuaternionToEuler : public FRigUnit_MathQuate
 	FRigUnit_MathQuaternionToEuler()
 	{
 		Value = FQuat::Identity;
-		RotationOrder = EControlRigRotationOrder::ZYX;
+		RotationOrder = EEulerRotationOrder::ZYX;
 		Result = FVector::ZeroVector;
 	}
 
@@ -230,7 +230,7 @@ struct CONTROLRIG_API FRigUnit_MathQuaternionToEuler : public FRigUnit_MathQuate
 	FQuat Value;
 
 	UPROPERTY(meta = (Input))
-	EControlRigRotationOrder RotationOrder;
+	EEulerRotationOrder RotationOrder;
 
 	UPROPERTY(meta=(Output))
 	FVector Result;
@@ -530,11 +530,11 @@ struct CONTROLRIG_API FRigUnit_MathQuaternionRotationOrder : public FRigUnit_Mat
 
 	FRigUnit_MathQuaternionRotationOrder()
 	{
-		RotationOrder = EControlRigRotationOrder::ZYX;
+		RotationOrder = EEulerRotationOrder::ZYX;
 	}
 
 	UPROPERTY(meta = (Input, Output))
-	EControlRigRotationOrder RotationOrder;
+	EEulerRotationOrder RotationOrder;
 
 
 	RIGVM_METHOD()
diff --git a/Engine/Plugins/Animation/ControlRig/Source/ControlRig/Private/Units/Math/RigUnit_MathQuaternionTests.cpp b/Engine/Plugins/Animation/ControlRig/Source/ControlRig/Private/Units/Math/RigUnit_MathQuaternionTests.cpp
--- a/Engine/Plugins/Animation/ControlRig/Source/ControlRig/Private/Units/Math/RigUnit_MathQuaternionTests.cpp
+++ b/Engine/Plugins/Animation/ControlRig/Source/ControlRig/Private/Units/Math/RigUnit_MathQuaternionTests.cpp
@@ -3,6 +3,7 @@
 #if WITH_EDITOR
 #include "Units/Math/RigUnit_MathQuaternion.h"
 #include "Units/RigUnitTest.h"
+#include "AnimationCoreLibrary.h"
 
 namespace FRigUnit_MathQuatTest_Utils
 {
@@ -53,11 +54,11 @@ IMPLEMENT_RIGUNIT_AUTOMATION_TEST(FRigUnit_MathQuaternionFromEuler)
 {
 	Unit.Euler = FVector(30.f, 20.f, 10.f);
 	InitAndExecute();
-	FVector Euler = FControlRigMathLibrary::EulerFromQuat(Unit.Result, Unit.RotationOrder);
+	FVector Euler = AnimationCore::EulerFromQuat(Unit.Result, Unit.RotationOrder);
 	AddErrorIfFalse(FRigUnit_MathQuatTest_Utils::IsNearlyEqual(Euler, FVector(30.f, 20.f, 10.f)), TEXT("unexpected result"));
-	Unit.RotationOrder = EControlRigRotationOrder::YZX;
+	Unit.RotationOrder = EEulerRotationOrder::YZX;
 	InitAndExecute();
-	Euler = FControlRigMathLibrary::EulerFromQuat(Unit.Result, Unit.RotationOrder);
+	Euler = AnimationCore::EulerFromQuat(Unit.Result, Unit.RotationOrder);
 	AddErrorIfFalse(FRigUnit_MathQuatTest_Utils::IsNearlyEqual(Euler, FVector(30.f, 20.f, 10.f)), TEXT("unexpected result"));
 	return true;
 }
@@ -90,11 +91,11 @@ IMPLEMENT_RIGUNIT_AUTOMATION_TEST(FRigUnit_MathQuaternionToAxisAndAngle)
 
 IMPLEMENT_RIGUNIT_AUTOMATION_TEST(FRigUnit_MathQuaternionToEuler)
 {
-	Unit.Value = FControlRigMathLibrary::QuatFromEuler(FVector(10.f, 20.f, 30.f), Unit.RotationOrder);
+	Unit.Value = AnimationCore::QuatFromEuler(FVector(10.f, 20.f, 30.f), Unit.RotationOrder);
 	InitAndExecute();
 	AddErrorIfFalse(FRigUnit_MathQuatTest_Utils::IsNearlyEqual(Unit.Result, FVector(10.f, 20.f, 30.f)), TEXT("unexpected result"));
-	Unit.RotationOrder = EControlRigRotationOrder::YZX;
-	Unit.Value = FControlRigMathLibrary::QuatFromEuler(FVector(10.f, 20.f, 30.f), Unit.RotationOrder);
+	Unit.RotationOrder = EEulerRotationOrder::YZX;
+	Unit.Value = AnimationCore::QuatFromEuler(FVector(10.f, 20.f, 30.f), Unit.RotationOrder);
 	InitAndExecute();
 	AddErrorIfFalse(FRigUnit_MathQuatTest_Utils::IsNearlyEqual(Unit.Result, FVector(10.f, 20.f, 30.f)), TEXT("unexpected result"));
 	return true;
diff --git a/Engine/Plugins/Animation/ControlRig/Source/ControlRig/Private/Units/Math/RigUnit_MathTransform.cpp b/Engine/Plugins/Animation/ControlRig/Source/ControlRig/Private/Units/Math/RigUnit_MathTransform.cpp
--- a/Engine/Plugins/Animation/ControlRig/Source/ControlRig/Private/Units/Math/RigUnit_MathTransform.cpp
+++ b/Engine/Plugins/Animation/ControlRig/Source/ControlRig/Private/Units/Math/RigUnit_MathTransform.cpp
@@ -4,6 +4,7 @@
 #include "Units/Math/RigUnit_MathVector.h"
 #include "Math/ControlRigMathLibrary.h"
 #include "Units/RigUnitContext.h"
+#include "AnimationCoreLibrary.h"
 
 FRigUnit_MathTransformFromEulerTransform_Execute()
 {
@@ -127,7 +128,7 @@ FRigUnit_MathTransformFromSRT_Execute()
 {
     DECLARE_SCOPE_HIERARCHICAL_COUNTER_RIGUNIT()
 	Transform.SetLocation(Location);
-	Transform.SetRotation(FControlRigMathLibrary::QuatFromEuler(Rotation, RotationOrder));
+	Transform.SetRotation(AnimationCore::QuatFromEuler(Rotation, RotationOrder));
 	Transform.SetScale3D(Scale);
 	EulerTransform.FromFTransform(Transform);
 }
diff --git a/Engine/Plugins/Animation/ControlRig/Source/ControlRig/Private/Units/Math/RigUnit_MathTransform.h b/Engine/Plugins/Animation/ControlRig/Source/ControlRig/Private/Units/Math/RigUnit_MathTransform.h
--- a/Engine/Plugins/Animation/ControlRig/Source/ControlRig/Private/Units/Math/RigUnit_MathTransform.h
+++ b/Engine/Plugins/Animation/ControlRig/Source/ControlRig/Private/Units/Math/RigUnit_MathTransform.h
@@ -346,7 +346,7 @@ struct CONTROLRIG_API FRigUnit_MathTransformFromSRT : public FRigUnit_MathTransf
 	{
 		Location = FVector::ZeroVector;
 		Rotation = FVector::ZeroVector;
-		RotationOrder = EControlRigRotationOrder::XYZ;
+		RotationOrder = EEulerRotationOrder::XYZ;
 		Scale = FVector::OneVector;
 		Transform = FTransform::Identity;
 		EulerTransform = FEulerTransform::Identity;
@@ -362,7 +362,7 @@ struct CONTROLRIG_API FRigUnit_MathTransformFromSRT : public FRigUnit_MathTransf
 	FVector Rotation;
 
 	UPROPERTY(meta=(Input))
-	EControlRigRotationOrder RotationOrder;
+	EEulerRotationOrder RotationOrder;
 
 	UPROPERTY(meta=(Input))
 	FVector Scale;
diff --git a/Engine/Plugins/Animation/ControlRig/Source/ControlRig/Public/Math/ControlRigMathLibrary.h b/Engine/Plugins/Animation/ControlRig/Source/ControlRig/Public/Math/ControlRigMathLibrary.h
--- a/Engine/Plugins/Animation/ControlRig/Source/ControlRig/Public/Math/ControlRigMathLibrary.h
+++ b/Engine/Plugins/Animation/ControlRig/Source/ControlRig/Public/Math/ControlRigMathLibrary.h
@@ -6,17 +6,6 @@
 #include "ControlRigDefines.h"
 #include "ControlRigMathLibrary.generated.h"
 
-UENUM()
-enum class EControlRigRotationOrder : uint8
-{
-	XYZ,
-	XZY,
-	YXZ,
-	YZX,
-	ZXY,
-	ZYX
-};
-
 UENUM()
 enum class EControlRigAnimEasingType : uint8
 {
@@ -80,8 +69,6 @@ class CONTROLRIG_API FControlRigMathLibrary
 {
 public:
 	static float AngleBetween(const FVector& A, const FVector& B);
-	static FQuat QuatFromEuler(const FVector& XYZAnglesInDegrees, EControlRigRotationOrder RotationOrderr = EControlRigRotationOrder::ZYX);
-	static FVector EulerFromQuat(const FQuat& Rotation, EControlRigRotationOrder RotationOrder = EControlRigRotationOrder::ZYX);
 	static void FourPointBezier(const FVector& A, const FVector& B, const FVector& C, const FVector& D, float T, FVector& OutPosition, FVector& OutTangent);
 	static void FourPointBezier(const FCRFourPointBezier& Bezier, float T, FVector& OutPosition, FVector& OutTangent);
 	static float EaseFloat(float Value, EControlRigAnimEasingType Type);
diff --git a/Engine/Source/Runtime/AnimationCore/Private/AnimationCoreLibrary.cpp b/Engine/Source/Runtime/AnimationCore/Private/AnimationCoreLibrary.cpp
--- a/Engine/Source/Runtime/AnimationCore/Private/AnimationCoreLibrary.cpp
+++ b/Engine/Source/Runtime/AnimationCore/Private/AnimationCoreLibrary.cpp
@@ -189,4 +189,201 @@ FTransform SolveConstraints(const FTransform& CurrentTransform, const FTransform
 
 	return BlendedLocalTransform * CurrentParentTransform;
 }
+
+FQuat QuatFromEuler(const FVector& XYZAnglesInDegrees, EEulerRotationOrder RotationOrder)
+{
+	float X = FMath::DegreesToRadians(XYZAnglesInDegrees.X);
+	float Y = FMath::DegreesToRadians(XYZAnglesInDegrees.Y);
+	float Z = FMath::DegreesToRadians(XYZAnglesInDegrees.Z);
+
+	float CosX = FMath::Cos( X * 0.5f );
+	float CosY = FMath::Cos( Y * 0.5f );
+	float CosZ = FMath::Cos( Z * 0.5f );
+
+	float SinX = FMath::Sin( X * 0.5f );
+	float SinY = FMath::Sin( Y * 0.5f );
+	float SinZ = FMath::Sin( Z * 0.5f );
+
+	if ( RotationOrder == EEulerRotationOrder::XYZ )
+	{
+		return FQuat(SinX * CosY * CosZ - CosX * SinY * SinZ, 
+					 CosX * SinY * CosZ + SinX * CosY * SinZ,
+					 CosX * CosY * SinZ - SinX * SinY * CosZ,
+					 CosX * CosY * CosZ + SinX * SinY * SinZ);
+
+	}
+	else if ( RotationOrder == EEulerRotationOrder::XZY )
+	{
+		return FQuat(SinX * CosY * CosZ + CosX * SinY * SinZ, 
+					 CosX * SinY * CosZ + SinX * CosY * SinZ,
+					 CosX * CosY * SinZ - SinX * SinY * CosZ,
+					 CosX * CosY * CosZ - SinX * SinY * SinZ);
+
+	}
+	else if ( RotationOrder == EEulerRotationOrder::YXZ )
+	{
+		return FQuat(SinX * CosY * CosZ - CosX * SinY * SinZ, 
+					 CosX * SinY * CosZ + SinX * CosY * SinZ,
+					 CosX * CosY * SinZ + SinX * SinY * CosZ,
+					 CosX * CosY * CosZ - SinX * SinY * SinZ);
+
+	}
+	else if ( RotationOrder == EEulerRotationOrder::YZX )
+	{
+		return FQuat(SinX * CosY * CosZ - CosX * SinY * SinZ, 
+					 CosX * SinY * CosZ - SinX * CosY * SinZ,
+					 CosX * CosY * SinZ + SinX * SinY * CosZ,
+					 CosX * CosY * CosZ + SinX * SinY * SinZ);
+	}
+	else if ( RotationOrder == EEulerRotationOrder::ZXY )
+	{
+		return FQuat(SinX * CosY * CosZ + CosX * SinY * SinZ, 
+					 CosX * SinY * CosZ - SinX * CosY * SinZ,
+					 CosX * CosY * SinZ - SinX * SinY * CosZ,
+					 CosX * CosY * CosZ + SinX * SinY * SinZ);
+
+	}
+	else if ( RotationOrder == EEulerRotationOrder::ZYX )
+	{
+		return FQuat(SinX * CosY * CosZ + CosX * SinY * SinZ, 
+					 CosX * SinY * CosZ - SinX * CosY * SinZ,
+					 CosX * CosY * SinZ + SinX * SinY * CosZ,
+					 CosX * CosY * CosZ - SinX * SinY * SinZ);
+
+	}
+
+	// should not happen
+	return FQuat::Identity;
+}
+
+FVector EulerFromQuat(const FQuat& Rotation, EEulerRotationOrder RotationOrder)
+{
+	float X = Rotation.X;
+	float Y = Rotation.Y;
+	float Z = Rotation.Z;
+	float W = Rotation.W;
+	float X2 = X * 2.f;
+	float Y2 = Y * 2.f;
+	float Z2 = Z * 2.f;
+	float XX2 = X * X2;
+	float XY2 = X * Y2;
+	float XZ2 = X * Z2;
+	float YX2 = Y * X2;
+	float YY2 = Y * Y2;
+	float YZ2 = Y * Z2;
+	float ZX2 = Z * X2;
+	float ZY2 = Z * Y2;
+	float ZZ2 = Z * Z2;
+	float WX2 = W * X2;
+	float WY2 = W * Y2;
+	float WZ2 = W * Z2;
+
+	FVector AxisX, AxisY, AxisZ;
+	AxisX.X = (1.f - (YY2 + ZZ2));
+	AxisY.X = (XY2 + WZ2);
+	AxisZ.X = (XZ2 - WY2);
+	AxisX.Y = (XY2 - WZ2);
+	AxisY.Y = (1.f - (XX2 + ZZ2));
+	AxisZ.Y = (YZ2 + WX2);
+	AxisX.Z = (XZ2 + WY2);
+	AxisY.Z = (YZ2 - WX2);
+	AxisZ.Z = (1.f - (XX2 + YY2));
+
+	FVector Result = FVector::ZeroVector;
+
+	if ( RotationOrder == EEulerRotationOrder::XYZ )
+	{
+		Result.Y = FMath::Asin( - FMath::Clamp<float>( AxisZ.X, -1.f, 1.f ) );
+
+		if ( FMath::Abs( AxisZ.X ) < 1.f - SMALL_NUMBER )
+		{
+			Result.X = FMath::Atan2( AxisZ.Y, AxisZ.Z );
+			Result.Z = FMath::Atan2( AxisY.X, AxisX.X );
+		}
+		else
+		{
+			Result.X = 0.f;
+			Result.Z = FMath::Atan2( -AxisX.Y, AxisY.Y );
+		}
+	}
+	else if ( RotationOrder == EEulerRotationOrder::XZY )
+	{
+
+		Result.Z = FMath::Asin( FMath::Clamp<float>( AxisY.X, -1.f, 1.f ) );
+
+		if ( FMath::Abs( AxisY.X ) < 1.f - SMALL_NUMBER )
+		{
+			Result.X = FMath::Atan2( -AxisY.Z, AxisY.Y );
+			Result.Y = FMath::Atan2( -AxisZ.X, AxisX.X );
+		}
+		else
+		{
+			Result.X = 0.f;
+			Result.Y = FMath::Atan2( AxisX.Z, AxisZ.Z );
+		}
+	}
+	else if ( RotationOrder == EEulerRotationOrder::YXZ )
+	{
+		Result.X = FMath::Asin( FMath::Clamp<float>( AxisZ.Y, -1.f, 1.f ) );
+
+		if ( FMath::Abs( AxisZ.Y ) < 1.f - SMALL_NUMBER )
+		{
+			Result.Y = FMath::Atan2( -AxisZ.X, AxisZ.Z );
+			Result.Z = FMath::Atan2( -AxisX.Y, AxisY.Y );
+		}
+		else
+		{
+			Result.Y = 0.f;
+			Result.Z = FMath::Atan2( AxisY.X, AxisX.X );
+		}
+	}
+	else if ( RotationOrder == EEulerRotationOrder::YZX )
+	{
+		Result.Z = FMath::Asin( - FMath::Clamp<float>( AxisX.Y, -1.f, 1.f ) );
+
+		if ( FMath::Abs( AxisX.Y ) < 1.f - SMALL_NUMBER )
+		{
+			Result.X = FMath::Atan2( AxisZ.Y, AxisY.Y );
+			Result.Y = FMath::Atan2( AxisX.Z, AxisX.X );
+		}
+		else
+		{
+			Result.X = FMath::Atan2( -AxisY.Z, AxisZ.Z );
+			Result.Y = 0.f;
+		}
+	}
+	else if ( RotationOrder == EEulerRotationOrder::ZXY )
+	{
+		Result.X = FMath::Asin( - FMath::Clamp<float>( AxisY.Z, -1.f, 1.f ) );
+
+		if ( FMath::Abs( AxisY.Z ) < 1.f - SMALL_NUMBER )
+		{
+			Result.Y = FMath::Atan2( AxisX.Z, AxisZ.Z );
+			Result.Z = FMath::Atan2( AxisY.X, AxisY.Y );
+		}
+		else
+		{
+			Result.Y = FMath::Atan2( -AxisZ.X, AxisX.X );
+			Result.Z = 0.f;
+		}
+	}
+	else if ( RotationOrder == EEulerRotationOrder::ZYX )
+	{
+		Result.Y = FMath::Asin( FMath::Clamp<float>( AxisX.Z, -1.f, 1.f ) );
+
+		if ( FMath::Abs( AxisX.Z ) < 1.f - SMALL_NUMBER )
+		{
+			Result.X = FMath::Atan2( -AxisY.Z, AxisZ.Z );
+			Result.Z = FMath::Atan2( -AxisX.Y, AxisX.X );
+		}
+		else
+		{
+			Result.X = FMath::Atan2( AxisZ.Y, AxisY.Y );
+			Result.Z = 0.f;
+		}
+	}
+
+	return Result * 180.f / PI;
+}
+
 }
\ No newline at end of file
diff --git a/Engine/Source/Runtime/AnimationCore/Public/AnimationCoreLibrary.h b/Engine/Source/Runtime/AnimationCore/Public/AnimationCoreLibrary.h
--- a/Engine/Source/Runtime/AnimationCore/Public/AnimationCoreLibrary.h
+++ b/Engine/Source/Runtime/AnimationCore/Public/AnimationCoreLibrary.h
@@ -7,6 +7,7 @@
 #pragma once
 
 #include "CoreMinimal.h"
+#include "EulerTransform.h"
 
 struct FTransformConstraint;
 struct FConstraintData;
@@ -51,4 +52,15 @@ namespace AnimationCore
 	 * @return  Delta Rotation to turn
 	 */
 	ANIMATIONCORE_API FQuat SolveAim(const FTransform& CurrentTransform, const FVector& TargetPosition, const FVector& AimVector, bool bUseUpVector = false, const FVector& UpVector = FVector::UpVector, float AimClampInDegree = 0.f);
+
+	/**
+	 * Converts a euler rotation represented by a vector of rotations in degrees and a rotation order to a quaternion
+	 */
+	ANIMATIONCORE_API FQuat QuatFromEuler(const FVector& XYZAnglesInDegrees, EEulerRotationOrder RotationOrderr = EEulerRotationOrder::ZYX);
+
+	/**
+	 * Converts a quaternion to a euler rotation represented by a vector of rotations in degrees and a rotation order
+	 */
+	ANIMATIONCORE_API FVector EulerFromQuat(const FQuat& Rotation, EEulerRotationOrder RotationOrder = EEulerRotationOrder::ZYX);
+
 }
diff --git a/Engine/Source/Runtime/AnimationCore/Public/EulerTransform.h b/Engine/Source/Runtime/AnimationCore/Public/EulerTransform.h
--- a/Engine/Source/Runtime/AnimationCore/Public/EulerTransform.h
+++ b/Engine/Source/Runtime/AnimationCore/Public/EulerTransform.h
@@ -5,6 +5,17 @@
 #include "CoreMinimal.h"
 #include "EulerTransform.generated.h"
 
+UENUM()
+enum class EEulerRotationOrder : uint8
+{
+	XYZ,
+	XZY,
+	YXZ,
+	YZX,
+	ZXY,
+	ZYX
+};
+
 USTRUCT(BlueprintType)
 struct ANIMATIONCORE_API FEulerTransform
 {
