ControlRig: Fixed syntax error in a redirector for horde issue https://horde.devtools.epicgames.com/log/6194556bcc0baff9b092b53c?lineindex=169

diff --git a/Engine/Plugins/Animation/ControlRig/Config/BaseControlRig.ini b/Engine/Plugins/Animation/ControlRig/Config/BaseControlRig.ini
--- a/Engine/Plugins/Animation/ControlRig/Config/BaseControlRig.ini
+++ b/Engine/Plugins/Animation/ControlRig/Config/BaseControlRig.ini
@@ -12,7 +12,8 @@
 +ClassRedirects=(OldName="RigVMGraph",NewName="RigVMGraph",ValueChanges=(("Model","RigVMModel")))
 
 +EnumRedirects=(OldName="ERigElementType",ValueChanges=(("Space","Null"))
-+EnumRedirects=(OldName="EControlRigRotationOrder",NewName="EEulerRotationOrder")
++EnumRedirects=(OldName="/Script/ControlRig.EControlRigRotationOrder",NewName="/Script/AnimationCore.EEulerRotationOrder")
+
 
 ; renaming gizmo to shape
 +ClassRedirects=(OldName="/Script/ControlRig.ControlRigGizmoLibrary", NewName="/Script/ControlRig.ControlRigShapeLibrary")
