Control Rig: Fix crash when renaming a function to a recently deleted function

diff --git a/Engine/Plugins/Animation/ControlRig/Source/ControlRigDeveloper/Private/Graph/ControlRigGraph.cpp b/Engine/Plugins/Animation/ControlRig/Source/ControlRigDeveloper/Private/Graph/ControlRigGraph.cpp
--- a/Engine/Plugins/Animation/ControlRig/Source/ControlRigDeveloper/Private/Graph/ControlRigGraph.cpp
+++ b/Engine/Plugins/Animation/ControlRig/Source/ControlRigDeveloper/Private/Graph/ControlRigGraph.cpp
@@ -395,6 +395,20 @@ void UControlRigGraph::HandleModifiedEvent(ERigVMGraphNotifType InNotifType, URi
 				UEdGraphNode* EdNode = FindNodeForModelNodeName(ModelNode->GetFName(), false);
 				if (EdNode)
 				{
+					// Rename the soon to be deleted object to a unique name, so that other objects can use
+					// the old name
+					{
+						FString DeletedName;
+						UObject* ExistingObject = nullptr;
+						int32 Index = FMath::Rand();
+						do
+						{
+							DeletedName = FString::Printf(TEXT("%s_Deleted_%d"), *EdNode->GetName(), Index++); 
+							ExistingObject = StaticFindObject(/*Class=*/ NULL, EdNode->GetOuter(), *DeletedName, true);						
+						}
+						while (ExistingObject);
+						EdNode->Rename(*DeletedName);
+					}
 					RemoveNode(EdNode, true);
 					NotifyGraphChanged();
 				}
