Sequencer: Control Rig: Add ability to edit multiple controls

Instead of one details panel with topmost nodes we have a detail panel per type, and then allow for multiple objects for each.  When theres multiple objects the system still asks for each customization for each object so we need to store the 'multiple' name so the title is correct. Also FTransform would not work correctly with multiple types so instead we are using FEulerTransform with coercion. (Note FTransforms aren't common anymore FEulerTransform is really the default CR transform type but we have lots of legacy assets still using Transform controls)

diff --git a/Engine/Plugins/Animation/ControlRig/Source/ControlRigEditor/Private/EditMode/ControlRigControlsProxy.cpp b/Engine/Plugins/Animation/ControlRig/Source/ControlRigEditor/Private/EditMode/ControlRigControlsProxy.cpp
--- a/Engine/Plugins/Animation/ControlRig/Source/ControlRigEditor/Private/EditMode/ControlRigControlsProxy.cpp
+++ b/Engine/Plugins/Animation/ControlRig/Source/ControlRigEditor/Private/EditMode/ControlRigControlsProxy.cpp
@@ -18,6 +18,23 @@
 #include "DetailWidgetRow.h"
 #include "SEnumCombo.h"
 
+
+void UControlRigControlsProxy::SetIsMultiple(bool bIsVal)
+{ 
+	bIsMultiple = bIsVal; 
+	if (bIsMultiple)
+	{
+		FString DisplayString = TEXT("Multiple");
+		FName DisplayName(*DisplayString);
+		Name = DisplayName;
+	}
+	else
+	{
+		Name = ControlName;
+	}
+}
+
+
 void FControlRigEnumControlProxyValueDetails::CustomizeHeader(TSharedRef<IPropertyHandle> InStructPropertyHandle, FDetailWidgetRow& HeaderRow, IPropertyTypeCustomizationUtils& StructCustomizationUtils)
 {
 	TArray<UObject*> Objects;
@@ -151,7 +168,8 @@ void UControlRigTransformControlProxy::PostEditChangeProperty(struct FPropertyCh
 		{
 			//MUST set through ControlRig
 			FControlRigInteractionScope InteractionScope(ControlRig.Get());
-			ControlRig->SetControlValue<FRigControlValue::FTransform_Float>(ControlName, Transform, true, EControlRigSetKey::DoNotCare,false);
+			FTransform RealTransform = Transform; //Transform is FEulerTransform
+			ControlRig->SetControlValue<FRigControlValue::FTransform_Float>(ControlName, RealTransform, true, EControlRigSetKey::DoNotCare,false);
 			ControlRig->Evaluate_AnyThread();
 
 		}
@@ -167,7 +185,8 @@ void UControlRigTransformControlProxy::ValueChanged()
 		const FName PropertyName("Transform");
 		FTrackInstancePropertyBindings Binding(PropertyName, PropertyName.ToString());
 		const FTransform NewTransform = ControlRig.Get()->GetHierarchy()->GetControlValue(ControlElement, ERigControlValueType::Current).Get<FRigControlValue::FTransform_Float>().ToTransform();
-		Binding.CallFunction<FTransform>(*this, NewTransform);
+		const FEulerTransform EulerTransform = NewTransform;
+		Binding.CallFunction<FEulerTransform>(*this, EulerTransform);
 	}
 }
 
@@ -182,7 +201,8 @@ void UControlRigTransformControlProxy::PostEditUndo()
 			CheckEditModeOnSelectionChange(ControlRig.Get());
 		}
 		ControlRig->SelectControl(ControlName, bSelected);
-		ControlRig->SetControlValue<FRigControlValue::FTransform_Float>(ControlName, Transform, true, EControlRigSetKey::Never,false);
+		FTransform RealTransform = Transform; //Transform is FEulerTransform
+		ControlRig->SetControlValue<FRigControlValue::FTransform_Float>(ControlName, RealTransform, true, EControlRigSetKey::Never,false);
 	}
 }
 #endif
@@ -192,7 +212,8 @@ void UControlRigTransformControlProxy::SetKey(const IPropertyHandle& KeyedProper
 	FRigControlElement* ControlElement = GetControlElement();
 	if (ControlElement)
 	{
-		ControlRig->SetControlValue<FRigControlValue::FTransform_Float>(ControlName, Transform, true, EControlRigSetKey::Always,false);
+		FTransform RealTransform = Transform; //Transform is FEulerTransform
+		ControlRig->SetControlValue<FRigControlValue::FTransform_Float>(ControlName, RealTransform, true, EControlRigSetKey::Always,false);
 	}
 }
 
diff --git a/Engine/Plugins/Animation/ControlRig/Source/ControlRigEditor/Private/EditMode/ControlRigControlsProxy.h b/Engine/Plugins/Animation/ControlRig/Source/ControlRigEditor/Private/EditMode/ControlRigControlsProxy.h
--- a/Engine/Plugins/Animation/ControlRig/Source/ControlRigEditor/Private/EditMode/ControlRigControlsProxy.h
+++ b/Engine/Plugins/Animation/ControlRig/Source/ControlRigEditor/Private/EditMode/ControlRigControlsProxy.h
@@ -24,8 +24,8 @@ class UControlRigControlsProxy : public UObject
 
 public:
 	UControlRigControlsProxy() : bSelected(false) {}
-	virtual void SetName(const FName& InName)  { ControlName = InName; }
-	virtual FName GetName() const { return ControlName; }
+	virtual void SetName(const FName& InName) { Name = ControlName = InName; }
+	virtual FName GetName() const { return Name; }
 	virtual void ValueChanged() {}
 	virtual void SelectionChanged(bool bInSelected);
 	virtual void SetKey(const IPropertyHandle& KeyedPropertyHandle) {};
@@ -35,7 +35,7 @@ public:
 #if WITH_EDITOR
 	virtual void PostEditUndo() override;
 #endif
-
+	void SetIsMultiple(bool bIsVal); 
 public:
 
 	FRigControlElement* GetControlElement() const;
@@ -46,6 +46,8 @@ public:
 	UPROPERTY(VisibleAnywhere, Category = "Control")
 	FName ControlName;
 protected:
+	bool bIsMultiple = 0;
+	FName Name;
 	void CheckEditModeOnSelectionChange(UControlRig* InControlRig);
 };
 
@@ -68,7 +70,7 @@ class UControlRigTransformControlProxy : public UControlRigControlsProxy
 public:
 	
 	UPROPERTY(EditAnywhere, Interp, Category = "Control")
-	FTransform Transform;
+	FEulerTransform Transform; //FTransform doesn't work with multiple values for some reason, so for now using eulertransform which does work
 };
 
 
diff --git a/Engine/Plugins/Animation/ControlRig/Source/ControlRigEditor/Private/EditMode/ControlRigEditMode.cpp b/Engine/Plugins/Animation/ControlRig/Source/ControlRigEditor/Private/EditMode/ControlRigEditMode.cpp
--- a/Engine/Plugins/Animation/ControlRig/Source/ControlRigEditor/Private/EditMode/ControlRigEditMode.cpp
+++ b/Engine/Plugins/Animation/ControlRig/Source/ControlRigEditor/Private/EditMode/ControlRigEditMode.cpp
@@ -230,18 +230,78 @@ void FControlRigEditMode::SetUpDetailPanel()
 {
 	if (IsInLevelEditor())
 	{
-		TArray<TWeakObjectPtr<>> SelectedObjects;
+		TArray<TWeakObjectPtr<>> Eulers;
+		TArray<TWeakObjectPtr<>> Transforms;
+		TArray<TWeakObjectPtr<>> TransformNoScales;
+		TArray<TWeakObjectPtr<>> Floats;
+		TArray<TWeakObjectPtr<>> Vectors;
+		TArray<TWeakObjectPtr<>> Vector2Ds;
+		TArray<TWeakObjectPtr<>> Bools;
+		TArray<TWeakObjectPtr<>> Integers;
+		TArray<TWeakObjectPtr<>> Enums;
 		if (UControlRig* ControlRig = GetControlRig(true))
 		{
 			const TArray<UControlRigControlsProxy*>& Proxies = ControlProxy->GetSelectedProxies();
 			for (UControlRigControlsProxy* Proxy : Proxies)
 			{
-				SelectedObjects.Add(Proxy);
+				if (Proxy->GetClass() == UControlRigTransformControlProxy::StaticClass())
+				{
+					Transforms.Add(Proxy);
+				}
+				else if (Proxy->GetClass() == UControlRigTransformNoScaleControlProxy::StaticClass())
+				{
+					TransformNoScales.Add(Proxy);
+				}
+				else if (Proxy->GetClass() == UControlRigEulerTransformControlProxy::StaticClass())
+				{
+					Eulers.Add(Proxy);
+				}
+				else if (Proxy->GetClass() == UControlRigFloatControlProxy::StaticClass())
+				{
+					Floats.Add(Proxy);
+				}
+				else if (Proxy->GetClass() == UControlRigVectorControlProxy::StaticClass())
+				{
+					Vectors.Add(Proxy);
+				}
+				else if (Proxy->GetClass() == UControlRigVector2DControlProxy::StaticClass())
+				{
+					Vector2Ds.Add(Proxy);
+				}
+				else if (Proxy->GetClass() == UControlRigBoolControlProxy::StaticClass())
+				{
+					Bools.Add(Proxy);
+				}
+				else if (Proxy->GetClass() == UControlRigEnumControlProxy::StaticClass())
+				{
+					Enums.Add(Proxy);
+				}
+				else if (Proxy->GetClass() == UControlRigIntegerControlProxy::StaticClass())
+				{
+					Integers.Add(Proxy);
+				}
 			}
-			SelectedObjects.Add(GetMutableDefault<UControlRigEditModeSettings>());
 		}
 		StaticCastSharedPtr<SControlRigEditModeTools>(Toolkit->GetInlineContent())->SetSequencer(WeakSequencer.Pin());
-		StaticCastSharedPtr<SControlRigEditModeTools>(Toolkit->GetInlineContent())->SetDetailsObjects(SelectedObjects);
+		for (TWeakObjectPtr<>& Object : Transforms)
+		{
+			UControlRigControlsProxy* Proxy = Cast<UControlRigControlsProxy>(Object.Get());
+			if (Proxy)
+			{
+				Proxy->SetIsMultiple(Transforms.Num() > 1);
+			}
+		}
+		StaticCastSharedPtr<SControlRigEditModeTools>(Toolkit->GetInlineContent())->SetTransformDetailsObjects(Transforms);
+		StaticCastSharedPtr<SControlRigEditModeTools>(Toolkit->GetInlineContent())->SetTransformNoScaleDetailsObjects(TransformNoScales);
+		StaticCastSharedPtr<SControlRigEditModeTools>(Toolkit->GetInlineContent())->SetEulerTransformDetailsObjects(Eulers);
+		StaticCastSharedPtr<SControlRigEditModeTools>(Toolkit->GetInlineContent())->SetFloatDetailsObjects(Floats);
+		StaticCastSharedPtr<SControlRigEditModeTools>(Toolkit->GetInlineContent())->SetVectorDetailsObjects(Vectors);
+		StaticCastSharedPtr<SControlRigEditModeTools>(Toolkit->GetInlineContent())->SetVector2DDetailsObjects(Vector2Ds);
+		StaticCastSharedPtr<SControlRigEditModeTools>(Toolkit->GetInlineContent())->SetBoolDetailsObjects(Bools);
+		StaticCastSharedPtr<SControlRigEditModeTools>(Toolkit->GetInlineContent())->SetIntegerDetailsObjects(Integers);
+		StaticCastSharedPtr<SControlRigEditModeTools>(Toolkit->GetInlineContent())->SetEnumDetailsObjects(Enums);
+
+		StaticCastSharedPtr<SControlRigEditModeTools>(Toolkit->GetInlineContent())->SetSettingsDetailsObject(GetMutableDefault<UControlRigEditModeSettings>());
 	}
 }
 
diff --git a/Engine/Plugins/Animation/ControlRig/Source/ControlRigEditor/Private/EditMode/SControlRigEditModeTools.cpp b/Engine/Plugins/Animation/ControlRig/Source/ControlRigEditor/Private/EditMode/SControlRigEditModeTools.cpp
--- a/Engine/Plugins/Animation/ControlRig/Source/ControlRigEditor/Private/EditMode/SControlRigEditModeTools.cpp
+++ b/Engine/Plugins/Animation/ControlRig/Source/ControlRigEditor/Private/EditMode/SControlRigEditModeTools.cpp
@@ -46,7 +46,7 @@ class FControlRigEditModeGenericDetails : public IDetailCustomization
 public:
 	FControlRigEditModeGenericDetails() = delete;
 	FControlRigEditModeGenericDetails(FEditorModeTools* InModeTools) : ModeTools(InModeTools) {}
-	
+
 	/** Makes a new instance of this detail layout class for a specific detail view requesting it */
 	static TSharedRef<IDetailCustomization> MakeInstance(FEditorModeTools* InModeTools)
 	{
@@ -68,13 +68,54 @@ public:
 			}
 		}
 
-		if (ProxiesBeingCustomized.Num() == 0)
+		if (ProxiesBeingCustomized.Num() == 0 || ProxiesBeingCustomized[0]->GetControlElement() == nullptr)
 		{
 			return;
 		}
+		FText ControlText = LOCTEXT("Channels", "Channels");
+		
+		if(ProxiesBeingCustomized[0]->GetClass() == UControlRigEditModeSettings::StaticClass())
+		{
+			ControlText = LOCTEXT("Animation", "Animation");
+		}
+		else if (ProxiesBeingCustomized[0]->GetClass() == UControlRigTransformControlProxy::StaticClass())
+		{
+			ControlText = LOCTEXT("TransformChannels", "Transform Channels");
+		}
+		else if (ProxiesBeingCustomized[0]->GetClass() == UControlRigTransformNoScaleControlProxy::StaticClass())
+		{
+			ControlText = LOCTEXT("TransformNoScaleChannels", "TransformNoScale Channels");
+		}
+		else if (ProxiesBeingCustomized[0]->GetClass() == UControlRigEulerTransformControlProxy::StaticClass())
+		{
+			ControlText = LOCTEXT("EulerTransformChannels", "Euler Transform Channels");
+		}
+		else if (ProxiesBeingCustomized[0]->GetClass() == UControlRigFloatControlProxy::StaticClass())
+		{
+			ControlText = LOCTEXT("FloatChannels", "Float Channels");
+		}
+		else if (ProxiesBeingCustomized[0]->GetClass() == UControlRigVectorControlProxy::StaticClass())
+		{
+			ControlText = LOCTEXT("VectorChannels", "Vector Channels");
+		}
+		else if (ProxiesBeingCustomized[0]->GetClass() == UControlRigVector2DControlProxy::StaticClass())
+		{
+			ControlText = LOCTEXT("Vector2DChannels", "Vector2D Channels");
+		}
+		else if (ProxiesBeingCustomized[0]->GetClass() == UControlRigBoolControlProxy::StaticClass())
+		{
+			ControlText = LOCTEXT("BoolChannels", "Bool Channels");
+		}
+		else if (ProxiesBeingCustomized[0]->GetClass() == UControlRigEnumControlProxy::StaticClass())
+		{
+			ControlText = LOCTEXT("EnumChannels", "Enum Channels");
+		}
+		else if (ProxiesBeingCustomized[0]->GetClass() == UControlRigIntegerControlProxy::StaticClass())
+		{
+			ControlText = LOCTEXT("IntegerChannels", "Integer Channels");
+		}
 
-		IDetailCategoryBuilder& Category = DetailLayout.EditCategory(TEXT("Control"), LOCTEXT("Channels", "Channels"));
-
+		IDetailCategoryBuilder& Category = DetailLayout.EditCategory(TEXT("Control"), ControlText);
 		for (UControlRigControlsProxy* Proxy : ProxiesBeingCustomized)
 		{
 			FRigControlElement* ControlElement = Proxy->GetControlElement();
@@ -112,16 +153,24 @@ public:
 			{
 				ValuePropertyName = TEXT("Vector2D");
 			}
+			else if (ControlElement->Settings.ControlType == ERigControlType::EulerTransform)
+			{
+				ValuePropertyName = TEXT("EulerTransform");
+			}
+			else if (ControlElement->Settings.ControlType == ERigControlType::TransformNoScale)
+			{
+				ValuePropertyName = TEXT("TransformNoScale");
+			}
 
 			TSharedPtr<IPropertyHandle> ValuePropertyHandle = DetailLayout.GetProperty(ValuePropertyName, Proxy->GetClass());
 			if (ValuePropertyHandle)
 			{
-				ValuePropertyHandle->SetPropertyDisplayName(FText::FromName(ControlElement->GetDisplayName()));
+				ValuePropertyHandle->SetPropertyDisplayName(FText::FromName(Proxy->GetName()));
 			}
 
 			URigHierarchy* Hierarchy = Proxy->ControlRig->GetHierarchy();
 			Hierarchy->ForEach<FRigControlElement>([Hierarchy, Proxy, &Category, this](FRigControlElement* ControlElement) -> bool
-            {
+	      	{
 				FName ParentControlName = NAME_None;
 				FRigControlElement* ParentControlElement = Cast<FRigControlElement>(Hierarchy->GetFirstParent(ControlElement));
 				if(ParentControlElement)
@@ -176,7 +225,6 @@ public:
 
 							FAddPropertyParams Params;
 							Params.CreateCategoryNodes(false);
-
 							IDetailPropertyRow* NestedRow = Category.AddExternalObjectProperty(
 								NestedProxies,
 								PropertyName,
@@ -255,17 +303,72 @@ void SControlRigEditModeTools::Construct(const FArguments& InArgs, TSharedPtr<FC
 		DetailsViewArgs.bCustomNameAreaLocation = true;
 		DetailsViewArgs.bCustomFilterAreaLocation = true;
 		DetailsViewArgs.NameAreaSettings = FDetailsViewArgs::HideNameArea;
-		DetailsViewArgs.bAllowMultipleTopLevelObjects = true;
+		DetailsViewArgs.bAllowMultipleTopLevelObjects = false;
 		DetailsViewArgs.bShowScrollBar = false; // Don't need to show this, as we are putting it in a scroll box
 	}
 	
 	ModeTools = InEditMode.GetModeManager();
 
-	ControlDetailsView = FModuleManager::GetModuleChecked<FPropertyEditorModule>("PropertyEditor").CreateDetailView(DetailsViewArgs);
-	ControlDetailsView->SetKeyframeHandler(SharedThis(this));
-	ControlDetailsView->SetIsPropertyVisibleDelegate(FIsPropertyVisible::CreateSP(this, &SControlRigEditModeTools::ShouldShowPropertyOnDetailCustomization));
-	ControlDetailsView->SetIsPropertyReadOnlyDelegate(FIsPropertyReadOnly::CreateSP(this, &SControlRigEditModeTools::IsReadOnlyPropertyOnDetailCustomization));
-	ControlDetailsView->SetGenericLayoutDetailsDelegate(FOnGetDetailCustomizationInstance::CreateStatic(&FControlRigEditModeGenericDetails::MakeInstance, ModeTools));
+	SettingsDetailsView = FModuleManager::GetModuleChecked<FPropertyEditorModule>("PropertyEditor").CreateDetailView(DetailsViewArgs);
+	SettingsDetailsView->SetKeyframeHandler(SharedThis(this));
+	SettingsDetailsView->SetIsPropertyVisibleDelegate(FIsPropertyVisible::CreateSP(this, &SControlRigEditModeTools::ShouldShowPropertyOnDetailCustomization));
+	SettingsDetailsView->SetIsPropertyReadOnlyDelegate(FIsPropertyReadOnly::CreateSP(this, &SControlRigEditModeTools::IsReadOnlyPropertyOnDetailCustomization));
+	SettingsDetailsView->SetGenericLayoutDetailsDelegate(FOnGetDetailCustomizationInstance::CreateStatic(&FControlRigEditModeGenericDetails::MakeInstance, ModeTools));
+
+	ControlEulerTransformDetailsView = FModuleManager::GetModuleChecked<FPropertyEditorModule>("PropertyEditor").CreateDetailView(DetailsViewArgs);
+	ControlEulerTransformDetailsView->SetKeyframeHandler(SharedThis(this));
+	ControlEulerTransformDetailsView->SetIsPropertyVisibleDelegate(FIsPropertyVisible::CreateSP(this, &SControlRigEditModeTools::ShouldShowPropertyOnDetailCustomization));
+	ControlEulerTransformDetailsView->SetIsPropertyReadOnlyDelegate(FIsPropertyReadOnly::CreateSP(this, &SControlRigEditModeTools::IsReadOnlyPropertyOnDetailCustomization));
+	ControlEulerTransformDetailsView->SetGenericLayoutDetailsDelegate(FOnGetDetailCustomizationInstance::CreateStatic(&FControlRigEditModeGenericDetails::MakeInstance, ModeTools));
+
+	ControlTransformDetailsView = FModuleManager::GetModuleChecked<FPropertyEditorModule>("PropertyEditor").CreateDetailView(DetailsViewArgs);
+	ControlTransformDetailsView->SetKeyframeHandler(SharedThis(this));
+	ControlTransformDetailsView->SetIsPropertyVisibleDelegate(FIsPropertyVisible::CreateSP(this, &SControlRigEditModeTools::ShouldShowPropertyOnDetailCustomization));
+	ControlTransformDetailsView->SetIsPropertyReadOnlyDelegate(FIsPropertyReadOnly::CreateSP(this, &SControlRigEditModeTools::IsReadOnlyPropertyOnDetailCustomization));
+	ControlTransformDetailsView->SetGenericLayoutDetailsDelegate(FOnGetDetailCustomizationInstance::CreateStatic(&FControlRigEditModeGenericDetails::MakeInstance, ModeTools));
+
+	ControlTransformNoScaleDetailsView = FModuleManager::GetModuleChecked<FPropertyEditorModule>("PropertyEditor").CreateDetailView(DetailsViewArgs);
+	ControlTransformNoScaleDetailsView->SetKeyframeHandler(SharedThis(this));
+	ControlTransformNoScaleDetailsView->SetIsPropertyVisibleDelegate(FIsPropertyVisible::CreateSP(this, &SControlRigEditModeTools::ShouldShowPropertyOnDetailCustomization));
+	ControlTransformNoScaleDetailsView->SetIsPropertyReadOnlyDelegate(FIsPropertyReadOnly::CreateSP(this, &SControlRigEditModeTools::IsReadOnlyPropertyOnDetailCustomization));
+	ControlTransformNoScaleDetailsView->SetGenericLayoutDetailsDelegate(FOnGetDetailCustomizationInstance::CreateStatic(&FControlRigEditModeGenericDetails::MakeInstance, ModeTools));
+
+	ControlFloatDetailsView = FModuleManager::GetModuleChecked<FPropertyEditorModule>("PropertyEditor").CreateDetailView(DetailsViewArgs);
+	ControlFloatDetailsView->SetKeyframeHandler(SharedThis(this));
+	ControlFloatDetailsView->SetIsPropertyVisibleDelegate(FIsPropertyVisible::CreateSP(this, &SControlRigEditModeTools::ShouldShowPropertyOnDetailCustomization));
+	ControlFloatDetailsView->SetIsPropertyReadOnlyDelegate(FIsPropertyReadOnly::CreateSP(this, &SControlRigEditModeTools::IsReadOnlyPropertyOnDetailCustomization));
+	ControlFloatDetailsView->SetGenericLayoutDetailsDelegate(FOnGetDetailCustomizationInstance::CreateStatic(&FControlRigEditModeGenericDetails::MakeInstance, ModeTools));
+
+	ControlEnumDetailsView = FModuleManager::GetModuleChecked<FPropertyEditorModule>("PropertyEditor").CreateDetailView(DetailsViewArgs);
+	ControlEnumDetailsView->SetKeyframeHandler(SharedThis(this));
+	ControlEnumDetailsView->SetIsPropertyVisibleDelegate(FIsPropertyVisible::CreateSP(this, &SControlRigEditModeTools::ShouldShowPropertyOnDetailCustomization));
+	ControlEnumDetailsView->SetIsPropertyReadOnlyDelegate(FIsPropertyReadOnly::CreateSP(this, &SControlRigEditModeTools::IsReadOnlyPropertyOnDetailCustomization));
+	ControlEnumDetailsView->SetGenericLayoutDetailsDelegate(FOnGetDetailCustomizationInstance::CreateStatic(&FControlRigEditModeGenericDetails::MakeInstance, ModeTools));
+
+	ControlIntegerDetailsView = FModuleManager::GetModuleChecked<FPropertyEditorModule>("PropertyEditor").CreateDetailView(DetailsViewArgs);
+	ControlIntegerDetailsView->SetKeyframeHandler(SharedThis(this));
+	ControlIntegerDetailsView->SetIsPropertyVisibleDelegate(FIsPropertyVisible::CreateSP(this, &SControlRigEditModeTools::ShouldShowPropertyOnDetailCustomization));
+	ControlIntegerDetailsView->SetIsPropertyReadOnlyDelegate(FIsPropertyReadOnly::CreateSP(this, &SControlRigEditModeTools::IsReadOnlyPropertyOnDetailCustomization));
+	ControlIntegerDetailsView->SetGenericLayoutDetailsDelegate(FOnGetDetailCustomizationInstance::CreateStatic(&FControlRigEditModeGenericDetails::MakeInstance, ModeTools));
+
+	ControlBoolDetailsView = FModuleManager::GetModuleChecked<FPropertyEditorModule>("PropertyEditor").CreateDetailView(DetailsViewArgs);
+	ControlBoolDetailsView->SetKeyframeHandler(SharedThis(this));
+	ControlBoolDetailsView->SetIsPropertyVisibleDelegate(FIsPropertyVisible::CreateSP(this, &SControlRigEditModeTools::ShouldShowPropertyOnDetailCustomization));
+	ControlBoolDetailsView->SetIsPropertyReadOnlyDelegate(FIsPropertyReadOnly::CreateSP(this, &SControlRigEditModeTools::IsReadOnlyPropertyOnDetailCustomization));
+	ControlBoolDetailsView->SetGenericLayoutDetailsDelegate(FOnGetDetailCustomizationInstance::CreateStatic(&FControlRigEditModeGenericDetails::MakeInstance, ModeTools));
+
+	ControlVectorDetailsView = FModuleManager::GetModuleChecked<FPropertyEditorModule>("PropertyEditor").CreateDetailView(DetailsViewArgs);
+	ControlVectorDetailsView->SetKeyframeHandler(SharedThis(this));
+	ControlVectorDetailsView->SetIsPropertyVisibleDelegate(FIsPropertyVisible::CreateSP(this, &SControlRigEditModeTools::ShouldShowPropertyOnDetailCustomization));
+	ControlVectorDetailsView->SetIsPropertyReadOnlyDelegate(FIsPropertyReadOnly::CreateSP(this, &SControlRigEditModeTools::IsReadOnlyPropertyOnDetailCustomization));
+	ControlVectorDetailsView->SetGenericLayoutDetailsDelegate(FOnGetDetailCustomizationInstance::CreateStatic(&FControlRigEditModeGenericDetails::MakeInstance, ModeTools));
+
+	ControlVector2DDetailsView = FModuleManager::GetModuleChecked<FPropertyEditorModule>("PropertyEditor").CreateDetailView(DetailsViewArgs);
+	ControlVector2DDetailsView->SetKeyframeHandler(SharedThis(this));
+	ControlVector2DDetailsView->SetIsPropertyVisibleDelegate(FIsPropertyVisible::CreateSP(this, &SControlRigEditModeTools::ShouldShowPropertyOnDetailCustomization));
+	ControlVector2DDetailsView->SetIsPropertyReadOnlyDelegate(FIsPropertyReadOnly::CreateSP(this, &SControlRigEditModeTools::IsReadOnlyPropertyOnDetailCustomization));
+	ControlVector2DDetailsView->SetGenericLayoutDetailsDelegate(FOnGetDetailCustomizationInstance::CreateStatic(&FControlRigEditModeGenericDetails::MakeInstance, ModeTools));
+
 
 	RigOptionsDetailsView = FModuleManager::GetModuleChecked<FPropertyEditorModule>("PropertyEditor").CreateDetailView(DetailsViewArgs);
 	RigOptionsDetailsView->SetKeyframeHandler(SharedThis(this));
@@ -309,7 +412,52 @@ void SControlRigEditModeTools::Construct(const FArguments& InArgs, TSharedPtr<FC
 			+SVerticalBox::Slot()
 			.AutoHeight()
 			[
-				ControlDetailsView.ToSharedRef()
+				SettingsDetailsView.ToSharedRef()
+			]
+			+ SVerticalBox::Slot()
+			.AutoHeight()
+			[
+				ControlEulerTransformDetailsView.ToSharedRef()
+			]
+			+ SVerticalBox::Slot()
+			.AutoHeight()
+			[
+				ControlTransformDetailsView.ToSharedRef()
+			]
+			+ SVerticalBox::Slot()
+			.AutoHeight()
+			[
+				ControlTransformNoScaleDetailsView.ToSharedRef()
+			]
+			+ SVerticalBox::Slot()
+			.AutoHeight()
+			[
+				ControlBoolDetailsView.ToSharedRef()
+			]
+			+ SVerticalBox::Slot()
+			.AutoHeight()
+			[
+				ControlIntegerDetailsView.ToSharedRef()
+			]
+			+ SVerticalBox::Slot()
+			.AutoHeight()
+			[
+				ControlEnumDetailsView.ToSharedRef()
+			]
+			+ SVerticalBox::Slot()
+			.AutoHeight()
+			[
+				ControlVectorDetailsView.ToSharedRef()
+			]
+			+ SVerticalBox::Slot()
+			.AutoHeight()
+			[
+				ControlVector2DDetailsView.ToSharedRef()
+			]
+			+ SVerticalBox::Slot()
+			.AutoHeight()
+			[
+				ControlFloatDetailsView.ToSharedRef()
 			]
 
 			+SVerticalBox::Slot()
@@ -395,12 +543,88 @@ void SControlRigEditModeTools::Construct(const FArguments& InArgs, TSharedPtr<FC
 	HierarchyTreeView->RefreshTreeView(true);
 }
 
-void SControlRigEditModeTools::SetDetailsObjects(const TArray<TWeakObjectPtr<>>& InObjects)
+void SControlRigEditModeTools::SetSettingsDetailsObject(const TWeakObjectPtr<>& InObject)
+{
+	if (SettingsDetailsView)
+	{
+		TArray<TWeakObjectPtr<>> Objects;
+		Objects.Add(InObject);
+		SettingsDetailsView->SetObjects(Objects);
+
+	}
+}
+void SControlRigEditModeTools::SetEulerTransformDetailsObjects(const TArray<TWeakObjectPtr<>>& InObjects)
+{
+	if (ControlEulerTransformDetailsView)
+	{
+		ControlEulerTransformDetailsView->SetObjects(InObjects);
+	}
+};
+
+void SControlRigEditModeTools::SetTransformDetailsObjects(const TArray<TWeakObjectPtr<>>& InObjects)
+{
+	if (ControlTransformDetailsView)
+	{
+		ControlTransformDetailsView->SetObjects(InObjects);
+	}
+}
+
+void SControlRigEditModeTools::SetTransformNoScaleDetailsObjects(const TArray<TWeakObjectPtr<>>& InObjects)
+{
+	if (ControlTransformNoScaleDetailsView)
+	{
+		ControlTransformNoScaleDetailsView->SetObjects(InObjects);
+	}
+}
+
+void SControlRigEditModeTools::SetFloatDetailsObjects(const TArray<TWeakObjectPtr<>>& InObjects)
+{
+	if (ControlFloatDetailsView)
+	{
+		ControlFloatDetailsView->SetObjects(InObjects);
+	}
+}
+
+void SControlRigEditModeTools::SetBoolDetailsObjects(const TArray<TWeakObjectPtr<>>& InObjects)
+{
+	if (ControlBoolDetailsView)
+	{
+		ControlBoolDetailsView->SetObjects(InObjects);
+	}
+}
+
+void SControlRigEditModeTools::SetIntegerDetailsObjects(const TArray<TWeakObjectPtr<>>& InObjects)
+{
+	if (ControlIntegerDetailsView)
+	{
+		ControlIntegerDetailsView->SetObjects(InObjects);
+	}
+}
+void SControlRigEditModeTools::SetEnumDetailsObjects(const TArray<TWeakObjectPtr<>>& InObjects)
+{
+	if (ControlVectorDetailsView)
+	{
+		ControlVectorDetailsView->SetObjects(InObjects);
+	}
+}
+
+void SControlRigEditModeTools::SetVectorDetailsObjects(const TArray<TWeakObjectPtr<>>& InObjects)
 {
-	ControlDetailsView->SetObjects(InObjects);
+	if (ControlEnumDetailsView)
+	{
+		ControlEnumDetailsView->SetObjects(InObjects);
+	}
+}
+
+void SControlRigEditModeTools::SetVector2DDetailsObjects(const TArray<TWeakObjectPtr<>>& InObjects)
+{
+	if (ControlVector2DDetailsView)
+	{
+		ControlVector2DDetailsView->SetObjects(InObjects);
+	}
 }
 
-void SControlRigEditModeTools::SetSequencer(TWeakPtr<ISequencer> InSequencer)
+void SControlRigEditModeTools::SControlRigEditModeTools::SetSequencer(TWeakPtr<ISequencer> InSequencer)
 {
 	WeakSequencer = InSequencer.Pin();
 }
diff --git a/Engine/Plugins/Animation/ControlRig/Source/ControlRigEditor/Private/EditMode/SControlRigEditModeTools.h b/Engine/Plugins/Animation/ControlRig/Source/ControlRigEditor/Private/EditMode/SControlRigEditModeTools.h
--- a/Engine/Plugins/Animation/ControlRig/Source/ControlRigEditor/Private/EditMode/SControlRigEditModeTools.h
+++ b/Engine/Plugins/Animation/ControlRig/Source/ControlRigEditor/Private/EditMode/SControlRigEditModeTools.h
@@ -31,7 +31,17 @@ public:
 	void Construct(const FArguments& InArgs, TSharedPtr<FControlRigEditModeToolkit> InOwningToolkit, FControlRigEditMode& InEditMode, UWorld* InWorld);
 
 	/** Set the objects to be displayed in the details panel */
-	void SetDetailsObjects(const TArray<TWeakObjectPtr<>>& InObjects);
+	void SetSettingsDetailsObject(const TWeakObjectPtr<>& InObject);
+	void SetEulerTransformDetailsObjects(const TArray<TWeakObjectPtr<>>& InObjects);
+	void SetTransformDetailsObjects(const TArray<TWeakObjectPtr<>>& InObjects);
+	void SetTransformNoScaleDetailsObjects(const TArray<TWeakObjectPtr<>>& InObjects);
+	void SetFloatDetailsObjects(const TArray<TWeakObjectPtr<>>& InObjects);
+	void SetBoolDetailsObjects(const TArray<TWeakObjectPtr<>>& InObjects);
+	void SetIntegerDetailsObjects(const TArray<TWeakObjectPtr<>>& InObjects);
+	void SetEnumDetailsObjects(const TArray<TWeakObjectPtr<>>& InObjects);
+	void SetVectorDetailsObjects(const TArray<TWeakObjectPtr<>>& InObjects);
+	void SetVector2DDetailsObjects(const TArray<TWeakObjectPtr<>>& InObjects);
+
 
 	/** Set the sequencer we are bound to */
 	void SetSequencer(TWeakPtr<ISequencer> InSequencer);
@@ -51,8 +61,17 @@ private:
 	/** Sequencer we are currently bound to */
 	TWeakPtr<ISequencer> WeakSequencer;
 
-	/** The details view we do most of our work within */
-	TSharedPtr<IDetailsView> ControlDetailsView;
+	/** The details views we do most of our work within */
+	TSharedPtr<IDetailsView> SettingsDetailsView;
+	TSharedPtr<IDetailsView> ControlEulerTransformDetailsView;
+	TSharedPtr<IDetailsView> ControlTransformDetailsView;
+	TSharedPtr<IDetailsView> ControlTransformNoScaleDetailsView;
+	TSharedPtr<IDetailsView> ControlFloatDetailsView;
+	TSharedPtr<IDetailsView> ControlBoolDetailsView;
+	TSharedPtr<IDetailsView> ControlIntegerDetailsView;
+	TSharedPtr<IDetailsView> ControlEnumDetailsView;
+	TSharedPtr<IDetailsView> ControlVector2DDetailsView;
+	TSharedPtr<IDetailsView> ControlVectorDetailsView;
 
 	/** Expander to interact with the options of the rig  */
 	TSharedPtr<SExpandableArea> RigOptionExpander;
