Control Rig: Add a scrollbar for the snapper

diff --git a/Engine/Plugins/Animation/ControlRig/Source/ControlRigEditor/Private/EditMode/SControlRigSnapper.cpp b/Engine/Plugins/Animation/ControlRig/Source/ControlRigEditor/Private/EditMode/SControlRigSnapper.cpp
--- a/Engine/Plugins/Animation/ControlRig/Source/ControlRigEditor/Private/EditMode/SControlRigSnapper.cpp
+++ b/Engine/Plugins/Animation/ControlRig/Source/ControlRigEditor/Private/EditMode/SControlRigSnapper.cpp
@@ -26,6 +26,7 @@
 #include "Widgets/Text/STextBlock.h"
 #include "Widgets/Views/STableRow.h"
 #include "Widgets/Input/SButton.h"
+#include "Widgets/Layout/SScrollBox.h"
 #include "Tools/ControlRigSnapSettings.h"
 
 #define LOCTEXT_NAMESPACE "ControlRigSnapper"
@@ -150,10 +151,8 @@ void SControlRigSnapper::Construct(const FArguments& InArgs)
 
 	ChildSlot
 		[
-			SNew(SBorder)
-			.HAlign(HAlign_Fill)
-		.VAlign(VAlign_Fill)
-		.Padding(FMargin(10.0,5.0,10.0,5.0))
+		SNew(SScrollBox)
+		+ SScrollBox::Slot()
 		[
 			SNew(SVerticalBox)
 			+ SVerticalBox::Slot()
