Control Rig: Re-add th empty line in Rig Hierarchy

diff --git a/Engine/Plugins/Animation/ControlRig/Source/ControlRigEditor/Private/Editor/SRigHierarchyTreeView.cpp b/Engine/Plugins/Animation/ControlRig/Source/ControlRigEditor/Private/Editor/SRigHierarchyTreeView.cpp
--- a/Engine/Plugins/Animation/ControlRig/Source/ControlRigEditor/Private/Editor/SRigHierarchyTreeView.cpp
+++ b/Engine/Plugins/Animation/ControlRig/Source/ControlRigEditor/Private/Editor/SRigHierarchyTreeView.cpp
@@ -374,6 +374,11 @@ bool SRigHierarchyTreeView::AddElement(const FRigBaseElement* InElement)
 	return true;
 }
 
+void SRigHierarchyTreeView::AddSpacerElement()
+{
+	AddElement(FRigElementKey(), FRigElementKey());
+}
+
 bool SRigHierarchyTreeView::ReparentElement(FRigElementKey InKey, FRigElementKey InParentKey)
 {
 	if (!InKey.IsValid() || InKey == InParentKey)
@@ -509,6 +514,11 @@ void SRigHierarchyTreeView::RefreshTreeView(bool bRebuildContent)
 					return Delegates.OnCompareKeys.Execute(A->Key, B->Key);
 				});
 			}
+
+			if (RootElements.Num() > 0)
+			{
+				AddSpacerElement();
+			}
 		}
 	}
 	else
@@ -521,6 +531,7 @@ void SRigHierarchyTreeView::RefreshTreeView(bool bRebuildContent)
 			{
 				return InElement.Get()->Key == FRigElementKey();
 			});
+			AddSpacerElement();
 		}
 	}
 
diff --git a/Engine/Plugins/Animation/ControlRig/Source/ControlRigEditor/Private/Editor/SRigHierarchyTreeView.h b/Engine/Plugins/Animation/ControlRig/Source/ControlRigEditor/Private/Editor/SRigHierarchyTreeView.h
--- a/Engine/Plugins/Animation/ControlRig/Source/ControlRigEditor/Private/Editor/SRigHierarchyTreeView.h
+++ b/Engine/Plugins/Animation/ControlRig/Source/ControlRigEditor/Private/Editor/SRigHierarchyTreeView.h
@@ -248,6 +248,7 @@ public:
 	static TSharedPtr<FRigTreeElement> FindElement(const FRigElementKey& InElementKey, TSharedPtr<FRigTreeElement> CurrentItem);
 	bool AddElement(FRigElementKey InKey, FRigElementKey InParentKey = FRigElementKey());
 	bool AddElement(const FRigBaseElement* InElement);
+	void AddSpacerElement();
 	bool ReparentElement(FRigElementKey InKey, FRigElementKey InParentKey);
 	bool RemoveElement(FRigElementKey InKey);
 	void RefreshTreeView(bool bRebuildContent = true);
