Control Rig: Search filter in function arguments not showing matching types

diff --git a/Engine/Plugins/Animation/ControlRig/Source/ControlRigEditor/Private/ControlRigGraphDetails.cpp b/Engine/Plugins/Animation/ControlRig/Source/ControlRigEditor/Private/ControlRigGraphDetails.cpp
--- a/Engine/Plugins/Animation/ControlRig/Source/ControlRigEditor/Private/ControlRigGraphDetails.cpp
+++ b/Engine/Plugins/Animation/ControlRig/Source/ControlRigEditor/Private/ControlRigGraphDetails.cpp
@@ -30,9 +30,11 @@
 FControlRigArgumentGroupLayout::FControlRigArgumentGroupLayout(
 	URigVMGraph* InGraph, 
 	UControlRigBlueprint* InBlueprint, 
+	TWeakPtr<IControlRigEditor> InEditor,
 	bool bInputs)
 	: GraphPtr(InGraph)
 	, ControlRigBlueprintPtr(InBlueprint)
+	, ControlRigEditorPtr(InEditor)
 	, bIsInputGroup(bInputs)
 {
 	if (ControlRigBlueprintPtr.IsValid())
@@ -65,7 +67,8 @@ void FControlRigArgumentGroupLayout::GenerateChildContent(IDetailChildrenBuilder
 					TSharedRef<class FControlRigArgumentLayout> ControlRigArgumentLayout = MakeShareable(new FControlRigArgumentLayout(
 						Pin,
 						Graph,
-						ControlRigBlueprintPtr.Get()
+						ControlRigBlueprintPtr.Get(),
+						ControlRigEditorPtr
 					));
 					ChildrenBuilder.AddCustomBuilder(ControlRigArgumentLayout);
 					WasContentAdded = true;
@@ -134,6 +137,13 @@ void FControlRigArgumentLayout::GenerateHeaderRowContent(FDetailWidgetRow& NodeR
 	ETypeTreeFilter TypeTreeFilter = ETypeTreeFilter::None;
 	TypeTreeFilter |= ETypeTreeFilter::AllowExec;
 
+	TSharedPtr<IPinTypeSelectorFilter> CustomPinTypeFilter;
+	if (ControlRigEditorPtr.IsValid())
+	{
+		CustomPinTypeFilter = ControlRigEditorPtr.Pin()->GetImportedPinTypeSelectorFilter();
+	}
+	
+
 	NodeRow
 		.NameContent()
 		[
@@ -185,6 +195,7 @@ void FControlRigArgumentLayout::GenerateHeaderRowContent(FDetailWidgetRow& NodeR
 				.TypeTreeFilter(TypeTreeFilter)
 				.bAllowArrays(!ShouldPinBeReadOnly())
 				.IsEnabled(!ShouldPinBeReadOnly(true))
+				.CustomFilter(CustomPinTypeFilter)
 				.Font(IDetailLayoutBuilder::GetDetailFont())
 			]
 			+ SHorizontalBox::Slot()
@@ -643,7 +654,8 @@ void FControlRigGraphDetails::CustomizeDetails(IDetailLayoutBuilder& DetailLayou
 	IDetailCategoryBuilder& InputsCategory = DetailLayout.EditCategory("Inputs", LOCTEXT("FunctionDetailsInputs", "Inputs"));
 	TSharedRef<FControlRigArgumentGroupLayout> InputArgumentGroup = MakeShareable(new FControlRigArgumentGroupLayout(
 		Model, 
-		Blueprint, 
+		Blueprint,
+		ControlRigEditorPtr,
 		true));
 	InputsCategory.AddCustomBuilder(InputArgumentGroup);
 
@@ -673,7 +685,8 @@ void FControlRigGraphDetails::CustomizeDetails(IDetailLayoutBuilder& DetailLayou
 	IDetailCategoryBuilder& OutputsCategory = DetailLayout.EditCategory("Outputs", LOCTEXT("FunctionDetailsOutputs", "Outputs"));
 	TSharedRef<FControlRigArgumentGroupLayout> OutputArgumentGroup = MakeShareable(new FControlRigArgumentGroupLayout(
 		Model, 
-		Blueprint, 
+		Blueprint,
+		ControlRigEditorPtr,
 		false));
 	OutputsCategory.AddCustomBuilder(OutputArgumentGroup);
 
diff --git a/Engine/Plugins/Animation/ControlRig/Source/ControlRigEditor/Private/ControlRigGraphDetails.h b/Engine/Plugins/Animation/ControlRig/Source/ControlRigEditor/Private/ControlRigGraphDetails.h
--- a/Engine/Plugins/Animation/ControlRig/Source/ControlRigEditor/Private/ControlRigGraphDetails.h
+++ b/Engine/Plugins/Animation/ControlRig/Source/ControlRigEditor/Private/ControlRigGraphDetails.h
@@ -22,7 +22,8 @@ class FControlRigArgumentGroupLayout : public IDetailCustomNodeBuilder, public T
 public:
 	FControlRigArgumentGroupLayout(
 		URigVMGraph* InGraph, 
-		UControlRigBlueprint* InBlueprint, 
+		UControlRigBlueprint* InBlueprint,
+		TWeakPtr<IControlRigEditor> InEditor,
 		bool bInputs);
 	virtual ~FControlRigArgumentGroupLayout();
 
@@ -42,6 +43,7 @@ private:
 
 	TWeakObjectPtr<URigVMGraph> GraphPtr;
 	TWeakObjectPtr<UControlRigBlueprint> ControlRigBlueprintPtr;
+	TWeakPtr<IControlRigEditor> ControlRigEditorPtr;
 	bool bIsInputGroup;
 	FSimpleDelegate OnRebuildChildren;
 };
@@ -53,10 +55,12 @@ public:
 	FControlRigArgumentLayout(
 		URigVMPin* InPin, 
 		URigVMGraph* InGraph, 
-		UControlRigBlueprint* InBlueprint)
+		UControlRigBlueprint* InBlueprint,
+		TWeakPtr<IControlRigEditor> InEditor)
 		: PinPtr(InPin)
 		, GraphPtr(InGraph)
 		, ControlRigBlueprintPtr(InBlueprint)
+		, ControlRigEditorPtr(InEditor)
 		, NameValidator(InBlueprint, InGraph, InPin->GetFName())
 	{}
 
@@ -102,6 +106,9 @@ private:
 	/** The blueprint we are editing */
 	TWeakObjectPtr<UControlRigBlueprint> ControlRigBlueprintPtr;
 
+	/** The editor we are editing */
+	TWeakPtr<IControlRigEditor> ControlRigEditorPtr;
+
 	/** Holds a weak pointer to the argument name widget, used for error notifications */
 	TWeakPtr<SEditableTextBox> ArgumentNameWidget;
 
