Sequencer: Control Rig: Filtering controls on additive sections wouldn't work well if there are any bool/int controls like in Metahuman rigs
#trivial

diff --git a/Engine/Plugins/Animation/ControlRig/Source/ControlRig/Private/Sequencer/MovieSceneControlRigParameterSection.cpp b/Engine/Plugins/Animation/ControlRig/Source/ControlRig/Private/Sequencer/MovieSceneControlRigParameterSection.cpp
--- a/Engine/Plugins/Animation/ControlRig/Source/ControlRig/Private/Sequencer/MovieSceneControlRigParameterSection.cpp
+++ b/Engine/Plugins/Animation/ControlRig/Source/ControlRig/Private/Sequencer/MovieSceneControlRigParameterSection.cpp
@@ -1212,6 +1212,7 @@ void UMovieSceneControlRigParameterSection::ReconstructChannelProxy()
 		}
 			
 		int32 ControlIndex = 0; 
+		int32 MaskIndex = 0;
 		int32 TotalIndex = 0; 
 		int32 FloatChannelIndex = 0;
 		int32 BoolChannelIndex = 0;
@@ -1240,150 +1241,91 @@ void UMovieSceneControlRigParameterSection::ReconstructChannelProxy()
 
 			switch (ControlElement->Settings.ControlType)
 			{
-			case ERigControlType::Float:
-			{
-				for (FScalarParameterNameAndCurve& Scalar : GetScalarParameterNamesAndCurves())
+				case ERigControlType::Float:
 				{
-					if (ControlElement->GetName() == Scalar.ParameterName)
+					for (FScalarParameterNameAndCurve& Scalar : GetScalarParameterNamesAndCurves())
 					{
-						bool bEnabled = ControlsMask[ControlIndex];
-
-						FText Group;
-						if(ParentControlElement)
+						if (ControlElement->GetName() == Scalar.ParameterName)
 						{
-							switch(ParentControlElement->Settings.ControlType)
+							bool bEnabled =ControlsMask[MaskIndex];
+
+							FText Group;
+							if(ParentControlElement)
 							{
-								case ERigControlType::Position:
-								case ERigControlType::Scale:
-								case ERigControlType::Rotator:
-								case ERigControlType::Transform:
-								case ERigControlType::EulerTransform:
-								case ERigControlType::TransformNoScale:
+								switch(ParentControlElement->Settings.ControlType)
 								{
-									Group = FText::FromName(ParentControlElement->GetDisplayName());
-									break;
+									case ERigControlType::Position:
+									case ERigControlType::Scale:
+									case ERigControlType::Rotator:
+									case ERigControlType::Transform:
+									case ERigControlType::EulerTransform:
+									case ERigControlType::TransformNoScale:
+									{
+										Group = FText::FromName(ParentControlElement->GetDisplayName());
+										break;
+									}
 								}
 							}
-						}
 
-						if (Group.IsEmpty())
-						{
-							ControlChannelMap.Add(Scalar.ParameterName, FChannelMapInfo(ControlIndex, TotalIndex, FloatChannelIndex));
-							Group = FText::FromName(ControlElement->GetDisplayName());
-							ControlIndex++;  //up the index only if no parent
-						}
-						else
-						{
-							FChannelMapInfo* pChannelIndex = ControlChannelMap.Find(ParentControlName);
-							if (pChannelIndex)
-							{
-								ControlChannelMap.Add(Scalar.ParameterName, FChannelMapInfo(ControlIndex, TotalIndex, FloatChannelIndex,pChannelIndex->ControlIndex));
-							}
-							else
+							if (Group.IsEmpty())
 							{
 								ControlChannelMap.Add(Scalar.ParameterName, FChannelMapInfo(ControlIndex, TotalIndex, FloatChannelIndex));
-
+								Group = FText::FromName(ControlElement->GetDisplayName());
+								ControlIndex++;  //up the index only if no parent
 							}
-						}
-
-						FParameterFloatChannelEditorData EditorData(ControlRig, Scalar.ParameterName, bEnabled, Group, TotalIndex);
-						EditorData.MetaData.DisplayText = FText::FromName(ControlElement->GetDisplayName());
-						Channels.Add(Scalar.ParameterCurve, EditorData.MetaData, EditorData.ExternalValues);
-						FloatChannelIndex += 1;
-						TotalIndex += 1;
-						break;
-					}
-				}
-				break;
-			}
-			case ERigControlType::Bool:
-			{
-				for (FBoolParameterNameAndCurve& Bool : GetBoolParameterNamesAndCurves())
-				{
-					if (ControlElement->GetName() == Bool.ParameterName)
-					{
-						bool bEnabled = ControlsMask[ControlIndex];
-
-						FText Group;
-						if(ParentControlElement)
-						{
-							switch(ParentControlElement->Settings.ControlType)
+							else
 							{
-								case ERigControlType::Position:
-								case ERigControlType::Scale:
-								case ERigControlType::Rotator:
-								case ERigControlType::Transform:
-								case ERigControlType::EulerTransform:
-								case ERigControlType::TransformNoScale:
+								FChannelMapInfo* pChannelIndex = ControlChannelMap.Find(ParentControlName);
+								if (pChannelIndex)
 								{
-									Group = FText::FromName(ParentControlElement->GetDisplayName());
-									break;
+									ControlChannelMap.Add(Scalar.ParameterName, FChannelMapInfo(ControlIndex, TotalIndex, FloatChannelIndex,pChannelIndex->ControlIndex));
 								}
-							}
-						}
+								else
+								{
+									ControlChannelMap.Add(Scalar.ParameterName, FChannelMapInfo(ControlIndex, TotalIndex, FloatChannelIndex));
 
-						if (Group.IsEmpty())
-						{
-							ControlChannelMap.Add(Bool.ParameterName, FChannelMapInfo(ControlIndex, TotalIndex, BoolChannelIndex, INDEX_NONE,BoolChannelTypeName));
-							Group = FText::FromName(ControlElement->GetDisplayName());
-							ControlIndex++;  //up the index only if no parent
-						}
-						else
-						{
-							FChannelMapInfo* pChannelIndex = ControlChannelMap.Find(ParentControlName);
-							if (pChannelIndex)
-							{
-								ControlChannelMap.Add(Bool.ParameterName, FChannelMapInfo(ControlIndex, TotalIndex, BoolChannelIndex,pChannelIndex->ControlIndex,BoolChannelTypeName));
-							}
-							else
-							{
-								ControlChannelMap.Add(Bool.ParameterName, FChannelMapInfo(ControlIndex, TotalIndex, BoolChannelIndex,INDEX_NONE, BoolChannelTypeName));
+								}
 							}
-						}
 
-						FMovieSceneChannelMetaData MetaData(Bool.ParameterName, Group, Group, bEnabled);
-						MetaData.DisplayText = FText::FromName(ControlElement->GetDisplayName());
-						MetaData.SortOrder = TotalIndex++;
-						BoolChannelIndex += 1;
-						// Prevent single channels from collapsing to the track node
-						MetaData.bCanCollapseToTrack = false;
-						Channels.Add(Bool.ParameterCurve, MetaData, TMovieSceneExternalValue<bool>());
-						break;
+							FParameterFloatChannelEditorData EditorData(ControlRig, Scalar.ParameterName, bEnabled, Group, TotalIndex);
+							EditorData.MetaData.DisplayText = FText::FromName(ControlElement->GetDisplayName());
+							Channels.Add(Scalar.ParameterCurve, EditorData.MetaData, EditorData.ExternalValues);
+							FloatChannelIndex += 1;
+							TotalIndex += 1;
+							break;
+						}
 					}
+					break;
 				}
-				break;
-			}
-			case ERigControlType::Integer:
-			{
-				if (ControlElement->Settings.ControlEnum)
+				case ERigControlType::Bool:
 				{
-					for (FEnumParameterNameAndCurve& Enum : GetEnumParameterNamesAndCurves())
+					for (FBoolParameterNameAndCurve& Bool : GetBoolParameterNamesAndCurves())
 					{
-						if (ControlElement->GetName() == Enum.ParameterName)
+						if (ControlElement->GetName() == Bool.ParameterName)
 						{
-							bool bEnabled = ControlsMask[ControlIndex];
+							bool bEnabled =ControlsMask[MaskIndex];
 
 							FText Group;
-							if (ParentControlElement)
+							if(ParentControlElement)
 							{
-								switch (ParentControlElement->Settings.ControlType)
-								{
-								case ERigControlType::Position:
-								case ERigControlType::Scale:
-								case ERigControlType::Rotator:
-								case ERigControlType::Transform:
-								case ERigControlType::EulerTransform:
-								case ERigControlType::TransformNoScale:
+								switch(ParentControlElement->Settings.ControlType)
 								{
-									Group = FText::FromName(ParentControlElement->GetDisplayName());
-									break;
-								}
+									case ERigControlType::Position:
+									case ERigControlType::Scale:
+									case ERigControlType::Rotator:
+									case ERigControlType::Transform:
+									case ERigControlType::EulerTransform:
+									case ERigControlType::TransformNoScale:
+									{
+										Group = FText::FromName(ParentControlElement->GetDisplayName());
+										break;
+									}
 								}
 							}
 
 							if (Group.IsEmpty())
 							{
-								ControlChannelMap.Add(Enum.ParameterName, FChannelMapInfo(ControlIndex, TotalIndex, EnumChannelIndex,INDEX_NONE, EnumChannelTypeName));
+								ControlChannelMap.Add(Bool.ParameterName, FChannelMapInfo(ControlIndex, TotalIndex, BoolChannelIndex, INDEX_NONE,BoolChannelTypeName));
 								Group = FText::FromName(ControlElement->GetDisplayName());
 								ControlIndex++;  //up the index only if no parent
 							}
@@ -1392,442 +1334,505 @@ void UMovieSceneControlRigParameterSection::ReconstructChannelProxy()
 								FChannelMapInfo* pChannelIndex = ControlChannelMap.Find(ParentControlName);
 								if (pChannelIndex)
 								{
-									ControlChannelMap.Add(Enum.ParameterName, FChannelMapInfo(ControlIndex, TotalIndex, EnumChannelIndex, pChannelIndex->ControlIndex, EnumChannelTypeName));
+									ControlChannelMap.Add(Bool.ParameterName, FChannelMapInfo(ControlIndex, TotalIndex, BoolChannelIndex,pChannelIndex->ControlIndex,BoolChannelTypeName));
 								}
 								else
 								{
-									ControlChannelMap.Add(Enum.ParameterName, FChannelMapInfo(ControlIndex, TotalIndex, EnumChannelIndex, INDEX_NONE, EnumChannelTypeName));
+									ControlChannelMap.Add(Bool.ParameterName, FChannelMapInfo(ControlIndex, TotalIndex, BoolChannelIndex,INDEX_NONE, BoolChannelTypeName));
 								}
 							}
 
-							FMovieSceneChannelMetaData MetaData(Enum.ParameterName, Group, Group, bEnabled);
+							FMovieSceneChannelMetaData MetaData(Bool.ParameterName, Group, Group, bEnabled);
 							MetaData.DisplayText = FText::FromName(ControlElement->GetDisplayName());
-							EnumChannelIndex += 1;
 							MetaData.SortOrder = TotalIndex++;
+							BoolChannelIndex += 1;
 							// Prevent single channels from collapsing to the track node
 							MetaData.bCanCollapseToTrack = false;
-							Channels.Add(Enum.ParameterCurve, MetaData, TMovieSceneExternalValue<uint8>());
+							Channels.Add(Bool.ParameterCurve, MetaData, TMovieSceneExternalValue<bool>());
 							break;
 						}
 					}
+					break;
 				}
-				else
+				case ERigControlType::Integer:
 				{
-					for (FIntegerParameterNameAndCurve& Integer : GetIntegerParameterNamesAndCurves())
+					if (ControlElement->Settings.ControlEnum)
 					{
-						if (ControlElement->GetName() == Integer.ParameterName)
+						for (FEnumParameterNameAndCurve& Enum : GetEnumParameterNamesAndCurves())
 						{
-							bool bEnabled = ControlsMask[ControlIndex];
-
-							FText Group;
-							if (ParentControlElement)
+							if (ControlElement->GetName() == Enum.ParameterName)
 							{
-								switch (ParentControlElement->Settings.ControlType)
+								bool bEnabled =ControlsMask[MaskIndex];
+
+								FText Group;
+								if (ParentControlElement)
 								{
-								case ERigControlType::Position:
-								case ERigControlType::Scale:
-								case ERigControlType::Rotator:
-								case ERigControlType::Transform:
-								case ERigControlType::EulerTransform:
-								case ERigControlType::TransformNoScale:
+									switch (ParentControlElement->Settings.ControlType)
+									{
+									case ERigControlType::Position:
+									case ERigControlType::Scale:
+									case ERigControlType::Rotator:
+									case ERigControlType::Transform:
+									case ERigControlType::EulerTransform:
+									case ERigControlType::TransformNoScale:
+									{
+										Group = FText::FromName(ParentControlElement->GetDisplayName());
+										break;
+									}
+									}
+								}
+
+								if (Group.IsEmpty())
 								{
-									Group = FText::FromName(ParentControlElement->GetDisplayName());
-									break;
+									ControlChannelMap.Add(Enum.ParameterName, FChannelMapInfo(ControlIndex, TotalIndex, EnumChannelIndex,INDEX_NONE, EnumChannelTypeName));
+									Group = FText::FromName(ControlElement->GetDisplayName());
+									ControlIndex++;  //up the index only if no parent
 								}
+								else
+								{
+									FChannelMapInfo* pChannelIndex = ControlChannelMap.Find(ParentControlName);
+									if (pChannelIndex)
+									{
+										ControlChannelMap.Add(Enum.ParameterName, FChannelMapInfo(ControlIndex, TotalIndex, EnumChannelIndex, pChannelIndex->ControlIndex, EnumChannelTypeName));
+									}
+									else
+									{
+										ControlChannelMap.Add(Enum.ParameterName, FChannelMapInfo(ControlIndex, TotalIndex, EnumChannelIndex, INDEX_NONE, EnumChannelTypeName));
+									}
 								}
+
+								FMovieSceneChannelMetaData MetaData(Enum.ParameterName, Group, Group, bEnabled);
+								MetaData.DisplayText = FText::FromName(ControlElement->GetDisplayName());
+								EnumChannelIndex += 1;
+								MetaData.SortOrder = TotalIndex++;
+								// Prevent single channels from collapsing to the track node
+								MetaData.bCanCollapseToTrack = false;
+								Channels.Add(Enum.ParameterCurve, MetaData, TMovieSceneExternalValue<uint8>());
+								break;
 							}
-							if (Group.IsEmpty())
-							{
-								ControlChannelMap.Add(Integer.ParameterName, FChannelMapInfo(ControlIndex, TotalIndex, IntegerChannelIndex,INDEX_NONE,IntegerChannelTypeName));
-								Group = FText::FromName(ControlElement->GetDisplayName());
-								ControlIndex++;  //up the index only if no parent
-							}
-							else
+						}
+					}
+					else
+					{
+						for (FIntegerParameterNameAndCurve& Integer : GetIntegerParameterNamesAndCurves())
+						{
+							if (ControlElement->GetName() == Integer.ParameterName)
 							{
-								FChannelMapInfo* pChannelIndex = ControlChannelMap.Find(ParentControlName);
-								if (pChannelIndex)
+								bool bEnabled =ControlsMask[MaskIndex];
+
+								FText Group;
+								if (ParentControlElement)
+								{
+									switch (ParentControlElement->Settings.ControlType)
+									{
+									case ERigControlType::Position:
+									case ERigControlType::Scale:
+									case ERigControlType::Rotator:
+									case ERigControlType::Transform:
+									case ERigControlType::EulerTransform:
+									case ERigControlType::TransformNoScale:
+									{
+										Group = FText::FromName(ParentControlElement->GetDisplayName());
+										break;
+									}
+									}
+								}
+								if (Group.IsEmpty())
 								{
-									ControlChannelMap.Add(Integer.ParameterName, FChannelMapInfo(ControlIndex, TotalIndex, IntegerChannelIndex, pChannelIndex->ControlIndex,IntegerChannelTypeName));
+									ControlChannelMap.Add(Integer.ParameterName, FChannelMapInfo(ControlIndex, TotalIndex, IntegerChannelIndex,INDEX_NONE,IntegerChannelTypeName));
+									Group = FText::FromName(ControlElement->GetDisplayName());
+									ControlIndex++;  //up the index only if no parent
 								}
 								else
 								{
-									ControlChannelMap.Add(Integer.ParameterName, FChannelMapInfo(ControlIndex, TotalIndex, IntegerChannelIndex, INDEX_NONE, IntegerChannelTypeName));
+									FChannelMapInfo* pChannelIndex = ControlChannelMap.Find(ParentControlName);
+									if (pChannelIndex)
+									{
+										ControlChannelMap.Add(Integer.ParameterName, FChannelMapInfo(ControlIndex, TotalIndex, IntegerChannelIndex, pChannelIndex->ControlIndex,IntegerChannelTypeName));
+									}
+									else
+									{
+										ControlChannelMap.Add(Integer.ParameterName, FChannelMapInfo(ControlIndex, TotalIndex, IntegerChannelIndex, INDEX_NONE, IntegerChannelTypeName));
+									}
 								}
-							}
 
-							FMovieSceneChannelMetaData MetaData(Integer.ParameterName, Group, Group, bEnabled);
-							MetaData.DisplayText = FText::FromName(ControlElement->GetDisplayName());
-							IntegerChannelIndex += 1;
-							MetaData.SortOrder = TotalIndex++;
-							// Prevent single channels from collapsing to the track node
-							MetaData.bCanCollapseToTrack = false;
-							Channels.Add(Integer.ParameterCurve, MetaData, TMovieSceneExternalValue<int32>());
-							break;
+								FMovieSceneChannelMetaData MetaData(Integer.ParameterName, Group, Group, bEnabled);
+								MetaData.DisplayText = FText::FromName(ControlElement->GetDisplayName());
+								IntegerChannelIndex += 1;
+								MetaData.SortOrder = TotalIndex++;
+								// Prevent single channels from collapsing to the track node
+								MetaData.bCanCollapseToTrack = false;
+								Channels.Add(Integer.ParameterCurve, MetaData, TMovieSceneExternalValue<int32>());
+								break;
+							}
 						}
-					}
 
+					}
+					break;
 				}
-				break;
-			}
-			case ERigControlType::Vector2D:
-			{
-				for (FVector2DParameterNameAndCurves& Vector2D : GetVector2DParameterNamesAndCurves())
+				case ERigControlType::Vector2D:
 				{
-					if (ControlElement->GetName() == Vector2D.ParameterName)
+					for (FVector2DParameterNameAndCurves& Vector2D : GetVector2DParameterNamesAndCurves())
 					{
-						ControlChannelMap.Add(Vector2D.ParameterName, FChannelMapInfo(ControlIndex, TotalIndex, FloatChannelIndex));
-						bool bEnabled = ControlsMask[ControlIndex++];
-						FText Group = FText::FromName(ControlElement->GetDisplayName());
-						FParameterVectorChannelEditorData EditorData(ControlRig, Vector2D.ParameterName, bEnabled, Group, TotalIndex, 2);
-						Channels.Add(Vector2D.XCurve, EditorData.MetaData[0], EditorData.ExternalValues[0]);
-						Channels.Add(Vector2D.YCurve, EditorData.MetaData[1], EditorData.ExternalValues[1]);
-						FloatChannelIndex += 2;
-						TotalIndex += 2;
-						break;
+						if (ControlElement->GetName() == Vector2D.ParameterName)
+						{
+							ControlChannelMap.Add(Vector2D.ParameterName, FChannelMapInfo(ControlIndex, TotalIndex, FloatChannelIndex));
+							ControlIndex++;
+							bool bEnabled =ControlsMask[MaskIndex];
+							FText Group = FText::FromName(ControlElement->GetDisplayName());
+							FParameterVectorChannelEditorData EditorData(ControlRig, Vector2D.ParameterName, bEnabled, Group, TotalIndex, 2);
+							Channels.Add(Vector2D.XCurve, EditorData.MetaData[0], EditorData.ExternalValues[0]);
+							Channels.Add(Vector2D.YCurve, EditorData.MetaData[1], EditorData.ExternalValues[1]);
+							FloatChannelIndex += 2;
+							TotalIndex += 2;
+							break;
+						}
 					}
+					break;
 				}
-				break;
-			}
-			case ERigControlType::Position:
-			case ERigControlType::Scale:
-			case ERigControlType::Rotator:
-			{
-				for (FVectorParameterNameAndCurves& Vector : GetVectorParameterNamesAndCurves())
+				case ERigControlType::Position:
+				case ERigControlType::Scale:
+				case ERigControlType::Rotator:
 				{
-					if (ControlElement->GetName() == Vector.ParameterName)
+					for (FVectorParameterNameAndCurves& Vector : GetVectorParameterNamesAndCurves())
 					{
-						if (ControlElement->Settings.ControlType == ERigControlType::Scale)
+						if (ControlElement->GetName() == Vector.ParameterName)
 						{
-							//only set defaults when not loading or undo
-							if (HasAnyFlags(RF_NeedLoad | RF_NeedPostLoad | RF_NeedPostLoadSubobjects) == false && bIsInUndo == false)
+							if (ControlElement->Settings.ControlType == ERigControlType::Scale)
 							{
-								if (BlendType == EMovieSceneBlendType::Additive)
+								//only set defaults when not loading or undo
+								if (HasAnyFlags(RF_NeedLoad | RF_NeedPostLoad | RF_NeedPostLoadSubobjects) == false && bIsInUndo == false)
 								{
-									Vector.XCurve.SetDefault(0.0f);
-									Vector.YCurve.SetDefault(0.0f);
-									Vector.ZCurve.SetDefault(0.0f);
+									if (BlendType == EMovieSceneBlendType::Additive)
+									{
+										Vector.XCurve.SetDefault(0.0f);
+										Vector.YCurve.SetDefault(0.0f);
+										Vector.ZCurve.SetDefault(0.0f);
+									}
+									else
+									{
+										Vector.XCurve.SetDefault(1.0f);
+										Vector.YCurve.SetDefault(1.0f);
+										Vector.ZCurve.SetDefault(1.0f);
+									}
 								}
-								else
+							}
+							ControlChannelMap.Add(Vector.ParameterName, FChannelMapInfo(ControlIndex, TotalIndex, FloatChannelIndex));
+							ControlIndex++;
+							bool bEnabled = ControlsMask[MaskIndex];
+							FText Group = FText::FromName(ControlElement->GetDisplayName());
+							if (FSpaceControlNameAndChannel* SpaceChannel = GetSpaceChannel(Vector.ParameterName))
+							{
+
+								FChannelMapInfo* pChannelIndex = ControlChannelMap.Find(Vector.ParameterName);
+								if (pChannelIndex)
 								{
-									Vector.XCurve.SetDefault(1.0f);
-									Vector.YCurve.SetDefault(1.0f);
-									Vector.ZCurve.SetDefault(1.0f);
+									pChannelIndex->bDoesHaveSpace = true;
+									pChannelIndex->SpaceChannelIndex = SpaceChannelIndex;
 								}
-							}
-						}
-						ControlChannelMap.Add(Vector.ParameterName, FChannelMapInfo(ControlIndex, TotalIndex, FloatChannelIndex));
-						bool bEnabled = ControlsMask[ControlIndex++];
-						FText Group = FText::FromName(ControlElement->GetDisplayName());
-						if (FSpaceControlNameAndChannel* SpaceChannel = GetSpaceChannel(Vector.ParameterName))
-						{
 
-							FChannelMapInfo* pChannelIndex = ControlChannelMap.Find(Vector.ParameterName);
-							if (pChannelIndex)
-							{
-								pChannelIndex->bDoesHaveSpace = true;
-								pChannelIndex->SpaceChannelIndex = SpaceChannelIndex;
+								FString TotalName = Vector.ParameterName.ToString(); //need ControlName.Space for selection to work.
+								FString SpaceString = SpaceName.ToString();
+								TotalName += ("." + SpaceString);
+								FMovieSceneChannelMetaData SpaceMetaData(FName(*TotalName), Group, Group, bEnabled);
+								SpaceMetaData.DisplayText = FText::FromName(SpaceName);
+								SpaceChannelIndex += 1;
+								SpaceMetaData.SortOrder = TotalIndex++;
+								// Prevent single channels from collapsing to the track node
+								SpaceMetaData.bCanCollapseToTrack = false;
+								Channels.Add(SpaceChannel->SpaceCurve, SpaceMetaData);
 							}
-
-							FString TotalName = Vector.ParameterName.ToString(); //need ControlName.Space for selection to work.
-							FString SpaceString = SpaceName.ToString();
-							TotalName += ("." + SpaceString);
-							FMovieSceneChannelMetaData SpaceMetaData(FName(*TotalName), Group, Group, bEnabled);
-							SpaceMetaData.DisplayText = FText::FromName(SpaceName);
-							SpaceChannelIndex += 1;
-							SpaceMetaData.SortOrder = TotalIndex++;
-							// Prevent single channels from collapsing to the track node
-							SpaceMetaData.bCanCollapseToTrack = false;
-							Channels.Add(SpaceChannel->SpaceCurve, SpaceMetaData);
-						}
 						
 
-						FParameterVectorChannelEditorData EditorData(ControlRig, Vector.ParameterName, bEnabled, Group, TotalIndex, 3);
-						Channels.Add(Vector.XCurve, EditorData.MetaData[0], EditorData.ExternalValues[0]);
-						Channels.Add(Vector.YCurve, EditorData.MetaData[1], EditorData.ExternalValues[1]);
-						Channels.Add(Vector.ZCurve, EditorData.MetaData[2], EditorData.ExternalValues[2]);
-						FloatChannelIndex += 3;
-						TotalIndex += 3;
-						break;
+							FParameterVectorChannelEditorData EditorData(ControlRig, Vector.ParameterName, bEnabled, Group, TotalIndex, 3);
+							Channels.Add(Vector.XCurve, EditorData.MetaData[0], EditorData.ExternalValues[0]);
+							Channels.Add(Vector.YCurve, EditorData.MetaData[1], EditorData.ExternalValues[1]);
+							Channels.Add(Vector.ZCurve, EditorData.MetaData[2], EditorData.ExternalValues[2]);
+							FloatChannelIndex += 3;
+							TotalIndex += 3;
+							break;
+						}
 					}
+					break;
 				}
-				break;
-			}
 
-			case ERigControlType::TransformNoScale:
-			case ERigControlType::Transform:
-			case ERigControlType::EulerTransform:
-			{
-				for (FTransformParameterNameAndCurves& Transform : GetTransformParameterNamesAndCurves())
+				case ERigControlType::TransformNoScale:
+				case ERigControlType::Transform:
+				case ERigControlType::EulerTransform:
 				{
-					if (ControlElement->GetName() == Transform.ParameterName)
+					for (FTransformParameterNameAndCurves& Transform : GetTransformParameterNamesAndCurves())
 					{
-						ControlChannelMap.Add(Transform.ParameterName, FChannelMapInfo(ControlIndex, TotalIndex, FloatChannelIndex));
-						bool bEnabled = ControlsMask[ControlIndex++];
-						FText Group = FText::FromName(ControlElement->GetDisplayName());
-
-						if (FSpaceControlNameAndChannel* SpaceChannel = GetSpaceChannel(Transform.ParameterName))
+						if (ControlElement->GetName() == Transform.ParameterName)
 						{
+							ControlChannelMap.Add(Transform.ParameterName, FChannelMapInfo(ControlIndex, TotalIndex, FloatChannelIndex));
+							ControlIndex++;
+							bool bEnabled = ControlsMask[MaskIndex];
+							FText Group = FText::FromName(ControlElement->GetDisplayName());
 
-							FChannelMapInfo* pChannelIndex = ControlChannelMap.Find(Transform.ParameterName);
-							if (pChannelIndex)
+							if (FSpaceControlNameAndChannel* SpaceChannel = GetSpaceChannel(Transform.ParameterName))
 							{
-								pChannelIndex->bDoesHaveSpace = true;
-								pChannelIndex->SpaceChannelIndex = SpaceChannelIndex;
-							}
 
-							FString TotalName = Transform.ParameterName.ToString(); //need ControlName.Space for selection to work.
-							FString SpaceString = SpaceName.ToString();
-							TotalName += ("." + SpaceString);
-							FMovieSceneChannelMetaData SpaceMetaData(FName(*TotalName), Group, Group, bEnabled);
-							SpaceMetaData.DisplayText = FText::FromName(SpaceName);
-							SpaceChannelIndex += 1;
-							SpaceMetaData.SortOrder = TotalIndex++;
-							// Prevent single channels from collapsing to the track node
-							SpaceMetaData.bCanCollapseToTrack = false;
-							//TMovieSceneExternalValue<FMovieSceneControlRigSpaceBaseKey> ExternalData;
-							Channels.Add(SpaceChannel->SpaceCurve, SpaceMetaData);
-						}
+								FChannelMapInfo* pChannelIndex = ControlChannelMap.Find(Transform.ParameterName);
+								if (pChannelIndex)
+								{
+									pChannelIndex->bDoesHaveSpace = true;
+									pChannelIndex->SpaceChannelIndex = SpaceChannelIndex;
+								}
+
+								FString TotalName = Transform.ParameterName.ToString(); //need ControlName.Space for selection to work.
+								FString SpaceString = SpaceName.ToString();
+								TotalName += ("." + SpaceString);
+								FMovieSceneChannelMetaData SpaceMetaData(FName(*TotalName), Group, Group, bEnabled);
+								SpaceMetaData.DisplayText = FText::FromName(SpaceName);
+								SpaceChannelIndex += 1;
+								SpaceMetaData.SortOrder = TotalIndex++;
+								// Prevent single channels from collapsing to the track node
+								SpaceMetaData.bCanCollapseToTrack = false;
+								//TMovieSceneExternalValue<FMovieSceneControlRigSpaceBaseKey> ExternalData;
+								Channels.Add(SpaceChannel->SpaceCurve, SpaceMetaData);
+							}
 					
 
-						FParameterTransformChannelEditorData EditorData(ControlRig, Transform.ParameterName, bEnabled, TransformMask.GetChannels(), Group, 
-							TotalIndex);
+							FParameterTransformChannelEditorData EditorData(ControlRig, Transform.ParameterName, bEnabled, TransformMask.GetChannels(), Group, 
+								TotalIndex);
 
-						Channels.Add(Transform.Translation[0], EditorData.MetaData[0], EditorData.ExternalValues[0]);
-						Channels.Add(Transform.Translation[1], EditorData.MetaData[1], EditorData.ExternalValues[1]);
-						Channels.Add(Transform.Translation[2], EditorData.MetaData[2], EditorData.ExternalValues[2]);
+							Channels.Add(Transform.Translation[0], EditorData.MetaData[0], EditorData.ExternalValues[0]);
+							Channels.Add(Transform.Translation[1], EditorData.MetaData[1], EditorData.ExternalValues[1]);
+							Channels.Add(Transform.Translation[2], EditorData.MetaData[2], EditorData.ExternalValues[2]);
 
-						Channels.Add(Transform.Rotation[0], EditorData.MetaData[3], EditorData.ExternalValues[3]);
-						Channels.Add(Transform.Rotation[1], EditorData.MetaData[4], EditorData.ExternalValues[4]);
-						Channels.Add(Transform.Rotation[2], EditorData.MetaData[5], EditorData.ExternalValues[5]);
+							Channels.Add(Transform.Rotation[0], EditorData.MetaData[3], EditorData.ExternalValues[3]);
+							Channels.Add(Transform.Rotation[1], EditorData.MetaData[4], EditorData.ExternalValues[4]);
+							Channels.Add(Transform.Rotation[2], EditorData.MetaData[5], EditorData.ExternalValues[5]);
 
-						if (ControlElement->Settings.ControlType == ERigControlType::Transform ||
-							ControlElement->Settings.ControlType == ERigControlType::EulerTransform)
-						{
-							//only set defaults when not loading or undo
-							if (HasAnyFlags(RF_NeedLoad | RF_NeedPostLoad | RF_NeedPostLoadSubobjects) == false && bIsInUndo == false)
+							if (ControlElement->Settings.ControlType == ERigControlType::Transform ||
+								ControlElement->Settings.ControlType == ERigControlType::EulerTransform)
 							{
-								if (BlendType == EMovieSceneBlendType::Additive)
+								//only set defaults when not loading or undo
+								if (HasAnyFlags(RF_NeedLoad | RF_NeedPostLoad | RF_NeedPostLoadSubobjects) == false && bIsInUndo == false)
 								{
-									Transform.Scale[0].SetDefault(0.0f);
-									Transform.Scale[1].SetDefault(0.0f);
-									Transform.Scale[2].SetDefault(0.0f);
+									if (BlendType == EMovieSceneBlendType::Additive)
+									{
+										Transform.Scale[0].SetDefault(0.0f);
+										Transform.Scale[1].SetDefault(0.0f);
+										Transform.Scale[2].SetDefault(0.0f);
+									}
+									else
+									{
+										Transform.Scale[0].SetDefault(1.0f);
+										Transform.Scale[1].SetDefault(1.0f);
+										Transform.Scale[2].SetDefault(1.0f);
+									}
 								}
-								else
-								{
-									Transform.Scale[0].SetDefault(1.0f);
-									Transform.Scale[1].SetDefault(1.0f);
-									Transform.Scale[2].SetDefault(1.0f);
-								}
-							}
-							Channels.Add(Transform.Scale[0], EditorData.MetaData[6], EditorData.ExternalValues[6]);
-							Channels.Add(Transform.Scale[1], EditorData.MetaData[7], EditorData.ExternalValues[7]);
-							Channels.Add(Transform.Scale[2], EditorData.MetaData[8], EditorData.ExternalValues[8]);
-							FloatChannelIndex += 9;
-							TotalIndex += 9;
+								Channels.Add(Transform.Scale[0], EditorData.MetaData[6], EditorData.ExternalValues[6]);
+								Channels.Add(Transform.Scale[1], EditorData.MetaData[7], EditorData.ExternalValues[7]);
+								Channels.Add(Transform.Scale[2], EditorData.MetaData[8], EditorData.ExternalValues[8]);
+								FloatChannelIndex += 9;
+								TotalIndex += 9;
 
-						}
-						else
-						{
-							FloatChannelIndex += 6;
-							TotalIndex += 6;
+							}
+							else
+							{
+								FloatChannelIndex += 6;
+								TotalIndex += 6;
 
+							}
+							break;
 						}
-						break;
 					}
 				}
-			}
-			default:
-				break;
-			}
-#else
-			switch (ControlElement->Settings.ControlType)
-			{
-			case ERigControlType::Float:
-			{
-				for (FScalarParameterNameAndCurve& Scalar : GetScalarParameterNamesAndCurves())
-				{
-					if (ControlElement->GetName() == Scalar.ParameterName)
-					{
-						ControlChannelMap.Add(Scalar.ParameterName, FChannelMapInfo(ControlIndex, TotalIndex,FloatChannelIndex));
-						Channels.Add(Scalar.ParameterCurve);
-						FloatChannelIndex += 1;
-						TotalIndex += 1;
-						ControlIndex++;
-						break;
-					}
+				default:
+					break;
 				}
-				break;
-			}
-			case ERigControlType::Bool:
-			{
-				for (FBoolParameterNameAndCurve& Bool : GetBoolParameterNamesAndCurves())
+	#else
+				switch (ControlElement->Settings.ControlType)
 				{
-					if (ControlElement->GetName() == Bool.ParameterName)
-					{
-						ControlChannelMap.Add(Bool.ParameterName, FChannelMapInfo(ControlIndex, TotalIndex, BoolChannelIndex));
-						Channels.Add(Bool.ParameterCurve);
-						BoolChannelIndex += 1;
-						TotalIndex += 1;
-						ControlIndex++;
-						break;
-					}
-				}
-				break;
-			}
-			case ERigControlType::Integer:
-			{
-				if (ControlElement->Settings.ControlEnum)
+				case ERigControlType::Float:
 				{
-					for (FEnumParameterNameAndCurve& Enum : GetEnumParameterNamesAndCurves())
+					for (FScalarParameterNameAndCurve& Scalar : GetScalarParameterNamesAndCurves())
 					{
-						if (ControlElement->GetName() == Enum.ParameterName)
+						if (ControlElement->GetName() == Scalar.ParameterName)
 						{
-							ControlChannelMap.Add(Enum.ParameterName, FChannelMapInfo(ControlIndex, TotalIndex, EnumChannelIndex));
-							Channels.Add(Enum.ParameterCurve);
-							EnumChannelIndex += 1;
+							ControlChannelMap.Add(Scalar.ParameterName, FChannelMapInfo(ControlIndex, TotalIndex,FloatChannelIndex));
+							Channels.Add(Scalar.ParameterCurve);
+							FloatChannelIndex += 1;
 							TotalIndex += 1;
 							ControlIndex++;
 							break;
 						}
 					}
+					break;
 				}
-				else
+				case ERigControlType::Bool:
 				{
-					for (FIntegerParameterNameAndCurve& Integer : GetIntegerParameterNamesAndCurves())
+					for (FBoolParameterNameAndCurve& Bool : GetBoolParameterNamesAndCurves())
 					{
-						if (ControlElement->GetName() == Integer.ParameterName)
+						if (ControlElement->GetName() == Bool.ParameterName)
 						{
-							ControlChannelMap.Add(Integer.ParameterName, FChannelMapInfo(ControlIndex, TotalIndex, IntegerChannelIndex));
-							Channels.Add(Integer.ParameterCurve);
-							IntegerChannelIndex += 1;
+							ControlChannelMap.Add(Bool.ParameterName, FChannelMapInfo(ControlIndex, TotalIndex, BoolChannelIndex));
+							Channels.Add(Bool.ParameterCurve);
+							BoolChannelIndex += 1;
 							TotalIndex += 1;
 							ControlIndex++;
 							break;
 						}
 					}
+					break;
 				}
-				break;
-			}
-			case ERigControlType::Vector2D:
-			{
-				for (FVector2DParameterNameAndCurves& Vector2D : GetVector2DParameterNamesAndCurves())
+				case ERigControlType::Integer:
 				{
-					if (ControlElement->GetName() == Vector2D.ParameterName)
+					if (ControlElement->Settings.ControlEnum)
 					{
-						ControlChannelMap.Add(Vector2D.ParameterName, FChannelMapInfo(ControlIndex, TotalIndex, FloatChannelIndex));
-						Channels.Add(Vector2D.XCurve);
-						Channels.Add(Vector2D.YCurve);
-						FloatChannelIndex += 2;
-						TotalIndex += 2;
-						ControlIndex++;
-						break;
+						for (FEnumParameterNameAndCurve& Enum : GetEnumParameterNamesAndCurves())
+						{
+							if (ControlElement->GetName() == Enum.ParameterName)
+							{
+								ControlChannelMap.Add(Enum.ParameterName, FChannelMapInfo(ControlIndex, TotalIndex, EnumChannelIndex));
+								Channels.Add(Enum.ParameterCurve);
+								EnumChannelIndex += 1;
+								TotalIndex += 1;
+								ControlIndex++;
+								break;
+							}
+						}
+					}
+					else
+					{
+						for (FIntegerParameterNameAndCurve& Integer : GetIntegerParameterNamesAndCurves())
+						{
+							if (ControlElement->GetName() == Integer.ParameterName)
+							{
+								ControlChannelMap.Add(Integer.ParameterName, FChannelMapInfo(ControlIndex, TotalIndex, IntegerChannelIndex));
+								Channels.Add(Integer.ParameterCurve);
+								IntegerChannelIndex += 1;
+								TotalIndex += 1;
+								ControlIndex++;
+								break;
+							}
+						}
 					}
+					break;
 				}
-				break;
-			}
-			case ERigControlType::Position:
-			case ERigControlType::Scale:
-			case ERigControlType::Rotator:
-			{
-				for (FVectorParameterNameAndCurves& Vector : GetVectorParameterNamesAndCurves())
+				case ERigControlType::Vector2D:
 				{
-					if (ControlElement->GetName() == Vector.ParameterName)
+					for (FVector2DParameterNameAndCurves& Vector2D : GetVector2DParameterNamesAndCurves())
 					{
-						ControlChannelMap.Add(Vector.ParameterName, FChannelMapInfo(ControlIndex, TotalIndex, FloatChannelIndex));
-						bool bDoSpaceChannel = true;
-						if (bDoSpaceChannel)
+						if (ControlElement->GetName() == Vector2D.ParameterName)
 						{
-							if (FSpaceControlNameAndChannel* SpaceChannel = GetSpaceChannel(Vector.ParameterName))
+							ControlChannelMap.Add(Vector2D.ParameterName, FChannelMapInfo(ControlIndex, TotalIndex, FloatChannelIndex));
+							Channels.Add(Vector2D.XCurve);
+							Channels.Add(Vector2D.YCurve);
+							FloatChannelIndex += 2;
+							TotalIndex += 2;
+							ControlIndex++;
+							break;
+						}
+					}
+					break;
+				}
+				case ERigControlType::Position:
+				case ERigControlType::Scale:
+				case ERigControlType::Rotator:
+				{
+					for (FVectorParameterNameAndCurves& Vector : GetVectorParameterNamesAndCurves())
+					{
+						if (ControlElement->GetName() == Vector.ParameterName)
+						{
+							ControlChannelMap.Add(Vector.ParameterName, FChannelMapInfo(ControlIndex, TotalIndex, FloatChannelIndex));
+							bool bDoSpaceChannel = true;
+							if (bDoSpaceChannel)
 							{
-								FChannelMapInfo* pChannelIndex = ControlChannelMap.Find(Vector.ParameterName);
-								if (pChannelIndex)
+								if (FSpaceControlNameAndChannel* SpaceChannel = GetSpaceChannel(Vector.ParameterName))
 								{
-									pChannelIndex->bDoesHaveSpace = true;
-									pChannelIndex->SpaceChannelIndex = SpaceChannelIndex;
+									FChannelMapInfo* pChannelIndex = ControlChannelMap.Find(Vector.ParameterName);
+									if (pChannelIndex)
+									{
+										pChannelIndex->bDoesHaveSpace = true;
+										pChannelIndex->SpaceChannelIndex = SpaceChannelIndex;
+									}
+									SpaceChannelIndex += 1;
+									Channels.Add(SpaceChannel->SpaceCurve);
 								}
-								SpaceChannelIndex += 1;
-								Channels.Add(SpaceChannel->SpaceCurve);
 							}
-						}
 
-						Channels.Add(Vector.XCurve);
-						Channels.Add(Vector.YCurve);
-						Channels.Add(Vector.ZCurve);
-						FloatChannelIndex += 3;
-						TotalIndex += 3;
-						ControlIndex++;
-						break;
+							Channels.Add(Vector.XCurve);
+							Channels.Add(Vector.YCurve);
+							Channels.Add(Vector.ZCurve);
+							FloatChannelIndex += 3;
+							TotalIndex += 3;
+							ControlIndex++;
+							break;
+						}
 					}
+					break;
 				}
-				break;
-			}
-			/*
-			for (FColorParameterNameAndCurves& Color : GetColorParameterNamesAndCurves())
-			{
-				Channels.Add(Color.RedCurve);
-				Channels.Add(Color.GreenCurve);
-				Channels.Add(Color.BlueCurve);
-				Channels.Add(Color.AlphaCurve);
-				break
-			}
-			*/
-			case ERigControlType::TransformNoScale:
-			case ERigControlType::Transform:
-			case ERigControlType::EulerTransform:
-			{
-				for (FTransformParameterNameAndCurves& Transform : GetTransformParameterNamesAndCurves())
+				/*
+				for (FColorParameterNameAndCurves& Color : GetColorParameterNamesAndCurves())
 				{
-					if (ControlElement->GetName() == Transform.ParameterName)
+					Channels.Add(Color.RedCurve);
+					Channels.Add(Color.GreenCurve);
+					Channels.Add(Color.BlueCurve);
+					Channels.Add(Color.AlphaCurve);
+					break
+				}
+				*/
+				case ERigControlType::TransformNoScale:
+				case ERigControlType::Transform:
+				case ERigControlType::EulerTransform:
+				{
+					for (FTransformParameterNameAndCurves& Transform : GetTransformParameterNamesAndCurves())
 					{
-						ControlChannelMap.Add(Transform.ParameterName, FChannelMapInfo(ControlIndex, TotalIndex,FloatChannelIndex));
-						
-						bool bDoSpaceChannel = true;
-						if (bDoSpaceChannel)
+						if (ControlElement->GetName() == Transform.ParameterName)
 						{
-							if (FSpaceControlNameAndChannel* SpaceChannel = GetSpaceChannel(Transform.ParameterName))
+							ControlChannelMap.Add(Transform.ParameterName, FChannelMapInfo(ControlIndex, TotalIndex,FloatChannelIndex));
+						
+							bool bDoSpaceChannel = true;
+							if (bDoSpaceChannel)
 							{
-
-								FChannelMapInfo* pChannelIndex = ControlChannelMap.Find(Transform.ParameterName);
-								if (pChannelIndex)
+								if (FSpaceControlNameAndChannel* SpaceChannel = GetSpaceChannel(Transform.ParameterName))
 								{
-									pChannelIndex->bDoesHaveSpace = true;
-									pChannelIndex->SpaceChannelIndex = SpaceChannelIndex;
+
+									FChannelMapInfo* pChannelIndex = ControlChannelMap.Find(Transform.ParameterName);
+									if (pChannelIndex)
+									{
+										pChannelIndex->bDoesHaveSpace = true;
+										pChannelIndex->SpaceChannelIndex = SpaceChannelIndex;
+									}
+									SpaceChannelIndex += 1;
+									Channels.Add(SpaceChannel->SpaceCurve);
 								}
-								SpaceChannelIndex += 1;
-								Channels.Add(SpaceChannel->SpaceCurve);
-							}
-						}	
+							}	
 						
-						Channels.Add(Transform.Translation[0]);
-						Channels.Add(Transform.Translation[1]);
-						Channels.Add(Transform.Translation[2]);
+							Channels.Add(Transform.Translation[0]);
+							Channels.Add(Transform.Translation[1]);
+							Channels.Add(Transform.Translation[2]);
 
-						Channels.Add(Transform.Rotation[0]);
-						Channels.Add(Transform.Rotation[1]);
-						Channels.Add(Transform.Rotation[2]);
+							Channels.Add(Transform.Rotation[0]);
+							Channels.Add(Transform.Rotation[1]);
+							Channels.Add(Transform.Rotation[2]);
 
-						if (ControlElement->Settings.ControlType == ERigControlType::Transform || ControlElement->Settings.ControlType == ERigControlType::EulerTransform)
-						{
-							Channels.Add(Transform.Scale[0]);
-							Channels.Add(Transform.Scale[1]);
-							Channels.Add(Transform.Scale[2]);
-							FloatChannelIndex += 9;
-							TotalIndex += 9;
-						}
-						else
-						{
-							FloatChannelIndex += 6;
-							TotalIndex += 6;
-						}
+							if (ControlElement->Settings.ControlType == ERigControlType::Transform || ControlElement->Settings.ControlType == ERigControlType::EulerTransform)
+							{
+								Channels.Add(Transform.Scale[0]);
+								Channels.Add(Transform.Scale[1]);
+								Channels.Add(Transform.Scale[2]);
+								FloatChannelIndex += 9;
+								TotalIndex += 9;
+							}
+							else
+							{
+								FloatChannelIndex += 6;
+								TotalIndex += 6;
+							}
 
-						ControlIndex++;
-						break;
+							ControlIndex++;
+							break;
+						}
 					}
+					break;
 				}
-				break;
-			}
 			}
 #endif
+			++MaskIndex;
 		}
 		
 #if WITH_EDITOR
