Control Rig: Fix error when loading rigs with old cache rig element index

diff --git a/Engine/Plugins/Animation/ControlRig/Source/ControlRigDeveloper/Private/ControlRigBlueprint.cpp b/Engine/Plugins/Animation/ControlRig/Source/ControlRigDeveloper/Private/ControlRigBlueprint.cpp
--- a/Engine/Plugins/Animation/ControlRig/Source/ControlRigDeveloper/Private/ControlRigBlueprint.cpp
+++ b/Engine/Plugins/Animation/ControlRig/Source/ControlRigDeveloper/Private/ControlRigBlueprint.cpp
@@ -446,6 +446,7 @@ void UControlRigBlueprint::PostLoad()
 
 		PatchFunctionReferencesOnLoad();
 		PatchVariableNodesOnLoad();
+		PatchRigElementKeyCacheOnLoad();
 
 #if WITH_EDITOR
 
@@ -3470,6 +3471,45 @@ void UControlRigBlueprint::PatchVariableNodesOnLoad()
 #endif
 }
 
+void UControlRigBlueprint::PatchRigElementKeyCacheOnLoad()
+{
+	if (GetLinkerCustomVersion(FControlRigObjectVersion::GUID) < FControlRigObjectVersion::RigElementKeyCache)
+	{
+		for (URigVMGraph* Graph : GetAllModels())
+		{
+			URigVMController* Controller = GetOrCreateController(Graph);
+			TGuardValue<bool> DisablePinDefaultValueValidation(Controller->bValidatePinDefaults, false);
+			Controller->SuspendNotifications(true);
+			for (URigVMNode* Node : Graph->GetNodes())
+			{
+				if (URigVMUnitNode* UnitNode = Cast<URigVMUnitNode>(Node))
+				{
+					UScriptStruct* ScriptStruct = UnitNode->GetScriptStruct();
+					FString FunctionName = FString::Printf(TEXT("F%s::%s"), *ScriptStruct->GetName(), *UnitNode->GetMethodName().ToString());
+					FRigVMFunction Function = FRigVMRegistry::Get().FindFunctionInfo(*FunctionName);
+					for (TFieldIterator<FProperty> It(Function.Struct); It; ++It)
+					{
+						if (It->GetCPPType() == TEXT("FCachedRigElement"))
+						{
+							if (URigVMPin* Pin = Node->FindPin(It->GetName()))
+							{
+								int32 BoneIndex = FCString::Atoi(*Pin->GetDefaultValue());
+								FRigElementKey Key = Hierarchy->GetKey(BoneIndex);
+								FCachedRigElement DefaultValueElement(Key, Hierarchy);
+								FString Result;
+								TBaseStructure<FCachedRigElement>::Get()->ExportText(Result, &DefaultValueElement, nullptr, nullptr, PPF_None, nullptr);								
+								Controller->SetPinDefaultValue(Pin->GetPinPath(), Result, true, false, false);
+								bDirtyDuringLoad = true;
+							}							
+						}
+					}
+				}
+			}
+			Controller->SuspendNotifications(false);
+		}
+	}
+}
+
 void UControlRigBlueprint::PropagatePoseFromInstanceToBP(UControlRig* InControlRig)
 {
 	check(InControlRig);
diff --git a/Engine/Plugins/Animation/ControlRig/Source/ControlRigDeveloper/Public/ControlRigBlueprint.h b/Engine/Plugins/Animation/ControlRig/Source/ControlRigDeveloper/Public/ControlRigBlueprint.h
--- a/Engine/Plugins/Animation/ControlRig/Source/ControlRigDeveloper/Public/ControlRigBlueprint.h
+++ b/Engine/Plugins/Animation/ControlRig/Source/ControlRigDeveloper/Public/ControlRigBlueprint.h
@@ -537,6 +537,7 @@ private:
 #endif
 	void PatchFunctionReferencesOnLoad();
 	void PatchVariableNodesOnLoad();
+	void PatchRigElementKeyCacheOnLoad();
 
 	TMap<FName, int32> AddedMemberVariableMap;
 	TArray<FBPVariableDescription> LastNewVariables;
diff --git a/Engine/Source/Developer/RigVMDeveloper/Public/RigVMModel/RigVMPin.h b/Engine/Source/Developer/RigVMDeveloper/Public/RigVMModel/RigVMPin.h
--- a/Engine/Source/Developer/RigVMDeveloper/Public/RigVMModel/RigVMPin.h
+++ b/Engine/Source/Developer/RigVMDeveloper/Public/RigVMModel/RigVMPin.h
@@ -462,6 +462,7 @@ private:
 	static const FString OrphanPinPrefix;
 
 	friend class URigVMController;
+	friend class UControlRigBlueprint;
 	friend class URigVMGraph;
 	friend class URigVMNode;
 	friend class FRigVMParserAST;
