Control Rig: Copy python script not setting the control transforms correctly

diff --git a/Engine/Plugins/Animation/ControlRig/Source/ControlRig/Private/Rigs/RigHierarchyController.cpp b/Engine/Plugins/Animation/ControlRig/Source/ControlRig/Private/Rigs/RigHierarchyController.cpp
--- a/Engine/Plugins/Animation/ControlRig/Source/ControlRig/Private/Rigs/RigHierarchyController.cpp
+++ b/Engine/Plugins/Animation/ControlRig/Source/ControlRig/Private/Rigs/RigHierarchyController.cpp
@@ -1445,11 +1445,11 @@ TArray<FString> URigHierarchyController::GetAddControlPythonCommands(FRigControl
 		*SettingsStr,
 		*ValueStr));
 
-	Commands.Append(GetSetControlValuePythonCommands(Control, Value, ERigControlValueType::Initial));
+	Commands.Append(GetSetControlShapeTransformPythonCommands(Control, Control->Shape.Initial.Local.Transform, true));
 	Commands.Append(GetSetControlValuePythonCommands(Control, Settings.MinimumValue, ERigControlValueType::Minimum));
 	Commands.Append(GetSetControlValuePythonCommands(Control, Settings.MaximumValue, ERigControlValueType::Maximum));
 	Commands.Append(GetSetControlOffsetTransformPythonCommands(Control, Control->Offset.Initial.Local.Transform, true, true));
-	Commands.Append(GetSetControlShapeTransformPythonCommands(Control, Control->Shape.Initial.Local.Transform, true));
+	Commands.Append(GetSetControlValuePythonCommands(Control, Value, ERigControlValueType::Current));
 
 	return Commands;
 }
diff --git a/Engine/Plugins/Animation/ControlRig/Source/ControlRig/Public/Rigs/RigHierarchyDefines.h b/Engine/Plugins/Animation/ControlRig/Source/ControlRig/Public/Rigs/RigHierarchyDefines.h
--- a/Engine/Plugins/Animation/ControlRig/Source/ControlRig/Public/Rigs/RigHierarchyDefines.h
+++ b/Engine/Plugins/Animation/ControlRig/Source/ControlRig/Public/Rigs/RigHierarchyDefines.h
@@ -298,40 +298,19 @@ public:
 				Get<FVector3f>().X, Get<FVector3f>().Z, Get<FVector3f>().Y); break;
 			case ERigControlType::Scale: ValueStr = FString::Printf(TEXT("unreal.RigHierarchy.make_control_value_from_vector(unreal.Vector(%.6f, %.6f, %.6f))"),
 				Get<FVector3f>().X, Get<FVector3f>().Y, Get<FVector3f>().Z); break;
-			/*
-			case ERigControlType::Transform: ValueStr = FString::Printf(TEXT("unreal.RigHierarchy.make_control_value_from_transform(unreal.Transform(location=[%.6f,%.6f,%.6f],rotation=[%.6f,%.6f,%.6f],scale=[%.6f,%.6f,%.6f]))"),
-				Get<FTransform_Float>().ToTransform().GetLocation().X,
-				Get<FTransform_Float>().ToTransform().GetLocation().Y,
-				Get<FTransform_Float>().ToTransform().GetLocation().Z,
-				Get<FTransform_Float>().ToTransform().Rotator().Pitch,
-				Get<FTransform_Float>().ToTransform().Rotator().Yaw,
-				Get<FTransform_Float>().ToTransform().Rotator().Roll,
-				Get<FTransform_Float>().ToTransform().GetScale3D().X,
-				Get<FTransform_Float>().ToTransform().GetScale3D().Y,
-				Get<FTransform_Float>().ToTransform().GetScale3D().Z); break;
-			*/
 			case ERigControlType::Transform:
 			case ERigControlType::EulerTransform: ValueStr = FString::Printf(TEXT("unreal.RigHierarchy.make_control_value_from_euler_transform(unreal.EulerTransform(location=[%.6f,%.6f,%.6f],rotation=[%.6f,%.6f,%.6f],scale=[%.6f,%.6f,%.6f]))"),
-				Get<FEulerTransform_Float>().ToTransform().Location.X,
-				Get<FEulerTransform_Float>().ToTransform().Location.Y,
-				Get<FEulerTransform_Float>().ToTransform().Location.Z,
-				Get<FEulerTransform_Float>().ToTransform().Rotation.Euler().X,
-				Get<FEulerTransform_Float>().ToTransform().Rotation.Euler().Y,
-				Get<FEulerTransform_Float>().ToTransform().Rotation.Euler().Z,
-				Get<FEulerTransform_Float>().ToTransform().Scale.X,
-				Get<FEulerTransform_Float>().ToTransform().Scale.Y,
-				Get<FEulerTransform_Float>().ToTransform().Scale.Z); break;
+				Get<FEulerTransform_Float>().TranslationX,
+				Get<FEulerTransform_Float>().TranslationY,
+				Get<FEulerTransform_Float>().TranslationZ,
+				Get<FEulerTransform_Float>().RotationPitch,
+				Get<FEulerTransform_Float>().RotationYaw,
+				Get<FEulerTransform_Float>().RotationRoll,
+				Get<FEulerTransform_Float>().ScaleX,
+				Get<FEulerTransform_Float>().ScaleY,
+				Get<FEulerTransform_Float>().ScaleZ); break;
 			case ERigControlType::Vector2D: ValueStr = FString::Printf(TEXT("unreal.RigHierarchy.make_control_value_from_vector2d(unreal.Vector2D(%.6f, %.6f))"),
 				Get<FVector3f>().X, Get<FVector3f>().Y); break;
-			/*
-			case ERigControlType::TransformNoScale: ValueStr = FString::Printf(TEXT("unreal.RigHierarchy.make_control_value_from_transform_no_scale(unreal.TransformNoScale(location=[%.6f,%.6f,%.6f],rotation=[%.6f,%.6f,%.6f]))"),
-				Get<FTransformNoScale_Float>().ToTransform().ToFTransform().GetLocation().X,
-				Get<FTransformNoScale_Float>().ToTransform().ToFTransform().GetLocation().Y,
-				Get<FTransformNoScale_Float>().ToTransform().ToFTransform().GetLocation().Z,
-				Get<FTransformNoScale_Float>().ToTransform().ToFTransform().Rotator().Pitch,
-				Get<FTransformNoScale_Float>().ToTransform().ToFTransform().Rotator().Yaw,
-				Get<FTransformNoScale_Float>().ToTransform().ToFTransform().Rotator().Roll); break;
-			*/
 			case ERigControlType::TransformNoScale: ValueStr = FString::Printf(TEXT("unreal.RigHierarchy.make_control_value_from_euler_transform(unreal.EulerTransform(location=[%.6f,%.6f,%.6f],rotation=[%.6f,%.6f,%.6f],scale=[%.6f,%.6f,%.6f]))"),
 				FEulerTransform(Get<FTransformNoScale_Float>().ToTransform().ToFTransform()).Location.X,
 				FEulerTransform(Get<FTransformNoScale_Float>().ToTransform().ToFTransform()).Location.Y,
diff --git a/Engine/Source/Developer/RigVMDeveloper/Private/RigVMModel/RigVMPin.cpp b/Engine/Source/Developer/RigVMDeveloper/Private/RigVMModel/RigVMPin.cpp
--- a/Engine/Source/Developer/RigVMDeveloper/Private/RigVMModel/RigVMPin.cpp
+++ b/Engine/Source/Developer/RigVMDeveloper/Private/RigVMModel/RigVMPin.cpp
@@ -517,7 +517,7 @@ FString URigVMPin::GetArrayElementCppType() const
 	}
 
 	const FString ResolvedType = GetCPPType();
-	return ResolvedType.Mid(7, ResolvedType.Len() - 8);
+	return RigVMTypeUtils::BaseTypeFromArrayType(ResolvedType);
 }
 
 bool URigVMPin::IsStringType() const
