Adding a delegate for mode toolkits to do any UI shutdown that depends on the mode UI layer. Small fix to completely unregister known tabs in the mode UI layer.

diff --git a/Engine/Plugins/Animation/ControlRig/Source/ControlRigEditor/Private/EditMode/ControlRigEditModeToolkit.cpp b/Engine/Plugins/Animation/ControlRig/Source/ControlRigEditor/Private/EditMode/ControlRigEditModeToolkit.cpp
--- a/Engine/Plugins/Animation/ControlRig/Source/ControlRigEditor/Private/EditMode/ControlRigEditModeToolkit.cpp
+++ b/Engine/Plugins/Animation/ControlRig/Source/ControlRigEditor/Private/EditMode/ControlRigEditModeToolkit.cpp
@@ -56,16 +56,6 @@ void FControlRigEditModeToolkit::Init(const TSharedPtr<IToolkitHost>& InitToolki
 	FModeToolkit::Init(InitToolkitHost);
 }
 
-FControlRigEditModeToolkit::~FControlRigEditModeToolkit()
-{
-	if (FSlateApplication::IsInitialized())
-	{
-		RemoveAndDestroyTweenOverlay();
-		FGlobalTabmanager::Get()->UnregisterNomadTabSpawner("ControlRigProfiler");
-	}
-}
-
-
 void FControlRigEditModeToolkit::GetToolPaletteNames(TArray<FName>& InPaletteName) const
 {
 	InPaletteName = AnimationPaletteNames;
@@ -98,7 +88,11 @@ void FControlRigEditModeToolkit::TryInvokeToolkitUI(const FName InName)
 	if (InName == MotionTrailTabName)
 	{
 		FTabId MotionTrailTabID(MotionTrailTabName);
-		FGlobalTabmanager::Get()->TryInvokeTab(MotionTrailTabID);
+		if (ModeUILayer.IsValid())
+		{
+			TSharedPtr<FAssetEditorModeUILayer> ModeUILayerPtr = ModeUILayer.Pin();
+			ModeUILayerPtr->GetTabManager()->TryInvokeTab(MotionTrailTabID);
+		}
 	}
 	if (InName == PoseTabName)
 	{
@@ -266,17 +260,19 @@ void FControlRigEditModeToolkit::RequestModeUITabs()
 		SnapperTabInfo.TabTooltip = LOCTEXT("ControlRigSnapperTabTooltip", "Snap child objects to a parent object over a set of frames.");
 		ModeUILayerPtr->SetModePanelInfo(UAssetEditorUISubsystem::TopRightTabID, SnapperTabInfo);
 
-		FGlobalTabmanager::Get()->RegisterNomadTabSpawner(MotionTrailTabName, FOnSpawnTab::CreateStatic(&SpawnMotionTrailTab))
+		ModeUILayerPtr->GetTabManager()->RegisterTabSpawner(MotionTrailTabName, FOnSpawnTab::CreateStatic(&SpawnMotionTrailTab))
 			.SetDisplayName(LOCTEXT("MotionTrailTab", "Motion Trail"))
 			.SetTooltipText(LOCTEXT("MotionTrailTabTooltip", "Display motion trails for animated objects."))
 			.SetGroup(MenuGroup)
 			.SetIcon(FSlateIcon(TEXT("ControlRigEditorStyle"), TEXT("HierarchicalProfiler.TabIcon")));
 
-		FGlobalTabmanager::Get()->RegisterNomadTabSpawner("HierarchicalProfiler", FOnSpawnTab::CreateStatic(&SpawnRigProfiler))
+		ModeUILayerPtr->GetTabManager()->RegisterTabSpawner("HierarchicalProfiler", FOnSpawnTab::CreateStatic(&SpawnRigProfiler))
 			.SetDisplayName(LOCTEXT("HierarchicalProfilerTab", "Hierarchical Profiler"))
 			.SetTooltipText(LOCTEXT("HierarchicalProfilerTooltip", "Open the Hierarchical Profiler tab."))
 			.SetGroup(MenuGroup)
 			.SetIcon(FSlateIcon(TEXT("ControlRigEditorStyle"), TEXT("HierarchicalProfiler.TabIcon")));
+
+		ModeUILayer.Pin()->ToolkitHostShutdownUI().BindSP(this, &FControlRigEditModeToolkit::UnregisterAndRemoveFloatingTabs);
 	}
 };
 
@@ -287,6 +283,30 @@ void FControlRigEditModeToolkit::InvokeUI()
 // TODO: any future default tabs will go here
 }
 
+void FControlRigEditModeToolkit::UnregisterAndRemoveFloatingTabs()
+{
+	if (FSlateApplication::IsInitialized())
+	{
+		RemoveAndDestroyTweenOverlay();
+		if (ModeUILayer.IsValid())
+		{
+			TSharedPtr<FAssetEditorModeUILayer> ModeUILayerPtr = ModeUILayer.Pin();
+			TSharedPtr<SDockTab> ProfilerTab = ModeUILayerPtr->GetTabManager()->FindExistingLiveTab(FTabId("HierarchicalProfiler"));
+			if(ProfilerTab)
+			{ 
+				ProfilerTab->RequestCloseTab();
+			}
+			ModeUILayerPtr->GetTabManager()->UnregisterTabSpawner("HierarchicalProfiler");
+			TSharedPtr<SDockTab> MotionTrailTab = ModeUILayerPtr->GetTabManager()->FindExistingLiveTab(FTabId(MotionTrailTabName));
+			if (MotionTrailTab)
+			{
+				MotionTrailTab->RequestCloseTab();
+			}
+			ModeUILayerPtr->GetTabManager()->UnregisterTabSpawner(MotionTrailTabName);
+		}
+	}
+}
+
 void FControlRigEditModeToolkit::TryInvokeSnapperTab()
 {
 	TSharedPtr<FAssetEditorModeUILayer> ModeUILayerPtr = ModeUILayer.Pin();
diff --git a/Engine/Plugins/Animation/ControlRig/Source/ControlRigEditor/Private/EditMode/ControlRigEditModeToolkit.h b/Engine/Plugins/Animation/ControlRig/Source/ControlRigEditor/Private/EditMode/ControlRigEditModeToolkit.h
--- a/Engine/Plugins/Animation/ControlRig/Source/ControlRigEditor/Private/EditMode/ControlRigEditModeToolkit.h
+++ b/Engine/Plugins/Animation/ControlRig/Source/ControlRigEditor/Private/EditMode/ControlRigEditModeToolkit.h
@@ -36,7 +36,7 @@ public:
 		return false;
 	}
 	virtual void Init(const TSharedPtr<IToolkitHost>& InitToolkitHost) override;
-	~FControlRigEditModeToolkit();
+
 	/** Mode Toolbar Palettes **/
 	virtual void GetToolPaletteNames(TArray<FName>& InPaletteName) const override;
 	virtual FText GetToolPaletteDisplayName(FName PaletteName) const override;
@@ -69,6 +69,8 @@ protected:
 	virtual void RequestModeUITabs() override;
 	virtual void InvokeUI() override;
 
+	void UnregisterAndRemoveFloatingTabs();
+
 private:
 	void TryInvokeSnapperTab();
 	void TryInvokePoseTab();
diff --git a/Engine/Source/Editor/EditorFramework/Private/Toolkits/AssetEditorModeUILayer.cpp b/Engine/Source/Editor/EditorFramework/Private/Toolkits/AssetEditorModeUILayer.cpp
--- a/Engine/Source/Editor/EditorFramework/Private/Toolkits/AssetEditorModeUILayer.cpp
+++ b/Engine/Source/Editor/EditorFramework/Private/Toolkits/AssetEditorModeUILayer.cpp
@@ -55,12 +55,9 @@ void FAssetEditorModeUILayer::OnToolkitHostingStarted(const TSharedRef<IToolkit>
 
 void FAssetEditorModeUILayer::OnToolkitHostingFinished(const TSharedRef<IToolkit>& Toolkit)
 {
+	OnToolkitHostShutdownUI.ExecuteIfBound();
 	if (HostedToolkit.IsValid() && HostedToolkit.Pin() == Toolkit)
 	{
-		for (TPair<FName, FMinorTabConfig>& TabSpawnerInfo : RequestedTabInfo)
-		{
-			TabSpawnerInfo.Value = FMinorTabConfig(TabSpawnerInfo.Key);
-		}
 		for (TPair<FName, TWeakPtr<SDockTab>>& SpawnedTab : SpawnedTabs)
 		{
 			if (SpawnedTab.Value.IsValid())
@@ -69,7 +66,11 @@ void FAssetEditorModeUILayer::OnToolkitHostingFinished(const TSharedRef<IToolkit
 				SpawnedTab.Value.Pin()->RequestCloseTab();
 			}
 		}
-
+		for (TPair<FName, FMinorTabConfig>& TabSpawnerInfo : RequestedTabInfo)
+		{
+			TabSpawnerInfo.Value = FMinorTabConfig(TabSpawnerInfo.Key);
+			GetTabManager()->UnregisterTabSpawner(TabSpawnerInfo.Key);
+		}
 		SpawnedTabs.Empty();
 	}
 }
diff --git a/Engine/Source/Editor/EditorFramework/Public/Toolkits/AssetEditorModeUILayer.h b/Engine/Source/Editor/EditorFramework/Public/Toolkits/AssetEditorModeUILayer.h
--- a/Engine/Source/Editor/EditorFramework/Public/Toolkits/AssetEditorModeUILayer.h
+++ b/Engine/Source/Editor/EditorFramework/Public/Toolkits/AssetEditorModeUILayer.h
@@ -57,6 +57,10 @@ public:
 	{
 		return OnToolkitHostReadyForUI;
 	};
+	virtual FSimpleDelegate& ToolkitHostShutdownUI()
+	{
+		return OnToolkitHostShutdownUI;
+	}
 	virtual const FName GetStatusBarName() const
 	{
 		return ToolkitHost->GetStatusBarName();
@@ -80,4 +84,5 @@ protected:
 	TMap<FName, FMinorTabConfig> RequestedTabInfo;
 	TMap<FName, TWeakPtr<SDockTab>> SpawnedTabs;
 	FSimpleDelegate OnToolkitHostReadyForUI;
+	FSimpleDelegate OnToolkitHostShutdownUI;
 };
