Control Rig:  Change TreeView so that it won't deselect items that are filtered out or collapsed.

diff --git a/Engine/Plugins/Animation/ControlRig/Source/ControlRigEditor/Private/Editor/SRigHierarchyTreeView.cpp b/Engine/Plugins/Animation/ControlRig/Source/ControlRigEditor/Private/Editor/SRigHierarchyTreeView.cpp
--- a/Engine/Plugins/Animation/ControlRig/Source/ControlRigEditor/Private/Editor/SRigHierarchyTreeView.cpp
+++ b/Engine/Plugins/Animation/ControlRig/Source/ControlRigEditor/Private/Editor/SRigHierarchyTreeView.cpp
@@ -155,7 +155,8 @@ void SRigHierarchyTreeView::Construct(const FArguments& InArgs)
 	SuperArgs.OnSetExpansionRecursive(Delegates.OnSetExpansionRecursive);
 	SuperArgs.HighlightParentNodesForSelection(true);
 	SuperArgs.ItemHeight(24);
-	
+	SuperArgs.AllowInvisibleItemSelection(true);  //without this we deselect everything when we filter or we collapse
+
 	STreeView<TSharedPtr<FRigTreeElement>>::Construct(SuperArgs);
 }
 
