Sequencer: Control Rig: Create CR outliner and details panel for Controls while animating.
Created base widget for these new dockable tabs. Moved logic to these new widgets.
However since we want some feedback we left a #define to keep the old code that we can bring back if needed, will be removed once vetted by design.
Note space picker has it's own tab also but isn't working, need to revist that with Engine UX, think they have a bug.

Note AG/Lyra failed do to other changes not related to this CL.

diff --git a/Engine/Plugins/Animation/ControlRig/Source/ControlRig/Private/Sequencer/MovieSceneControlRigParameterTemplate.cpp b/Engine/Plugins/Animation/ControlRig/Source/ControlRig/Private/Sequencer/MovieSceneControlRigParameterTemplate.cpp
--- a/Engine/Plugins/Animation/ControlRig/Source/ControlRig/Private/Sequencer/MovieSceneControlRigParameterTemplate.cpp
+++ b/Engine/Plugins/Animation/ControlRig/Source/ControlRig/Private/Sequencer/MovieSceneControlRigParameterTemplate.cpp
@@ -684,6 +684,7 @@ struct FControlRigParameterPreAnimatedTokenProducer : IMovieScenePreAnimatedToke
 						const bool bSetupUndo = false;
 						if (URigHierarchy* RigHierarchy = ControlRig->GetHierarchy())
 						{
+
 							FRigElementKey ControlKey;
 							ControlKey.Type = ERigElementType::Control;
 							for (FControlSpaceAndValue& SpaceNameAndValue : SpaceValues)
@@ -709,74 +710,77 @@ struct FControlRigParameterPreAnimatedTokenProducer : IMovieScenePreAnimatedToke
 								break;
 								}
 							}
-						}
+							
 
-						for (TNameAndValue<float>& Value : ScalarValues)
-						{
-							if (ControlRig->FindControl(Value.Name))
+							for (TNameAndValue<float>& Value : ScalarValues)
 							{
-								ControlRig->SetControlValue<float>(Value.Name, Value.Value, true, FRigControlModifiedContext(EControlRigSetKey::Never),bSetupUndo);
-							}
-						}
-
-						for (TNameAndValue<bool>& Value : BoolValues)
-						{
-							if (ControlRig->FindControl(Value.Name))
-							{
-								ControlRig->SetControlValue<bool>(Value.Name, Value.Value, true, FRigControlModifiedContext(EControlRigSetKey::Never),bSetupUndo);
-							}
-						}
-
-						for (TNameAndValue<int32>& Value : IntegerValues)
-						{
-							if (ControlRig->FindControl(Value.Name))
-							{
-								ControlRig->SetControlValue<int32>(Value.Name, Value.Value, true, FRigControlModifiedContext(EControlRigSetKey::Never),bSetupUndo);
+								if (ControlRig->FindControl(Value.Name))
+								{
+									ControlRig->SetControlValue<float>(Value.Name, Value.Value, true, FRigControlModifiedContext(EControlRigSetKey::Never), bSetupUndo);
+								}
 							}
-						}
 
-						for (TNameAndValue<FVector2D>& Value : Vector2DValues)
-						{
-							if (ControlRig->FindControl(Value.Name))
+							for (TNameAndValue<bool>& Value : BoolValues)
 							{
-								const FVector3f Vector3(Value.Value.X, Value.Value.Y, 0.f);
-								//okay to use vector3 for 2d here
-								ControlRig->SetControlValue<FVector3f>(Value.Name, Vector3, true, FRigControlModifiedContext(EControlRigSetKey::Never),bSetupUndo);
+								if (ControlRig->FindControl(Value.Name))
+								{
+									ControlRig->SetControlValue<bool>(Value.Name, Value.Value, true, FRigControlModifiedContext(EControlRigSetKey::Never), bSetupUndo);
+								}
 							}
-						}
 
-						for (TNameAndValue<FVector>& Value : VectorValues)
-						{
-							if (ControlRig->FindControl(Value.Name))
+							for (TNameAndValue<int32>& Value : IntegerValues)
 							{
-								ControlRig->SetControlValue<FVector3f>(Value.Name, Value.Value, true, FRigControlModifiedContext(EControlRigSetKey::Never),bSetupUndo);
+								if (ControlRig->FindControl(Value.Name))
+								{
+									ControlRig->SetControlValue<int32>(Value.Name, Value.Value, true, FRigControlModifiedContext(EControlRigSetKey::Never), bSetupUndo);
+								}
 							}
-						}
-
-						for (TNameAndValue<FTransform>& Value : TransformValues)
-						{
-							if (FRigControlElement* ControlElement = ControlRig->FindControl(Value.Name))
+							for (int32 TwiceHack = 0; TwiceHack < 2; ++TwiceHack)
 							{
-								switch (ControlElement->Settings.ControlType)
+								for (TNameAndValue<FVector2D>& Value : Vector2DValues)
 								{
-								case ERigControlType::Transform:
-								{
-									ControlRig->SetControlValue<FRigControlValue::FTransform_Float>(Value.Name, Value.Value, true, FRigControlModifiedContext(EControlRigSetKey::Never),bSetupUndo);
-									break;
+									if (ControlRig->FindControl(Value.Name))
+									{
+										const FVector3f Vector3(Value.Value.X, Value.Value.Y, 0.f);
+										//okay to use vector3 for 2d here
+										ControlRig->SetControlValue<FVector3f>(Value.Name, Vector3, true, FRigControlModifiedContext(EControlRigSetKey::Never), bSetupUndo);
+									}
 								}
-								case ERigControlType::TransformNoScale:
+
+								for (TNameAndValue<FVector>& Value : VectorValues)
 								{
-									FTransformNoScale NoScale = Value.Value;
-									ControlRig->SetControlValue<FRigControlValue::FTransformNoScale_Float>(Value.Name, NoScale, true, FRigControlModifiedContext(EControlRigSetKey::Never),bSetupUndo);
-									break;
+									if (ControlRig->FindControl(Value.Name))
+									{
+										ControlRig->SetControlValue<FVector3f>(Value.Name, Value.Value, true, FRigControlModifiedContext(EControlRigSetKey::Never), bSetupUndo);
+									}
 								}
-								case ERigControlType::EulerTransform:
+
+								for (TNameAndValue<FTransform>& Value : TransformValues)
 								{
-									FEulerTransform EulerTransform = Value.Value;
-									ControlRig->SetControlValue<FRigControlValue::FEulerTransform_Float>(Value.Name, EulerTransform, true, FRigControlModifiedContext(EControlRigSetKey::Never),bSetupUndo);
-									break;
-								}
+									if (FRigControlElement* ControlElement = ControlRig->FindControl(Value.Name))
+									{
+										switch (ControlElement->Settings.ControlType)
+										{
+										case ERigControlType::Transform:
+										{
+											ControlRig->SetControlValue<FRigControlValue::FTransform_Float>(Value.Name, Value.Value, true, FRigControlModifiedContext(EControlRigSetKey::Never), bSetupUndo);
+											break;
+										}
+										case ERigControlType::TransformNoScale:
+										{
+											FTransformNoScale NoScale = Value.Value;
+											ControlRig->SetControlValue<FRigControlValue::FTransformNoScale_Float>(Value.Name, NoScale, true, FRigControlModifiedContext(EControlRigSetKey::Never), bSetupUndo);
+											break;
+										}
+										case ERigControlType::EulerTransform:
+										{
+											FEulerTransform EulerTransform = Value.Value;
+											ControlRig->SetControlValue<FRigControlValue::FEulerTransform_Float>(Value.Name, EulerTransform, true, FRigControlModifiedContext(EControlRigSetKey::Never), bSetupUndo);
+											break;
+										}
 
+										}
+									}
 								}
 							}
 						}
diff --git a/Engine/Plugins/Animation/ControlRig/Source/ControlRigEditor/Private/EditMode/ControlRigBaseDockableView.cpp b/Engine/Plugins/Animation/ControlRig/Source/ControlRigEditor/Private/EditMode/ControlRigBaseDockableView.cpp
new file mode 100644
--- /dev/null
+++ b/Engine/Plugins/Animation/ControlRig/Source/ControlRigEditor/Private/EditMode/ControlRigBaseDockableView.cpp
@@ -0,0 +1,129 @@
+// Copyright Epic Games, Inc. All Rights Reserved.
+
+#include "ControlRigBaseDockableView.h"
+#include "AssetData.h"
+#include "EditorStyleSet.h"
+#include "Widgets/Text/SInlineEditableTextBlock.h"
+#include "EditorStyleSet.h"
+#include "Styling/CoreStyle.h"
+#include "ScopedTransaction.h"
+#include "ControlRig.h"
+#include "UnrealEdGlobals.h"
+#include "ControlRigEditMode.h"
+#include "EditorModeManager.h"
+#include "ISequencer.h"
+
+void FControlRigBaseDockableView::SetEditMode(FControlRigEditMode& InEditMode)
+{
+
+	CurrentControlRig = nullptr;
+	ModeTools = InEditMode.GetModeManager();
+	if (FControlRigEditMode* EditMode = static_cast<FControlRigEditMode*>(ModeTools->GetActiveMode(FControlRigEditMode::ModeName)))
+	{
+		EditMode->OnControlRigAddedOrRemoved().AddRaw(this, &FControlRigBaseDockableView::HandleControlAdded);
+		HandleControlAdded(GetControlRig(), true);
+	}
+}
+
+FControlRigBaseDockableView::FControlRigBaseDockableView()
+{
+	FCoreUObjectDelegates::OnObjectsReplaced.AddRaw(this, &FControlRigBaseDockableView::OnObjectsReplaced);
+}
+
+FControlRigBaseDockableView::~FControlRigBaseDockableView()
+{
+	if (FControlRigEditMode* EditMode = static_cast<FControlRigEditMode*>(ModeTools->GetActiveMode(FControlRigEditMode::ModeName)))
+	{
+		EditMode->OnControlRigAddedOrRemoved().RemoveAll(this);
+		if (EditMode->GetControlRig(true))
+		{
+			EditMode->GetControlRig(true)->ControlSelected().RemoveAll(this);
+		}
+	}
+	else
+	{
+		if (CurrentControlRig.IsValid())
+		{
+			(CurrentControlRig.Get())->ControlSelected().RemoveAll(this);
+		}
+	}
+	CurrentControlRig = nullptr;
+
+	FCoreUObjectDelegates::OnObjectsReplaced.RemoveAll(this);
+
+}
+
+UControlRig* FControlRigBaseDockableView::GetControlRig() 
+{
+	UControlRig* NewControlRig = nullptr;
+	if (FControlRigEditMode* EditMode = static_cast<FControlRigEditMode*>(ModeTools->GetActiveMode(FControlRigEditMode::ModeName)))
+	{
+		NewControlRig = EditMode->GetControlRig(true);
+	}
+	bool bNewControlRig = NewControlRig != CurrentControlRig;
+	if (bNewControlRig)
+	{
+		if (CurrentControlRig.IsValid())
+		{
+			(CurrentControlRig.Get())->ControlSelected().RemoveAll(this);
+
+		}
+		CurrentControlRig = NewControlRig;
+		NewControlRigSet(NewControlRig);
+	}
+	return NewControlRig;
+}
+
+
+void FControlRigBaseDockableView::OnObjectsReplaced(const TMap<UObject*, UObject*>& OldToNewInstanceMap)
+{
+	if(CurrentControlRig.IsValid())
+	{ 
+		UObject* OldObject = CurrentControlRig.Get();
+		UObject* NewObject = OldToNewInstanceMap.FindRef(OldObject);
+		if (NewObject)
+		{
+			if (UControlRig* ControlRig = Cast<UControlRig>(NewObject))
+			{
+				NewControlRigSet(ControlRig);
+			}
+		}
+	}
+}
+
+void FControlRigBaseDockableView::NewControlRigSet(UControlRig* ControlRig)
+{
+	CurrentControlRig = ControlRig;
+}
+
+void FControlRigBaseDockableView::HandleControlAdded(UControlRig* ControlRig, bool bIsAdded)
+{
+	if (ControlRig)
+	{
+		if (bIsAdded)
+		{
+			(ControlRig)->ControlSelected().RemoveAll(this);
+			(ControlRig)->ControlSelected().AddRaw(this, &FControlRigBaseDockableView::HandleControlSelected);
+			CurrentControlRig = ControlRig;
+		}
+		else
+		{
+			(ControlRig)->ControlSelected().RemoveAll(this);
+		}
+	}
+}
+
+void FControlRigBaseDockableView::HandleControlSelected(UControlRig* Subject, FRigControlElement* InControl, bool bSelected)
+{
+
+}
+
+ISequencer* FControlRigBaseDockableView::GetSequencer() const
+{
+	if (FControlRigEditMode* EditMode = static_cast<FControlRigEditMode*>(ModeTools->GetActiveMode(FControlRigEditMode::ModeName)))
+	{
+		TWeakPtr<ISequencer> Sequencer = EditMode->GetWeakSequencer();
+		return Sequencer.Pin().Get();
+	}
+	return nullptr;
+}
diff --git a/Engine/Plugins/Animation/ControlRig/Source/ControlRigEditor/Private/EditMode/ControlRigBaseDockableView.h b/Engine/Plugins/Animation/ControlRig/Source/ControlRigEditor/Private/EditMode/ControlRigBaseDockableView.h
new file mode 100644
--- /dev/null
+++ b/Engine/Plugins/Animation/ControlRig/Source/ControlRigEditor/Private/EditMode/ControlRigBaseDockableView.h
@@ -0,0 +1,38 @@
+// Copyright Epic Games, Inc. All Rights Reserved.
+/**
+* Base View for Dockable Control Rig Animation widgets Details/Outliner
+*/
+#pragma once
+
+#include "CoreMinimal.h"
+#include "EditorModeManager.h"
+
+
+class UControlRig;
+class ISequencer;
+class FControlRigEditMode;
+struct FRigControlElement;
+
+class FControlRigBaseDockableView 
+{
+public:
+	FControlRigBaseDockableView();
+	virtual ~FControlRigBaseDockableView();
+
+	void SetEditMode(FControlRigEditMode& InEditMode);
+
+protected:
+	virtual void HandleControlSelected(UControlRig* Subject, FRigControlElement* InControl, bool bSelected);
+	virtual void HandleControlAdded(UControlRig* ControlRig, bool bIsAdded);
+	virtual void NewControlRigSet(UControlRig* ControlRig);
+
+	void OnObjectsReplaced(const TMap<UObject*, UObject*>& OldToNewInstanceMap);
+
+	UControlRig* GetControlRig();
+	ISequencer* GetSequencer() const;
+	TWeakObjectPtr<UControlRig> CurrentControlRig;
+
+	FEditorModeTools* ModeTools = nullptr;
+
+};
+
diff --git a/Engine/Plugins/Animation/ControlRig/Source/ControlRigEditor/Private/EditMode/ControlRigEditMode.cpp b/Engine/Plugins/Animation/ControlRig/Source/ControlRigEditor/Private/EditMode/ControlRigEditMode.cpp
--- a/Engine/Plugins/Animation/ControlRig/Source/ControlRigEditor/Private/EditMode/ControlRigEditMode.cpp
+++ b/Engine/Plugins/Animation/ControlRig/Source/ControlRigEditor/Private/EditMode/ControlRigEditMode.cpp
@@ -225,11 +225,12 @@ bool FControlRigEditMode::IsInLevelEditor() const
 {
 	return GetModeManager() == &GLevelEditorModeTools();
 }
-
 void FControlRigEditMode::SetUpDetailPanel()
 {
 	if (IsInLevelEditor())
 	{
+		StaticCastSharedPtr<SControlRigEditModeTools>(Toolkit->GetInlineContent())->SetSequencer(WeakSequencer.Pin());
+#if USE_LOCAL_DETAILS
 		TArray<TWeakObjectPtr<>> Eulers;
 		TArray<TWeakObjectPtr<>> Transforms;
 		TArray<TWeakObjectPtr<>> TransformNoScales;
@@ -282,7 +283,6 @@ void FControlRigEditMode::SetUpDetailPanel()
 				}
 			}
 		}
-		StaticCastSharedPtr<SControlRigEditModeTools>(Toolkit->GetInlineContent())->SetSequencer(WeakSequencer.Pin());
 		for (TWeakObjectPtr<>& Object : Transforms)
 		{
 			UControlRigControlsProxy* Proxy = Cast<UControlRigControlsProxy>(Object.Get());
@@ -300,9 +300,12 @@ void FControlRigEditMode::SetUpDetailPanel()
 		StaticCastSharedPtr<SControlRigEditModeTools>(Toolkit->GetInlineContent())->SetBoolDetailsObjects(Bools);
 		StaticCastSharedPtr<SControlRigEditModeTools>(Toolkit->GetInlineContent())->SetIntegerDetailsObjects(Integers);
 		StaticCastSharedPtr<SControlRigEditModeTools>(Toolkit->GetInlineContent())->SetEnumDetailsObjects(Enums);
+#else 
+		StaticCastSharedPtr<SControlRigEditModeTools>(Toolkit->GetInlineContent())->SetSettingsDetailsObject(GetMutableDefault<UControlRigEditModeSettings>());	
+#endif
 
-		StaticCastSharedPtr<SControlRigEditModeTools>(Toolkit->GetInlineContent())->SetSettingsDetailsObject(GetMutableDefault<UControlRigEditModeSettings>());
 	}
+
 }
 
 void FControlRigEditMode::SetObjects_Internal()
diff --git a/Engine/Plugins/Animation/ControlRig/Source/ControlRigEditor/Private/EditMode/ControlRigEditMode.h b/Engine/Plugins/Animation/ControlRig/Source/ControlRigEditor/Private/EditMode/ControlRigEditMode.h
--- a/Engine/Plugins/Animation/ControlRig/Source/ControlRigEditor/Private/EditMode/ControlRigEditMode.h
+++ b/Engine/Plugins/Animation/ControlRig/Source/ControlRigEditor/Private/EditMode/ControlRigEditMode.h
@@ -303,8 +303,17 @@ public:
 	/** Get the number of selected RigElements */
 	int32 GetNumSelectedRigElements(uint32 InTypes) const;
 
+	/** Get the control rig*/
 	UControlRig* GetControlRig(bool bInteractionRig, int32 InIndex = 0) const;
 
+	/** Get the detail proxiescontrol rig*/
+	UControlRigDetailPanelControlProxies* GetDetailProxies() { return ControlProxy; }
+
+	/** Get Sequencer Driving This*/
+	TWeakPtr<ISequencer> GetWeakSequencer() { return WeakSequencer; }
+
+	/** Suspend Rig Hierarchy Notifies*/
+	void SuspendHierarchyNotifs(bool bVal) { bSuspendHierarchyNotifs = bVal; }
 private:
 	/** Set a RigElement's selection state */
 	void SetRigElementSelectionInternal(ERigElementType Type, const FName& InRigElementName, bool bSelected);
diff --git a/Engine/Plugins/Animation/ControlRig/Source/ControlRigEditor/Private/EditMode/ControlRigEditModeSettings.h b/Engine/Plugins/Animation/ControlRig/Source/ControlRigEditor/Private/EditMode/ControlRigEditModeSettings.h
--- a/Engine/Plugins/Animation/ControlRig/Source/ControlRigEditor/Private/EditMode/ControlRigEditModeSettings.h
+++ b/Engine/Plugins/Animation/ControlRig/Source/ControlRigEditor/Private/EditMode/ControlRigEditModeSettings.h
@@ -36,39 +36,39 @@ class UControlRigEditModeSettings : public UObject
 public:
 
 	/** Whether to show all bones in the hierarchy */
-	UPROPERTY(config, EditAnywhere, AdvancedDisplay, Category = "Animation")
+	UPROPERTY(config, EditAnywhere, AdvancedDisplay, Category = "Animation Settings")
 	bool bDisplayHierarchy;
 
 	/** Whether to show all nulls in the hierarchy */
-	UPROPERTY(config, EditAnywhere, AdvancedDisplay, Category = "Animation")
+	UPROPERTY(config, EditAnywhere, AdvancedDisplay, Category = "Animation Settings")
 	bool bDisplayNulls;
 
 	/** Should we always hide manipulators in viewport */
-	UPROPERTY(config, EditAnywhere, AdvancedDisplay, Category = "Animation")
+	UPROPERTY(config, EditAnywhere, AdvancedDisplay, Category = "Animation Settings")
 	bool bHideManipulators;
 
 	/** Should we show axes for the selected elements */
-	UPROPERTY(config, EditAnywhere, AdvancedDisplay, Category = "Animation")
+	UPROPERTY(config, EditAnywhere, AdvancedDisplay, Category = "Animation Settings")
 	bool bDisplayAxesOnSelection;
 
 	/** The scale for axes to draw on the selection */
-	UPROPERTY(config, EditAnywhere, AdvancedDisplay, Category = "Animation")
+	UPROPERTY(config, EditAnywhere, AdvancedDisplay, Category = "Animation Settings")
 	float AxisScale;
 
 	/** If true we restore the coordinate space when changing Widget Modes in the Viewport*/
-	UPROPERTY(config, EditAnywhere, AdvancedDisplay, Category = "Animation")
+	UPROPERTY(config, EditAnywhere, AdvancedDisplay, Category = "Animation Settings")
 	bool bCoordSystemPerWidgetMode;
 
 	/** If true we can only select Rig Controls in the scene not other Actors. */
-	UPROPERTY(config, EditAnywhere, AdvancedDisplay, Category = "Animation")
+	UPROPERTY(config, EditAnywhere, AdvancedDisplay, Category = "Animation Settings")
 	bool bOnlySelectRigControls;
 
 	/** If true when we transform multiple selected objects in the viewport they each transforms along their own local transform space */
-	UPROPERTY(config, EditAnywhere, AdvancedDisplay, Category = "Animation")
+	UPROPERTY(config, EditAnywhere, AdvancedDisplay, Category = "Animation Settings")
 	bool bLocalTransformsInEachLocalSpace;
 	
 	/** The scale for Gizmos */
-	UPROPERTY(config, EditAnywhere, AdvancedDisplay, Category = "Animation")
+	UPROPERTY(config, EditAnywhere, AdvancedDisplay, Category = "Animation Settings")
 	float ShapeScale;
 
 	UPROPERTY(config)
diff --git a/Engine/Plugins/Animation/ControlRig/Source/ControlRigEditor/Private/EditMode/ControlRigEditModeToolkit.cpp b/Engine/Plugins/Animation/ControlRig/Source/ControlRigEditor/Private/EditMode/ControlRigEditModeToolkit.cpp
--- a/Engine/Plugins/Animation/ControlRig/Source/ControlRigEditor/Private/EditMode/ControlRigEditModeToolkit.cpp
+++ b/Engine/Plugins/Animation/ControlRig/Source/ControlRigEditor/Private/EditMode/ControlRigEditModeToolkit.cpp
@@ -20,7 +20,9 @@
 #include "Toolkits/AssetEditorModeUILayer.h"
 #include "Widgets/Docking/SDockTab.h"
 #include "ControlRigEditModeSettings.h"
-
+#include "SControlRigDetails.h"
+#include "SControlRigOutliner.h"
+#include "SControlRigSpacePicker.h"
 #define LOCTEXT_NAMESPACE "FControlRigEditModeToolkit"
 
 namespace 
@@ -33,6 +35,9 @@ const FName FControlRigEditModeToolkit::PoseTabName = FName(TEXT("PoseTab"));
 const FName FControlRigEditModeToolkit::MotionTrailTabName = FName(TEXT("MotionTrailTab"));
 const FName FControlRigEditModeToolkit::SnapperTabName = FName(TEXT("SnapperTab"));
 const FName FControlRigEditModeToolkit::TweenOverlayName = FName(TEXT("TweenOverlay"));
+const FName FControlRigEditModeToolkit::OutlinerTabName = FName(TEXT("ControlRigOutlinerTab"));
+const FName FControlRigEditModeToolkit::DetailsTabName = FName(TEXT("ControlRigDetailsTab"));
+const FName FControlRigEditModeToolkit::SpacePickerTabName = FName(TEXT("ControlRigSpacePicker"));
 
 void FControlRigEditModeToolkit::Init(const TSharedPtr<IToolkitHost>& InitToolkitHost)
 {
@@ -85,22 +90,34 @@ void FControlRigEditModeToolkit::OnToolPaletteChanged(FName PaletteName)
 
 void FControlRigEditModeToolkit::TryInvokeToolkitUI(const FName InName)
 {
+	TSharedPtr<FAssetEditorModeUILayer> ModeUILayerPtr = ModeUILayer.Pin();
+
 	if (InName == MotionTrailTabName)
 	{
-		FTabId MotionTrailTabID(MotionTrailTabName);
-		if (ModeUILayer.IsValid())
-		{
-			TSharedPtr<FAssetEditorModeUILayer> ModeUILayerPtr = ModeUILayer.Pin();
-			ModeUILayerPtr->GetTabManager()->TryInvokeTab(MotionTrailTabID);
-		}
+		FTabId TabID(MotionTrailTabName);
+		ModeUILayerPtr->GetTabManager()->TryInvokeTab(TabID);
 	}
-	if (InName == PoseTabName)
+	else if (InName == PoseTabName)
 	{
-		TryInvokePoseTab();
+		FTabId TabID(PoseTabName);
+		ModeUILayerPtr->GetTabManager()->TryInvokeTab(TabID);
 	}
-	if (InName == SnapperTabName)
+	else if (InName == SnapperTabName)
 	{
-		TryInvokeSnapperTab();
+		FTabId TabID(SnapperTabName);
+		ModeUILayerPtr->GetTabManager()->TryInvokeTab(TabID);
+	}
+	else if (InName == OutlinerTabName)
+	{
+		ModeUILayerPtr->GetTabManager()->TryInvokeTab(UAssetEditorUISubsystem::TopRightTabID);
+	}
+	else if (InName == SpacePickerTabName)
+	{
+		ModeUILayerPtr->GetTabManager()->TryInvokeTab(UAssetEditorUISubsystem::BottomLeftTabID);
+	}
+	else if (InName == DetailsTabName)
+	{
+		ModeUILayerPtr->GetTabManager()->TryInvokeTab(UAssetEditorUISubsystem::BottomRightTabID);
 	}
 	else if (InName == TweenOverlayName)
 	{
@@ -159,6 +176,32 @@ TSharedRef<SDockTab> SpawnRigProfiler(const FSpawnTabArgs& Args)
 		];
 }
 
+
+TSharedRef<SDockTab> SpawnOutlinerTab(const FSpawnTabArgs& Args, FControlRigEditMode* InEditorMode)
+{
+	return SNew(SDockTab)
+		[
+			SNew(SControlRigOutliner, *InEditorMode)
+		];
+}
+
+TSharedRef<SDockTab> SpawnSpacePickerTab(const FSpawnTabArgs& Args, FControlRigEditMode* InEditorMode)
+{
+	return SNew(SDockTab)
+		[
+			SNew(SControlRigSpacePicker, *InEditorMode)
+		];
+}
+
+TSharedRef<SDockTab> SpawnDetailsTab(const FSpawnTabArgs& Args, FControlRigEditMode* InEditorMode)
+{
+	return SNew(SDockTab)
+		[
+			SNew(SControlRigDetails, *InEditorMode)
+		];
+}
+
+
 void FControlRigEditModeToolkit::CreateAndShowTweenOverlay()
 {
 	FVector2D NewTweenWidgetLocation = GetDefault<UControlRigEditModeSettings>()->LastInViewportTweenWidgetLocation;
@@ -247,25 +290,51 @@ void FControlRigEditModeToolkit::RequestModeUITabs()
 	{
 		TSharedPtr<FAssetEditorModeUILayer> ModeUILayerPtr = ModeUILayer.Pin();
 		TSharedRef<FWorkspaceItem> MenuGroup = ModeUILayerPtr->GetModeMenuCategory().ToSharedRef();
-	
-		FMinorTabConfig PoseTabInfo;
-		PoseTabInfo.OnSpawnTab = FOnSpawnTab::CreateStatic(&SpawnPoseTab);
-		PoseTabInfo.TabLabel = LOCTEXT("ControlRigPoseTab", "Control Rig Pose");
-		PoseTabInfo.TabTooltip = LOCTEXT("ControlRigPoseTabTooltip", "Show Poses.");
-		ModeUILayerPtr->SetModePanelInfo(UAssetEditorUISubsystem::BottomRightTabID, PoseTabInfo);
-
-		FMinorTabConfig SnapperTabInfo;
-		SnapperTabInfo.OnSpawnTab = FOnSpawnTab::CreateStatic(&SpawnSnapperTab);
-		SnapperTabInfo.TabLabel = LOCTEXT("ControlRigSnapperTab", "Control Rig Snapper");
-		SnapperTabInfo.TabTooltip = LOCTEXT("ControlRigSnapperTabTooltip", "Snap child objects to a parent object over a set of frames.");
-		ModeUILayerPtr->SetModePanelInfo(UAssetEditorUISubsystem::TopRightTabID, SnapperTabInfo);
 
+		FMinorTabConfig DetailTabInfo;
+		DetailTabInfo.OnSpawnTab = FOnSpawnTab::CreateStatic(&SpawnDetailsTab, &EditMode);
+		DetailTabInfo.TabLabel = LOCTEXT("ControlRigDetailTab", "Anim Details");
+		DetailTabInfo.TabTooltip = LOCTEXT("ControlRigDetailTabTooltip", "Show Details For Selected Controls.");
+		ModeUILayerPtr->SetModePanelInfo(UAssetEditorUISubsystem::BottomRightTabID, DetailTabInfo);
+
+		FMinorTabConfig OutlinerTabInfo;
+		OutlinerTabInfo.OnSpawnTab = FOnSpawnTab::CreateStatic(&SpawnOutlinerTab, &EditMode);
+		OutlinerTabInfo.TabLabel = LOCTEXT("AnimationOutlinerTab", "Anim Outliner");
+		OutlinerTabInfo.TabTooltip = LOCTEXT("AnimationOutlinerTabTooltip", "Control Rig Controls");
+		ModeUILayerPtr->SetModePanelInfo(UAssetEditorUISubsystem::TopRightTabID, OutlinerTabInfo);
+		/* doesn't work as expected
+		FMinorTabConfig SpawnSpacePickerTabInfo;
+		SpawnSpacePickerTabInfo.OnSpawnTab = FOnSpawnTab::CreateStatic(&SpawnSpacePickerTab, &EditMode);
+		SpawnSpacePickerTabInfo.TabLabel = LOCTEXT("ControlRigSpacePickerTab", "Control Rig Space Picker");
+		SpawnSpacePickerTabInfo.TabTooltip = LOCTEXT("ControlRigSpacePickerTabTooltip", "Control Rig Space Picker");
+		ModeUILayerPtr->SetModePanelInfo(UAssetEditorUISubsystem::TopLeftTabID, SpawnSpacePickerTabInfo);
+		*/
+
+		ModeUILayerPtr->GetTabManager()->UnregisterTabSpawner(SnapperTabName);
+		ModeUILayerPtr->GetTabManager()->RegisterTabSpawner(SnapperTabName, FOnSpawnTab::CreateStatic(&SpawnSnapperTab))
+			.SetDisplayName(LOCTEXT("ControlRigSnapperTab", "Control Rig Snapper"))
+			.SetTooltipText(LOCTEXT("ControlRigSnapperTabTooltip", "Snap child objects to a parent object over a set of frames."))
+			.SetGroup(MenuGroup)
+			.SetIcon(FSlateIcon(TEXT("ControlRigEditorStyle"), TEXT("ControlRig.SnapperTool")));
+		ModeUILayerPtr->GetTabManager()->RegisterDefaultTabWindowSize(SnapperTabName, FVector2D(300, 325));
+
+		ModeUILayerPtr->GetTabManager()->UnregisterTabSpawner(PoseTabName);
+		ModeUILayerPtr->GetTabManager()->RegisterTabSpawner(PoseTabName, FOnSpawnTab::CreateStatic(&SpawnPoseTab))
+			.SetDisplayName(LOCTEXT("ControlRigPoseTab", "Control Rig Pose"))
+			.SetTooltipText(LOCTEXT("ControlRigPoseTabTooltip", "Show Poses."))
+			.SetGroup(MenuGroup)
+			.SetIcon(FSlateIcon(TEXT("ControlRigEditorStyle"), TEXT("ControlRig.PoseTool")));
+		ModeUILayerPtr->GetTabManager()->RegisterDefaultTabWindowSize(PoseTabName, FVector2D(675, 625));
+
+		ModeUILayerPtr->GetTabManager()->UnregisterTabSpawner(MotionTrailTabName);
 		ModeUILayerPtr->GetTabManager()->RegisterTabSpawner(MotionTrailTabName, FOnSpawnTab::CreateStatic(&SpawnMotionTrailTab))
 			.SetDisplayName(LOCTEXT("MotionTrailTab", "Motion Trail"))
 			.SetTooltipText(LOCTEXT("MotionTrailTabTooltip", "Display motion trails for animated objects."))
 			.SetGroup(MenuGroup)
-			.SetIcon(FSlateIcon(TEXT("ControlRigEditorStyle"), TEXT("HierarchicalProfiler.TabIcon")));
+			.SetIcon(FSlateIcon(TEXT("ControlRigEditorStyle"), TEXT("ControlRig.EditableMotionTrails")));
+		ModeUILayerPtr->GetTabManager()->RegisterDefaultTabWindowSize(MotionTrailTabName, FVector2D(425, 575));
 
+		ModeUILayerPtr->GetTabManager()->UnregisterTabSpawner("HierarchicalProfiler");
 		ModeUILayerPtr->GetTabManager()->RegisterTabSpawner("HierarchicalProfiler", FOnSpawnTab::CreateStatic(&SpawnRigProfiler))
 			.SetDisplayName(LOCTEXT("HierarchicalProfilerTab", "Hierarchical Profiler"))
 			.SetTooltipText(LOCTEXT("HierarchicalProfilerTooltip", "Open the Hierarchical Profiler tab."))
@@ -273,6 +342,8 @@ void FControlRigEditModeToolkit::RequestModeUITabs()
 			.SetIcon(FSlateIcon(TEXT("ControlRigEditorStyle"), TEXT("HierarchicalProfiler.TabIcon")));
 
 		ModeUILayer.Pin()->ToolkitHostShutdownUI().BindSP(this, &FControlRigEditModeToolkit::UnregisterAndRemoveFloatingTabs);
+
+
 	}
 };
 
@@ -280,7 +351,13 @@ void FControlRigEditModeToolkit::InvokeUI()
 {
 	FModeToolkit::InvokeUI();
 
-// TODO: any future default tabs will go here
+	if (ModeUILayer.IsValid())
+	{
+		TSharedPtr<FAssetEditorModeUILayer> ModeUILayerPtr = ModeUILayer.Pin();
+		ModeUILayerPtr->GetTabManager()->TryInvokeTab(UAssetEditorUISubsystem::TopRightTabID);
+		// doesn't work as expected todo ModeUILayerPtr->GetTabManager()->TryInvokeTab(UAssetEditorUISubsystem::TopLeftTabID);
+		ModeUILayerPtr->GetTabManager()->TryInvokeTab(UAssetEditorUISubsystem::BottomRightTabID);
+	}	
 }
 
 void FControlRigEditModeToolkit::UnregisterAndRemoveFloatingTabs()
@@ -292,32 +369,34 @@ void FControlRigEditModeToolkit::UnregisterAndRemoveFloatingTabs()
 		{
 			TSharedPtr<FAssetEditorModeUILayer> ModeUILayerPtr = ModeUILayer.Pin();
 			TSharedPtr<SDockTab> ProfilerTab = ModeUILayerPtr->GetTabManager()->FindExistingLiveTab(FTabId("HierarchicalProfiler"));
-			if(ProfilerTab)
-			{ 
+			if (ProfilerTab)
+			{
 				ProfilerTab->RequestCloseTab();
 			}
 			ModeUILayerPtr->GetTabManager()->UnregisterTabSpawner("HierarchicalProfiler");
+
 			TSharedPtr<SDockTab> MotionTrailTab = ModeUILayerPtr->GetTabManager()->FindExistingLiveTab(FTabId(MotionTrailTabName));
 			if (MotionTrailTab)
 			{
 				MotionTrailTab->RequestCloseTab();
 			}
 			ModeUILayerPtr->GetTabManager()->UnregisterTabSpawner(MotionTrailTabName);
+
+			TSharedPtr<SDockTab> SnapperTab = ModeUILayerPtr->GetTabManager()->FindExistingLiveTab(FTabId(SnapperTabName));
+			if (SnapperTab)
+			{
+				SnapperTab->RequestCloseTab();
+			}
+			ModeUILayerPtr->GetTabManager()->UnregisterTabSpawner(SnapperTabName);
+		
+			TSharedPtr<SDockTab> PoseTab = ModeUILayerPtr->GetTabManager()->FindExistingLiveTab(FTabId(PoseTabName));
+			if (PoseTab)
+			{
+				PoseTab->RequestCloseTab();
+			}
+			ModeUILayerPtr->GetTabManager()->UnregisterTabSpawner(PoseTabName);
 		}
 	}
 }
 
-void FControlRigEditModeToolkit::TryInvokeSnapperTab()
-{
-	TSharedPtr<FAssetEditorModeUILayer> ModeUILayerPtr = ModeUILayer.Pin();
-	ModeUILayerPtr->GetTabManager()->TryInvokeTab(UAssetEditorUISubsystem::TopRightTabID);
-}
-
-void FControlRigEditModeToolkit::TryInvokePoseTab()
-{
-	TSharedPtr<FAssetEditorModeUILayer> ModeUILayerPtr = ModeUILayer.Pin();
-	ModeUILayerPtr->GetTabManager()->TryInvokeTab(UAssetEditorUISubsystem::BottomRightTabID);
-}
-
-
 #undef LOCTEXT_NAMESPACE
\ No newline at end of file
diff --git a/Engine/Plugins/Animation/ControlRig/Source/ControlRigEditor/Private/EditMode/ControlRigEditModeToolkit.h b/Engine/Plugins/Animation/ControlRig/Source/ControlRigEditor/Private/EditMode/ControlRigEditModeToolkit.h
--- a/Engine/Plugins/Animation/ControlRig/Source/ControlRigEditor/Private/EditMode/ControlRigEditModeToolkit.h
+++ b/Engine/Plugins/Animation/ControlRig/Source/ControlRigEditor/Private/EditMode/ControlRigEditModeToolkit.h
@@ -54,8 +54,12 @@ public:
 	static const FName MotionTrailTabName;
 	static const FName TweenOverlayName;
 	static const FName SnapperTabName;
+	static const FName DetailsTabName;
+	static const FName OutlinerTabName;
+	static const FName SpacePickerTabName;
 
 protected:
+
 	void CreateAndShowTweenOverlay();
 	void TryShowTweenOverlay();
 	void RemoveAndDestroyTweenOverlay();
@@ -72,8 +76,7 @@ protected:
 	void UnregisterAndRemoveFloatingTabs();
 
 private:
-	void TryInvokeSnapperTab();
-	void TryInvokePoseTab();
+
 
 private:
 	/** The edit mode we are bound to */
diff --git a/Engine/Plugins/Animation/ControlRig/Source/ControlRigEditor/Private/EditMode/SControlRigDetails.cpp b/Engine/Plugins/Animation/ControlRig/Source/ControlRigEditor/Private/EditMode/SControlRigDetails.cpp
new file mode 100644
--- /dev/null
+++ b/Engine/Plugins/Animation/ControlRig/Source/ControlRigEditor/Private/EditMode/SControlRigDetails.cpp
@@ -0,0 +1,668 @@
+// Copyright Epic Games, Inc. All Rights Reserved.
+
+#include "SControlRigDetails.h"
+#include "Framework/Notifications/NotificationManager.h"
+#include "Widgets/Notifications/SNotificationList.h"
+#include "Framework/MultiBox/MultiBoxBuilder.h"
+#include "AssetData.h"
+#include "EditorStyleSet.h"
+#include "Widgets/Text/SInlineEditableTextBlock.h"
+#include "EditorStyleSet.h"
+#include "Styling/CoreStyle.h"
+#include "ScopedTransaction.h"
+#include "ControlRig.h"
+#include "UnrealEdGlobals.h"
+#include "ControlRigEditMode.h"
+#include "EditorModeManager.h"
+#include "ISequencer.h"
+#include "LevelSequence.h"
+#include "UnrealEd/Public/Selection.h"
+#include "Editor.h"
+#include "LevelEditor.h"
+#include "Widgets/Views/SListView.h"
+#include "Widgets/Text/STextBlock.h"
+#include "Widgets/Input/SButton.h"
+#include "Widgets/Layout/SScrollBox.h"
+#include "ControlRigControlsProxy.h"
+#include "ControlRigEditModeSettings.h"
+#include "TimerManager.h"
+
+#define LOCTEXT_NAMESPACE "ControlRigDetails"
+
+void FControlRigEditModeGenericDetails::CustomizeDetails(class IDetailLayoutBuilder& DetailLayout)
+{
+	TArray<TWeakObjectPtr<UObject>> ObjectsBeingCustomized;
+	DetailLayout.GetObjectsBeingCustomized(ObjectsBeingCustomized);
+
+	TArray<UControlRigControlsProxy*> ProxiesBeingCustomized;
+	for (TWeakObjectPtr<UObject> ObjectBeingCustomized : ObjectsBeingCustomized)
+	{
+		if (UControlRigControlsProxy* Proxy = Cast< UControlRigControlsProxy>(ObjectBeingCustomized.Get()))
+		{
+			ProxiesBeingCustomized.Add(Proxy);
+		}
+	}
+
+	if (ProxiesBeingCustomized.Num() == 0 || ProxiesBeingCustomized[0]->GetControlElement() == nullptr)
+	{
+		return;
+	}
+	FText ControlText = FText::FromName(ProxiesBeingCustomized[0]->GetName());
+
+	if (ProxiesBeingCustomized.Num() > 1)
+	{
+		if (ProxiesBeingCustomized[0]->GetClass() == UControlRigTransformControlProxy::StaticClass())
+		{
+			ControlText = LOCTEXT("TransformChannels", "Transform Channels");
+		}
+		else if (ProxiesBeingCustomized[0]->GetClass() == UControlRigTransformNoScaleControlProxy::StaticClass())
+		{
+			ControlText = LOCTEXT("TransformNoScaleChannels", "TransformNoScale Channels");
+		}
+		else if (ProxiesBeingCustomized[0]->GetClass() == UControlRigEulerTransformControlProxy::StaticClass())
+		{
+			ControlText = LOCTEXT("EulerTransformChannels", "Euler Transform Channels");
+		}
+		else if (ProxiesBeingCustomized[0]->GetClass() == UControlRigFloatControlProxy::StaticClass())
+		{
+			ControlText = LOCTEXT("FloatChannels", "Float Channels");
+		}
+		else if (ProxiesBeingCustomized[0]->GetClass() == UControlRigVectorControlProxy::StaticClass())
+		{
+			ControlText = LOCTEXT("VectorChannels", "Vector Channels");
+		}
+		else if (ProxiesBeingCustomized[0]->GetClass() == UControlRigVector2DControlProxy::StaticClass())
+		{
+			ControlText = LOCTEXT("Vector2DChannels", "Vector2D Channels");
+		}
+		else if (ProxiesBeingCustomized[0]->GetClass() == UControlRigBoolControlProxy::StaticClass())
+		{
+			ControlText = LOCTEXT("BoolChannels", "Bool Channels");
+		}
+		else if (ProxiesBeingCustomized[0]->GetClass() == UControlRigEnumControlProxy::StaticClass())
+		{
+			ControlText = LOCTEXT("EnumChannels", "Enum Channels");
+		}
+		else if (ProxiesBeingCustomized[0]->GetClass() == UControlRigIntegerControlProxy::StaticClass())
+		{
+			ControlText = LOCTEXT("IntegerChannels", "Integer Channels");
+		}
+	}
+
+	IDetailCategoryBuilder& Category = DetailLayout.EditCategory(TEXT("Control"), ControlText);
+	for (UControlRigControlsProxy* Proxy : ProxiesBeingCustomized)
+	{
+		FRigControlElement* ControlElement = Proxy->GetControlElement();
+		if (ControlElement == nullptr)
+		{
+			continue;
+		}
+
+		FName ValuePropertyName = TEXT("Transform");
+		if (ControlElement->Settings.ControlType == ERigControlType::Float)
+		{
+			ValuePropertyName = TEXT("Float");
+		}
+		else if (ControlElement->Settings.ControlType == ERigControlType::Integer)
+		{
+			if (ControlElement->Settings.ControlEnum == nullptr)
+			{
+				ValuePropertyName = TEXT("Integer");
+			}
+			else
+			{
+				ValuePropertyName = TEXT("Enum");
+			}
+		}
+		else if (ControlElement->Settings.ControlType == ERigControlType::Bool)
+		{
+			ValuePropertyName = TEXT("Bool");
+		}
+		else if (ControlElement->Settings.ControlType == ERigControlType::Position ||
+			ControlElement->Settings.ControlType == ERigControlType::Scale)
+		{
+			ValuePropertyName = TEXT("Vector");
+		}
+		else if (ControlElement->Settings.ControlType == ERigControlType::Vector2D)
+		{
+			ValuePropertyName = TEXT("Vector2D");
+		}
+		else if (ControlElement->Settings.ControlType == ERigControlType::EulerTransform)
+		{
+			ValuePropertyName = TEXT("EulerTransform");
+		}
+		else if (ControlElement->Settings.ControlType == ERigControlType::TransformNoScale)
+		{
+			ValuePropertyName = TEXT("TransformNoScale");
+		}
+
+		TSharedPtr<IPropertyHandle> ValuePropertyHandle = DetailLayout.GetProperty(ValuePropertyName, Proxy->GetClass());
+		if (ValuePropertyHandle)
+		{
+			ValuePropertyHandle->SetPropertyDisplayName(FText::FromName(Proxy->GetName()));
+		}
+
+		URigHierarchy* Hierarchy = Proxy->ControlRig->GetHierarchy();
+		Hierarchy->ForEach<FRigControlElement>([Hierarchy, Proxy, &Category, this](FRigControlElement* ControlElement) -> bool
+			{
+				FName ParentControlName = NAME_None;
+				FRigControlElement* ParentControlElement = Cast<FRigControlElement>(Hierarchy->GetFirstParent(ControlElement));
+				if (ParentControlElement)
+				{
+					ParentControlName = ParentControlElement->GetName();
+				}
+
+				if (ParentControlName == ControlElement->GetName())
+				{
+					if (FControlRigEditMode* EditMode = static_cast<FControlRigEditMode*>(ModeTools->GetActiveMode(FControlRigEditMode::ModeName)))
+					{
+						if (UObject* NestedProxy = EditMode->ControlProxy->FindProxy(ControlElement->GetName()))
+						{
+							FName PropertyName(NAME_None);
+							switch (ControlElement->Settings.ControlType)
+							{
+							case ERigControlType::Bool:
+							{
+								PropertyName = TEXT("Bool");
+								break;
+							}
+							case ERigControlType::Float:
+							{
+								PropertyName = TEXT("Float");
+								break;
+							}
+							case ERigControlType::Integer:
+							{
+								if (ControlElement->Settings.ControlEnum == nullptr)
+								{
+									PropertyName = TEXT("Integer");
+								}
+								else
+								{
+									PropertyName = TEXT("Enum");
+								}
+								break;
+							}
+							default:
+							{
+								break;
+							}
+							}
+
+							if (PropertyName.IsNone())
+							{
+								return true;
+							}
+
+							TArray<UObject*> NestedProxies;
+							NestedProxies.Add(NestedProxy);
+
+							FAddPropertyParams Params;
+							Params.CreateCategoryNodes(false);
+							IDetailPropertyRow* NestedRow = Category.AddExternalObjectProperty(
+								NestedProxies,
+								PropertyName,
+								EPropertyLocation::Advanced,
+								Params);
+							NestedRow->DisplayName(FText::FromName(ControlElement->Settings.DisplayName));
+
+							Category.SetShowAdvanced(true);
+						}
+					}
+				}
+				return true;
+			});
+	}
+}
+
+void SControlRigDetails::Construct(const FArguments& InArgs, FControlRigEditMode& InEditMode)
+{
+	ModeTools = InEditMode.GetModeManager();
+	FDetailsViewArgs DetailsViewArgs;
+	{
+		DetailsViewArgs.bAllowSearch = true;
+		DetailsViewArgs.bHideSelectionTip = true;
+		DetailsViewArgs.bLockable = false;
+		DetailsViewArgs.bSearchInitialKeyFocus = true;
+		DetailsViewArgs.bUpdatesFromSelection = false;
+		DetailsViewArgs.bShowOptions = false;
+		DetailsViewArgs.bShowModifiedPropertiesOption = true;
+		DetailsViewArgs.bCustomNameAreaLocation = true;
+		DetailsViewArgs.bCustomFilterAreaLocation = true;
+		DetailsViewArgs.NameAreaSettings = FDetailsViewArgs::HideNameArea;
+		DetailsViewArgs.bAllowMultipleTopLevelObjects = false;
+		DetailsViewArgs.bShowScrollBar = false; // Don't need to show this, as we are putting it in a scroll box
+	}
+
+	ControlEulerTransformDetailsView = FModuleManager::GetModuleChecked<FPropertyEditorModule>("PropertyEditor").CreateDetailView(DetailsViewArgs);
+	ControlEulerTransformDetailsView->SetKeyframeHandler(SharedThis(this));
+	ControlEulerTransformDetailsView->SetIsPropertyVisibleDelegate(FIsPropertyVisible::CreateSP(this, &SControlRigDetails::ShouldShowPropertyOnDetailCustomization));
+	ControlEulerTransformDetailsView->SetIsPropertyReadOnlyDelegate(FIsPropertyReadOnly::CreateSP(this, &SControlRigDetails::IsReadOnlyPropertyOnDetailCustomization));
+	ControlEulerTransformDetailsView->SetGenericLayoutDetailsDelegate(FOnGetDetailCustomizationInstance::CreateStatic(&FControlRigEditModeGenericDetails::MakeInstance, ModeTools));
+
+	ControlTransformDetailsView = FModuleManager::GetModuleChecked<FPropertyEditorModule>("PropertyEditor").CreateDetailView(DetailsViewArgs);
+	ControlTransformDetailsView->SetKeyframeHandler(SharedThis(this));
+	ControlTransformDetailsView->SetIsPropertyVisibleDelegate(FIsPropertyVisible::CreateSP(this, &SControlRigDetails::ShouldShowPropertyOnDetailCustomization));
+	ControlTransformDetailsView->SetIsPropertyReadOnlyDelegate(FIsPropertyReadOnly::CreateSP(this, &SControlRigDetails::IsReadOnlyPropertyOnDetailCustomization));
+	ControlTransformDetailsView->SetGenericLayoutDetailsDelegate(FOnGetDetailCustomizationInstance::CreateStatic(&FControlRigEditModeGenericDetails::MakeInstance, ModeTools));
+
+	ControlTransformNoScaleDetailsView = FModuleManager::GetModuleChecked<FPropertyEditorModule>("PropertyEditor").CreateDetailView(DetailsViewArgs);
+	ControlTransformNoScaleDetailsView->SetKeyframeHandler(SharedThis(this));
+	ControlTransformNoScaleDetailsView->SetIsPropertyVisibleDelegate(FIsPropertyVisible::CreateSP(this, &SControlRigDetails::ShouldShowPropertyOnDetailCustomization));
+	ControlTransformNoScaleDetailsView->SetIsPropertyReadOnlyDelegate(FIsPropertyReadOnly::CreateSP(this, &SControlRigDetails::IsReadOnlyPropertyOnDetailCustomization));
+	ControlTransformNoScaleDetailsView->SetGenericLayoutDetailsDelegate(FOnGetDetailCustomizationInstance::CreateStatic(&FControlRigEditModeGenericDetails::MakeInstance, ModeTools));
+
+	ControlFloatDetailsView = FModuleManager::GetModuleChecked<FPropertyEditorModule>("PropertyEditor").CreateDetailView(DetailsViewArgs);
+	ControlFloatDetailsView->SetKeyframeHandler(SharedThis(this));
+	ControlFloatDetailsView->SetIsPropertyVisibleDelegate(FIsPropertyVisible::CreateSP(this, &SControlRigDetails::ShouldShowPropertyOnDetailCustomization));
+	ControlFloatDetailsView->SetIsPropertyReadOnlyDelegate(FIsPropertyReadOnly::CreateSP(this, &SControlRigDetails::IsReadOnlyPropertyOnDetailCustomization));
+	ControlFloatDetailsView->SetGenericLayoutDetailsDelegate(FOnGetDetailCustomizationInstance::CreateStatic(&FControlRigEditModeGenericDetails::MakeInstance, ModeTools));
+
+	ControlEnumDetailsView = FModuleManager::GetModuleChecked<FPropertyEditorModule>("PropertyEditor").CreateDetailView(DetailsViewArgs);
+	ControlEnumDetailsView->SetKeyframeHandler(SharedThis(this));
+	ControlEnumDetailsView->SetIsPropertyVisibleDelegate(FIsPropertyVisible::CreateSP(this, &SControlRigDetails::ShouldShowPropertyOnDetailCustomization));
+	ControlEnumDetailsView->SetIsPropertyReadOnlyDelegate(FIsPropertyReadOnly::CreateSP(this, &SControlRigDetails::IsReadOnlyPropertyOnDetailCustomization));
+	ControlEnumDetailsView->SetGenericLayoutDetailsDelegate(FOnGetDetailCustomizationInstance::CreateStatic(&FControlRigEditModeGenericDetails::MakeInstance, ModeTools));
+
+	ControlIntegerDetailsView = FModuleManager::GetModuleChecked<FPropertyEditorModule>("PropertyEditor").CreateDetailView(DetailsViewArgs);
+	ControlIntegerDetailsView->SetKeyframeHandler(SharedThis(this));
+	ControlIntegerDetailsView->SetIsPropertyVisibleDelegate(FIsPropertyVisible::CreateSP(this, &SControlRigDetails::ShouldShowPropertyOnDetailCustomization));
+	ControlIntegerDetailsView->SetIsPropertyReadOnlyDelegate(FIsPropertyReadOnly::CreateSP(this, &SControlRigDetails::IsReadOnlyPropertyOnDetailCustomization));
+	ControlIntegerDetailsView->SetGenericLayoutDetailsDelegate(FOnGetDetailCustomizationInstance::CreateStatic(&FControlRigEditModeGenericDetails::MakeInstance, ModeTools));
+
+	ControlBoolDetailsView = FModuleManager::GetModuleChecked<FPropertyEditorModule>("PropertyEditor").CreateDetailView(DetailsViewArgs);
+	ControlBoolDetailsView->SetKeyframeHandler(SharedThis(this));
+	ControlBoolDetailsView->SetIsPropertyVisibleDelegate(FIsPropertyVisible::CreateSP(this, &SControlRigDetails::ShouldShowPropertyOnDetailCustomization));
+	ControlBoolDetailsView->SetIsPropertyReadOnlyDelegate(FIsPropertyReadOnly::CreateSP(this, &SControlRigDetails::IsReadOnlyPropertyOnDetailCustomization));
+	ControlBoolDetailsView->SetGenericLayoutDetailsDelegate(FOnGetDetailCustomizationInstance::CreateStatic(&FControlRigEditModeGenericDetails::MakeInstance, ModeTools));
+
+	ControlVectorDetailsView = FModuleManager::GetModuleChecked<FPropertyEditorModule>("PropertyEditor").CreateDetailView(DetailsViewArgs);
+	ControlVectorDetailsView->SetKeyframeHandler(SharedThis(this));
+	ControlVectorDetailsView->SetIsPropertyVisibleDelegate(FIsPropertyVisible::CreateSP(this, &SControlRigDetails::ShouldShowPropertyOnDetailCustomization));
+	ControlVectorDetailsView->SetIsPropertyReadOnlyDelegate(FIsPropertyReadOnly::CreateSP(this, &SControlRigDetails::IsReadOnlyPropertyOnDetailCustomization));
+	ControlVectorDetailsView->SetGenericLayoutDetailsDelegate(FOnGetDetailCustomizationInstance::CreateStatic(&FControlRigEditModeGenericDetails::MakeInstance, ModeTools));
+
+	ControlVector2DDetailsView = FModuleManager::GetModuleChecked<FPropertyEditorModule>("PropertyEditor").CreateDetailView(DetailsViewArgs);
+	ControlVector2DDetailsView->SetKeyframeHandler(SharedThis(this));
+	ControlVector2DDetailsView->SetIsPropertyVisibleDelegate(FIsPropertyVisible::CreateSP(this, &SControlRigDetails::ShouldShowPropertyOnDetailCustomization));
+	ControlVector2DDetailsView->SetIsPropertyReadOnlyDelegate(FIsPropertyReadOnly::CreateSP(this, &SControlRigDetails::IsReadOnlyPropertyOnDetailCustomization));
+	ControlVector2DDetailsView->SetGenericLayoutDetailsDelegate(FOnGetDetailCustomizationInstance::CreateStatic(&FControlRigEditModeGenericDetails::MakeInstance, ModeTools));
+
+	ChildSlot
+		[
+			SNew(SScrollBox)
+			+ SScrollBox::Slot()
+			[
+			SNew(SVerticalBox)
+
+			
+			+ SVerticalBox::Slot()
+				.AutoHeight()
+				[
+					ControlEulerTransformDetailsView.ToSharedRef()
+				]
+			+ SVerticalBox::Slot()
+				.AutoHeight()
+				[
+					ControlTransformDetailsView.ToSharedRef()
+				]
+			+ SVerticalBox::Slot()
+				.AutoHeight()
+				[
+					ControlTransformNoScaleDetailsView.ToSharedRef()
+				]
+			+ SVerticalBox::Slot()
+				.AutoHeight()
+				[
+					ControlBoolDetailsView.ToSharedRef()
+				]
+			+ SVerticalBox::Slot()
+				.AutoHeight()
+				[
+					ControlIntegerDetailsView.ToSharedRef()
+				]
+			+ SVerticalBox::Slot()
+				.AutoHeight()
+				[
+					ControlEnumDetailsView.ToSharedRef()
+				]
+			+ SVerticalBox::Slot()
+				.AutoHeight()
+				[
+					ControlVectorDetailsView.ToSharedRef()
+				]
+			+ SVerticalBox::Slot()
+				.AutoHeight()
+				[
+					ControlVector2DDetailsView.ToSharedRef()
+				]
+			+ SVerticalBox::Slot()
+				.AutoHeight()
+				[
+					ControlFloatDetailsView.ToSharedRef()
+				]
+
+			]
+		];
+
+	SetEditMode(InEditMode);
+}
+
+SControlRigDetails::~SControlRigDetails()
+{
+	//base class handles control rig related cleanup
+}
+
+
+void SControlRigDetails::HandleControlAdded(UControlRig* ControlRig, bool bIsAdded)
+{
+	FControlRigBaseDockableView::HandleControlAdded(ControlRig, bIsAdded);
+	UpdateProxies();
+}
+
+void SControlRigDetails::HandleControlSelected(UControlRig* Subject, FRigControlElement* InControl, bool bSelected)
+{
+	FControlRigBaseDockableView::HandleControlSelected(Subject, InControl, bSelected);
+	UpdateProxies();
+}
+
+void SControlRigDetails::UpdateProxies()
+{
+	//proxies that are in edit mode are also listening to the same messages so they may not be set up yet so need to wait
+	GEditor->GetTimerManager()->SetTimerForNextTick([this]()
+	{
+		TArray<TWeakObjectPtr<>> Eulers;
+		TArray<TWeakObjectPtr<>> Transforms;
+		TArray<TWeakObjectPtr<>> TransformNoScales;
+		TArray<TWeakObjectPtr<>> Floats;
+		TArray<TWeakObjectPtr<>> Vectors;
+		TArray<TWeakObjectPtr<>> Vector2Ds;
+		TArray<TWeakObjectPtr<>> Bools;
+		TArray<TWeakObjectPtr<>> Integers;
+		TArray<TWeakObjectPtr<>> Enums;
+		if (UControlRig* ControlRig = GetControlRig())
+		{
+			if (FControlRigEditMode* EditMode = static_cast<FControlRigEditMode*>(ModeTools->GetActiveMode(FControlRigEditMode::ModeName)))
+			{
+				if (UControlRigDetailPanelControlProxies* ControlProxy = EditMode->GetDetailProxies())
+				{
+					const TArray<UControlRigControlsProxy*>& Proxies = ControlProxy->GetSelectedProxies();
+					for (UControlRigControlsProxy* Proxy : Proxies)
+					{
+						if (Proxy->GetClass() == UControlRigTransformControlProxy::StaticClass())
+						{
+							Transforms.Add(Proxy);
+						}
+						else if (Proxy->GetClass() == UControlRigTransformNoScaleControlProxy::StaticClass())
+						{
+							TransformNoScales.Add(Proxy);
+						}
+						else if (Proxy->GetClass() == UControlRigEulerTransformControlProxy::StaticClass())
+						{
+							Eulers.Add(Proxy);
+						}
+						else if (Proxy->GetClass() == UControlRigFloatControlProxy::StaticClass())
+						{
+							Floats.Add(Proxy);
+						}
+						else if (Proxy->GetClass() == UControlRigVectorControlProxy::StaticClass())
+						{
+							Vectors.Add(Proxy);
+						}
+						else if (Proxy->GetClass() == UControlRigVector2DControlProxy::StaticClass())
+						{
+							Vector2Ds.Add(Proxy);
+						}
+						else if (Proxy->GetClass() == UControlRigBoolControlProxy::StaticClass())
+						{
+							Bools.Add(Proxy);
+						}
+						else if (Proxy->GetClass() == UControlRigEnumControlProxy::StaticClass())
+						{
+							Enums.Add(Proxy);
+						}
+						else if (Proxy->GetClass() == UControlRigIntegerControlProxy::StaticClass())
+						{
+							Integers.Add(Proxy);
+						}
+					}
+				}
+				for (TWeakObjectPtr<>& Object : Transforms)
+				{
+					UControlRigControlsProxy* Proxy = Cast<UControlRigControlsProxy>(Object.Get());
+					if (Proxy)
+					{
+						Proxy->SetIsMultiple(Transforms.Num() > 1);
+					}
+				}
+			}
+		}
+		SetTransformDetailsObjects(Transforms);
+		SetTransformNoScaleDetailsObjects(TransformNoScales);
+		SetEulerTransformDetailsObjects(Eulers);
+		SetFloatDetailsObjects(Floats);
+		SetVectorDetailsObjects(Vectors);
+		SetVector2DDetailsObjects(Vector2Ds);
+		SetBoolDetailsObjects(Bools);
+		SetIntegerDetailsObjects(Integers);
+		SetEnumDetailsObjects(Enums);
+	});
+}
+
+void SControlRigDetails::SetEulerTransformDetailsObjects(const TArray<TWeakObjectPtr<>>& InObjects)
+{
+	if (ControlEulerTransformDetailsView)
+	{
+		ControlEulerTransformDetailsView->SetObjects(InObjects);
+	}
+};
+
+void SControlRigDetails::SetTransformDetailsObjects(const TArray<TWeakObjectPtr<>>& InObjects)
+{
+	if (ControlTransformDetailsView)
+	{
+		ControlTransformDetailsView->SetObjects(InObjects);
+	}
+}
+
+void SControlRigDetails::SetTransformNoScaleDetailsObjects(const TArray<TWeakObjectPtr<>>& InObjects)
+{
+	if (ControlTransformNoScaleDetailsView)
+	{
+		ControlTransformNoScaleDetailsView->SetObjects(InObjects);
+	}
+}
+
+void SControlRigDetails::SetFloatDetailsObjects(const TArray<TWeakObjectPtr<>>& InObjects)
+{
+	if (ControlFloatDetailsView)
+	{
+		ControlFloatDetailsView->SetObjects(InObjects);
+	}
+}
+
+void SControlRigDetails::SetBoolDetailsObjects(const TArray<TWeakObjectPtr<>>& InObjects)
+{
+	if (ControlBoolDetailsView)
+	{
+		ControlBoolDetailsView->SetObjects(InObjects);
+	}
+}
+
+void SControlRigDetails::SetIntegerDetailsObjects(const TArray<TWeakObjectPtr<>>& InObjects)
+{
+	if (ControlIntegerDetailsView)
+	{
+		ControlIntegerDetailsView->SetObjects(InObjects);
+	}
+}
+void SControlRigDetails::SetEnumDetailsObjects(const TArray<TWeakObjectPtr<>>& InObjects)
+{
+	if (ControlVectorDetailsView)
+	{
+		ControlVectorDetailsView->SetObjects(InObjects);
+	}
+}
+
+void SControlRigDetails::SetVectorDetailsObjects(const TArray<TWeakObjectPtr<>>& InObjects)
+{
+	if (ControlEnumDetailsView)
+	{
+		ControlEnumDetailsView->SetObjects(InObjects);
+	}
+}
+
+void SControlRigDetails::SetVector2DDetailsObjects(const TArray<TWeakObjectPtr<>>& InObjects)
+{
+	if (ControlVector2DDetailsView)
+	{
+		ControlVector2DDetailsView->SetObjects(InObjects);
+	}
+}
+bool SControlRigDetails::IsPropertyKeyable(const UClass* InObjectClass, const IPropertyHandle& InPropertyHandle) const
+{
+	if (InObjectClass && InObjectClass->IsChildOf(UControlRigTransformNoScaleControlProxy::StaticClass()) && InObjectClass->IsChildOf(UControlRigEulerTransformControlProxy::StaticClass()) && InPropertyHandle.GetProperty()
+		&& InPropertyHandle.GetProperty()->GetFName() == GET_MEMBER_NAME_CHECKED(UControlRigTransformControlProxy, Transform))
+	{
+		return true;
+	}
+	FCanKeyPropertyParams CanKeyPropertyParams(InObjectClass, InPropertyHandle);
+	ISequencer* Sequencer = GetSequencer();
+	if (Sequencer && Sequencer->CanKeyProperty(CanKeyPropertyParams))
+	{
+		return true;
+	}
+
+	return false;
+}
+
+bool SControlRigDetails::IsPropertyKeyingEnabled() const
+{
+	ISequencer* Sequencer = GetSequencer();
+	if (Sequencer && Sequencer->GetFocusedMovieSceneSequence())
+	{
+		return true;
+	}
+
+	return false;
+}
+
+bool SControlRigDetails::IsPropertyAnimated(const IPropertyHandle& PropertyHandle, UObject *ParentObject) const
+{
+	ISequencer* Sequencer = GetSequencer();
+	if (Sequencer && Sequencer->GetFocusedMovieSceneSequence())
+	{
+		FGuid ObjectHandle = Sequencer->GetHandleToObject(ParentObject);
+		if (ObjectHandle.IsValid())
+		{
+			UMovieScene* MovieScene = Sequencer->GetFocusedMovieSceneSequence()->GetMovieScene();
+			FProperty* Property = PropertyHandle.GetProperty();
+			TSharedRef<FPropertyPath> PropertyPath = FPropertyPath::CreateEmpty();
+			PropertyPath->AddProperty(FPropertyInfo(Property));
+			FName PropertyName(*PropertyPath->ToString(TEXT(".")));
+			TSubclassOf<UMovieSceneTrack> TrackClass; //use empty @todo find way to get the UMovieSceneTrack from the Property type.
+			return MovieScene->FindTrack(TrackClass, ObjectHandle, PropertyName) != nullptr;
+		}
+	}
+	return false;
+}
+
+void SControlRigDetails::OnKeyPropertyClicked(const IPropertyHandle& KeyedPropertyHandle)
+{
+	ISequencer* Sequencer = GetSequencer();
+	if (Sequencer && !Sequencer->IsAllowedToChange())
+	{
+		return;
+	}
+
+	TArray<UObject*> Objects;
+	KeyedPropertyHandle.GetOuterObjects(Objects);
+	for (UObject *Object : Objects)
+	{
+		UControlRigControlsProxy* Proxy = Cast< UControlRigControlsProxy>(Object);
+		if (Proxy)
+		{
+			Proxy->SetKey(KeyedPropertyHandle);
+		}
+	}
+}
+
+bool SControlRigDetails::ShouldShowPropertyOnDetailCustomization(const FPropertyAndParent& InPropertyAndParent) const
+{
+	auto ShouldPropertyBeVisible = [](const FProperty& InProperty)
+	{
+		bool bShow = InProperty.HasAnyPropertyFlags(CPF_Interp) || InProperty.HasMetaData(FRigVMStruct::InputMetaName) || InProperty.HasMetaData(FRigVMStruct::OutputMetaName);
+
+		// Always show settings properties
+		const UClass* OwnerClass = InProperty.GetOwner<UClass>();
+		bShow |= OwnerClass == UControlRigTransformControlProxy::StaticClass();
+		bShow |= OwnerClass == UControlRigTransformNoScaleControlProxy::StaticClass();
+		bShow |= OwnerClass == UControlRigEulerTransformControlProxy::StaticClass();
+		bShow |= OwnerClass == UControlRigFloatControlProxy::StaticClass();
+		bShow |= OwnerClass == UControlRigVectorControlProxy::StaticClass();
+		bShow |= OwnerClass == UControlRigVector2DControlProxy::StaticClass();
+		bShow |= OwnerClass == UControlRigBoolControlProxy::StaticClass();
+		bShow |= OwnerClass == UControlRigEnumControlProxy::StaticClass();
+		bShow |= OwnerClass == UControlRigIntegerControlProxy::StaticClass();
+
+		return bShow;
+	};
+
+	bool bContainsVisibleProperty = false;
+	if (InPropertyAndParent.Property.IsA<FStructProperty>())
+	{
+		const FStructProperty* StructProperty = CastField<FStructProperty>(&InPropertyAndParent.Property);
+		for (TFieldIterator<FProperty> PropertyIt(StructProperty->Struct); PropertyIt; ++PropertyIt)
+		{
+			if (ShouldPropertyBeVisible(**PropertyIt))
+			{
+				return true;
+			}
+		}
+	}
+
+	return ShouldPropertyBeVisible(InPropertyAndParent.Property) ||
+		(InPropertyAndParent.ParentProperties.Num() > 0 && ShouldPropertyBeVisible(*InPropertyAndParent.ParentProperties[0]));
+}
+
+bool SControlRigDetails::IsReadOnlyPropertyOnDetailCustomization(const FPropertyAndParent& InPropertyAndParent) const
+{
+	auto ShouldPropertyBeEnabled = [](const FProperty& InProperty)
+	{
+		bool bShow = InProperty.HasAnyPropertyFlags(CPF_Interp) || InProperty.HasMetaData(FRigVMStruct::InputMetaName);
+
+		// Always show settings properties
+		const UClass* OwnerClass = InProperty.GetOwner<UClass>();
+		bShow |= OwnerClass == UControlRigTransformControlProxy::StaticClass();
+		bShow |= OwnerClass == UControlRigTransformNoScaleControlProxy::StaticClass();
+		bShow |= OwnerClass == UControlRigEulerTransformControlProxy::StaticClass();
+		bShow |= OwnerClass == UControlRigFloatControlProxy::StaticClass();
+		bShow |= OwnerClass == UControlRigVectorControlProxy::StaticClass();
+		bShow |= OwnerClass == UControlRigVector2DControlProxy::StaticClass();
+		bShow |= OwnerClass == UControlRigBoolControlProxy::StaticClass();
+		bShow |= OwnerClass == UControlRigEnumControlProxy::StaticClass();
+		bShow |= OwnerClass == UControlRigIntegerControlProxy::StaticClass();
+
+
+		return bShow;
+	};
+
+	bool bContainsVisibleProperty = false;
+	if (InPropertyAndParent.Property.IsA<FStructProperty>())
+	{
+		const FStructProperty* StructProperty = CastField<FStructProperty>(&InPropertyAndParent.Property);
+		for (TFieldIterator<FProperty> PropertyIt(StructProperty->Struct); PropertyIt; ++PropertyIt)
+		{
+			if (ShouldPropertyBeEnabled(**PropertyIt))
+			{
+				return false;
+			}
+		}
+	}
+
+	return !(ShouldPropertyBeEnabled(InPropertyAndParent.Property) ||
+		(InPropertyAndParent.ParentProperties.Num() > 0 && ShouldPropertyBeEnabled(*InPropertyAndParent.ParentProperties[0])));
+}
+
+
+#undef LOCTEXT_NAMESPACE
diff --git a/Engine/Plugins/Animation/ControlRig/Source/ControlRigEditor/Private/EditMode/SControlRigDetails.h b/Engine/Plugins/Animation/ControlRig/Source/ControlRigEditor/Private/EditMode/SControlRigDetails.h
new file mode 100644
--- /dev/null
+++ b/Engine/Plugins/Animation/ControlRig/Source/ControlRigEditor/Private/EditMode/SControlRigDetails.h
@@ -0,0 +1,91 @@
+// Copyright Epic Games, Inc. All Rights Reserved.
+/**
+* View for containing details for various controls
+*/
+#pragma once
+
+#include "CoreMinimal.h"
+#include "ControlRigBaseDockableView.h"
+#include "Widgets/SWidget.h"
+#include "Widgets/DeclarativeSyntaxSupport.h"
+#include "Widgets/SCompoundWidget.h"
+#include "Misc/FrameNumber.h"
+#include "IDetailsView.h"
+#include "ControlRig.h"
+#include "Rigs/RigHierarchy.h"
+#include "IDetailKeyframeHandler.h"
+#include "IDetailCustomization.h"
+#include "DetailLayoutBuilder.h"
+#include "DetailCategoryBuilder.h"
+#include "DetailWidgetRow.h"
+
+class ISequencer;
+
+class FControlRigEditModeGenericDetails : public IDetailCustomization
+{
+public:
+	FControlRigEditModeGenericDetails() = delete;
+	FControlRigEditModeGenericDetails(FEditorModeTools* InModeTools) : ModeTools(InModeTools) {}
+
+	/** Makes a new instance of this detail layout class for a specific detail view requesting it */
+	static TSharedRef<IDetailCustomization> MakeInstance(FEditorModeTools* InModeTools)
+	{
+		return MakeShareable(new FControlRigEditModeGenericDetails(InModeTools));
+	}
+
+	/** IDetailCustomization interface */
+	virtual void CustomizeDetails(class IDetailLayoutBuilder& DetailLayout) override;
+
+protected:
+	FEditorModeTools* ModeTools = nullptr;
+};
+
+class SControlRigDetails: public SCompoundWidget, public FControlRigBaseDockableView, public IDetailKeyframeHandler
+{
+
+	SLATE_BEGIN_ARGS(SControlRigDetails)
+	{}
+	SLATE_END_ARGS()
+	~SControlRigDetails();
+
+	void Construct(const FArguments& InArgs, FControlRigEditMode& InEditMode);
+
+	/** Set the objects to be displayed in the details panel */
+	void SetSettingsDetailsObject(const TWeakObjectPtr<>& InObject);
+	void SetEulerTransformDetailsObjects(const TArray<TWeakObjectPtr<>>& InObjects);
+	void SetTransformDetailsObjects(const TArray<TWeakObjectPtr<>>& InObjects);
+	void SetTransformNoScaleDetailsObjects(const TArray<TWeakObjectPtr<>>& InObjects);
+	void SetFloatDetailsObjects(const TArray<TWeakObjectPtr<>>& InObjects);
+	void SetBoolDetailsObjects(const TArray<TWeakObjectPtr<>>& InObjects);
+	void SetIntegerDetailsObjects(const TArray<TWeakObjectPtr<>>& InObjects);
+	void SetEnumDetailsObjects(const TArray<TWeakObjectPtr<>>& InObjects);
+	void SetVectorDetailsObjects(const TArray<TWeakObjectPtr<>>& InObjects);
+	void SetVector2DDetailsObjects(const TArray<TWeakObjectPtr<>>& InObjects);
+
+	// IDetailKeyframeHandler interface
+	virtual bool IsPropertyKeyable(const UClass* InObjectClass, const class IPropertyHandle& PropertyHandle) const override;
+	virtual bool IsPropertyKeyingEnabled() const override;
+	virtual void OnKeyPropertyClicked(const IPropertyHandle& KeyedPropertyHandle) override;
+	virtual bool IsPropertyAnimated(const class IPropertyHandle& PropertyHandle, UObject *ParentObject) const override;
+
+	/** Display or edit set up for property */
+	bool ShouldShowPropertyOnDetailCustomization(const struct FPropertyAndParent& InPropertyAndParent) const;
+	bool IsReadOnlyPropertyOnDetailCustomization(const struct FPropertyAndParent& InPropertyAndParent) const;
+
+private:
+	void UpdateProxies();
+	virtual void HandleControlSelected(UControlRig* Subject, FRigControlElement* InControl, bool bSelected) override;
+	virtual void HandleControlAdded(UControlRig* ControlRig, bool bIsAdded) override;
+
+	TSharedPtr<IDetailsView> ControlEulerTransformDetailsView;
+	TSharedPtr<IDetailsView> ControlTransformDetailsView;
+	TSharedPtr<IDetailsView> ControlTransformNoScaleDetailsView;
+	TSharedPtr<IDetailsView> ControlFloatDetailsView;
+	TSharedPtr<IDetailsView> ControlBoolDetailsView;
+	TSharedPtr<IDetailsView> ControlIntegerDetailsView;
+	TSharedPtr<IDetailsView> ControlEnumDetailsView;
+	TSharedPtr<IDetailsView> ControlVector2DDetailsView;
+	TSharedPtr<IDetailsView> ControlVectorDetailsView;
+
+};
+
diff --git a/Engine/Plugins/Animation/ControlRig/Source/ControlRigEditor/Private/EditMode/SControlRigEditModeTools.cpp b/Engine/Plugins/Animation/ControlRig/Source/ControlRigEditor/Private/EditMode/SControlRigEditModeTools.cpp
--- a/Engine/Plugins/Animation/ControlRig/Source/ControlRigEditor/Private/EditMode/SControlRigEditModeTools.cpp
+++ b/Engine/Plugins/Animation/ControlRig/Source/ControlRigEditor/Private/EditMode/SControlRigEditModeTools.cpp
@@ -17,7 +17,6 @@
 #include "Widgets/Layout/SScrollBox.h"
 #include "MovieSceneSequence.h"
 #include "MovieScene.h"
-#include "SControlPicker.h"
 #include "IDetailCustomization.h"
 #include "DetailLayoutBuilder.h"
 #include "DetailCategoryBuilder.h"
@@ -38,211 +37,10 @@
 #include "Framework/Notifications/NotificationManager.h"
 #include "ScopedTransaction.h"
 #include "ControlRigEditModeToolkit.h"
+#include "SControlRigDetails.h"
+#define LOCTEXT_NAMESPACE "ControlRigEditModeTools"
 
-#define LOCTEXT_NAMESPACE "ControlRigRootCustomization"
 
-class FControlRigEditModeGenericDetails : public IDetailCustomization
-{
-public:
-	FControlRigEditModeGenericDetails() = delete;
-	FControlRigEditModeGenericDetails(FEditorModeTools* InModeTools) : ModeTools(InModeTools) {}
-
-	/** Makes a new instance of this detail layout class for a specific detail view requesting it */
-	static TSharedRef<IDetailCustomization> MakeInstance(FEditorModeTools* InModeTools)
-	{
-		return MakeShareable(new FControlRigEditModeGenericDetails(InModeTools));
-	}
-
-	/** IDetailCustomization interface */
-	virtual void CustomizeDetails(class IDetailLayoutBuilder& DetailLayout) override
-	{
-		TArray<TWeakObjectPtr<UObject>> ObjectsBeingCustomized;
-		DetailLayout.GetObjectsBeingCustomized(ObjectsBeingCustomized);
-
-		TArray<UControlRigControlsProxy*> ProxiesBeingCustomized;
-		for (TWeakObjectPtr<UObject> ObjectBeingCustomized : ObjectsBeingCustomized)
-		{
-			if (UControlRigControlsProxy* Proxy = Cast< UControlRigControlsProxy>(ObjectBeingCustomized.Get()))
-			{
-				ProxiesBeingCustomized.Add(Proxy);
-			}
-		}
-
-		if (ProxiesBeingCustomized.Num() == 0 || ProxiesBeingCustomized[0]->GetControlElement() == nullptr)
-		{
-			return;
-		}
-		FText ControlText = LOCTEXT("Channels", "Channels");
-		
-		if(ProxiesBeingCustomized[0]->GetClass() == UControlRigEditModeSettings::StaticClass())
-		{
-			ControlText = LOCTEXT("Animation", "Animation");
-		}
-		else if (ProxiesBeingCustomized[0]->GetClass() == UControlRigTransformControlProxy::StaticClass())
-		{
-			ControlText = LOCTEXT("TransformChannels", "Transform Channels");
-		}
-		else if (ProxiesBeingCustomized[0]->GetClass() == UControlRigTransformNoScaleControlProxy::StaticClass())
-		{
-			ControlText = LOCTEXT("TransformNoScaleChannels", "TransformNoScale Channels");
-		}
-		else if (ProxiesBeingCustomized[0]->GetClass() == UControlRigEulerTransformControlProxy::StaticClass())
-		{
-			ControlText = LOCTEXT("EulerTransformChannels", "Euler Transform Channels");
-		}
-		else if (ProxiesBeingCustomized[0]->GetClass() == UControlRigFloatControlProxy::StaticClass())
-		{
-			ControlText = LOCTEXT("FloatChannels", "Float Channels");
-		}
-		else if (ProxiesBeingCustomized[0]->GetClass() == UControlRigVectorControlProxy::StaticClass())
-		{
-			ControlText = LOCTEXT("VectorChannels", "Vector Channels");
-		}
-		else if (ProxiesBeingCustomized[0]->GetClass() == UControlRigVector2DControlProxy::StaticClass())
-		{
-			ControlText = LOCTEXT("Vector2DChannels", "Vector2D Channels");
-		}
-		else if (ProxiesBeingCustomized[0]->GetClass() == UControlRigBoolControlProxy::StaticClass())
-		{
-			ControlText = LOCTEXT("BoolChannels", "Bool Channels");
-		}
-		else if (ProxiesBeingCustomized[0]->GetClass() == UControlRigEnumControlProxy::StaticClass())
-		{
-			ControlText = LOCTEXT("EnumChannels", "Enum Channels");
-		}
-		else if (ProxiesBeingCustomized[0]->GetClass() == UControlRigIntegerControlProxy::StaticClass())
-		{
-			ControlText = LOCTEXT("IntegerChannels", "Integer Channels");
-		}
-
-		IDetailCategoryBuilder& Category = DetailLayout.EditCategory(TEXT("Control"), ControlText);
-		for (UControlRigControlsProxy* Proxy : ProxiesBeingCustomized)
-		{
-			FRigControlElement* ControlElement = Proxy->GetControlElement();
-			if(ControlElement == nullptr)
-			{
-				continue;
-			}
-			
-			FName ValuePropertyName = TEXT("Transform");
-			if (ControlElement->Settings.ControlType == ERigControlType::Float)
-			{
-				ValuePropertyName = TEXT("Float");
-			}
-			else if (ControlElement->Settings.ControlType == ERigControlType::Integer)
-			{
-				if (ControlElement->Settings.ControlEnum == nullptr)
-				{
-					ValuePropertyName = TEXT("Integer");
-				}
-				else
-				{
-					ValuePropertyName = TEXT("Enum");
-				}
-			}
-			else if (ControlElement->Settings.ControlType == ERigControlType::Bool)
-			{
-				ValuePropertyName = TEXT("Bool");
-			}
-			else if (ControlElement->Settings.ControlType == ERigControlType::Position ||
-				ControlElement->Settings.ControlType == ERigControlType::Scale)
-			{
-				ValuePropertyName = TEXT("Vector");
-			}
-			else if (ControlElement->Settings.ControlType == ERigControlType::Vector2D)
-			{
-				ValuePropertyName = TEXT("Vector2D");
-			}
-			else if (ControlElement->Settings.ControlType == ERigControlType::EulerTransform)
-			{
-				ValuePropertyName = TEXT("EulerTransform");
-			}
-			else if (ControlElement->Settings.ControlType == ERigControlType::TransformNoScale)
-			{
-				ValuePropertyName = TEXT("TransformNoScale");
-			}
-
-			TSharedPtr<IPropertyHandle> ValuePropertyHandle = DetailLayout.GetProperty(ValuePropertyName, Proxy->GetClass());
-			if (ValuePropertyHandle)
-			{
-				ValuePropertyHandle->SetPropertyDisplayName(FText::FromName(Proxy->GetName()));
-			}
-
-			URigHierarchy* Hierarchy = Proxy->ControlRig->GetHierarchy();
-			Hierarchy->ForEach<FRigControlElement>([Hierarchy, Proxy, &Category, this](FRigControlElement* ControlElement) -> bool
-	      	{
-				FName ParentControlName = NAME_None;
-				FRigControlElement* ParentControlElement = Cast<FRigControlElement>(Hierarchy->GetFirstParent(ControlElement));
-				if(ParentControlElement)
-				{
-					ParentControlName = ParentControlElement->GetName();
-				}
-				
-				if (ParentControlName == ControlElement->GetName())
-				{
-					if (FControlRigEditMode* EditMode = static_cast<FControlRigEditMode*>(ModeTools->GetActiveMode(FControlRigEditMode::ModeName)))
-					{
-						if (UObject* NestedProxy = EditMode->ControlProxy->FindProxy(ControlElement->GetName()))
-						{
-							FName PropertyName(NAME_None);
-							switch (ControlElement->Settings.ControlType)
-							{
-								case ERigControlType::Bool:
-								{
-									PropertyName = TEXT("Bool");
-									break;
-								}
-								case ERigControlType::Float:
-								{
-									PropertyName = TEXT("Float");
-									break;
-								}
-								case ERigControlType::Integer:
-								{
-									if (ControlElement->Settings.ControlEnum == nullptr)
-									{
-										PropertyName = TEXT("Integer");
-									}
-									else
-									{
-										PropertyName = TEXT("Enum");
-									}
-									break;
-								}
-								default:
-								{
-									break;
-								}
-							}
-
-							if (PropertyName.IsNone())
-							{
-								return true;
-							}
-
-							TArray<UObject*> NestedProxies;
-							NestedProxies.Add(NestedProxy);
-
-							FAddPropertyParams Params;
-							Params.CreateCategoryNodes(false);
-							IDetailPropertyRow* NestedRow = Category.AddExternalObjectProperty(
-								NestedProxies,
-								PropertyName,
-								EPropertyLocation::Advanced,
-								Params);
-							NestedRow->DisplayName(FText::FromName(ControlElement->Settings.DisplayName));
-
-							Category.SetShowAdvanced(true);
-						}
-					}
-				}
-				return true;
-			});
-		}
-	}
-protected:
-	FEditorModeTools* ModeTools = nullptr;
-};
 
 void SControlRigEditModeTools::SetControlRig(UControlRig* ControlRig)
 {
@@ -264,9 +62,9 @@ void SControlRigEditModeTools::SetControlRig(UControlRig* ControlRig)
 	TArray<TWeakObjectPtr<>> Objects;
 	Objects.Add(SequencerRig);
 	RigOptionsDetailsView->SetObjects(Objects);
-
+#if USE_LOCAL_DETAILS
 	HierarchyTreeView->RefreshTreeView(true);
-
+#endif
 	if (ViewportRig.IsValid())
 	{
 		ViewportRig->ControlSelected().AddRaw(this, &SControlRigEditModeTools::OnRigElementSelected);
@@ -314,7 +112,7 @@ void SControlRigEditModeTools::Construct(const FArguments& InArgs, TSharedPtr<FC
 	SettingsDetailsView->SetIsPropertyVisibleDelegate(FIsPropertyVisible::CreateSP(this, &SControlRigEditModeTools::ShouldShowPropertyOnDetailCustomization));
 	SettingsDetailsView->SetIsPropertyReadOnlyDelegate(FIsPropertyReadOnly::CreateSP(this, &SControlRigEditModeTools::IsReadOnlyPropertyOnDetailCustomization));
 	SettingsDetailsView->SetGenericLayoutDetailsDelegate(FOnGetDetailCustomizationInstance::CreateStatic(&FControlRigEditModeGenericDetails::MakeInstance, ModeTools));
-
+#if USE_LOCAL_DETAILS
 	ControlEulerTransformDetailsView = FModuleManager::GetModuleChecked<FPropertyEditorModule>("PropertyEditor").CreateDetailView(DetailsViewArgs);
 	ControlEulerTransformDetailsView->SetKeyframeHandler(SharedThis(this));
 	ControlEulerTransformDetailsView->SetIsPropertyVisibleDelegate(FIsPropertyVisible::CreateSP(this, &SControlRigEditModeTools::ShouldShowPropertyOnDetailCustomization));
@@ -368,7 +166,7 @@ void SControlRigEditModeTools::Construct(const FArguments& InArgs, TSharedPtr<FC
 	ControlVector2DDetailsView->SetIsPropertyVisibleDelegate(FIsPropertyVisible::CreateSP(this, &SControlRigEditModeTools::ShouldShowPropertyOnDetailCustomization));
 	ControlVector2DDetailsView->SetIsPropertyReadOnlyDelegate(FIsPropertyReadOnly::CreateSP(this, &SControlRigEditModeTools::IsReadOnlyPropertyOnDetailCustomization));
 	ControlVector2DDetailsView->SetGenericLayoutDetailsDelegate(FOnGetDetailCustomizationInstance::CreateStatic(&FControlRigEditModeGenericDetails::MakeInstance, ModeTools));
-
+#endif
 
 	RigOptionsDetailsView = FModuleManager::GetModuleChecked<FPropertyEditorModule>("PropertyEditor").CreateDetailView(DetailsViewArgs);
 	RigOptionsDetailsView->SetKeyframeHandler(SharedThis(this));
@@ -381,19 +179,19 @@ void SControlRigEditModeTools::Construct(const FArguments& InArgs, TSharedPtr<FC
 	DisplaySettings.bShowRigidBodies = false;
 	DisplaySettings.bHideParentsOnFilter = true;
 	DisplaySettings.bFlattenHierarchyOnFilter = true;
-
+#if USE_LOCAL_DETAILS
 	FRigTreeDelegates RigTreeDelegates;
 	RigTreeDelegates.OnGetHierarchy = FOnGetRigTreeHierarchy::CreateSP(this, &SControlRigEditModeTools::GetHierarchy);
 	RigTreeDelegates.OnGetDisplaySettings = FOnGetRigTreeDisplaySettings::CreateSP(this, &SControlRigEditModeTools::GetDisplaySettings);
 	RigTreeDelegates.OnSelectionChanged = FOnRigTreeSelectionChanged::CreateSP(this, &SControlRigEditModeTools::HandleSelectionChanged);
-
+#endif
 	ChildSlot
 	[
 		SNew(SScrollBox)
 		+ SScrollBox::Slot()
 		[
 			SNew(SVerticalBox)
-
+#if USE_LOCAL_DETAILS
 			+SVerticalBox::Slot()
 			.AutoHeight()
 			[
@@ -408,12 +206,13 @@ void SControlRigEditModeTools::Construct(const FArguments& InArgs, TSharedPtr<FC
 					.RigTreeDelegates(RigTreeDelegates)
 				]
 			]
-
+#endif
 			+SVerticalBox::Slot()
 			.AutoHeight()
 			[
 				SettingsDetailsView.ToSharedRef()
 			]
+#if USE_LOCAL_DETAILS
 			+ SVerticalBox::Slot()
 			.AutoHeight()
 			[
@@ -459,6 +258,7 @@ void SControlRigEditModeTools::Construct(const FArguments& InArgs, TSharedPtr<FC
 			[
 				ControlFloatDetailsView.ToSharedRef()
 			]
+#endif
 
 			+SVerticalBox::Slot()
 			.AutoHeight()
@@ -539,8 +339,9 @@ void SControlRigEditModeTools::Construct(const FArguments& InArgs, TSharedPtr<FC
 			]
 		]
 	];
-
+#if USE_LOCAL_DETAILS
 	HierarchyTreeView->RefreshTreeView(true);
+#endif
 }
 
 void SControlRigEditModeTools::SetSettingsDetailsObject(const TWeakObjectPtr<>& InObject)
@@ -553,6 +354,8 @@ void SControlRigEditModeTools::SetSettingsDetailsObject(const TWeakObjectPtr<>&
 
 	}
 }
+#if USE_LOCAL_DETAILS
+
 void SControlRigEditModeTools::SetEulerTransformDetailsObjects(const TArray<TWeakObjectPtr<>>& InObjects)
 {
 	if (ControlEulerTransformDetailsView)
@@ -623,8 +426,8 @@ void SControlRigEditModeTools::SetVector2DDetailsObjects(const TArray<TWeakObjec
 		ControlVector2DDetailsView->SetObjects(InObjects);
 	}
 }
-
-void SControlRigEditModeTools::SControlRigEditModeTools::SetSequencer(TWeakPtr<ISequencer> InSequencer)
+#endif
+void SControlRigEditModeTools::SetSequencer(TWeakPtr<ISequencer> InSequencer)
 {
 	WeakSequencer = InSequencer.Pin();
 }
@@ -780,6 +583,7 @@ bool SControlRigEditModeTools::IsReadOnlyPropertyOnDetailCustomization(const FPr
 		(InPropertyAndParent.ParentProperties.Num() > 0 && ShouldPropertyBeEnabled(*InPropertyAndParent.ParentProperties[0])));
 }
 
+#if USE_LOCAL_DETAILS
 static bool bPickerChangingSelection = false;
 
 void SControlRigEditModeTools::OnManipulatorsPicked(const TArray<FName>& Manipulators)
@@ -846,9 +650,11 @@ void SControlRigEditModeTools::HandleSelectionChanged(TSharedPtr<FRigTreeElement
 		}
 	}
 }
+#endif
 
 void SControlRigEditModeTools::OnRigElementSelected(UControlRig* Subject, FRigControlElement* ControlElement, bool bSelected)
 {
+#if USE_LOCAL_DETAILS
 	const FRigElementKey Key = ControlElement->GetKey();
 	for (int32 RootIndex = 0; RootIndex < HierarchyTreeView->GetRootElements().Num(); ++RootIndex)
 	{
@@ -869,7 +675,7 @@ void SControlRigEditModeTools::OnRigElementSelected(UControlRig* Subject, FRigCo
 			}
 		}
 	}
-
+#endif
 	if (UControlRig* ControlRig = SequencerRig.Get())
 	{
 		// get the selected controls
@@ -878,6 +684,7 @@ void SControlRigEditModeTools::OnRigElementSelected(UControlRig* Subject, FRigCo
 	}
 }
 
+
 const FRigControlElementCustomization* SControlRigEditModeTools::HandleGetControlElementCustomization(URigHierarchy* InHierarchy, const FRigElementKey& InControlKey)
 {
 	if (UControlRig* ControlRig = SequencerRig.Get())
@@ -1066,7 +873,6 @@ FReply SControlRigEditModeTools::OnBakeControlsToNewSpaceButtonClicked()
 		return BakeWidget->OpenDialog(true);
 	}
 	return FReply::Unhandled();
-
 }
 
 EVisibility SControlRigEditModeTools::GetRigOptionExpanderVisibility() const
diff --git a/Engine/Plugins/Animation/ControlRig/Source/ControlRigEditor/Private/EditMode/SControlRigEditModeTools.h b/Engine/Plugins/Animation/ControlRig/Source/ControlRigEditor/Private/EditMode/SControlRigEditModeTools.h
--- a/Engine/Plugins/Animation/ControlRig/Source/ControlRigEditor/Private/EditMode/SControlRigEditModeTools.h
+++ b/Engine/Plugins/Animation/ControlRig/Source/ControlRigEditor/Private/EditMode/SControlRigEditModeTools.h
@@ -22,6 +22,8 @@ class FToolBarBuilder;
 class FEditorModeTools;
 class FControlRigEditModeToolkit;
 
+#define USE_LOCAL_DETAILS 0
+
 class SControlRigEditModeTools : public SCompoundWidget, public IDetailKeyframeHandler
 {
 public:
@@ -29,9 +31,9 @@ public:
 	SLATE_END_ARGS();
 
 	void Construct(const FArguments& InArgs, TSharedPtr<FControlRigEditModeToolkit> InOwningToolkit, FControlRigEditMode& InEditMode, UWorld* InWorld);
-
 	/** Set the objects to be displayed in the details panel */
 	void SetSettingsDetailsObject(const TWeakObjectPtr<>& InObject);
+#if USE_LOCAL_DETAILS
 	void SetEulerTransformDetailsObjects(const TArray<TWeakObjectPtr<>>& InObjects);
 	void SetTransformDetailsObjects(const TArray<TWeakObjectPtr<>>& InObjects);
 	void SetTransformNoScaleDetailsObjects(const TArray<TWeakObjectPtr<>>& InObjects);
@@ -41,7 +43,7 @@ public:
 	void SetEnumDetailsObjects(const TArray<TWeakObjectPtr<>>& InObjects);
 	void SetVectorDetailsObjects(const TArray<TWeakObjectPtr<>>& InObjects);
 	void SetVector2DDetailsObjects(const TArray<TWeakObjectPtr<>>& InObjects);
-
+#endif
 
 	/** Set the sequencer we are bound to */
 	void SetSequencer(TWeakPtr<ISequencer> InSequencer);
@@ -51,7 +53,6 @@ public:
 
 	/** Returns the hierarchy currently being used */
 	const URigHierarchy* GetHierarchy() const;
-
 	// IDetailKeyframeHandler interface
 	virtual bool IsPropertyKeyable(const UClass* InObjectClass, const class IPropertyHandle& PropertyHandle) const override;
 	virtual bool IsPropertyKeyingEnabled() const override;
@@ -59,10 +60,12 @@ public:
 	virtual bool IsPropertyAnimated(const class IPropertyHandle& PropertyHandle, UObject *ParentObject) const override;
 private:
 	/** Sequencer we are currently bound to */
+
 	TWeakPtr<ISequencer> WeakSequencer;
+	TSharedPtr<IDetailsView> SettingsDetailsView;
 
+#if USE_LOCAL_DETAILS
 	/** The details views we do most of our work within */
-	TSharedPtr<IDetailsView> SettingsDetailsView;
 	TSharedPtr<IDetailsView> ControlEulerTransformDetailsView;
 	TSharedPtr<IDetailsView> ControlTransformDetailsView;
 	TSharedPtr<IDetailsView> ControlTransformNoScaleDetailsView;
@@ -72,20 +75,18 @@ private:
 	TSharedPtr<IDetailsView> ControlEnumDetailsView;
 	TSharedPtr<IDetailsView> ControlVector2DDetailsView;
 	TSharedPtr<IDetailsView> ControlVectorDetailsView;
-
+#endif
 	/** Expander to interact with the options of the rig  */
 	TSharedPtr<SExpandableArea> RigOptionExpander;
 	TSharedPtr<IDetailsView> RigOptionsDetailsView;
-
+#if USE_LOCAL_DETAILS
 	/** Hierarchy picker for controls*/
 	TSharedPtr<SRigHierarchyTreeView> HierarchyTreeView;
+#endif
 
-	/** Hierarchy picker for controls*/
-	TSharedPtr<SRigSpacePickerWidget> SpacePickerWidget;
-
-	/** Special picker for controls, no longer used */
-	TSharedPtr<SControlPicker> ControlPicker;
+	/** Space Picker controls*/
 	TSharedPtr<SExpandableArea> PickerExpander;
+	TSharedPtr<SRigSpacePickerWidget> SpacePickerWidget;
 
 	/** Storage for both sequencer and viewport rigs */
 	TWeakObjectPtr<UControlRig> SequencerRig;
@@ -122,7 +123,7 @@ private:
 	FEditorModeTools* ModeTools = nullptr;
 	FRigTreeDisplaySettings DisplaySettings;
 	const FRigTreeDisplaySettings& GetDisplaySettings() const { return DisplaySettings; }
-	bool bIsChangingRigHierarchy;
+	bool bIsChangingRigHierarchy = false;
 
 	// The toolkit that created this UI
 	TWeakPtr<FControlRigEditModeToolkit> OwningToolkit;
diff --git a/Engine/Plugins/Animation/ControlRig/Source/ControlRigEditor/Private/EditMode/SControlRigOutliner.cpp b/Engine/Plugins/Animation/ControlRig/Source/ControlRigEditor/Private/EditMode/SControlRigOutliner.cpp
new file mode 100644
--- /dev/null
+++ b/Engine/Plugins/Animation/ControlRig/Source/ControlRigEditor/Private/EditMode/SControlRigOutliner.cpp
@@ -0,0 +1,157 @@
+// Copyright Epic Games, Inc. All Rights Reserved.
+
+#include "SControlRigOutliner.h"
+#include "Framework/Notifications/NotificationManager.h"
+#include "Widgets/Notifications/SNotificationList.h"
+#include "Framework/MultiBox/MultiBoxBuilder.h"
+#include "Widgets/Input/SButton.h"
+#include "Widgets/Layout/SExpandableArea.h"
+#include "Widgets/Layout/SScrollBox.h"
+#include "AssetData.h"
+#include "EditorStyleSet.h"
+#include "Widgets/Text/SInlineEditableTextBlock.h"
+#include "EditorStyleSet.h"
+#include "Styling/CoreStyle.h"
+#include "ScopedTransaction.h"
+#include "ControlRig.h"
+#include "UnrealEdGlobals.h"
+#include "ControlRigEditMode.h"
+#include "EditorModeManager.h"
+#include "ISequencer.h"
+#include "LevelSequence.h"
+#include "UnrealEd/Public/Selection.h"
+#include "Editor.h"
+#include "LevelEditor.h"
+#include "Widgets/Views/SListView.h"
+#include "Widgets/Text/STextBlock.h"
+#include "Widgets/Input/SButton.h"
+#include "Widgets/Layout/SScrollBox.h"
+#include "ControlRigControlsProxy.h"
+#include "ControlRigEditModeSettings.h"
+#include "TimerManager.h"
+#include "SRigHierarchyTreeView.h"
+#include "Rigs/RigHierarchyController.h"
+
+#define LOCTEXT_NAMESPACE "ControlRigOutliner"
+
+void SControlRigOutliner::Construct(const FArguments& InArgs, FControlRigEditMode& InEditMode)
+{
+	bIsChangingRigHierarchy = false;
+
+	DisplaySettings.bShowBones = false;
+	DisplaySettings.bShowControls = true;
+	DisplaySettings.bShowNulls = false;
+	DisplaySettings.bShowReferences = false;
+	DisplaySettings.bShowRigidBodies = false;
+	DisplaySettings.bHideParentsOnFilter = true;
+	DisplaySettings.bFlattenHierarchyOnFilter = true;
+
+	FRigTreeDelegates RigTreeDelegates;
+	RigTreeDelegates.OnGetHierarchy = FOnGetRigTreeHierarchy::CreateSP(this, &SControlRigOutliner::GetHierarchy);
+	RigTreeDelegates.OnGetDisplaySettings = FOnGetRigTreeDisplaySettings::CreateSP(this, &SControlRigOutliner::GetDisplaySettings);
+	RigTreeDelegates.OnSelectionChanged = FOnRigTreeSelectionChanged::CreateSP(this, &SControlRigOutliner::HandleSelectionChanged);
+
+	ChildSlot
+		[
+			SNew(SScrollBox)
+			+ SScrollBox::Slot()
+			[
+			SNew(SVerticalBox)
+			+ SVerticalBox::Slot()
+				.AutoHeight()
+				[
+					SAssignNew(PickerExpander, SExpandableArea)
+					.InitiallyCollapsed(false)
+					.AreaTitle(LOCTEXT("Picker_Header", "Controls"))
+					.AreaTitleFont(FEditorStyle::GetFontStyle("DetailsView.CategoryFontStyle"))
+					.BorderBackgroundColor(FLinearColor(.6f, .6f, .6f))
+					.BodyContent()
+					[
+						SAssignNew(HierarchyTreeView, SSearchableRigHierarchyTreeView)
+						.RigTreeDelegates(RigTreeDelegates)
+					]
+				]
+		
+			]
+		];
+
+	SetEditMode(InEditMode);
+	HierarchyTreeView->GetTreeView()->RefreshTreeView(true);
+}
+
+SControlRigOutliner::~SControlRigOutliner()
+{
+	//base class handles control rig related cleanup
+}
+
+
+void SControlRigOutliner::HandleControlAdded(UControlRig* ControlRig, bool bIsAdded)
+{
+	FControlRigBaseDockableView::HandleControlAdded(ControlRig, bIsAdded);
+	HierarchyTreeView->GetTreeView()->RefreshTreeView(true);
+}
+
+void SControlRigOutliner::NewControlRigSet(UControlRig* ControlRig)
+{
+	FControlRigBaseDockableView::NewControlRigSet(ControlRig);
+	HierarchyTreeView->GetTreeView()->RefreshTreeView(true);
+}
+void SControlRigOutliner::HandleControlSelected(UControlRig* Subject, FRigControlElement* ControlElement, bool bSelected)
+{
+	FControlRigBaseDockableView::HandleControlSelected(Subject, ControlElement, bSelected);
+	const FRigElementKey Key = ControlElement->GetKey();
+	for (int32 RootIndex = 0; RootIndex < HierarchyTreeView->GetTreeView()->GetRootElements().Num(); ++RootIndex)
+	{
+		TSharedPtr<FRigTreeElement> Found = HierarchyTreeView->GetTreeView()->FindElement(Key, HierarchyTreeView->GetTreeView()->GetRootElements()[RootIndex]);
+		if (Found.IsValid())
+		{
+			HierarchyTreeView->GetTreeView()->SetItemSelection(Found, bSelected, ESelectInfo::OnNavigation);
+
+			TArray<TSharedPtr<FRigTreeElement>> SelectedItems = HierarchyTreeView->GetTreeView()->GetSelectedItems();
+			for (TSharedPtr<FRigTreeElement> SelectedItem : SelectedItems)
+			{
+				HierarchyTreeView->GetTreeView()->SetExpansionRecursive(SelectedItem, false, true);
+			}
+
+			if (SelectedItems.Num() > 0)
+			{
+				HierarchyTreeView->GetTreeView()->RequestScrollIntoView(SelectedItems.Last());
+			}
+		}
+	}
+}
+
+const URigHierarchy* SControlRigOutliner::GetHierarchy() const
+{
+	if (CurrentControlRig.IsValid())
+	{
+		return CurrentControlRig.Get()->GetHierarchy();
+	}
+	return nullptr;
+}
+
+void SControlRigOutliner::HandleSelectionChanged(TSharedPtr<FRigTreeElement> Selection, ESelectInfo::Type SelectInfo)
+{
+	if (bIsChangingRigHierarchy)
+	{
+		return;
+	}
+
+	const URigHierarchy* Hierarchy = GetHierarchy();
+	if (Hierarchy)
+	{
+		URigHierarchyController* Controller = ((URigHierarchy*)Hierarchy)->GetController(true);
+		check(Controller);
+
+		TGuardValue<bool> GuardRigHierarchyChanges(bIsChangingRigHierarchy, true);
+
+		const TArray<FRigElementKey> NewSelection = HierarchyTreeView->GetTreeView()->GetSelectedKeys();
+		if (!Controller->SetSelection(NewSelection))
+		{
+			return;
+		}
+	}
+}
+
+
+#undef LOCTEXT_NAMESPACE
diff --git a/Engine/Plugins/Animation/ControlRig/Source/ControlRigEditor/Private/EditMode/SControlRigOutliner.h b/Engine/Plugins/Animation/ControlRig/Source/ControlRigEditor/Private/EditMode/SControlRigOutliner.h
new file mode 100644
--- /dev/null
+++ b/Engine/Plugins/Animation/ControlRig/Source/ControlRigEditor/Private/EditMode/SControlRigOutliner.h
@@ -0,0 +1,45 @@
+// Copyright Epic Games, Inc. All Rights Reserved.
+/**
+* View for holding ControlRig Animation Outliner
+*/
+#pragma once
+
+#include "CoreMinimal.h"
+#include "ControlRigBaseDockableView.h"
+#include "Widgets/SWidget.h"
+#include "Widgets/SCompoundWidget.h"
+#include "ControlRig.h"
+#include "Rigs/RigHierarchy.h"
+#include "SRigHierarchyTreeView.h"
+
+
+class ISequencer;
+class SExpandableArea;
+class SSearchableRigHierarchyTreeView;
+
+class SControlRigOutliner: public SCompoundWidget, public FControlRigBaseDockableView
+{
+
+	SLATE_BEGIN_ARGS(SControlRigOutliner)
+	{}
+	SLATE_END_ARGS()
+	~SControlRigOutliner();
+
+	void Construct(const FArguments& InArgs, FControlRigEditMode& InEditMode);
+
+private:
+	virtual void HandleControlSelected(UControlRig* Subject, FRigControlElement* InControl, bool bSelected) override;
+	virtual void HandleControlAdded(UControlRig* ControlRig, bool bIsAdded) override;
+	virtual void NewControlRigSet(UControlRig* ControlRig) override;
+
+	const URigHierarchy* GetHierarchy() const;
+	void HandleSelectionChanged(TSharedPtr<FRigTreeElement> Selection, ESelectInfo::Type SelectInfo);
+
+	/** Hierarchy picker for controls*/
+	TSharedPtr<SSearchableRigHierarchyTreeView> HierarchyTreeView;
+	FRigTreeDisplaySettings DisplaySettings;
+	const FRigTreeDisplaySettings& GetDisplaySettings() const { return DisplaySettings; }
+	bool bIsChangingRigHierarchy = false;
+	TSharedPtr<SExpandableArea> PickerExpander;
+};
+
diff --git a/Engine/Plugins/Animation/ControlRig/Source/ControlRigEditor/Private/EditMode/SControlRigSpacePicker.cpp b/Engine/Plugins/Animation/ControlRig/Source/ControlRigEditor/Private/EditMode/SControlRigSpacePicker.cpp
new file mode 100644
--- /dev/null
+++ b/Engine/Plugins/Animation/ControlRig/Source/ControlRigEditor/Private/EditMode/SControlRigSpacePicker.cpp
@@ -0,0 +1,315 @@
+// Copyright Epic Games, Inc. All Rights Reserved.
+
+#include "SControlRigSpacePicker.h"
+#include "Framework/Notifications/NotificationManager.h"
+#include "Widgets/Notifications/SNotificationList.h"
+#include "Framework/MultiBox/MultiBoxBuilder.h"
+#include "Widgets/Input/SButton.h"
+#include "Widgets/Layout/SExpandableArea.h"
+#include "Widgets/Layout/SScrollBox.h"
+#include "AssetData.h"
+#include "EditorStyleSet.h"
+#include "Widgets/Text/SInlineEditableTextBlock.h"
+#include "EditorStyleSet.h"
+#include "ISequencer.h"
+#include "ScopedTransaction.h"
+#include "ControlRig.h"
+#include "ControlRigEditMode.h"
+#include "MovieSceneSequence.h"
+#include "MovieScene.h"
+#include "ControlRigSpaceChannelEditors.h"
+
+#define LOCTEXT_NAMESPACE "ControlRigSpacePicker"
+
+void SControlRigSpacePicker::Construct(const FArguments& InArgs, FControlRigEditMode& InEditMode)
+{
+	
+	ChildSlot
+		[
+			SNew(SScrollBox)
+			+ SScrollBox::Slot()
+			[
+				SNew(SVerticalBox)
+				+SVerticalBox::Slot()
+				.AutoHeight()
+				[
+					SAssignNew(PickerExpander, SExpandableArea)
+					.InitiallyCollapsed(true)
+					.AreaTitle(LOCTEXT("Picker_SpaceWidget", "Spaces"))
+					.AreaTitleFont(FEditorStyle::GetFontStyle("DetailsView.CategoryFontStyle"))
+					.BorderBackgroundColor(FLinearColor(.6f, .6f, .6f))
+					.Padding(FMargin(8.f))
+					.HeaderContent()
+					[
+						SNew(SHorizontalBox)
+
+						+SHorizontalBox::Slot()
+						.AutoWidth()
+						.HAlign(HAlign_Left)
+						.VAlign(VAlign_Center)
+						.Padding(0.f, 0.f, 0.f, 0.f)
+						[
+							SNew(STextBlock)
+							.Text(LOCTEXT("Picker_SpaceWidget", "Spaces"))
+							.Font(FCoreStyle::Get().GetFontStyle("ExpandableArea.TitleFont"))
+						]
+					
+						+SHorizontalBox::Slot()
+						.FillWidth(1.f)
+						[
+							SNew(SSpacer)
+						]
+
+						+SHorizontalBox::Slot()
+						.AutoWidth()
+						.HAlign(HAlign_Right)
+						.VAlign(VAlign_Center)
+						.Padding(0.f, 2.f, 8.f, 2.f)
+						[
+							SNew(SButton)
+							.ContentPadding(0.0f)
+							.ButtonStyle(FEditorStyle::Get(), "NoBorder")
+							.OnClicked(this, &SControlRigSpacePicker::HandleAddSpaceClicked)
+							.Cursor(EMouseCursor::Default)
+							.ToolTipText(LOCTEXT("AddSpace", "Add Space"))
+							[
+								SNew(SImage)
+								.Image(FEditorStyle::GetBrush(TEXT("Icons.PlusCircle")))
+							]
+						]
+					]
+					.BodyContent()
+					[
+						SAssignNew(SpacePickerWidget, SRigSpacePickerWidget)
+						.AllowDelete(true)
+						.AllowReorder(true)
+						.AllowAdd(false)
+						.ShowBakeButton(true)
+						.GetControlCustomization(this, &SControlRigSpacePicker::HandleGetControlElementCustomization)
+						.OnActiveSpaceChanged(this, &SControlRigSpacePicker::HandleActiveSpaceChanged)
+						.OnSpaceListChanged(this, &SControlRigSpacePicker::HandleSpaceListChanged)
+						.OnBakeButtonClicked(this, &SControlRigSpacePicker::OnBakeControlsToNewSpaceButtonClicked)
+						// todo: implement GetAdditionalSpacesDelegate to pull spaces from sequencer
+					]
+				]
+
+			]
+		];
+
+	SetEditMode(InEditMode);
+}
+
+SControlRigSpacePicker::~SControlRigSpacePicker()
+{
+	//base class handles control rig related cleanup
+}
+
+
+void SControlRigSpacePicker::HandleControlAdded(UControlRig* ControlRig, bool bIsAdded)
+{
+	FControlRigBaseDockableView::HandleControlAdded(ControlRig, bIsAdded);
+}
+
+void SControlRigSpacePicker::NewControlRigSet(UControlRig* ControlRig)
+{
+	FControlRigBaseDockableView::NewControlRigSet(ControlRig);
+}
+
+void SControlRigSpacePicker::HandleControlSelected(UControlRig* Subject, FRigControlElement* ControlElement, bool bSelected)
+{
+	if (UControlRig* ControlRig = GetControlRig())
+	{
+		// get the selected controls
+		TArray<FRigElementKey> SelectedControls = ControlRig->GetHierarchy()->GetSelectedKeys(ERigElementType::Control);
+		SpacePickerWidget->SetControls(ControlRig->GetHierarchy(), SelectedControls);
+	}
+}
+
+const FRigControlElementCustomization* SControlRigSpacePicker::HandleGetControlElementCustomization(URigHierarchy* InHierarchy, const FRigElementKey& InControlKey)
+{
+	if (UControlRig* ControlRig = GetControlRig())
+	{
+		return ControlRig->GetControlCustomization(InControlKey);
+	}
+	return nullptr;
+}
+
+void SControlRigSpacePicker::HandleActiveSpaceChanged(URigHierarchy* InHierarchy, const FRigElementKey& InControlKey,
+	const FRigElementKey& InSpaceKey)
+{
+
+	if (ISequencer* Sequencer = GetSequencer())
+	{
+		if (UControlRig* ControlRig = GetControlRig())
+		{
+			FString FailureReason;
+			URigHierarchy::TElementDependencyMap DependencyMap = InHierarchy->GetDependenciesForVM(ControlRig->GetVM());
+			if (!InHierarchy->CanSwitchToParent(InControlKey, InSpaceKey, DependencyMap, &FailureReason))
+			{
+				// notification
+				FNotificationInfo Info(FText::FromString(FailureReason));
+				Info.bFireAndForget = true;
+				Info.FadeOutDuration = 2.0f;
+				Info.ExpireDuration = 8.0f;
+
+				const TSharedPtr<SNotificationItem> NotificationPtr = FSlateNotificationManager::Get().AddNotification(Info);
+				NotificationPtr->SetCompletionState(SNotificationItem::CS_Fail);
+				return;
+			}
+
+			if (const FRigControlElement* ControlElement = InHierarchy->Find<FRigControlElement>(InControlKey))
+			{
+				FScopedTransaction Transaction(LOCTEXT("KeyControlRigSpace", "Key Control Rig Space"));
+				ControlRig->Modify();
+
+				FSpaceChannelAndSection SpaceChannelAndSection = FControlRigSpaceChannelHelpers::FindSpaceChannelAndSectionForControl(ControlRig, InControlKey.Name, Sequencer, true /*bCreateIfNeeded*/);
+				if (SpaceChannelAndSection.SpaceChannel)
+				{
+					const FFrameRate TickResolution = Sequencer->GetFocusedTickResolution();
+					const FFrameTime FrameTime = Sequencer->GetLocalTime().ConvertTo(TickResolution);
+					FFrameNumber CurrentTime = FrameTime.GetFrame();
+					FControlRigSpaceChannelHelpers::SequencerKeyControlRigSpaceChannel(ControlRig, Sequencer, SpaceChannelAndSection.SpaceChannel, SpaceChannelAndSection.SectionToKey, CurrentTime, InHierarchy, InControlKey, InSpaceKey);
+				}
+			}
+		}
+	}
+}
+
+void SControlRigSpacePicker::HandleSpaceListChanged(URigHierarchy* InHierarchy, const FRigElementKey& InControlKey,
+	const TArray<FRigElementKey>& InSpaceList)
+{
+	if (UControlRig* ControlRig = GetControlRig())
+	{
+		if (const FRigControlElement* ControlElement = InHierarchy->Find<FRigControlElement>(InControlKey))
+		{
+			FRigControlElementCustomization ControlCustomization = *ControlRig->GetControlCustomization(InControlKey);
+			ControlCustomization.AvailableSpaces = InSpaceList;
+			ControlCustomization.RemovedSpaces.Reset();
+
+			// remember  the elements which are in the asset's available list but removed by the user
+			for (const FRigElementKey& AvailableSpace : ControlElement->Settings.Customization.AvailableSpaces)
+			{
+				if (!ControlCustomization.AvailableSpaces.Contains(AvailableSpace))
+				{
+					ControlCustomization.RemovedSpaces.Add(AvailableSpace);
+				}
+			}
+
+			ControlRig->SetControlCustomization(InControlKey, ControlCustomization);
+
+			if (FControlRigEditMode* EditMode = static_cast<FControlRigEditMode*>(ModeTools->GetActiveMode(FControlRigEditMode::ModeName)))
+			{
+				EditMode->SuspendHierarchyNotifs(true);
+				InHierarchy->Notify(ERigHierarchyNotification::ControlSettingChanged, ControlElement);
+				EditMode->SuspendHierarchyNotifs(false);
+			}
+			else
+			{
+				InHierarchy->Notify(ERigHierarchyNotification::ControlSettingChanged, ControlElement);
+			}
+
+			SpacePickerWidget->RefreshContents();
+		}
+	}
+}
+
+FReply SControlRigSpacePicker::HandleAddSpaceClicked()
+{
+	return SpacePickerWidget->HandleAddElementClicked();
+}
+
+FReply SControlRigSpacePicker::OnBakeControlsToNewSpaceButtonClicked()
+{
+	if (SpacePickerWidget->GetHierarchy() == nullptr)
+	{
+		return FReply::Unhandled();
+	}
+	if (SpacePickerWidget->GetControls().Num() == 0)
+	{
+		return FReply::Unhandled();
+	}
+	if (!GetControlRig())
+	{
+		return FReply::Unhandled();
+	}
+	ISequencer* Sequencer = GetSequencer();
+	if (Sequencer == nullptr || Sequencer->GetFocusedMovieSceneSequence() == nullptr || Sequencer->GetFocusedMovieSceneSequence()->GetMovieScene() == nullptr)
+	{
+		return FReply::Unhandled();
+	}
+	UControlRig* ControlRig = GetControlRig();
+
+	FRigSpacePickerBakeSettings Settings;
+	//Find default target space, just use first control and find space at current sequencer time
+	//Then Find range
+
+	// FindSpaceChannelAndSectionForControl() will trigger RecreateCurveEditor(), which will deselect the controls
+	// but in theory the selection will be recovered in the next tick, so here we just cache the selected controls
+	// and use it throughout this function. If this deselection is causing other problems, this part could use a revisit.
+	TArray<FRigElementKey> ControlKeys = SpacePickerWidget->GetControls();
+
+	FSpaceChannelAndSection SpaceChannelAndSection = FControlRigSpaceChannelHelpers::FindSpaceChannelAndSectionForControl(ControlRig, ControlKeys[0].Name, Sequencer, true /*bCreateIfNeeded*/);
+	if (SpaceChannelAndSection.SpaceChannel != nullptr)
+	{
+		const FFrameRate TickResolution = Sequencer->GetFocusedTickResolution();
+		const FFrameTime FrameTime = Sequencer->GetLocalTime().ConvertTo(TickResolution);
+		FFrameNumber CurrentTime = FrameTime.GetFrame();
+		FMovieSceneControlRigSpaceBaseKey Value;
+		using namespace UE::MovieScene;
+		URigHierarchy* RigHierarchy = SpacePickerWidget->GetHierarchy();
+		Settings.TargetSpace = URigHierarchy::GetDefaultParentKey();
+
+		TRange<FFrameNumber> Range = Sequencer->GetFocusedMovieSceneSequence()->GetMovieScene()->GetPlaybackRange();
+		TArray<FFrameNumber> Keys;
+		TArray < FKeyHandle> KeyHandles;
+
+		Settings.StartFrame = Range.GetLowerBoundValue();
+		Settings.EndFrame = Range.GetUpperBoundValue();
+		if (Keys.Num() > 0)
+		{
+			int32 Index = Algo::LowerBound(Keys, CurrentTime);
+			if (Index >= 0 && Index < (Keys.Num() - 1))
+			{
+				Settings.StartFrame = Keys[Index];
+				Settings.EndFrame = Keys[Index + 1];
+
+			}
+		}
+
+		TSharedRef<SRigSpacePickerBakeWidget> BakeWidget =
+			SNew(SRigSpacePickerBakeWidget)
+			.Settings(Settings)
+			.Hierarchy(SpacePickerWidget->GetHierarchy())
+			.Controls(ControlKeys) // use the cached controls here since the selection is not recovered until next tick.
+			.Sequencer(Sequencer)
+			.GetControlCustomization(this, &SControlRigSpacePicker::HandleGetControlElementCustomization)
+			.OnBake_Lambda([Sequencer, ControlRig, TickResolution](URigHierarchy* InHierarchy, TArray<FRigElementKey> InControls, FRigSpacePickerBakeSettings InSettings)
+		{
+			TArray<FFrameNumber> Frames;
+
+			const FFrameRate& FrameRate = Sequencer->GetFocusedDisplayRate();
+			FFrameNumber FrameRateInFrameNumber = TickResolution.AsFrameNumber(FrameRate.AsInterval());
+			for (FFrameNumber& Frame = InSettings.StartFrame; Frame <= InSettings.EndFrame; Frame += FrameRateInFrameNumber)
+			{
+				Frames.Add(Frame);
+			}
+			FScopedTransaction Transaction(LOCTEXT("BakeControlToSpace", "Bake Control In Space"));
+			for (const FRigElementKey& ControlKey : InControls)
+			{
+				//when baking we will now create a channel if one doesn't exist, was causing confusion
+				FSpaceChannelAndSection SpaceChannelAndSection = FControlRigSpaceChannelHelpers::FindSpaceChannelAndSectionForControl(ControlRig, ControlKey.Name, Sequencer, true /*bCreateIfNeeded*/);
+				if (SpaceChannelAndSection.SpaceChannel)
+				{
+					FControlRigSpaceChannelHelpers::SequencerBakeControlInSpace(ControlRig, Sequencer, SpaceChannelAndSection.SpaceChannel, SpaceChannelAndSection.SectionToKey,
+						Frames, InHierarchy, ControlKey, InSettings);
+				}
+			}
+			return FReply::Handled();
+		});
+
+		return BakeWidget->OpenDialog(true);
+	}
+	return FReply::Unhandled();
+}
+
+#undef LOCTEXT_NAMESPACE
diff --git a/Engine/Plugins/Animation/ControlRig/Source/ControlRigEditor/Private/EditMode/SControlRigSpacePicker.h b/Engine/Plugins/Animation/ControlRig/Source/ControlRigEditor/Private/EditMode/SControlRigSpacePicker.h
new file mode 100644
--- /dev/null
+++ b/Engine/Plugins/Animation/ControlRig/Source/ControlRigEditor/Private/EditMode/SControlRigSpacePicker.h
@@ -0,0 +1,50 @@
+// Copyright Epic Games, Inc. All Rights Reserved.
+
+/**
+* Space Picker View
+*/
+#pragma once
+
+#include "CoreMinimal.h"
+#include "ControlRigBaseDockableView.h"
+#include "Widgets/SWidget.h"
+#include "Widgets/SCompoundWidget.h"
+#include "Misc/FrameNumber.h"
+#include "ControlRig.h"
+#include "SRigSpacePickerWidget.h"
+
+
+class ISequencer;
+class SExpandableArea;
+class SSearchableRigHierarchyTreeView;
+
+class SControlRigSpacePicker : public SCompoundWidget, public FControlRigBaseDockableView
+{
+
+	SLATE_BEGIN_ARGS(SControlRigSpacePicker)
+	{}
+	SLATE_END_ARGS()
+		~SControlRigSpacePicker();
+
+	void Construct(const FArguments& InArgs, FControlRigEditMode& InEditMode);
+
+private:
+	virtual void HandleControlSelected(UControlRig* Subject, FRigControlElement* InControl, bool bSelected) override;
+	virtual void HandleControlAdded(UControlRig* ControlRig, bool bIsAdded) override;
+	virtual void NewControlRigSet(UControlRig* ControlRig) override;
+
+	const URigHierarchy* GetHierarchy() const;
+	void HandleSelectionChanged(TSharedPtr<FRigTreeElement> Selection, ESelectInfo::Type SelectInfo);
+
+	/** Space picker widget*/
+	TSharedPtr<SRigSpacePickerWidget> SpacePickerWidget;
+	TSharedPtr<SExpandableArea> PickerExpander;
+
+	const FRigControlElementCustomization* HandleGetControlElementCustomization(URigHierarchy* InHierarchy, const FRigElementKey& InControlKey);
+	void HandleActiveSpaceChanged(URigHierarchy* InHierarchy, const FRigElementKey& InControlKey, const FRigElementKey& InSpaceKey);
+	void HandleSpaceListChanged(URigHierarchy* InHierarchy, const FRigElementKey& InControlKey, const TArray<FRigElementKey>& InSpaceList);
+	FReply HandleAddSpaceClicked();
+	FReply OnBakeControlsToNewSpaceButtonClicked();
+
+};
+
