ControlRig: make sure duplicated temp assets used for editor validation cannot dirty the original assets

diff --git a/Engine/Plugins/Animation/ControlRig/Source/ControlRigDeveloper/Private/ControlRigBlueprint.cpp b/Engine/Plugins/Animation/ControlRig/Source/ControlRigDeveloper/Private/ControlRigBlueprint.cpp
--- a/Engine/Plugins/Animation/ControlRig/Source/ControlRigDeveloper/Private/ControlRigBlueprint.cpp
+++ b/Engine/Plugins/Animation/ControlRig/Source/ControlRigDeveloper/Private/ControlRigBlueprint.cpp
@@ -511,15 +511,23 @@ void UControlRigBlueprint::PostLoad()
 					}
 				}
 
-				if(URigVMFunctionReferenceNode* FunctionReferenceNode = Cast<URigVMFunctionReferenceNode>(Node))
+				
+				// avoid function reference related validation for temp assets, a temp asset may get generated during
+				// certain content validation process. It is usually just a simple file-level copy of the source asset
+				// so these references are usually not fixed-up properly. Thus, it is meaningless to validate them.
+				// They should not be allowed to dirty the source asset either.
+				if (!this->GetPackage()->GetName().StartsWith("/Temp/"))
 				{
-					if(URigVMLibraryNode* DependencyNode = FunctionReferenceNode->GetReferencedNode())
+					if(URigVMFunctionReferenceNode* FunctionReferenceNode = Cast<URigVMFunctionReferenceNode>(Node))
 					{
-						if(UControlRigBlueprint* DependencyBlueprint = DependencyNode->GetTypedOuter<UControlRigBlueprint>())
+						if(URigVMLibraryNode* DependencyNode = FunctionReferenceNode->GetReferencedNode())
 						{
-							if(DependencyBlueprint != this)
+							if(UControlRigBlueprint* DependencyBlueprint = DependencyNode->GetTypedOuter<UControlRigBlueprint>())
 							{
-								DependencyBlueprint->GetLocalFunctionLibrary()->UpdateReferencesForReferenceNode(FunctionReferenceNode);
+								if(DependencyBlueprint != this)
+								{
+									DependencyBlueprint->GetLocalFunctionLibrary()->UpdateReferencesForReferenceNode(FunctionReferenceNode);
+								}
 							}
 						}
 					}
@@ -532,6 +540,11 @@ void UControlRigBlueprint::PostLoad()
 #endif
 	}
 
+	if (FunctionLibrary)
+	{
+		FunctionLibrary->ClearInvalidReferences();
+	}
+	
 	// upgrade the gizmo libraries to shape libraries
 	if(GizmoLibrary_DEPRECATED.IsValid() || GetLinkerCustomVersion(FControlRigObjectVersion::GUID) < FControlRigObjectVersion::RenameGizmoToShape)
 	{
diff --git a/Engine/Source/Developer/RigVMDeveloper/Private/RigVMCompiler/RigVMCompiler.cpp b/Engine/Source/Developer/RigVMDeveloper/Private/RigVMCompiler/RigVMCompiler.cpp
--- a/Engine/Source/Developer/RigVMDeveloper/Private/RigVMCompiler/RigVMCompiler.cpp
+++ b/Engine/Source/Developer/RigVMDeveloper/Private/RigVMCompiler/RigVMCompiler.cpp
@@ -182,13 +182,19 @@ bool URigVMCompiler::Compile(URigVMGraph* InGraph, URigVMController* InControlle
 				bEncounteredGraphError = true;
 			}
 
-			if(URigVMFunctionReferenceNode* FunctionReferenceNode = Cast<URigVMFunctionReferenceNode>(ModelNode))
+			// avoid function reference related validation for temp assets, a temp asset may get generated during
+			// certain content validation process. It is usually just a simple file-level copy of the source asset
+			// so these references are usually not fixed-up properly. Thus, it is meaningless to validate them.
+			if (!ModelNode->GetPackage()->GetName().StartsWith(TEXT("/Temp/")))
 			{
-				if(!FunctionReferenceNode->IsFullyRemapped())
+				if(URigVMFunctionReferenceNode* FunctionReferenceNode = Cast<URigVMFunctionReferenceNode>(ModelNode))
 				{
-					static const FString UnmappedMessage = TEXT("Node @@ has unmapped variables. Please adjust the node and re-compile.");
-					Settings.ASTSettings.Report(EMessageSeverity::Error, ModelNode, UnmappedMessage);
-					bEncounteredGraphError = true;
+					if(!FunctionReferenceNode->IsFullyRemapped())
+					{
+						static const FString UnmappedMessage = TEXT("Node @@ has unmapped variables. Please adjust the node and re-compile.");
+						Settings.ASTSettings.Report(EMessageSeverity::Error, ModelNode, UnmappedMessage);
+						bEncounteredGraphError = true;
+					}
 				}
 			}
 
diff --git a/Engine/Source/Developer/RigVMDeveloper/Private/RigVMModel/RigVMFunctionLibrary.cpp b/Engine/Source/Developer/RigVMDeveloper/Private/RigVMModel/RigVMFunctionLibrary.cpp
--- a/Engine/Source/Developer/RigVMDeveloper/Private/RigVMModel/RigVMFunctionLibrary.cpp
+++ b/Engine/Source/Developer/RigVMDeveloper/Private/RigVMModel/RigVMFunctionLibrary.cpp
@@ -204,3 +204,24 @@ URigVMLibraryNode* URigVMFunctionLibrary::FindPreviouslyLocalizedFunction(URigVM
 	
 	return LocalizedFunction;
 }
+
+void URigVMFunctionLibrary::ClearInvalidReferences()
+{
+	// check each function's each reference
+	for (TTuple<TObjectPtr<URigVMLibraryNode>, FRigVMFunctionReferenceArray>& FunctionReferenceInfo : FunctionReferences)
+	{
+		FRigVMFunctionReferenceArray* ReferencesEntry = &FunctionReferenceInfo.Value;
+
+		Modify();
+		int32 NumRemoved = ReferencesEntry->FunctionReferences.RemoveAll([](TSoftObjectPtr<URigVMFunctionReferenceNode> Referencer)
+		{
+			return Referencer.ToString().StartsWith(TEXT("/Temp/"));
+		});
+
+		if (NumRemoved > 0)
+		{
+			MarkPackageDirty();
+		}
+	}
+}
+
diff --git a/Engine/Source/Developer/RigVMDeveloper/Public/RigVMModel/RigVMFunctionLibrary.h b/Engine/Source/Developer/RigVMDeveloper/Public/RigVMModel/RigVMFunctionLibrary.h
--- a/Engine/Source/Developer/RigVMDeveloper/Public/RigVMModel/RigVMFunctionLibrary.h
+++ b/Engine/Source/Developer/RigVMDeveloper/Public/RigVMModel/RigVMFunctionLibrary.h
@@ -89,6 +89,9 @@ public:
 	// and use this method to avoid redundant localizations.
 	URigVMLibraryNode* FindPreviouslyLocalizedFunction(URigVMLibraryNode* InFunctionToLocalize);
 
+	// Clear references to temp assets
+	void ClearInvalidReferences();
+	
 private:
 
 	UPROPERTY()
