Control Rig: Fix errors on nodes don't disappear after undoing

diff --git a/Engine/Plugins/Animation/ControlRig/Source/ControlRigDeveloper/Private/ControlRigBlueprint.cpp b/Engine/Plugins/Animation/ControlRig/Source/ControlRigDeveloper/Private/ControlRigBlueprint.cpp
--- a/Engine/Plugins/Animation/ControlRig/Source/ControlRigDeveloper/Private/ControlRigBlueprint.cpp
+++ b/Engine/Plugins/Animation/ControlRig/Source/ControlRigDeveloper/Private/ControlRigBlueprint.cpp
@@ -1975,17 +1975,17 @@ void UControlRigBlueprint::PostTransacted(const FTransactionObjectEvent& Transac
 			MarkPackageDirty();
 		}
 
-		else if (PropertiesChanged.Contains(TEXT("DrawContainer")))
+		if (PropertiesChanged.Contains(TEXT("DrawContainer")))
 		{
 			PropagateDrawInstructionsFromBPToInstances();
 		}
 
-		else if (PropertiesChanged.Contains(TEXT("VMRuntimeSettings")))
+		if (PropertiesChanged.Contains(TEXT("VMRuntimeSettings")))
 		{
 			PropagateRuntimeSettingsFromBPToInstances();
 		}
 
-		else if (PropertiesChanged.Contains(TEXT("NewVariables")))
+		if (PropertiesChanged.Contains(TEXT("NewVariables")))
 		{
 			if (RefreshEditorEvent.IsBound())
 			{
diff --git a/Engine/Source/Developer/RigVMDeveloper/Private/RigVMModel/RigVMController.cpp b/Engine/Source/Developer/RigVMDeveloper/Private/RigVMModel/RigVMController.cpp
--- a/Engine/Source/Developer/RigVMDeveloper/Private/RigVMModel/RigVMController.cpp
+++ b/Engine/Source/Developer/RigVMDeveloper/Private/RigVMModel/RigVMController.cpp
@@ -12402,7 +12402,7 @@ bool URigVMController::ChangePinType(URigVMPin* InPin, const FString& InCPPType,
 	UObject* CPPTypeObject = URigVMPin::FindObjectFromCPPTypeObjectPath(InCPPTypeObjectPath.ToString());
 	if (InPin->CPPType == InCPPType && InPin->CPPTypeObject == CPPTypeObject)
 	{
-		return false;
+		return true;
 	}
 
 	return ChangePinType(InPin, InCPPType, CPPTypeObject, bSetupUndoRedo, bSetupOrphanPins, bBreakLinks, bRemoveSubPins);
