Non-Unity fixes

diff --git a/Engine/Plugins/Animation/ControlRig/Source/ControlRigEditor/Private/EditMode/SControlRigSpacePicker.cpp b/Engine/Plugins/Animation/ControlRig/Source/ControlRigEditor/Private/EditMode/SControlRigSpacePicker.cpp
--- a/Engine/Plugins/Animation/ControlRig/Source/ControlRigEditor/Private/EditMode/SControlRigSpacePicker.cpp
+++ b/Engine/Plugins/Animation/ControlRig/Source/ControlRigEditor/Private/EditMode/SControlRigSpacePicker.cpp
@@ -10,6 +10,7 @@
 #include "AssetData.h"
 #include "EditorStyleSet.h"
 #include "Widgets/Text/SInlineEditableTextBlock.h"
+#include "Widgets/Layout/SSpacer.h"
 #include "EditorStyleSet.h"
 #include "ISequencer.h"
 #include "ScopedTransaction.h"
