Control Rig: Fix cannot add execute context to a function once it is removed

diff --git a/Engine/Plugins/Animation/ControlRig/Source/ControlRigDeveloper/Private/Graph/ControlRigGraphSchema.cpp b/Engine/Plugins/Animation/ControlRig/Source/ControlRigDeveloper/Private/Graph/ControlRigGraphSchema.cpp
--- a/Engine/Plugins/Animation/ControlRig/Source/ControlRigDeveloper/Private/Graph/ControlRigGraphSchema.cpp
+++ b/Engine/Plugins/Animation/ControlRig/Source/ControlRigDeveloper/Private/Graph/ControlRigGraphSchema.cpp
@@ -902,9 +902,12 @@ bool UControlRigGraphSchema::SupportsPinType(TWeakPtr<const FEdGraphSchemaAction
 	{
 		if (UScriptStruct* ScriptStruct = Cast<UScriptStruct>(PinType.PinSubCategoryObject))
 		{
-			if(ScriptStruct->IsChildOf(FRigVMExecuteContext::StaticStruct()))
+			if(SchemaAction.IsValid() && SchemaAction.Pin()->IsAVariable())
 			{
-				return false;
+				if(ScriptStruct->IsChildOf(FRigVMExecuteContext::StaticStruct()))
+				{
+					return false;
+				}
 			}
 			return SupportsPinType(ScriptStruct);
 		}
diff --git a/Engine/Plugins/Animation/ControlRig/Source/ControlRigEditor/Private/ControlRigGraphDetails.cpp b/Engine/Plugins/Animation/ControlRig/Source/ControlRigEditor/Private/ControlRigGraphDetails.cpp
--- a/Engine/Plugins/Animation/ControlRig/Source/ControlRigEditor/Private/ControlRigGraphDetails.cpp
+++ b/Engine/Plugins/Animation/ControlRig/Source/ControlRigEditor/Private/ControlRigGraphDetails.cpp
@@ -115,6 +115,7 @@ void FControlRigArgumentGroupLayout::HandleModifiedEvent(ERigVMGraphNotifType In
 		case ERigVMGraphNotifType::PinAdded:
 		case ERigVMGraphNotifType::PinRemoved:
 		case ERigVMGraphNotifType::PinIndexChanged:
+		case ERigVMGraphNotifType::PinTypeChanged:
 		{
 			URigVMPin* Pin = CastChecked<URigVMPin>(InSubject);
 			if (Pin->GetNode() == LibraryNode)
@@ -130,6 +131,63 @@ void FControlRigArgumentGroupLayout::HandleModifiedEvent(ERigVMGraphNotifType In
 	}
 }
 
+class FControlRigArgumentPinTypeSelectorFilter : public IPinTypeSelectorFilter
+{
+public:
+	FControlRigArgumentPinTypeSelectorFilter(TWeakPtr<IControlRigEditor> InControlRigEditor, TWeakObjectPtr<URigVMGraph> InGraph)
+		: ControlRigEditorPtr(InControlRigEditor), GraphPtr(InGraph)
+	{
+	}
+	
+	virtual bool ShouldShowPinTypeTreeItem(FPinTypeTreeItem InItem) const override
+	{
+		if (!InItem.IsValid())
+		{
+			return false;
+		}
+
+		// Only allow an execute context pin if the graph doesnt have one already
+		FString CPPType;
+		UObject* CPPTypeObject = nullptr;
+		RigVMTypeUtils::CPPTypeFromPinType(InItem.Get()->GetPinType(false), CPPType, &CPPTypeObject);
+		if (UScriptStruct* ScriptStruct = Cast<UScriptStruct>(CPPTypeObject))
+		{
+			if (ScriptStruct->IsChildOf(FRigVMExecuteContext::StaticStruct()))
+			{
+				if (GraphPtr.IsValid())
+				{
+					if (URigVMFunctionEntryNode* EntryNode = GraphPtr.Get()->GetEntryNode())
+					{
+						for (URigVMPin* Pin : EntryNode->GetPins())
+						{
+							if (Pin->IsExecuteContext())
+							{
+								return false;
+							}
+						}
+					}
+				}
+			}
+		}
+		
+		if (ControlRigEditorPtr.IsValid())
+		{
+			if (ControlRigEditorPtr.Pin()->GetImportedPinTypeSelectorFilter().IsValid())
+			{
+				return ControlRigEditorPtr.Pin()->GetImportedPinTypeSelectorFilter()->ShouldShowPinTypeTreeItem(InItem);				
+			}
+		}
+
+		return false;
+	}
+
+private:
+
+	TWeakPtr<IControlRigEditor> ControlRigEditorPtr;
+	
+	TWeakObjectPtr<URigVMGraph> GraphPtr;
+};
+
 void FControlRigArgumentLayout::GenerateHeaderRowContent(FDetailWidgetRow& NodeRow)
 {
 	const UEdGraphSchema* Schema = GetDefault<UControlRigGraphSchema>();
@@ -140,10 +198,9 @@ void FControlRigArgumentLayout::GenerateHeaderRowContent(FDetailWidgetRow& NodeR
 	TSharedPtr<IPinTypeSelectorFilter> CustomPinTypeFilter;
 	if (ControlRigEditorPtr.IsValid())
 	{
-		CustomPinTypeFilter = ControlRigEditorPtr.Pin()->GetImportedPinTypeSelectorFilter();
+		CustomPinTypeFilter = MakeShared<FControlRigArgumentPinTypeSelectorFilter>(ControlRigEditorPtr, GraphPtr);
 	}
 	
-
 	NodeRow
 		.NameContent()
 		[
@@ -413,29 +470,10 @@ void FControlRigArgumentLayout::PinInfoChanged(const FEdGraphPinType& PinType)
 		{
 			if (URigVMController* Controller = Blueprint->GetController(LibraryNode->GetContainedGraph()))
 			{
-				FRigVMExternalVariable ExternalVariable = RigVMTypeUtils::ExternalVariableFromPinType(Pin->GetFName(), PinType, true, false);
-				if (!ExternalVariable.IsValid(true /* allow nullptr memory */))
-				{
-					return;
-				}
-
-				FString CPPType = ExternalVariable.TypeName.ToString();
+				FString CPPType;
 				FName CPPTypeObjectName = NAME_None;
-				if (ExternalVariable.TypeObject)
-				{
-					CPPTypeObjectName = *ExternalVariable.TypeObject->GetPathName();
-
-					if (UScriptStruct* ScriptStruct = Cast<UScriptStruct>(ExternalVariable.TypeObject))
-					{
-						CPPType = ScriptStruct->GetStructCPPName();
-					}
-				}
-
-				if (ExternalVariable.bIsArray)
-				{
-					CPPType = FString::Printf(TEXT("TArray<%s>"), *CPPType);
-				}
-
+				RigVMTypeUtils::CPPTypeFromPinType(PinType, CPPType, CPPTypeObjectName);
+				
 				bool bSetupUndoRedo = true;
 				Controller->ChangeExposedPinType(Pin->GetFName(), CPPType, CPPTypeObjectName, bSetupUndoRedo, false, true);
 
diff --git a/Engine/Source/Developer/RigVMDeveloper/Private/RigVMModel/RigVMController.cpp b/Engine/Source/Developer/RigVMDeveloper/Private/RigVMModel/RigVMController.cpp
--- a/Engine/Source/Developer/RigVMDeveloper/Private/RigVMModel/RigVMController.cpp
+++ b/Engine/Source/Developer/RigVMDeveloper/Private/RigVMModel/RigVMController.cpp
@@ -12487,6 +12487,11 @@ bool URigVMController::ChangePinType(URigVMPin* InPin, const FString& InCPPType,
 	// we might want to use GetPinInitialDefaultValue here for a better default value
 	InPin->DefaultValue = FString();
 
+	if (InPin->IsExecuteContext())
+	{
+		InPin->Direction = ERigVMPinDirection::IO;
+	}
+
 	if (InPin->IsStruct())
 	{
 		FString DefaultValue = InPin->DefaultValue;
diff --git a/Engine/Source/Developer/RigVMDeveloper/Private/RigVMTypeUtils.cpp b/Engine/Source/Developer/RigVMDeveloper/Private/RigVMTypeUtils.cpp
--- a/Engine/Source/Developer/RigVMDeveloper/Private/RigVMTypeUtils.cpp
+++ b/Engine/Source/Developer/RigVMDeveloper/Private/RigVMTypeUtils.cpp
@@ -432,6 +432,7 @@ bool RigVMTypeUtils::CPPTypeFromPinType(const FEdGraphPinType& InPinType, FStrin
 		Suffix = TEXT(">");
 	}
 
+	OutCPPType = FString();
 	*OutCPPTypeObject = nullptr;
 	if (InPinType.PinCategory == UEdGraphSchema_K2::PC_Boolean)
 	{
@@ -496,6 +497,34 @@ bool RigVMTypeUtils::CPPTypeFromPinType(const FEdGraphPinType& InPinType, FStrin
 	return true;
 }
 
+bool RigVMTypeUtils::CPPTypeFromPinType(const FEdGraphPinType& InPinType, FString& OutCPPType, FName& OutCPPTypeObjectPath)
+{
+	OutCPPType = FString();
+	OutCPPTypeObjectPath = NAME_None;
+	UObject* CPPTypeObject = nullptr;
+	if (RigVMTypeUtils::CPPTypeFromPinType(InPinType, OutCPPType, &CPPTypeObject))
+	{
+		if (CPPTypeObject)
+		{
+			OutCPPTypeObjectPath = *CPPTypeObject->GetPathName();
+			if (UScriptStruct* ScriptStruct = Cast<UScriptStruct>(CPPTypeObject))
+			{
+				FString StructName = ScriptStruct->GetStructCPPName();
+				while (IsArrayType(OutCPPType))
+				{
+					OutCPPType = BaseTypeFromArrayType(OutCPPType);
+					StructName = ArrayTypeFromBaseType(StructName);
+				}
+				OutCPPType = StructName;
+			}
+		}
+
+		return true;
+	}
+
+	return false;
+}
+
 bool RigVMTypeUtils::CPPTypeFromExternalVariable(const FRigVMExternalVariable& InExternalVariable, FString& OutCPPType, UObject** OutCPPTypeObject)
 {
 	FString Prefix = "";
diff --git a/Engine/Source/Developer/RigVMDeveloper/Public/RigVMTypeUtils.h b/Engine/Source/Developer/RigVMDeveloper/Public/RigVMTypeUtils.h
--- a/Engine/Source/Developer/RigVMDeveloper/Public/RigVMTypeUtils.h
+++ b/Engine/Source/Developer/RigVMDeveloper/Public/RigVMTypeUtils.h
@@ -46,6 +46,8 @@ namespace RigVMTypeUtils
 
 	RIGVMDEVELOPER_API bool CPPTypeFromPinType(const FEdGraphPinType& InPinType, FString& OutCPPType, UObject** OutCPPTypeObject);
 
+	RIGVMDEVELOPER_API bool CPPTypeFromPinType(const FEdGraphPinType& InPinType, FString& OutCPPType, FName& OutCPPTypeObjectPath);
+
 	RIGVMDEVELOPER_API bool CPPTypeFromExternalVariable(const FRigVMExternalVariable& InExternalVariable, FString& OutCPPType, UObject** OutCPPTypeObject);
 
 	RIGVMDEVELOPER_API bool AreCompatible(const FRigVMExternalVariable& InTypeA, const FRigVMExternalVariable& InTypeB, const FString& InSegmentPathA = FString(), const FString& InSegmentPathB = FString());
diff --git a/Engine/Source/Editor/Kismet/Private/SBlueprintPalette.cpp b/Engine/Source/Editor/Kismet/Private/SBlueprintPalette.cpp
--- a/Engine/Source/Editor/Kismet/Private/SBlueprintPalette.cpp
+++ b/Engine/Source/Editor/Kismet/Private/SBlueprintPalette.cpp
@@ -725,30 +725,16 @@ public:
 		SLATE_EVENT(FOnPinTypeChanged, OnTypeChanged)
 	SLATE_END_ARGS()
 
-	/**
-	 * Constructs a PinTypeSelector widget (for variable actions only, so that 
-	 * the user can modify the variable's type without going to the details panel).
-	 * 
-	 * @param  InArgs					A set of slate arguments, defined above.
-	 * @param  InVariableProperty		The variable property to select
-	 * @param  InBlueprintEditor			A pointer to the blueprint editor that the palette belongs to.
-	 */
-	void Construct(const FArguments& InArgs, FProperty* InVariableProperty, UBlueprint* InBlueprint, TWeakPtr<FBlueprintEditor> InBlueprintEditor)
-	{
-		BlueprintObj = InBlueprint;
-		BlueprintEditorPtr = InBlueprintEditor;
-		VariableProperty = InVariableProperty;
-		ActionPtr = nullptr;
-
-		ConstructInternal(InArgs);
-	}
-
 	void Construct(const FArguments& InArgs, TWeakPtr<FEdGraphSchemaAction_BlueprintVariableBase> InAction, UBlueprint* InBlueprint, TWeakPtr<FBlueprintEditor> InBlueprintEditor)
 	{
 		BlueprintObj = InBlueprint;
 		BlueprintEditorPtr = InBlueprintEditor;
-		VariableProperty = nullptr;
 		ActionPtr = InAction;
+		VariableProperty = nullptr;
+		if (ActionPtr.IsValid())
+		{
+			VariableProperty = ActionPtr.Pin()->GetProperty();
+		}
 
 		ConstructInternal(InArgs);
 	}
@@ -1161,8 +1147,13 @@ void SBlueprintPaletteItem::Construct(const FArguments& InArgs, FCreateWidgetFor
 	TSharedRef<SWidget> NameSlotWidget = CreateTextSlotWidget(InCreateData, bIsReadOnly );
 	
 	// Will set the icon of this property to be a Pin Type selector. 
-	auto GenerateVariableSettings = [&](FProperty* VariableProp)
+	auto GenerateVariableSettings = [&](TSharedPtr<FEdGraphSchemaAction_BlueprintVariableBase> Action)
 	{
+		FProperty* VariableProp = nullptr;
+		if (Action.IsValid())
+		{
+			VariableProp = Action->GetProperty();
+		}
 		if (VariableProp)
 		{
 			if (bShouldCheckForAccessSpec)
@@ -1184,7 +1175,7 @@ void SBlueprintPaletteItem::Construct(const FArguments& InArgs, FCreateWidgetFor
 			if (FBlueprintEditorUtils::IsVariableCreatedByBlueprint(Blueprint, VariableProp) || VariableProp->GetOwner<UFunction>())
 			{
 				const UEdGraphSchema_K2* Schema = GetDefault<UEdGraphSchema_K2>();
-				IconWidget = SNew(SPinTypeSelectorHelper, VariableProp, Blueprint, BlueprintEditorPtr)
+				IconWidget = SNew(SPinTypeSelectorHelper, Action, Blueprint, BlueprintEditorPtr)
 					.IsEnabled(bIsEditingEnabled)
 					.ReadOnly_Lambda([this]() {return !IsHovered(); });
 			}
@@ -1194,11 +1185,11 @@ void SBlueprintPaletteItem::Construct(const FArguments& InArgs, FCreateWidgetFor
 	// For Variables and Local Variables, we will convert the icon widget into a pin type selector.
 	if (GraphAction->GetTypeId() == FEdGraphSchemaAction_K2Var::StaticGetTypeId())
 	{	
-		GenerateVariableSettings(StaticCastSharedPtr<FEdGraphSchemaAction_K2Var>(GraphAction)->GetProperty());
+		GenerateVariableSettings(StaticCastSharedPtr<FEdGraphSchemaAction_K2Var>(GraphAction));
 	}
 	else if (GraphAction->GetTypeId() == FEdGraphSchemaAction_K2LocalVar::StaticGetTypeId())
 	{
-		GenerateVariableSettings(StaticCastSharedPtr<FEdGraphSchemaAction_K2LocalVar>(GraphAction)->GetProperty());
+		GenerateVariableSettings(StaticCastSharedPtr<FEdGraphSchemaAction_K2LocalVar>(GraphAction));
 	}
 	else if (GraphAction->IsA(FEdGraphSchemaAction_BlueprintVariableBase::StaticGetTypeId()))
 	{
