Control Rig: Don't exit prematurely when selecting with the controller otherwise elements may not be selected/deselected correctly.  This comes up for example with Sequencer since it will auto-select child nodes, and often does a ClearSelection based upon a selection change, both of which can cause issues with this function.

diff --git a/Engine/Plugins/Animation/ControlRig/Source/ControlRig/Private/Rigs/RigHierarchyController.cpp b/Engine/Plugins/Animation/ControlRig/Source/ControlRig/Private/Rigs/RigHierarchyController.cpp
--- a/Engine/Plugins/Animation/ControlRig/Source/ControlRig/Private/Rigs/RigHierarchyController.cpp
+++ b/Engine/Plugins/Animation/ControlRig/Source/ControlRig/Private/Rigs/RigHierarchyController.cpp
@@ -113,7 +113,7 @@ bool URigHierarchyController::SetSelection(const TArray<FRigElementKey>& InKeys,
 	}
 
 	TArray<FRigElementKey> PreviousSelection = Hierarchy->GetSelectedKeys();
-	bool bResult = false;
+	bool bResult = true;
 
 	{
 		// disable python printing here as we only want to print a single command instead of one per selected item
@@ -125,9 +125,8 @@ bool URigHierarchyController::SetSelection(const TArray<FRigElementKey>& InKeys,
 			{
 				if(!SelectElement(KeyToDeselect, false))
 				{
-					return false;
+					bResult = false;
 				}
-				bResult = true;
 			}
 		}
 
@@ -137,9 +136,8 @@ bool URigHierarchyController::SetSelection(const TArray<FRigElementKey>& InKeys,
 			{
 				if(!SelectElement(KeyToSelect, true))
 				{
-					return false;
+					bResult = false;
 				}
-				bResult = true;
 			}
 		}
 	}
