Fix nullptr dereference in FControlRigComponentMappedElement::GetAnimProxyOnGameThread()
 - UObject Cast can fail, returning nullptr, if the object is not the correct type. This led to calling GetControlRigProxyOnGameThread() on the nullptr.

diff --git a/Engine/Plugins/Animation/ControlRig/Source/ControlRig/Private/ControlRigComponent.cpp b/Engine/Plugins/Animation/ControlRig/Source/ControlRig/Private/ControlRigComponent.cpp
--- a/Engine/Plugins/Animation/ControlRig/Source/ControlRig/Private/ControlRigComponent.cpp
+++ b/Engine/Plugins/Animation/ControlRig/Source/ControlRig/Private/ControlRigComponent.cpp
@@ -15,8 +15,10 @@ FControlRigAnimInstanceProxy* FControlRigComponentMappedElement::GetAnimProxyOnG
 {
 	if (USkeletalMeshComponent* SkeletalMeshComponent = Cast<USkeletalMeshComponent>(SceneComponent))
 	{
-		UControlRigAnimInstance* AnimInstance = Cast<UControlRigAnimInstance>(SkeletalMeshComponent->GetAnimInstance());
-		return AnimInstance->GetControlRigProxyOnGameThread();
+		if (UControlRigAnimInstance* AnimInstance = Cast<UControlRigAnimInstance>(SkeletalMeshComponent->GetAnimInstance()))
+		{
+			return AnimInstance->GetControlRigProxyOnGameThread();
+		}
 	}
 
 	return nullptr;
