ControlRig:
- Fix missing call to super in UControlRigComponent::InitializeComponent.
- Don't call UControlRigComponent::Initialize from UControlRigComponent::InitializeComponent.

Saves multiple frames over 200ms due to calling UControlRigComponent::Initialize from ULevel::RouteActorInitialize in cooked builds.

diff --git a/Engine/Plugins/Animation/ControlRig/Source/ControlRig/Private/ControlRigComponent.cpp b/Engine/Plugins/Animation/ControlRig/Source/ControlRig/Private/ControlRigComponent.cpp
--- a/Engine/Plugins/Animation/ControlRig/Source/ControlRig/Private/ControlRigComponent.cpp
+++ b/Engine/Plugins/Animation/ControlRig/Source/ControlRig/Private/ControlRigComponent.cpp
@@ -92,10 +92,13 @@ void UControlRigComponent::BeginDestroy()
 	gPendingSkeletalMeshes.Remove(this);
 }
 
+#if WITH_EDITOR
 void UControlRigComponent::InitializeComponent()
 {
+	Super::InitializeComponent();
 	Initialize();
 }
+#endif
 
 void UControlRigComponent::OnRegister()
 {
diff --git a/Engine/Plugins/Animation/ControlRig/Source/ControlRig/Public/ControlRigComponent.h b/Engine/Plugins/Animation/ControlRig/Source/ControlRig/Public/ControlRigComponent.h
--- a/Engine/Plugins/Animation/ControlRig/Source/ControlRig/Public/ControlRigComponent.h
+++ b/Engine/Plugins/Animation/ControlRig/Source/ControlRig/Public/ControlRigComponent.h
@@ -215,7 +215,9 @@ public:
 	//~ End UObject Interface
 	
 	//~ Begin UActorComponent Interface
+#if WITH_EDITOR
 	virtual void InitializeComponent() override;
+#endif
 	virtual void OnRegister() override;
 	virtual void OnUnregister() override;
 	virtual void TickComponent(float DeltaTime, enum ELevelTick TickType, FActorComponentTickFunction* ThisTickFunction) override;
