Control Rig: Adjust icons for Tab Summoner, Rig Tree

diff --git a/Engine/Plugins/Animation/ControlRig/Source/ControlRig/Public/Rigs/RigHierarchyDefines.h b/Engine/Plugins/Animation/ControlRig/Source/ControlRig/Public/Rigs/RigHierarchyDefines.h
--- a/Engine/Plugins/Animation/ControlRig/Source/ControlRig/Public/Rigs/RigHierarchyDefines.h
+++ b/Engine/Plugins/Animation/ControlRig/Source/ControlRig/Public/Rigs/RigHierarchyDefines.h
@@ -22,7 +22,7 @@ enum class ERigElementType : uint8
 	Space = Null UMETA(Hidden),
 	Control = 0x004,
 	Curve = 0x008,
-	RigidBody = 0x010,
+	RigidBody = 0x010 UMETA(Hidden), 
 	Reference = 0x020,
 	Last = 0x040 UMETA(Hidden),
 	All = Bone | Null | Control | Curve | RigidBody | Reference,
diff --git a/Engine/Plugins/Animation/ControlRig/Source/ControlRigEditor/Private/ControlRigEditorStyle.h b/Engine/Plugins/Animation/ControlRig/Source/ControlRigEditor/Private/ControlRigEditorStyle.h
--- a/Engine/Plugins/Animation/ControlRig/Source/ControlRigEditor/Private/ControlRigEditorStyle.h
+++ b/Engine/Plugins/Animation/ControlRig/Source/ControlRigEditor/Private/ControlRigEditorStyle.h
@@ -14,6 +14,25 @@
 class FControlRigEditorStyle final
 	: public FSlateStyleSet
 {
+	class FContentRootBracket
+	{
+	public:
+		FContentRootBracket(FControlRigEditorStyle* InStyle, const FString& NewContentRoot)
+			: Style(InStyle)
+			, PreviousContentRoot(InStyle->GetContentRootDir())
+		{
+			Style->SetContentRoot(NewContentRoot);
+		}
+
+		~FContentRootBracket()
+		{
+			Style->SetContentRoot(PreviousContentRoot);
+		}
+	private:
+		FControlRigEditorStyle* Style;
+		FString PreviousContentRoot;
+	};
+	
 public:
 	FControlRigEditorStyle()
 		: FSlateStyleSet("ControlRigEditorStyle")
@@ -25,16 +44,17 @@ public:
 		const FVector2D Icon24x24(24.0f, 24.0f);
 		const FVector2D Icon32x32(32.0f, 32.0f);
 		const FVector2D Icon40x40(40.0f, 40.0f);
-		const FString PluginContentDir = FPaths::EnginePluginsDir() / TEXT("Animation/ControlRig/Content");
+		const FString ControlRigPluginContentDir = FPaths::EnginePluginsDir() / TEXT("Animation/ControlRig/Content");
 		const FString EngineEditorSlateDir = FPaths::EngineContentDir() / TEXT("Editor/Slate");
-		SetContentRoot(PluginContentDir);
+		SetContentRoot(ControlRigPluginContentDir);
 
 		const FSlateColor DefaultForeground(FLinearColor(0.72f, 0.72f, 0.72f, 1.f));
 
 		// Class Icons
 		{
 			Set("ClassIcon.ControlRigSequence", new IMAGE_BRUSH("Slate/ControlRigSequence_16x", Icon16x16));
-			Set("ClassIcon.ControlRigBlueprint", new IMAGE_BRUSH("Slate/ControlRigBlueprint_16x", Icon16x16));
+			Set("ClassIcon.ControlRigBlueprint", new IMAGE_BRUSH("Slate/ControlRig_16", Icon16x16));
+			Set("ClassIcon.ControlRigPose", new IMAGE_BRUSH("Slate/ControlRigPose_16", Icon16x16));
 		}
 
 		// Edit mode styles
@@ -66,28 +86,29 @@ public:
 
 		// Control Rig Editor styles
 		{
+			// tab icons
 			Set("ControlRig.TabIcon", new IMAGE_BRUSH("Slate/ControlRigTab_16x", Icon16x16));
-			Set("ControlRig.RigUnit", new IMAGE_BRUSH("Slate/ControlRigUnit_16x", Icon16x16));
-
+			Set("ExecutionStack.TabIcon", new IMAGE_BRUSH_SVG("Slate/ExecutionStack", Icon16x16));
+			Set("RigHierarchy.TabIcon", new IMAGE_BRUSH_SVG("Slate/RigHierarchy", Icon16x16));
+			Set("RigValidation.TabIcon", new IMAGE_BRUSH_SVG("Slate/RigValidation", Icon16x16));
+			Set("CurveContainer.TabIcon", new IMAGE_BRUSH_SVG("Slate/CurveContainer", Icon16x16));
 			Set("HierarchicalProfiler.TabIcon", new IMAGE_BRUSH_SVG("Slate/HierarchicalProfiler", Icon16x16));
-
+			
 			// icons for control units
-			Set("ControlRig.ControlUnitOn", new IMAGE_BRUSH("Slate/ControlUnit_On", Icon32x32));
-			Set("ControlRig.ControlUnitOff", new IMAGE_BRUSH("Slate/ControlUnit_Off", Icon32x32));
-
-			Set("ControlRig.ExecuteGraph", new IMAGE_BRUSH("Slate/ExecuteGraph", Icon40x40));
-			Set("ControlRig.ExecuteGraph.Small", new IMAGE_BRUSH("Slate/ExecuteGraph", Icon20x20));
+			Set("ControlRig.RigUnit", new IMAGE_BRUSH("Slate/ControlRigUnit_16x", Icon16x16));
 
-			Set("ControlRig.AutoCompileGraph", new IMAGE_BRUSH("Slate/AutoCompile", Icon40x40));
+			Set("ControlRig.AutoCompileGraph", new IMAGE_BRUSH("Slate/AutoCompile", Icon20x20));
 			Set("ControlRig.AutoCompileGraph.Small", new IMAGE_BRUSH("Slate/AutoCompile", Icon20x20));
 
-			Set("ControlRig.SetupMode", new IMAGE_BRUSH("Slate/SetupMode", Icon40x40));
-			Set("ControlRig.SetupMode.Small", new IMAGE_BRUSH("Slate/SetupMode", Icon20x20));
-
-			Set("ControlRig.UpdateEvent", new IMAGE_BRUSH("Slate/UpdateEvent", Icon40x40));
-			Set("ControlRig.InverseEvent", new IMAGE_BRUSH("Slate/InverseEvent", Icon40x40));
-			Set("ControlRig.UpdateAndInverse", new IMAGE_BRUSH("Slate/UpdateAndInverse", Icon40x40));
-			Set("ControlRig.InverseAndUpdate", new IMAGE_BRUSH("Slate/InverseAndUpdate", Icon40x40));
+			{
+				FContentRootBracket Bracket(this, EngineEditorSlateDir);
+				Set("ControlRig.SetupMode", new IMAGE_BRUSH_SVG("Starship/Common/Adjust", Icon40x40));
+				Set("ControlRig.SetupMode.Small", new IMAGE_BRUSH_SVG("Starship/Common/Adjust", Icon20x20));
+				Set("ControlRig.UpdateEvent", new IMAGE_BRUSH("Icons/diff_next_40x", Icon40x40));
+				Set("ControlRig.InverseEvent", new IMAGE_BRUSH("Icons/diff_prev_40x", Icon40x40));
+				Set("ControlRig.UpdateAndInverse", new IMAGE_BRUSH("Icons/Loop_40x", Icon40x40));
+				Set("ControlRig.InverseAndUpdate", new IMAGE_BRUSH("Icons/Loop_40x", Icon40x40));
+			}
 
 			Set("ControlRig.Bug.Dot", new IMAGE_BRUSH("Slate/ControlRig_BugDot_32x", Icon16x16));
 			Set("ControlRig.Bug.Normal", new IMAGE_BRUSH("Slate/ControlRig_Bug_28x", Icon14x14));
@@ -97,21 +118,13 @@ public:
 			Set("ControlRig.ResumeExecution", new IMAGE_BRUSH_SVG("Slate/simulate", Icon40x40));
 			Set("ControlRig.ReleaseMode", new IMAGE_BRUSH_SVG("Slate/animation", Icon40x40));
 			Set("ControlRig.DebugMode", new IMAGE_BRUSH_SVG("Slate/Bug", Icon40x40));
-			
-			Set( "ControlRig.StepOut", new IMAGE_BRUSH("Slate/icon_DebugStepOut_40x", Icon40x40));
-			Set( "ControlRig.StepInto", new IMAGE_BRUSH( "Slate/icon_DebugStepIn_40x", Icon40x40 ) );
-			Set( "ControlRig.StepOver", new IMAGE_BRUSH( "Slate/icon_DebugStepOver_40x", Icon40x40 ) );
 
 			{
-				const FString PreviousContentRoot = GetContentRootDir();
-				// Common/RoundedSelection_16x lives in the editor slate folder, change the root temporarily
-				SetContentRoot(EngineEditorSlateDir);
+				FContentRootBracket Bracket(this, EngineEditorSlateDir);
 				// similar style to "LevelViewport.StartingPlayInEditorBorder"
 				Set( "ControlRig.Viewport.Border", new BOX_BRUSH( "Old/Window/ViewportDebugBorder", 0.8f, FLinearColor(1.0f,1.0f,1.0f,1.0f) ) );
 				// similar style to "AnimViewport.Notification.Warning"
 				Set( "ControlRig.Viewport.Notification.ChangeShapeTransform", new BOX_BRUSH("Common/RoundedSelection_16x", 4.0f/16.0f, FLinearColor(FColor(169, 0, 148))));
-				// restore the previous content root
-				SetContentRoot(PreviousContentRoot);
 			}
 		}
 
@@ -132,9 +145,9 @@ public:
 
 		// Tree styles
 		{
-			Set("ControlRig.Tree.BoneUser", new IMAGE_BRUSH("Slate/BoneUser_16x", Icon16x16));
-			Set("ControlRig.Tree.BoneImported", new IMAGE_BRUSH("Slate/BoneImported_16x", Icon16x16));
-			Set("ControlRig.Tree.Control", new IMAGE_BRUSH("Slate/Control_16x", Icon16x16));
+			Set("ControlRig.Tree.BoneUser", new IMAGE_BRUSH("Slate/BoneNonWeighted_16x", Icon16x16));
+			Set("ControlRig.Tree.BoneImported", new IMAGE_BRUSH("Slate/Bone_16x", Icon16x16));
+			Set("ControlRig.Tree.Control", new IMAGE_BRUSH("Slate/RigControlCircle_16x", Icon16x16));
 			Set("ControlRig.Tree.Null", new IMAGE_BRUSH("Slate/Null_16x", Icon16x16));
 			Set("ControlRig.Tree.RigidBody", new IMAGE_BRUSH("Slate/RigidBody_16x", Icon16x16));
 			Set("ControlRig.Tree.Socket", new IMAGE_BRUSH("Slate/Socket_16x", Icon16x16));
diff --git a/Engine/Plugins/Animation/ControlRig/Source/ControlRigEditor/Private/Editor/RigCurveContainerTabSummoner.cpp b/Engine/Plugins/Animation/ControlRig/Source/ControlRigEditor/Private/Editor/RigCurveContainerTabSummoner.cpp
--- a/Engine/Plugins/Animation/ControlRig/Source/ControlRigEditor/Private/Editor/RigCurveContainerTabSummoner.cpp
+++ b/Engine/Plugins/Animation/ControlRig/Source/ControlRigEditor/Private/Editor/RigCurveContainerTabSummoner.cpp
@@ -14,7 +14,7 @@ FRigCurveContainerTabSummoner::FRigCurveContainerTabSummoner(const TSharedRef<FC
 	, ControlRigEditor(InControlRigEditor)
 {
 	TabLabel = LOCTEXT("RigCurveContainerTabLabel", "Curve Container");
-	TabIcon = FSlateIcon(FControlRigEditorStyle::Get().GetStyleSetName(), "ControlRig.TabIcon");
+	TabIcon = FSlateIcon(FControlRigEditorStyle::Get().GetStyleSetName(), "CurveContainer.TabIcon");
 
 	bIsSingleton = true;
 
diff --git a/Engine/Plugins/Animation/ControlRig/Source/ControlRigEditor/Private/Editor/RigHierarchyTabSummoner.cpp b/Engine/Plugins/Animation/ControlRig/Source/ControlRigEditor/Private/Editor/RigHierarchyTabSummoner.cpp
--- a/Engine/Plugins/Animation/ControlRig/Source/ControlRigEditor/Private/Editor/RigHierarchyTabSummoner.cpp
+++ b/Engine/Plugins/Animation/ControlRig/Source/ControlRigEditor/Private/Editor/RigHierarchyTabSummoner.cpp
@@ -16,7 +16,7 @@ FRigHierarchyTabSummoner::FRigHierarchyTabSummoner(const TSharedRef<FControlRigE
 	, ControlRigEditor(InControlRigEditor)
 {
 	TabLabel = LOCTEXT("RigHierarchyTabLabel", "Rig Hierarchy");
-	TabIcon = FSlateIcon(FControlRigEditorStyle::Get().GetStyleSetName(), "ControlRig.TabIcon");
+	TabIcon = FSlateIcon(FControlRigEditorStyle::Get().GetStyleSetName(), "RigHierarchy.TabIcon");
 
 	ViewMenuDescription = LOCTEXT("RigHierarchy_ViewMenu_Desc", "Rig Hierarchy");
 	ViewMenuTooltip = LOCTEXT("RigHierarchy_ViewMenu_ToolTip", "Show the Rig Hierarchy tab");
diff --git a/Engine/Plugins/Animation/ControlRig/Source/ControlRigEditor/Private/Editor/RigStackTabSummoner.cpp b/Engine/Plugins/Animation/ControlRig/Source/ControlRigEditor/Private/Editor/RigStackTabSummoner.cpp
--- a/Engine/Plugins/Animation/ControlRig/Source/ControlRigEditor/Private/Editor/RigStackTabSummoner.cpp
+++ b/Engine/Plugins/Animation/ControlRig/Source/ControlRigEditor/Private/Editor/RigStackTabSummoner.cpp
@@ -14,7 +14,7 @@ FRigStackTabSummoner::FRigStackTabSummoner(const TSharedRef<FControlRigEditor>&
 	, ControlRigEditor(InControlRigEditor)
 {
 	TabLabel = LOCTEXT("RigStackTabLabel", "Execution Stack");
-	TabIcon = FSlateIcon(FControlRigEditorStyle::Get().GetStyleSetName(), "ControlRig.TabIcon");
+	TabIcon = FSlateIcon(FControlRigEditorStyle::Get().GetStyleSetName(), "ExecutionStack.TabIcon");
 
 	bIsSingleton = true;
 
diff --git a/Engine/Plugins/Animation/ControlRig/Source/ControlRigEditor/Private/Editor/RigValidationTabSummoner.cpp b/Engine/Plugins/Animation/ControlRig/Source/ControlRigEditor/Private/Editor/RigValidationTabSummoner.cpp
--- a/Engine/Plugins/Animation/ControlRig/Source/ControlRigEditor/Private/Editor/RigValidationTabSummoner.cpp
+++ b/Engine/Plugins/Animation/ControlRig/Source/ControlRigEditor/Private/Editor/RigValidationTabSummoner.cpp
@@ -16,7 +16,7 @@ FRigValidationTabSummoner::FRigValidationTabSummoner(const TSharedRef<FControlRi
 	, WeakControlRigEditor(InControlRigEditor)
 {
 	TabLabel = LOCTEXT("RigValidationTabLabel", "Rig Validation");
-	TabIcon = FSlateIcon(FControlRigEditorStyle::Get().GetStyleSetName(), "ControlRig.TabIcon");
+	TabIcon = FSlateIcon(FControlRigEditorStyle::Get().GetStyleSetName(), "RigValidation.TabIcon");
 
 	bIsSingleton = true;
 
diff --git a/Engine/Plugins/Animation/ControlRig/Source/ControlRigEditor/Private/Editor/SRigHierarchyTreeView.cpp b/Engine/Plugins/Animation/ControlRig/Source/ControlRigEditor/Private/Editor/SRigHierarchyTreeView.cpp
--- a/Engine/Plugins/Animation/ControlRig/Source/ControlRigEditor/Private/Editor/SRigHierarchyTreeView.cpp
+++ b/Engine/Plugins/Animation/ControlRig/Source/ControlRigEditor/Private/Editor/SRigHierarchyTreeView.cpp
@@ -87,8 +87,18 @@ void SRigHierarchyItem::Construct(const FArguments& InArgs, const TSharedRef<STa
 
 	TSharedPtr< SInlineEditableTextBlock > InlineWidget;
 
-	const FSlateBrush* Brush = GetBrushForElementType(Delegates.GetHierarchy(), InRigTreeElement->Key);
-	FSlateColor Color = InRigTreeElement->FilterResult == ERigTreeFilterResult::Shown ? FSlateColor::UseForeground() : FSlateColor(FLinearColor::Gray * 0.5f); 
+	TPair<const FSlateBrush*, FSlateColor> Result = GetBrushForElementType(Delegates.GetHierarchy(), InRigTreeElement->Key);
+	const FSlateBrush* Brush = Result.Key;
+	FSlateColor IconColor = Result.Value;
+	if(IconColor.IsColorSpecified())
+	{
+		IconColor = InRigTreeElement->FilterResult == ERigTreeFilterResult::Shown ? Result.Value : FSlateColor(Result.Value.GetSpecifiedColor() * 0.5f);
+	}
+	else
+	{
+		IconColor = InRigTreeElement->FilterResult == ERigTreeFilterResult::Shown ? FSlateColor::UseForeground() : FSlateColor(FLinearColor::Gray * 0.5f);
+	}
+	const FSlateColor TextColor = InRigTreeElement->FilterResult == ERigTreeFilterResult::Shown ? FSlateColor::UseForeground() : FSlateColor(FLinearColor::Gray * 0.5f);
 
 	STableRow<TSharedPtr<FRigTreeElement>>::Construct(
 		STableRow<TSharedPtr<FRigTreeElement>>::FArguments()
@@ -104,10 +114,11 @@ void SRigHierarchyItem::Construct(const FArguments& InArgs, const TSharedRef<STa
 			.FillWidth(1.0)
 			.HAlign(HAlign_Left)
 			.VAlign(VAlign_Center)
+			.Padding(FMargin(0.f, 0.f, 3.f, 0.f))
 			[
 				SNew(SImage)
 				.Image(Brush)
-				.ColorAndOpacity(Color)
+				.ColorAndOpacity(IconColor)
 			]
 			+ SHorizontalBox::Slot()
 			.AutoWidth()
@@ -118,7 +129,7 @@ void SRigHierarchyItem::Construct(const FArguments& InArgs, const TSharedRef<STa
 				.OnVerifyTextChanged(this, &SRigHierarchyItem::OnVerifyNameChanged)
 				.OnTextCommitted(this, &SRigHierarchyItem::OnNameCommitted)
 				.MultiLine(false)
-				.ColorAndOpacity(Color)
+				.ColorAndOpacity(TextColor)
 			]
 		], OwnerTable);
 
@@ -612,14 +623,19 @@ bool SRigHierarchyItem::OnVerifyNameChanged(const FText& InText, FText& OutError
 	return Delegates.HandleVerifyElementNameChanged(OldKey, NewName, OutErrorMessage);
 }
 
-const FSlateBrush* SRigHierarchyItem::GetBrushForElementType(const URigHierarchy* InHierarchy, const FRigElementKey& InKey)
+TPair<const FSlateBrush*, FSlateColor> SRigHierarchyItem::GetBrushForElementType(const URigHierarchy* InHierarchy, const FRigElementKey& InKey)
 {
 	const FSlateBrush* Brush = nullptr;
+	FSlateColor Color = FSlateColor::UseForeground();
 	switch (InKey.Type)
 	{
 		case ERigElementType::Control:
 		{
 			Brush = FControlRigEditorStyle::Get().GetBrush("ControlRig.Tree.Control");
+			if(const FRigControlElement* Control = InHierarchy->Find<FRigControlElement>(InKey))
+			{
+				Color = FSlateColor(Control->Settings.ShapeColor);
+			}
 			break;
 		}
 		case ERigElementType::Null:
@@ -673,7 +689,7 @@ const FSlateBrush* SRigHierarchyItem::GetBrushForElementType(const URigHierarchy
 		}
 	}
 
-	return Brush;
+	return TPair<const FSlateBrush*, FSlateColor>(Brush, Color);
 }
 
 void SRigHierarchyItem::OnNameCommitted(const FText& InText, ETextCommit::Type InCommitType) const
diff --git a/Engine/Plugins/Animation/ControlRig/Source/ControlRigEditor/Private/Editor/SRigHierarchyTreeView.h b/Engine/Plugins/Animation/ControlRig/Source/ControlRigEditor/Private/Editor/SRigHierarchyTreeView.h
--- a/Engine/Plugins/Animation/ControlRig/Source/ControlRigEditor/Private/Editor/SRigHierarchyTreeView.h
+++ b/Engine/Plugins/Animation/ControlRig/Source/ControlRigEditor/Private/Editor/SRigHierarchyTreeView.h
@@ -189,7 +189,7 @@ public:
 	void Construct(const FArguments& InArgs, const TSharedRef<STableViewBase>& OwnerTable, TSharedRef<FRigTreeElement> InRigTreeElement, TSharedPtr<SRigHierarchyTreeView> InTreeView);
  	void OnNameCommitted(const FText& InText, ETextCommit::Type InCommitType) const;
 	bool OnVerifyNameChanged(const FText& InText, FText& OutErrorMessage);
-	static const FSlateBrush* GetBrushForElementType(const URigHierarchy* InHierarchy, const FRigElementKey& InKey);
+	static TPair<const FSlateBrush*, FSlateColor> GetBrushForElementType(const URigHierarchy* InHierarchy, const FRigElementKey& InKey);
 
 private:
 	TWeakPtr<FRigTreeElement> WeakRigTreeElement;
diff --git a/Engine/Plugins/Animation/ControlRig/Source/ControlRigEditor/Private/Editor/SRigSpacePickerWidget.cpp b/Engine/Plugins/Animation/ControlRig/Source/ControlRigEditor/Private/Editor/SRigSpacePickerWidget.cpp
--- a/Engine/Plugins/Animation/ControlRig/Source/ControlRigEditor/Private/Editor/SRigSpacePickerWidget.cpp
+++ b/Engine/Plugins/Animation/ControlRig/Source/ControlRigEditor/Private/Editor/SRigSpacePickerWidget.cpp
@@ -95,6 +95,7 @@ void SRigSpacePickerWidget::Construct(const FArguments& InArgs)
 			ESpacePickerType_Parent,
 			URigHierarchy::GetDefaultParentKey(),
 			FAppStyle::Get().GetBrush("Icons.Transform"),
+			FSlateColor::UseForeground(),
 			LOCTEXT("Parent", "Parent"),
 			FOnClicked::CreateSP(this, &SRigSpacePickerWidget::HandleParentSpaceClicked)
 		);
@@ -104,6 +105,7 @@ void SRigSpacePickerWidget::Construct(const FArguments& InArgs)
 			ESpacePickerType_World,
 			URigHierarchy::GetWorldSpaceReferenceKey(),
 			FEditorStyle::GetBrush("EditorViewport.RelativeCoordinateSystem_World"),
+			FSlateColor::UseForeground(),
 			LOCTEXT("World", "World"),
 			FOnClicked::CreateSP(this, &SRigSpacePickerWidget::HandleWorldSpaceClicked)
 		);
@@ -379,6 +381,7 @@ void SRigSpacePickerWidget::AddSpacePickerRow(
 	ESpacePickerType InType,
 	const FRigElementKey& InKey,
 	const FSlateBrush* InBush,
+	const FSlateColor& InColor,
 	const FText& InTitle,
     FOnClicked OnClickedDelegate)
 {
@@ -414,10 +417,11 @@ void SRigSpacePickerWidget::AddSpacePickerRow(
 					.AutoWidth()
 					.VAlign(VAlign_Center)
 					.HAlign(HAlign_Left)
-					.Padding(0)
+					.Padding(FMargin(0.f, 0.f, 3.f, 0.f))
 					[
 						SNew(SImage)
 						.Image(InBush)
+						.ColorAndOpacity(InColor)
 					]
 
 					+ SHorizontalBox::Slot()
@@ -944,11 +948,14 @@ void SRigSpacePickerWidget::RepopulateItemSpaces()
 
 	for(const FRigElementKey& Key : Keys)
 	{
+		TPair<const FSlateBrush*, FSlateColor> IconAndColor = SRigHierarchyItem::GetBrushForElementType(Hierarchy, Key);
+		
 		AddSpacePickerRow(
 			ItemSpacesListBox,
 			ESpacePickerType_Item,
 			Key,
-			SRigHierarchyItem::GetBrushForElementType(Hierarchy, Key),
+			IconAndColor.Key,
+			IconAndColor.Value,
 			FText::FromName(Key.Name),
 			FOnClicked::CreateSP(this, &SRigSpacePickerWidget::HandleElementSpaceClicked, Key)
 		);
diff --git a/Engine/Plugins/Animation/ControlRig/Source/ControlRigEditor/Private/Editor/SRigSpacePickerWidget.h b/Engine/Plugins/Animation/ControlRig/Source/ControlRigEditor/Private/Editor/SRigSpacePickerWidget.h
--- a/Engine/Plugins/Animation/ControlRig/Source/ControlRigEditor/Private/Editor/SRigSpacePickerWidget.h
+++ b/Engine/Plugins/Animation/ControlRig/Source/ControlRigEditor/Private/Editor/SRigSpacePickerWidget.h
@@ -129,6 +129,7 @@ private:
 		ESpacePickerType InType,
 		const FRigElementKey& InKey,
 		const FSlateBrush* InBush,
+		const FSlateColor& InColor,
 		const FText& InTitle,
 		FOnClicked OnClickedDelegate);
 
