Control Rig: Fix regression with input widgets showing on structs

diff --git a/Engine/Plugins/Animation/ControlRig/Source/ControlRigEditor/Private/Graph/SControlRigGraphNode.cpp b/Engine/Plugins/Animation/ControlRig/Source/ControlRigEditor/Private/Graph/SControlRigGraphNode.cpp
--- a/Engine/Plugins/Animation/ControlRig/Source/ControlRigEditor/Private/Graph/SControlRigGraphNode.cpp
+++ b/Engine/Plugins/Animation/ControlRig/Source/ControlRigEditor/Private/Graph/SControlRigGraphNode.cpp
@@ -189,6 +189,10 @@ void SControlRigGraphNode::Construct( const FArguments& InArgs )
 				PinInfo.bHideInputWidget = !bIsStructEditable;
 				PinInfo.bAutoHeight = bIsStructEditable;
 			}
+			else if(PinInfo.bHasChildren)
+			{
+				PinInfo.bHideInputWidget = true;
+			}
 		}
 		
 		if(URigVMPin* ParentPin = ModelPin->GetParentPin())
