Sequencer: Control Rig Anim Layer Improvements, pt1.
Fixes issue with selections when a control rig's section has channels masked out, it could select the wrong control, added CategoryIndex to fix that. Also improved selection/deselction speed by avoiding reselecting Controls when they are selected. That was done to make sure Python selections worked on the first time, but what the only that that was needed there was enabling the Control Rig Edit Mode, so we just do that.

diff --git a/Engine/Plugins/Animation/ControlRig/Source/ControlRig/Private/Sequencer/MovieSceneControlRigParameterSection.cpp b/Engine/Plugins/Animation/ControlRig/Source/ControlRig/Private/Sequencer/MovieSceneControlRigParameterSection.cpp
--- a/Engine/Plugins/Animation/ControlRig/Source/ControlRig/Private/Sequencer/MovieSceneControlRigParameterSection.cpp
+++ b/Engine/Plugins/Animation/ControlRig/Source/ControlRig/Private/Sequencer/MovieSceneControlRigParameterSection.cpp
@@ -1219,6 +1219,7 @@ void UMovieSceneControlRigParameterSection::ReconstructChannelProxy()
 		int32 EnumChannelIndex = 0;
 		int32 IntegerChannelIndex = 0;
 		int32 SpaceChannelIndex = 0;
+		int32 CategoryIndex = 0;
 		const FName BoolChannelTypeName = FMovieSceneBoolChannel::StaticStruct()->GetFName();
 		const FName EnumChannelTypeName = FMovieSceneByteChannel::StaticStruct()->GetFName();
 		const FName IntegerChannelTypeName = FMovieSceneIntegerChannel::StaticStruct()->GetFName();
@@ -1237,8 +1238,8 @@ void UMovieSceneControlRigParameterSection::ReconstructChannelProxy()
 			{
 				ParentControlName = ParentControlElement->GetName();
 			}
+			bool bEnabled = ControlsMask[MaskIndex];
 #if WITH_EDITOR
-
 			switch (ControlElement->Settings.ControlType)
 			{
 				case ERigControlType::Float:
@@ -1247,7 +1248,6 @@ void UMovieSceneControlRigParameterSection::ReconstructChannelProxy()
 					{
 						if (ControlElement->GetName() == Scalar.ParameterName)
 						{
-							bool bEnabled =ControlsMask[MaskIndex];
 
 							FText Group;
 							if(ParentControlElement)
@@ -1269,20 +1269,24 @@ void UMovieSceneControlRigParameterSection::ReconstructChannelProxy()
 
 							if (Group.IsEmpty())
 							{
-								ControlChannelMap.Add(Scalar.ParameterName, FChannelMapInfo(ControlIndex, TotalIndex, FloatChannelIndex));
+								ControlChannelMap.Add(Scalar.ParameterName, FChannelMapInfo(ControlIndex, TotalIndex, FloatChannelIndex, INDEX_NONE, NAME_None, MaskIndex, CategoryIndex));
 								Group = FText::FromName(ControlElement->GetDisplayName());
 								ControlIndex++;  //up the index only if no parent
+								if (bEnabled)
+								{
+									++CategoryIndex;
+								}
 							}
 							else
 							{
 								FChannelMapInfo* pChannelIndex = ControlChannelMap.Find(ParentControlName);
 								if (pChannelIndex)
 								{
-									ControlChannelMap.Add(Scalar.ParameterName, FChannelMapInfo(ControlIndex, TotalIndex, FloatChannelIndex,pChannelIndex->ControlIndex));
+									ControlChannelMap.Add(Scalar.ParameterName, FChannelMapInfo(ControlIndex, TotalIndex, FloatChannelIndex,pChannelIndex->ControlIndex, NAME_None, MaskIndex, CategoryIndex));
 								}
 								else
 								{
-									ControlChannelMap.Add(Scalar.ParameterName, FChannelMapInfo(ControlIndex, TotalIndex, FloatChannelIndex));
+									ControlChannelMap.Add(Scalar.ParameterName, FChannelMapInfo(ControlIndex, TotalIndex, FloatChannelIndex, INDEX_NONE, NAME_None, MaskIndex, CategoryIndex));
 
 								}
 							}
@@ -1303,8 +1307,6 @@ void UMovieSceneControlRigParameterSection::ReconstructChannelProxy()
 					{
 						if (ControlElement->GetName() == Bool.ParameterName)
 						{
-							bool bEnabled =ControlsMask[MaskIndex];
-
 							FText Group;
 							if(ParentControlElement)
 							{
@@ -1325,20 +1327,24 @@ void UMovieSceneControlRigParameterSection::ReconstructChannelProxy()
 
 							if (Group.IsEmpty())
 							{
-								ControlChannelMap.Add(Bool.ParameterName, FChannelMapInfo(ControlIndex, TotalIndex, BoolChannelIndex, INDEX_NONE,BoolChannelTypeName));
+								ControlChannelMap.Add(Bool.ParameterName, FChannelMapInfo(ControlIndex, TotalIndex, BoolChannelIndex, INDEX_NONE,BoolChannelTypeName,MaskIndex, CategoryIndex));
 								Group = FText::FromName(ControlElement->GetDisplayName());
 								ControlIndex++;  //up the index only if no parent
+								if (bEnabled)
+								{
+									++CategoryIndex;
+								}
 							}
 							else
 							{
 								FChannelMapInfo* pChannelIndex = ControlChannelMap.Find(ParentControlName);
 								if (pChannelIndex)
 								{
-									ControlChannelMap.Add(Bool.ParameterName, FChannelMapInfo(ControlIndex, TotalIndex, BoolChannelIndex,pChannelIndex->ControlIndex,BoolChannelTypeName));
+									ControlChannelMap.Add(Bool.ParameterName, FChannelMapInfo(ControlIndex, TotalIndex, BoolChannelIndex,pChannelIndex->ControlIndex,BoolChannelTypeName,MaskIndex, CategoryIndex));
 								}
 								else
 								{
-									ControlChannelMap.Add(Bool.ParameterName, FChannelMapInfo(ControlIndex, TotalIndex, BoolChannelIndex,INDEX_NONE, BoolChannelTypeName));
+									ControlChannelMap.Add(Bool.ParameterName, FChannelMapInfo(ControlIndex, TotalIndex, BoolChannelIndex,INDEX_NONE, BoolChannelTypeName,MaskIndex, CategoryIndex));
 								}
 							}
 
@@ -1362,8 +1368,6 @@ void UMovieSceneControlRigParameterSection::ReconstructChannelProxy()
 						{
 							if (ControlElement->GetName() == Enum.ParameterName)
 							{
-								bool bEnabled =ControlsMask[MaskIndex];
-
 								FText Group;
 								if (ParentControlElement)
 								{
@@ -1384,20 +1388,24 @@ void UMovieSceneControlRigParameterSection::ReconstructChannelProxy()
 
 								if (Group.IsEmpty())
 								{
-									ControlChannelMap.Add(Enum.ParameterName, FChannelMapInfo(ControlIndex, TotalIndex, EnumChannelIndex,INDEX_NONE, EnumChannelTypeName));
+									ControlChannelMap.Add(Enum.ParameterName, FChannelMapInfo(ControlIndex, TotalIndex, EnumChannelIndex,INDEX_NONE, EnumChannelTypeName,MaskIndex, CategoryIndex));
 									Group = FText::FromName(ControlElement->GetDisplayName());
 									ControlIndex++;  //up the index only if no parent
+									if (bEnabled)
+									{
+										++CategoryIndex;
+									}
 								}
 								else
 								{
 									FChannelMapInfo* pChannelIndex = ControlChannelMap.Find(ParentControlName);
 									if (pChannelIndex)
 									{
-										ControlChannelMap.Add(Enum.ParameterName, FChannelMapInfo(ControlIndex, TotalIndex, EnumChannelIndex, pChannelIndex->ControlIndex, EnumChannelTypeName));
+										ControlChannelMap.Add(Enum.ParameterName, FChannelMapInfo(ControlIndex, TotalIndex, EnumChannelIndex, pChannelIndex->ControlIndex, EnumChannelTypeName,MaskIndex, CategoryIndex));
 									}
 									else
 									{
-										ControlChannelMap.Add(Enum.ParameterName, FChannelMapInfo(ControlIndex, TotalIndex, EnumChannelIndex, INDEX_NONE, EnumChannelTypeName));
+										ControlChannelMap.Add(Enum.ParameterName, FChannelMapInfo(ControlIndex, TotalIndex, EnumChannelIndex, INDEX_NONE, EnumChannelTypeName,MaskIndex, CategoryIndex));
 									}
 								}
 
@@ -1418,8 +1426,6 @@ void UMovieSceneControlRigParameterSection::ReconstructChannelProxy()
 						{
 							if (ControlElement->GetName() == Integer.ParameterName)
 							{
-								bool bEnabled =ControlsMask[MaskIndex];
-
 								FText Group;
 								if (ParentControlElement)
 								{
@@ -1439,20 +1445,24 @@ void UMovieSceneControlRigParameterSection::ReconstructChannelProxy()
 								}
 								if (Group.IsEmpty())
 								{
-									ControlChannelMap.Add(Integer.ParameterName, FChannelMapInfo(ControlIndex, TotalIndex, IntegerChannelIndex,INDEX_NONE,IntegerChannelTypeName));
+									ControlChannelMap.Add(Integer.ParameterName, FChannelMapInfo(ControlIndex, TotalIndex, IntegerChannelIndex,INDEX_NONE,IntegerChannelTypeName,MaskIndex, CategoryIndex));
 									Group = FText::FromName(ControlElement->GetDisplayName());
 									ControlIndex++;  //up the index only if no parent
+									if (bEnabled)
+									{
+										++CategoryIndex;
+									}
 								}
 								else
 								{
 									FChannelMapInfo* pChannelIndex = ControlChannelMap.Find(ParentControlName);
 									if (pChannelIndex)
 									{
-										ControlChannelMap.Add(Integer.ParameterName, FChannelMapInfo(ControlIndex, TotalIndex, IntegerChannelIndex, pChannelIndex->ControlIndex,IntegerChannelTypeName));
+										ControlChannelMap.Add(Integer.ParameterName, FChannelMapInfo(ControlIndex, TotalIndex, IntegerChannelIndex, pChannelIndex->ControlIndex,IntegerChannelTypeName,MaskIndex, CategoryIndex));
 									}
 									else
 									{
-										ControlChannelMap.Add(Integer.ParameterName, FChannelMapInfo(ControlIndex, TotalIndex, IntegerChannelIndex, INDEX_NONE, IntegerChannelTypeName));
+										ControlChannelMap.Add(Integer.ParameterName, FChannelMapInfo(ControlIndex, TotalIndex, IntegerChannelIndex, INDEX_NONE, IntegerChannelTypeName,MaskIndex, CategoryIndex));
 									}
 								}
 
@@ -1476,9 +1486,12 @@ void UMovieSceneControlRigParameterSection::ReconstructChannelProxy()
 					{
 						if (ControlElement->GetName() == Vector2D.ParameterName)
 						{
-							ControlChannelMap.Add(Vector2D.ParameterName, FChannelMapInfo(ControlIndex, TotalIndex, FloatChannelIndex));
+							ControlChannelMap.Add(Vector2D.ParameterName, FChannelMapInfo(ControlIndex, TotalIndex, FloatChannelIndex, INDEX_NONE, NAME_None, MaskIndex, CategoryIndex));
 							ControlIndex++;
-							bool bEnabled =ControlsMask[MaskIndex];
+							if (bEnabled)
+							{
+								++CategoryIndex;
+							}
 							FText Group = FText::FromName(ControlElement->GetDisplayName());
 							FParameterVectorChannelEditorData EditorData(ControlRig, Vector2D.ParameterName, bEnabled, Group, TotalIndex, 2);
 							Channels.Add(Vector2D.XCurve, EditorData.MetaData[0], EditorData.ExternalValues[0]);
@@ -1517,9 +1530,12 @@ void UMovieSceneControlRigParameterSection::ReconstructChannelProxy()
 									}
 								}
 							}
-							ControlChannelMap.Add(Vector.ParameterName, FChannelMapInfo(ControlIndex, TotalIndex, FloatChannelIndex));
+							ControlChannelMap.Add(Vector.ParameterName, FChannelMapInfo(ControlIndex, TotalIndex, FloatChannelIndex, INDEX_NONE, NAME_None, MaskIndex, CategoryIndex));
 							ControlIndex++;
-							bool bEnabled = ControlsMask[MaskIndex];
+							if (bEnabled)
+							{
+								++CategoryIndex;
+							}
 							FText Group = FText::FromName(ControlElement->GetDisplayName());
 							if (FSpaceControlNameAndChannel* SpaceChannel = GetSpaceChannel(Vector.ParameterName))
 							{
@@ -1564,9 +1580,12 @@ void UMovieSceneControlRigParameterSection::ReconstructChannelProxy()
 					{
 						if (ControlElement->GetName() == Transform.ParameterName)
 						{
-							ControlChannelMap.Add(Transform.ParameterName, FChannelMapInfo(ControlIndex, TotalIndex, FloatChannelIndex));
+							ControlChannelMap.Add(Transform.ParameterName, FChannelMapInfo(ControlIndex, TotalIndex, FloatChannelIndex, INDEX_NONE, NAME_None, MaskIndex, CategoryIndex));
 							ControlIndex++;
-							bool bEnabled = ControlsMask[MaskIndex];
+							if (bEnabled)
+							{
+								++CategoryIndex;
+							}
 							FText Group = FText::FromName(ControlElement->GetDisplayName());
 
 							if (FSpaceControlNameAndChannel* SpaceChannel = GetSpaceChannel(Transform.ParameterName))
@@ -1833,6 +1852,7 @@ void UMovieSceneControlRigParameterSection::ReconstructChannelProxy()
 			}
 #endif
 			++MaskIndex;
+
 		}
 		
 #if WITH_EDITOR
@@ -2591,6 +2611,18 @@ void UMovieSceneControlRigParameterSection::ClearAllParameters()
 	IntegerParameterNamesAndCurves.SetNum(0);
 }
 
+int32 UMovieSceneControlRigParameterSection::GetActiveCategoryIndex(FName ControlName) const
+{
+	int32 CategoryIndex = INDEX_NONE;
+	FChannelMapInfo* pChannelIndex = ControlChannelMap.Find(ControlName);
+	if (pChannelIndex != nullptr && ControlsMask[pChannelIndex->MaskIndex])
+	{
+		CategoryIndex = pChannelIndex->CategoryIndex;
+	}
+	return CategoryIndex;
+}
+
+
 TOptional<float> UMovieSceneControlRigParameterSection::EvaluateScalarParameter(const  FFrameTime& InTime, FName InParameterName)
 {
 	TOptional<float> OptValue;	
@@ -2758,3 +2790,4 @@ FIntegerParameterNameAndCurve::FIntegerParameterNameAndCurve(FName InParameterNa
 }
 
 #undef LOCTEXT_NAMESPACE 
+
diff --git a/Engine/Plugins/Animation/ControlRig/Source/ControlRig/Public/Sequencer/MovieSceneControlRigParameterSection.h b/Engine/Plugins/Animation/ControlRig/Source/ControlRig/Public/Sequencer/MovieSceneControlRigParameterSection.h
--- a/Engine/Plugins/Animation/ControlRig/Source/ControlRig/Public/Sequencer/MovieSceneControlRigParameterSection.h
+++ b/Engine/Plugins/Animation/ControlRig/Source/ControlRig/Public/Sequencer/MovieSceneControlRigParameterSection.h
@@ -137,10 +137,11 @@ struct CONTROLRIG_API FChannelMapInfo
 
 	FChannelMapInfo() = default;
 
-	FChannelMapInfo(int32 InControlIndex, int32 InTotalChannelIndex,  int32 InChannelIndex, int32 InParentControlIndex = INDEX_NONE, FName InChannelTypeName = NAME_None) :
-		ControlIndex(InControlIndex),TotalChannelIndex(InTotalChannelIndex), ChannelIndex(InChannelIndex), ParentControlIndex(InParentControlIndex), ChannelTypeName(InChannelTypeName) {};
+	FChannelMapInfo(int32 InControlIndex, int32 InTotalChannelIndex, int32 InChannelIndex, int32 InParentControlIndex = INDEX_NONE, FName InChannelTypeName = NAME_None,
+		int32 InMaskIndex = INDEX_NONE, int32 InCategoryIndex = INDEX_NONE) :
+		ControlIndex(InControlIndex),TotalChannelIndex(InTotalChannelIndex), ChannelIndex(InChannelIndex), ParentControlIndex(InParentControlIndex), ChannelTypeName(InChannelTypeName),MaskIndex(InMaskIndex),CategoryIndex(InCategoryIndex) {};
 	UPROPERTY()
-	int32 ControlIndex = 0;
+	int32 ControlIndex = 0; 
 	UPROPERTY()
 	int32 TotalChannelIndex = 0;
 	UPROPERTY()
@@ -153,7 +154,11 @@ struct CONTROLRIG_API FChannelMapInfo
 	bool bDoesHaveSpace = false;
 	UPROPERTY()
 	int32 SpaceChannelIndex = -1; //if it has space what's the space channel index
-	
+	UPROPERTY()
+	int32 MaskIndex = -1; //index for the mask
+	UPROPERTY()
+	int32 CategoryIndex = -1; //index for the Sequencer category node
+
 	int32 GeneratedKeyIndex = -1; //temp index set by the ControlRigParameterTrack, not saved
 
 
@@ -288,7 +293,11 @@ public:
 		ControlsMask.Init(Val, ControlsMask.Num());
 		ReconstructChannelProxy();
 	}
-
+	/**
+	* This function returns the active category index of the control, based upon what controls are active/masked or not
+	* If itself is masked it returns INDEX_NONE
+	*/
+	int32 GetActiveCategoryIndex(FName ControlName) const;
 	/**
 	* Access the transform mask that defines which channels this track should animate
 	*/
diff --git a/Engine/Plugins/Animation/ControlRig/Source/ControlRigEditor/Private/Sequencer/ControlRigParameterTrackEditor.cpp b/Engine/Plugins/Animation/ControlRig/Source/ControlRigEditor/Private/Sequencer/ControlRigParameterTrackEditor.cpp
--- a/Engine/Plugins/Animation/ControlRig/Source/ControlRigEditor/Private/Sequencer/ControlRigParameterTrackEditor.cpp
+++ b/Engine/Plugins/Animation/ControlRig/Source/ControlRigEditor/Private/Sequencer/ControlRigParameterTrackEditor.cpp
@@ -2049,7 +2049,11 @@ void FControlRigParameterTrackEditor::SelectSequencerNodeInSection(UMovieSceneCo
 		{
 			if (pChannelIndex->ParentControlIndex == INDEX_NONE)
 			{
-				GetSequencer()->SelectByNthCategoryNode(ParamSection, pChannelIndex->ControlIndex, bSelected);
+				int32 CategoryIndex = ParamSection->GetActiveCategoryIndex(ControlName);
+				if (CategoryIndex != INDEX_NONE)
+				{
+					GetSequencer()->SelectByNthCategoryNode(ParamSection, CategoryIndex, bSelected);
+			}
 			}
 			else
 			{
@@ -2216,13 +2220,48 @@ void FControlRigParameterTrackEditor::HandleSpaceKeyMoved(UMovieSceneControlRigP
 	}
 }
 
+void FControlRigParameterTrackEditor::SetUpEditModeIfNeeded(UControlRig* ControlRig)
+{
+
+	FControlRigEditMode* ControlRigEditMode = GetEditMode();
+	if (!ControlRigEditMode)
+	{
+		ControlRigEditMode = GetEditMode(true);
+		if (TSharedPtr<IControlRigObjectBinding> ObjectBinding = ControlRig->GetObjectBinding())
+		{
+			if (ControlRigEditMode)
+			{
+				ControlRigEditMode->SetObjects(ControlRig, nullptr, GetSequencer());
+			}
+		}
+	}
+	else
+	{
+		if (ControlRigEditMode->GetControlRig(false) != ControlRig)
+		{
+			if (ControlRigEditMode->GetControlRig(false))
+			{
+				ControlRigEditMode->GetControlRig(false)->ClearControlSelection();
+			}
+			ControlRigEditMode->SetObjects(ControlRig, nullptr, GetSequencer());
+			//force an evaluation, this will get the control rig setup so edit mode looks good.
+			if (GetSequencer().IsValid())
+			{
+				GetSequencer()->NotifyMovieSceneDataChanged(EMovieSceneDataChangeType::Unknown);
+			}
+		}
+	}
+}
+
 void FControlRigParameterTrackEditor::HandleControlSelected(UControlRig* Subject, FRigControlElement* ControlElement, bool bSelected)
 {
 	//if parent selected we select child here if it's a bool,integer or single float
 	TArray<FRigControl> Controls;
 
 	URigHierarchy* Hierarchy = Subject->GetHierarchy();
-	if(URigHierarchyController* Controller = Hierarchy->GetController())
+	if (bSelected)
+	{
+	if (URigHierarchyController* Controller = Hierarchy->GetController())
 	{
 		Hierarchy->ForEach<FRigControlElement>([ControlElement, Controller, bSelected](FRigControlElement* OtherControlElement) -> bool
 		{
@@ -2244,6 +2283,7 @@ void FControlRigParameterTrackEditor::HandleControlSelected(UControlRig* Subject
 		});
 	}
 
+	}
 	if (bIsDoingSelection)
 	{
 		return;
@@ -2255,7 +2295,7 @@ void FControlRigParameterTrackEditor::HandleControlSelected(UControlRig* Subject
 	FName ControlRigName(*Subject->GetName());
 	if (TSharedPtr<IControlRigObjectBinding> ObjectBinding = Subject->GetObjectBinding())
 	{
-		UObject *ActorObject = nullptr;
+		UObject* ActorObject = nullptr;
 		USceneComponent* Component = Cast<USceneComponent>(ObjectBinding->GetBoundObject());
 		if (!Component)
 		{
@@ -2269,23 +2309,20 @@ void FControlRigParameterTrackEditor::HandleControlSelected(UControlRig* Subject
 		{
 			return;
 		}
-		
+
 		FFindOrCreateTrackResult TrackResult = FindOrCreateControlRigTrackForObject(ObjectHandle, Subject, ControlRigName, bCreateTrack);
 		UMovieSceneControlRigParameterTrack* Track = CastChecked<UMovieSceneControlRigParameterTrack>(TrackResult.Track, ECastCheckedType::NullAllowed);
 		if (Track)
 		{
 			GetSequencer()->SuspendSelectionBroadcast();
-			for (UMovieSceneSection* Section : Track->GetAllSections())
-			{
+			//Just set in the section to key not all
+			UMovieSceneSection* Section = Track->GetSectionToKey();
 				UMovieSceneControlRigParameterSection* ParamSection = Cast<UMovieSceneControlRigParameterSection>(Section);
-				SelectSequencerNodeInSection(ParamSection,ControlElement->GetName(), bSelected);	
-			}
+				SelectSequencerNodeInSection(ParamSection, ControlElement->GetName(), bSelected);
+			
 			GetSequencer()->ResumeSelectionBroadcast();
 
-			TArray<const IKeyArea*> KeyAreas;
-			const bool UseSelectedKeys = CVarSelectedKeysSelectControls.GetValueOnGameThread();
-			GetSequencer()->GetSelectedKeyAreas(KeyAreas, UseSelectedKeys);
-			SelectRigsAndControls(Subject, KeyAreas);
+			SetUpEditModeIfNeeded(Subject);
 
 			//Force refresh later, not now
 			GetSequencer()->NotifyMovieSceneDataChanged(EMovieSceneDataChangeType::RefreshTree);				
@@ -2314,7 +2351,7 @@ void FControlRigParameterTrackEditor::HandleControlModified(UControlRig* Control
 	{
 		return;
 	}
-	FTransform  Transform = ControlRig->GetControlLocalTransform(ControlElement->GetName());	
+	FTransform  Transform = ControlRig->GetControlLocalTransform(ControlElement->GetName());
 	UMovieScene* MovieScene = GetSequencer()->GetFocusedMovieSceneSequence()->GetMovieScene();
 	const TArray<FMovieSceneBinding>& Bindings = MovieScene->GetBindings();
 	for (const FMovieSceneBinding& Binding : Bindings)
@@ -2341,7 +2378,7 @@ void FControlRigParameterTrackEditor::HandleControlModified(UControlRig* Control
 	}
 }
 
-void FControlRigParameterTrackEditor::GetControlRigKeys(UControlRig* InControlRig, FName ParameterName, EControlRigContextChannelToKey ChannelsToKey,  ESequencerKeyMode KeyMode, UMovieSceneControlRigParameterSection* SectionToKey,FGeneratedTrackKeys& OutGeneratedKeys)
+void FControlRigParameterTrackEditor::GetControlRigKeys(UControlRig* InControlRig, FName ParameterName, EControlRigContextChannelToKey ChannelsToKey, ESequencerKeyMode KeyMode, UMovieSceneControlRigParameterSection* SectionToKey, FGeneratedTrackKeys& OutGeneratedKeys)
 {
 	const TArray<bool>& ControlsMask = SectionToKey->GetControlsMask();
 	EMovieSceneTransformChannel TransformMask = SectionToKey->GetTransformMask().GetChannels();
@@ -2374,7 +2411,7 @@ void FControlRigParameterTrackEditor::GetControlRigKeys(UControlRig* InControlRi
 			bool bSetKey = ControlElement->GetName() == ParameterName && !bMaskKeyOut;
 
 			FRigControlValue ControlValue = InControlRig->GetHierarchy()->GetControlValue(ControlElement, ERigControlValueType::Current);
-		
+
 			switch (ControlElement->Settings.ControlType)
 			{
 			case ERigControlType::Bool:
@@ -2584,7 +2621,7 @@ void FControlRigParameterTrackEditor::GetControlRigKeys(UControlRig* InControlRi
 	}
 }
 
-FKeyPropertyResult FControlRigParameterTrackEditor::AddKeysToControlRigHandle(USceneComponent *InSceneComp, UControlRig* InControlRig,
+FKeyPropertyResult FControlRigParameterTrackEditor::AddKeysToControlRigHandle(USceneComponent* InSceneComp, UControlRig* InControlRig,
 	FGuid ObjectHandle, FFrameNumber KeyTime, FGeneratedTrackKeys& GeneratedKeys,
 	ESequencerKeyMode KeyMode, TSubclassOf<UMovieSceneTrack> TrackClass, FName ControlRigName, FName RigControlName)
 {
@@ -2660,7 +2697,7 @@ FKeyPropertyResult FControlRigParameterTrackEditor::AddKeysToControlRigHandle(US
 }
 
 FKeyPropertyResult FControlRigParameterTrackEditor::AddKeysToControlRig(
-	USceneComponent *InSceneComp, UControlRig* InControlRig, FFrameNumber KeyTime,  FGeneratedTrackKeys& GeneratedKeys,
+	USceneComponent* InSceneComp, UControlRig* InControlRig, FFrameNumber KeyTime, FGeneratedTrackKeys& GeneratedKeys,
 	ESequencerKeyMode KeyMode, TSubclassOf<UMovieSceneTrack> TrackClass, FName ControlRigName, FName RigControlName)
 {
 	FKeyPropertyResult KeyPropertyResult;
@@ -2683,7 +2720,7 @@ FKeyPropertyResult FControlRigParameterTrackEditor::AddKeysToControlRig(
 	return KeyPropertyResult;
 }
 
-void FControlRigParameterTrackEditor::AddControlKeys(USceneComponent *InSceneComp, UControlRig* InControlRig, FName ControlRigName, FName RigControlName, EControlRigContextChannelToKey ChannelsToKey, ESequencerKeyMode KeyMode, float InLocalTime)
+void FControlRigParameterTrackEditor::AddControlKeys(USceneComponent* InSceneComp, UControlRig* InControlRig, FName ControlRigName, FName RigControlName, EControlRigContextChannelToKey ChannelsToKey, ESequencerKeyMode KeyMode, float InLocalTime)
 {
 	if (KeyMode == ESequencerKeyMode::ManualKey || (GetSequencer().IsValid() && !GetSequencer()->IsAllowedToChange()))
 	{
@@ -2739,7 +2776,7 @@ void FControlRigParameterTrackEditor::AddControlKeys(USceneComponent *InSceneCom
 }
 
 
-bool FControlRigParameterTrackEditor::ModifyOurGeneratedKeysByCurrentAndWeight(UObject *Object, UControlRig* InControlRig, FName RigControlName, UMovieSceneTrack *Track, UMovieSceneSection* SectionToKey, FFrameNumber KeyTime, FGeneratedTrackKeys& GeneratedTotalKeys, float Weight) const
+bool FControlRigParameterTrackEditor::ModifyOurGeneratedKeysByCurrentAndWeight(UObject* Object, UControlRig* InControlRig, FName RigControlName, UMovieSceneTrack* Track, UMovieSceneSection* SectionToKey, FFrameNumber KeyTime, FGeneratedTrackKeys& GeneratedTotalKeys, float Weight) const
 {
 	FFrameRate TickResolution = GetSequencer()->GetFocusedTickResolution();
 	FMovieSceneEvaluationTrack EvalTrack = CastChecked<UMovieSceneControlRigParameterTrack>(Track)->GenerateTrackTemplate(Track);
@@ -2857,20 +2894,20 @@ bool FControlRigParameterTrackEditor::ModifyOurGeneratedKeysByCurrentAndWeight(U
 						
 						FVector3f CurrentPos = Val.Val.GetTranslation();
 						FRotator3f CurrentRot = FRotator3f(Val.Val.GetRotation().Rotator());
-						GeneratedTotalKeys[ChannelIndex]->ModifyByCurrentAndWeight(Proxy, KeyTime, (void *)&CurrentPos.X, Weight);
-						GeneratedTotalKeys[ChannelIndex + 1]->ModifyByCurrentAndWeight(Proxy, KeyTime, (void *)&CurrentPos.Y, Weight);
-						GeneratedTotalKeys[ChannelIndex + 2]->ModifyByCurrentAndWeight(Proxy, KeyTime, (void *)&CurrentPos.Z, Weight);
+						GeneratedTotalKeys[ChannelIndex]->ModifyByCurrentAndWeight(Proxy, KeyTime, (void*)&CurrentPos.X, Weight);
+						GeneratedTotalKeys[ChannelIndex + 1]->ModifyByCurrentAndWeight(Proxy, KeyTime, (void*)&CurrentPos.Y, Weight);
+						GeneratedTotalKeys[ChannelIndex + 2]->ModifyByCurrentAndWeight(Proxy, KeyTime, (void*)&CurrentPos.Z, Weight);
 
-						GeneratedTotalKeys[ChannelIndex + 3]->ModifyByCurrentAndWeight(Proxy, KeyTime, (void *)&CurrentRot.Roll, Weight);
-						GeneratedTotalKeys[ChannelIndex + 4]->ModifyByCurrentAndWeight(Proxy, KeyTime, (void *)&CurrentRot.Pitch, Weight);
-						GeneratedTotalKeys[ChannelIndex + 5]->ModifyByCurrentAndWeight(Proxy, KeyTime, (void *)&CurrentRot.Yaw, Weight);
+						GeneratedTotalKeys[ChannelIndex + 3]->ModifyByCurrentAndWeight(Proxy, KeyTime, (void*)&CurrentRot.Roll, Weight);
+						GeneratedTotalKeys[ChannelIndex + 4]->ModifyByCurrentAndWeight(Proxy, KeyTime, (void*)&CurrentRot.Pitch, Weight);
+						GeneratedTotalKeys[ChannelIndex + 5]->ModifyByCurrentAndWeight(Proxy, KeyTime, (void*)&CurrentRot.Yaw, Weight);
 
 						if (ControlElement->Settings.ControlType == ERigControlType::Transform || ControlElement->Settings.ControlType == ERigControlType::EulerTransform)
 						{
 							FVector3f CurrentScale = Val.Val.GetScale3D();
-							GeneratedTotalKeys[ChannelIndex + 6]->ModifyByCurrentAndWeight(Proxy, KeyTime, (void *)&CurrentScale.X, Weight);
-							GeneratedTotalKeys[ChannelIndex + 7]->ModifyByCurrentAndWeight(Proxy, KeyTime, (void *)&CurrentScale.Y, Weight);
-							GeneratedTotalKeys[ChannelIndex + 8]->ModifyByCurrentAndWeight(Proxy, KeyTime, (void *)&CurrentScale.Z, Weight);
+							GeneratedTotalKeys[ChannelIndex + 6]->ModifyByCurrentAndWeight(Proxy, KeyTime, (void*)&CurrentScale.X, Weight);
+							GeneratedTotalKeys[ChannelIndex + 7]->ModifyByCurrentAndWeight(Proxy, KeyTime, (void*)&CurrentScale.Y, Weight);
+							GeneratedTotalKeys[ChannelIndex + 8]->ModifyByCurrentAndWeight(Proxy, KeyTime, (void*)&CurrentScale.Z, Weight);
 						}
 					}
 					break;
@@ -2954,7 +2991,7 @@ bool FControlRigParameterTrackEditor::HandleAssetAdded(UObject* Asset, const FGu
 	{
 		return false;
 	}
-	
+
 	UControlRigBlueprint* ControlRigBlueprint = Cast<UControlRigBlueprint>(Asset);
 	UControlRigBlueprintGeneratedClass* RigClass = ControlRigBlueprint->GetControlRigBlueprintGeneratedClass();
 	if (!RigClass)
@@ -3255,9 +3292,9 @@ private:
 			{
 				FFKBoneCheckInfo& Info = Pair.Value;
 				BoneCheckArray[Index++] = Info;
-			
+
 			}
-			if (Track  && Sequencer)
+			if (Track && Sequencer)
 			{
 				TArray<bool> Mask;
 				Mask.SetNum(BoneCheckArray.Num());
@@ -3367,8 +3404,8 @@ void FControlRigParameterSection::BuildSectionContextMenu(FMenuBuilder& MenuBuil
 				}
 			}
 		}
-	    TArray<FRigControlElement*> Controls;
-	    ControlRig->GetControlsInOrder(Controls);
+		TArray<FRigControlElement*> Controls;
+		ControlRig->GetControlsInOrder(Controls);
 
 		auto MakeUIAction = [=](EMovieSceneTransformChannel ChannelsToToggle)
 		{
@@ -3737,7 +3774,7 @@ FEditorModeTools* FControlRigParameterTrackEditor::GetEditorModeTools() const
 	if (SharedSequencer.IsValid())
 	{
 		TSharedPtr<IToolkitHost> ToolkitHost = SharedSequencer->GetToolkitHost();
-		if(ToolkitHost.IsValid())
+		if (ToolkitHost.IsValid())
 		{
 			return &ToolkitHost->GetEditorModeManager();
 		}
@@ -3750,11 +3787,11 @@ FControlRigEditMode* FControlRigParameterTrackEditor::GetEditMode(bool bForceAct
 {
 	if (FEditorModeTools* EditorModetools = GetEditorModeTools())
 	{
-		if(bForceActivate && !EditorModetools->IsModeActive(FControlRigEditMode::ModeName))
+		if (bForceActivate && !EditorModetools->IsModeActive(FControlRigEditMode::ModeName))
 		{
 			EditorModetools->ActivateMode(FControlRigEditMode::ModeName);
 		}
-		
+
 		return static_cast<FControlRigEditMode*>(EditorModetools->GetActiveMode(FControlRigEditMode::ModeName));
 	}
 
diff --git a/Engine/Plugins/Animation/ControlRig/Source/ControlRigEditor/Private/Sequencer/ControlRigParameterTrackEditor.h b/Engine/Plugins/Animation/ControlRig/Source/ControlRigEditor/Private/Sequencer/ControlRigParameterTrackEditor.h
--- a/Engine/Plugins/Animation/ControlRig/Source/ControlRigEditor/Private/Sequencer/ControlRigParameterTrackEditor.h
+++ b/Engine/Plugins/Animation/ControlRig/Source/ControlRigEditor/Private/Sequencer/ControlRigParameterTrackEditor.h
@@ -164,6 +164,9 @@ private:
 	/** Bake To Control Rig Sub Menu*/
 	void BakeToControlRig(UClass* InClass, FGuid ObjectBinding,UObject* BoundObject, USkeletalMeshComponent* SkelMeshComp, USkeleton* Skeleton);
 
+	/** Set Up EditMode for Specified Control Rig*/
+	void SetUpEditModeIfNeeded(UControlRig* ControlRig);
+
 private:
 
 	/** Command Bindings added by the Transform Track Editor to Sequencer and curve editor. */
