CIS Fix
https://horde.devtools.epicgames.com/job/61afd845ee6e47a827ef2be3?step=4aff

diff --git a/Engine/Plugins/Animation/ControlRig/Source/ControlRig/Private/Sequencer/MovieSceneControlRigParameterSection.cpp b/Engine/Plugins/Animation/ControlRig/Source/ControlRig/Private/Sequencer/MovieSceneControlRigParameterSection.cpp
--- a/Engine/Plugins/Animation/ControlRig/Source/ControlRig/Private/Sequencer/MovieSceneControlRigParameterSection.cpp
+++ b/Engine/Plugins/Animation/ControlRig/Source/ControlRig/Private/Sequencer/MovieSceneControlRigParameterSection.cpp
@@ -2614,7 +2614,7 @@ void UMovieSceneControlRigParameterSection::ClearAllParameters()
 int32 UMovieSceneControlRigParameterSection::GetActiveCategoryIndex(FName ControlName) const
 {
 	int32 CategoryIndex = INDEX_NONE;
-	FChannelMapInfo* pChannelIndex = ControlChannelMap.Find(ControlName);
+	const FChannelMapInfo* pChannelIndex = ControlChannelMap.Find(ControlName);
 	if (pChannelIndex != nullptr && ControlsMask[pChannelIndex->MaskIndex])
 	{
 		CategoryIndex = pChannelIndex->CategoryIndex;
