Control Rig: Disable default value validation on load, which was causing array pins to lose their subpins. sometimes a serialized default value can become invalid when the pin's struct type's cpp declaration was changed. However we don't want to completely reject the default value since it still holds valid values for pins that did not change

diff --git a/Engine/Plugins/Animation/ControlRig/Source/ControlRigDeveloper/Private/ControlRigBlueprint.cpp b/Engine/Plugins/Animation/ControlRig/Source/ControlRigDeveloper/Private/ControlRigBlueprint.cpp
--- a/Engine/Plugins/Animation/ControlRig/Source/ControlRigDeveloper/Private/ControlRigBlueprint.cpp
+++ b/Engine/Plugins/Animation/ControlRig/Source/ControlRigDeveloper/Private/ControlRigBlueprint.cpp
@@ -460,6 +460,8 @@ void UControlRigBlueprint::PostLoad()
 			for (URigVMGraph* GraphToDetach : GraphsToDetach)
 			{
 				URigVMController* Controller = GetOrCreateController(GraphToDetach);
+				// temporarily disable default value validation during load time, serialized values should always be accepted
+				TGuardValue<bool> PerGraphDisablePinDefaultValueValidation(Controller->bValidatePinDefaults, false);
 				Controller->DetachLinksFromPinObjects();
 				TArray<URigVMNode*> Nodes = GraphToDetach->GetNodes();
 				for (URigVMNode* Node : Nodes)
