Control Rig: Allow copy/paste of control shape properties

diff --git a/Engine/Plugins/Animation/ControlRig/Source/ControlRigEditor/Private/ControlRigElementDetails.cpp b/Engine/Plugins/Animation/ControlRig/Source/ControlRigEditor/Private/ControlRigElementDetails.cpp
--- a/Engine/Plugins/Animation/ControlRig/Source/ControlRigEditor/Private/ControlRigElementDetails.cpp
+++ b/Engine/Plugins/Animation/ControlRig/Source/ControlRigEditor/Private/ControlRigElementDetails.cpp
@@ -24,6 +24,7 @@
 #include "RigVMModel/RigVMGraph.h"
 #include "RigVMModel/RigVMNode.h"
 #include "Graph/SControlRigGraphPinVariableBinding.h"
+#include "HAL/PlatformApplicationMisc.h"
 
 #define LOCTEXT_NAMESPACE "ControlRigElementDetails"
 
@@ -1911,34 +1912,50 @@ void FRigControlElementDetails::CustomizeDetails(IDetailLayoutBuilder& DetailBui
 		ShapeCategory.AddProperty(SettingsHandle->GetChildHandle(TEXT("bShapeVisible")).ToSharedRef())
 		.DisplayName(FText::FromString(TEXT("Visible")));
 
-		// setup shape transform
-		{
-			const TSharedPtr<IPropertyHandle> ShapeHandle = DetailBuilder.GetProperty(TEXT("Shape"));
-			const TSharedPtr<IPropertyHandle> InitialHandle = ShapeHandle->GetChildHandle(TEXT("Initial"));
-			const TSharedPtr<IPropertyHandle> LocalHandle = InitialHandle->GetChildHandle(TEXT("Local"));
-			const TSharedPtr<IPropertyHandle> TransformHandle = LocalHandle->GetChildHandle(TEXT("Transform"));
-			ShapeCategory.AddProperty(TransformHandle.ToSharedRef())
-			.IsEnabled(TAttribute<bool>::CreateSP(this, &FRigControlElementDetails::IsShapeEnabled));
-		}
-
-		const TSharedPtr<IPropertyHandle> ShapeShapeHandle = SettingsHandle->GetChildHandle(TEXT("ShapeName"));
-		ShapeCategory.AddProperty(ShapeShapeHandle.ToSharedRef()).CustomWidget()
-		.NameContent()
+		IDetailGroup& ShapeProperitesGroup = ShapeCategory.AddGroup(TEXT("Shape Properties"), LOCTEXT("ShapeProperties", "Shape Properties"));
+		ShapeProperitesGroup.HeaderRow().NameContent()
 		[
 			SNew(STextBlock)
-			.Text(FText::FromString(TEXT("Shape")))
 			.Font(IDetailLayoutBuilder::GetDetailFont())
-			.IsEnabled(this, &FRigControlElementDetails::IsShapeEnabled)
+			.Text(LOCTEXT("ShapeProperties", "Shape Properties"))
+			.ToolTipText(LOCTEXT("ShapeProperties", "Customize the properties of the shape"))
 		]
-		.ValueContent()
-		[
-			SNew(SControlRigShapeNameList, ControlElements, BlueprintBeingCustomized)
-			.OnGetNameListContent(this, &FRigControlElementDetails::GetShapeNameList)
-			.IsEnabled(this, &FRigControlElementDetails::IsShapeEnabled)
-		];
+		.CopyAction(FUIAction(
+			FExecuteAction::CreateSP(this, &FRigControlElementDetails::OnCopyShapeProperties)))
+		.PasteAction(FUIAction(
+			FExecuteAction::CreateSP(this, &FRigControlElementDetails::OnPasteShapeProperties)));
+		
+		{
+			// setup shape transform
+			{
+				TSharedPtr<IPropertyHandle> ShapeHandle = DetailBuilder.GetProperty(TEXT("Shape"));
+				TSharedPtr<IPropertyHandle> InitialHandle = ShapeHandle->GetChildHandle(TEXT("Initial"));
+				TSharedPtr<IPropertyHandle> LocalHandle = InitialHandle->GetChildHandle(TEXT("Local"));
+				ShapeTransformHandle = LocalHandle->GetChildHandle(TEXT("Transform"));
+				ShapeProperitesGroup.AddPropertyRow(ShapeTransformHandle.ToSharedRef())
+				.IsEnabled(TAttribute<bool>::CreateSP(this, &FRigControlElementDetails::IsShapeEnabled));
+			}
+
+			ShapeNameHandle = SettingsHandle->GetChildHandle(TEXT("ShapeName"));
+			ShapeProperitesGroup.AddPropertyRow(ShapeNameHandle.ToSharedRef()).CustomWidget()
+			.NameContent()
+			[
+				SNew(STextBlock)
+				.Text(FText::FromString(TEXT("Shape")))
+				.Font(IDetailLayoutBuilder::GetDetailFont())
+				.IsEnabled(this, &FRigControlElementDetails::IsShapeEnabled)
+			]
+			.ValueContent()
+			[
+				SNew(SControlRigShapeNameList, ControlElements, BlueprintBeingCustomized)
+				.OnGetNameListContent(this, &FRigControlElementDetails::GetShapeNameList)
+				.IsEnabled(this, &FRigControlElementDetails::IsShapeEnabled)
+			];
 
-		ShapeCategory.AddProperty(SettingsHandle->GetChildHandle(TEXT("ShapeColor")).ToSharedRef())
-		.DisplayName(FText::FromString(TEXT("Color")));
+			ShapeColorHandle = SettingsHandle->GetChildHandle(TEXT("ShapeColor"));
+			ShapeProperitesGroup.AddPropertyRow(ShapeColorHandle.ToSharedRef())
+			.DisplayName(FText::FromString(TEXT("Color")));
+		}
 	}
 
 	if(IsAnyControlOfType(ERigControlType::Integer))
@@ -2139,6 +2156,92 @@ void FRigControlElementDetails::SetDisplayName(const FText& InNewText, ETextComm
 	}
 }
 
+void FRigControlElementDetails::OnCopyShapeProperties()
+{
+	FString Value;
+	if (ObjectsBeingCustomized.Num() > 0)
+	{
+		if(ObjectsBeingCustomized[0].IsValid())
+		{
+			if(ObjectsBeingCustomized[0]->IsChildOf<FRigControlElement>())
+			{
+				const FRigControlElement ControlElement = ObjectsBeingCustomized[0]->GetContent<FRigControlElement>();
+
+				Value = FString::Printf(TEXT("(ShapeName=\"%s\",ShapeColor=%s,Transform=%s)"),
+					*ControlElement.Settings.ShapeName.ToString(),
+					*ControlElement.Settings.ShapeColor.ToString(),
+					*ControlElement.Shape.Initial.Local.Transform.ToString());
+			}
+		}
+	}
+		
+	if (!Value.IsEmpty())
+	{
+		// Copy.
+		FPlatformApplicationMisc::ClipboardCopy(*Value);
+	}
+}
+
+void FRigControlElementDetails::OnPasteShapeProperties()
+{
+	FString PastedText;
+	FPlatformApplicationMisc::ClipboardPaste(PastedText);
+	
+	FString TrimmedText = PastedText.LeftChop(1).RightChop(1);
+	FString ShapeName;
+	FString ShapeColorStr;
+	FString TransformStr;
+	bool bSuccessful = FParse::Value(*TrimmedText, TEXT("ShapeName="), ShapeName) &&
+					   FParse::Value(*TrimmedText, TEXT("ShapeColor="), ShapeColorStr, false) &&
+					   FParse::Value(*TrimmedText, TEXT("Transform="), TransformStr, false);
+
+	if (bSuccessful)
+	{
+		FScopedTransaction Transaction(LOCTEXT("PasteShape", "Paste Shape"));
+		
+		// Name
+		{
+			ShapeNameHandle->NotifyPreChange();
+			ShapeNameHandle->SetValue(ShapeName);
+			ShapeNameHandle->NotifyPostChange(EPropertyChangeType::ValueSet);
+		}
+		
+		// Color
+		{
+			ShapeColorHandle->NotifyPreChange();
+			TArray<void*> RawDataPtrs;
+			ShapeColorHandle->AccessRawData(RawDataPtrs);
+			for (void* RawPtr: RawDataPtrs)
+			{
+				bSuccessful &= static_cast<FLinearColor*>(RawPtr)->InitFromString(ShapeColorStr);
+				if (!bSuccessful)
+				{
+					Transaction.Cancel();
+					return;
+				}
+			}		
+			ShapeColorHandle->NotifyPostChange(EPropertyChangeType::ValueSet);
+		}
+
+		// Transform
+		{
+			ShapeTransformHandle->NotifyPreChange();
+			TArray<void*> RawDataPtrs;
+			ShapeTransformHandle->AccessRawData(RawDataPtrs);
+			for (void* RawPtr: RawDataPtrs)
+			{
+				bSuccessful &= static_cast<FTransform*>(RawPtr)->InitFromString(TransformStr);
+				if (!bSuccessful)
+				{
+					Transaction.Cancel();
+					return;
+				}
+			}		
+			ShapeTransformHandle->NotifyPostChange(EPropertyChangeType::ValueSet);
+		}		
+	}
+}
+
 void FRigNullElementDetails::CustomizeDetails(IDetailLayoutBuilder& DetailBuilder)
 {
 	FRigTransformElementDetails::CustomizeDetails(DetailBuilder);
diff --git a/Engine/Plugins/Animation/ControlRig/Source/ControlRigEditor/Private/ControlRigElementDetails.h b/Engine/Plugins/Animation/ControlRig/Source/ControlRigEditor/Private/ControlRigElementDetails.h
--- a/Engine/Plugins/Animation/ControlRig/Source/ControlRigEditor/Private/ControlRigElementDetails.h
+++ b/Engine/Plugins/Animation/ControlRig/Source/ControlRigEditor/Private/ControlRigElementDetails.h
@@ -245,11 +245,18 @@ public:
 	FText GetDisplayName() const;
 	void SetDisplayName(const FText& InNewText, ETextCommit::Type InCommitType);
 
+	void OnCopyShapeProperties();
+	void OnPasteShapeProperties();
+
 private:
 	TArray<TSharedPtr<FString>> ShapeNameList;
 	static TArray<TSharedPtr<FString>> ControlTypeList;
 	TSharedPtr<FRigInfluenceEntryModifier> InfluenceModifier;
 	TSharedPtr<FStructOnScope> InfluenceModifierStruct;
+
+	TSharedPtr<IPropertyHandle> ShapeNameHandle;
+	TSharedPtr<IPropertyHandle> ShapeColorHandle;
+	TSharedPtr<IPropertyHandle> ShapeTransformHandle;
 };
 
 class FRigNullElementDetails : public FRigTransformElementDetails
