Allow the interpolation type to be passed an argument of UMovieSceneControlRigParameterSection::RecordControlRigKey

diff --git a/Engine/Plugins/Animation/ControlRig/Source/ControlRig/Private/Sequencer/MovieSceneControlRigParameterSection.cpp b/Engine/Plugins/Animation/ControlRig/Source/ControlRig/Private/Sequencer/MovieSceneControlRigParameterSection.cpp
--- a/Engine/Plugins/Animation/ControlRig/Source/ControlRig/Private/Sequencer/MovieSceneControlRigParameterSection.cpp
+++ b/Engine/Plugins/Animation/ControlRig/Source/ControlRig/Private/Sequencer/MovieSceneControlRigParameterSection.cpp
@@ -2120,7 +2120,7 @@ void UMovieSceneControlRigParameterSection::FixRotationWinding(FName ControlName
 
 #if WITH_EDITOR
 
-void UMovieSceneControlRigParameterSection::RecordControlRigKey(FFrameNumber FrameNumber, bool bSetDefault, bool bAutoKey)
+void UMovieSceneControlRigParameterSection::RecordControlRigKey(FFrameNumber FrameNumber, bool bSetDefault, ERichCurveInterpMode InInterpMode)
 {
 	if (ControlRig)
 	{
@@ -2129,6 +2129,32 @@ void UMovieSceneControlRigParameterSection::RecordControlRigKey(FFrameNumber Fra
 		TArrayView<FMovieSceneIntegerChannel*> IntChannels = GetChannelProxy().GetChannels<FMovieSceneIntegerChannel>();
 		TArrayView<FMovieSceneByteChannel*> EnumChannels = GetChannelProxy().GetChannels<FMovieSceneByteChannel>();
 
+		// Helper lambda to add a FVector key to the float channels
+		auto AddVectorKeyToFloatChannels = [&FloatChannels, InInterpMode](int32& ChannelIndex, FFrameNumber FrameNumber, const FVector& Value)
+		{
+			switch (InInterpMode)
+			{
+			case RCIM_Linear:
+				FloatChannels[ChannelIndex++]->AddLinearKey(FrameNumber, Value.X);
+				FloatChannels[ChannelIndex++]->AddLinearKey(FrameNumber, Value.Y);
+				FloatChannels[ChannelIndex++]->AddLinearKey(FrameNumber, Value.Z);
+				break;
+
+			case RCIM_Constant:
+				FloatChannels[ChannelIndex++]->AddConstantKey(FrameNumber, Value.X);
+				FloatChannels[ChannelIndex++]->AddConstantKey(FrameNumber, Value.Y);
+				FloatChannels[ChannelIndex++]->AddConstantKey(FrameNumber, Value.Z);
+				break;
+
+			case RCIM_Cubic:
+			default:
+				FloatChannels[ChannelIndex++]->AddCubicKey(FrameNumber, Value.X, ERichCurveTangentMode::RCTM_Auto);
+				FloatChannels[ChannelIndex++]->AddCubicKey(FrameNumber, Value.Y, ERichCurveTangentMode::RCTM_Auto);
+				FloatChannels[ChannelIndex++]->AddCubicKey(FrameNumber, Value.Z, ERichCurveTangentMode::RCTM_Auto);
+				break;
+			}
+		};
+
 		TArray<FRigControlElement*> Controls;
 		ControlRig->GetControlsInOrder(Controls);
 		
@@ -2187,14 +2213,23 @@ void UMovieSceneControlRigParameterSection::RecordControlRigKey(FFrameNumber Fra
 					{
 						FloatChannels[ChannelIndex]->SetDefault(Val);
 					}
-					if (bAutoKey)
-					{
-						FloatChannels[ChannelIndex++]->AddCubicKey(FrameNumber, Val, ERichCurveTangentMode::RCTM_Auto);
-					}
-					else
+
+					switch (InInterpMode)
 					{
+					case RCIM_Linear:
 						FloatChannels[ChannelIndex++]->AddLinearKey(FrameNumber, Val);
+						break;
+
+					case RCIM_Constant:
+						FloatChannels[ChannelIndex++]->AddConstantKey(FrameNumber, Val);
+						break;
+
+					case RCIM_Cubic:
+					default:
+						FloatChannels[ChannelIndex++]->AddCubicKey(FrameNumber, Val, ERichCurveTangentMode::RCTM_Auto);
+						break;
 					}
+
 					break;
 				}
 				case ERigControlType::Vector2D:
@@ -2205,15 +2240,24 @@ void UMovieSceneControlRigParameterSection::RecordControlRigKey(FFrameNumber Fra
 						FloatChannels[ChannelIndex]->SetDefault(Val.X);
 						FloatChannels[ChannelIndex + 1]->SetDefault(Val.Y);
 					}
-					if (bAutoKey)
-					{
-						FloatChannels[ChannelIndex++]->AddCubicKey(FrameNumber, Val.X, ERichCurveTangentMode::RCTM_Auto);
-						FloatChannels[ChannelIndex++]->AddCubicKey(FrameNumber, Val.Y, ERichCurveTangentMode::RCTM_Auto);
-					}
-					else
+
+					switch (InInterpMode)
 					{
+					case RCIM_Linear:
 						FloatChannels[ChannelIndex++]->AddLinearKey(FrameNumber, Val.X);
 						FloatChannels[ChannelIndex++]->AddLinearKey(FrameNumber, Val.Y);
+						break;
+
+					case RCIM_Constant:
+						FloatChannels[ChannelIndex++]->AddConstantKey(FrameNumber, Val.X);
+						FloatChannels[ChannelIndex++]->AddConstantKey(FrameNumber, Val.Y);
+						break;
+
+					case RCIM_Cubic:
+					default:
+						FloatChannels[ChannelIndex++]->AddCubicKey(FrameNumber, Val.X, ERichCurveTangentMode::RCTM_Auto);
+						FloatChannels[ChannelIndex++]->AddCubicKey(FrameNumber, Val.Y, ERichCurveTangentMode::RCTM_Auto);
+						break;
 					}
 
 					break;
@@ -2239,18 +2283,9 @@ void UMovieSceneControlRigParameterSection::RecordControlRigKey(FFrameNumber Fra
 						FloatChannels[ChannelIndex + 1]->SetDefault(Val.Y);
 						FloatChannels[ChannelIndex + 2]->SetDefault(Val.Z);
 					}
-					if (bAutoKey)
-					{
-						FloatChannels[ChannelIndex++]->AddCubicKey(FrameNumber, Val.X, ERichCurveTangentMode::RCTM_Auto);
-						FloatChannels[ChannelIndex++]->AddCubicKey(FrameNumber, Val.Y, ERichCurveTangentMode::RCTM_Auto);
-						FloatChannels[ChannelIndex++]->AddCubicKey(FrameNumber, Val.Z, ERichCurveTangentMode::RCTM_Auto);
-					}
-					else		
-					{
-						FloatChannels[ChannelIndex++]->AddLinearKey(FrameNumber, Val.X);
-						FloatChannels[ChannelIndex++]->AddLinearKey(FrameNumber, Val.Y);
-						FloatChannels[ChannelIndex++]->AddLinearKey(FrameNumber, Val.Z);
-					}
+
+					AddVectorKeyToFloatChannels(ChannelIndex, FrameNumber, Val);
+
 					break;
 				}
 
@@ -2286,18 +2321,9 @@ void UMovieSceneControlRigParameterSection::RecordControlRigKey(FFrameNumber Fra
 						FloatChannels[ChannelIndex + 1]->SetDefault(CurrentVector.Y);
 						FloatChannels[ChannelIndex + 2]->SetDefault(CurrentVector.Z);
 					}
-					if(bAutoKey)
-					{
-						FloatChannels[ChannelIndex++]->AddCubicKey(FrameNumber, CurrentVector.X, ERichCurveTangentMode::RCTM_Auto);
-						FloatChannels[ChannelIndex++]->AddCubicKey(FrameNumber, CurrentVector.Y, ERichCurveTangentMode::RCTM_Auto);
-						FloatChannels[ChannelIndex++]->AddCubicKey(FrameNumber, CurrentVector.Z, ERichCurveTangentMode::RCTM_Auto);
-					}
-					else
-					{
-						FloatChannels[ChannelIndex++]->AddLinearKey(FrameNumber, CurrentVector.X);
-						FloatChannels[ChannelIndex++]->AddLinearKey(FrameNumber, CurrentVector.Y);
-						FloatChannels[ChannelIndex++]->AddLinearKey(FrameNumber, CurrentVector.Z);
-					}
+
+					AddVectorKeyToFloatChannels(ChannelIndex, FrameNumber, CurrentVector);
+
 					CurrentVector = Val.GetRotation().Euler();
 					if (FloatChannels[ChannelIndex]->GetNumKeys() > 0)
 					{
@@ -2314,18 +2340,8 @@ void UMovieSceneControlRigParameterSection::RecordControlRigKey(FFrameNumber Fra
 						FloatChannels[ChannelIndex + 1]->SetDefault(CurrentVector.Y);
 						FloatChannels[ChannelIndex + 2]->SetDefault(CurrentVector.Z);
 					}
-					if (bAutoKey)
-					{
-						FloatChannels[ChannelIndex++]->AddCubicKey(FrameNumber, CurrentVector.X, ERichCurveTangentMode::RCTM_Auto);
-						FloatChannels[ChannelIndex++]->AddCubicKey(FrameNumber, CurrentVector.Y, ERichCurveTangentMode::RCTM_Auto);
-						FloatChannels[ChannelIndex++]->AddCubicKey(FrameNumber, CurrentVector.Z, ERichCurveTangentMode::RCTM_Auto);
-					}
-					else
-					{
-						FloatChannels[ChannelIndex++]->AddLinearKey(FrameNumber, CurrentVector.X);
-						FloatChannels[ChannelIndex++]->AddLinearKey(FrameNumber, CurrentVector.Y);
-						FloatChannels[ChannelIndex++]->AddLinearKey(FrameNumber, CurrentVector.Z);
-					}
+
+					AddVectorKeyToFloatChannels(ChannelIndex, FrameNumber, CurrentVector);
 
 					if (ControlElement->Settings.ControlType == ERigControlType::Transform ||
 						ControlElement->Settings.ControlType == ERigControlType::EulerTransform)
@@ -2337,19 +2353,8 @@ void UMovieSceneControlRigParameterSection::RecordControlRigKey(FFrameNumber Fra
 							FloatChannels[ChannelIndex + 1]->SetDefault(CurrentVector.Y);
 							FloatChannels[ChannelIndex + 2]->SetDefault(CurrentVector.Z);
 						}
-						if (bAutoKey)
-						{
-							FloatChannels[ChannelIndex++]->AddCubicKey(FrameNumber, CurrentVector.X, ERichCurveTangentMode::RCTM_Auto);
-							FloatChannels[ChannelIndex++]->AddCubicKey(FrameNumber, CurrentVector.Y, ERichCurveTangentMode::RCTM_Auto);
-							FloatChannels[ChannelIndex++]->AddCubicKey(FrameNumber, CurrentVector.Z, ERichCurveTangentMode::RCTM_Auto);
-						}
-						else
-						{
-							FloatChannels[ChannelIndex++]->AddLinearKey(FrameNumber, CurrentVector.X);
-							FloatChannels[ChannelIndex++]->AddLinearKey(FrameNumber, CurrentVector.Y);
-							FloatChannels[ChannelIndex++]->AddLinearKey(FrameNumber, CurrentVector.Z);
-						}
 
+						AddVectorKeyToFloatChannels(ChannelIndex, FrameNumber, CurrentVector);
 					}
 					break;
 				}
@@ -2465,7 +2470,8 @@ bool UMovieSceneControlRigParameterSection::LoadAnimSequenceIntoThisSection(UAni
 		}
 		ControlRig->Execute(EControlRigState::Update, FRigUnit_InverseExecution::EventName);
 
-		RecordControlRigKey(FrameNumber, true, bKeyReduce);
+		const ERichCurveInterpMode InterpMode = bKeyReduce ? RCIM_Cubic : RCIM_Linear;
+		RecordControlRigKey(FrameNumber, true, InterpMode);
 		Progress.EnterProgressFrame(1);
 		if (Progress.ShouldCancel())
 		{
diff --git a/Engine/Plugins/Animation/ControlRig/Source/ControlRig/Public/Sequencer/MovieSceneControlRigParameterSection.h b/Engine/Plugins/Animation/ControlRig/Source/ControlRig/Public/Sequencer/MovieSceneControlRigParameterSection.h
--- a/Engine/Plugins/Animation/ControlRig/Source/ControlRig/Public/Sequencer/MovieSceneControlRigParameterSection.h
+++ b/Engine/Plugins/Animation/ControlRig/Source/ControlRig/Public/Sequencer/MovieSceneControlRigParameterSection.h
@@ -252,7 +252,7 @@ public:
 
 #if WITH_EDITOR
 	//Function to save control rig key when recording.
-	void RecordControlRigKey(FFrameNumber FrameNumber, bool bSetDefault, bool bDoAutoKey);
+	void RecordControlRigKey(FFrameNumber FrameNumber, bool bSetDefault, ERichCurveInterpMode InInterpMode);
 
 	//Function to load an Anim Sequence into this section. It will automatically resize to the section size.
 	//Will return false if fails or is canceled
