Sequencer: Control Rig: Support Re-Baking over a track with a Control Rig and crash with Pose Library open on exit due to edit mode being de-activated but not calling Exit which does it cleanup to well later.

diff --git a/Engine/Plugins/Animation/ControlRig/Source/ControlRigEditor/Private/EditMode/ControlRigEditMode.cpp b/Engine/Plugins/Animation/ControlRig/Source/ControlRigEditor/Private/EditMode/ControlRigEditMode.cpp
--- a/Engine/Plugins/Animation/ControlRig/Source/ControlRigEditor/Private/EditMode/ControlRigEditMode.cpp
+++ b/Engine/Plugins/Animation/ControlRig/Source/ControlRigEditor/Private/EditMode/ControlRigEditMode.cpp
@@ -227,7 +227,7 @@ bool FControlRigEditMode::IsInLevelEditor() const
 }
 void FControlRigEditMode::SetUpDetailPanel()
 {
-	if (IsInLevelEditor())
+	if (IsInLevelEditor() && Toolkit)
 	{
 		StaticCastSharedPtr<SControlRigEditModeTools>(Toolkit->GetInlineContent())->SetSequencer(WeakSequencer.Pin());
 #if USE_LOCAL_DETAILS
diff --git a/Engine/Plugins/Animation/ControlRig/Source/ControlRigEditor/Private/Sequencer/ControlRigParameterTrackEditor.cpp b/Engine/Plugins/Animation/ControlRig/Source/ControlRigEditor/Private/Sequencer/ControlRigParameterTrackEditor.cpp
--- a/Engine/Plugins/Animation/ControlRig/Source/ControlRigEditor/Private/Sequencer/ControlRigParameterTrackEditor.cpp
+++ b/Engine/Plugins/Animation/ControlRig/Source/ControlRigEditor/Private/Sequencer/ControlRigParameterTrackEditor.cpp
@@ -461,6 +461,7 @@ void FControlRigParameterTrackEditor::OnRelease()
 	FControlRigEditMode* ControlRigEditMode = GetEditMode();
 	if (ControlRigEditMode)
 	{
+		ControlRigEditMode->Exit(); //deactive mode below doesn't exit for some reason so need to make sure things are cleaned up
 		if (FEditorModeTools* Tools = GetEditorModeTools())
 		{
 			Tools->DeactivateMode(FControlRigEditMode::ModeName);
@@ -935,16 +936,18 @@ void FControlRigParameterTrackEditor::BakeToControlRig(UClass* InClass, FGuid Ob
 				}
 
 				const FScopedTransaction Transaction(LOCTEXT("BakeToControlRig_Transaction", "Bake To Control Rig"));
-
+				
+				bool bReuseControlRig = false; //if same Class just re-use it, and put into a new section
 				OwnerMovieScene->Modify();
 				UMovieSceneControlRigParameterTrack* Track = OwnerMovieScene->FindTrack<UMovieSceneControlRigParameterTrack>(ObjectBinding);
 				if (Track)
 				{
-					Track->Modify();
-					for (UMovieSceneSection* Section : Track->GetAllSections())
+					if (Track->GetControlRig() && Track->GetControlRig()->GetClass() == InClass)
 					{
-						Section->SetIsActive(false);
+						bReuseControlRig = true;
 					}
+					Track->Modify();
+					Track->RemoveAllAnimationData();//removes all sections and sectiontokey
 				}
 				else
 				{
@@ -955,13 +958,12 @@ void FControlRigParameterTrackEditor::BakeToControlRig(UClass* InClass, FGuid Ob
 					}
 				}
 
-
 				if (Track)
 				{
 
 					FString ObjectName = InClass->GetName();
 					ObjectName.RemoveFromEnd(TEXT("_C"));
-					UControlRig* ControlRig = NewObject<UControlRig>(Track, InClass, FName(*ObjectName), RF_Transactional);
+					UControlRig* ControlRig = bReuseControlRig ? Track->GetControlRig() : NewObject<UControlRig>(Track, InClass, FName(*ObjectName), RF_Transactional);
 					if (InClass != UFKControlRig::StaticClass() && !ControlRig->SupportsEvent(FRigUnit_InverseExecution::EventName))
 					{
 						TempAnimSequence->MarkAsGarbage();
@@ -984,17 +986,19 @@ void FControlRigParameterTrackEditor::BakeToControlRig(UClass* InClass, FGuid Ob
 						}
 					}
 
-					bool bSequencerOwnsControlRig = true;
-
-					ControlRig->Modify();
-					ControlRig->SetObjectBinding(MakeShared<FControlRigObjectBinding>());
-					ControlRig->GetObjectBinding()->BindToObject(BoundActor);
-					ControlRig->GetDataSourceRegistry()->RegisterDataSource(UControlRig::OwnerComponent, ControlRig->GetObjectBinding()->GetBoundObject());
-					ControlRig->Initialize();
-					ControlRig->RequestInit();
-					ControlRig->SetBoneInitialTransformsFromSkeletalMeshComponent(SkelMeshComp, true);
-					ControlRig->Evaluate_AnyThread();
+					if (bReuseControlRig == false)
+					{
+						ControlRig->Modify();
+						ControlRig->SetObjectBinding(MakeShared<FControlRigObjectBinding>());
+						ControlRig->GetObjectBinding()->BindToObject(BoundActor);
+						ControlRig->GetDataSourceRegistry()->RegisterDataSource(UControlRig::OwnerComponent, ControlRig->GetObjectBinding()->GetBoundObject());
+						ControlRig->Initialize();
+						ControlRig->RequestInit();
+						ControlRig->SetBoneInitialTransformsFromSkeletalMeshComponent(SkelMeshComp, true);
+						ControlRig->Evaluate_AnyThread();
+					}
 
+					const bool bSequencerOwnsControlRig = true;
 					UMovieSceneSection* NewSection = Track->CreateControlRigSection(0, ControlRig, bSequencerOwnsControlRig);
 					UMovieSceneControlRigParameterSection* ParamSection = Cast<UMovieSceneControlRigParameterSection>(NewSection);
 
