Control Rig: fix crash on function's local variable deletion undo

diff --git a/Engine/Plugins/Animation/ControlRig/Source/ControlRigEditor/Private/Graph/SControlRigGraphNode.cpp b/Engine/Plugins/Animation/ControlRig/Source/ControlRigEditor/Private/Graph/SControlRigGraphNode.cpp
--- a/Engine/Plugins/Animation/ControlRig/Source/ControlRigEditor/Private/Graph/SControlRigGraphNode.cpp
+++ b/Engine/Plugins/Animation/ControlRig/Source/ControlRigEditor/Private/Graph/SControlRigGraphNode.cpp
@@ -860,37 +860,39 @@ void SControlRigGraphNode::GetOverlayBrushes(bool bSelected, const FVector2D Wid
 {
 	UControlRigGraphNode* RigGraphNode = Cast<UControlRigGraphNode>(GraphNode);
 
-	const URigVMNode* VMNode = RigGraphNode->GetModelNode();
-	const bool bHasBreakpoint = VMNode->HasBreakpoint();
-	if (bHasBreakpoint)
+	if (const URigVMNode* VMNode = RigGraphNode->GetModelNode())
 	{
-		FOverlayBrushInfo BreakpointOverlayInfo;
-
-		BreakpointOverlayInfo.Brush = FEditorStyle::GetBrush(TEXT("Kismet.DebuggerOverlay.Breakpoint.EnabledAndValid"));
-		if (BreakpointOverlayInfo.Brush != NULL)
+		const bool bHasBreakpoint = VMNode->HasBreakpoint();
+		if (bHasBreakpoint)
 		{
-			BreakpointOverlayInfo.OverlayOffset -= BreakpointOverlayInfo.Brush->ImageSize / 2.f;
-		}
+			FOverlayBrushInfo BreakpointOverlayInfo;
 
-		Brushes.Add(BreakpointOverlayInfo);
-	}
+			BreakpointOverlayInfo.Brush = FEditorStyle::GetBrush(TEXT("Kismet.DebuggerOverlay.Breakpoint.EnabledAndValid"));
+			if (BreakpointOverlayInfo.Brush != NULL)
+			{
+				BreakpointOverlayInfo.OverlayOffset -= BreakpointOverlayInfo.Brush->ImageSize / 2.f;
+			}
 
-	// Paint red arrow pointing at breakpoint node that caused a halt in execution
-	{
-		FOverlayBrushInfo IPOverlayInfo;
-		if (VMNode->ExecutionIsHaltedAtThisNode())
+			Brushes.Add(BreakpointOverlayInfo);
+		}
+
+		// Paint red arrow pointing at breakpoint node that caused a halt in execution
 		{
-			IPOverlayInfo.Brush = FEditorStyle::GetBrush( TEXT("Kismet.DebuggerOverlay.InstructionPointerBreakpoint") );
-			if (IPOverlayInfo.Brush != NULL)
+			FOverlayBrushInfo IPOverlayInfo;
+			if (VMNode->ExecutionIsHaltedAtThisNode())
 			{
-				float Overlap = 10.f;
-				IPOverlayInfo.OverlayOffset.X = (WidgetSize.X/2.f) - (IPOverlayInfo.Brush->ImageSize.X/2.f);
-				IPOverlayInfo.OverlayOffset.Y = (Overlap - IPOverlayInfo.Brush->ImageSize.Y);
-			}
+				IPOverlayInfo.Brush = FEditorStyle::GetBrush( TEXT("Kismet.DebuggerOverlay.InstructionPointerBreakpoint") );
+				if (IPOverlayInfo.Brush != NULL)
+				{
+					float Overlap = 10.f;
+					IPOverlayInfo.OverlayOffset.X = (WidgetSize.X/2.f) - (IPOverlayInfo.Brush->ImageSize.X/2.f);
+					IPOverlayInfo.OverlayOffset.Y = (Overlap - IPOverlayInfo.Brush->ImageSize.Y);
+				}
 
-			IPOverlayInfo.AnimationEnvelope = FVector2D(0.f, 10.f);
+				IPOverlayInfo.AnimationEnvelope = FVector2D(0.f, 10.f);
 
-			Brushes.Add(IPOverlayInfo);
+				Brushes.Add(IPOverlayInfo);
+			}
 		}
 	}
 }
diff --git a/Engine/Source/Developer/RigVMDeveloper/Private/RigVMModel/RigVMController.cpp b/Engine/Source/Developer/RigVMDeveloper/Private/RigVMModel/RigVMController.cpp
--- a/Engine/Source/Developer/RigVMDeveloper/Private/RigVMModel/RigVMController.cpp
+++ b/Engine/Source/Developer/RigVMDeveloper/Private/RigVMModel/RigVMController.cpp
@@ -8847,11 +8847,11 @@ bool URigVMController::RemoveLocalVariable(const FName& InVariableName, bool bSe
 	if (FoundIndex != INDEX_NONE)
 	{
 		FRigVMControllerCompileBracketScope CompileScope(this);
-		FRigVMInverseAction InverseAction;
+		FRigVMBaseAction BaseAction;
 		if (bSetupUndoRedo)
 		{
-			InverseAction.Title = FString::Printf(TEXT("Remove Local Variable %s"), *InVariableName.ToString());
-			ActionStack->BeginAction(InverseAction);			
+			BaseAction.Title = FString::Printf(TEXT("Remove Local Variable %s"), *InVariableName.ToString());
+			ActionStack->BeginAction(BaseAction);			
 		}	
 		
 		const FString VarNameStr = InVariableName.ToString();
@@ -8932,13 +8932,13 @@ bool URigVMController::RemoveLocalVariable(const FName& InVariableName, bool bSe
 
 		if (bSetupUndoRedo)
 		{
-			ActionStack->AddAction(FRigVMAddLocalVariableAction(LocalVariables[FoundIndex]));
+			ActionStack->AddAction(FRigVMRemoveLocalVariableAction(LocalVariables[FoundIndex]));
 		}
 		LocalVariables.RemoveAt(FoundIndex);
 
 		if (bSetupUndoRedo)
 		{
-			ActionStack->EndAction(InverseAction);
+			ActionStack->EndAction(BaseAction);
 		}
 
 		if (bPrintPythonCommand)
