Added support for "Use Selected Item" and "Browse to Item" buttons on Control Rig node pins that use the custom bone widget.

diff --git a/Engine/Plugins/Animation/ControlRig/Source/ControlRigEditor/Private/Graph/SControlRigGraphPinNameList.cpp b/Engine/Plugins/Animation/ControlRig/Source/ControlRigEditor/Private/Graph/SControlRigGraphPinNameList.cpp
--- a/Engine/Plugins/Animation/ControlRig/Source/ControlRigEditor/Private/Graph/SControlRigGraphPinNameList.cpp
+++ b/Engine/Plugins/Animation/ControlRig/Source/ControlRigEditor/Private/Graph/SControlRigGraphPinNameList.cpp
@@ -232,6 +232,18 @@ FReply SControlRigGraphPinNameList::OnGetSelectedClicked()
 							CurrentList = GetNameList();
 						}
 					}
+
+					else if (ModelPin->GetCustomWidgetName() == TEXT("BoneName"))
+					{
+						URigHierarchy* Hierarchy = Graph->GetBlueprint()->Hierarchy;
+						TArray<FRigElementKey> Keys = Hierarchy->GetSelectedKeys();
+						FRigBaseElement* Element = Hierarchy->FindChecked(Keys[0]);
+						if (Element->GetType() == ERigElementType::Bone)
+						{
+							Graph->GetController()->SetPinDefaultValue(ModelPin->GetPinPath(), Keys[0].Name.ToString(), true, true, false, true);
+							CurrentList = GetNameList();
+						}
+					}
 					
 					// if we don't have a key pin - this is just a plain name.
 					// let's derive the type of element this node deals with from its name.
@@ -275,22 +287,28 @@ FReply SControlRigGraphPinNameList::OnBrowseClicked()
 		TSharedPtr<FString> Selected = NameListComboBox->GetSelectedItem();
 		if (Selected.IsValid() && ModelPin)
 		{
-			if(URigVMPin* KeyPin = ModelPin->GetParentPin())
+			URigVMPin* KeyPin = ModelPin->GetParentPin();
+			if(KeyPin && KeyPin->GetCPPTypeObject() == FRigElementKey::StaticStruct())
 			{
-				if(KeyPin->GetCPPTypeObject() == FRigElementKey::StaticStruct())
+				// browse to rig element key
+				FString DefaultValue = ModelPin->GetParentPin()->GetDefaultValue();
+				if (!DefaultValue.IsEmpty())
 				{
-					FString DefaultValue = ModelPin->GetParentPin()->GetDefaultValue();
-					if (!DefaultValue.IsEmpty())
+					FRigElementKey Key;
+					FRigElementKey::StaticStruct()->ImportText(*DefaultValue, &Key, nullptr, EPropertyPortFlags::PPF_None, nullptr, FRigElementKey::StaticStruct()->GetName(), true);
+					if (Key.IsValid())
 					{
-						FRigElementKey Key;
-						FRigElementKey::StaticStruct()->ImportText(*DefaultValue, &Key, nullptr, EPropertyPortFlags::PPF_None, nullptr, FRigElementKey::StaticStruct()->GetName(), true);
-						if (Key.IsValid())
-						{
-							Graph->GetBlueprint()->GetHierarchyController()->SetSelection({Key});
-						}
+						Graph->GetBlueprint()->GetHierarchyController()->SetSelection({Key});
 					}
 				}
 			}
+			else if (ModelPin->GetCustomWidgetName() == TEXT("BoneName"))
+			{
+				// browse to named bone
+				const FString DefaultValue = ModelPin->GetDefaultValue();
+				FRigElementKey Key(*DefaultValue, ERigElementType::Bone);
+				Graph->GetBlueprint()->GetHierarchyController()->SetSelection({Key});
+			}
 			else
 			{
 				// if we don't have a key pin - this is just a plain name.
