Eliminate whitespace only differences between Release-Engine-Staging and Release-Engine-Test
#[fyi] Aurel.Cordonnier

diff --git a/Engine/Plugins/Animation/ControlRig/Source/ControlRigEditor/Private/Sequencer/ControlRigParameterTrackEditor.cpp b/Engine/Plugins/Animation/ControlRig/Source/ControlRigEditor/Private/Sequencer/ControlRigParameterTrackEditor.cpp
--- a/Engine/Plugins/Animation/ControlRig/Source/ControlRigEditor/Private/Sequencer/ControlRigParameterTrackEditor.cpp
+++ b/Engine/Plugins/Animation/ControlRig/Source/ControlRigEditor/Private/Sequencer/ControlRigParameterTrackEditor.cpp
@@ -90,7 +90,7 @@ static USkeletalMeshComponent* AcquireSkeletalMeshFromObject(UObject* BoundObjec
 	{
 		TArray<USkeletalMeshComponent*> SkeletalMeshComponents;
 		Actor->GetComponents(SkeletalMeshComponents);
-		
+
 		if (SkeletalMeshComponents.Num() == 1)
 		{
 			return SkeletalMeshComponents[0];
@@ -206,111 +206,111 @@ FControlRigParameterTrackEditor::FControlRigParameterTrackEditor(TSharedRef<ISeq
 		//we check for two things, one if the control rig has been replaced if so we need to switch.
 		//the other is if bound object on the edit mode is null we request a re-evaluate which will reset it up.
 		FDelegateHandle OnObjectsReplacedHandle = FCoreUObjectDelegates::OnObjectsReplaced.AddLambda([&](const TMap<UObject*, UObject*>& ReplacementMap)
-		{
-			if (GetSequencer().IsValid())
 			{
-				TMap<UControlRig*, UControlRig*> OldToNewControlRigs;
-				FControlRigEditMode* ControlRigEditMode = GetEditMode();
-				if (ControlRigEditMode && ControlRigEditMode->GetControlRig(true) && ControlRigEditMode->GetControlRig(true)->GetObjectBinding())
+				if (GetSequencer().IsValid())
 				{
-					if (ControlRigEditMode->GetControlRig(true)->GetObjectBinding()->GetBoundObject() == nullptr)
+					TMap<UControlRig*, UControlRig*> OldToNewControlRigs;
+					FControlRigEditMode* ControlRigEditMode = GetEditMode();
+					if (ControlRigEditMode && ControlRigEditMode->GetControlRig(true) && ControlRigEditMode->GetControlRig(true)->GetObjectBinding())
 					{
-						GetSequencer()->RequestEvaluate();
+						if (ControlRigEditMode->GetControlRig(true)->GetObjectBinding()->GetBoundObject() == nullptr)
+						{
+							GetSequencer()->RequestEvaluate();
+						}
 					}
-				}
-				//Reset Bindings for replaced objects.
-				for (TPair<UObject*, UObject*> ReplacedObject : ReplacementMap)
-				{
-					if (UControlRigComponent* OldControlRigComponent = Cast<UControlRigComponent>(ReplacedObject.Key))
+					//Reset Bindings for replaced objects.
+					for (TPair<UObject*, UObject*> ReplacedObject : ReplacementMap)
 					{
-						UControlRigComponent* NewControlRigComponent = Cast<UControlRigComponent>(ReplacedObject.Value);
-						if (OldControlRigComponent->GetControlRig())
+						if (UControlRigComponent* OldControlRigComponent = Cast<UControlRigComponent>(ReplacedObject.Key))
 						{
-							UControlRig* NewControlRig = nullptr;
-							if(NewControlRigComponent)
+							UControlRigComponent* NewControlRigComponent = Cast<UControlRigComponent>(ReplacedObject.Value);
+							if (OldControlRigComponent->GetControlRig())
 							{
-								NewControlRig = NewControlRigComponent->GetControlRig();
+								UControlRig* NewControlRig = nullptr;
+							if(NewControlRigComponent)
+								{
+									NewControlRig = NewControlRigComponent->GetControlRig();
+								}
+								OldToNewControlRigs.Emplace(OldControlRigComponent->GetControlRig(), NewControlRig);
 							}
-							OldToNewControlRigs.Emplace(OldControlRigComponent->GetControlRig(), NewControlRig);
+						}
+						else if (UControlRig* OldControlRig = Cast<UControlRig>(ReplacedObject.Key))
+						{
+							UControlRig* NewControlRig = Cast<UControlRig>(ReplacedObject.Value);
+							OldToNewControlRigs.Emplace(OldControlRig, NewControlRig);
 						}
 					}
-					else if (UControlRig* OldControlRig = Cast<UControlRig>(ReplacedObject.Key))
-					{
-						UControlRig* NewControlRig = Cast<UControlRig>(ReplacedObject.Value);
-						OldToNewControlRigs.Emplace(OldControlRig, NewControlRig);
-					}
-				}
-				UMovieScene* MovieScene = GetSequencer()->GetFocusedMovieSceneSequence()->GetMovieScene();
-				const TArray<FMovieSceneBinding>& Bindings = MovieScene->GetBindings();
-				for (const FMovieSceneBinding& Binding : Bindings)
-				{
-					UMovieSceneControlRigParameterTrack* Track = Cast<UMovieSceneControlRigParameterTrack>(MovieScene->FindTrack(UMovieSceneControlRigParameterTrack::StaticClass(), Binding.GetObjectGuid(), NAME_None));
-					if (Track && Track->GetControlRig())
+					UMovieScene* MovieScene = GetSequencer()->GetFocusedMovieSceneSequence()->GetMovieScene();
+					const TArray<FMovieSceneBinding>& Bindings = MovieScene->GetBindings();
+					for (const FMovieSceneBinding& Binding : Bindings)
 					{
-						UControlRig* OldControlRig = Track->GetControlRig();
-						UControlRig** NewControlRig = OldToNewControlRigs.Find(OldControlRig);
-						if (NewControlRig)
-						{  
-							TArray<FName> SelectedControls = OldControlRig->CurrentControlSelection();
-							OldControlRig->ClearControlSelection();
-							UnbindControlRig(OldControlRig);
-							if (*NewControlRig)
+						UMovieSceneControlRigParameterTrack* Track = Cast<UMovieSceneControlRigParameterTrack>(MovieScene->FindTrack(UMovieSceneControlRigParameterTrack::StaticClass(), Binding.GetObjectGuid(), NAME_None));
+						if (Track && Track->GetControlRig())
+						{
+							UControlRig* OldControlRig = Track->GetControlRig();
+							UControlRig** NewControlRig = OldToNewControlRigs.Find(OldControlRig);
+							if (NewControlRig)
 							{
-								Track->ReplaceControlRig(*NewControlRig, OldControlRig->GetClass() != (*NewControlRig)->GetClass());
-								BindControlRig(*NewControlRig);
+								TArray<FName> SelectedControls = OldControlRig->CurrentControlSelection();
+								OldControlRig->ClearControlSelection();
+								UnbindControlRig(OldControlRig);
+								if (*NewControlRig)
+								{
+									Track->ReplaceControlRig(*NewControlRig, OldControlRig->GetClass() != (*NewControlRig)->GetClass());
+									BindControlRig(*NewControlRig);
 
-								GetSequencer()->NotifyMovieSceneDataChanged(EMovieSceneDataChangeType::MovieSceneStructureItemsChanged);
-							}
-							else
-							{
-								Track->ReplaceControlRig(nullptr, true);
-							}
-							if (ControlRigEditMode)
-							{
-								if (ControlRigEditMode->GetControlRig(false) == OldControlRig)
+									GetSequencer()->NotifyMovieSceneDataChanged(EMovieSceneDataChangeType::MovieSceneStructureItemsChanged);
+								}
+								else
 								{
-									ControlRigEditMode->SetObjects(*NewControlRig, nullptr, GetSequencer());
+									Track->ReplaceControlRig(nullptr, true);
 								}
-								UControlRig* PtrNewControlRig = *NewControlRig;
-
-								auto UpdateSelectionDelegate = [this, SelectedControls, PtrNewControlRig]()
+								if (ControlRigEditMode)
 								{
+									if (ControlRigEditMode->GetControlRig(false) == OldControlRig)
+									{
+										ControlRigEditMode->SetObjects(*NewControlRig, nullptr, GetSequencer());
+									}
+									UControlRig* PtrNewControlRig = *NewControlRig;
 
-									UE_LOG(LogControlRigEditor, Log, TEXT("UpdateSelectionTimer"));
-									if (!(FSlateApplication::Get().HasAnyMouseCaptor() || GUnrealEd->IsUserInteracting()))
+									auto UpdateSelectionDelegate = [this, SelectedControls, PtrNewControlRig]()
 									{
-										UE_LOG(LogControlRigEditor, Log, TEXT("UpdateSelectionTimer - Sync"));
-										TGuardValue<bool> Guard(bIsDoingSelection, true);
-										GetSequencer()->ExternalSelectionHasChanged();
-										if (PtrNewControlRig)
+
+										UE_LOG(LogControlRigEditor, Log, TEXT("UpdateSelectionTimer"));
+										if (!(FSlateApplication::Get().HasAnyMouseCaptor() || GUnrealEd->IsUserInteracting()))
 										{
-											GEditor->GetTimerManager()->SetTimerForNextTick([SelectedControls,PtrNewControlRig]()
+											UE_LOG(LogControlRigEditor, Log, TEXT("UpdateSelectionTimer - Sync"));
+											TGuardValue<bool> Guard(bIsDoingSelection, true);
+											GetSequencer()->ExternalSelectionHasChanged();
+											if (PtrNewControlRig)
 											{
-												PtrNewControlRig->ClearControlSelection();
-												for (const FName& ControlName : SelectedControls)
+												GEditor->GetTimerManager()->SetTimerForNextTick([SelectedControls,PtrNewControlRig]()
 												{
-													PtrNewControlRig->SelectControl(ControlName, true);
-												}
-											});
-										}
-										if (UpdateSelectionTimerHandle.IsValid())
-										{
-											UE_LOG(LogControlRigEditor, Log, TEXT("UpdateSelectionTimer - Clear"));
-											GEditor->GetTimerManager()->ClearTimer(UpdateSelectionTimerHandle);
+													PtrNewControlRig->ClearControlSelection();
+													for (const FName& ControlName : SelectedControls)
+													{
+														PtrNewControlRig->SelectControl(ControlName, true);
+													}
+												});
+											}
+											if (UpdateSelectionTimerHandle.IsValid())
+											{
+												UE_LOG(LogControlRigEditor, Log, TEXT("UpdateSelectionTimer - Clear"));
+												GEditor->GetTimerManager()->ClearTimer(UpdateSelectionTimerHandle);
+											}
 										}
-									}
 
 
-								};
+									};
 
-								GEditor->GetTimerManager()->SetTimer(UpdateSelectionTimerHandle, UpdateSelectionDelegate, 0.01f, true);
+									GEditor->GetTimerManager()->SetTimer(UpdateSelectionTimerHandle, UpdateSelectionDelegate, 0.01f, true);
+								}
 							}
 						}
 					}
 				}
-			}
 
-		});
+			});
 		AcquiredResources.Add([=] { FCoreUObjectDelegates::OnObjectsReplaced.Remove(OnObjectsReplacedHandle); });
 	}
 	//register all modified/selections for control rigs
@@ -448,7 +448,7 @@ void FControlRigParameterTrackEditor::OnRelease()
 		{
 			GetSequencer()->OnChannelChanged().Remove(OnChannelChangedHandle);
 		}
-		
+
 		if (GetSequencer()->GetFocusedMovieSceneSequence() && GetSequencer()->GetFocusedMovieSceneSequence()->GetMovieScene())
 		{
 			UMovieScene* MovieScene = GetSequencer()->GetFocusedMovieSceneSequence()->GetMovieScene();
@@ -500,7 +500,7 @@ void FControlRigParameterTrackEditor::BuildObjectBindingContextMenu(FMenuBuilder
 		const TSharedPtr<ISequencer> ParentSequencer = GetSequencer();
 		UObject* BoundObject = nullptr;
 		USkeleton* Skeleton = AcquireSkeletonFromObjectGuid(ObjectBindings[0], &BoundObject, ParentSequencer);
-		USkeletalMeshComponent*  SkelMeshComp = AcquireSkeletalMeshFromObject(BoundObject, ParentSequencer);
+		USkeletalMeshComponent* SkelMeshComp = AcquireSkeletalMeshFromObject(BoundObject, ParentSequencer);
 
 		if (Skeleton && SkelMeshComp)
 		{
@@ -540,9 +540,9 @@ void FControlRigParameterTrackEditor::BuildObjectBindingContextMenu(FMenuBuilder
 class FControlRigClassFilter : public IClassViewerFilter
 {
 public:
-	FControlRigClassFilter(bool bInCheckSkeleton, bool bInCheckAnimatable, bool bInCheckInversion, USkeleton* InSkeleton) : 
+	FControlRigClassFilter(bool bInCheckSkeleton, bool bInCheckAnimatable, bool bInCheckInversion, USkeleton* InSkeleton) :
 		bFilterAssetBySkeleton(bInCheckSkeleton),
-		bFilterExposesAnimatableControls(bInCheckAnimatable), 
+		bFilterExposesAnimatableControls(bInCheckAnimatable),
 		bFilterInversion(bInCheckInversion),
 		AssetRegistry(FModuleManager::GetModuleChecked<FAssetRegistryModule>(TEXT("AssetRegistry")).Get())
 	{
@@ -574,7 +574,7 @@ public:
 				FString EventString = FRigUnit_InverseExecution::EventName.ToString();
 				TArray<FString> SupportedEventNames;
 				Tag.GetValue().ParseIntoArray(SupportedEventNames, TEXT(","), true);
-	
+
 				for (const FString& Name : SupportedEventNames)
 				{
 					if (Name.Contains(EventString))
@@ -688,9 +688,9 @@ public:
 	{}
 
 	SLATE_ARGUMENT(UAnimSeqExportOption*, ExportOptions)
-	SLATE_ARGUMENT(UBakeToControlRigSettings*, BakeSettings)
-	SLATE_ARGUMENT(TSharedPtr<SWindow>, WidgetWindow)
-	SLATE_END_ARGS()
+		SLATE_ARGUMENT(UBakeToControlRigSettings*, BakeSettings)
+		SLATE_ARGUMENT(TSharedPtr<SWindow>, WidgetWindow)
+		SLATE_END_ARGS()
 
 public:
 	void Construct(const FArguments& InArgs);
@@ -1045,7 +1045,7 @@ void FControlRigParameterTrackEditor::BakeToControlRig(UClass* InClass, FGuid Ob
 					TempAnimSequence->MarkAsGarbage();
 					AnimSeqExportOption->MarkAsGarbage();
 					GetSequencer()->NotifyMovieSceneDataChanged(EMovieSceneDataChangeType::MovieSceneStructureItemAdded);
-					
+
 
 				}
 			}
@@ -1313,7 +1313,7 @@ void FControlRigParameterTrackEditor::AddControlRig(UClass* InClass, UObject* Bo
 			ControlRig->GetDataSourceRegistry()->RegisterDataSource(UControlRig::OwnerComponent, ControlRig->GetObjectBinding()->GetBoundObject());
 			// Do not re-initialize existing control rig
 			if (!InExistingControlRig)
-			{ 
+			{
 				ControlRig->Initialize();
 			}
 			ControlRig->Evaluate_AnyThread();
@@ -1334,13 +1334,13 @@ void FControlRigParameterTrackEditor::AddControlRig(UClass* InClass, UObject* Bo
 			GetSequencer()->NotifyMovieSceneDataChanged(EMovieSceneDataChangeType::MovieSceneStructureItemAdded);
 
 			FControlRigEditMode* ControlRigEditMode = GetEditMode(true);
-			
+
 			if (ControlRigEditMode)
 			{
 				ControlRigEditMode->SetObjects(ControlRig, nullptr, GetSequencer());
 			}
 			BindControlRig(ControlRig);
-			
+
 		}
 	}
 }
@@ -1367,7 +1367,7 @@ void FControlRigParameterTrackEditor::AddControlRigFromComponent(FGuid InGuid)
 				AddControlRig(CR->GetClass(), BoundActor, InGuid, CR);
 			}
 		}
-			
+
 	}
 }
 
@@ -1785,48 +1785,48 @@ void FControlRigParameterTrackEditor::OnCurveDisplayChanged(FCurveModel* CurveMo
 			{
 				UE_LOG(LogControlRigEditor, Display, TEXT("Could not find Rig Control From FCurveModel::LongName"));
 			}
-		
+
 			if (bCurveDisplayTickIsPending == false)
 			{
 				bCurveDisplayTickIsPending = true;
 				GEditor->GetTimerManager()->SetTimerForNextTick([MovieSection, bDisplayed, this]()
-				{
-
-					if (DisplayedControls.Num() > 0 || UnDisplayedControls.Num() > 0)
 					{
-						TGuardValue<bool> Guard(bIsDoingSelection, true);
-						UMovieSceneControlRigParameterSection* ParamSection = Cast<UMovieSceneControlRigParameterSection>(MovieSection);
-						bool bSync = GetSequencer()->GetSequencerSettings()->ShouldSyncCurveEditorSelection();
-						GetSequencer()->SuspendSelectionBroadcast();
-						GetSequencer()->GetSequencerSettings()->SyncCurveEditorSelection(false);
-						if (UnDisplayedControls.Num() > 0)
+
+						if (DisplayedControls.Num() > 0 || UnDisplayedControls.Num() > 0)
 						{
-							for (const FName& ControlName : UnDisplayedControls)
+							TGuardValue<bool> Guard(bIsDoingSelection, true);
+							UMovieSceneControlRigParameterSection* ParamSection = Cast<UMovieSceneControlRigParameterSection>(MovieSection);
+							bool bSync = GetSequencer()->GetSequencerSettings()->ShouldSyncCurveEditorSelection();
+							GetSequencer()->SuspendSelectionBroadcast();
+							GetSequencer()->GetSequencerSettings()->SyncCurveEditorSelection(false);
+							if (UnDisplayedControls.Num() > 0)
 							{
-								SelectSequencerNodeInSection(ParamSection, ControlName, false);
+								for (const FName& ControlName : UnDisplayedControls)
+								{
+									SelectSequencerNodeInSection(ParamSection, ControlName, false);
+								}
+								UnDisplayedControls.Empty();
 							}
-							UnDisplayedControls.Empty();
-						}
-						if (DisplayedControls.Num() > 0)
-						{
-							for (const FName& ControlName : DisplayedControls)
+							if (DisplayedControls.Num() > 0)
 							{
-								SelectSequencerNodeInSection(ParamSection, ControlName, true);
+								for (const FName& ControlName : DisplayedControls)
+								{
+									SelectSequencerNodeInSection(ParamSection, ControlName, true);
+								}
+								DisplayedControls.Empty();
 							}
-							DisplayedControls.Empty();
-						}
-						GetSequencer()->ResumeSelectionBroadcast(); //need to resume first so when we refreh the tree we do the Selection.Tick, which since syncing is off won't 
-																	//mess up the curve editor.
-						GetSequencer()->RefreshTree();
-						GetSequencer()->GetSequencerSettings()->SyncCurveEditorSelection(bSync);
-					};
-					bCurveDisplayTickIsPending = false;
-	
-				});
+							GetSequencer()->ResumeSelectionBroadcast(); //need to resume first so when we refreh the tree we do the Selection.Tick, which since syncing is off won't 
+																		//mess up the curve editor.
+							GetSequencer()->RefreshTree();
+							GetSequencer()->GetSequencerSettings()->SyncCurveEditorSelection(bSync);
+						};
+						bCurveDisplayTickIsPending = false;
+
+					});
 
 			}
 		}
-		
+
 	}
 }
 
@@ -1860,7 +1860,7 @@ void FControlRigParameterTrackEditor::OnSelectionChanged(TArray<UMovieSceneTrack
 	{
 		return;
 	}
-	
+
 	TGuardValue<bool> Guard(bIsDoingSelection, true);
 
 	FControlRigEditMode* ControlRigEditMode = GetEditMode();
@@ -2268,23 +2268,23 @@ void FControlRigParameterTrackEditor::HandleControlSelected(UControlRig* Subject
 	if (URigHierarchyController* Controller = Hierarchy->GetController())
 	{
 		Hierarchy->ForEach<FRigControlElement>([ControlElement, Controller, bSelected](FRigControlElement* OtherControlElement) -> bool
-		{
-			if(OtherControlElement->Settings.ControlType == ERigControlType::Bool ||
-                    OtherControlElement->Settings.ControlType == ERigControlType::Float ||
-                    OtherControlElement->Settings.ControlType == ERigControlType::Integer)
 			{
-				for(const FRigElementParentConstraint& ParentConstraint : OtherControlElement->ParentConstraints)
+			if(OtherControlElement->Settings.ControlType == ERigControlType::Bool ||
+					OtherControlElement->Settings.ControlType == ERigControlType::Float ||
+					OtherControlElement->Settings.ControlType == ERigControlType::Integer)
 				{
-					if(ParentConstraint.ParentElement == ControlElement)
+				for(const FRigElementParentConstraint& ParentConstraint : OtherControlElement->ParentConstraints)
 					{
-						Controller->SelectElement(OtherControlElement->GetKey(), bSelected);
-						break;
+					if(ParentConstraint.ParentElement == ControlElement)
+						{
+							Controller->SelectElement(OtherControlElement->GetKey(), bSelected);
+							break;
+						}
 					}
 				}
-			}
 
-			return true;
-		});
+				return true;
+			});
 	}
 
 	}
@@ -3086,8 +3086,8 @@ public:
 	SLATE_BEGIN_ARGS(SFKControlRigBoneSelect) {}
 	SLATE_ATTRIBUTE(UFKControlRig*, AutoRig)
 	SLATE_ATTRIBUTE(UMovieSceneControlRigParameterTrack*,Track)
-	SLATE_ATTRIBUTE(ISequencer*, Sequencer)
-	SLATE_END_ARGS()
+		SLATE_ATTRIBUTE(ISequencer*, Sequencer)
+		SLATE_END_ARGS()
 
 		void Construct(const FArguments& InArgs)
 	{
@@ -3415,105 +3415,105 @@ void FControlRigParameterSection::BuildSectionContextMenu(FMenuBuilder& MenuBuil
 		{
 			return FUIAction(
 				FExecuteAction::CreateLambda([=]
-			{
-				FScopedTransaction Transaction(LOCTEXT("SetActiveChannelsTransaction", "Set Active Channels"));
-				ParameterSection->Modify();
-				EMovieSceneTransformChannel Channels = ParameterSection->GetTransformMask().GetChannels();
+					{
+						FScopedTransaction Transaction(LOCTEXT("SetActiveChannelsTransaction", "Set Active Channels"));
+						ParameterSection->Modify();
+						EMovieSceneTransformChannel Channels = ParameterSection->GetTransformMask().GetChannels();
 
-				if (EnumHasAllFlags(Channels, ChannelsToToggle) || (Channels & ChannelsToToggle) == EMovieSceneTransformChannel::None)
-				{
-					ParameterSection->SetTransformMask(ParameterSection->GetTransformMask().GetChannels() ^ ChannelsToToggle);
-				}
-				else
-				{
-					ParameterSection->SetTransformMask(ParameterSection->GetTransformMask().GetChannels() | ChannelsToToggle);
-				}
+						if (EnumHasAllFlags(Channels, ChannelsToToggle) || (Channels & ChannelsToToggle) == EMovieSceneTransformChannel::None)
+						{
+							ParameterSection->SetTransformMask(ParameterSection->GetTransformMask().GetChannels() ^ ChannelsToToggle);
+						}
+						else
+						{
+							ParameterSection->SetTransformMask(ParameterSection->GetTransformMask().GetChannels() | ChannelsToToggle);
+						}
 
-				// Restore pre-animated state for the bound objects so that inactive channels will return to their default values.
-				for (TWeakObjectPtr<> WeakObject : SequencerPtr->FindBoundObjects(InObjectBinding, SequencerPtr->GetFocusedTemplateID()))
-				{
-					if (UObject* Object = WeakObject.Get())
-					{
-						SequencerPtr->RestorePreAnimatedState();
-					}
-				}
+						// Restore pre-animated state for the bound objects so that inactive channels will return to their default values.
+						for (TWeakObjectPtr<> WeakObject : SequencerPtr->FindBoundObjects(InObjectBinding, SequencerPtr->GetFocusedTemplateID()))
+						{
+							if (UObject* Object = WeakObject.Get())
+							{
+								SequencerPtr->RestorePreAnimatedState();
+							}
+						}
 
-				SequencerPtr->NotifyMovieSceneDataChanged(EMovieSceneDataChangeType::MovieSceneStructureItemsChanged);
-			}
+						SequencerPtr->NotifyMovieSceneDataChanged(EMovieSceneDataChangeType::MovieSceneStructureItemsChanged);
+					}
 				),
 				FCanExecuteAction(),
-				FGetActionCheckState::CreateLambda([=]
-			{
-				EMovieSceneTransformChannel Channels = ParameterSection->GetTransformMask().GetChannels();
-				if (EnumHasAllFlags(Channels, ChannelsToToggle))
-				{
-					return ECheckBoxState::Checked;
-				}
-				else if (EnumHasAnyFlags(Channels, ChannelsToToggle))
-				{
-					return ECheckBoxState::Undetermined;
-				}
-				return ECheckBoxState::Unchecked;
-			})
-				);
+						FGetActionCheckState::CreateLambda([=]
+							{
+								EMovieSceneTransformChannel Channels = ParameterSection->GetTransformMask().GetChannels();
+								if (EnumHasAllFlags(Channels, ChannelsToToggle))
+								{
+									return ECheckBoxState::Checked;
+								}
+								else if (EnumHasAnyFlags(Channels, ChannelsToToggle))
+								{
+									return ECheckBoxState::Undetermined;
+								}
+								return ECheckBoxState::Unchecked;
+							})
+						);
 		};
 		auto ToggleControls = [=](int32 Index)
 		{
 			return FUIAction(
 				FExecuteAction::CreateLambda([=]
-			{
-				FScopedTransaction Transaction(LOCTEXT("ToggleRigControlFiltersTransaction", "Toggle Rig Control Filters"));
-				ParameterSection->Modify();
-				if (Index >= 0)
-				{
-					ParameterSection->SetControlsMask(Index, !ParameterSection->GetControlsMask(Index));
-				}
-				else
-				{
-					ParameterSection->FillControlsMask(!ParameterSection->GetControlsMask(0));
-				}
-				SequencerPtr->NotifyMovieSceneDataChanged(EMovieSceneDataChangeType::MovieSceneStructureItemsChanged);
-
-			}
-				),
-				FCanExecuteAction(),
-				FGetActionCheckState::CreateLambda([=]
-			{
-				TArray<bool> ControlBool = ParameterSection->GetControlsMask();
-				if (Index >= 0)
-				{
-					if (ControlBool[Index])
-					{
-						return ECheckBoxState::Checked;
-
-					}
-					else
-					{
-						return ECheckBoxState::Unchecked;
-					}
-				}
-				else
-				{
-					TOptional<bool> FirstVal;
-					for (bool Val : ControlBool)
 					{
-						if (FirstVal.IsSet())
+						FScopedTransaction Transaction(LOCTEXT("ToggleRigControlFiltersTransaction", "Toggle Rig Control Filters"));
+						ParameterSection->Modify();
+						if (Index >= 0)
 						{
-							if (Val != FirstVal)
-							{
-								return ECheckBoxState::Undetermined;
-							}
+							ParameterSection->SetControlsMask(Index, !ParameterSection->GetControlsMask(Index));
 						}
 						else
 						{
-							FirstVal = Val;
+							ParameterSection->FillControlsMask(!ParameterSection->GetControlsMask(0));
 						}
+						SequencerPtr->NotifyMovieSceneDataChanged(EMovieSceneDataChangeType::MovieSceneStructureItemsChanged);
 
 					}
-					return (FirstVal.IsSet() && FirstVal.GetValue()) ? ECheckBoxState::Checked : ECheckBoxState::Unchecked;
-				}
-			})
-				);
+				),
+				FCanExecuteAction(),
+						FGetActionCheckState::CreateLambda([=]
+							{
+								TArray<bool> ControlBool = ParameterSection->GetControlsMask();
+								if (Index >= 0)
+								{
+									if (ControlBool[Index])
+									{
+										return ECheckBoxState::Checked;
+
+									}
+									else
+									{
+										return ECheckBoxState::Unchecked;
+									}
+								}
+								else
+								{
+									TOptional<bool> FirstVal;
+									for (bool Val : ControlBool)
+									{
+										if (FirstVal.IsSet())
+										{
+											if (Val != FirstVal)
+											{
+												return ECheckBoxState::Undetermined;
+											}
+										}
+										else
+										{
+											FirstVal = Val;
+										}
+
+									}
+									return (FirstVal.IsSet() && FirstVal.GetValue()) ? ECheckBoxState::Checked : ECheckBoxState::Unchecked;
+								}
+							})
+						);
 		};
 		MenuBuilder.BeginSection(NAME_None, LOCTEXT("RigSectionActiveChannels", "Active Channels"));
 		{
@@ -3534,53 +3534,53 @@ void FControlRigParameterSection::BuildSectionContextMenu(FMenuBuilder& MenuBuil
 			MenuBuilder.AddSubMenu(
 				LOCTEXT("AllTranslation", "Translation"), LOCTEXT("AllTranslation_ToolTip", "Causes this section to affect the translation of rig control transforms"),
 				FNewMenuDelegate::CreateLambda([=](FMenuBuilder& SubMenuBuilder) {
-				SubMenuBuilder.AddMenuEntry(
-					LOCTEXT("TranslationX", "X"), LOCTEXT("TranslationX_ToolTip", "Causes this section to affect the X channel of the transform's translation"),
-					FSlateIcon(), MakeUIAction(EMovieSceneTransformChannel::TranslationX), NAME_None, EUserInterfaceActionType::ToggleButton);
-				SubMenuBuilder.AddMenuEntry(
-					LOCTEXT("TranslationY", "Y"), LOCTEXT("TranslationY_ToolTip", "Causes this section to affect the Y channel of the transform's translation"),
-					FSlateIcon(), MakeUIAction(EMovieSceneTransformChannel::TranslationY), NAME_None, EUserInterfaceActionType::ToggleButton);
-				SubMenuBuilder.AddMenuEntry(
-					LOCTEXT("TranslationZ", "Z"), LOCTEXT("TranslationZ_ToolTip", "Causes this section to affect the Z channel of the transform's translation"),
-					FSlateIcon(), MakeUIAction(EMovieSceneTransformChannel::TranslationZ), NAME_None, EUserInterfaceActionType::ToggleButton);
-			}),
+					SubMenuBuilder.AddMenuEntry(
+						LOCTEXT("TranslationX", "X"), LOCTEXT("TranslationX_ToolTip", "Causes this section to affect the X channel of the transform's translation"),
+						FSlateIcon(), MakeUIAction(EMovieSceneTransformChannel::TranslationX), NAME_None, EUserInterfaceActionType::ToggleButton);
+					SubMenuBuilder.AddMenuEntry(
+						LOCTEXT("TranslationY", "Y"), LOCTEXT("TranslationY_ToolTip", "Causes this section to affect the Y channel of the transform's translation"),
+						FSlateIcon(), MakeUIAction(EMovieSceneTransformChannel::TranslationY), NAME_None, EUserInterfaceActionType::ToggleButton);
+					SubMenuBuilder.AddMenuEntry(
+						LOCTEXT("TranslationZ", "Z"), LOCTEXT("TranslationZ_ToolTip", "Causes this section to affect the Z channel of the transform's translation"),
+						FSlateIcon(), MakeUIAction(EMovieSceneTransformChannel::TranslationZ), NAME_None, EUserInterfaceActionType::ToggleButton);
+					}),
 				MakeUIAction(EMovieSceneTransformChannel::Translation),
-				NAME_None,
-				EUserInterfaceActionType::ToggleButton);
+						NAME_None,
+						EUserInterfaceActionType::ToggleButton);
 
 			MenuBuilder.AddSubMenu(
 				LOCTEXT("AllRotation", "Rotation"), LOCTEXT("AllRotation_ToolTip", "Causes this section to affect the rotation of the rig control transform"),
 				FNewMenuDelegate::CreateLambda([=](FMenuBuilder& SubMenuBuilder) {
-				SubMenuBuilder.AddMenuEntry(
-					LOCTEXT("RotationX", "Roll (X)"), LOCTEXT("RotationX_ToolTip", "Causes this section to affect the roll (X) channel the transform's rotation"),
-					FSlateIcon(), MakeUIAction(EMovieSceneTransformChannel::RotationX), NAME_None, EUserInterfaceActionType::ToggleButton);
-				SubMenuBuilder.AddMenuEntry(
-					LOCTEXT("RotationY", "Pitch (Y)"), LOCTEXT("RotationY_ToolTip", "Causes this section to affect the pitch (Y) channel the transform's rotation"),
-					FSlateIcon(), MakeUIAction(EMovieSceneTransformChannel::RotationY), NAME_None, EUserInterfaceActionType::ToggleButton);
-				SubMenuBuilder.AddMenuEntry(
-					LOCTEXT("RotationZ", "Yaw (Z)"), LOCTEXT("RotationZ_ToolTip", "Causes this section to affect the yaw (Z) channel the transform's rotation"),
-					FSlateIcon(), MakeUIAction(EMovieSceneTransformChannel::RotationZ), NAME_None, EUserInterfaceActionType::ToggleButton);
-			}),
+					SubMenuBuilder.AddMenuEntry(
+						LOCTEXT("RotationX", "Roll (X)"), LOCTEXT("RotationX_ToolTip", "Causes this section to affect the roll (X) channel the transform's rotation"),
+						FSlateIcon(), MakeUIAction(EMovieSceneTransformChannel::RotationX), NAME_None, EUserInterfaceActionType::ToggleButton);
+					SubMenuBuilder.AddMenuEntry(
+						LOCTEXT("RotationY", "Pitch (Y)"), LOCTEXT("RotationY_ToolTip", "Causes this section to affect the pitch (Y) channel the transform's rotation"),
+						FSlateIcon(), MakeUIAction(EMovieSceneTransformChannel::RotationY), NAME_None, EUserInterfaceActionType::ToggleButton);
+					SubMenuBuilder.AddMenuEntry(
+						LOCTEXT("RotationZ", "Yaw (Z)"), LOCTEXT("RotationZ_ToolTip", "Causes this section to affect the yaw (Z) channel the transform's rotation"),
+						FSlateIcon(), MakeUIAction(EMovieSceneTransformChannel::RotationZ), NAME_None, EUserInterfaceActionType::ToggleButton);
+					}),
 				MakeUIAction(EMovieSceneTransformChannel::Rotation),
-				NAME_None,
-				EUserInterfaceActionType::ToggleButton);
+						NAME_None,
+						EUserInterfaceActionType::ToggleButton);
 
 			MenuBuilder.AddSubMenu(
 				LOCTEXT("AllScale", "Scale"), LOCTEXT("AllScale_ToolTip", "Causes this section to affect the scale of the rig control transform"),
 				FNewMenuDelegate::CreateLambda([=](FMenuBuilder& SubMenuBuilder) {
-				SubMenuBuilder.AddMenuEntry(
-					LOCTEXT("ScaleX", "X"), LOCTEXT("ScaleX_ToolTip", "Causes this section to affect the X channel of the transform's scale"),
-					FSlateIcon(), MakeUIAction(EMovieSceneTransformChannel::ScaleX), NAME_None, EUserInterfaceActionType::ToggleButton);
-				SubMenuBuilder.AddMenuEntry(
-					LOCTEXT("ScaleY", "Y"), LOCTEXT("ScaleY_ToolTip", "Causes this section to affect the Y channel of the transform's scale"),
-					FSlateIcon(), MakeUIAction(EMovieSceneTransformChannel::ScaleY), NAME_None, EUserInterfaceActionType::ToggleButton);
-				SubMenuBuilder.AddMenuEntry(
-					LOCTEXT("ScaleZ", "Z"), LOCTEXT("ScaleZ_ToolTip", "Causes this section to affect the Z channel of the transform's scale"),
-					FSlateIcon(), MakeUIAction(EMovieSceneTransformChannel::ScaleZ), NAME_None, EUserInterfaceActionType::ToggleButton);
-			}),
+					SubMenuBuilder.AddMenuEntry(
+						LOCTEXT("ScaleX", "X"), LOCTEXT("ScaleX_ToolTip", "Causes this section to affect the X channel of the transform's scale"),
+						FSlateIcon(), MakeUIAction(EMovieSceneTransformChannel::ScaleX), NAME_None, EUserInterfaceActionType::ToggleButton);
+					SubMenuBuilder.AddMenuEntry(
+						LOCTEXT("ScaleY", "Y"), LOCTEXT("ScaleY_ToolTip", "Causes this section to affect the Y channel of the transform's scale"),
+						FSlateIcon(), MakeUIAction(EMovieSceneTransformChannel::ScaleY), NAME_None, EUserInterfaceActionType::ToggleButton);
+					SubMenuBuilder.AddMenuEntry(
+						LOCTEXT("ScaleZ", "Z"), LOCTEXT("ScaleZ_ToolTip", "Causes this section to affect the Z channel of the transform's scale"),
+						FSlateIcon(), MakeUIAction(EMovieSceneTransformChannel::ScaleZ), NAME_None, EUserInterfaceActionType::ToggleButton);
+					}),
 				MakeUIAction(EMovieSceneTransformChannel::Scale),
-				NAME_None,
-				EUserInterfaceActionType::ToggleButton);
+						NAME_None,
+						EUserInterfaceActionType::ToggleButton);
 
 			//mz todo h
 			MenuBuilder.AddMenuEntry(
diff --git a/Engine/Plugins/Messaging/UdpMessaging/Source/UdpMessaging/Private/Transport/UdpMessageProcessor.cpp b/Engine/Plugins/Messaging/UdpMessaging/Source/UdpMessaging/Private/Transport/UdpMessageProcessor.cpp
--- a/Engine/Plugins/Messaging/UdpMessaging/Source/UdpMessaging/Private/Transport/UdpMessageProcessor.cpp
+++ b/Engine/Plugins/Messaging/UdpMessaging/Source/UdpMessaging/Private/Transport/UdpMessageProcessor.cpp
@@ -149,7 +149,7 @@ bool FUdpMessageProcessor::EnqueueOutboundMessage(const TSharedRef<IMessageConte
 	{
 		return false;
 	}
-	
+
 	TMap<uint8, TArray<FGuid>> RecipientPerVersions = GetRecipientsPerProtocolVersion(Recipients);
 	for (const auto& RecipientVersion : RecipientPerVersions)
 	{
@@ -248,13 +248,13 @@ uint32 FUdpMessageProcessor::Run()
 
 		do
 		{
-		FDateTime LastTime = CurrentTime;
-		CurrentTime = FDateTime::UtcNow();
-		DeltaTime = CurrentTime - LastTime;
+			FDateTime LastTime = CurrentTime;
+			CurrentTime = FDateTime::UtcNow();
+			DeltaTime = CurrentTime - LastTime;
 
 			ConsumeInboundSegments();
 			ConsumeOutboundMessages();
-		UpdateKnownNodes();
+			UpdateKnownNodes();
 			UpdateNetworkStatistics();
 		} while ((!InboundSegments.IsEmpty() || MoreToSend()) && !bStopping);
 	}
@@ -416,7 +416,7 @@ void FUdpMessageProcessor::ConsumeInboundSegments()
 				ProcessByeSegment(Segment, NodeInfo);
 				break;
 
-			case EUdpMessageSegments::Data:			
+			case EUdpMessageSegments::Data:
 				ProcessDataSegment(Segment, NodeInfo);
 				break;
 
@@ -498,8 +498,8 @@ bool FUdpMessageProcessor::ConsumeOneOutboundMessage(const FOutboundMessage& Out
 
 			RecipientNodeInfo->WorkQueue.Enqueue(LastSentMessage);
 			SendSegmenterStatsToListeners(LastSentMessage, RecipientNodeInfo->NodeId, Segmenter);
+		}
 	}
-}
 
 	// Return true that it was consumed or not reliable. Unreliable messages can be discarded.
 	return bCanConsume || !bIsReliable;
@@ -518,17 +518,17 @@ void FUdpMessageProcessor::ConsumeOutboundMessages()
 	};
 
 	if (DeferredOutboundMessage)
-		{
+	{
 		if (ConsumeOneOutboundMessage(DeferredOutboundMessage.GetValue()))
-			{
+		{
 			DeferredOutboundMessage.Reset();
-			}
-			else
-			{
+		}
+		else
+		{
 			CannotConsumeMsg();
 			return ;
-			}
 		}
+	}
 
 	while (OutboundMessages.Dequeue(OutboundMessage))
 	{
@@ -537,9 +537,9 @@ void FUdpMessageProcessor::ConsumeOutboundMessages()
 			CannotConsumeMsg();
 			DeferredOutboundMessage = OutboundMessage;
 			return ;
+		}
 	}
 }
-}
 
 bool FUdpMessageProcessor::FilterSegment(const FUdpMessageSegment::FHeader& Header)
 {
@@ -598,7 +598,7 @@ void FUdpMessageProcessor::ProcessAcknowledgeSegmentsSegment(FInboundSegment& Se
 	UE_LOG(LogUdpMessaging, Verbose, TEXT("Received AcknowledgeSegments for %d from %s"), AcknowledgeChunk.MessageId, *NodeInfo.NodeId.ToString());
 
 	NodeInfo.MarkAcks(AcknowledgeChunk.MessageId, AcknowledgeChunk.Segments, CurrentTime);
-		}
+}
 
 
 void FUdpMessageProcessor::ProcessByeSegment(FInboundSegment& Segment, FNodeInfo& NodeInfo)
@@ -617,7 +617,7 @@ void FUdpMessageProcessor::ProcessDataSegment(FInboundSegment& Segment, FNodeInf
 {
 	FUdpMessageSegment::FDataChunk DataChunk;
 	DataChunk.Serialize(*Segment.Data, NodeInfo.ProtocolVersion);
-	
+
 	if (Segment.Data->IsError())
 	{
 		UE_LOG(LogUdpMessaging, Warning, TEXT("FUdpMessageProcessor::ProcessDataSegment: Failed to serialize DataChunk. Sender=%s"),
@@ -653,8 +653,8 @@ void FUdpMessageProcessor::ProcessDataSegment(FInboundSegment& Segment, FNodeInf
 	}
 
 	UE_LOG(LogUdpMessaging, Verbose, TEXT("FUdpMessageProcessor::ProcessDataSegment: Reassembled %d bytes message %s with %d for %s (%s)"),
-		ReassembledMessage->GetData().Num(), 
-		*ReassembledMessage->Describe(), 
+		ReassembledMessage->GetData().Num(),
+		*ReassembledMessage->Describe(),
 		DataChunk.MessageId,
 		*NodeInfo.NodeId.ToString(),
 		*NodeInfo.Endpoint.ToString());
@@ -686,7 +686,7 @@ void FUdpMessageProcessor::ProcessPingSegment(FInboundSegment& Segment, FNodeInf
 	{
 		NodeInfo.ResetIfRestarted(RemoteNodeId);
 	}
-	
+
 	// The protocol version we are going to use to communicate to this node is the smallest between its version and our own
 	uint8 ProtocolVersion = FMath::Min<uint8>(NodeProtocolVersion, UDP_MESSAGING_TRANSPORT_PROTOCOL_VERSION);
 
@@ -782,10 +782,10 @@ void FUdpMessageProcessor::DeliverMessage(const TSharedPtr<FUdpReassembledMessag
 		return;
 	}
 
-		if (NodeInfo.NodeId.IsValid())
-		{
-			MessageReassembledDelegate.ExecuteIfBound(*ReassembledMessage, nullptr, NodeInfo.NodeId);
-		}
+	if (NodeInfo.NodeId.IsValid())
+	{
+		MessageReassembledDelegate.ExecuteIfBound(*ReassembledMessage, nullptr, NodeInfo.NodeId);
+	}
 
 	// Mark the message delivered but do not remove it from the list yet, this is to prevent the double delivery of reliable message
 	ReassembledMessage->MarkDelivered();
@@ -830,13 +830,13 @@ void FUdpMessageProcessor::RemoveDeadNodes()
 }
 
 bool FUdpMessageProcessor::MoreToSend()
-	{
+{
 	for (auto& KnownNodePair : KnownNodes)
 	{
 		if (KnownNodePair.Value.CanSendSegments() && KnownNodePair.Value.HasSegmenterThatCanSend(CurrentTime))
 		{
 			return true;
-	}
+		}
 	}
 	return false;
 }
@@ -851,16 +851,16 @@ void FUdpMessageProcessor::UpdateKnownNodes()
 	Beacon->SetEndpointCount(KnownNodes.Num() + 1);
 
 	bool bSuccess = true;
-		for (auto& KnownNodePair : KnownNodes)
-		{
+	for (auto& KnownNodePair : KnownNodes)
+	{
 		int32 NodeByteSent = UpdateSegmenters(KnownNodePair.Value);
-			// if NodByteSent is negative, there is a socket error, continuing is useless
-			bSuccess = NodeByteSent >= 0;
-			if (!bSuccess)
-			{
-				UE_LOG(LogUdpMessaging, Verbose, TEXT("FUdpMessageProcessor::UpdateKnownNodes received negative NodeByteSent (%d) from socket."), NodeByteSent);
-				break;
-			}
+		// if NodByteSent is negative, there is a socket error, continuing is useless
+		bSuccess = NodeByteSent >= 0;
+		if (!bSuccess)
+		{
+			UE_LOG(LogUdpMessaging, Verbose, TEXT("FUdpMessageProcessor::UpdateKnownNodes received negative NodeByteSent (%d) from socket sender."), NodeByteSent);
+			break;
+		}
 
 		bSuccess = UpdateReassemblers(KnownNodePair.Value);
 		// if there is a socket error, continuing is useless
@@ -883,7 +883,7 @@ TSharedPtr<FUdpMessageSegmenter>* FUdpMessageProcessor::GetInitializedSegmenter(
 {
 	TSharedPtr<FUdpMessageSegmenter>* FoundSegmenter = NodeInfo.Segmenters.Find(MessageId);
 	if (!FoundSegmenter)
-{
+	{
 		// It's possible that a segmenter was fully ack and removed from the segmenter list.
 		return nullptr;
 	}
@@ -905,20 +905,20 @@ TSharedPtr<FUdpMessageSegmenter>* FUdpMessageProcessor::GetInitializedSegmenter(
 }
 
 FUdpMessageSegment::FDataChunk FUdpMessageProcessor::GetDataChunk(FNodeInfo& NodeInfo, TSharedPtr<FUdpMessageSegmenter>& Segmenter, int32 MessageId)
-		{
-			FUdpMessageSegment::FDataChunk DataChunk;
+{
+	FUdpMessageSegment::FDataChunk DataChunk;
 
-			// Track the segments we sent as we'll update the segmenter to keep track
+	// Track the segments we sent as we'll update the segmenter to keep track
 	TConstSetBitIterator<> BIt(Segmenter->GetPendingSendSegments());
-				Segmenter->GetPendingSegment(BIt.GetIndex(), DataChunk.Data);
-				DataChunk.SegmentNumber = BIt.GetIndex();
+	Segmenter->GetPendingSegment(BIt.GetIndex(), DataChunk.Data);
+	DataChunk.SegmentNumber = BIt.GetIndex();
 
 	DataChunk.MessageId = MessageId;
-				DataChunk.MessageFlags = Segmenter->GetMessageFlags();
-				DataChunk.MessageSize = Segmenter->GetMessageSize();
-				DataChunk.SegmentOffset = UDP_MESSAGING_SEGMENT_SIZE * DataChunk.SegmentNumber;
+	DataChunk.MessageFlags = Segmenter->GetMessageFlags();
+	DataChunk.MessageSize = Segmenter->GetMessageSize();
+	DataChunk.SegmentOffset = UDP_MESSAGING_SEGMENT_SIZE * DataChunk.SegmentNumber;
 	DataChunk.Sequence = 0; // This should be kept 0 for legacy support.
-				DataChunk.TotalSegments = Segmenter->GetSegmentCount();
+	DataChunk.TotalSegments = Segmenter->GetSegmentCount();
 	return DataChunk;
 }
 
@@ -938,27 +938,27 @@ FSentSegmentInfo FUdpMessageProcessor::SendNextSegmentForMessageId(FNodeInfo& No
 	{
 		FUdpMessageSegment::FDataChunk DataChunk = GetDataChunk(NodeInfo, Segmenter, MessageId);
 
-				// validate with are sending message in the proper protocol version
-				check(Header.ProtocolVersion == Segmenter->GetProtocolVersion());
+		// validate with are sending message in the proper protocol version
+		check(Header.ProtocolVersion == Segmenter->GetProtocolVersion());
 
-				TSharedRef<FArrayWriter, ESPMode::ThreadSafe> Writer = MakeShared<FArrayWriter, ESPMode::ThreadSafe>();
-				{
-					*Writer << Header;
-					DataChunk.Serialize(*Writer, Header.ProtocolVersion);
-				}
+		TSharedRef<FArrayWriter, ESPMode::ThreadSafe> Writer = MakeShared<FArrayWriter, ESPMode::ThreadSafe>();
+		{
+			*Writer << Header;
+			DataChunk.Serialize(*Writer, Header.ProtocolVersion);
+		}
 
 		UE_LOG(LogUdpMessaging, Verbose, TEXT("FUdpMessageProcessor::UpdateSegmenters: Sending msg %d as segment %d/%d of %d bytes to %s"),
-					DataChunk.MessageId, 
-					DataChunk.SegmentNumber + 1, 
-					DataChunk.TotalSegments, 
-					Segmenter->GetMessageSize(), 
-					*NodeInfo.NodeId.ToString());
+			   DataChunk.MessageId,
+			   DataChunk.SegmentNumber + 1,
+			   DataChunk.TotalSegments,
+			   Segmenter->GetMessageSize(),
+			   *NodeInfo.NodeId.ToString());
 
-				if (!SocketSender->Send(Writer, NodeInfo.Endpoint))
-				{
+		if (!SocketSender->Send(Writer, NodeInfo.Endpoint))
+		{
 			SentInfo.bSendSocketError = true;
 			return MoveTemp(SentInfo);
- 				}
+		}
 
 		Segmenter->MarkAsSent(DataChunk.SegmentNumber);
 
@@ -977,25 +977,25 @@ FSentSegmentInfo FUdpMessageProcessor::SendNextSegmentForMessageId(FNodeInfo& No
 		SendSegmenterStatsToListeners(MessageId, NodeInfo.NodeId, Segmenter);
 
 		if (Segmenter->IsSendingComplete() && !SentInfo.bIsReliable)
-			{
+		{
 			// Not reliably sent so we don't need to wait for acks.
-					UE_LOG(LogUdpMessaging, VeryVerbose, TEXT("FUdpMessageProcessor::UpdateSegmenters: Finished with message segmenter for %s"), *NodeInfo.NodeId.ToString());
+			UE_LOG(LogUdpMessaging, VeryVerbose, TEXT("FUdpMessageProcessor::UpdateSegmenters: Finished with message segmenter for %s"), *NodeInfo.NodeId.ToString());
 			NodeInfo.Segmenters.Remove(MessageId);
-				}
-			}
-			else
-			{
+		}
+	}
+	else
+	{
 		SentInfo.bRequiresRequeue = true;
-			}
+	}
 
 	NodeInfo.Statistics.BytesSent += SentInfo.BytesSent;
 	NodeInfo.Statistics.IpAddress = NodeInfo.Endpoint;
 	NodeInfo.Statistics.SegmentsInFlight = NodeInfo.InflightSegments.Num();
 	return MoveTemp(SentInfo);
-		}
+}
 
 int32 FUdpMessageProcessor::UpdateSegmenters(FNodeInfo& NodeInfo)
-		{
+{
 	SCOPED_MESSAGING_TRACE(FUdpMessageProcessor_UpdateSegmenters);
 	FUdpMessageSegment::FHeader Header
 	{
@@ -1029,7 +1029,7 @@ int32 FUdpMessageProcessor::UpdateSegmenters(FNodeInfo& NodeInfo)
 	for (int32 RequeueMessageId : MessagesForRequeue)
 	{
 		NodeInfo.WorkQueue.Enqueue(RequeueMessageId);
-}
+	}
 	return BytesSent;
 }
 
@@ -1067,7 +1067,7 @@ bool FUdpMessageProcessor::UpdateReassemblers(FNodeInfo& NodeInfo)
 			TSharedRef<FArrayWriter, ESPMode::ThreadSafe> Writer = MakeShared<FArrayWriter, ESPMode::ThreadSafe>();
 			{
 				*Writer << Header;
-				AcknowledgeChunk.Serialize(*Writer, Header.ProtocolVersion);				
+				AcknowledgeChunk.Serialize(*Writer, Header.ProtocolVersion);
 			}
 			BytesSent += Writer->Num();
 
@@ -1079,10 +1079,10 @@ bool FUdpMessageProcessor::UpdateReassemblers(FNodeInfo& NodeInfo)
 
 				return false;
 			}
-			
-			UE_LOG(LogUdpMessaging, Verbose, TEXT("FUdpMessageProcessor::UpdateReassemblers: sending acknowledgement for reliable msg %d from %s"), 
-				It.Key(), 
-				*NodeInfo.NodeId.ToString());			
+
+			UE_LOG(LogUdpMessaging, Verbose, TEXT("FUdpMessageProcessor::UpdateReassemblers: sending acknowledgement for reliable msg %d from %s"),
+				It.Key(),
+				*NodeInfo.NodeId.ToString());
 		}
 
 		// Try to deliver completed message that couldn't be delivered the first time around
@@ -1096,9 +1096,9 @@ bool FUdpMessageProcessor::UpdateReassemblers(FNodeInfo& NodeInfo)
 			(!EnumHasAnyFlags(ReassembledMessage->GetFlags(), EMessageFlags::Reliable) || ReassembledMessage->IsDelivered()))
 		{
 			if (!ReassembledMessage->IsDelivered())
-			{ 
+			{
 				const int ReceivedSegments = ReassembledMessage->GetTotalSegmentsCount() - ReassembledMessage->GetPendingSegmentsCount();
-				UE_LOG(LogUdpMessaging, Warning, TEXT("FUdpMessageProcessor::UpdateReassemblers Discarding %d/%d of stale message segements from %s"), 
+				UE_LOG(LogUdpMessaging, Warning, TEXT("FUdpMessageProcessor::UpdateReassemblers Discarding %d/%d of stale message segements from %s"),
 					ReceivedSegments,
 					ReassembledMessage->GetTotalSegmentsCount(),
 					*ReassembledMessage->Describe());
diff --git a/Engine/Plugins/Runtime/nDisplay/Source/DisplayCluster/Private/Game/EngineClasses/Basics/DisplayClusterViewportClient.cpp b/Engine/Plugins/Runtime/nDisplay/Source/DisplayCluster/Private/Game/EngineClasses/Basics/DisplayClusterViewportClient.cpp
--- a/Engine/Plugins/Runtime/nDisplay/Source/DisplayCluster/Private/Game/EngineClasses/Basics/DisplayClusterViewportClient.cpp
+++ b/Engine/Plugins/Runtime/nDisplay/Source/DisplayCluster/Private/Game/EngineClasses/Basics/DisplayClusterViewportClient.cpp
@@ -224,77 +224,77 @@ void UDisplayClusterViewportClient::Draw(FViewport* InViewport, FCanvas* SceneCa
 
 		for (FDisplayClusterRenderFrame::FFrameViewFamily& DCViewFamily : DCRenderTarget.ViewFamilies)
 		{
-				// Create the view family for rendering the world scene to the viewport's render target
-				FSceneViewFamilyContext ViewFamily(FSceneViewFamily::ConstructionValues(DCRenderTarget.RenderTargetPtr, MyWorld->Scene, EngineShowFlags)
-					.SetRealtimeUpdate(true)
-					.SetAdditionalViewFamily(bAdditionalViewFamily));
+			// Create the view family for rendering the world scene to the viewport's render target
+			FSceneViewFamilyContext ViewFamily(FSceneViewFamily::ConstructionValues(DCRenderTarget.RenderTargetPtr, MyWorld->Scene, EngineShowFlags)
+				.SetRealtimeUpdate(true)
+				.SetAdditionalViewFamily(bAdditionalViewFamily));
 
-				// Disable clean op for all next families on this render target
-				bAdditionalViewFamily = true;
+			// Disable clean op for all next families on this render target
+			bAdditionalViewFamily = true;
 
-				// Configure family flags
-				RenderFrame.ViewportManager->ConfigureViewFamily(DCRenderTarget, DCViewFamily, ViewFamily);
+			// Configure family flags
+			RenderFrame.ViewportManager->ConfigureViewFamily(DCRenderTarget, DCViewFamily, ViewFamily);
 
 #if WITH_EDITOR
-				if (GIsEditor)
-				{
-					// Force enable view family show flag for HighDPI derived's screen percentage.
-					ViewFamily.EngineShowFlags.ScreenPercentage = true;
-				}
+			if (GIsEditor)
+			{
+				// Force enable view family show flag for HighDPI derived's screen percentage.
+				ViewFamily.EngineShowFlags.ScreenPercentage = true;
+			}
 #endif
 
-				ViewFamily.ViewMode = EViewModeIndex(ViewModeIndex);
-				EngineShowFlagOverride(ESFIM_Game, ViewFamily.ViewMode, ViewFamily.EngineShowFlags, false);
+			ViewFamily.ViewMode = EViewModeIndex(ViewModeIndex);
+			EngineShowFlagOverride(ESFIM_Game, ViewFamily.ViewMode, ViewFamily.EngineShowFlags, false);
 
-				if (ViewFamily.EngineShowFlags.VisualizeBuffer && AllowDebugViewmodes())
+			if (ViewFamily.EngineShowFlags.VisualizeBuffer && AllowDebugViewmodes())
+			{
+				// Process the buffer visualization console command
+				FName NewBufferVisualizationMode = NAME_None;
+				static IConsoleVariable* ICVar = IConsoleManager::Get().FindConsoleVariable(FBufferVisualizationData::GetVisualizationTargetConsoleCommandName());
+				if (ICVar)
 				{
-					// Process the buffer visualization console command
-					FName NewBufferVisualizationMode = NAME_None;
-					static IConsoleVariable* ICVar = IConsoleManager::Get().FindConsoleVariable(FBufferVisualizationData::GetVisualizationTargetConsoleCommandName());
-					if (ICVar)
+					static const FName OverviewName = TEXT("Overview");
+					FString ModeNameString = ICVar->GetString();
+					FName ModeName = *ModeNameString;
+					if (ModeNameString.IsEmpty() || ModeName == OverviewName || ModeName == NAME_None)
 					{
-						static const FName OverviewName = TEXT("Overview");
-						FString ModeNameString = ICVar->GetString();
-						FName ModeName = *ModeNameString;
-						if (ModeNameString.IsEmpty() || ModeName == OverviewName || ModeName == NAME_None)
+						NewBufferVisualizationMode = NAME_None;
+					}
+					else
+					{
+						if (GetBufferVisualizationData().GetMaterial(ModeName) == NULL)
 						{
-							NewBufferVisualizationMode = NAME_None;
+							// Mode is out of range, so display a message to the user, and reset the mode back to the previous valid one
+							UE_LOG(LogConsoleResponse, Warning, TEXT("Buffer visualization mode '%s' does not exist"), *ModeNameString);
+							NewBufferVisualizationMode = GetCurrentBufferVisualizationMode();
+							// todo: cvars are user settings, here the cvar state is used to avoid log spam and to auto correct for the user (likely not what the user wants)
+							ICVar->Set(*NewBufferVisualizationMode.GetPlainNameString(), ECVF_SetByCode);
 						}
 						else
 						{
-							if (GetBufferVisualizationData().GetMaterial(ModeName) == NULL)
-							{
-								// Mode is out of range, so display a message to the user, and reset the mode back to the previous valid one
-								UE_LOG(LogConsoleResponse, Warning, TEXT("Buffer visualization mode '%s' does not exist"), *ModeNameString);
-								NewBufferVisualizationMode = GetCurrentBufferVisualizationMode();
-								// todo: cvars are user settings, here the cvar state is used to avoid log spam and to auto correct for the user (likely not what the user wants)
-								ICVar->Set(*NewBufferVisualizationMode.GetPlainNameString(), ECVF_SetByCode);
-							}
-							else
-							{
-								NewBufferVisualizationMode = ModeName;
-							}
+							NewBufferVisualizationMode = ModeName;
 						}
 					}
+				}
 
-					if (NewBufferVisualizationMode != GetCurrentBufferVisualizationMode())
-					{
-						SetCurrentBufferVisualizationMode(NewBufferVisualizationMode);
-					}
+				if (NewBufferVisualizationMode != GetCurrentBufferVisualizationMode())
+				{
+					SetCurrentBufferVisualizationMode(NewBufferVisualizationMode);
 				}
+			}
 
-				TMap<ULocalPlayer*, FSceneView*> PlayerViewMap;
-				FAudioDeviceHandle RetrievedAudioDevice = MyWorld->GetAudioDevice();
-				TArray<FSceneView*> Views;
+			TMap<ULocalPlayer*, FSceneView*> PlayerViewMap;
+			FAudioDeviceHandle RetrievedAudioDevice = MyWorld->GetAudioDevice();
+			TArray<FSceneView*> Views;
 
-				for (FDisplayClusterRenderFrame::FFrameView& DCView : DCViewFamily.Views)
-				{
-						const FDisplayClusterViewport_Context ViewportContext = DCView.Viewport->GetContexts()[DCView.ContextNum];
+			for (FDisplayClusterRenderFrame::FFrameView& DCView : DCViewFamily.Views)
+			{
+				const FDisplayClusterViewport_Context ViewportContext = DCView.Viewport->GetContexts()[DCView.ContextNum];
 
-						// Calculate the player's view information.
-						FVector		ViewLocation;
-						FRotator	ViewRotation;
-						FSceneView* View = LocalPlayer->CalcSceneView(&ViewFamily, ViewLocation, ViewRotation, InViewport, nullptr, ViewportContext.StereoViewIndex);
+				// Calculate the player's view information.
+				FVector		ViewLocation;
+				FRotator	ViewRotation;
+				FSceneView* View = LocalPlayer->CalcSceneView(&ViewFamily, ViewLocation, ViewRotation, InViewport, nullptr, ViewportContext.StereoViewIndex);
 
 				if (View && DCView.bDisableRender)
 				{
@@ -304,128 +304,128 @@ void UDisplayClusterViewportClient::Draw(FViewport* InViewport, FCanvas* SceneCa
 					View = nullptr;
 				}
 
-						if (View)
-						{
-							Views.Add(View);
+				if (View)
+				{
+					Views.Add(View);
 
-							// Apply viewport context settings to view (crossGPU, visibility, etc)
-							DCView.Viewport->SetupSceneView(DCView.ContextNum, World, ViewFamily, *View);
+					// Apply viewport context settings to view (crossGPU, visibility, etc)
+					DCView.Viewport->SetupSceneView(DCView.ContextNum, World, ViewFamily, *View);
 
-							// We don't allow instanced stereo currently
-							View->bIsInstancedStereoEnabled = false;
-							View->bShouldBindInstancedViewUB = false;
+					// We don't allow instanced stereo currently
+					View->bIsInstancedStereoEnabled = false;
+					View->bShouldBindInstancedViewUB = false;
 
-							if (View->Family->EngineShowFlags.Wireframe)
-							{
-								// Wireframe color is emissive-only, and mesh-modifying materials do not use material substitution, hence...
-								View->DiffuseOverrideParameter = FVector4(0.f, 0.f, 0.f, 0.f);
-								View->SpecularOverrideParameter = FVector4(0.f, 0.f, 0.f, 0.f);
-							}
-							else if (View->Family->EngineShowFlags.OverrideDiffuseAndSpecular)
-							{
-								View->DiffuseOverrideParameter = FVector4(GEngine->LightingOnlyBrightness.R, GEngine->LightingOnlyBrightness.G, GEngine->LightingOnlyBrightness.B, 0.0f);
-								View->SpecularOverrideParameter = FVector4(.1f, .1f, .1f, 0.0f);
-							}
-							else if (View->Family->EngineShowFlags.LightingOnlyOverride)
-							{
-								View->DiffuseOverrideParameter = FVector4(GEngine->LightingOnlyBrightness.R, GEngine->LightingOnlyBrightness.G, GEngine->LightingOnlyBrightness.B, 0.0f);
-								View->SpecularOverrideParameter = FVector4(0.f, 0.f, 0.f, 0.f);
-							}
-							else if (View->Family->EngineShowFlags.ReflectionOverride)
-							{
-								View->DiffuseOverrideParameter = FVector4(0.f, 0.f, 0.f, 0.f);
-								View->SpecularOverrideParameter = FVector4(1, 1, 1, 0.0f);
-								View->NormalOverrideParameter = FVector4(0, 0, 1, 0.0f);
-								View->RoughnessOverrideParameter = FVector2D(0.0f, 0.0f);
-							}
+					if (View->Family->EngineShowFlags.Wireframe)
+					{
+						// Wireframe color is emissive-only, and mesh-modifying materials do not use material substitution, hence...
+						View->DiffuseOverrideParameter = FVector4(0.f, 0.f, 0.f, 0.f);
+						View->SpecularOverrideParameter = FVector4(0.f, 0.f, 0.f, 0.f);
+					}
+					else if (View->Family->EngineShowFlags.OverrideDiffuseAndSpecular)
+					{
+						View->DiffuseOverrideParameter = FVector4(GEngine->LightingOnlyBrightness.R, GEngine->LightingOnlyBrightness.G, GEngine->LightingOnlyBrightness.B, 0.0f);
+						View->SpecularOverrideParameter = FVector4(.1f, .1f, .1f, 0.0f);
+					}
+					else if (View->Family->EngineShowFlags.LightingOnlyOverride)
+					{
+						View->DiffuseOverrideParameter = FVector4(GEngine->LightingOnlyBrightness.R, GEngine->LightingOnlyBrightness.G, GEngine->LightingOnlyBrightness.B, 0.0f);
+						View->SpecularOverrideParameter = FVector4(0.f, 0.f, 0.f, 0.f);
+					}
+					else if (View->Family->EngineShowFlags.ReflectionOverride)
+					{
+						View->DiffuseOverrideParameter = FVector4(0.f, 0.f, 0.f, 0.f);
+						View->SpecularOverrideParameter = FVector4(1, 1, 1, 0.0f);
+						View->NormalOverrideParameter = FVector4(0, 0, 1, 0.0f);
+						View->RoughnessOverrideParameter = FVector2D(0.0f, 0.0f);
+					}
 
-							if (!View->Family->EngineShowFlags.Diffuse)
-							{
-								View->DiffuseOverrideParameter = FVector4(0.f, 0.f, 0.f, 0.f);
-							}
+					if (!View->Family->EngineShowFlags.Diffuse)
+					{
+						View->DiffuseOverrideParameter = FVector4(0.f, 0.f, 0.f, 0.f);
+					}
 
-							if (!View->Family->EngineShowFlags.Specular)
-							{
-								View->SpecularOverrideParameter = FVector4(0.f, 0.f, 0.f, 0.f);
-							}
+					if (!View->Family->EngineShowFlags.Specular)
+					{
+						View->SpecularOverrideParameter = FVector4(0.f, 0.f, 0.f, 0.f);
+					}
 
-							View->CurrentBufferVisualizationMode = GetCurrentBufferVisualizationMode();
+					View->CurrentBufferVisualizationMode = GetCurrentBufferVisualizationMode();
 
-							View->CameraConstrainedViewRect = View->UnscaledViewRect;
+					View->CameraConstrainedViewRect = View->UnscaledViewRect;
 
 
-							{
-								// Save the location of the view.
-								LocalPlayer->LastViewLocation = ViewLocation;
+					{
+						// Save the location of the view.
+						LocalPlayer->LastViewLocation = ViewLocation;
 
-								PlayerViewMap.Add(LocalPlayer, View);
+						PlayerViewMap.Add(LocalPlayer, View);
 
-								// Update the listener.
-								if (RetrievedAudioDevice && PlayerController != NULL)
-								{
-									bool bUpdateListenerPosition = true;
+						// Update the listener.
+						if (RetrievedAudioDevice && PlayerController != NULL)
+						{
+							bool bUpdateListenerPosition = true;
 
-									// If the main audio device is used for multiple PIE viewport clients, we only
-									// want to update the main audio device listener position if it is in focus
-									if (GEngine)
-									{
-										FAudioDeviceManager* AudioDeviceManager = GEngine->GetAudioDeviceManager();
-
-										// If there is more than one world referencing the main audio device
-										if (AudioDeviceManager->GetNumMainAudioDeviceWorlds() > 1)
-										{
-											uint32 MainAudioDeviceID = GEngine->GetMainAudioDeviceID();
-											if (AudioDevice->DeviceID == MainAudioDeviceID && !HasAudioFocus())
-											{
-												bUpdateListenerPosition = false;
-											}
-										}
-									}
+							// If the main audio device is used for multiple PIE viewport clients, we only
+							// want to update the main audio device listener position if it is in focus
+							if (GEngine)
+							{
+								FAudioDeviceManager* AudioDeviceManager = GEngine->GetAudioDeviceManager();
 
-									if (bUpdateListenerPosition)
+								// If there is more than one world referencing the main audio device
+								if (AudioDeviceManager->GetNumMainAudioDeviceWorlds() > 1)
+								{
+									uint32 MainAudioDeviceID = GEngine->GetMainAudioDeviceID();
+									if (AudioDevice->DeviceID == MainAudioDeviceID && !HasAudioFocus())
 									{
-										FVector Location;
-										FVector ProjFront;
-										FVector ProjRight;
-										PlayerController->GetAudioListenerPosition(Location, ProjFront, ProjRight);
-
-										FTransform ListenerTransform(FRotationMatrix::MakeFromXY(ProjFront, ProjRight));
-
-										// Allow the HMD to adjust based on the head position of the player, as opposed to the view location
-										if (GEngine->XRSystem.IsValid() && GEngine->StereoRenderingDevice.IsValid() && GEngine->StereoRenderingDevice->IsStereoEnabled())
-										{
-											const FVector Offset = GEngine->XRSystem->GetAudioListenerOffset();
-											Location += ListenerTransform.TransformPositionNoScale(Offset);
-										}
-
-										ListenerTransform.SetTranslation(Location);
-										ListenerTransform.NormalizeRotation();
-
-										uint32 ViewportIndex = PlayerViewMap.Num() - 1;
-										RetrievedAudioDevice->SetListener(MyWorld, ViewportIndex, ListenerTransform, (View->bCameraCut ? 0.f : MyWorld->GetDeltaSeconds()));
-
-										FVector OverrideAttenuation;
-										if (PlayerController->GetAudioListenerAttenuationOverridePosition(OverrideAttenuation))
-										{
-											RetrievedAudioDevice->SetListenerAttenuationOverride(ViewportIndex, OverrideAttenuation);
-										}
-										else
-										{
-											RetrievedAudioDevice->ClearListenerAttenuationOverride(ViewportIndex);
-										}
+										bUpdateListenerPosition = false;
 									}
 								}
 							}
 
-							// Add view information for resource streaming. Allow up to 5X boost for small FOV.
-							const float StreamingScale = 1.f / FMath::Clamp<float>(View->LODDistanceFactor, .2f, 1.f);
-							IStreamingManager::Get().AddViewInformation(View->ViewMatrices.GetViewOrigin(), View->UnscaledViewRect.Width(), View->UnscaledViewRect.Width() * View->ViewMatrices.GetProjectionMatrix().M[0][0], StreamingScale);
-							MyWorld->ViewLocationsRenderedLastFrame.Add(View->ViewMatrices.GetViewOrigin());
+							if (bUpdateListenerPosition)
+							{
+								FVector Location;
+								FVector ProjFront;
+								FVector ProjRight;
+								PlayerController->GetAudioListenerPosition(Location, ProjFront, ProjRight);
+
+								FTransform ListenerTransform(FRotationMatrix::MakeFromXY(ProjFront, ProjRight));
+
+								// Allow the HMD to adjust based on the head position of the player, as opposed to the view location
+								if (GEngine->XRSystem.IsValid() && GEngine->StereoRenderingDevice.IsValid() && GEngine->StereoRenderingDevice->IsStereoEnabled())
+								{
+									const FVector Offset = GEngine->XRSystem->GetAudioListenerOffset();
+									Location += ListenerTransform.TransformPositionNoScale(Offset);
+								}
+
+								ListenerTransform.SetTranslation(Location);
+								ListenerTransform.NormalizeRotation();
+
+								uint32 ViewportIndex = PlayerViewMap.Num() - 1;
+								RetrievedAudioDevice->SetListener(MyWorld, ViewportIndex, ListenerTransform, (View->bCameraCut ? 0.f : MyWorld->GetDeltaSeconds()));
+
+								FVector OverrideAttenuation;
+								if (PlayerController->GetAudioListenerAttenuationOverridePosition(OverrideAttenuation))
+								{
+									RetrievedAudioDevice->SetListenerAttenuationOverride(ViewportIndex, OverrideAttenuation);
+								}
+								else
+								{
+									RetrievedAudioDevice->ClearListenerAttenuationOverride(ViewportIndex);
+								}
+							}
 						}
 					}
 
+					// Add view information for resource streaming. Allow up to 5X boost for small FOV.
+					const float StreamingScale = 1.f / FMath::Clamp<float>(View->LODDistanceFactor, .2f, 1.f);
+					IStreamingManager::Get().AddViewInformation(View->ViewMatrices.GetViewOrigin(), View->UnscaledViewRect.Width(), View->UnscaledViewRect.Width() * View->ViewMatrices.GetProjectionMatrix().M[0][0], StreamingScale);
+					MyWorld->ViewLocationsRenderedLastFrame.Add(View->ViewMatrices.GetViewOrigin());
+				}
+			}
+
 #if CSV_PROFILER
-				UpdateCsvCameraStats(PlayerViewMap);
+			UpdateCsvCameraStats(PlayerViewMap);
 #endif
 			if (ViewFamily.Views.Num() > 0)
 			{
@@ -459,7 +459,7 @@ void UDisplayClusterViewportClient::Draw(FViewport* InViewport, FCanvas* SceneCa
 					TEXT("Some code has tried to set up an alien screen percentage driver, that could be wrong if not supported very well by the RHI."));
 
 				// Setup main view family with screen percentage interface by dynamic resolution if screen percentage is enabled.
-#if WITH_DYNAMIC_RESOLUTION
+	#if WITH_DYNAMIC_RESOLUTION
 				if (ViewFamily.EngineShowFlags.ScreenPercentage)
 				{
 					FDynamicResolutionStateInfos DynamicResolutionStateInfos;
@@ -481,14 +481,14 @@ void UDisplayClusterViewportClient::Draw(FViewport* InViewport, FCanvas* SceneCa
 							DynamicResolutionStateInfos.ResolutionFractionUpperBound));
 					}
 
-#if CSV_PROFILER
+	#if CSV_PROFILER
 					if (DynamicResolutionStateInfos.ResolutionFractionApproximation >= 0.0f)
 					{
 						CSV_CUSTOM_STAT_GLOBAL(DynamicResolutionPercentage, DynamicResolutionStateInfos.ResolutionFractionApproximation * 100.0f, ECsvCustomStatOp::Set);
 					}
-#endif
+	#endif
 				}
-#endif
+	#endif
 
 				// If a screen percentage interface was not set by dynamic resolution, then create one matching legacy behavior.
 				if (ViewFamily.GetScreenPercentageInterface() == nullptr)
diff --git a/Engine/Plugins/VirtualProduction/RemoteControlWebInterface/WebApp/Client/package.json b/Engine/Plugins/VirtualProduction/RemoteControlWebInterface/WebApp/Client/package.json
deleted file mode 100644
diff --git a/Engine/Plugins/VirtualProduction/RemoteControlWebInterface/WebApp/Server/package.json b/Engine/Plugins/VirtualProduction/RemoteControlWebInterface/WebApp/Server/package.json
deleted file mode 100644
diff --git a/Engine/Source/Runtime/D3D12RHI/Private/D3D12Stats.cpp b/Engine/Source/Runtime/D3D12RHI/Private/D3D12Stats.cpp
--- a/Engine/Source/Runtime/D3D12RHI/Private/D3D12Stats.cpp
+++ b/Engine/Source/Runtime/D3D12RHI/Private/D3D12Stats.cpp
@@ -378,19 +378,19 @@ bool D3D12RHI::FD3DGPUProfiler::CheckGpuHeartbeat() const
 					{
 						GFSDK_Aftermath_ContextData& ContextData = ContextDataOut[ContextIdx];
 						uint32 NumMarkers = ContextData.markerSize / sizeof(uint32);
-							uint32* Data = (uint32*)ContextData.markerData;
+						uint32* Data = (uint32*)ContextData.markerData;
 
 						const TCHAR* StatusNames[] = { TEXT("NotStarted"), TEXT("Executing"), TEXT("Finished"), TEXT("Invalid") };
 						const TCHAR* ContextStatusName = ContextData.status < UE_ARRAY_COUNT(StatusNames) ? StatusNames[ContextData.status] : TEXT("UNKNOWN");
 						UE_LOG(LogRHI, Error, TEXT("[Aftermath] Context %d, command list %016llX, status %s, %u markers. Begin GPU Stack Dump"), ContextIdx, AftermathCommandLists[ContextIdx], ContextStatusName, NumMarkers);
 						for (uint32 MarkerIdx = 0; MarkerIdx < NumMarkers; ++MarkerIdx)
-							{
+						{
 							const FString* MarkerName = CachedEventStrings.Find(Data[MarkerIdx]);
 							UE_LOG(LogRHI, Error, TEXT("[Aftermath] %d: %s"), MarkerIdx, MarkerName ? *(*MarkerName) : TEXT("NULL"));
-								}
-								UE_LOG(LogRHI, Error, TEXT("[Aftermath] End GPU Stack Dump"));
-							}
 						}
+						UE_LOG(LogRHI, Error, TEXT("[Aftermath] End GPU Stack Dump"));
+					}
+				}
 				else
 				{
 					UE_LOG(LogRHI, Error, TEXT("[Aftermath] Failed to get Aftermath stack data"));
diff --git a/Engine/Source/Runtime/Experimental/Chaos/Public/Chaos/PBDSphericalConstraint.h b/Engine/Source/Runtime/Experimental/Chaos/Public/Chaos/PBDSphericalConstraint.h
--- a/Engine/Source/Runtime/Experimental/Chaos/Public/Chaos/PBDSphericalConstraint.h
+++ b/Engine/Source/Runtime/Experimental/Chaos/Public/Chaos/PBDSphericalConstraint.h
@@ -132,34 +132,34 @@ namespace Chaos
 
 			if (bEnabled)
 			{
-			if (bUseLegacyBackstop)
-			{
-				// SphereOffsetDistances includes the sphere radius
-				// This is harder to author, and does not follow the NvCloth specs.
-				// However, this is how it's been done in the Unreal Engine PhysX cloth implementation.
-				if (bRealTypeCompatibleWithISPC && bChaos_Spherical_ISPC_Enabled)
-				{
-					ApplyLegacyHelperISPC(Particles, Dt);
-				}
-				else
-				{
-					ApplyLegacyHelper(Particles, Dt);
-				}
-			}
-			else
-			{
-				// SphereOffsetDistances doesn't include the sphere radius
-				if (bRealTypeCompatibleWithISPC && bChaos_Spherical_ISPC_Enabled)
+				if (bUseLegacyBackstop)
 				{
-					ApplyHelperISPC(Particles, Dt);
+					// SphereOffsetDistances includes the sphere radius
+					// This is harder to author, and does not follow the NvCloth specs.
+					// However, this is how it's been done in the Unreal Engine PhysX cloth implementation.
+					if (bRealTypeCompatibleWithISPC && bChaos_Spherical_ISPC_Enabled)
+					{
+						ApplyLegacyHelperISPC(Particles, Dt);
+					}
+					else
+					{
+						ApplyLegacyHelper(Particles, Dt);
+					}
 				}
 				else
 				{
-					ApplyHelper(Particles, Dt);
+					// SphereOffsetDistances doesn't include the sphere radius
+					if (bRealTypeCompatibleWithISPC && bChaos_Spherical_ISPC_Enabled)
+					{
+						ApplyHelperISPC(Particles, Dt);
+					}
+					else
+					{
+						ApplyHelper(Particles, Dt);
+					}
 				}
 			}
 		}
-		}
 
 		void SetSphereRadiiMultiplier(const FReal InSphereRadiiMultiplier)
 		{
diff --git a/Engine/Source/Runtime/Experimental/GeometryCollectionEngine/Private/GeometryCollection/GeometryCollectionComponent.cpp b/Engine/Source/Runtime/Experimental/GeometryCollectionEngine/Private/GeometryCollection/GeometryCollectionComponent.cpp
--- a/Engine/Source/Runtime/Experimental/GeometryCollectionEngine/Private/GeometryCollection/GeometryCollectionComponent.cpp
+++ b/Engine/Source/Runtime/Experimental/GeometryCollectionEngine/Private/GeometryCollection/GeometryCollectionComponent.cpp
@@ -423,7 +423,7 @@ FBoxSphereBounds UGeometryCollectionComponent::CalcBounds(const FTransform& Loca
 		const TManagedArray<FTransform>& Transforms = GetTransformArray();
 
 		const int32 NumBoxes = BoundingBoxes.Num();
-
+	
 		int32 NumElements = HackGeometryCollectionPtr->NumElements(FGeometryCollection::TransformGroup);
 		if (RestCollection->EnableNanite && HackGeometryCollectionPtr->HasAttribute("BoundingBox", FGeometryCollection::TransformGroup) && NumElements)
 		{
@@ -895,14 +895,14 @@ void UGeometryCollectionComponent::RefreshEmbeddedGeometry()
 			{
 				if (!HideArray || !(*HideArray)[Idx])
 				{ 
-				InstanceTransforms.Add(FTransform(GlobalMatrices[Idx]));
+					InstanceTransforms.Add(FTransform(GlobalMatrices[Idx]));
 #if WITH_EDITOR
 					int32 InstanceIndex = EmbeddedBoneMaps[ExemplarIndex].Add(Idx);
 					EmbeddedInstanceIndex[Idx] = InstanceIndex;
 #endif
+				}
 			}
 		}
-		}
 
 		if (EmbeddedGeometryComponents[ExemplarIndex])
 		{
@@ -1112,20 +1112,20 @@ void UGeometryCollectionComponent::RegisterForEvents()
 #if INCLUDE_CHAOS
 		Chaos::FPhysicsSolver* Solver = GetWorld()->GetPhysicsScene()->GetSolver();
 		if (Solver)
-	{
-		if (bNotifyCollisions || BodyInstance.bNotifyRigidBodyCollision)
 		{
-			EventDispatcher->RegisterForCollisionEvents(this, this);
+			if (bNotifyCollisions || BodyInstance.bNotifyRigidBodyCollision)
+			{
+				EventDispatcher->RegisterForCollisionEvents(this, this);
 
 				Solver->EnqueueCommandImmediate([Solver]()
 					{
 						Solver->SetGenerateCollisionData(true);
 					});
-		}
+			}
 
-		if (bNotifyBreaks)
-		{
-			EventDispatcher->RegisterForBreakEvents(this, &DispatchGeometryCollectionBreakEvent);
+			if (bNotifyBreaks)
+			{
+				EventDispatcher->RegisterForBreakEvents(this, &DispatchGeometryCollectionBreakEvent);
 
 				Solver->EnqueueCommandImmediate([Solver]()
 					{
@@ -1464,7 +1464,7 @@ void UGeometryCollectionComponent::InitConstantData(FGeometryCollectionConstantD
 		const TManagedArray<TArray<FVector2f>>& UVs = Collection->UVs;
 		const TManagedArray<FLinearColor>& Color = Collection->Color;
 		const TManagedArray<FLinearColor>& BoneColors = Collection->BoneColor;
-
+		
 		const int32 NumGeom = Collection->NumElements(FGeometryCollection::GeometryGroup);
 		const TManagedArray<int32>& TransformIndex = Collection->TransformIndex;
 		const TManagedArray<int32>& FaceStart = Collection->FaceStart;
@@ -1601,7 +1601,7 @@ void UGeometryCollectionComponent::InitConstantData(FGeometryCollectionConstantD
 
 		// store the index buffer and render sections for the base unfractured mesh
 		const TManagedArray<int32>& TransformToGeometryIndex = Collection->TransformToGeometryIndex;
-
+		
 		const int32 NumFaces = Collection->NumElements(FGeometryCollection::FacesGroup);
 		TArray<FIntVector> BaseMeshIndices;
 		TArray<int32> BaseMeshOriginalFaceIndices;
@@ -2472,30 +2472,30 @@ void FScopedColorEdit::SelectBones(GeometryCollection::ESelectionMode SelectionM
 			FGeometryCollectionProximityUtility ProximityUtility(GeometryCollectionPtr.Get());
 			ProximityUtility.UpdateProximity();
 
-				const TManagedArray<int32>& TransformIndex = GeometryCollectionPtr->TransformIndex;
-				const TManagedArray<int32>& TransformToGeometryIndex = GeometryCollectionPtr->TransformToGeometryIndex;
-				const TManagedArray<TSet<int32>>& Proximity = GeometryCollectionPtr->GetAttribute<TSet<int32>>("Proximity", FGeometryCollection::GeometryGroup);
+			const TManagedArray<int32>& TransformIndex = GeometryCollectionPtr->TransformIndex;
+			const TManagedArray<int32>& TransformToGeometryIndex = GeometryCollectionPtr->TransformToGeometryIndex;
+			const TManagedArray<TSet<int32>>& Proximity = GeometryCollectionPtr->GetAttribute<TSet<int32>>("Proximity", FGeometryCollection::GeometryGroup);
 
-				const TArray<int32> SelectedBones = GetSelectedBones();
+			const TArray<int32> SelectedBones = GetSelectedBones();
 
-				TArray<int32> NewSelection;
-				for (int32 Bone : SelectedBones)
+			TArray<int32> NewSelection;
+			for (int32 Bone : SelectedBones)
+			{
+				NewSelection.AddUnique(Bone);
+				int32 GeometryIdx = TransformToGeometryIndex[Bone];
+				if (GeometryIdx != INDEX_NONE)
 				{
-					NewSelection.AddUnique(Bone);
-					int32 GeometryIdx = TransformToGeometryIndex[Bone];
-					if (GeometryIdx != INDEX_NONE)
+					const TSet<int32>& Neighbors = Proximity[GeometryIdx];
+					for (int32 NeighborGeometryIndex : Neighbors)
 					{
-						const TSet<int32>& Neighbors = Proximity[GeometryIdx];
-						for (int32 NeighborGeometryIndex : Neighbors)
-						{
-							NewSelection.AddUnique(TransformIndex[NeighborGeometryIndex]);
-						}
+						NewSelection.AddUnique(TransformIndex[NeighborGeometryIndex]);
 					}
 				}
-
-				ResetBoneSelection();
-				AppendSelectedBones(NewSelection);
 			}
+
+			ResetBoneSelection();
+			AppendSelectedBones(NewSelection);
+		}
 		break;
 
 		case GeometryCollection::ESelectionMode::Parent:
@@ -2568,27 +2568,27 @@ void FScopedColorEdit::SelectBones(GeometryCollection::ESelectionMode SelectionM
 		case GeometryCollection::ESelectionMode::Level:
 		{
 			if (GeometryCollectionPtr->HasAttribute("Level", FTransformCollection::TransformGroup))
-		{
+			{
 				const TManagedArray<int32>& Levels = GeometryCollectionPtr->GetAttribute<int32>("Level", FTransformCollection::TransformGroup);
 
-			const TArray<int32> SelectedBones = GetSelectedBones();
+				const TArray<int32> SelectedBones = GetSelectedBones();
 
-			TArray<int32> NewSelection;
-			for (int32 Bone : SelectedBones)
-			{
+				TArray<int32> NewSelection;
+				for (int32 Bone : SelectedBones)
+				{
 					int32 Level = Levels[Bone];
 					for (int32 TransformIdx = 0; TransformIdx < GeometryCollectionPtr->NumElements(FTransformCollection::TransformGroup); ++TransformIdx)
 					{
 						if (Levels[TransformIdx] == Level)
-				{
+						{
 							NewSelection.AddUnique(TransformIdx);
 						}
+					}
 				}
-			}
 
-			ResetBoneSelection();
-			AppendSelectedBones(NewSelection);
-		}
+				ResetBoneSelection();
+				AppendSelectedBones(NewSelection);
+			}	
 		}
 		break;
 
@@ -2932,10 +2932,10 @@ void UGeometryCollectionComponent::CalculateGlobalMatrices()
 			}
 			else
 			{ 
-			GeometryCollectionAlgo::GlobalMatrices(GetTransformArray(), GetParentArray(), GlobalMatrices);
+				GeometryCollectionAlgo::GlobalMatrices(GetTransformArray(), GetParentArray(), GlobalMatrices);		
+			}
 		}
 	}
-	}
 	
 #if WITH_EDITOR
 	if (GlobalMatrices.Num() > 0)
