Sequencer: Control Rig: Make sure pose is saved in default parent space, and when pasting compensate if not in default parent space.

diff --git a/Engine/Plugins/Animation/ControlRig/Source/ControlRig/Private/Tools/ControlRigPose.cpp b/Engine/Plugins/Animation/ControlRig/Source/ControlRig/Private/Tools/ControlRigPose.cpp
--- a/Engine/Plugins/Animation/ControlRig/Source/ControlRig/Private/Tools/ControlRigPose.cpp
+++ b/Engine/Plugins/Animation/ControlRig/Source/ControlRig/Private/Tools/ControlRigPose.cpp
@@ -18,8 +18,29 @@ void FControlRigControlPose::SavePose(UControlRig* ControlRig, bool bUseAll)
 	{
 		if (ControlElement->Settings.bAnimatable && (bUseAll || ControlRig->IsControlSelected(ControlElement->GetName())))
 		{
+			//we store poses in default parent space so if not in that space we need to compensate it
+			bool bHasNonDefaultParent = false;
+			FTransform GlobalTransform;
+			FRigElementKey SpaceKey = ControlRig->GetHierarchy()->GetActiveParent(ControlElement->GetKey());
+			if (SpaceKey != ControlRig->GetHierarchy()->GetDefaultParentKey())
+			{
+				bHasNonDefaultParent = true;
+				//to compensate we get the global, switch space, then reset global in that new space
+				GlobalTransform = ControlRig->GetHierarchy()->GetGlobalTransform(ControlElement->GetKey());
+				ControlRig->GetHierarchy()->SwitchToDefaultParent(ControlElement->GetKey());
+				ControlRig->GetHierarchy()->SetGlobalTransform(ControlElement->GetKey(),GlobalTransform);
+
+				ControlRig->Evaluate_AnyThread();
+			}
 			FRigControlCopy Copy(ControlElement, ControlRig->GetHierarchy());
 			CopyOfControls.Add(Copy);
+	
+			if (bHasNonDefaultParent == true) 
+			{
+				ControlRig->GetHierarchy()->SwitchToParent(ControlElement->GetKey(),SpaceKey);
+				ControlRig->GetHierarchy()->SetGlobalTransform(ControlElement->GetKey(), GlobalTransform);
+			}
+			
 		}
 	}
 	SetUpControlMap();
@@ -74,6 +95,16 @@ void FControlRigControlPose::PastePoseInternal(UControlRig* ControlRig, bool bDo
 		FRigControlCopy* CopyRigControl = MirrorTable.GetControl(*this, ControlElement->GetName());
 		if (CopyRigControl)
 		{
+			//if not in default parent space we need to move it to default parent space first and then reset the global transforms
+			bool bHasNonDefaultParent = false;
+			FRigElementKey SpaceKey = ControlRig->GetHierarchy()->GetActiveParent(ControlElement->GetKey());
+			if (SpaceKey != ControlRig->GetHierarchy()->GetDefaultParentKey())
+			{
+				bHasNonDefaultParent = true;
+				ControlRig->GetHierarchy()->SwitchToDefaultParent(ControlElement->GetKey());
+				ControlRig->Evaluate_AnyThread();
+			}
+
 			switch (ControlElement->Settings.ControlType)
 			{
 			case ERigControlType::Transform:
@@ -134,6 +165,13 @@ void FControlRigControlPose::PastePoseInternal(UControlRig* ControlRig, bool bDo
 				//TODO add log
 				break;
 			};
+
+			if (bHasNonDefaultParent == true)
+			{
+				FTransform GlobalTransform = ControlRig->GetHierarchy()->GetGlobalTransform(ControlElement->GetKey());
+				ControlRig->GetHierarchy()->SwitchToParent(ControlElement->GetKey(), SpaceKey);
+				ControlRig->SetControlGlobalTransform(ControlElement->GetName(), GlobalTransform, true, Context, bSetupUndo);
+			}
 		}
 	}
 }
@@ -182,6 +220,15 @@ void FControlRigControlPose::BlendWithInitialPoses(FControlRigControlPose& Initi
 					CopyRigControl->ControlType == ERigControlType::Rotator || CopyRigControl->ControlType == ERigControlType::Scale
 					))
 				{
+					//if not in default parent space we need to move it to default parent space first and then reset the global transforms
+					bool bHasNonDefaultParent = false;
+					FRigElementKey SpaceKey = ControlRig->GetHierarchy()->GetActiveParent(ControlElement->GetKey());
+					if (SpaceKey != ControlRig->GetHierarchy()->GetDefaultParentKey())
+					{
+						bHasNonDefaultParent = true;
+						ControlRig->GetHierarchy()->SwitchToDefaultParent(ControlElement->GetKey());
+						ControlRig->Evaluate_AnyThread();
+					}
 					if (bDoMirror == false)
 					{
 						if (bDoLocal == true)    // -V547  
@@ -227,6 +274,12 @@ void FControlRigControlPose::BlendWithInitialPoses(FControlRigControlPose& Initi
 						LocalRotation = FQuat::Slerp(InitialLocationRotation, LocalRotation, BlendValue); //doing slerp here not fast lerp, can be slow this is for content creation
 						SetControlMirrorTransform(bDoLocal,ControlRig, ControlElement->GetName(), bIsMatched, GlobalTranslation, GlobalRotation,LocalTranslation, LocalRotation, bDoKey,Context,bSetupUndo);							
 					}
+					if (bHasNonDefaultParent == true)
+					{
+						FTransform GlobalTransform = ControlRig->GetHierarchy()->GetGlobalTransform(ControlElement->GetKey());
+						ControlRig->GetHierarchy()->SwitchToParent(ControlElement->GetKey(), SpaceKey);
+						ControlRig->SetControlGlobalTransform(ControlElement->GetName(), GlobalTransform, true, Context, bSetupUndo);
+					}
 				}
 			}
 		}
