Sequencer: Anim Layer Improvments Pt 2. Added functions to Collapse all layers onto a Control Rig, and added functions to set zero values and weight values on the additive sections based upon animator feedback.
#prefight 61b7c32f1a4451f7ba1d71ae

diff --git a/Engine/Plugins/Animation/ControlRig/Source/ControlRig/Private/Sequencer/MovieSceneControlRigParameterSection.cpp b/Engine/Plugins/Animation/ControlRig/Source/ControlRig/Private/Sequencer/MovieSceneControlRigParameterSection.cpp
--- a/Engine/Plugins/Animation/ControlRig/Source/ControlRig/Private/Sequencer/MovieSceneControlRigParameterSection.cpp
+++ b/Engine/Plugins/Animation/ControlRig/Source/ControlRig/Private/Sequencer/MovieSceneControlRigParameterSection.cpp
@@ -1915,6 +1915,81 @@ float UMovieSceneControlRigParameterSection::GetTotalWeightValue(FFrameTime InTi
 	return WeightVal;
 }
 
+void UMovieSceneControlRigParameterSection::KeyZeroValue(FFrameNumber InFrame, bool bSelectedControls)
+{
+	TArray<FName> SelectedControls;
+	if (bSelectedControls && ControlRig)
+	{
+		SelectedControls = ControlRig->CurrentControlSelection();
+	}
+	/* Don't set zero values on these doesn't make sense
+	
+	for (FBoolParameterNameAndCurve& Bool : GetBoolParameterNamesAndCurves())
+	for (FEnumParameterNameAndCurve& Enum : GetEnumParameterNamesAndCurves())
+	for (FIntegerParameterNameAndCurve& Integer : GetIntegerParameterNamesAndCurves())
+
+	*/
+	for (FScalarParameterNameAndCurve& Scalar : GetScalarParameterNamesAndCurves())
+	{
+		if (SelectedControls.Num() == 0 || SelectedControls.Contains(Scalar.ParameterName))
+		{
+			Scalar.ParameterCurve.AddCubicKey(InFrame, 0.0f, ERichCurveTangentMode::RCTM_Auto);
+			Scalar.ParameterCurve.AutoSetTangents();
+		}
+	}
+	for (FVector2DParameterNameAndCurves& Vector2D : GetVector2DParameterNamesAndCurves())
+	{
+		if (SelectedControls.Num() == 0 || SelectedControls.Contains(Vector2D.ParameterName))
+		{
+			Vector2D.XCurve.AddCubicKey(InFrame, 0.0f, ERichCurveTangentMode::RCTM_Auto);
+			Vector2D.XCurve.AutoSetTangents();
+			Vector2D.YCurve.AddCubicKey(InFrame, 0.0f, ERichCurveTangentMode::RCTM_Auto);
+			Vector2D.YCurve.AutoSetTangents();
+		}
+	}
+	for (FVectorParameterNameAndCurves& Vector : GetVectorParameterNamesAndCurves())
+	{
+		if (SelectedControls.Num() == 0 || SelectedControls.Contains(Vector.ParameterName))
+		{
+			Vector.XCurve.AddCubicKey(InFrame, 0.0f, ERichCurveTangentMode::RCTM_Auto);
+			Vector.XCurve.AutoSetTangents();
+			Vector.YCurve.AddCubicKey(InFrame, 0.0f, ERichCurveTangentMode::RCTM_Auto);
+			Vector.YCurve.AutoSetTangents();
+			Vector.ZCurve.AddCubicKey(InFrame, 0.0f, ERichCurveTangentMode::RCTM_Auto);
+			Vector.ZCurve.AutoSetTangents();
+		}
+	}
+	for (FTransformParameterNameAndCurves& Transform : GetTransformParameterNamesAndCurves())
+	{
+		if (SelectedControls.Num() == 0 || SelectedControls.Contains(Transform.ParameterName))
+		{
+			for (int32 Index = 0; Index < 3; ++Index)
+			{
+				Transform.Translation[Index].AddCubicKey(InFrame, 0.0f, ERichCurveTangentMode::RCTM_Auto);
+				Transform.Translation[Index].AutoSetTangents();
+				Transform.Rotation[Index].AddCubicKey(InFrame, 0.0f, ERichCurveTangentMode::RCTM_Auto);
+				Transform.Rotation[Index].AutoSetTangents();
+				if (GetBlendType() == EMovieSceneBlendType::Additive)
+				{
+					Transform.Scale[Index].AddCubicKey(InFrame, 0.0f, ERichCurveTangentMode::RCTM_Auto);
+				}
+				else
+				{
+					Transform.Scale[Index].AddCubicKey(InFrame, 1.0f, ERichCurveTangentMode::RCTM_Auto);
+				}
+				Transform.Scale[Index].AutoSetTangents();
+
+			}
+		}
+	}
+}
+
+void UMovieSceneControlRigParameterSection::KeyWeightValue(FFrameNumber InFrame, float InVal)
+{
+	Weight.AddCubicKey(InFrame, InVal, ERichCurveTangentMode::RCTM_Auto);
+	Weight.AutoSetTangents();
+}
+
 void UMovieSceneControlRigParameterSection::RecreateWithThisControlRig(UControlRig* InControlRig, bool bSetDefault)
 {
 	bool bSameControlRig = (ControlRig == InControlRig);
@@ -2609,8 +2684,98 @@ void UMovieSceneControlRigParameterSection::ClearAllParameters()
 	TransformParameterNamesAndCurves.SetNum(0);
 	EnumParameterNamesAndCurves.SetNum(0);
 	IntegerParameterNamesAndCurves.SetNum(0);
+	SpaceChannels.SetNum(0);
+}
+void UMovieSceneControlRigParameterSection::RemoveAllKeys(bool bIncludeSpaceKeys)
+{
+	TArray<FFrameNumber> KeyTimes;
+	TArray<FKeyHandle> Handles;
+	if (bIncludeSpaceKeys)
+	{
+		for (FSpaceControlNameAndChannel& Space : SpaceChannels)
+		{
+			KeyTimes.SetNum(0);
+			Handles.SetNum(0);
+			Space.SpaceCurve.GetKeys(TRange<FFrameNumber>(), &KeyTimes, &Handles);
+			Space.SpaceCurve.DeleteKeys(Handles);
+		}
+	}
+	for (FBoolParameterNameAndCurve& Bool : GetBoolParameterNamesAndCurves())
+	{
+		KeyTimes.SetNum(0);
+		Handles.SetNum(0);
+		Bool.ParameterCurve.GetKeys(TRange<FFrameNumber>(), &KeyTimes, &Handles);
+		Bool.ParameterCurve.DeleteKeys(Handles);
+	}
+	for (FEnumParameterNameAndCurve& Enum : GetEnumParameterNamesAndCurves())
+	{
+		KeyTimes.SetNum(0);
+		Handles.SetNum(0);
+		Enum.ParameterCurve.GetKeys(TRange<FFrameNumber>(), &KeyTimes, &Handles);
+		Enum.ParameterCurve.DeleteKeys(Handles);
+	}
+	for (FIntegerParameterNameAndCurve& Integer : GetIntegerParameterNamesAndCurves())
+	{
+		KeyTimes.SetNum(0);
+		Handles.SetNum(0);
+		Integer.ParameterCurve.GetKeys(TRange<FFrameNumber>(), &KeyTimes, &Handles);
+		Integer.ParameterCurve.DeleteKeys(Handles);
+	}
+
+	for (FScalarParameterNameAndCurve& Scalar : GetScalarParameterNamesAndCurves())
+	{
+		KeyTimes.SetNum(0);
+		Handles.SetNum(0);
+		Scalar.ParameterCurve.GetKeys(TRange<FFrameNumber>(), &KeyTimes, &Handles);
+		Scalar.ParameterCurve.DeleteKeys(Handles);
+	}
+	for (FVector2DParameterNameAndCurves& Vector2D : GetVector2DParameterNamesAndCurves())
+	{
+		KeyTimes.SetNum(0);
+		Handles.SetNum(0);
+		Vector2D.XCurve.GetKeys(TRange<FFrameNumber>(), &KeyTimes, &Handles);
+		Vector2D.XCurve.DeleteKeys(Handles);
+		KeyTimes.SetNum(0);
+		Handles.SetNum(0);
+		Vector2D.YCurve.GetKeys(TRange<FFrameNumber>(), &KeyTimes, &Handles);
+		Vector2D.YCurve.DeleteKeys(Handles);
+	}
+	for (FVectorParameterNameAndCurves& Vector : GetVectorParameterNamesAndCurves())
+	{
+		KeyTimes.SetNum(0);
+		Handles.SetNum(0);
+		Vector.XCurve.GetKeys(TRange<FFrameNumber>(), &KeyTimes, &Handles);
+		Vector.XCurve.DeleteKeys(Handles);
+		KeyTimes.SetNum(0);
+		Handles.SetNum(0);
+		Vector.YCurve.GetKeys(TRange<FFrameNumber>(), &KeyTimes, &Handles);
+		Vector.YCurve.DeleteKeys(Handles);
+		KeyTimes.SetNum(0);
+		Handles.SetNum(0);
+		Vector.ZCurve.GetKeys(TRange<FFrameNumber>(), &KeyTimes, &Handles);
+		Vector.ZCurve.DeleteKeys(Handles);
+	}
+	for (FTransformParameterNameAndCurves& Transform : GetTransformParameterNamesAndCurves())
+	{
+		for (int32 Index = 0; Index < 3; ++Index)
+		{
+			KeyTimes.SetNum(0);
+			Handles.SetNum(0);
+			Transform.Translation[Index].GetKeys(TRange<FFrameNumber>(), &KeyTimes, &Handles);
+			Transform.Translation[Index].DeleteKeys(Handles);
+			KeyTimes.SetNum(0);
+			Handles.SetNum(0);
+			Transform.Rotation[Index].GetKeys(TRange<FFrameNumber>(), &KeyTimes, &Handles);
+			Transform.Rotation[Index].DeleteKeys(Handles);
+			KeyTimes.SetNum(0);
+			Handles.SetNum(0);
+			Transform.Scale[Index].GetKeys(TRange<FFrameNumber>(), &KeyTimes, &Handles);
+			Transform.Scale[Index].DeleteKeys(Handles);
+		}
+	}
 }
 
+
 int32 UMovieSceneControlRigParameterSection::GetActiveCategoryIndex(FName ControlName) const
 {
 	int32 CategoryIndex = INDEX_NONE;
diff --git a/Engine/Plugins/Animation/ControlRig/Source/ControlRig/Public/Sequencer/MovieSceneControlRigParameterSection.h b/Engine/Plugins/Animation/ControlRig/Source/ControlRig/Public/Sequencer/MovieSceneControlRigParameterSection.h
--- a/Engine/Plugins/Animation/ControlRig/Source/ControlRig/Public/Sequencer/MovieSceneControlRigParameterSection.h
+++ b/Engine/Plugins/Animation/ControlRig/Source/ControlRig/Public/Sequencer/MovieSceneControlRigParameterSection.h
@@ -417,6 +417,15 @@ public:
 	/** Evaluates a a key for a specific space parameter. Will not get set if not found */
 	TOptional<FMovieSceneControlRigSpaceBaseKey> EvaluateSpaceChannel(const  FFrameTime& InTime, FName InParameterName);
 
+	/** Key Zero Values on all or just selected controls in these section at the specified time */
+	void KeyZeroValue(FFrameNumber InFrame, bool bSelected);
+
+	/** Key the Weights to the specified value */
+	void KeyWeightValue(FFrameNumber InFrame, float InVal);
+;
+	/** Remove All Keys, but maybe not space keys if bIncludeSpaceKeys is false */
+	void RemoveAllKeys(bool bIncludeSpaceKeys);
+
 public:
 	/**
 	* Access the interrogation key for control rig data 
diff --git a/Engine/Plugins/Animation/ControlRig/Source/ControlRigEditor/Private/Sequencer/ControlRigParameterTrackEditor.cpp b/Engine/Plugins/Animation/ControlRig/Source/ControlRigEditor/Private/Sequencer/ControlRigParameterTrackEditor.cpp
--- a/Engine/Plugins/Animation/ControlRig/Source/ControlRigEditor/Private/Sequencer/ControlRigParameterTrackEditor.cpp
+++ b/Engine/Plugins/Animation/ControlRig/Source/ControlRigEditor/Private/Sequencer/ControlRigParameterTrackEditor.cpp
@@ -3373,6 +3373,128 @@ void FControlRigParameterTrackEditor::SelectFKBonesToAnimate(UFKControlRig* Auto
 	//reconstruct all channel proxies TODO or not to do that is the question
 }
 
+void FControlRigParameterSection::CollapseAllLayers()
+{
+	FScopedTransaction Transaction(LOCTEXT("CollapseAllSections", "CollapseAllSections"));
+
+	UMovieSceneControlRigParameterSection* ParameterSection = CastChecked<UMovieSceneControlRigParameterSection>(WeakSection.Get());
+	if (ParameterSection && ParameterSection->GetControlRig())
+	{
+		UControlRig* ControlRig = ParameterSection->GetControlRig();
+		TSharedPtr<ISequencer> SequencerPtr = WeakSequencer.Pin();
+		TRange<FFrameNumber> Range = SequencerPtr->GetFocusedMovieSceneSequence()->GetMovieScene()->GetPlaybackRange();
+		FFrameNumber StartFrame = Range.GetLowerBoundValue();
+		FFrameNumber EndFrame = Range.GetUpperBoundValue();
+		const FFrameRate& FrameRate = SequencerPtr->GetFocusedDisplayRate();
+		const FFrameRate& TickResolution = SequencerPtr->GetFocusedTickResolution();
+
+		FFrameNumber FrameRateInFrameNumber = TickResolution.AsFrameNumber(FrameRate.AsInterval());
+		TArray<FFrameNumber> Frames;
+		for (FFrameNumber& Frame = StartFrame; Frame <= EndFrame; Frame += FrameRateInFrameNumber)
+		{
+			Frames.Add(Frame);
+		}
+		//Store transforms
+		TArray<TPair<FName, TArray<FTransform>>> ControlLocalTransforms;
+		TArray<FRigControlElement*> Controls;
+		ControlRig->GetControlsInOrder(Controls);
+
+		for (FRigControlElement* ControlElement : Controls)
+		{
+			if (!ControlElement->Settings.bAnimatable)
+			{
+				continue;
+			}
+			TPair<FName, TArray<FTransform>> NameTransforms;
+			NameTransforms.Key = ControlElement->GetName();
+			NameTransforms.Value.SetNum(Frames.Num());
+			ControlLocalTransforms.Add(NameTransforms);
+		}
+
+		//get all of the local 
+		int32 Index = 0;
+		for (Index = 0; Index < Frames.Num(); ++Index)
+		{
+			const FFrameNumber& FrameNumber = Frames[Index];
+			FFrameTime GlobalTime(FrameNumber);
+
+			FMovieSceneContext Context = FMovieSceneContext(FMovieSceneEvaluationRange(GlobalTime, TickResolution), SequencerPtr->GetPlaybackStatus()).SetHasJumped(true);
+
+			SequencerPtr->GetEvaluationTemplate().Evaluate(Context, *SequencerPtr);
+			ControlRig->Evaluate_AnyThread();
+			for (TPair<FName, TArray<FTransform>>& TrailControlTransform : ControlLocalTransforms)
+			{
+				TrailControlTransform.Value[Index] = ControlRig->GetControlLocalTransform(TrailControlTransform.Key);
+			}
+		}
+		//delete other sections
+		UMovieSceneTrack* OwnerTrack = ParameterSection->GetTypedOuter<UMovieSceneTrack>();
+		OwnerTrack->Modify();
+		const TArray<UMovieSceneSection*>& Sections = OwnerTrack->GetAllSections();
+		for (Index = Sections.Num() - 1; Index >= 0; --Index)
+		{
+			if (Sections[Index] != ParameterSection)
+			{
+				OwnerTrack->RemoveSectionAt(Index);
+			}
+		}
+
+		ParameterSection->RemoveAllKeys(false); //remove all keys but the space keys.
+		FRigControlModifiedContext Context;
+		Context.SetKey = EControlRigSetKey::Always;
+
+		Index = 0;
+		for (Index = 0; Index < Frames.Num(); ++Index)
+		{
+			const FFrameNumber& FrameNumber = Frames[Index];
+			Context.LocalTime = TickResolution.AsSeconds(FFrameTime(FrameNumber));
+			//need to do the twice hack since controls aren't really in order
+			for (int32 TwiceHack = 0; TwiceHack < 2; ++TwiceHack)
+			{
+				for (TPair<FName, TArray<FTransform>>& TrailControlTransform : ControlLocalTransforms)
+				{
+					ControlRig->SetControlLocalTransform(TrailControlTransform.Key, TrailControlTransform.Value[Index],true,Context,false);
+				}
+			}
+		}
+
+		//reset everything back
+		SequencerPtr->NotifyMovieSceneDataChanged(EMovieSceneDataChangeType::MovieSceneStructureItemAdded);
+		/*
+			const FFrameTime StartTime = SequencerPtr->GetLocalTime().Time;
+			FMovieSceneContext Context = FMovieSceneContext(FMovieSceneEvaluationRange(StartTime, TickResolution), SequencerPtr->GetPlaybackStatus()).SetHasJumped(true);
+			SequencerPtr->GetEvaluationTemplate().Evaluate(Context, *SequencerPtr);
+			ControlRig->Evaluate_AnyThread();
+			*/
+	}
+
+}
+void FControlRigParameterSection::KeyZeroValue()
+{
+	UMovieSceneControlRigParameterSection* ParameterSection = CastChecked<UMovieSceneControlRigParameterSection>(WeakSection.Get());
+	TSharedPtr<ISequencer> SequencerPtr = WeakSequencer.Pin();
+	FScopedTransaction Transaction(LOCTEXT("KeyZeroValue", "Key Zero Value"));
+	ParameterSection->Modify();
+	FFrameTime Time = SequencerPtr->GetLocalTime().Time;
+	ParameterSection->KeyZeroValue(Time.GetFrame(),true);
+	SequencerPtr->NotifyMovieSceneDataChanged(EMovieSceneDataChangeType::TrackValueChanged);
+}
+void FControlRigParameterSection::KeyWeightValue(float Val)
+{
+	UMovieSceneControlRigParameterSection* ParameterSection = CastChecked<UMovieSceneControlRigParameterSection>(WeakSection.Get());
+	TSharedPtr<ISequencer> SequencerPtr = WeakSequencer.Pin();
+	FScopedTransaction Transaction(LOCTEXT("KeyWeightZero", "Key Weight Zero"));
+	ParameterSection->Modify();
+	EMovieSceneTransformChannel Channels = ParameterSection->GetTransformMask().GetChannels();
+	if ((Channels & EMovieSceneTransformChannel::Weight) == EMovieSceneTransformChannel::None)
+	{
+		ParameterSection->SetTransformMask(ParameterSection->GetTransformMask().GetChannels() | EMovieSceneTransformChannel::Weight);
+		SequencerPtr->NotifyMovieSceneDataChanged(EMovieSceneDataChangeType::MovieSceneStructureItemsChanged);
+	}
+	FFrameTime Time = SequencerPtr->GetLocalTime().Time;
+	ParameterSection->KeyWeightValue(Time.GetFrame(), Val);
+	SequencerPtr->NotifyMovieSceneDataChanged(EMovieSceneDataChangeType::TrackValueChanged);
+}
 
 void FControlRigParameterSection::BuildSectionContextMenu(FMenuBuilder& MenuBuilder, const FGuid& InObjectBinding)
 {
@@ -3515,6 +3637,52 @@ void FControlRigParameterSection::BuildSectionContextMenu(FMenuBuilder& MenuBuil
 							})
 						);
 		};
+		UMovieSceneControlRigParameterTrack* Track = ParameterSection->GetTypedOuter<UMovieSceneControlRigParameterTrack>();
+		if (Track)
+		{
+			TArray<UMovieSceneSection*> Sections = Track->GetAllSections();
+			//If Base Absolute section 
+			if (ParameterSection->GetBlendType().Get() == EMovieSceneBlendType::Absolute && Sections[0] == ParameterSection)
+			{
+				MenuBuilder.BeginSection(NAME_None, LOCTEXT("AnimationLayers", "Animation Layers"));
+				{
+					MenuBuilder.AddMenuEntry(
+						LOCTEXT("MergeAllSections", "Merge All Sections"),
+						LOCTEXT("MergeAllSections_ToolTip", "Merge all sections to this"),
+						FSlateIcon(),
+						FUIAction(FExecuteAction::CreateLambda([=] { CollapseAllLayers(); }))
+					);
+				}
+			}
+			if (ParameterSection->GetBlendType().Get() == EMovieSceneBlendType::Additive)
+			{
+
+				MenuBuilder.BeginSection(NAME_None, LOCTEXT("AnimationLayers", "Animation Layers"));
+				{
+					MenuBuilder.AddMenuEntry(
+						LOCTEXT("KeyZeroValue", "Key Zero Value"),
+						LOCTEXT("KeyZeroValue_Tooltip", "Set zero key on all controls in this section"),
+						FSlateIcon(),
+						FUIAction(FExecuteAction::CreateLambda([=] { KeyZeroValue(); }))
+					);
+				}
+				
+				MenuBuilder.AddMenuEntry(
+					LOCTEXT("KeyWeightZero", "Key Weight Zero"),
+					LOCTEXT("KeyWeightZero_Tooltip", "Key a zero value on the Weight channel"),
+					FSlateIcon(),
+					FUIAction(FExecuteAction::CreateLambda([=] { KeyWeightValue(0.0f); }))
+				);
+				
+				MenuBuilder.AddMenuEntry(
+					LOCTEXT("KeyWeightOne", "Key Weight One"),
+					LOCTEXT("KeyWeightOne_Tooltip", "Key a one value on the Weight channel"),
+					FSlateIcon(),
+					FUIAction(FExecuteAction::CreateLambda([=] { KeyWeightValue(1.0f); }))
+				);
+				
+			}
+		}
 		MenuBuilder.BeginSection(NAME_None, LOCTEXT("RigSectionActiveChannels", "Active Channels"));
 		{
 			MenuBuilder.AddMenuEntry(
diff --git a/Engine/Plugins/Animation/ControlRig/Source/ControlRigEditor/Private/Sequencer/ControlRigParameterTrackEditor.h b/Engine/Plugins/Animation/ControlRig/Source/ControlRigEditor/Private/Sequencer/ControlRigParameterTrackEditor.h
--- a/Engine/Plugins/Animation/ControlRig/Source/ControlRigEditor/Private/Sequencer/ControlRigParameterTrackEditor.h
+++ b/Engine/Plugins/Animation/ControlRig/Source/ControlRigEditor/Private/Sequencer/ControlRigParameterTrackEditor.h
@@ -278,6 +278,10 @@ protected:
 	/** Show all control channels*/
 	void ShowAllControlsChannels();
 
+	void KeyZeroValue();
+	void KeyWeightValue(float Val);
+	void CollapseAllLayers();
+
 	/** The sequencer which is controlling this section. */
 	TWeakPtr<ISequencer> WeakSequencer;
 
