Control Rig: Add option for icon coloring

diff --git a/Engine/Plugins/Animation/ControlRig/Source/ControlRigEditor/Private/ControlRigHierarchyCommands.cpp b/Engine/Plugins/Animation/ControlRig/Source/ControlRigEditor/Private/ControlRigHierarchyCommands.cpp
--- a/Engine/Plugins/Animation/ControlRig/Source/ControlRigEditor/Private/ControlRigHierarchyCommands.cpp
+++ b/Engine/Plugins/Animation/ControlRig/Source/ControlRigEditor/Private/ControlRigHierarchyCommands.cpp
@@ -38,6 +38,7 @@ void FControlRigHierarchyCommands::RegisterCommands()
 	UI_COMMAND(ShowDynamicHierarchy, "Dynamic Hierarchy", "Whether to show the static or dynamic hierarchy", EUserInterfaceActionType::ToggleButton, FInputChord());
 	UI_COMMAND(ToggleControlShapeTransformEdit, "Toggle Shape Transform Edit", "Toggle Editing Selected Control's Shape Transform", EUserInterfaceActionType::ToggleButton, FInputChord(EKeys::Period, EModifierKey::Control)); 
 	UI_COMMAND(TestSpaceSwitching, "Test Space Switching", "Try out space switching on the control", EUserInterfaceActionType::Button, FInputChord(EKeys::Q, EModifierKey::Control)); 
+	UI_COMMAND(ShowIconColors, "Show Icon Colors", "Whether to tint the icons with the element color", EUserInterfaceActionType::ToggleButton, FInputChord());
 }
 
 #undef LOCTEXT_NAMESPACE
diff --git a/Engine/Plugins/Animation/ControlRig/Source/ControlRigEditor/Private/ControlRigHierarchyCommands.h b/Engine/Plugins/Animation/ControlRig/Source/ControlRigEditor/Private/ControlRigHierarchyCommands.h
--- a/Engine/Plugins/Animation/ControlRig/Source/ControlRigEditor/Private/ControlRigHierarchyCommands.h
+++ b/Engine/Plugins/Animation/ControlRig/Source/ControlRigEditor/Private/ControlRigHierarchyCommands.h
@@ -114,6 +114,9 @@ public:
 	/** Test out space switching as it would look like for animator */
 	TSharedPtr< FUICommandInfo > TestSpaceSwitching;
 
+	/** Whether to tint the icons with the element color */
+	TSharedPtr< FUICommandInfo > ShowIconColors;
+
 	/**
 	 * Initialize commands
 	 */
diff --git a/Engine/Plugins/Animation/ControlRig/Source/ControlRigEditor/Private/EditMode/SControlRigEditModeTools.cpp b/Engine/Plugins/Animation/ControlRig/Source/ControlRigEditor/Private/EditMode/SControlRigEditModeTools.cpp
--- a/Engine/Plugins/Animation/ControlRig/Source/ControlRigEditor/Private/EditMode/SControlRigEditModeTools.cpp
+++ b/Engine/Plugins/Animation/ControlRig/Source/ControlRigEditor/Private/EditMode/SControlRigEditModeTools.cpp
@@ -179,6 +179,7 @@ void SControlRigEditModeTools::Construct(const FArguments& InArgs, TSharedPtr<FC
 	DisplaySettings.bShowRigidBodies = false;
 	DisplaySettings.bHideParentsOnFilter = true;
 	DisplaySettings.bFlattenHierarchyOnFilter = true;
+	DisplaySettings.bShowIconColors = true;
 #if USE_LOCAL_DETAILS
 	FRigTreeDelegates RigTreeDelegates;
 	RigTreeDelegates.OnGetHierarchy = FOnGetRigTreeHierarchy::CreateSP(this, &SControlRigEditModeTools::GetHierarchy);
diff --git a/Engine/Plugins/Animation/ControlRig/Source/ControlRigEditor/Private/Editor/SRigHierarchy.cpp b/Engine/Plugins/Animation/ControlRig/Source/ControlRigEditor/Private/Editor/SRigHierarchy.cpp
--- a/Engine/Plugins/Animation/ControlRig/Source/ControlRigEditor/Private/Editor/SRigHierarchy.cpp
+++ b/Engine/Plugins/Animation/ControlRig/Source/ControlRigEditor/Private/Editor/SRigHierarchy.cpp
@@ -447,6 +447,12 @@ void SRigHierarchy::BindCommands()
 		Commands.TestSpaceSwitching,
 		FExecuteAction::CreateSP(this, &SRigHierarchy::HandleTestSpaceSwitching),
 		FCanExecuteAction::CreateSP(this, &SRigHierarchy::IsControlSelected));
+
+	CommandList->MapAction(
+		Commands.ShowIconColors,
+		FExecuteAction::CreateLambda([this]() { DisplaySettings.bShowIconColors = !DisplaySettings.bShowIconColors; RefreshTreeView(); }),
+		FCanExecuteAction(),
+		FIsActionChecked::CreateLambda([this]() { return DisplaySettings.bShowIconColors; }));
 }
 
 FReply SRigHierarchy::OnKeyDown(const FGeometry& MyGeometry, const FKeyEvent& InKeyEvent)
@@ -860,6 +866,7 @@ TSharedRef< SWidget > SRigHierarchy::CreateFilterMenu()
 	{
 		MenuBuilder.AddMenuEntry(Actions.ShowControls);
 		MenuBuilder.AddMenuEntry(Actions.ShowNulls);
+		MenuBuilder.AddMenuEntry(Actions.ShowIconColors);
 	}
 	MenuBuilder.EndSection();
 
diff --git a/Engine/Plugins/Animation/ControlRig/Source/ControlRigEditor/Private/Editor/SRigHierarchyTreeView.cpp b/Engine/Plugins/Animation/ControlRig/Source/ControlRigEditor/Private/Editor/SRigHierarchyTreeView.cpp
--- a/Engine/Plugins/Animation/ControlRig/Source/ControlRigEditor/Private/Editor/SRigHierarchyTreeView.cpp
+++ b/Engine/Plugins/Animation/ControlRig/Source/ControlRigEditor/Private/Editor/SRigHierarchyTreeView.cpp
@@ -45,9 +45,9 @@ FRigTreeElement::FRigTreeElement(const FRigElementKey& InKey, TWeakPtr<SRigHiera
 }
 
 
-TSharedRef<ITableRow> FRigTreeElement::MakeTreeRowWidget(const TSharedRef<STableViewBase>& InOwnerTable, TSharedRef<FRigTreeElement> InRigTreeElement, TSharedPtr<SRigHierarchyTreeView> InTreeView)
+TSharedRef<ITableRow> FRigTreeElement::MakeTreeRowWidget(const TSharedRef<STableViewBase>& InOwnerTable, TSharedRef<FRigTreeElement> InRigTreeElement, TSharedPtr<SRigHierarchyTreeView> InTreeView, const FRigTreeDisplaySettings& InSettings)
 {
-	return SNew(SRigHierarchyItem, InOwnerTable, InRigTreeElement, InTreeView);
+	return SNew(SRigHierarchyItem, InOwnerTable, InRigTreeElement, InTreeView, InSettings);
 }
 
 void FRigTreeElement::RequestRename()
@@ -61,7 +61,7 @@ void FRigTreeElement::RequestRename()
 //////////////////////////////////////////////////////////////
 /// SRigHierarchyItem
 ///////////////////////////////////////////////////////////
-void SRigHierarchyItem::Construct(const FArguments& InArgs, const TSharedRef<STableViewBase>& OwnerTable, TSharedRef<FRigTreeElement> InRigTreeElement, TSharedPtr<SRigHierarchyTreeView> InTreeView)
+void SRigHierarchyItem::Construct(const FArguments& InArgs, const TSharedRef<STableViewBase>& OwnerTable, TSharedRef<FRigTreeElement> InRigTreeElement, TSharedPtr<SRigHierarchyTreeView> InTreeView, const FRigTreeDisplaySettings& InSettings)
 {
 	WeakRigTreeElement = InRigTreeElement;
 	Delegates = InTreeView->GetRigTreeDelegates();
@@ -90,7 +90,7 @@ void SRigHierarchyItem::Construct(const FArguments& InArgs, const TSharedRef<STa
 	TPair<const FSlateBrush*, FSlateColor> Result = GetBrushForElementType(Delegates.GetHierarchy(), InRigTreeElement->Key);
 	const FSlateBrush* Brush = Result.Key;
 	FSlateColor IconColor = Result.Value;
-	if(IconColor.IsColorSpecified())
+	if(IconColor.IsColorSpecified() && InSettings.bShowIconColors)
 	{
 		IconColor = InRigTreeElement->FilterResult == ERigTreeFilterResult::Shown ? Result.Value : FSlateColor(Result.Value.GetSpecifiedColor() * 0.5f);
 	}
@@ -596,7 +596,8 @@ void SRigHierarchyTreeView::SetExpansionRecursive(TSharedPtr<FRigTreeElement> In
 TSharedRef<ITableRow> SRigHierarchyTreeView::MakeTableRowWidget(TSharedPtr<FRigTreeElement> InItem,
 	const TSharedRef<STableViewBase>& OwnerTable)
 {
-	return InItem->MakeTreeRowWidget(OwnerTable, InItem.ToSharedRef(), SharedThis(this));
+	const FRigTreeDisplaySettings& Settings = Delegates.GetDisplaySettings();
+	return InItem->MakeTreeRowWidget(OwnerTable, InItem.ToSharedRef(), SharedThis(this), Settings);
 }
 
 void SRigHierarchyTreeView::HandleGetChildrenForTree(TSharedPtr<FRigTreeElement> InItem,
diff --git a/Engine/Plugins/Animation/ControlRig/Source/ControlRigEditor/Private/Editor/SRigHierarchyTreeView.h b/Engine/Plugins/Animation/ControlRig/Source/ControlRigEditor/Private/Editor/SRigHierarchyTreeView.h
--- a/Engine/Plugins/Animation/ControlRig/Source/ControlRigEditor/Private/Editor/SRigHierarchyTreeView.h
+++ b/Engine/Plugins/Animation/ControlRig/Source/ControlRigEditor/Private/Editor/SRigHierarchyTreeView.h
@@ -25,6 +25,7 @@ struct CONTROLRIGEDITOR_API FRigTreeDisplaySettings
 		bShowRigidBodies = true;
 		bShowReferences = true;
 		bShowDynamicHierarchy = false;
+		bShowIconColors = true;
 	}
 	
 	FText FilterText;
@@ -55,6 +56,9 @@ struct CONTROLRIGEDITOR_API FRigTreeDisplaySettings
 
 	/** Whether or not to show the static or dynamic hierarchy */
 	bool bShowDynamicHierarchy;
+
+	/** Whether to tint the icons with the element color */
+	bool bShowIconColors;
 };
 
 DECLARE_DELEGATE_RetVal(const URigHierarchy*, FOnGetRigTreeHierarchy);
@@ -170,7 +174,7 @@ public:
 	bool bSupportsRename;
 	TArray<TSharedPtr<FRigTreeElement>> Children;
 
-	TSharedRef<ITableRow> MakeTreeRowWidget(const TSharedRef<STableViewBase>& InOwnerTable, TSharedRef<FRigTreeElement> InRigTreeElement, TSharedPtr<SRigHierarchyTreeView> InTreeView);
+	TSharedRef<ITableRow> MakeTreeRowWidget(const TSharedRef<STableViewBase>& InOwnerTable, TSharedRef<FRigTreeElement> InRigTreeElement, TSharedPtr<SRigHierarchyTreeView> InTreeView, const FRigTreeDisplaySettings& InSettings);
 
 	void RequestRename();
 
@@ -186,7 +190,7 @@ class SRigHierarchyItem : public STableRow<TSharedPtr<FRigTreeElement>>
 {
 public:
 	
-	void Construct(const FArguments& InArgs, const TSharedRef<STableViewBase>& OwnerTable, TSharedRef<FRigTreeElement> InRigTreeElement, TSharedPtr<SRigHierarchyTreeView> InTreeView);
+	void Construct(const FArguments& InArgs, const TSharedRef<STableViewBase>& OwnerTable, TSharedRef<FRigTreeElement> InRigTreeElement, TSharedPtr<SRigHierarchyTreeView> InTreeView, const FRigTreeDisplaySettings& InSettings);
  	void OnNameCommitted(const FText& InText, ETextCommit::Type InCommitType) const;
 	bool OnVerifyNameChanged(const FText& InText, FText& OutErrorMessage);
 	static TPair<const FSlateBrush*, FSlateColor> GetBrushForElementType(const URigHierarchy* InHierarchy, const FRigElementKey& InKey);
