Expose Collapse Layers and Space Switching functionality to Python.

diff --git a/Engine/Plugins/Animation/ControlRig/Source/ControlRigEditor/Private/ControlRigSequencerEditorLibrary.cpp b/Engine/Plugins/Animation/ControlRig/Source/ControlRigEditor/Private/ControlRigSequencerEditorLibrary.cpp
--- a/Engine/Plugins/Animation/ControlRig/Source/ControlRigEditor/Private/ControlRigSequencerEditorLibrary.cpp
+++ b/Engine/Plugins/Animation/ControlRig/Source/ControlRigEditor/Private/ControlRigSequencerEditorLibrary.cpp
@@ -34,6 +34,8 @@
 #include "Exporters/AnimSeqExportOption.h"
 #include "MovieSceneTimeHelpers.h"
 #include "ScopedTransaction.h"
+#include "ControlRigParameterTrackEditor.h"
+#include "ControlRigSpaceChannelEditors.h"
 
 #define LOCTEXT_NAMESPACE "ControlrigSequencerEditorLibrary"
 
@@ -2132,6 +2134,135 @@ bool UControlRigSequencerEditorLibrary::ImportFBXToControlRigTrack(UWorld* World
 	return bValid;
 }
 
+bool UControlRigSequencerEditorLibrary::CollapseControlRigAnimLayers(ULevelSequence* LevelSequence, UMovieSceneControlRigParameterTrack* InTrack, bool bKeyReduce, float Tolerance)
+{
+	TWeakPtr<ISequencer> WeakSequencer = GetSequencerFromAsset(LevelSequence);
+	bool bValid = false;
+
+	if (WeakSequencer.IsValid() && InTrack)
+	{
+		TArray<UMovieSceneSection*> Sections = InTrack->GetAllSections();
+		if (Sections.Num() > 0)
+		{
+			TSharedPtr<ISequencer>  SequencerPtr = WeakSequencer.Pin();
+			UMovieSceneControlRigParameterSection* ParameterSection = Cast<UMovieSceneControlRigParameterSection>(Sections[0]);
+			bValid = FControlRigParameterTrackEditor::CollapseAllLayers(SequencerPtr,InTrack, ParameterSection, bKeyReduce, Tolerance);
+		}
+	}
+	return bValid;
+}
+
+bool UControlRigSequencerEditorLibrary::SetControlRigSpace(ULevelSequence* LevelSequence, UControlRig* ControlRig, FName ControlName, const FRigElementKey& InSpaceKey, FFrameNumber InTime, ESequenceTimeUnit TimeUnit)
+{
+	TWeakPtr<ISequencer> WeakSequencer = GetSequencerFromAsset(LevelSequence);
+	bool bValid = false;
+
+	if (WeakSequencer.IsValid() && ControlRig && ControlName != NAME_None)
+	{
+		if (FRigControlElement* Element = ControlRig->FindControl(ControlName))
+		{
+			TSharedPtr<ISequencer>  Sequencer = WeakSequencer.Pin();
+			FScopedTransaction Transaction(LOCTEXT("KeyControlRigSpace", "Key Control Rig Space"));
+			FSpaceChannelAndSection SpaceChannelAndSection = FControlRigSpaceChannelHelpers::FindSpaceChannelAndSectionForControl(ControlRig, ControlName, Sequencer.Get(), true /*bCreateIfNeeded*/);
+			if (SpaceChannelAndSection.SpaceChannel)
+			{
+				if (TimeUnit == ESequenceTimeUnit::DisplayRate)
+				{
+					InTime = FFrameRate::TransformTime(FFrameTime(InTime, 0), LevelSequence->GetMovieScene()->GetDisplayRate(), LevelSequence->GetMovieScene()->GetTickResolution()).RoundToFrame();
+				}
+				FKeyHandle Handle = FControlRigSpaceChannelHelpers::SequencerKeyControlRigSpaceChannel(ControlRig, Sequencer.Get(), SpaceChannelAndSection.SpaceChannel, SpaceChannelAndSection.SectionToKey, InTime, ControlRig->GetHierarchy(), Element->GetKey(), InSpaceKey);
+				bValid = Handle != FKeyHandle::Invalid();
+			}
+			else
+			{
+				UE_LOG(LogControlRig, Error, TEXT("Can not find Space Channel"));
+				return false;
+			}
+		}
+		else
+		{
+			UE_LOG(LogControlRig, Error, TEXT("Can not find Control with that Name"));
+			return false;
+		}
+	}
+	return bValid;
+}
+
+bool UControlRigSequencerEditorLibrary::DeleteControlRigSpace(ULevelSequence* LevelSequence, UControlRig* ControlRig, FName ControlName, FFrameNumber InTime, ESequenceTimeUnit TimeUnit)
+{
+	TWeakPtr<ISequencer> WeakSequencer = GetSequencerFromAsset(LevelSequence);
+	bool bValid = false;
+
+	if (WeakSequencer.IsValid() && ControlRig && ControlName != NAME_None)
+	{
+		if (FRigControlElement* Element = ControlRig->FindControl(ControlName))
+		{
+			TSharedPtr<ISequencer>  Sequencer = WeakSequencer.Pin();
+			FScopedTransaction Transaction(LOCTEXT("KeyControlRigSpace", "Key Control Rig Space"));
+			FSpaceChannelAndSection SpaceChannelAndSection = FControlRigSpaceChannelHelpers::FindSpaceChannelAndSectionForControl(ControlRig, ControlName, Sequencer.Get(), false /*bCreateIfNeeded*/);
+			if (SpaceChannelAndSection.SpaceChannel)
+			{
+				if (TimeUnit == ESequenceTimeUnit::DisplayRate)
+				{
+					InTime = FFrameRate::TransformTime(FFrameTime(InTime, 0), LevelSequence->GetMovieScene()->GetDisplayRate(), LevelSequence->GetMovieScene()->GetTickResolution()).RoundToFrame();
+				}
+				UMovieSceneControlRigParameterSection* ParamSection = Cast<UMovieSceneControlRigParameterSection>(SpaceChannelAndSection.SectionToKey);
+				FControlRigSpaceChannelHelpers::SequencerSpaceChannelKeyDeleted(ControlRig, Sequencer.Get(), ControlName, SpaceChannelAndSection.SpaceChannel, ParamSection, InTime);
+				bValid = true;
+			}
+			else
+			{
+				UE_LOG(LogControlRig, Error, TEXT("Can not find Space Channel"));
+				return false;
+			}
+		}
+		else
+		{
+			UE_LOG(LogControlRig, Error, TEXT("Can not find Control with that Name"));
+			return false;
+		}
+	}
+	return bValid;
+}
+
+bool UControlRigSequencerEditorLibrary::BakeControlRigSpace(ULevelSequence* InSequence, UControlRig* InControlRig, TArray<FName>& InControlNames, FRigSpacePickerBakeSettings InSettings, ESequenceTimeUnit TimeUnit)
+{
+	TWeakPtr<ISequencer> WeakSequencer = GetSequencerFromAsset(InSequence);
+	bool bValid = false;
 
+	if (WeakSequencer.IsValid() && InControlRig && InControlNames.Num() > 0)
+	{
+		TSharedPtr<ISequencer>  Sequencer = WeakSequencer.Pin();
+		const FFrameRate TickResolution = Sequencer->GetFocusedTickResolution();
+		TArray<FFrameNumber> Frames;
+		const FFrameRate& FrameRate = Sequencer->GetFocusedDisplayRate();
+		FFrameNumber FrameRateInFrameNumber = TickResolution.AsFrameNumber(FrameRate.AsInterval());
+		if (TimeUnit == ESequenceTimeUnit::DisplayRate)
+		{
+			InSettings.StartFrame = FFrameRate::TransformTime(FFrameTime(InSettings.StartFrame, 0), FrameRate, TickResolution).RoundToFrame();
+			InSettings.EndFrame = FFrameRate::TransformTime(FFrameTime(InSettings.EndFrame, 0), FrameRate, TickResolution).RoundToFrame();
+		}
+		for (FFrameNumber& Frame = InSettings.StartFrame; Frame <= InSettings.EndFrame; Frame += FrameRateInFrameNumber)
+		{
+			Frames.Add(Frame);
+		}
+
+		FScopedTransaction Transaction(LOCTEXT("BakeControlToSpace", "Bake Control In Space"));
+		for (const FName& ControlName : InControlNames)
+		{
+			if (FRigControlElement* Element = InControlRig->FindControl(ControlName))
+			{
+				FSpaceChannelAndSection SpaceChannelAndSection = FControlRigSpaceChannelHelpers::FindSpaceChannelAndSectionForControl(InControlRig, ControlName, Sequencer.Get(), false /*bCreateIfNeeded*/);
+				if (SpaceChannelAndSection.SpaceChannel)
+				{
+					FControlRigSpaceChannelHelpers::SequencerBakeControlInSpace(InControlRig, Sequencer.Get(), SpaceChannelAndSection.SpaceChannel, SpaceChannelAndSection.SectionToKey,
+						Frames, InControlRig->GetHierarchy(), Element->GetKey(), InSettings);
+				}
+			}
+		}
+		bValid = true;
+	}
+	return bValid;
+}
 
 #undef LOCTEXT_NAMESPACE
diff --git a/Engine/Plugins/Animation/ControlRig/Source/ControlRigEditor/Private/Editor/SRigSpacePickerWidget.h b/Engine/Plugins/Animation/ControlRig/Source/ControlRigEditor/Private/Editor/SRigSpacePickerWidget.h
--- a/Engine/Plugins/Animation/ControlRig/Source/ControlRigEditor/Private/Editor/SRigSpacePickerWidget.h
+++ b/Engine/Plugins/Animation/ControlRig/Source/ControlRigEditor/Private/Editor/SRigSpacePickerWidget.h
@@ -11,37 +11,8 @@
 #include "IStructureDetailsView.h"
 #include "Misc/FrameNumber.h"
 #include "EditorUndoClient.h"
-#include "SRigSpacePickerWidget.generated.h"
+#include "RigSpacePickerBakeSettings.h"
 
-USTRUCT()
-struct CONTROLRIGEDITOR_API FRigSpacePickerBakeSettings
-{
-	GENERATED_BODY();
-
-	FRigSpacePickerBakeSettings()
-	{
-		TargetSpace = FRigElementKey();
-		StartFrame = 0;
-		EndFrame = 100;
-		bReduceKeys = false;
-		Tolerance = 0.001f;
-	}
-
-	UPROPERTY()
-	FRigElementKey TargetSpace;
-
-	UPROPERTY(EditAnywhere, Category = "Settings")
-	FFrameNumber StartFrame;
-
-	UPROPERTY(EditAnywhere, Category = "Settings")
-	FFrameNumber EndFrame;
-
-	UPROPERTY(EditAnywhere, Category = "Settings")
-	bool bReduceKeys;
-
-	UPROPERTY(EditAnywhere, Category = "Settings")
-	float Tolerance;
-};
 
 DECLARE_DELEGATE_RetVal_TwoParams(FRigElementKey, FRigSpacePickerGetActiveSpace, URigHierarchy*, const FRigElementKey&);
 DECLARE_DELEGATE_RetVal_TwoParams(const FRigControlElementCustomization*, FRigSpacePickerGetControlCustomization, URigHierarchy*, const FRigElementKey&);
diff --git a/Engine/Plugins/Animation/ControlRig/Source/ControlRigEditor/Private/Sequencer/ControlRigParameterTrackEditor.cpp b/Engine/Plugins/Animation/ControlRig/Source/ControlRigEditor/Private/Sequencer/ControlRigParameterTrackEditor.cpp
--- a/Engine/Plugins/Animation/ControlRig/Source/ControlRigEditor/Private/Sequencer/ControlRigParameterTrackEditor.cpp
+++ b/Engine/Plugins/Animation/ControlRig/Source/ControlRigEditor/Private/Sequencer/ControlRigParameterTrackEditor.cpp
@@ -3372,16 +3372,20 @@ void FControlRigParameterTrackEditor::SelectFKBonesToAnimate(UFKControlRig* Auto
 
 	//reconstruct all channel proxies TODO or not to do that is the question
 }
-
-void FControlRigParameterSection::CollapseAllLayers()
+bool FControlRigParameterTrackEditor::CollapseAllLayers(TSharedPtr<ISequencer>& SequencerPtr, UMovieSceneTrack* OwnerTrack, UMovieSceneControlRigParameterSection* ParameterSection, bool bKeyReduce, float Tolerance)
 {
-	FScopedTransaction Transaction(LOCTEXT("CollapseAllSections", "CollapseAllSections"));
 
-	UMovieSceneControlRigParameterSection* ParameterSection = CastChecked<UMovieSceneControlRigParameterSection>(WeakSection.Get());
-	if (ParameterSection && ParameterSection->GetControlRig())
+	if (SequencerPtr.IsValid() && OwnerTrack && ParameterSection && ParameterSection->GetControlRig())
 	{
+		TArray<UMovieSceneSection*> Sections = OwnerTrack->GetAllSections();
+		//make sure right type
+		if (ParameterSection->GetBlendType().Get() != EMovieSceneBlendType::Absolute && Sections.Num() > 0 && Sections[0] != ParameterSection)
+		{
+			UE_LOG(LogControlRigEditor, Log, TEXT("Section wrong type or not first when collapsing layers"));
+			return false;
+		}
+		FScopedTransaction Transaction(LOCTEXT("CollapseAllSections", "CollapseAllSections"));
 		UControlRig* ControlRig = ParameterSection->GetControlRig();
-		TSharedPtr<ISequencer> SequencerPtr = WeakSequencer.Pin();
 		TRange<FFrameNumber> Range = SequencerPtr->GetFocusedMovieSceneSequence()->GetMovieScene()->GetPlaybackRange();
 		FFrameNumber StartFrame = Range.GetLowerBoundValue();
 		FFrameNumber EndFrame = Range.GetUpperBoundValue();
@@ -3428,9 +3432,7 @@ void FControlRigParameterSection::CollapseAllLayers()
 			}
 		}
 		//delete other sections
-		UMovieSceneTrack* OwnerTrack = ParameterSection->GetTypedOuter<UMovieSceneTrack>();
 		OwnerTrack->Modify();
-		const TArray<UMovieSceneSection*>& Sections = OwnerTrack->GetAllSections();
 		for (Index = Sections.Num() - 1; Index >= 0; --Index)
 		{
 			if (Sections[Index] != ParameterSection)
@@ -3453,20 +3455,35 @@ void FControlRigParameterSection::CollapseAllLayers()
 			{
 				for (TPair<FName, TArray<FTransform>>& TrailControlTransform : ControlLocalTransforms)
 				{
-					ControlRig->SetControlLocalTransform(TrailControlTransform.Key, TrailControlTransform.Value[Index],true,Context,false);
+					ControlRig->SetControlLocalTransform(TrailControlTransform.Key, TrailControlTransform.Value[Index], true, Context, false);
 				}
 			}
 		}
+		if (bKeyReduce)
+		{
+			FKeyDataOptimizationParams Params;
+			Params.bAutoSetInterpolation = true;
+			Params.Tolerance = Tolerance;
+			FMovieSceneChannelProxy& ChannelProxy = ParameterSection->GetChannelProxy();
+			TArrayView<FMovieSceneFloatChannel*> FloatChannels = ChannelProxy.GetChannels<FMovieSceneFloatChannel>();
 
+			for (FMovieSceneFloatChannel* Channel : FloatChannels)
+			{
+				Channel->Optimize(Params); //should also auto tangent
+			}
+		}
 		//reset everything back
 		SequencerPtr->NotifyMovieSceneDataChanged(EMovieSceneDataChangeType::MovieSceneStructureItemAdded);
-		/*
-			const FFrameTime StartTime = SequencerPtr->GetLocalTime().Time;
-			FMovieSceneContext Context = FMovieSceneContext(FMovieSceneEvaluationRange(StartTime, TickResolution), SequencerPtr->GetPlaybackStatus()).SetHasJumped(true);
-			SequencerPtr->GetEvaluationTemplate().Evaluate(Context, *SequencerPtr);
-			ControlRig->Evaluate_AnyThread();
-			*/
+		return true;
 	}
+	return false;
+}
+void FControlRigParameterSection::CollapseAllLayers()
+{
+	TSharedPtr<ISequencer> SequencerPtr = WeakSequencer.Pin();
+	UMovieSceneControlRigParameterSection* ParameterSection = CastChecked<UMovieSceneControlRigParameterSection>(WeakSection.Get());
+	UMovieSceneTrack* OwnerTrack = ParameterSection->GetTypedOuter<UMovieSceneTrack>();
+	FControlRigParameterTrackEditor::CollapseAllLayers(SequencerPtr,OwnerTrack, ParameterSection);
 
 }
 void FControlRigParameterSection::KeyZeroValue()
@@ -3647,8 +3664,8 @@ void FControlRigParameterSection::BuildSectionContextMenu(FMenuBuilder& MenuBuil
 				MenuBuilder.BeginSection(NAME_None, LOCTEXT("AnimationLayers", "Animation Layers"));
 				{
 					MenuBuilder.AddMenuEntry(
-						LOCTEXT("MergeAllSections", "Merge All Sections"),
-						LOCTEXT("MergeAllSections_ToolTip", "Merge all sections to this"),
+						LOCTEXT("CollapseAllSections", "Collapse All Sections"),
+						LOCTEXT("CollapseAllSections_ToolTip", "Collapse all sections onto this section"),
 						FSlateIcon(),
 						FUIAction(FExecuteAction::CreateLambda([=] { CollapseAllLayers(); }))
 					);
diff --git a/Engine/Plugins/Animation/ControlRig/Source/ControlRigEditor/Private/Sequencer/ControlRigParameterTrackEditor.h b/Engine/Plugins/Animation/ControlRig/Source/ControlRigEditor/Private/Sequencer/ControlRigParameterTrackEditor.h
--- a/Engine/Plugins/Animation/ControlRig/Source/ControlRigEditor/Private/Sequencer/ControlRigParameterTrackEditor.h
+++ b/Engine/Plugins/Animation/ControlRig/Source/ControlRigEditor/Private/Sequencer/ControlRigParameterTrackEditor.h
@@ -195,6 +195,8 @@ public:
 	 */
 	bool ModifyOurGeneratedKeysByCurrentAndWeight(UObject* Object, UControlRig* InControlRig, FName RigControlName, UMovieSceneTrack *Track, UMovieSceneSection* SectionToKey, FFrameNumber Time, FGeneratedTrackKeys& InOutGeneratedTotalKeys, float Weight) const;
 
+	//**Function to collapse all layers from this section onto the first absoluate layer.*/
+	static bool CollapseAllLayers(TSharedPtr<ISequencer>& SequencerPtr, UMovieSceneTrack* OwnerTrack, UMovieSceneControlRigParameterSection* ParameterSection, bool bKeyReduce = false, float Tolerance = 0.001f);
 
 private:
 	FDelegateHandle SelectionChangedHandle;
@@ -259,7 +261,6 @@ public:
 	//~ ISequencerSection interface
 	virtual bool RequestDeleteCategory(const TArray<FName>& CategoryNamePath) override;
 	virtual bool RequestDeleteKeyArea(const TArray<FName>& KeyAreaNamePath) override;
-
 protected:
 	/** Add Sub Menu */
 	void AddAnimationSubMenuForFK(FMenuBuilder& MenuBuilder, FGuid ObjectBinding, USkeleton* Skeleton, UMovieSceneControlRigParameterSection* Section);
diff --git a/Engine/Plugins/Animation/ControlRig/Source/ControlRigEditor/Public/ControlRigSequencerEditorLibrary.h b/Engine/Plugins/Animation/ControlRig/Source/ControlRigEditor/Public/ControlRigSequencerEditorLibrary.h
--- a/Engine/Plugins/Animation/ControlRig/Source/ControlRigEditor/Public/ControlRigSequencerEditorLibrary.h
+++ b/Engine/Plugins/Animation/ControlRig/Source/ControlRigEditor/Public/ControlRigSequencerEditorLibrary.h
@@ -16,6 +16,7 @@
 #include "EulerTransform.h"
 #include "Tools/ControlRigSnapSettings.h"
 #include "SequenceTimeUnit.h"
+#include "RigSpacePickerBakeSettings.h"
 #include "ControlRigSequencerEditorLibrary.generated.h"
 
 class ULevelSequence;
@@ -782,18 +783,67 @@ public:
 	/*
 	 * Import FBX onto a control rig with the specified track and section
 	 *
-	 * @InWorld World to import to
-	 * @InSequence Sequence to import
-	 * @InTrack Track to import onto
-	 * @InSection Section to import onto, may be null in which case we use the track's section to key
-	 * @SelectedControlRigNames  List of selected control rig names. Will use them if  ImportFBXControlRigSettings->bImportOntoSelectedControls is true
-	 * @ImportFBXControlRigSettings Settings to control import.
-	 * @InImportFileName Path to fbx file to create
+	 * @param InWorld World to import to
+	 * @param InSequence Sequence to import
+	 * @param InTrack Track to import onto
+	 * @param InSection Section to import onto, may be null in which case we use the track's section to key
+	 * @param SelectedControlRigNames  List of selected control rig names. Will use them if  ImportFBXControlRigSettings->bImportOntoSelectedControls is true
+	 * @param ImportFBXControlRigSettings Settings to control import.
+	 * @param InImportFileName Path to fbx file to create
 	 */
 	UFUNCTION(BlueprintCallable, Category = "Editor Scripting | Sequencer Tools | FBX")
 	static bool ImportFBXToControlRigTrack(UWorld* World, ULevelSequence* InSequence, UMovieSceneControlRigParameterTrack* InTrack, UMovieSceneControlRigParameterSection* InSection,
 			const TArray<FString>& SelectedControlRigNames,
 			UMovieSceneUserImportFBXControlRigSettings* ImportFBXControlRigSettings,
 			const FString& ImportFilename);
+
+	/*
+	 * Collapse and bake all sections and layers on a control rig track to just one section.
+	 *
+	 * @param InSequence Sequence that has track to collapse
+	 * @param InTrack Track for layers to collapse
+	 * @param bKeyReduce If true do key reduction based upon Tolerance, if false don't
+	 * @param Tolerance If reducing keys, tolerance about which keys will be removed, smaller tolerance, more keys usually.
+	 */
+	UFUNCTION(BlueprintCallable, Category = "Editor Scripting | Sequencer Tools | Control Rig")
+	static bool CollapseControlRigAnimLayers(ULevelSequence* InSequence,UMovieSceneControlRigParameterTrack* InTrack, bool bKeyReduce = false, float Tolerance = 0.001f);
+
+	/*
+	 * Set the a key for the Control Rig Space for the Control at the specified time. If space is the same as the current no key witll be set.
+	 *
+	 * @param InSequence Sequence to set the space
+	 * @param InControlRig ControlRig with the Control
+	 * @param InControlName The name of the Control
+	 * @param InSpaceKey  The new space for the Control
+	 * @param InTime Time to change the space.
+	 * @param TimeUnit Unit for the InTime, either in display rate or tick resolution
+	 */
+	UFUNCTION(BlueprintCallable, Category = "Editor Scripting | Sequencer Tools | Control Rig")
+	static bool SetControlRigSpace(ULevelSequence* InSequence, UControlRig* InControlRig, FName InControlName, const FRigElementKey& InSpaceKey, FFrameNumber InTime,  ESequenceTimeUnit TimeUnit = ESequenceTimeUnit::DisplayRate);
+
+	/*Set the a key for the Control Rig Space for the Control at the specified time.If space is the same as the current no key witll be set.
+	*
+	* @param InSequence Sequence to set the space
+	* @param InControlRig ControlRig with the Control
+	* @param InControlName The name of the Control
+	* @param InSpaceKey  The new space for the Control
+	* @param InTime Time to change the space.
+	* @param TimeUnit Unit for the InTime, either in display rate or tick resolution
+	* 
+	*/
+	UFUNCTION(BlueprintCallable, Category = "Editor Scripting | Sequencer Tools | Control Rig")
+	static bool DeleteControlRigSpace(ULevelSequence* InSequence, UControlRig* InControlRig, FName InControlName,  FFrameNumber InTime, ESequenceTimeUnit TimeUnit = ESequenceTimeUnit::DisplayRate);
+
+	/**
+	* @param InSequence Sequence to bake
+	* @param InControlRig ControlRig to bake
+	* @param InControlNames The name of the Controls to bake
+	* @param InSettings  The settings for the bake, e.g, how long to bake, to key reduce etc.
+	* @param TimeUnit Unit for the start and end times in the InSettings parameter.
+	*/
+	UFUNCTION(BlueprintCallable, Category = "Editor Scripting | Sequencer Tools | Control Rig")
+	static bool BakeControlRigSpace(ULevelSequence* InSequence, UControlRig* InControlRig, TArray<FName>& InControlNames, FRigSpacePickerBakeSettings InSettings, ESequenceTimeUnit TimeUnit = ESequenceTimeUnit::DisplayRate);
+	
+	
 };
 
diff --git a/Engine/Plugins/Animation/ControlRig/Source/ControlRigEditor/Public/RigSpacePickerBakeSettings.h b/Engine/Plugins/Animation/ControlRig/Source/ControlRigEditor/Public/RigSpacePickerBakeSettings.h
new file mode 100644
--- /dev/null
+++ b/Engine/Plugins/Animation/ControlRig/Source/ControlRigEditor/Public/RigSpacePickerBakeSettings.h
@@ -0,0 +1,36 @@
+// Copyright Epic Games, Inc. All Rights Reserved.
+
+#pragma once
+
+#include "CoreMinimal.h"
+#include "RigSpacePickerBakeSettings.generated.h"
+
+USTRUCT(BlueprintType)
+struct CONTROLRIGEDITOR_API FRigSpacePickerBakeSettings
+{
+	GENERATED_BODY();
+
+	FRigSpacePickerBakeSettings()
+	{
+		TargetSpace = FRigElementKey();
+		StartFrame = 0;
+		EndFrame = 100;
+		bReduceKeys = false;
+		Tolerance = 0.001f;
+	}
+
+	UPROPERTY(BlueprintReadWrite, Category = "Settings")
+	FRigElementKey TargetSpace;
+
+	UPROPERTY(BlueprintReadWrite,EditAnywhere, Category = "Settings")
+	FFrameNumber StartFrame;
+
+	UPROPERTY(BlueprintReadWrite,EditAnywhere, Category = "Settings")
+	FFrameNumber EndFrame;
+
+	UPROPERTY(BlueprintReadWrite,EditAnywhere, Category = "Settings")
+	bool bReduceKeys;
+
+	UPROPERTY(BlueprintReadWrite,EditAnywhere, Category = "Settings")
+	float Tolerance;
+};
