Control Rig: fix user able to set array type to exec pin

diff --git a/Engine/Plugins/Animation/ControlRig/Source/ControlRigDeveloper/Private/Graph/ControlRigGraphSchema.cpp b/Engine/Plugins/Animation/ControlRig/Source/ControlRigDeveloper/Private/Graph/ControlRigGraphSchema.cpp
--- a/Engine/Plugins/Animation/ControlRig/Source/ControlRigDeveloper/Private/Graph/ControlRigGraphSchema.cpp
+++ b/Engine/Plugins/Animation/ControlRig/Source/ControlRigDeveloper/Private/Graph/ControlRigGraphSchema.cpp
@@ -927,6 +927,15 @@ bool UControlRigGraphSchema::SupportsPinType(TWeakPtr<const FEdGraphSchemaAction
 bool UControlRigGraphSchema::SupportsPinTypeContainer(TWeakPtr<const FEdGraphSchemaAction> SchemaAction,
 	const FEdGraphPinType& PinType, const EPinContainerType& ContainerType) const
 {
+	// Do not allow containers for execute context type
+	if(const UScriptStruct* ExecuteContextScriptStruct = Cast<UScriptStruct>(PinType.PinSubCategoryObject))
+	{
+		if (ExecuteContextScriptStruct->IsChildOf(FRigVMExecuteContext::StaticStruct()))
+		{
+			return ContainerType == EPinContainerType::None;
+		}
+	}
+	
 	return ContainerType == EPinContainerType::None || ContainerType == EPinContainerType::Array;
 }
 
diff --git a/Engine/Source/Developer/RigVMDeveloper/Private/RigVMModel/RigVMController.cpp b/Engine/Source/Developer/RigVMDeveloper/Private/RigVMModel/RigVMController.cpp
--- a/Engine/Source/Developer/RigVMDeveloper/Private/RigVMModel/RigVMController.cpp
+++ b/Engine/Source/Developer/RigVMDeveloper/Private/RigVMModel/RigVMController.cpp
@@ -8255,13 +8255,21 @@ bool URigVMController::ChangeExposedPinType(const FName& InPinName, const FStrin
 
 	if (URigVMFunctionEntryNode* EntryNode = Graph->GetEntryNode())
 	{
+		TArray<URigVMLink*> Links = EntryNode->GetLinks();
+		DetachLinksFromPinObjects(&Links, true);
 		RepopulatePinsOnNode(EntryNode, true, true);
+		ReattachLinksToPinObjects(false, &Links, true);
+		
 		RemoveUnusedOrphanedPins(EntryNode, true);
 	}
 	
 	if (URigVMFunctionReturnNode* ReturnNode = Graph->GetReturnNode())
 	{
+		TArray<URigVMLink*> Links = ReturnNode->GetLinks();
+		DetachLinksFromPinObjects(&Links, true);
 		RepopulatePinsOnNode(ReturnNode, true, true);
+		ReattachLinksToPinObjects(false, &Links, true);
+		
 		RemoveUnusedOrphanedPins(ReturnNode, true);
 	}
 
diff --git a/Engine/Source/Editor/KismetWidgets/Private/SPinTypeSelector.cpp b/Engine/Source/Editor/KismetWidgets/Private/SPinTypeSelector.cpp
--- a/Engine/Source/Editor/KismetWidgets/Private/SPinTypeSelector.cpp
+++ b/Engine/Source/Editor/KismetWidgets/Private/SPinTypeSelector.cpp
@@ -948,6 +948,11 @@ void SPinTypeSelector::OnSelectPinType(FPinTypeTreeItem InItem, FName InPinCateg
 		FSlateNotificationManager::Get().AddNotification(Info);
 	}
 
+	if (!Schema->SupportsPinTypeContainer(SchemaAction, NewTargetPinType, NewTargetPinType.ContainerType))
+	{
+		NewTargetPinType.ContainerType = EPinContainerType::None;
+	}
+
 	OnTypeChanged.ExecuteIfBound(NewTargetPinType);
 }
 
