Control Rig: Adapt comments for several rig units

diff --git a/Engine/Plugins/Animation/ControlRig/Source/ControlRig/Private/Units/Execution/RigUnit_DynamicHierarchy.h b/Engine/Plugins/Animation/ControlRig/Source/ControlRig/Private/Units/Execution/RigUnit_DynamicHierarchy.h
--- a/Engine/Plugins/Animation/ControlRig/Source/ControlRig/Private/Units/Execution/RigUnit_DynamicHierarchy.h
+++ b/Engine/Plugins/Animation/ControlRig/Source/ControlRig/Private/Units/Execution/RigUnit_DynamicHierarchy.h
@@ -35,9 +35,15 @@ struct CONTROLRIG_API FRigUnit_AddParent : public FRigUnit_DynamicHierarchyBaseM
 	RIGVM_METHOD()
 	virtual void Execute(const FRigUnitContext& Context) override;
 
+	/*
+	 * The child to be parented under the new parent
+	 */
 	UPROPERTY(meta = (Input, ExpandByDefault))
 	FRigElementKey Child;
 
+	/*
+	 * The new parent to be added to the child
+	 */
 	UPROPERTY(meta = (Input, ExpandByDefault))
 	FRigElementKey Parent;
 };
@@ -95,7 +101,7 @@ struct CONTROLRIG_API FRigUnit_SwitchParent : public FRigUnit_DynamicHierarchyBa
 };
 
 /**
-* Returns the item's parent's weights
+* Returns the item's parents' weights
 */
 USTRUCT(meta=(DisplayName="Get Parent Weights", Keywords="Chain,Parents,Hierarchy", Varying, Deprecated = "5.0"))
 struct CONTROLRIG_API FRigUnit_HierarchyGetParentWeights : public FRigUnit_DynamicHierarchyBase
@@ -110,18 +116,27 @@ struct CONTROLRIG_API FRigUnit_HierarchyGetParentWeights : public FRigUnit_Dynam
 	RIGVM_METHOD()
 	virtual void Execute(const FRigUnitContext& Context) override;
 
+	/*
+	 * The child to retrieve the weights for
+	 */
 	UPROPERTY(meta = (Input, ExpandByDefault))
 	FRigElementKey Child;
 
+	/*
+	 * The weight of each parent
+	 */
 	UPROPERTY(meta = (Output))
 	TArray<FRigElementWeight> Weights;
 
+	/*
+	 * The key for each parent
+	 */
 	UPROPERTY(meta = (Output))
 	FRigElementKeyCollection Parents;
 };
 
 /**
-* Returns the item's parent's weights
+* Returns the item's parents' weights
 */
 USTRUCT(meta=(DisplayName="Get Parent Weights", Keywords="Chain,Parents,Hierarchy", Varying))
 struct CONTROLRIG_API FRigUnit_HierarchyGetParentWeightsArray : public FRigUnit_DynamicHierarchyBase
@@ -136,18 +151,27 @@ struct CONTROLRIG_API FRigUnit_HierarchyGetParentWeightsArray : public FRigUnit_
 	RIGVM_METHOD()
 	virtual void Execute(const FRigUnitContext& Context) override;
 
+	/*
+	 * The child to retrieve the weights for
+	 */
 	UPROPERTY(meta = (Input, ExpandByDefault))
 	FRigElementKey Child;
 
+	/*
+	 * The weight of each parent
+	 */
 	UPROPERTY(meta = (Output))
 	TArray<FRigElementWeight> Weights;
 
+	/*
+	 * The key for each parent
+	 */
 	UPROPERTY(meta = (Output))
 	TArray<FRigElementKey> Parents;
 };
 
 /**
-* Sets the item's parent's weights
+* Sets the item's parents' weights
 */
 USTRUCT(meta=(DisplayName="Set Parent Weights", Keywords="Chain,Parents,Hierarchy", Varying))
 struct CONTROLRIG_API FRigUnit_HierarchySetParentWeights : public FRigUnit_DynamicHierarchyBaseMutable
@@ -162,9 +186,16 @@ struct CONTROLRIG_API FRigUnit_HierarchySetParentWeights : public FRigUnit_Dynam
 	RIGVM_METHOD()
 	virtual void Execute(const FRigUnitContext& Context) override;
 
+	/*
+	 * The child to set the parents' weights for
+	 */
 	UPROPERTY(meta = (Input, ExpandByDefault))
 	FRigElementKey Child;
 
+	/*
+	 * The weights to set for the child's parents.
+	 * The number of weights needs to match the current number of parents.
+	 */
 	UPROPERTY(meta = (Input))
 	TArray<FRigElementWeight> Weights;
 };
\ No newline at end of file
