Control Rig: Support debugging of Sequencer Owned Control Rigs

diff --git a/Engine/Plugins/Animation/ControlRig/Source/ControlRigEditor/Private/EditMode/ControlRigEditMode.h b/Engine/Plugins/Animation/ControlRig/Source/ControlRigEditor/Private/EditMode/ControlRigEditMode.h
--- a/Engine/Plugins/Animation/ControlRig/Source/ControlRigEditor/Private/EditMode/ControlRigEditMode.h
+++ b/Engine/Plugins/Animation/ControlRig/Source/ControlRigEditor/Private/EditMode/ControlRigEditMode.h
@@ -357,6 +357,9 @@ private:
 	// Object binding
 	/** Setup bindings to a runtime object (or clear by passing in nullptr). */
 	void SetObjectBinding(TSharedPtr<IControlRigObjectBinding> InObjectBinding);
+
+protected:
+	
 	/** Get bindings to a runtime object */
 	TSharedPtr<IControlRigObjectBinding> GetObjectBinding() const;
 
diff --git a/Engine/Plugins/Animation/ControlRig/Source/ControlRigEditor/Private/Editor/ControlRigEditor.cpp b/Engine/Plugins/Animation/ControlRig/Source/ControlRigEditor/Private/Editor/ControlRigEditor.cpp
--- a/Engine/Plugins/Animation/ControlRig/Source/ControlRigEditor/Private/Editor/ControlRigEditor.cpp
+++ b/Engine/Plugins/Animation/ControlRig/Source/ControlRigEditor/Private/Editor/ControlRigEditor.cpp
@@ -1282,6 +1282,14 @@ void FControlRigEditor::HandleSetObjectBeingDebugged(UObject* InObject)
 		}
 
 		DebuggedControlRig->GetHierarchy()->OnModified().AddSP(this, &FControlRigEditor::OnHierarchyModified_AnyThread);
+
+		if(!bIsExternalControlRig)
+		{
+			if(EditorSkelComp)
+			{
+				EditorSkelComp->SetComponentToWorld(FTransform::Identity);
+			}
+		}
 	}
 	else
 	{
@@ -3012,6 +3020,21 @@ void FControlRigEditor::Tick(float DeltaTime)
 			bDrawHierarchyBones = true;
 		}
 
+		if(UControlRig* DebuggedControlRig = Cast<UControlRig>(GetBlueprintObj()->GetObjectBeingDebugged()))
+		{
+			const bool bIsExternalControlRig = DebuggedControlRig != ControlRig;			
+			if(bIsExternalControlRig)
+			{
+				if(UControlRigSkeletalMeshComponent* EditorSkelComp = Cast<UControlRigSkeletalMeshComponent>(GetPersonaToolkit()->GetPreviewScene()->GetPreviewMeshComponent()))
+				{
+					if(DebuggedControlRig->OuterSceneComponent.IsValid())
+					{
+						EditorSkelComp->SetComponentToWorld(DebuggedControlRig->OuterSceneComponent->GetComponentToWorld());
+					}
+				}
+			}
+		}
+
 		if (LastDebuggedRig != GetCustomDebugObjectLabel(Blueprint->GetObjectBeingDebugged()))
 		{
 			TArray<FCustomDebugObject> DebugList;
diff --git a/Engine/Plugins/Animation/ControlRig/Source/ControlRigEditor/Private/Editor/ControlRigEditorEditMode.cpp b/Engine/Plugins/Animation/ControlRig/Source/ControlRigEditor/Private/Editor/ControlRigEditorEditMode.cpp
--- a/Engine/Plugins/Animation/ControlRig/Source/ControlRigEditor/Private/Editor/ControlRigEditorEditMode.cpp
+++ b/Engine/Plugins/Animation/ControlRig/Source/ControlRigEditor/Private/Editor/ControlRigEditorEditMode.cpp
@@ -110,6 +110,12 @@ void FControlRigEditorEditMode::Render(const FSceneView* View, FViewport* Viewpo
 
 bool FControlRigEditorEditMode::GetCameraTarget(FSphere& OutTarget) const
 {
+	FTransform ComponentToWorld = FTransform::Identity;
+	if(const USceneComponent* SceneComponent = GetHostingSceneComponent())
+	{
+		ComponentToWorld = SceneComponent->GetComponentToWorld();
+	}
+
 	FBox Box(ForceInit);
 	TArray<FRigElementKey> SelectedRigElements = GetSelectedRigElements();
 	for (int32 Index = 0; Index < SelectedRigElements.Num(); ++Index)
@@ -118,12 +124,14 @@ bool FControlRigEditorEditMode::GetCameraTarget(FSphere& OutTarget) const
 		if (SelectedRigElements[Index].Type == ERigElementType::Bone || SelectedRigElements[Index].Type == ERigElementType::Null)
 		{
 			FTransform Transform = OnGetRigElementTransformDelegate.Execute(SelectedRigElements[Index], false, true);
+			Transform = Transform * ComponentToWorld;
 			Box += Transform.TransformPosition(FVector::OneVector * Radius);
 			Box += Transform.TransformPosition(FVector::OneVector * -Radius);
 		}
 		else if (SelectedRigElements[Index].Type == ERigElementType::Control)
 		{
 			FTransform Transform = OnGetRigElementTransformDelegate.Execute(SelectedRigElements[Index], false, true);
+			Transform = Transform * ComponentToWorld;
 			Box += Transform.TransformPosition(FVector::OneVector * Radius);
 			Box += Transform.TransformPosition(FVector::OneVector * -Radius);
 		}
