Sequencer: Pivot Tool fix for keyed Controls, don't evaluate if we are at the same time.
#trivial

diff --git a/Engine/Plugins/Animation/ControlRig/Source/ControlRigEditor/Private/Tools/ControlRigSnapper.cpp b/Engine/Plugins/Animation/ControlRig/Source/ControlRigEditor/Private/Tools/ControlRigSnapper.cpp
--- a/Engine/Plugins/Animation/ControlRig/Source/ControlRigEditor/Private/Tools/ControlRigSnapper.cpp
+++ b/Engine/Plugins/Animation/ControlRig/Source/ControlRigEditor/Private/Tools/ControlRigSnapper.cpp
@@ -111,7 +111,7 @@ TWeakPtr<ISequencer> FControlRigSnapper::GetSequencer()
 	return WeakSequencer;
 }
 
-static bool LocalGetControlRigControlTransforms(IMovieScenePlayer* Player, UMovieSceneSequence* MovieSceneSequence, FMovieSceneSequenceIDRef Template, FMovieSceneSequenceTransform& RootToLocalTransform,
+static bool LocalGetControlRigControlTransforms(IMovieScenePlayer* Player, const TOptional<FFrameNumber>& CurrentFrame, UMovieSceneSequence* MovieSceneSequence, FMovieSceneSequenceIDRef Template, FMovieSceneSequenceTransform& RootToLocalTransform,
 	UControlRig* ControlRig, const FName& ControlName,
 	const TArray<FFrameNumber>& Frames, const TArray<FTransform>& ParentTransforms, TArray<FTransform>& OutTransforms)
 {
@@ -135,11 +135,12 @@ static bool LocalGetControlRigControlTransforms(IMovieScenePlayer* Player, UMovi
 		for (int32 Index = 0; Index < Frames.Num(); ++Index)
 		{
 			const FFrameNumber& FrameNumber = Frames[Index];
-			FFrameTime GlobalTime(FrameNumber);
-
-			FMovieSceneContext Context = FMovieSceneContext(FMovieSceneEvaluationRange(GlobalTime, TickResolution), Player->GetPlaybackStatus()).SetHasJumped(true);
-
-			Player->GetEvaluationTemplate().Evaluate(Context, *Player);
+			if (CurrentFrame.IsSet() == false || CurrentFrame.GetValue() != FrameNumber)
+			{
+				FFrameTime GlobalTime(FrameNumber);
+				FMovieSceneContext Context = FMovieSceneContext(FMovieSceneEvaluationRange(GlobalTime, TickResolution), Player->GetPlaybackStatus()).SetHasJumped(true);
+				Player->GetEvaluationTemplate().Evaluate(Context, *Player);
+			}
 			ControlRig->Evaluate_AnyThread();
 			OutTransforms[Index] = ControlRig->GetControlGlobalTransform(ControlName) * ParentTransforms[Index];
 		}
@@ -154,7 +155,8 @@ bool FControlRigSnapper::GetControlRigControlTransforms(ISequencer* Sequencer,
 	{
 		FMovieSceneSequenceIDRef Template = Sequencer->GetFocusedTemplateID();
 		FMovieSceneSequenceTransform RootToLocalTransform;
-		return LocalGetControlRigControlTransforms(Sequencer, Sequencer->GetFocusedMovieSceneSequence(), Template, RootToLocalTransform,
+		TOptional<FFrameNumber> FrameNumber = Sequencer->GetLocalTime().Time.RoundToFrame();
+		return LocalGetControlRigControlTransforms(Sequencer, FrameNumber, Sequencer->GetFocusedMovieSceneSequence(), Template, RootToLocalTransform,
 			ControlRig, ControlName, Frames, ParentTransforms, OutTransforms);
 	
 	}
@@ -174,7 +176,8 @@ bool FControlRigSnapper::GetControlRigControlTransforms(UWorld* World,ULevelSequ
 		ULevelSequencePlayer* Player = ULevelSequencePlayer::CreateLevelSequencePlayer(World, LevelSequence, Settings, OutActor);
 		Player->Initialize(LevelSequence, World->PersistentLevel, Settings, CameraSettings);
 		Player->State.AssignSequence(MovieSceneSequenceID::Root, *LevelSequence, *Player);
-		bool Success = LocalGetControlRigControlTransforms(Player, LevelSequence, Template, RootToLocalTransform,
+		TOptional<FFrameNumber> OptFrame;
+		bool Success = LocalGetControlRigControlTransforms(Player, OptFrame, LevelSequence, Template, RootToLocalTransform,
 			ControlRig, ControlName, Frames, ParentTransforms, OutTransforms);
 		World->DestroyActor(OutActor);
 		return Success;
