Control Rig: Rename "Hide Manipulators" to "Hide Control Shapes" and "Shape Scale" to "Gizmo Scale".

diff --git a/Engine/Plugins/Animation/ControlRig/Source/ControlRigEditor/Private/EditMode/ControlRigEditMode.cpp b/Engine/Plugins/Animation/ControlRig/Source/ControlRigEditor/Private/EditMode/ControlRigEditMode.cpp
--- a/Engine/Plugins/Animation/ControlRig/Source/ControlRigEditor/Private/EditMode/ControlRigEditMode.cpp
+++ b/Engine/Plugins/Animation/ControlRig/Source/ControlRigEditor/Private/EditMode/ControlRigEditMode.cpp
@@ -692,7 +692,7 @@ void FControlRigEditMode::Render(const FSceneView* View, FViewport* Viewport, FP
 
 	const UControlRigEditModeSettings* Settings = GetDefault<UControlRigEditModeSettings>();
 
-	bool bRender = !Settings->bHideManipulators;
+	bool bRender = !Settings->bHideControlShapes;
 
 	FTransform ComponentTransform = GetHostingSceneComponentTransform();
 	if (bRender)
@@ -2155,22 +2155,22 @@ void FControlRigEditMode::FrameItems(const TArray<FRigElementKey>& InItems)
 void FControlRigEditMode::IncreaseShapeSize()
 {
 	UControlRigEditModeSettings* Settings = GetMutableDefault<UControlRigEditModeSettings>();
-	Settings->ShapeScale += 0.1f;
-	GetModeManager()->SetWidgetScale(Settings->ShapeScale);
+	Settings->GizmoScale += 0.1f;
+	GetModeManager()->SetWidgetScale(Settings->GizmoScale);
 }
 
 void FControlRigEditMode::DecreaseShapeSize()
 {
 	UControlRigEditModeSettings* Settings = GetMutableDefault<UControlRigEditModeSettings>();
-	Settings->ShapeScale -= 0.1f;
-	GetModeManager()->SetWidgetScale(Settings->ShapeScale);
+	Settings->GizmoScale -= 0.1f;
+	GetModeManager()->SetWidgetScale(Settings->GizmoScale);
 }
 
 void FControlRigEditMode::ResetControlShapeSize()
 {
 	UControlRigEditModeSettings* Settings = GetMutableDefault<UControlRigEditModeSettings>();
-	Settings->ShapeScale = 1.0f;
-	GetModeManager()->SetWidgetScale(Settings->ShapeScale);
+	Settings->GizmoScale = 1.0f;
+	GetModeManager()->SetWidgetScale(Settings->GizmoScale);
 }
 
 void FControlRigEditMode::ToggleControlShapeTransformEdit()
@@ -2319,7 +2319,7 @@ void FControlRigEditMode::ToggleManipulators()
 {
 	// Toggle flag (is used in drawing code)
 	UControlRigEditModeSettings* Settings = GetMutableDefault<UControlRigEditModeSettings>();
-	Settings->bHideManipulators = !Settings->bHideManipulators;
+	Settings->bHideControlShapes = !Settings->bHideControlShapes;
 }
 
 void FControlRigEditMode::ResetTransforms(bool bSelectionOnly)
@@ -3141,10 +3141,10 @@ void FControlRigEditMode::TickControlShape(AControlRigShapeActor* ShapeActor, co
 			if (FRigControlElement* ControlElement = ControlRig->FindControl(ShapeActor->ControlName))
 			{
 				ShapeActor->SetShapeColor(ControlElement->Settings.ShapeColor);
-				ShapeActor->SetIsTemporarilyHiddenInEditor(!ControlElement->Settings.bShapeVisible || Settings->bHideManipulators);
+				ShapeActor->SetIsTemporarilyHiddenInEditor(!ControlElement->Settings.bShapeVisible || Settings->bHideControlShapes);
 				if (!IsInLevelEditor()) //don't change this in level editor otherwise we can never select it
 				{
-					ShapeActor->SetSelectable(ControlElement->Settings.bShapeVisible && !Settings->bHideManipulators && ControlElement->Settings.bAnimatable);
+					ShapeActor->SetSelectable(ControlElement->Settings.bShapeVisible && !Settings->bHideControlShapes && ControlElement->Settings.bAnimatable);
 				}
 			}
 		}
diff --git a/Engine/Plugins/Animation/ControlRig/Source/ControlRigEditor/Private/EditMode/ControlRigEditModeSettings.cpp b/Engine/Plugins/Animation/ControlRig/Source/ControlRigEditor/Private/EditMode/ControlRigEditModeSettings.cpp
--- a/Engine/Plugins/Animation/ControlRig/Source/ControlRigEditor/Private/EditMode/ControlRigEditModeSettings.cpp
+++ b/Engine/Plugins/Animation/ControlRig/Source/ControlRigEditor/Private/EditMode/ControlRigEditModeSettings.cpp
@@ -13,10 +13,10 @@ void UControlRigEditModeSettings::PostEditChangeProperty(struct FPropertyChanged
 {
 	Super::PostEditChangeProperty(PropertyChangedEvent);
 #if WITH_EDITOR
-	if (PropertyChangedEvent.Property && PropertyChangedEvent.Property->GetFName() == GET_MEMBER_NAME_CHECKED(UControlRigEditModeSettings, ShapeScale))
+	if (PropertyChangedEvent.Property && PropertyChangedEvent.Property->GetFName() == GET_MEMBER_NAME_CHECKED(UControlRigEditModeSettings, GizmoScale))
 	{
 		FEditorModeTools& Tools = GLevelEditorModeTools();
-		Tools.SetWidgetScale(ShapeScale);
+		Tools.SetWidgetScale(GizmoScale);
 	}
 #endif
 	if (PropertyChangedEvent.ChangeType != EPropertyChangeType::Interactive)
@@ -30,7 +30,7 @@ void UControlRigEditModeSettings::PostEditChangeProperty(struct FPropertyChanged
 void UControlRigEditModeSettings::PostEditUndo()
 {
 	FEditorModeTools& Tools = GLevelEditorModeTools();
-	Tools.SetWidgetScale(ShapeScale);
+	Tools.SetWidgetScale(GizmoScale);
 }
 #endif
 
diff --git a/Engine/Plugins/Animation/ControlRig/Source/ControlRigEditor/Private/EditMode/ControlRigEditModeSettings.h b/Engine/Plugins/Animation/ControlRig/Source/ControlRigEditor/Private/EditMode/ControlRigEditModeSettings.h
--- a/Engine/Plugins/Animation/ControlRig/Source/ControlRigEditor/Private/EditMode/ControlRigEditModeSettings.h
+++ b/Engine/Plugins/Animation/ControlRig/Source/ControlRigEditor/Private/EditMode/ControlRigEditModeSettings.h
@@ -16,13 +16,13 @@ class UControlRigEditModeSettings : public UObject
 	UControlRigEditModeSettings()
 		: bDisplayHierarchy(false)
 		, bDisplayNulls(false)
-		, bHideManipulators(false)
+		, bHideControlShapes(false)
 		, bDisplayAxesOnSelection(false)
 		, AxisScale(10.f)
 		, bCoordSystemPerWidgetMode(true)
 		, bOnlySelectRigControls(false)
 		, bLocalTransformsInEachLocalSpace(true)
-		, ShapeScale(1.0f)
+		, GizmoScale(1.0f)
 	{
 		LastInViewportTweenWidgetLocation = FVector2D(EForceInit::ForceInitToZero);
 	}
@@ -43,9 +43,9 @@ public:
 	UPROPERTY(config, EditAnywhere, AdvancedDisplay, Category = "Animation Settings")
 	bool bDisplayNulls;
 
-	/** Should we always hide manipulators in viewport */
+	/** Should we always hide control shapes in viewport */
 	UPROPERTY(config, EditAnywhere, AdvancedDisplay, Category = "Animation Settings")
-	bool bHideManipulators;
+	bool bHideControlShapes;
 
 	/** Should we show axes for the selected elements */
 	UPROPERTY(config, EditAnywhere, AdvancedDisplay, Category = "Animation Settings")
@@ -69,7 +69,7 @@ public:
 	
 	/** The scale for Gizmos */
 	UPROPERTY(config, EditAnywhere, AdvancedDisplay, Category = "Animation Settings")
-	float ShapeScale;
+	float GizmoScale;
 
 	UPROPERTY(config)
 	FVector2D LastInViewportTweenWidgetLocation;
