Control Rig: Change the order of the represented dependencies in the VM map

diff --git a/Engine/Plugins/Animation/ControlRig/Source/ControlRig/Private/Rigs/RigHierarchy.cpp b/Engine/Plugins/Animation/ControlRig/Source/ControlRig/Private/Rigs/RigHierarchy.cpp
--- a/Engine/Plugins/Animation/ControlRig/Source/ControlRig/Private/Rigs/RigHierarchy.cpp
+++ b/Engine/Plugins/Animation/ControlRig/Source/ControlRig/Private/Rigs/RigHierarchy.cpp
@@ -2983,14 +2983,14 @@ bool URigHierarchy::IsDependentOn(FRigBaseElement* InDependent, FRigBaseElement*
 	}
 
 	// check the optional dependency map
-	if(const TArray<int32>* DependentIndicesPtr = InDependencyMap.Find(InDependency->GetIndex()))
+	if(const TArray<int32>* DependentIndicesPtr = InDependencyMap.Find(InDependent->GetIndex()))
 	{
 		const TArray<int32>& DependentIndices = *DependentIndicesPtr;
 		for(int32 DependentIndex : DependentIndices)
 		{
 			ensure(Elements.IsValidIndex(DependentIndex));
 			
-			if(IsDependentOn(InDependent, Elements[DependentIndex], InDependencyMap))
+			if(IsDependentOn(Elements[DependentIndex], InDependency, InDependencyMap))
 			{
 				return true;
 			}
@@ -3948,7 +3948,7 @@ URigHierarchy::TElementDependencyMap URigHierarchy::GetDependenciesForVM(URigVM*
 				{
 					if(ReadTransform != WrittenTransform)
 					{
-						Dependencies.FindOrAdd(ReadTransform).AddUnique(WrittenTransform);
+						Dependencies.FindOrAdd(WrittenTransform).AddUnique(ReadTransform);
 					}
 				}
 			}
diff --git a/Engine/Plugins/Animation/ControlRig/Source/ControlRigDeveloper/Private/ControlRigVisualGraphUtils.cpp b/Engine/Plugins/Animation/ControlRig/Source/ControlRigDeveloper/Private/ControlRigVisualGraphUtils.cpp
--- a/Engine/Plugins/Animation/ControlRig/Source/ControlRigDeveloper/Private/ControlRigVisualGraphUtils.cpp
+++ b/Engine/Plugins/Animation/ControlRig/Source/ControlRigDeveloper/Private/ControlRigVisualGraphUtils.cpp
@@ -233,9 +233,9 @@ FString FControlRigVisualGraphUtils::DumpRigHierarchyToDotGraph(URigHierarchy* I
 
 			for(const URigHierarchy::TElementDependencyMapPair& Dependency : Dependencies)
 			{
-				const int32 SourceElementIndex = Dependency.Key;
+				const int32 TargetElementIndex = Dependency.Key;
 				
-				for(const int32 TargetElementIndex : Dependency.Value)
+				for(const int32 SourceElementIndex : Dependency.Value)
 				{
 					FName EdgeName = *FString::Printf(TEXT("Dependency_%d_%d"), SourceElementIndex, TargetElementIndex);
 					if(Graph.FindEdge(EdgeName) != INDEX_NONE)
