Control Rig: Fix python log printing wrong control type when rotator control is created

diff --git a/Engine/Plugins/Animation/ControlRig/Source/ControlRig/Private/Rigs/RigHierarchy.cpp b/Engine/Plugins/Animation/ControlRig/Source/ControlRig/Private/Rigs/RigHierarchy.cpp
--- a/Engine/Plugins/Animation/ControlRig/Source/ControlRig/Private/Rigs/RigHierarchy.cpp
+++ b/Engine/Plugins/Animation/ControlRig/Source/ControlRig/Private/Rigs/RigHierarchy.cpp
@@ -4815,7 +4815,7 @@ TArray<FString> URigHierarchy::ControlSettingsToPythonCommands(const FRigControl
 		case ERigControlType::Float: TypeStr = TEXT("FLOAT"); break;
 		case ERigControlType::Integer: TypeStr = TEXT("INTEGER"); break;
 		case ERigControlType::Position: TypeStr = TEXT("POSITION"); break;
-		case ERigControlType::Rotator: TypeStr = TEXT("POSITION"); break;
+		case ERigControlType::Rotator: TypeStr = TEXT("ROTATOR"); break;
 		case ERigControlType::Scale: TypeStr = TEXT("SCALE"); break;
 		case ERigControlType::Transform: TypeStr = TEXT("EULER_TRANSFORM"); break;
 		case ERigControlType::EulerTransform: TypeStr = TEXT("EULER_TRANSFORM"); break;
