NonUnity Build fix

diff --git a/Engine/Plugins/Animation/ControlRig/Source/ControlRigEditor/Public/RigSpacePickerBakeSettings.h b/Engine/Plugins/Animation/ControlRig/Source/ControlRigEditor/Public/RigSpacePickerBakeSettings.h
--- a/Engine/Plugins/Animation/ControlRig/Source/ControlRigEditor/Public/RigSpacePickerBakeSettings.h
+++ b/Engine/Plugins/Animation/ControlRig/Source/ControlRigEditor/Public/RigSpacePickerBakeSettings.h
@@ -20,7 +20,7 @@ struct CONTROLRIGEDITOR_API FRigSpacePickerBakeSettings
 		Tolerance = 0.001f;
 	}
 
-	UPROPERTY(BlueprintReadWrite, Category = "Settings")
+	UPROPERTY(VisibleAnywhere,BlueprintReadWrite, Category = "Settings")
 	FRigElementKey TargetSpace;
 
 	UPROPERTY(BlueprintReadWrite,EditAnywhere, Category = "Settings")
