Refactor Control Rig Element Detaills to use new transform widgets

diff --git a/Engine/Plugins/Animation/ControlRig/Source/ControlRigEditor/ControlRigEditor.Build.cs b/Engine/Plugins/Animation/ControlRig/Source/ControlRigEditor/ControlRigEditor.Build.cs
--- a/Engine/Plugins/Animation/ControlRig/Source/ControlRigEditor/ControlRigEditor.Build.cs
+++ b/Engine/Plugins/Animation/ControlRig/Source/ControlRigEditor/ControlRigEditor.Build.cs
@@ -84,6 +84,7 @@ namespace UnrealBuildTool.Rules
 					"PythonScriptPlugin",
 					"AdvancedPreviewScene",
 					"ToolWidgets",
+                    "AnimationWidgets",
 				}
             );
 
diff --git a/Engine/Plugins/Animation/ControlRig/Source/ControlRigEditor/Private/ControlRigElementDetails.cpp b/Engine/Plugins/Animation/ControlRig/Source/ControlRigEditor/Private/ControlRigElementDetails.cpp
--- a/Engine/Plugins/Animation/ControlRig/Source/ControlRigEditor/Private/ControlRigElementDetails.cpp
+++ b/Engine/Plugins/Animation/ControlRig/Source/ControlRigEditor/Private/ControlRigElementDetails.cpp
@@ -2,9 +2,6 @@
 
 #include "ControlRigElementDetails.h"
 #include "Widgets/SWidget.h"
-#include "DetailLayoutBuilder.h"
-#include "DetailCategoryBuilder.h"
-#include "DetailWidgetRow.h"
 #include "IDetailChildrenBuilder.h"
 #include "Widgets/Text/STextBlock.h"
 #include "Widgets/Input/SEditableTextBox.h"
@@ -16,7 +13,6 @@
 #include "Widgets/Input/SButton.h"
 #include "Editor/SControlRigGizmoNameList.h"
 #include "ControlRigBlueprint.h"
-#include "IDetailGroup.h"
 #include "Graph/ControlRigGraph.h"
 #include "PropertyCustomizationHelpers.h"
 #include "SEnumCombo.h"
@@ -26,6 +22,8 @@
 #include "Graph/SControlRigGraphPinVariableBinding.h"
 #include "HAL/PlatformApplicationMisc.h"
 #include "ScopedTransaction.h"
+#include "AnimationCoreLibrary.h"
+#include "HAL/PlatformApplicationMisc.h"
 
 #define LOCTEXT_NAMESPACE "ControlRigElementDetails"
 
@@ -900,6 +898,36 @@ TArray<FRigElementKey> FRigBaseElementDetails::GetElementKeys() const
 	return Keys;
 }
 
+bool FRigBaseElementDetails::IsAnyElementOfType(ERigElementType InType) const
+{
+	for(TWeakObjectPtr<UDetailsViewWrapperObject> ObjectBeingCustomized : ObjectsBeingCustomized)
+	{
+		if(ObjectBeingCustomized.IsValid())
+		{
+			if(ObjectBeingCustomized->GetContent<FRigBaseElement>().GetKey().Type == InType)
+			{
+				return true;
+			}
+		}
+	}
+	return false;
+}
+
+bool FRigBaseElementDetails::IsAnyElementNotOfType(ERigElementType InType) const
+{
+	for(TWeakObjectPtr<UDetailsViewWrapperObject> ObjectBeingCustomized : ObjectsBeingCustomized)
+	{
+		if(ObjectBeingCustomized.IsValid())
+		{
+			if(ObjectBeingCustomized->GetContent<FRigBaseElement>().GetKey().Type != InType)
+			{
+				return true;
+			}
+		}
+	}
+	return false;
+}
+
 bool FRigBaseElementDetails::IsAnyControlOfType(ERigControlType InType) const
 {
 	for(TWeakObjectPtr<UDetailsViewWrapperObject> ObjectBeingCustomized : ObjectsBeingCustomized)
@@ -949,6 +977,17 @@ void FRigBaseElementDetails::RegisterSectionMappings(FPropertyEditorModule& Prop
 {
 }
 
+namespace ERigTransformElementDetailsTransform
+{
+	enum Type
+	{
+		Initial,
+		Current,
+		Offset,
+		Max
+	};
+}
+
 void FRigTransformElementDetails::CustomizeDetails(IDetailLayoutBuilder& DetailBuilder)
 {
 	FRigBaseElementDetails::CustomizeDetails(DetailBuilder);
@@ -965,51 +1004,577 @@ void FRigTransformElementDetails::RegisterSectionMappings(FPropertyEditorModule&
 
 void FRigTransformElementDetails::CustomizeTransform(IDetailLayoutBuilder& DetailBuilder)
 {
-	IDetailCategoryBuilder& TransformCategory = DetailBuilder.EditCategory(TEXT("Transform"), LOCTEXT("Transform", "Transform"));
+	URigHierarchy* HierarchyBeingDebugged = HierarchyBeingCustomized;
+	bool IsSetupModeEnabled = false;
+	if (UControlRig* DebuggedRig = Cast<UControlRig>(BlueprintBeingCustomized->GetObjectBeingDebugged()))
 	{
-		const TSharedPtr<IPropertyHandle> PoseHandle = DetailBuilder.GetProperty(TEXT("Pose"));
-		const TSharedPtr<IPropertyHandle> InitialHandle = PoseHandle->GetChildHandle(TEXT("Initial"));
-		const TSharedPtr<IPropertyHandle> CurrentHandle = PoseHandle->GetChildHandle(TEXT("Current"));
+		IsSetupModeEnabled = DebuggedRig->IsSetupModeEnabled(); 
+ 		if(!IsSetupModeEnabled)
+		{
+			HierarchyBeingDebugged = DebuggedRig->GetHierarchy();
+		}
+	}
 
-		// setup initial global
+	TArray<FRigElementKey> Keys = GetElementKeys();
+	Keys = HierarchyBeingDebugged->SortKeys(Keys);
+	const bool bAllControls = !IsAnyElementNotOfType(ERigElementType::Control);
+	TArray<ERigTransformElementDetailsTransform::Type> TransformTypes;
+	TArray<FText> ButtonLabels;
+	TArray<FText> ButtonTooltips;
+
+	if(bAllControls)
+	{
+		TransformTypes = {
+			ERigTransformElementDetailsTransform::Initial,
+			ERigTransformElementDetailsTransform::Current,
+			ERigTransformElementDetailsTransform::Offset
+		};
+		ButtonLabels = {
+			LOCTEXT("Initial", "Initial"),
+			LOCTEXT("Current", "Current"),
+			LOCTEXT("Offset", "Offset")
+		};
+		ButtonTooltips = {
+			LOCTEXT("InitialTooltip", "Initial transform in the reference pose"),
+			LOCTEXT("CurrentTooltip", "Current animation transform"),
+			LOCTEXT("OffsetTooltip", "Offset transform under the control")
+		};
+	}
+	else
+	{
+		TransformTypes = {
+			ERigTransformElementDetailsTransform::Initial,
+			ERigTransformElementDetailsTransform::Current
+		};
+		ButtonLabels = {
+			LOCTEXT("Initial", "Initial"),
+			LOCTEXT("Current", "Current")
+		};
+		ButtonTooltips = {
+			LOCTEXT("InitialTooltip", "Initial transform in the reference pose"),
+			LOCTEXT("CurrentTooltip", "Current animation transform")
+		};
+	}
+
+	TArray<bool> bTransformsEnabled;
+
+	// determine if the transforms are enabled
+	for(int32 Index = 0; Index < TransformTypes.Num(); Index++)
+	{
+		ERigTransformElementDetailsTransform::Type CurrentTransformType = (ERigTransformElementDetailsTransform::Type)Index;
+
+		bool bIsTransformEnabled = true;
+		if (IsAnyElementOfType(ERigElementType::Control))
 		{
-			const TSharedPtr<IPropertyHandle> GlobalHandle = InitialHandle->GetChildHandle(TEXT("Global"));
-			const TSharedPtr<IPropertyHandle> TransformHandle = GlobalHandle->GetChildHandle(TEXT("Transform"));
+			bIsTransformEnabled = IsAnyControlOfType(ERigControlType::EulerTransform) ||
+				IsAnyControlOfType(ERigControlType::Transform) ||
+				CurrentTransformType == ERigTransformElementDetailsTransform::Offset;
 
-			TransformCategory.AddProperty(TransformHandle.ToSharedRef(), EPropertyLocation::Advanced)
-				.DisplayName(FText::FromString(TEXT("Initial Global")))
-				.IsEnabled(TAttribute<bool>::Create(TAttribute<bool>::FGetter::CreateSP(this, &FRigBaseElementDetails::IsSetupModeEnabled)));
+			if(!bIsTransformEnabled)
+			{
+				ButtonTooltips[Index] = FText::FromString(
+					FString::Printf(TEXT("%s\n%s"),
+						*ButtonTooltips[Index].ToString(), 
+						TEXT("Only transform controls can be edited here. Refer to the 'Value' section instead.")));
+			}
 		}
-
-		// setup initial local
+		else if (IsAnyElementOfType(ERigElementType::Bone) && CurrentTransformType == ERigTransformElementDetailsTransform::Initial)
 		{
-			const TSharedPtr<IPropertyHandle> LocalHandle = InitialHandle->GetChildHandle(TEXT("Local"));
-			const TSharedPtr<IPropertyHandle> TransformHandle = LocalHandle->GetChildHandle(TEXT("Transform"));
+			for(TWeakObjectPtr<UDetailsViewWrapperObject> ObjectBeingCustomized : ObjectsBeingCustomized)
+			{
+				if(ObjectBeingCustomized.IsValid())
+				{
+					if(ObjectBeingCustomized->IsChildOf<FRigBoneElement>())
+					{
+						const FRigBoneElement BoneElement = ObjectBeingCustomized->GetContent<FRigBoneElement>();
+						bIsTransformEnabled = BoneElement.BoneType == ERigBoneType::User;
 
-			TransformCategory.AddProperty(TransformHandle.ToSharedRef(), EPropertyLocation::Advanced)
-				.DisplayName(FText::FromString(TEXT("Initial Local")))
-				.IsEnabled(TAttribute<bool>::Create(TAttribute<bool>::FGetter::CreateSP(this, &FRigBaseElementDetails::IsSetupModeEnabled)));
+						if(!bIsTransformEnabled)
+						{
+							ButtonTooltips[Index] = FText::FromString(
+								FString::Printf(TEXT("%s\n%s"),
+									*ButtonTooltips[Index].ToString(), 
+									TEXT("Imported Bones' initial transform cannot be edited.")));
+						}
+					}
+				}
+			}			
 		}
 
-		// setup current global
+		bTransformsEnabled.Add(bIsTransformEnabled);
+	}
+
+	static TAttribute<TArray<ERigTransformElementDetailsTransform::Type>> PickedTransforms = 
+		TArray<ERigTransformElementDetailsTransform::Type>({ERigTransformElementDetailsTransform::Current});
+
+	TSharedPtr<SSegmentedControl<ERigTransformElementDetailsTransform::Type>> TransformChoiceWidget =
+		SSegmentedControl<ERigTransformElementDetailsTransform::Type>::Create(
+			TransformTypes,
+			ButtonLabels,
+			ButtonTooltips,
+			PickedTransforms
+		);
+
+	IDetailCategoryBuilder& TransformCategory = DetailBuilder.EditCategory(TEXT("Transform"), LOCTEXT("Transform", "Transform"));
+	AddChoiceWidgetRow(TransformCategory, FText::FromString(TEXT("TransformType")), TransformChoiceWidget.ToSharedRef());
+	
+	TSharedRef<UE::Math::TVector<float>> IsComponentRelative = MakeShareable(new UE::Math::TVector<float>(1.f, 1.f, 1.f));
+
+	using TransformType = FEulerTransform;
+	using NumericType = FVector::FReal;
+
+	SAdvancedTransformInputBox<TransformType>::FArguments TransformWidgetArgs = SAdvancedTransformInputBox<TransformType>::FArguments()
+		.DisplayToggle(false)
+		.DisplayRelativeWorld(true)
+		.Font(IDetailLayoutBuilder::GetDetailFont())
+		.OnGetIsComponentRelative_Lambda(
+			[IsComponentRelative](ESlateTransformComponent::Type InComponent)
+			{
+				return IsComponentRelative->operator[]((int32)InComponent) > 0.f;
+			})
+		.OnIsComponentRelativeChanged_Lambda(
+			[IsComponentRelative](ESlateTransformComponent::Type InComponent, bool bIsRelative)
+			{
+				IsComponentRelative->operator[]((int32)InComponent) = bIsRelative ? 1.f : 0.f;
+			});
+	
+	for(int32 Index = 0; Index < ButtonLabels.Num(); Index++)
+	{
+		ERigTransformElementDetailsTransform::Type CurrentTransformType = (ERigTransformElementDetailsTransform::Type)Index;
+
+		auto GetRelativeAbsoluteTransforms = [CurrentTransformType, Keys, HierarchyBeingDebugged](
+			const FRigElementKey& Key,
+			ERigTransformElementDetailsTransform::Type InTransformType = ERigTransformElementDetailsTransform::Max
+			) -> TPair<TransformType, TransformType>
 		{
-			const TSharedPtr<IPropertyHandle> GlobalHandle = CurrentHandle->GetChildHandle(TEXT("Global"));
-			const TSharedPtr<IPropertyHandle> TransformHandle = GlobalHandle->GetChildHandle(TEXT("Transform"));
+			if(InTransformType == ERigTransformElementDetailsTransform::Max)
+			{
+				InTransformType = CurrentTransformType;
+			}
 
-			TransformCategory.AddProperty(TransformHandle.ToSharedRef(), EPropertyLocation::Advanced)
-				.DisplayName(FText::FromString(TEXT("Current Global")))
-				.IsEnabled(TAttribute<bool>::Create(TAttribute<bool>::FGetter::CreateSP(this, &FRigBaseElementDetails::IsSetupModeEnabled)));
-		}
+			TransformType RelativeTransform = TransformType::Identity;
+			TransformType AbsoluteTransform = TransformType::Identity;
+
+			const bool bInitial = InTransformType == ERigTransformElementDetailsTransform::Initial; 
+			if(bInitial || InTransformType == ERigTransformElementDetailsTransform::Current)
+			{
+				RelativeTransform = HierarchyBeingDebugged->GetLocalTransform(Key, bInitial);
+				AbsoluteTransform = HierarchyBeingDebugged->GetGlobalTransform(Key, bInitial);
+			}
+			else
+			{
+				if(FRigControlElement* ControlElement = HierarchyBeingDebugged->Find<FRigControlElement>(Key))
+				{
+					RelativeTransform = HierarchyBeingDebugged->GetControlOffsetTransform(ControlElement, bInitial ? ERigTransformType::InitialLocal : ERigTransformType::CurrentLocal);
+					AbsoluteTransform = HierarchyBeingDebugged->GetControlOffsetTransform(ControlElement, bInitial ? ERigTransformType::InitialGlobal : ERigTransformType::CurrentGlobal);
+				}
+			}
+
+			return TPair<FTransform, FTransform>(RelativeTransform, AbsoluteTransform);
+		};
 
-		// setup current local
+		auto GetCombinedTransform = [IsComponentRelative, GetRelativeAbsoluteTransforms](
+			const FRigElementKey& Key,
+			ERigTransformElementDetailsTransform::Type InTransformType = ERigTransformElementDetailsTransform::Max
+			) -> TransformType
 		{
-			const TSharedPtr<IPropertyHandle> LocalHandle = CurrentHandle->GetChildHandle(TEXT("Local"));
-			const TSharedPtr<IPropertyHandle> TransformHandle = LocalHandle->GetChildHandle(TEXT("Transform"));
+			const TPair<TransformType, TransformType> TransformPair = GetRelativeAbsoluteTransforms(Key, InTransformType);
+			const TransformType RelativeTransform = TransformPair.Key;
+			const TransformType AbsoluteTransform = TransformPair.Value;
+
+			TransformType Xfo;
+			Xfo.SetLocation((IsComponentRelative->X > 0.f) ? RelativeTransform.GetLocation() : AbsoluteTransform.GetLocation());
+			Xfo.SetRotation((IsComponentRelative->Y > 0.f) ? RelativeTransform.GetRotation() : AbsoluteTransform.GetRotation());
+			Xfo.SetScale3D((IsComponentRelative->Z > 0.f) ? RelativeTransform.GetScale3D() : AbsoluteTransform.GetScale3D());
+
+			return Xfo;
+		};
+
+		auto GetSingleTransform = [GetRelativeAbsoluteTransforms](
+			const FRigElementKey& Key,
+			bool bIsRelative,
+			ERigTransformElementDetailsTransform::Type InTransformType = ERigTransformElementDetailsTransform::Max
+			) -> TransformType
+		{
+			const TPair<TransformType, TransformType> TransformPair = GetRelativeAbsoluteTransforms(Key, InTransformType);
+			const TransformType RelativeTransform = TransformPair.Key;
+			const TransformType AbsoluteTransform = TransformPair.Value;
+			return bIsRelative ? RelativeTransform : AbsoluteTransform;
+		};
+
+		auto SetSingleTransform = [CurrentTransformType, GetRelativeAbsoluteTransforms, this, IsSetupModeEnabled, HierarchyBeingDebugged](
+			const FRigElementKey& Key,
+			TransformType InTransform,
+			bool bIsRelative)
+		{
+			const bool bInitial = CurrentTransformType == ERigTransformElementDetailsTransform::Initial; 
 
-			TransformCategory.AddProperty(TransformHandle.ToSharedRef())
-				.DisplayName(FText::FromString(TEXT("Current Local")))
-				.IsEnabled(TAttribute<bool>::Create(TAttribute<bool>::FGetter::CreateSP(this, &FRigTransformElementDetails::IsCurrentLocalEnabled)));
-		}
+			TArray<URigHierarchy*> HierarchiesToUpdate;
+			HierarchiesToUpdate.Add(HierarchyBeingDebugged);
+			if(bInitial || IsSetupModeEnabled)
+			{
+				HierarchiesToUpdate.Add(HierarchyBeingCustomized);
+			}
+
+			for(URigHierarchy* HierarchyToUpdate : HierarchiesToUpdate)
+			{
+				if(bInitial || CurrentTransformType == ERigTransformElementDetailsTransform::Current)
+				{
+					if(bIsRelative)
+					{
+						HierarchyToUpdate->SetLocalTransform(Key, InTransform, bInitial, true, true);
+					}
+					else
+					{
+						HierarchyToUpdate->SetGlobalTransform(Key, InTransform, bInitial, true, true);
+					}
+				}
+				else
+				{
+					if(!bIsRelative)
+					{
+						const FTransform ParentTransform = HierarchyToUpdate->GetParentTransform(Key, bInitial);
+						InTransform = FTransform(InTransform).GetRelativeTransform(ParentTransform);
+					}
+							
+					if(FRigControlElement* ControlElement = HierarchyToUpdate->Find<FRigControlElement>(Key))
+					{
+						HierarchyToUpdate->SetControlOffsetTransform(Key, InTransform, bInitial, true, true);
+					}
+				}
+			}
+		};
+
+		TransformWidgetArgs.Visibility_Lambda([TransformChoiceWidget, Index]() -> EVisibility
+		{
+			return TransformChoiceWidget->HasValue((ERigTransformElementDetailsTransform::Type)Index) ? EVisibility::Visible : EVisibility::Collapsed;
+		});
+
+		TransformWidgetArgs.IsEnabled(bTransformsEnabled[Index]);
+
+		TransformWidgetArgs.OnGetNumericValue_Lambda([Keys, GetCombinedTransform](
+			ESlateTransformComponent::Type Component,
+			ESlateRotationRepresentation::Type Representation,
+			ESlateTransformSubComponent::Type SubComponent) -> TOptional<NumericType>
+		{
+			TOptional<NumericType> FirstValue;
+
+			for(int32 Index = 0; Index < Keys.Num(); Index++)
+			{
+				const FRigElementKey& Key = Keys[Index];
+				TransformType Xfo = GetCombinedTransform(Key);
+
+				TOptional<NumericType> CurrentValue = SAdvancedTransformInputBox<TransformType>::GetNumericValueFromTransform(Xfo, Component, Representation, SubComponent);
+				if(!CurrentValue.IsSet())
+				{
+					return CurrentValue;
+				}
+
+				if(Index == 0)
+				{
+					FirstValue = CurrentValue;
+				}
+				else
+				{
+					if(!FMath::IsNearlyEqual(FirstValue.GetValue(), CurrentValue.GetValue()))
+					{
+						return TOptional<NumericType>();
+					}
+				}
+			}
+			
+			return FirstValue;
+		});
+
+		TransformWidgetArgs.OnNumericValueChanged_Lambda([
+			Keys,
+			this,
+			IsComponentRelative,
+			GetSingleTransform,
+			SetSingleTransform
+		](
+			ESlateTransformComponent::Type Component,
+			ESlateRotationRepresentation::Type Representation,
+			ESlateTransformSubComponent::Type SubComponent,
+			NumericType InNumericValue)
+		{
+			const bool bIsRelative = IsComponentRelative->Component((int32)Component) > 0.f;
+
+			FScopedTransaction Transaction(LOCTEXT("ChangeNumericValue", "Change Numeric Value"));
+
+			for(const FRigElementKey& Key : Keys)
+			{
+				TransformType Transform = GetSingleTransform(Key, bIsRelative);
+				SAdvancedTransformInputBox<TransformType>::ApplyNumericValueChange(Transform, InNumericValue, Component, Representation, SubComponent);
+				SetSingleTransform(Key, Transform, bIsRelative);
+			}
+		});
+
+		TransformWidgetArgs.OnCopyToClipboard_Lambda([Keys, IsComponentRelative, GetSingleTransform](
+			ESlateTransformComponent::Type InComponent
+			)
+		{
+			if(Keys.Num() == 0)
+			{
+				return;
+			}
+
+			// make sure that we use the same relative setting on all components when copying
+			IsComponentRelative->Y = IsComponentRelative->Z = IsComponentRelative->X;
+			const bool bIsRelative = IsComponentRelative->Y > 0.f; 
+
+			const FRigElementKey& FirstKey = Keys[0];
+			TransformType Xfo = GetSingleTransform(FirstKey, bIsRelative);
+
+			FString Content;
+			switch(InComponent)
+			{
+				case ESlateTransformComponent::Location:
+				{
+					const FVector Data = Xfo.GetLocation();
+					TBaseStructure<FVector>::Get()->ExportText(Content, &Data, &Data, nullptr, PPF_None, nullptr);
+					break;
+				}
+				case ESlateTransformComponent::Rotation:
+				{
+					const FRotator Data = Xfo.Rotator();
+					TBaseStructure<FRotator>::Get()->ExportText(Content, &Data, &Data, nullptr, PPF_None, nullptr);
+					break;
+				}
+				case ESlateTransformComponent::Scale:
+				{
+					const FVector Data = Xfo.GetScale3D();
+					TBaseStructure<FVector>::Get()->ExportText(Content, &Data, &Data, nullptr, PPF_None, nullptr);
+					break;
+				}
+				case ESlateTransformComponent::Max:
+				default:
+				{
+					TBaseStructure<TransformType>::Get()->ExportText(Content, &Xfo, &Xfo, nullptr, PPF_None, nullptr);
+					break;
+				}
+			}
+
+			if(!Content.IsEmpty())
+			{
+				FPlatformApplicationMisc::ClipboardCopy(*Content);
+			}
+		});
+
+		TransformWidgetArgs.OnPasteFromClipboard_Lambda([Keys, IsComponentRelative, GetSingleTransform, SetSingleTransform](
+			ESlateTransformComponent::Type InComponent
+			)
+		{
+			if(Keys.Num() == 0)
+			{
+				return;
+			}
+			
+			
+			// make sure that we use the same relative setting on all components when pasting
+			IsComponentRelative->Y = IsComponentRelative->Z = IsComponentRelative->X;
+			const bool bIsRelative = IsComponentRelative->Y > 0.f; 
+
+			FString Content;
+			FPlatformApplicationMisc::ClipboardPaste(Content);
+
+			if(Content.IsEmpty())
+			{
+				return;
+			}
+
+			FScopedTransaction Transaction(LOCTEXT("PasteTransform", "Paste Transform"));
+
+			for(const FRigElementKey& Key : Keys)
+			{
+				TransformType Xfo = GetSingleTransform(Key, bIsRelative);
+
+				{
+					class FRigPasteTransformWidgetErrorPipe : public FOutputDevice
+					{
+					public:
+
+						int32 NumErrors;
+
+						FRigPasteTransformWidgetErrorPipe()
+							: FOutputDevice()
+							, NumErrors(0)
+						{
+						}
+
+						virtual void Serialize(const TCHAR* V, ELogVerbosity::Type Verbosity, const class FName& Category) override
+						{
+							UE_LOG(LogControlRig, Error, TEXT("Error Pasting to Widget: %s"), V);
+							NumErrors++;
+						}
+					};
+
+					FRigPasteTransformWidgetErrorPipe ErrorPipe;
+					
+					switch(InComponent)
+					{
+						case ESlateTransformComponent::Location:
+						{
+							FVector Data = Xfo.GetLocation();
+							TBaseStructure<FVector>::Get()->ImportText(*Content, &Data, nullptr, PPF_None, &ErrorPipe, TBaseStructure<FVector>::Get()->GetName(), true);
+							Xfo.SetLocation(Data);
+							break;
+						}
+						case ESlateTransformComponent::Rotation:
+						{
+							FRotator Data = Xfo.Rotator();
+							TBaseStructure<FRotator>::Get()->ImportText(*Content, &Data, nullptr, PPF_None, &ErrorPipe, TBaseStructure<FRotator>::Get()->GetName(), true);
+							Xfo.SetRotator(Data);
+							break;
+						}
+						case ESlateTransformComponent::Scale:
+						{
+							FVector Data = Xfo.GetScale3D();
+							TBaseStructure<FVector>::Get()->ImportText(*Content, &Data, nullptr, PPF_None, &ErrorPipe, TBaseStructure<FVector>::Get()->GetName(), true);
+							Xfo.SetScale3D(Data);
+							break;
+						}
+						case ESlateTransformComponent::Max:
+						default:
+						{
+							TBaseStructure<TransformType>::Get()->ImportText(*Content, &Xfo, nullptr, PPF_None, &ErrorPipe, TBaseStructure<TransformType>::Get()->GetName(), true);
+							break;
+						}
+					}
+
+					if(ErrorPipe.NumErrors == 0)
+					{
+						SetSingleTransform(Key, Xfo, bIsRelative);
+					}
+				}
+			}
+		});
+
+		TransformWidgetArgs.DiffersFromDefault_Lambda([
+			CurrentTransformType,
+			Keys,
+			GetSingleTransform
+			
+		](
+			ESlateTransformComponent::Type InComponent) -> bool
+		{
+			for(const FRigElementKey& Key : Keys)
+			{
+				const TransformType CurrentTransform = GetSingleTransform(Key, true);
+				TransformType DefaultTransform;
+
+				switch(CurrentTransformType)
+				{
+					case ERigTransformElementDetailsTransform::Current:
+					{
+						DefaultTransform = GetSingleTransform(Key, true, ERigTransformElementDetailsTransform::Initial);
+						break;
+					}
+					default:
+					{
+						DefaultTransform = TransformType::Identity; 
+						break;
+					}
+				}
+
+				switch(InComponent)
+				{
+					case ESlateTransformComponent::Location:
+					{
+						if(!DefaultTransform.GetLocation().Equals(CurrentTransform.GetLocation()))
+						{
+							return true;
+						}
+						break;
+					}
+					case ESlateTransformComponent::Rotation:
+					{
+						if(!DefaultTransform.Rotator().Equals(CurrentTransform.Rotator()))
+						{
+							return true;
+						}
+						break;
+					}
+					case ESlateTransformComponent::Scale:
+					{
+						if(!DefaultTransform.GetScale3D().Equals(CurrentTransform.GetScale3D()))
+						{
+							return true;
+						}
+						break;
+					}
+					default: // also no component whole transform
+					{
+						if(!DefaultTransform.GetLocation().Equals(CurrentTransform.GetLocation()) ||
+							!DefaultTransform.Rotator().Equals(CurrentTransform.Rotator()) ||
+							!DefaultTransform.GetScale3D().Equals(CurrentTransform.GetScale3D()))
+						{
+							return true;
+						}
+						break;
+					}
+				}
+			}
+			return false;
+		});
+
+		TransformWidgetArgs.OnResetToDefault_Lambda([CurrentTransformType, Keys, GetSingleTransform, SetSingleTransform](
+			ESlateTransformComponent::Type InComponent)
+		{
+			FScopedTransaction Transaction(LOCTEXT("ResetTransformToDefault", "Reset Transform to Default"));
+			
+			for(const FRigElementKey& Key : Keys)
+			{
+				TransformType CurrentTransform = GetSingleTransform(Key, true);
+				TransformType DefaultTransform;
+
+				switch(CurrentTransformType)
+				{
+					case ERigTransformElementDetailsTransform::Current:
+					{
+						DefaultTransform = GetSingleTransform(Key, true, ERigTransformElementDetailsTransform::Initial);
+						break;
+					}
+					default:
+					{
+						DefaultTransform = TransformType::Identity; 
+						break;
+					}
+				}
+
+				switch(InComponent)
+				{
+					case ESlateTransformComponent::Location:
+					{
+						CurrentTransform.SetLocation(DefaultTransform.GetLocation());
+						break;
+					}
+					case ESlateTransformComponent::Rotation:
+					{
+						CurrentTransform.SetRotator(DefaultTransform.Rotator());
+						break;
+					}
+					case ESlateTransformComponent::Scale:
+					{
+						CurrentTransform.SetScale3D(DefaultTransform.GetScale3D());
+						break;
+					}
+					default: // whole transform / max component
+					{
+						CurrentTransform = DefaultTransform;
+						break;
+					}
+				}
+
+				SetSingleTransform(Key, CurrentTransform, true);
+			}
+		});
+
+		SAdvancedTransformInputBox<TransformType>::ConstructGroupedTransformRows(
+			TransformCategory, 
+			ButtonLabels[Index], 
+			ButtonTooltips[Index], 
+			TransformWidgetArgs);
 	}
 }
 
@@ -1028,6 +1593,25 @@ bool FRigTransformElementDetails::IsCurrentLocalEnabled() const
 	return true;
 }
 
+void FRigTransformElementDetails::AddChoiceWidgetRow(IDetailCategoryBuilder& InCategory, const FText& InSearchText, TSharedRef<SWidget> InWidget)
+{
+	InCategory.AddCustomRow(FText::FromString(TEXT("TransformType")))
+	.ValueContent()
+	.MinDesiredWidth(375.f)
+	.MaxDesiredWidth(375.f)
+	.HAlign(HAlign_Left)
+	[
+		SNew(SHorizontalBox)
+		+SHorizontalBox::Slot()
+		.AutoWidth()
+		.HAlign(HAlign_Left)
+		.VAlign(VAlign_Center)
+		[
+			InWidget
+		]
+	];
+}
+
 void FRigControlElementDetails_SetupBoolValueWidget(IDetailCategoryBuilder& InCategory, const TSharedRef<IPropertyUtilities>& InPropertyUtilities, ERigControlValueType InValueType, const FRigElementKey& InKey, URigHierarchy* InHierarchy)
 {
 	UEnum* ControlTypeEnum = StaticEnum<ERigControlType>();
diff --git a/Engine/Plugins/Animation/ControlRig/Source/ControlRigEditor/Private/ControlRigElementDetails.h b/Engine/Plugins/Animation/ControlRig/Source/ControlRigEditor/Private/ControlRigElementDetails.h
--- a/Engine/Plugins/Animation/ControlRig/Source/ControlRigEditor/Private/ControlRigElementDetails.h
+++ b/Engine/Plugins/Animation/ControlRig/Source/ControlRigEditor/Private/ControlRigElementDetails.h
@@ -16,6 +16,12 @@
 #include "IPropertyUtilities.h"
 #include "SSearchableComboBox.h"
 #include "DetailsViewWrapperObject.h"
+#include "Widgets/Input/SSegmentedControl.h"
+#include "DetailCategoryBuilder.h"
+#include "DetailLayoutBuilder.h"
+#include "DetailWidgetRow.h"
+#include "IDetailGroup.h"
+#include "SAdvancedTransformInputBox.h"
 
 class IPropertyHandle;
 
@@ -199,11 +205,8 @@ public:
 protected:
 
 	bool IsCurrentLocalEnabled() const;
-	
-	FEditPropertyChain PoseInitialGlobal;
-	FEditPropertyChain PoseInitialLocal;
-	FEditPropertyChain PoseCurrentGlobal;
-	FEditPropertyChain PoseCurrentLocal;
+
+	void AddChoiceWidgetRow(IDetailCategoryBuilder& InCategory, const FText& InSearchText, TSharedRef<SWidget> InWidget);
 };
 
 class FRigBoneElementDetails : public FRigTransformElementDetails
diff --git a/Engine/Source/Developer/AnimationWidgets/Public/SAdvancedTransformInputBox.h b/Engine/Source/Developer/AnimationWidgets/Public/SAdvancedTransformInputBox.h
--- a/Engine/Source/Developer/AnimationWidgets/Public/SAdvancedTransformInputBox.h
+++ b/Engine/Source/Developer/AnimationWidgets/Public/SAdvancedTransformInputBox.h
@@ -11,6 +11,15 @@
 #include "Framework/Commands/UIAction.h"
 #include "Framework/Commands/UICommandInfo.h"
 
+#if WITH_EDITOR
+
+#include "DetailCategoryBuilder.h"
+#include "DetailLayoutBuilder.h"
+#include "DetailWidgetRow.h"
+#include "IDetailGroup.h"
+
+#endif
+
 namespace ESlateTransformComponent
 {
 	enum Type
@@ -65,6 +74,18 @@ public:
 	/** Notification for toggle checkstate change */
 	DECLARE_DELEGATE_FourParams(FOnToggleChanged, ESlateTransformComponent::Type, ESlateRotationRepresentation::Type, ESlateTransformSubComponent::Type, ECheckBoxState);
 
+	/** Delegate to fire when a row or the whole transform (type == max) is being copied to the clipboard */
+	DECLARE_DELEGATE_OneParam(FOnCopyToClipboard, ESlateTransformComponent::Type);
+
+	/** Delegate to fire when a row or the whole transform (type == max) is being pasted into from the clipboard */
+	DECLARE_DELEGATE_OneParam(FOnPasteFromClipboard, ESlateTransformComponent::Type);
+
+	/** Delegate to fire when a row or the whole transform (type == max) needs to be reset to its default */
+	DECLARE_DELEGATE_OneParam(FOnResetToDefault, ESlateTransformComponent::Type);
+
+	/** Delegate to determine if the row or the whole transform (type == max) differs from its default */
+	DECLARE_DELEGATE_RetVal_OneParam(bool, FDiffersFromDefault, ESlateTransformComponent::Type);
+
 	SLATE_BEGIN_ARGS( SAdvancedTransformInputBox )
 		: _ConstructLocation(true)
 		, _ConstructRotation(true)
@@ -188,6 +209,18 @@ public:
 		/** Padding around the toggle checkbox */
 		SLATE_ARGUMENT( FMargin, TogglePadding )
 
+		/** Delegate to react to copy to clipboard */
+		SLATE_EVENT( FOnCopyToClipboard, OnCopyToClipboard )
+
+		/** Delegate to react to paste from clipboard */
+		SLATE_EVENT( FOnPasteFromClipboard, OnPasteFromClipboard )
+
+		/** Delegate to react to reset to default */
+		SLATE_EVENT( FOnResetToDefault, OnResetToDefault )
+
+		/** Delegate to determine if a row (or the whole transform) differs from its default */
+		SLATE_EVENT( FDiffersFromDefault, DiffersFromDefault )
+
 	SLATE_END_ARGS()
 
 	/**
@@ -1058,6 +1091,113 @@ public:
 		return WorldIcon;
 	}
 
+#if WITH_EDITOR
+	
+	static void ConstructGroupedTransformRows(
+		IDetailCategoryBuilder& InCategory,
+		const FText& InLabel,
+		const FText& InTooltip,
+		typename SAdvancedTransformInputBox<TransformType>::FArguments WidgetArgs
+		)
+	{
+		auto ConstructComponentWidgetRow = [WidgetArgs](FDetailWidgetRow& WidgetRow, ESlateTransformComponent::Type InComponent)
+		{
+			if(WidgetArgs._OnCopyToClipboard.IsBound() && WidgetArgs._OnPasteFromClipboard.IsBound())
+			{
+				WidgetRow
+				.CopyAction(FUIAction(
+					FExecuteAction::CreateLambda([WidgetArgs, InComponent]()
+					{
+						WidgetArgs._OnCopyToClipboard.ExecuteIfBound(InComponent);
+					}),
+					FCanExecuteAction())
+				);
+
+				WidgetRow.PasteAction(FUIAction(
+				FExecuteAction::CreateLambda([WidgetArgs, InComponent]()
+				{
+					WidgetArgs._OnPasteFromClipboard.ExecuteIfBound(InComponent);
+				}),
+				FCanExecuteAction::CreateLambda([WidgetArgs]() -> bool
+				{
+					return WidgetArgs._IsEnabled.Get();
+				})));
+			}
+
+			if(WidgetArgs._OnResetToDefault.IsBound() && WidgetArgs._DiffersFromDefault.IsBound())
+			{
+				WidgetRow.OverrideResetToDefault(FResetToDefaultOverride::Create(
+					TAttribute<bool>::CreateLambda([WidgetArgs, InComponent]() -> bool
+					{
+						if(!WidgetArgs._IsEnabled.Get())
+						{
+							return false;
+						}
+						return WidgetArgs._DiffersFromDefault.Execute(InComponent);
+					}),
+					FSimpleDelegate::CreateLambda([WidgetArgs, InComponent]()
+					{
+						WidgetArgs._OnResetToDefault.Execute(InComponent);
+					})
+				));
+			}
+			
+			if(InComponent != ESlateTransformComponent::Max)
+			{
+				WidgetRow
+				.NameContent()
+				.HAlign(HAlign_Fill)
+				[
+					SAdvancedTransformInputBox<TransformType>::ConstructLabel(WidgetArgs, InComponent)
+				]
+				.ValueContent()
+				.MinDesiredWidth(375.f)
+				.MaxDesiredWidth(375.f)
+				[
+					SNew(SHorizontalBox)
+					+SHorizontalBox::Slot()
+					[
+						SAdvancedTransformInputBox<TransformType>::ConstructWidget(WidgetArgs, InComponent)
+					]
+				];
+			}			
+		};
+		
+		IDetailGroup& Group = InCategory.AddGroup(*InLabel.ToString(), InLabel, false, true);
+		FDetailWidgetRow& HeaderRow = Group.HeaderRow();
+		ConstructComponentWidgetRow(HeaderRow, ESlateTransformComponent::Max);
+
+		HeaderRow
+		.Visibility(WidgetArgs._Visibility)
+		.NameContent()
+		[
+			SNew(STextBlock)
+			.Font(IDetailLayoutBuilder::GetDetailFont())
+			.Text(InLabel)
+			.ToolTipText(InTooltip)
+		];
+
+		if(WidgetArgs._ConstructLocation)
+		{
+			FDetailWidgetRow& WidgetRow = Group.AddWidgetRow();
+			ConstructComponentWidgetRow(WidgetRow, ESlateTransformComponent::Location);
+		}
+
+		if(WidgetArgs._ConstructRotation)
+		{
+			FDetailWidgetRow& WidgetRow = Group.AddWidgetRow();
+			ConstructComponentWidgetRow(WidgetRow, ESlateTransformComponent::Rotation);
+		}
+
+		if(WidgetArgs._ConstructScale)
+		{
+			FDetailWidgetRow& WidgetRow = Group.AddWidgetRow();
+			ConstructComponentWidgetRow(WidgetRow, ESlateTransformComponent::Scale);
+		}
+	}
+
+#endif
+
 	static const TArray<FText>& GetRotationRepresentationLabels()
 	{
 		static TArray<FText> Labels;
diff --git a/Engine/Source/Runtime/AnimationCore/Public/EulerTransform.h b/Engine/Source/Runtime/AnimationCore/Public/EulerTransform.h
--- a/Engine/Source/Runtime/AnimationCore/Public/EulerTransform.h
+++ b/Engine/Source/Runtime/AnimationCore/Public/EulerTransform.h
@@ -101,4 +101,9 @@ struct ANIMATIONCORE_API FEulerTransform
 	FORCEINLINE void SetRotator(const FRotator& InValue) { Rotation = InValue; }
 	FORCEINLINE void SetScale3D(const FVector& InValue) { Scale = InValue; }
 	FORCEINLINE void NormalizeRotation() {}
-};
\ No newline at end of file
+};
+
+template<> struct TBaseStructure<FEulerTransform>
+{
+	ANIMATIONCORE_API static UScriptStruct* Get() { return FEulerTransform::StaticStruct(); }
+};
diff --git a/Engine/Source/Runtime/Slate/Public/Widgets/Input/SSegmentedControl.h b/Engine/Source/Runtime/Slate/Public/Widgets/Input/SSegmentedControl.h
--- a/Engine/Source/Runtime/Slate/Public/Widgets/Input/SSegmentedControl.h
+++ b/Engine/Source/Runtime/Slate/Public/Widgets/Input/SSegmentedControl.h
@@ -400,6 +400,36 @@ public:
 		}
 	}
 
+	static  TSharedPtr<SSegmentedControl<OptionType>> Create(
+		const TArray<OptionType>& InKeys,
+		const TArray<FText>& InLabels,
+		const TArray<FText>& InTooltips,
+		const TAttribute<TArray<OptionType>>& InValues,
+		bool bSupportsMultiSelection = true,
+		FOnValuesChanged OnValuesChanged = FOnValuesChanged())
+	{
+		TSharedPtr<SSegmentedControl<OptionType>> Widget;
+
+		SAssignNew(Widget, SSegmentedControl<OptionType>)
+		.SupportsMultiSelection(bSupportsMultiSelection)
+		.Values(InValues)
+		.OnValuesChanged(OnValuesChanged);
+
+		ensure(InKeys.Num() == InLabels.Num());
+		ensure(InKeys.Num() == InTooltips.Num());
+
+		for(int32 Index = 0; Index < InKeys.Num(); Index++)
+		{
+			Widget->AddSlot(InKeys[Index], false)
+			.Text(InLabels[Index])
+			.ToolTip(InTooltips[Index]);
+		}
+
+		Widget->RebuildChildren();
+
+		return Widget;
+	}
+
 private:
 	
 	TAttribute<ECheckBoxState> GetCheckBoxStateAttribute(OptionType InValue) const
