Sequencer: Trivial fix for masked sections evaluating wrong control.
#trivial

diff --git a/Engine/Plugins/Animation/ControlRig/Source/ControlRig/Private/Sequencer/MovieSceneControlRigParameterTemplate.cpp b/Engine/Plugins/Animation/ControlRig/Source/ControlRig/Private/Sequencer/MovieSceneControlRigParameterTemplate.cpp
--- a/Engine/Plugins/Animation/ControlRig/Source/ControlRig/Private/Sequencer/MovieSceneControlRigParameterTemplate.cpp
+++ b/Engine/Plugins/Animation/ControlRig/Source/ControlRig/Private/Sequencer/MovieSceneControlRigParameterTemplate.cpp
@@ -300,19 +300,19 @@ struct FEvaluatedControlRigParameterSectionChannelMasks : IPersistentEvaluationD
 		{
 			const FScalarParameterNameAndCurve& Scalar = Scalars[Index];
 			ChannelInfo = Section->ControlChannelMap.Find(Scalar.ParameterName);
-			ScalarCurveMask[Index] = (!ChannelInfo || ControlsMask[ChannelInfo->ControlIndex]);
+			ScalarCurveMask[Index] = (!ChannelInfo || ControlsMask[ChannelInfo->MaskIndex]);
 		}
 		for (int32 Index = 0; Index < Bools.Num(); ++Index)
 		{
 			const FBoolParameterNameAndCurve& Bool = Bools[Index];
 			ChannelInfo = Section->ControlChannelMap.Find(Bool.ParameterName);
-			BoolCurveMask[Index] = (!ChannelInfo || ControlsMask[ChannelInfo->ControlIndex]);
+			BoolCurveMask[Index] = (!ChannelInfo || ControlsMask[ChannelInfo->MaskIndex]);
 		}
 		for (int32 Index = 0; Index < Integers.Num(); ++Index)
 		{
 			const FIntegerParameterNameAndCurve& Integer = Integers[Index];
 			ChannelInfo = Section->ControlChannelMap.Find(Integer.ParameterName);
-			IntegerCurveMask[Index] = (!ChannelInfo || ControlsMask[ChannelInfo->ControlIndex]);
+			IntegerCurveMask[Index] = (!ChannelInfo || ControlsMask[ChannelInfo->MaskIndex]);
 		}
 		for (int32 Index = 0; Index < Enums.Num(); ++Index)
 		{
@@ -324,25 +324,25 @@ struct FEvaluatedControlRigParameterSectionChannelMasks : IPersistentEvaluationD
 		{
 			const FVector2DParameterNameAndCurves& Vector2D = Vector2Ds[Index];
 			ChannelInfo = Section->ControlChannelMap.Find(Vector2D.ParameterName);
-			Vector2DCurveMask[Index] = (!ChannelInfo || ControlsMask[ChannelInfo->ControlIndex]);
+			Vector2DCurveMask[Index] = (!ChannelInfo || ControlsMask[ChannelInfo->MaskIndex]);
 		}
 		for (int32 Index = 0; Index < Vectors.Num(); ++Index)
 		{
 			const FVectorParameterNameAndCurves& Vector = Vectors[Index];
 			ChannelInfo = Section->ControlChannelMap.Find(Vector.ParameterName);
-			VectorCurveMask[Index] = (!ChannelInfo || ControlsMask[ChannelInfo->ControlIndex]);
+			VectorCurveMask[Index] = (!ChannelInfo || ControlsMask[ChannelInfo->MaskIndex]);
 		}
 		for (int32 Index = 0; Index < Colors.Num(); ++Index)
 		{
 			const FColorParameterNameAndCurves& Color = Colors[Index];
 			ChannelInfo = Section->ControlChannelMap.Find(Color.ParameterName);
-			ColorCurveMask[Index] = (!ChannelInfo || ControlsMask[ChannelInfo->ControlIndex]);
+			ColorCurveMask[Index] = (!ChannelInfo || ControlsMask[ChannelInfo->MaskIndex]);
 		}
 		for (int32 Index = 0; Index < Transforms.Num(); ++Index)
 		{
 			const FTransformParameterNameAndCurves& Transform = Transforms[Index];
 			ChannelInfo = Section->ControlChannelMap.Find(Transform.ParameterName);
-			TransformCurveMask[Index] = (!ChannelInfo || ControlsMask[ChannelInfo->ControlIndex]);
+			TransformCurveMask[Index] = (!ChannelInfo || ControlsMask[ChannelInfo->MaskIndex]);
 		}
 	}
 };
diff --git a/Engine/Plugins/Animation/ControlRig/Source/ControlRigEditor/Private/Sequencer/ControlRigParameterTrackEditor.cpp b/Engine/Plugins/Animation/ControlRig/Source/ControlRigEditor/Private/Sequencer/ControlRigParameterTrackEditor.cpp
--- a/Engine/Plugins/Animation/ControlRig/Source/ControlRigEditor/Private/Sequencer/ControlRigParameterTrackEditor.cpp
+++ b/Engine/Plugins/Animation/ControlRig/Source/ControlRigEditor/Private/Sequencer/ControlRigParameterTrackEditor.cpp
@@ -79,6 +79,7 @@
 #include "Toolkits/IToolkitHost.h"
 #include "ControlRigEditModeSettings.h"
 #include "ControlRigSpaceChannelEditors.h"
+#include "Misc/ScopedSlowTask.h"
 
 #define LOCTEXT_NAMESPACE "FControlRigParameterTrackEditor"
 
@@ -2348,6 +2349,39 @@ void FControlRigParameterTrackEditor::HandleOnInitialized(UControlRig* ControlRi
 	}
 }
 
+void FControlRigParameterTrackEditor::HandleHierarchyModified(ERigHierarchyNotification InNotification, URigHierarchy* InHierarchy, const FRigBaseElement* InElement)
+{
+	//Only handle renamed events, the other ones like ERigHierarchyNotification::ElementRemoved will get handle when we recompile.
+	switch (InNotification)
+	{
+
+	case ERigHierarchyNotification::ElementRenamed:
+	{
+		if (const FRigControlElement* ControlElement = Cast<FRigControlElement>(InElement))
+		{
+			if (!GetSequencer().IsValid() || !GetSequencer()->IsAllowedToChange())
+			{
+				return;
+			}
+			UMovieScene* MovieScene = GetSequencer()->GetFocusedMovieSceneSequence()->GetMovieScene();
+			const TArray<FMovieSceneBinding>& Bindings = MovieScene->GetBindings();
+			for (const FMovieSceneBinding& Binding : Bindings)
+			{
+				UMovieSceneControlRigParameterTrack* Track = Cast<UMovieSceneControlRigParameterTrack>(MovieScene->FindTrack(UMovieSceneControlRigParameterTrack::StaticClass(), Binding.GetObjectGuid(), NAME_None));
+				if (Track && Track->GetControlRig() && Track->GetControlRig()->GetHierarchy() == InHierarchy)
+				{
+
+				}
+			}
+
+		}
+		break;
+	}
+	default:
+		break;
+	}
+	
+}
 
 void FControlRigParameterTrackEditor::HandleControlModified(UControlRig* ControlRig, FRigControlElement* ControlElement, const FRigControlModifiedContext& Context)
 {
@@ -2779,7 +2813,6 @@ void FControlRigParameterTrackEditor::AddControlKeys(USceneComponent* InSceneCom
 
 }
 
-
 bool FControlRigParameterTrackEditor::ModifyOurGeneratedKeysByCurrentAndWeight(UObject* Object, UControlRig* InControlRig, FName RigControlName, UMovieSceneTrack* Track, UMovieSceneSection* SectionToKey, FFrameNumber KeyTime, FGeneratedTrackKeys& GeneratedTotalKeys, float Weight) const
 {
 	FFrameRate TickResolution = GetSequencer()->GetFocusedTickResolution();
@@ -3385,6 +3418,7 @@ bool FControlRigParameterTrackEditor::CollapseAllLayers(TSharedPtr<ISequencer>&
 			return false;
 		}
 		FScopedTransaction Transaction(LOCTEXT("CollapseAllSections", "CollapseAllSections"));
+		ParameterSection->Modify();
 		UControlRig* ControlRig = ParameterSection->GetControlRig();
 		TRange<FFrameNumber> Range = SequencerPtr->GetFocusedMovieSceneSequence()->GetMovieScene()->GetPlaybackRange();
 		FFrameNumber StartFrame = Range.GetLowerBoundValue();
@@ -3441,13 +3475,17 @@ bool FControlRigParameterTrackEditor::CollapseAllLayers(TSharedPtr<ISequencer>&
 			}
 		}
 
-		ParameterSection->RemoveAllKeys(false); //remove all keys but the space keys.
 		FRigControlModifiedContext Context;
 		Context.SetKey = EControlRigSetKey::Always;
 
+
+		FScopedSlowTask Feedback(Frames.Num(), LOCTEXT("CollapsingSections", "Collapsing Sections"));
+		Feedback.MakeDialog(true);
+
 		Index = 0;
 		for (Index = 0; Index < Frames.Num(); ++Index)
 		{
+			Feedback.EnterProgressFrame(1, LOCTEXT("CollapsingSections", "Collapsing Sections"));
 			const FFrameNumber& FrameNumber = Frames[Index];
 			Context.LocalTime = TickResolution.AsSeconds(FFrameTime(FrameNumber));
 			//need to do the twice hack since controls aren't really in order
@@ -3455,9 +3493,18 @@ bool FControlRigParameterTrackEditor::CollapseAllLayers(TSharedPtr<ISequencer>&
 			{
 				for (TPair<FName, TArray<FTransform>>& TrailControlTransform : ControlLocalTransforms)
 				{
-					ControlRig->SetControlLocalTransform(TrailControlTransform.Key, TrailControlTransform.Value[Index], true, Context, false);
+					ControlRig->SetControlLocalTransform(TrailControlTransform.Key, TrailControlTransform.Value[Index],false, Context, false);
 				}
 			}
+			ControlRig->Evaluate_AnyThread();
+			ParameterSection->RecordControlRigKey(FrameNumber, true, bKeyReduce);
+
+			if (Feedback.ShouldCancel())
+			{
+				Transaction.Cancel();
+				SequencerPtr->NotifyMovieSceneDataChanged(EMovieSceneDataChangeType::MovieSceneStructureItemAdded);
+				return false;
+			}
 		}
 		if (bKeyReduce)
 		{
diff --git a/Engine/Plugins/Animation/ControlRig/Source/ControlRigEditor/Private/Sequencer/ControlRigParameterTrackEditor.h b/Engine/Plugins/Animation/ControlRig/Source/ControlRigEditor/Private/Sequencer/ControlRigParameterTrackEditor.h
--- a/Engine/Plugins/Animation/ControlRig/Source/ControlRigEditor/Private/Sequencer/ControlRigParameterTrackEditor.h
+++ b/Engine/Plugins/Animation/ControlRig/Source/ControlRigEditor/Private/Sequencer/ControlRigParameterTrackEditor.h
@@ -126,6 +126,7 @@ private:
 
 	/** Control Rig Delegates*/
 	void HandleControlModified(UControlRig* Subject, FRigControlElement* ControlElement, const FRigControlModifiedContext& Context);
+	void HandleHierarchyModified(ERigHierarchyNotification InNotification, URigHierarchy* InHierarchy, const FRigBaseElement* InElement);
 	void HandleControlSelected(UControlRig* Subject, FRigControlElement* ControlElement, bool bSelected);
 	void HandleOnInitialized(UControlRig* Subject, const EControlRigState InState, const FName& InEventName);
 
