CIS Auto Merge Fix:  File came over with no merge issue for some reason, but won't compile.

diff --git a/Engine/Plugins/Animation/ControlRig/Source/ControlRigEditor/Private/Sequencer/ControlRigParameterTrackEditor.cpp b/Engine/Plugins/Animation/ControlRig/Source/ControlRigEditor/Private/Sequencer/ControlRigParameterTrackEditor.cpp
--- a/Engine/Plugins/Animation/ControlRig/Source/ControlRigEditor/Private/Sequencer/ControlRigParameterTrackEditor.cpp
+++ b/Engine/Plugins/Animation/ControlRig/Source/ControlRigEditor/Private/Sequencer/ControlRigParameterTrackEditor.cpp
@@ -3482,6 +3482,9 @@ bool FControlRigParameterTrackEditor::CollapseAllLayers(TSharedPtr<ISequencer>&
 		FScopedSlowTask Feedback(Frames.Num(), LOCTEXT("CollapsingSections", "Collapsing Sections"));
 		Feedback.MakeDialog(true);
 
+		const bool bKeyReduce = true;
+		const ERichCurveInterpMode InterpMode = bKeyReduce ? RCIM_Cubic : RCIM_Linear;
+
 		Index = 0;
 		for (Index = 0; Index < Frames.Num(); ++Index)
 		{
@@ -3497,7 +3500,7 @@ bool FControlRigParameterTrackEditor::CollapseAllLayers(TSharedPtr<ISequencer>&
 				}
 			}
 			ControlRig->Evaluate_AnyThread();
-			ParameterSection->RecordControlRigKey(FrameNumber, true, bKeyReduce);
+			ParameterSection->RecordControlRigKey(FrameNumber, true, InterpMode);
 
 			if (Feedback.ShouldCancel())
 			{
