Control Rig: Shorten context menu section title

diff --git a/Engine/Plugins/Animation/ControlRig/Source/ControlRigEditor/Private/Editor/ControlRigContextMenuContext.cpp b/Engine/Plugins/Animation/ControlRig/Source/ControlRigEditor/Private/Editor/ControlRigContextMenuContext.cpp
--- a/Engine/Plugins/Animation/ControlRig/Source/ControlRigEditor/Private/Editor/ControlRigContextMenuContext.cpp
+++ b/Engine/Plugins/Animation/ControlRig/Source/ControlRigEditor/Private/Editor/ControlRigContextMenuContext.cpp
@@ -13,7 +13,12 @@ FString FControlRigRigHierarchyToGraphDragAndDropContext::GetSectionTitle() cons
 	{
 		ElementNameStrings.Add(Element.Name.ToString());
 	}
-	return FString::Join(ElementNameStrings, TEXT(","));
+	FString SectionTitle = FString::Join(ElementNameStrings, TEXT(","));
+	if(SectionTitle.Len() > 64)
+	{
+		SectionTitle = SectionTitle.Left(61) + TEXT("...");
+	}
+	return SectionTitle;
 }
 
 void UControlRigContextMenuContext::Init(TWeakPtr<FControlRigEditor> InControlRigEditor, const FControlRigMenuSpecificContext& InMenuSpecificContext)
