Control Rig: Use Alpha == 1 for icons

diff --git a/Engine/Plugins/Animation/ControlRig/Source/ControlRigDeveloper/Private/ControlRigVisualGraphUtils.cpp b/Engine/Plugins/Animation/ControlRig/Source/ControlRigDeveloper/Private/ControlRigVisualGraphUtils.cpp
--- a/Engine/Plugins/Animation/ControlRig/Source/ControlRigDeveloper/Private/ControlRigVisualGraphUtils.cpp
+++ b/Engine/Plugins/Animation/ControlRig/Source/ControlRigDeveloper/Private/ControlRigVisualGraphUtils.cpp
@@ -170,7 +170,9 @@ FString FControlRigVisualGraphUtils::DumpRigHierarchyToDotGraph(URigHierarchy* I
 				}
 				case ERigElementType::Control:
 				{
-					Color = Cast<FRigControlElement>(InElement)->Settings.ShapeColor;
+					FLinearColor ShapeColor = Cast<FRigControlElement>(InElement)->Settings.ShapeColor; 
+					ShapeColor.A = 1.f;
+					Color = ShapeColor;
 					break;
 				}
 				default:
diff --git a/Engine/Plugins/Animation/ControlRig/Source/ControlRigEditor/Private/Editor/SRigHierarchyTreeView.cpp b/Engine/Plugins/Animation/ControlRig/Source/ControlRigEditor/Private/Editor/SRigHierarchyTreeView.cpp
--- a/Engine/Plugins/Animation/ControlRig/Source/ControlRigEditor/Private/Editor/SRigHierarchyTreeView.cpp
+++ b/Engine/Plugins/Animation/ControlRig/Source/ControlRigEditor/Private/Editor/SRigHierarchyTreeView.cpp
@@ -635,7 +635,10 @@ TPair<const FSlateBrush*, FSlateColor> SRigHierarchyItem::GetBrushForElementType
 			Brush = FControlRigEditorStyle::Get().GetBrush("ControlRig.Tree.Control");
 			if(const FRigControlElement* Control = InHierarchy->Find<FRigControlElement>(InKey))
 			{
-				Color = FSlateColor(Control->Settings.ShapeColor);
+				FLinearColor ShapeColor = Control->Settings.ShapeColor;
+				// ensure the alpha is always visible
+				ShapeColor.A = 1.f;
+				Color = FSlateColor(ShapeColor);
 			}
 			break;
 		}
