Control Rig; Space Switching: Remove visibility for property so it doesn't show up in details panel.
#trivial

diff --git a/Engine/Plugins/Animation/ControlRig/Source/ControlRigEditor/Public/RigSpacePickerBakeSettings.h b/Engine/Plugins/Animation/ControlRig/Source/ControlRigEditor/Public/RigSpacePickerBakeSettings.h
--- a/Engine/Plugins/Animation/ControlRig/Source/ControlRigEditor/Public/RigSpacePickerBakeSettings.h
+++ b/Engine/Plugins/Animation/ControlRig/Source/ControlRigEditor/Public/RigSpacePickerBakeSettings.h
@@ -20,7 +20,7 @@ struct CONTROLRIGEDITOR_API FRigSpacePickerBakeSettings
 		Tolerance = 0.001f;
 	}
 
-	UPROPERTY(VisibleAnywhere,BlueprintReadWrite, Category = "Settings")
+	UPROPERTY(BlueprintReadWrite, Category = "Settings")
 	FRigElementKey TargetSpace;
 
 	UPROPERTY(BlueprintReadWrite,EditAnywhere, Category = "Settings")
