Control Rig: Fix start watch action greyed out after change of variable type

diff --git a/Engine/Plugins/Animation/ControlRig/Source/ControlRigDeveloper/Private/Graph/ControlRigGraphNode.cpp b/Engine/Plugins/Animation/ControlRig/Source/ControlRigDeveloper/Private/Graph/ControlRigGraphNode.cpp
--- a/Engine/Plugins/Animation/ControlRig/Source/ControlRigDeveloper/Private/Graph/ControlRigGraphNode.cpp
+++ b/Engine/Plugins/Animation/ControlRig/Source/ControlRigDeveloper/Private/Graph/ControlRigGraphNode.cpp
@@ -5,6 +5,8 @@
 #include "Graph/ControlRigGraph.h"
 #include "Graph/ControlRigGraphSchema.h"
 #include "Kismet2/BlueprintEditorUtils.h"
+#include "Kismet2/KismetDebugUtilities.h"
+#include "Kismet2/WatchedPin.h"
 #include "KismetCompiler.h"
 #include "BlueprintNodeSpawner.h"
 #include "EditorCategoryUtils.h"
@@ -215,6 +217,26 @@ void UControlRigGraphNode::ReconstructNode_Internal(bool bForce)
 	CachedPins.Reset();
 	CachedModelPins.Reset();
 	ReallocatePinsDuringReconstruction(OldPins);
+
+	// Maintain watches up to date
+	if (URigVMNode* Node = GetModelNode())
+	{
+		UBlueprint* Blueprint = GetBlueprint();
+		for (UEdGraphPin* NewPin : Pins)
+		{
+			const FString PinName = NewPin->GetName();
+			FString Left, Right = PinName;
+			URigVMPin::SplitPinPathAtStart(PinName, Left, Right);
+			if (URigVMPin* ModelPin = Node->FindPin(Right))
+			{
+				if (ModelPin->RequiresWatch())
+				{
+					FKismetDebugUtilities::AddPinWatch(Blueprint, FBlueprintWatchedPin(NewPin));
+				}
+			}
+		}
+	}
+	
 	RewireOldPinsToNewPins(OldPins, Pins);
 
 	// Let subclasses do any additional work
