Control Rig: Improve invalid watch message

diff --git a/Engine/Plugins/Animation/ControlRig/Source/ControlRigEditor/Private/Graph/SControlRigGraphNode.cpp b/Engine/Plugins/Animation/ControlRig/Source/ControlRigEditor/Private/Graph/SControlRigGraphNode.cpp
--- a/Engine/Plugins/Animation/ControlRig/Source/ControlRigEditor/Private/Graph/SControlRigGraphNode.cpp
+++ b/Engine/Plugins/Animation/ControlRig/Source/ControlRigEditor/Private/Graph/SControlRigGraphNode.cpp
@@ -920,6 +920,7 @@ void SControlRigGraphNode::GetNodeInfoPopups(FNodeInfoContext* Context, TArray<F
 
 			FString PinnedWatchText;
 			int32 ValidWatchCount = 0;
+			int32 InvalidWatchCount = 0;
 			for (int32 PinIndex = 0; PinIndex < GraphNode->Pins.Num(); ++PinIndex)
 			{
 				UEdGraphPin* WatchPin = GraphNode->Pins[PinIndex];
@@ -994,6 +995,7 @@ void SControlRigGraphNode::GetNodeInfoPopups(FNodeInfoContext* Context, TArray<F
 						else
 						{
 							PinnedWatchText += FText::Format(LOCTEXT("InvalidPropertyFmt", "No watch found for {0}"), Schema->GetPinDisplayName(WatchPin)).ToString();//@TODO: Print out object being debugged name?
+							InvalidWatchCount++;
 						}
 						ValidWatchCount++;
 					}
@@ -1002,6 +1004,18 @@ void SControlRigGraphNode::GetNodeInfoPopups(FNodeInfoContext* Context, TArray<F
 
 			if (ValidWatchCount)
 			{
+				if (InvalidWatchCount && ModelNode.IsValid())
+				{
+					if(UControlRig* DebuggedControlRig = Cast<UControlRig>(Blueprint->GetObjectBeingDebugged()))
+					{
+						const int32 Count = ModelNode->GetInstructionVisitedCount(DebuggedControlRig->GetVM(), FRigVMASTProxy());
+						if(Count == 0)
+						{
+							PinnedWatchText = FString::Printf(TEXT("Node is not running - wrong event?\n%s"), *PinnedWatchText);
+						}
+					}
+				}
+
 				new (Popups) FGraphInformationPopupInfo(NULL, PinnedWatchColor, PinnedWatchText);
 			}
 		}
