Control Rig: Clear Curve selection on import

diff --git a/Engine/Plugins/Animation/ControlRig/Source/ControlRigEditor/Private/Editor/SRigHierarchy.cpp b/Engine/Plugins/Animation/ControlRig/Source/ControlRigEditor/Private/Editor/SRigHierarchy.cpp
--- a/Engine/Plugins/Animation/ControlRig/Source/ControlRigEditor/Private/Editor/SRigHierarchy.cpp
+++ b/Engine/Plugins/Animation/ControlRig/Source/ControlRigEditor/Private/Editor/SRigHierarchy.cpp
@@ -520,6 +520,15 @@ void SRigHierarchy::OnFilterTextChanged(const FText& SearchText)
 
 void SRigHierarchy::RefreshTreeView(bool bRebuildContent)
 {
+	bool bDummySuspensionFlag = false;
+	bool* SuspensionFlagPtr = &bDummySuspensionFlag;
+	if (ControlRigEditor.IsValid())
+	{
+		SuspensionFlagPtr = &ControlRigEditor.Pin()->bSuspendDetailsPanelRefresh;
+	}
+	TGuardValue<bool> SuspendDetailsPanelRefreshGuard(*SuspensionFlagPtr, true);
+	TGuardValue<bool> GuardRigHierarchyChanges(bIsChangingRigHierarchy, true);
+
 	TreeView->RefreshTreeView(bRebuildContent);
 }
 
@@ -1290,7 +1299,7 @@ void SRigHierarchy::ImportHierarchy(const FAssetData& InAssetData)
 		check(Controller);
 
 		TArray<FRigElementKey> ImportedBones = Controller->ImportBones(Mesh->GetSkeleton(), NAME_None, false, false, bSelectBones, true, true);
-		Controller->ImportCurves(Mesh->GetSkeleton(), NAME_None, true, true);
+		Controller->ImportCurves(Mesh->GetSkeleton(), NAME_None, false, true);
 
 		ControlRigBlueprint->SourceHierarchyImport = Mesh->GetSkeleton();
 		ControlRigBlueprint->SourceCurveImport = Mesh->GetSkeleton();
