Control Rig: Revamp Details Panel

diff --git a/Engine/Plugins/Animation/ControlRig/Source/ControlRig/Private/ControlRig.cpp b/Engine/Plugins/Animation/ControlRig/Source/ControlRig/Private/ControlRig.cpp
--- a/Engine/Plugins/Animation/ControlRig/Source/ControlRig/Private/ControlRig.cpp
+++ b/Engine/Plugins/Animation/ControlRig/Source/ControlRig/Private/ControlRig.cpp
@@ -896,21 +896,13 @@ void UControlRig::Execute(const EControlRigState InState, const FName& InEventNa
 		GetHierarchy()->ForEach<FRigControlElement>([this](FRigControlElement* ControlElement) -> bool
 		{
 			const FRigControlSettings& Settings = ControlElement->Settings;
-			
+
 			if (Settings.bShapeEnabled &&
 				Settings.bShapeVisible &&
 				!Settings.bIsTransientControl &&
 				Settings.bDrawLimits &&
-				(Settings.bLimitTranslation
-					|| Settings.bLimitRotation
-					|| Settings.bLimitScale))
+				Settings.LimitEnabled.Contains(FRigControlLimitEnabled(true, true)))
 			{
-				// for now we don't draw rotational limits
-				if(!Settings.bLimitTranslation)
-				{
-					return true;
-				}
-
 				FTransform Transform = GetHierarchy()->GetGlobalControlOffsetTransformByIndex(ControlElement->GetIndex());
 				FControlRigDrawInstruction Instruction(EControlRigDrawSettings::Lines, Settings.ShapeColor, 0.f, Transform);
 
@@ -918,6 +910,11 @@ void UControlRig::Execute(const EControlRigState InState, const FName& InEventNa
 				{
 					case ERigControlType::Float:
 					{
+						if(Settings.LimitEnabled[0].IsOff())
+						{
+							break;
+						}
+
 						FVector MinPos = FVector::ZeroVector;
 						FVector MaxPos = FVector::ZeroVector;
 
@@ -949,6 +946,11 @@ void UControlRig::Execute(const EControlRigState InState, const FName& InEventNa
 					}
 					case ERigControlType::Integer:
 					{
+						if(Settings.LimitEnabled[0].IsOff())
+						{
+							break;
+						}
+
 						FVector MinPos = FVector::ZeroVector;
 						FVector MaxPos = FVector::ZeroVector;
 
@@ -980,6 +982,15 @@ void UControlRig::Execute(const EControlRigState InState, const FName& InEventNa
 					}
 					case ERigControlType::Vector2D:
 					{
+						if(Settings.LimitEnabled.Num() < 2)
+						{
+							break;
+						}
+						if(Settings.LimitEnabled[0].IsOff() && Settings.LimitEnabled[1].IsOff())
+						{
+							break;
+						}
+
 						Instruction.PrimitiveType = EControlRigDrawSettings::LineStrip;
 						FVector3f MinPos = Settings.MinimumValue.Get<FVector3f>();
 						FVector3f MaxPos = Settings.MaximumValue.Get<FVector3f>();
@@ -1022,7 +1033,20 @@ void UControlRig::Execute(const EControlRigState InState, const FName& InEventNa
 					case ERigControlType::TransformNoScale:
 					case ERigControlType::EulerTransform:
 					{
-						FVector3f MinPos, MaxPos;
+						FVector3f MinPos = FVector3f::ZeroVector;
+						FVector3f MaxPos = FVector3f::ZeroVector;
+
+						// we only check the first three here
+						// since we only consider translation anyway
+						// for scale it's also the first three
+						if(Settings.LimitEnabled.Num() < 3)
+						{
+							break;
+						}
+						if(!Settings.LimitEnabled[0].IsOn() && !Settings.LimitEnabled[1].IsOn() && !Settings.LimitEnabled[2].IsOn())
+						{
+							break;
+						}
 
 						switch (Settings.ControlType)
 						{
@@ -1715,6 +1739,7 @@ FName UControlRig::AddTransientControl(URigVMPin* InPin, FRigElementKey SpaceKey
 	ShapeTransform.SetScale3D(FVector::ZeroVector);
 
 	FRigControlSettings Settings;
+	Settings.ControlType = ERigControlType::Transform;
 	if (URigVMPin* ColorPin = PinForLink->GetNode()->FindPin(TEXT("Color")))
 	{
 		if (ColorPin->GetCPPType() == TEXT("FLinearColor"))
@@ -1842,6 +1867,7 @@ FName UControlRig::AddTransientControl(const FRigElementKey& InElement)
 	ShapeTransform.SetScale3D(FVector::ZeroVector);
 
 	FRigControlSettings Settings;
+	Settings.ControlType = ERigControlType::Transform;
 	Settings.bIsTransientControl = true;
 	Settings.DisplayName = TEXT("Temporary Control");
 
diff --git a/Engine/Plugins/Animation/ControlRig/Source/ControlRig/Private/Rigs/AdditiveControlRig.cpp b/Engine/Plugins/Animation/ControlRig/Source/ControlRig/Private/Rigs/AdditiveControlRig.cpp
--- a/Engine/Plugins/Animation/ControlRig/Source/ControlRig/Private/Rigs/AdditiveControlRig.cpp
+++ b/Engine/Plugins/Animation/ControlRig/Source/ControlRig/Private/Rigs/AdditiveControlRig.cpp
@@ -135,6 +135,7 @@ void UAdditiveControlRig::CreateRigElements(const FReferenceSkeleton& InReferenc
 
             FRigControlSettings Settings;
             Settings.DisplayName = BoneName;
+			Settings.ControlType = ERigControlType::Transform;
             Controller->AddControl(ControlName, NullKey, Settings, FRigControlValue::Make(FTransform::Identity), FTransform::Identity, FTransform::Identity, false);
 
             return true;
diff --git a/Engine/Plugins/Animation/ControlRig/Source/ControlRig/Private/Rigs/RigControlHierarchy.cpp b/Engine/Plugins/Animation/ControlRig/Source/ControlRig/Private/Rigs/RigControlHierarchy.cpp
--- a/Engine/Plugins/Animation/ControlRig/Source/ControlRig/Private/Rigs/RigControlHierarchy.cpp
+++ b/Engine/Plugins/Animation/ControlRig/Source/ControlRig/Private/Rigs/RigControlHierarchy.cpp
@@ -1,6 +1,7 @@
 // Copyright Epic Games, Inc. All Rights Reserved.
 
 #include "Rigs/RigControlHierarchy.h"
+#include "Rigs/RigHierarchyElements.h"
 
 ////////////////////////////////////////////////////////////////////////////////
 // FRigControl
@@ -8,7 +9,10 @@
 
 void FRigControl::ApplyLimits(FRigControlValue& InOutValue) const
 {
-	InOutValue.ApplyLimits(bLimitTranslation, bLimitRotation, bLimitScale, ControlType, MinimumValue, MaximumValue);
+	FRigControlSettings Settings;
+	Settings.ControlType = ControlType;
+	Settings.SetupLimitArrayForType(bLimitTranslation, bLimitRotation, bLimitScale);
+	InOutValue.ApplyLimits(Settings.LimitEnabled, ControlType, MinimumValue, MaximumValue);
 }
 
 FTransform FRigControl::GetTransformFromValue(ERigControlValueType InValueType) const
diff --git a/Engine/Plugins/Animation/ControlRig/Source/ControlRig/Private/Rigs/RigHierarchy.cpp b/Engine/Plugins/Animation/ControlRig/Source/ControlRig/Private/Rigs/RigHierarchy.cpp
--- a/Engine/Plugins/Animation/ControlRig/Source/ControlRig/Private/Rigs/RigHierarchy.cpp
+++ b/Engine/Plugins/Animation/ControlRig/Source/ControlRig/Private/Rigs/RigHierarchy.cpp
@@ -156,6 +156,9 @@ void URigHierarchy::Save(FArchive& Ar)
 		{
 			return;
 		}
+
+		TArray<FRigElementKey> SelectedKeys = GetSelectedKeys();
+		Ar << SelectedKeys;
 	}
 
 	// make sure all parts of pose are valid.
@@ -188,6 +191,7 @@ void URigHierarchy::Save(FArchive& Ar)
 
 void URigHierarchy::Load(FArchive& Ar)
 {
+	TArray<FRigElementKey> SelectedKeys;
 	if(Ar.IsTransacting())
 	{
 		bool bOnlySerializedTransformStackIndex = false;
@@ -198,6 +202,8 @@ void URigHierarchy::Load(FArchive& Ar)
 		{
 			return;
 		}
+
+		Ar << SelectedKeys;
 	}
 
 	Reset();
@@ -251,6 +257,18 @@ void URigHierarchy::Load(FArchive& Ar)
 	}
 
 	UpdateAllCachedChildren();
+
+	if(Ar.IsTransacting())
+	{
+		for(const FRigElementKey& SelectedKey : SelectedKeys)
+		{
+			if(FRigBaseElement* Element = Find<FRigBaseElement>(SelectedKey))
+			{
+				Element->bSelected = true;
+			}
+		}
+	}
+
 	Notify(ERigHierarchyNotification::HierarchyReset, nullptr);
 }
 
@@ -4823,42 +4841,50 @@ TArray<FString> URigHierarchy::ControlSettingsToPythonCommands(const FRigControl
 		case ERigControlType::TransformNoScale: TypeStr = TEXT("EULER_TRANSFORM"); break;
 		default: ensure(false);
 	}
+
+	static const TCHAR* TrueText = TEXT("True");
+	static const TCHAR* FalseText = TEXT("False");
+
+	TArray<FString> LimitEnabledParts;
+	for(const FRigControlLimitEnabled& LimitEnabled : Settings.LimitEnabled)
+	{
+		LimitEnabledParts.Add(FString::Printf(TEXT("unreal.RigControlLimitEnabled(%s, %s)"),
+						   LimitEnabled.bMinimum ? TrueText : FalseText,
+						   LimitEnabled.bMaximum ? TrueText : FalseText));
+	}
+	
+	const FString LimitEnabledStr = FString::Join(LimitEnabledParts, TEXT(", "));
+	
 	Commands.Add(FString::Printf(TEXT("%s.control_type = unreal.RigControlType.%s"),
 									*NameSettings,
 									*TypeStr));
 	Commands.Add(FString::Printf(TEXT("%s.animatable = %s"),
 		*NameSettings,
-		Settings.bAnimatable ? TEXT("True") : TEXT("False")));
+		Settings.bAnimatable ? TrueText : FalseText));
 	Commands.Add(FString::Printf(TEXT("%s.display_name = '%s'"),
 		*NameSettings,
 		*Settings.DisplayName.ToString()));
 	Commands.Add(FString::Printf(TEXT("%s.draw_limits = %s"),
 		*NameSettings,
-		Settings.bDrawLimits ? TEXT("True") : TEXT("False")));
+		Settings.bDrawLimits ? TrueText : FalseText));
 	Commands.Add(FString::Printf(TEXT("%s.shape_color = %s"),
 		*NameSettings,
 		*RigVMPythonUtils::LinearColorToPythonString(Settings.ShapeColor)));
 	Commands.Add(FString::Printf(TEXT("%s.shape_enabled = %s"),
 		*NameSettings,
-		Settings.bShapeEnabled ? TEXT("True") : TEXT("False")));
+		Settings.bShapeEnabled ? TrueText : FalseText));
 	Commands.Add(FString::Printf(TEXT("%s.shape_name = '%s'"),
 		*NameSettings,
 		*Settings.ShapeName.ToString()));
 	Commands.Add(FString::Printf(TEXT("%s.shape_visible = %s"),
 		*NameSettings,
-		Settings.bShapeVisible ? TEXT("True") : TEXT("False")));
+		Settings.bShapeVisible ? TrueText : FalseText));
 	Commands.Add(FString::Printf(TEXT("%s.is_transient_control = %s"),
 		*NameSettings,
-		Settings.bIsTransientControl ? TEXT("True") : TEXT("False")));
-	Commands.Add(FString::Printf(TEXT("%s.limit_rotation = %s"),
-		*NameSettings,
-		Settings.bLimitRotation ? TEXT("True") : TEXT("False")));
-	Commands.Add(FString::Printf(TEXT("%s.limit_translation = %s"),
-		*NameSettings,
-		Settings.bLimitTranslation ? TEXT("True") : TEXT("False")));
-	Commands.Add(FString::Printf(TEXT("%s.limit_scale = %s"),
+		Settings.bIsTransientControl ? TrueText : FalseText));
+	Commands.Add(FString::Printf(TEXT("%s.limit_enabled = [%s]"),
 		*NameSettings,
-		Settings.bLimitScale ? TEXT("True") : TEXT("False")));
+		*LimitEnabledStr));
 	Commands.Add(FString::Printf(TEXT("%s.primary_axis = unreal.RigControlAxis.%s"),
 		*NameSettings,
 		Settings.PrimaryAxis == ERigControlAxis::X ? TEXT("X") : Settings.PrimaryAxis == ERigControlAxis::Y ? TEXT("Y") : TEXT("Z")));
diff --git a/Engine/Plugins/Animation/ControlRig/Source/ControlRig/Private/Rigs/RigHierarchyController.cpp b/Engine/Plugins/Animation/ControlRig/Source/ControlRig/Private/Rigs/RigHierarchyController.cpp
--- a/Engine/Plugins/Animation/ControlRig/Source/ControlRig/Private/Rigs/RigHierarchyController.cpp
+++ b/Engine/Plugins/Animation/ControlRig/Source/ControlRig/Private/Rigs/RigHierarchyController.cpp
@@ -313,6 +313,10 @@ FRigElementKey URigHierarchyController::AddControl(
 		NewElement->Key.Type = ERigElementType::Control;
 		NewElement->Key.Name = Hierarchy->GetSafeNewName(InName.ToString(), NewElement->Key.Type);
 		NewElement->Settings = InSettings;
+		if(NewElement->Settings.LimitEnabled.IsEmpty())
+		{
+			NewElement->Settings.SetupLimitArrayForType();
+		}
 		AddElement(NewElement, Hierarchy->Get(Hierarchy->GetIndex(InParent)), false);
 		
 		NewElement->Offset.Set(ERigTransformType::InitialLocal, InOffsetTransform);  
@@ -529,6 +533,10 @@ bool URigHierarchyController::SetControlSettings(FRigElementKey InKey, FRigContr
 #endif
 
 	ControlElement->Settings = InSettings;
+	if(ControlElement->Settings.LimitEnabled.IsEmpty())
+	{
+		ControlElement->Settings.SetupLimitArrayForType(false, false, false);
+	}
 
 	FRigControlValue InitialValue = Hierarchy->GetControlValue(ControlElement, ERigControlValueType::Initial);
 	FRigControlValue CurrentValue = Hierarchy->GetControlValue(ControlElement, ERigControlValueType::Current);
@@ -1237,9 +1245,7 @@ TArray<FRigElementKey> URigHierarchyController::ImportFromHierarchyContainer(con
 		Settings.PrimaryAxis = Control.PrimaryAxis;
 		Settings.bIsCurve = Control.bIsCurve;
 		Settings.bAnimatable = Control.bAnimatable;
-		Settings.bLimitTranslation = Control.bLimitTranslation;
-		Settings.bLimitRotation = Control.bLimitRotation;
-		Settings.bLimitScale = Control.bLimitScale;
+		Settings.SetupLimitArrayForType(Control.bLimitTranslation, Control.bLimitRotation, Control.bLimitScale);
 		Settings.bDrawLimits = Control.bDrawLimits;
 		Settings.MinimumValue = Control.MinimumValue;
 		Settings.MaximumValue = Control.MaximumValue;
@@ -2427,9 +2433,7 @@ TArray<FRigElementKey> URigHierarchyController::MirrorElements(TArray<FRigElemen
 		{
 			if(FRigControlElement* DuplicatedControlElement = Hierarchy->Find<FRigControlElement>(DuplicatedKeys[Index]))
 			{
-				TGuardValue<bool> DisableTranslationLimit(DuplicatedControlElement->Settings.bLimitTranslation, false);
-				TGuardValue<bool> DisableRotationLimit(DuplicatedControlElement->Settings.bLimitRotation, false);
-				TGuardValue<bool> DisableScaleLimit(DuplicatedControlElement->Settings.bLimitScale, false);
+				TGuardValue<TArray<FRigControlLimitEnabled>> DisableLimits(DuplicatedControlElement->Settings.LimitEnabled, TArray<FRigControlLimitEnabled>());
 
 				// mirror offset
 				FTransform OriginalGlobalOffsetTransform = Hierarchy->GetGlobalControlOffsetTransform(OriginalKeys[Index]);
diff --git a/Engine/Plugins/Animation/ControlRig/Source/ControlRig/Private/Rigs/RigHierarchyDefines.cpp b/Engine/Plugins/Animation/ControlRig/Source/ControlRig/Private/Rigs/RigHierarchyDefines.cpp
--- a/Engine/Plugins/Animation/ControlRig/Source/ControlRig/Private/Rigs/RigHierarchyDefines.cpp
+++ b/Engine/Plugins/Animation/ControlRig/Source/ControlRig/Private/Rigs/RigHierarchyDefines.cpp
@@ -3,6 +3,37 @@
 #include "Rigs/RigHierarchyDefines.h"
 #include "Rigs/RigHierarchy.h"
 
+////////////////////////////////////////////////////////////////////////////////
+// FRigControlLimitEnabled
+////////////////////////////////////////////////////////////////////////////////
+
+void FRigControlLimitEnabled::Serialize(FArchive& Ar)
+{
+	Ar << bMinimum;
+	Ar << bMaximum;
+}
+
+bool FRigControlLimitEnabled::GetForValueType(ERigControlValueType InValueType) const
+{
+	if(InValueType == ERigControlValueType::Minimum)
+	{
+		return bMinimum;
+	}
+	return bMaximum;
+}
+
+void FRigControlLimitEnabled::SetForValueType(ERigControlValueType InValueType, bool InValue)
+{
+	if(InValueType == ERigControlValueType::Minimum)
+	{
+		bMinimum = InValue;
+	}
+	else
+	{
+		bMaximum = InValue;
+	}
+}
+
 ////////////////////////////////////////////////////////////////////////////////
 // FRigElementKey
 ////////////////////////////////////////////////////////////////////////////////
diff --git a/Engine/Plugins/Animation/ControlRig/Source/ControlRig/Private/Rigs/RigHierarchyElements.cpp b/Engine/Plugins/Animation/ControlRig/Source/ControlRig/Private/Rigs/RigHierarchyElements.cpp
--- a/Engine/Plugins/Animation/ControlRig/Source/ControlRig/Private/Rigs/RigHierarchyElements.cpp
+++ b/Engine/Plugins/Animation/ControlRig/Source/ControlRig/Private/Rigs/RigHierarchyElements.cpp
@@ -408,14 +408,12 @@ void FRigBoneElement::CopyFrom(URigHierarchy* InHierarchy, FRigBaseElement* InOt
 ////////////////////////////////////////////////////////////////////////////////
 
 FRigControlSettings::FRigControlSettings()
-: ControlType(ERigControlType::Transform)
+: ControlType(ERigControlType::EulerTransform)
 , DisplayName(NAME_None)
 , PrimaryAxis(ERigControlAxis::X)
 , bIsCurve(false)
 , bAnimatable(true)
-, bLimitTranslation(false)
-, bLimitRotation(false)
-, bLimitScale(false)
+, LimitEnabled()
 , bDrawLimits(true)
 , MinimumValue()
 , MaximumValue()
@@ -452,9 +450,7 @@ void FRigControlSettings::Save(FArchive& Ar)
 	Ar << PrimaryAxisName;
 	Ar << bIsCurve;
 	Ar << bAnimatable;
-	Ar << bLimitTranslation;
-	Ar << bLimitRotation;
-	Ar << bLimitScale;
+	Ar << LimitEnabled;
 	Ar << bDrawLimits;
 	Ar << MinimumValue;
 	Ar << MaximumValue;
@@ -477,14 +473,25 @@ void FRigControlSettings::Load(FArchive& Ar)
 	FName ControlTypeName, PrimaryAxisName;
 	FString ControlEnumPathName;
 
+	bool bLimitTranslation_DEPRECATED = false;
+	bool bLimitRotation_DEPRECATED = false;
+	bool bLimitScale_DEPRECATED = false;
+
 	Ar << ControlTypeName;
 	Ar << DisplayName;
 	Ar << PrimaryAxisName;
 	Ar << bIsCurve;
 	Ar << bAnimatable;
-	Ar << bLimitTranslation;
-	Ar << bLimitRotation;
-	Ar << bLimitScale;
+	if (Ar.CustomVer(FControlRigObjectVersion::GUID) < FControlRigObjectVersion::PerChannelLimits)
+	{
+		Ar << bLimitTranslation_DEPRECATED;
+		Ar << bLimitRotation_DEPRECATED;
+		Ar << bLimitScale_DEPRECATED;
+	}
+	else
+	{
+		Ar << LimitEnabled;
+	}
 	Ar << bDrawLimits;
 
 	FTransform MinimumTransform, MaximumTransform;
@@ -529,6 +536,11 @@ void FRigControlSettings::Load(FArchive& Ar)
 	{
 		Customization.AvailableSpaces.Reset();
 	}
+
+	if (Ar.CustomVer(FControlRigObjectVersion::GUID) < FControlRigObjectVersion::PerChannelLimits)
+	{
+		SetupLimitArrayForType(bLimitTranslation_DEPRECATED, bLimitRotation_DEPRECATED, bLimitScale_DEPRECATED);
+	}
 }
 
 ////////////////////////////////////////////////////////////////////////////////
@@ -557,15 +569,7 @@ bool FRigControlSettings::operator==(const FRigControlSettings& InOther) const
 	{
 		return false;
 	}
-	if(bLimitTranslation != InOther.bLimitTranslation)
-	{
-		return false;
-	}
-	if(bLimitRotation != InOther.bLimitRotation)
-	{
-		return false;
-	}
-	if(bLimitScale != InOther.bLimitScale)
+	if(LimitEnabled != InOther.LimitEnabled)
 	{
 		return false;
 	}
@@ -619,6 +623,66 @@ bool FRigControlSettings::operator==(const FRigControlSettings& InOther) const
 	return true;
 }
 
+void FRigControlSettings::SetupLimitArrayForType(bool bLimitTranslation, bool bLimitRotation, bool bLimitScale)
+{
+	switch(ControlType)
+	{
+		case ERigControlType::Integer:
+		case ERigControlType::Float:
+		{
+			LimitEnabled.SetNum(1);
+			LimitEnabled[0].Set(bLimitTranslation);
+			break;
+		}
+		case ERigControlType::Vector2D:
+		{
+			LimitEnabled.SetNum(2);
+			LimitEnabled[0] = LimitEnabled[1].Set(bLimitTranslation);
+			break;
+		}
+		case ERigControlType::Position:
+		{
+			LimitEnabled.SetNum(3);
+			LimitEnabled[0] = LimitEnabled[1] = LimitEnabled[2].Set(bLimitTranslation);
+			break;
+		}
+		case ERigControlType::Scale:
+		{
+			LimitEnabled.SetNum(3);
+			LimitEnabled[0] = LimitEnabled[1] = LimitEnabled[2].Set(bLimitScale);
+			break;
+		}
+		case ERigControlType::Rotator:
+		{
+			LimitEnabled.SetNum(3);
+			LimitEnabled[0] = LimitEnabled[1] = LimitEnabled[2].Set(bLimitRotation);
+			break;
+		}
+		case ERigControlType::TransformNoScale:
+		{
+			LimitEnabled.SetNum(6);
+			LimitEnabled[0] = LimitEnabled[1] = LimitEnabled[2].Set(bLimitTranslation);
+			LimitEnabled[3] = LimitEnabled[4] = LimitEnabled[5].Set(bLimitRotation);
+			break;
+		}
+		case ERigControlType::EulerTransform:
+		case ERigControlType::Transform:
+		{
+			LimitEnabled.SetNum(9);
+			LimitEnabled[0] = LimitEnabled[1] = LimitEnabled[2].Set(bLimitTranslation);
+			LimitEnabled[3] = LimitEnabled[4] = LimitEnabled[5].Set(bLimitRotation);
+			LimitEnabled[6] = LimitEnabled[7] = LimitEnabled[8].Set(bLimitScale);
+			break;
+		}
+		case ERigControlType::Bool:
+		default:
+		{
+			LimitEnabled.Reset();
+			break;
+		}
+	}
+}
+
 void FRigControlElement::Save(FArchive& Ar, URigHierarchy* Hierarchy, ESerializationPhase SerializationPhase)
 {
 	Super::Save(Ar, Hierarchy, SerializationPhase);
diff --git a/Engine/Plugins/Animation/ControlRig/Source/ControlRig/Private/Units/Hierarchy/RigUnit_GetControlTransform.cpp b/Engine/Plugins/Animation/ControlRig/Source/ControlRig/Private/Units/Hierarchy/RigUnit_GetControlTransform.cpp
--- a/Engine/Plugins/Animation/ControlRig/Source/ControlRig/Private/Units/Hierarchy/RigUnit_GetControlTransform.cpp
+++ b/Engine/Plugins/Animation/ControlRig/Source/ControlRig/Private/Units/Hierarchy/RigUnit_GetControlTransform.cpp
@@ -309,10 +309,13 @@ FRigUnit_GetControlTransform_Execute()
 
 IMPLEMENT_RIGUNIT_AUTOMATION_TEST(FRigUnit_GetControlTransform)
 {
+	FRigControlSettings Settings;
+	Settings.ControlType = ERigControlType::Transform;
+
 	const FRigElementKey Root = Controller->AddControl(
 		TEXT("Root"),
 		FRigElementKey(),
-		FRigControlSettings(),
+		Settings,
 		FRigControlValue::Make(FTransform(FVector(1.f, 0.f, 0.f))),
 		FTransform::Identity,
 		FTransform::Identity);
@@ -320,7 +323,7 @@ IMPLEMENT_RIGUNIT_AUTOMATION_TEST(FRigUnit_GetControlTransform)
 	const FRigElementKey ControlA = Controller->AddControl(
 	    TEXT("ControlA"),
 	    Root,
-	    FRigControlSettings(),
+	    Settings,
 	    FRigControlValue::Make(FTransform(FVector(1.f, 2.f, 3.f))),
 	    FTransform::Identity,
 	    FTransform::Identity);
diff --git a/Engine/Plugins/Animation/ControlRig/Source/ControlRig/Public/ControlRigObjectVersion.h b/Engine/Plugins/Animation/ControlRig/Source/ControlRig/Public/ControlRigObjectVersion.h
--- a/Engine/Plugins/Animation/ControlRig/Source/ControlRig/Public/ControlRigObjectVersion.h
+++ b/Engine/Plugins/Animation/ControlRig/Source/ControlRig/Public/ControlRigObjectVersion.h
@@ -48,6 +48,9 @@ struct CONTROLRIG_API FControlRigObjectVersion
 		// BoundVariableWithInjectionNode 
 		BoundVariableWithInjectionNode,
 
+		// Switch limit control over to per channel limits 
+		PerChannelLimits,
+
 		// -----<new versions can be added above this line>-------------------------------------------------
 		VersionPlusOne,
 		LatestVersion = VersionPlusOne - 1
diff --git a/Engine/Plugins/Animation/ControlRig/Source/ControlRig/Public/Rigs/RigHierarchyDefines.h b/Engine/Plugins/Animation/ControlRig/Source/ControlRig/Public/Rigs/RigHierarchyDefines.h
--- a/Engine/Plugins/Animation/ControlRig/Source/ControlRig/Public/Rigs/RigHierarchyDefines.h
+++ b/Engine/Plugins/Animation/ControlRig/Source/ControlRig/Public/Rigs/RigHierarchyDefines.h
@@ -109,6 +109,92 @@ enum class ERigControlAxis : uint8
     Z
 };
 
+USTRUCT(BlueprintType)
+struct CONTROLRIG_API FRigControlLimitEnabled
+{
+	GENERATED_BODY()
+	
+	FRigControlLimitEnabled()
+		: bMinimum(false)
+		, bMaximum(false)
+	{}
+
+	FRigControlLimitEnabled(bool InValue)
+		: bMinimum(false)
+		, bMaximum(false)
+	{
+		Set(InValue);
+	}
+
+	FRigControlLimitEnabled(bool InMinimum, bool InMaximum)
+		: bMinimum(false)
+		, bMaximum(false)
+	{
+		Set(InMinimum, InMaximum);
+	}
+
+	FRigControlLimitEnabled& Set(bool InValue)
+	{
+		return Set(InValue, InValue);
+	}
+
+	FRigControlLimitEnabled& Set(bool InMinimum, bool InMaximum)
+	{
+		bMinimum = InMinimum;
+		bMaximum = InMaximum;
+		return *this;
+	}
+
+	void Serialize(FArchive& Ar);
+	FORCEINLINE friend FArchive& operator<<(FArchive& Ar, FRigControlLimitEnabled& P)
+	{
+		P.Serialize(Ar);
+		return Ar;
+	}
+
+	FORCEINLINE bool operator ==(const FRigControlLimitEnabled& Other) const
+	{
+		return bMinimum == Other.bMinimum && bMaximum == Other.bMaximum;
+	}
+
+	FORCEINLINE bool operator !=(const FRigControlLimitEnabled& Other) const
+	{
+		return bMinimum != Other.bMinimum || bMaximum != Other.bMaximum;
+	}
+
+	bool IsOn() const { return bMinimum || bMaximum; }
+	bool IsOff() const { return !bMinimum && !bMaximum; }
+	bool GetForValueType(ERigControlValueType InValueType) const;
+	void SetForValueType(ERigControlValueType InValueType, bool InValue);
+
+	template<typename T>
+	T Apply(const T& InValue, const T& InMinimum, const T& InMaximum) const
+	{
+		if(IsOff())
+		{
+			return InValue;
+		}
+
+		if(bMinimum && bMaximum)
+		{
+			return FMath::Clamp<T>(InValue, InMinimum, InMaximum);
+		}
+
+		if(bMinimum)
+		{
+			return FMath::Max<T>(InValue, InMinimum);
+		}
+
+		return FMath::Min<T>(InValue, InMaximum);
+	}
+
+	UPROPERTY(BlueprintReadWrite, EditAnywhere, Category = Limit)
+	bool bMinimum;
+
+	UPROPERTY(BlueprintReadWrite, EditAnywhere, Category = Limit)
+	bool bMaximum;
+};
+
 USTRUCT(BlueprintType)
 struct CONTROLRIG_API FRigControlValueStorage
 {
@@ -605,14 +691,12 @@ public:
 	}
 
 	FORCEINLINE void ApplyLimits(
-		bool bLimitTranslation,
-		bool bLimitRotation,
-		bool bLimitScale,
+		const TArray<FRigControlLimitEnabled>& LimitEnabled,
 		ERigControlType InControlType,
 		const FRigControlValue& InMinimumValue,
 		const FRigControlValue& InMaximumValue)
 	{
-		if (!bLimitTranslation && !bLimitRotation && !bLimitScale)
+		if(LimitEnabled.IsEmpty())
 		{
 			return;
 		}
@@ -621,7 +705,7 @@ public:
 		{
 			case ERigControlType::Float:
 			{
-				if (bLimitTranslation)
+				if (LimitEnabled[0].IsOn())
 				{
 					float& ValueRef = GetRef<float>();
 					ValueRef = Clamp<float>(ValueRef, InMinimumValue.Get<float>(), InMaximumValue.Get<float>());
@@ -630,7 +714,7 @@ public:
 			}
 			case ERigControlType::Integer:
 			{
-				if (bLimitTranslation)
+				if (LimitEnabled[0].IsOn())
 				{
 					int32& ValueRef = GetRef<int32>();
 					ValueRef = Clamp<int32>(ValueRef, InMinimumValue.Get<int32>(), InMaximumValue.Get<int32>());
@@ -639,115 +723,143 @@ public:
 			}
 			case ERigControlType::Vector2D:
 			{
-				if (bLimitTranslation)
+				if(LimitEnabled.Num() < 2)
+				{
+					return;
+				}
+				if (LimitEnabled[0].IsOn() || LimitEnabled[1].IsOn())
 				{
 					FVector3f& ValueRef = GetRef<FVector3f>();
 					const FVector3f& Min = InMinimumValue.GetRef<FVector3f>();
 					const FVector3f& Max = InMaximumValue.GetRef<FVector3f>();
-					ValueRef.X = Clamp<float>(ValueRef.X, Min.X, Max.X);
-					ValueRef.Y = Clamp<float>(ValueRef.Y, Min.Y, Max.Y);
+					ValueRef.X = LimitEnabled[0].Apply<float>(ValueRef.X, Min.X, Max.X);
+					ValueRef.Y = LimitEnabled[1].Apply<float>(ValueRef.Y, Min.Y, Max.Y);
 				}
 				break;
 			}
 			case ERigControlType::Position:
 			{
-				if (bLimitTranslation)
+				if(LimitEnabled.Num() < 3)
+				{
+					return;
+				}
+				if (LimitEnabled[0].IsOn() || LimitEnabled[2].IsOn() || LimitEnabled[2].IsOn())
 				{
 					FVector3f& ValueRef = GetRef<FVector3f>();
 					const FVector3f& Min = InMinimumValue.GetRef<FVector3f>();
 					const FVector3f& Max = InMaximumValue.GetRef<FVector3f>();
-					ValueRef.X = Clamp<float>(ValueRef.X, Min.X, Max.X);
-					ValueRef.Y = Clamp<float>(ValueRef.Y, Min.Y, Max.Y);
-					ValueRef.Z = Clamp<float>(ValueRef.Z, Min.Z, Max.Z);
+					ValueRef.X = LimitEnabled[0].Apply<float>(ValueRef.X, Min.X, Max.X);
+					ValueRef.Y = LimitEnabled[1].Apply<float>(ValueRef.Y, Min.Y, Max.Y);
+					ValueRef.Z = LimitEnabled[2].Apply<float>(ValueRef.Z, Min.Z, Max.Z);
 				}
 				break;
 			}
 			case ERigControlType::Scale:
 			{
-				if (bLimitScale)
+				if(LimitEnabled.Num() < 3)
+				{
+					return;
+				}
+				if (LimitEnabled[0].IsOn() || LimitEnabled[2].IsOn() || LimitEnabled[2].IsOn())
 				{
 					FVector3f& ValueRef = GetRef<FVector3f>();
 					const FVector3f& Min = InMinimumValue.GetRef<FVector3f>();
 					const FVector3f& Max = InMaximumValue.GetRef<FVector3f>();
-					ValueRef.X = Clamp<float>(ValueRef.X, Min.X, Max.X);
-					ValueRef.Y = Clamp<float>(ValueRef.Y, Min.Y, Max.Y);
-					ValueRef.Z = Clamp<float>(ValueRef.Z, Min.Z, Max.Z);
+					ValueRef.X = LimitEnabled[0].Apply<float>(ValueRef.X, Min.X, Max.X);
+					ValueRef.Y = LimitEnabled[1].Apply<float>(ValueRef.Y, Min.Y, Max.Y);
+					ValueRef.Z = LimitEnabled[2].Apply<float>(ValueRef.Z, Min.Z, Max.Z);
 				}
 				break;
 			}
 			case ERigControlType::Rotator:
 			{
-				if (bLimitRotation)
+				if(LimitEnabled.Num() < 3)
+				{
+					return;
+				}
+				if (LimitEnabled[0].IsOn() || LimitEnabled[2].IsOn() || LimitEnabled[2].IsOn())
 				{
 					FVector3f& ValueRef = GetRef<FVector3f>();
 					const FVector3f& Min = InMinimumValue.GetRef<FVector3f>();
 					const FVector3f& Max = InMaximumValue.GetRef<FVector3f>();
-					ValueRef.X = Clamp<float>(ValueRef.X, Min.X, Max.X);
-					ValueRef.Y = Clamp<float>(ValueRef.Y, Min.Y, Max.Y);
-					ValueRef.Z = Clamp<float>(ValueRef.Z, Min.Z, Max.Z);
+					ValueRef.X = LimitEnabled[0].Apply<float>(ValueRef.X, Min.X, Max.X);
+					ValueRef.Y = LimitEnabled[1].Apply<float>(ValueRef.Y, Min.Y, Max.Y);
+					ValueRef.Z = LimitEnabled[2].Apply<float>(ValueRef.Z, Min.Z, Max.Z);
 				}
 				break;
 			}
 			case ERigControlType::Transform:
 			{
+				if(LimitEnabled.Num() < 9)
+				{
+					return;
+				}
+
 				FTransform_Float& ValueRef = GetRef<FTransform_Float>();
 				const FTransform Min = InMinimumValue.GetRef<FTransform_Float>().ToTransform();
 				const FTransform Max = InMaximumValue.GetRef<FTransform_Float>().ToTransform();
 
-				if (bLimitTranslation)
+				if (LimitEnabled[0].IsOn() || LimitEnabled[2].IsOn() || LimitEnabled[2].IsOn())
 				{
-					ValueRef.TranslationX = Clamp<float>(ValueRef.TranslationX, Min.GetLocation().X, Max.GetLocation().X);
-					ValueRef.TranslationY = Clamp<float>(ValueRef.TranslationY, Min.GetLocation().Y, Max.GetLocation().Y);
-					ValueRef.TranslationZ = Clamp<float>(ValueRef.TranslationZ, Min.GetLocation().Z, Max.GetLocation().Z);
+					ValueRef.TranslationX = LimitEnabled[0].Apply<float>(ValueRef.TranslationX, Min.GetLocation().X, Max.GetLocation().X);
+					ValueRef.TranslationY = LimitEnabled[1].Apply<float>(ValueRef.TranslationY, Min.GetLocation().Y, Max.GetLocation().Y);
+					ValueRef.TranslationZ = LimitEnabled[2].Apply<float>(ValueRef.TranslationZ, Min.GetLocation().Z, Max.GetLocation().Z);
 				}
-				if (bLimitRotation)
+				if (LimitEnabled[3].IsOn() || LimitEnabled[4].IsOn() || LimitEnabled[5].IsOn())
 				{
 					const FRotator Rotator = FQuat(ValueRef.RotationX, ValueRef.RotationY, ValueRef.RotationZ, ValueRef.RotationW).Rotator();
 					const FRotator MinRotator = Min.GetRotation().Rotator();
 					const FRotator MaxRotator = Max.GetRotation().Rotator();
 
-					FQuat LimitedQuat(FRotator(
-						Clamp<float>(Rotator.Pitch, MinRotator.Pitch, MaxRotator.Pitch),
-						Clamp<float>(Rotator.Yaw, MinRotator.Yaw, MaxRotator.Yaw),
-						Clamp<float>(Rotator.Roll, MinRotator.Roll, MaxRotator.Roll)
-					));
+					FRotator LimitedRotator = Rotator;
+					LimitedRotator.Pitch = LimitEnabled[3].Apply<float>(LimitedRotator.Pitch, MinRotator.Pitch, MaxRotator.Pitch);
+					LimitedRotator.Yaw = LimitEnabled[4].Apply<float>(LimitedRotator.Yaw, MinRotator.Yaw, MaxRotator.Yaw);
+					LimitedRotator.Roll = LimitEnabled[5].Apply<float>(LimitedRotator.Roll, MinRotator.Roll, MaxRotator.Roll);
+
+					FQuat LimitedQuat(LimitedRotator);
 					
 					ValueRef.RotationX = LimitedQuat.X;
 					ValueRef.RotationY = LimitedQuat.Y;
 					ValueRef.RotationZ = LimitedQuat.Z;
 					ValueRef.RotationW = LimitedQuat.W;
 				}
-				if (bLimitScale)
+				if (LimitEnabled[6].IsOn() || LimitEnabled[7].IsOn() || LimitEnabled[8].IsOn())
 				{
-					ValueRef.ScaleX = Clamp<float>(ValueRef.ScaleX, Min.GetScale3D().X, Max.GetScale3D().X);
-					ValueRef.ScaleY = Clamp<float>(ValueRef.ScaleY, Min.GetScale3D().Y, Max.GetScale3D().Y);
-					ValueRef.ScaleZ = Clamp<float>(ValueRef.ScaleZ, Min.GetScale3D().Z, Max.GetScale3D().Z);
+					ValueRef.ScaleX = LimitEnabled[6].Apply<float>(ValueRef.ScaleX, Min.GetScale3D().X, Max.GetScale3D().X);
+					ValueRef.ScaleY = LimitEnabled[7].Apply<float>(ValueRef.ScaleY, Min.GetScale3D().Y, Max.GetScale3D().Y);
+					ValueRef.ScaleZ = LimitEnabled[8].Apply<float>(ValueRef.ScaleZ, Min.GetScale3D().Z, Max.GetScale3D().Z);
 				}
 				break;
 			}
 			case ERigControlType::TransformNoScale:
 			{
+				if(LimitEnabled.Num() < 6)
+				{
+					return;
+				}
+
 				FTransformNoScale_Float& ValueRef = GetRef<FTransformNoScale_Float>();
 				const FTransformNoScale Min = InMinimumValue.GetRef<FTransformNoScale_Float>().ToTransform();
 				const FTransformNoScale Max = InMaximumValue.GetRef<FTransformNoScale_Float>().ToTransform();
 
-				if (bLimitTranslation)
+				if (LimitEnabled[0].IsOn() || LimitEnabled[1].IsOn() || LimitEnabled[2].IsOn())
 				{
-					ValueRef.TranslationX = Clamp<float>(ValueRef.TranslationX, Min.Location.X, Max.Location.X);
-					ValueRef.TranslationY = Clamp<float>(ValueRef.TranslationY, Min.Location.Y, Max.Location.Y);
-					ValueRef.TranslationZ = Clamp<float>(ValueRef.TranslationZ, Min.Location.Z, Max.Location.Z);
+					ValueRef.TranslationX = LimitEnabled[0].Apply<float>(ValueRef.TranslationX, Min.Location.X, Max.Location.X);
+					ValueRef.TranslationY = LimitEnabled[1].Apply<float>(ValueRef.TranslationY, Min.Location.Y, Max.Location.Y);
+					ValueRef.TranslationZ = LimitEnabled[2].Apply<float>(ValueRef.TranslationZ, Min.Location.Z, Max.Location.Z);
 				}
-				if (bLimitRotation)
+				if (LimitEnabled[3].IsOn() || LimitEnabled[4].IsOn() || LimitEnabled[5].IsOn())
 				{
 					const FRotator Rotator = FQuat(ValueRef.RotationX, ValueRef.RotationY, ValueRef.RotationZ, ValueRef.RotationW).Rotator();
 					const FRotator MinRotator = Min.Rotation.Rotator();
 					const FRotator MaxRotator = Max.Rotation.Rotator();
 
-					FQuat LimitedQuat(FRotator(
-						Clamp<float>(Rotator.Pitch, MinRotator.Pitch, MaxRotator.Pitch),
-						Clamp<float>(Rotator.Yaw, MinRotator.Yaw, MaxRotator.Yaw),
-						Clamp<float>(Rotator.Roll, MinRotator.Roll, MaxRotator.Roll)
-					));
+					FRotator LimitedRotator = Rotator;
+					LimitedRotator.Pitch = LimitEnabled[3].Apply<float>(LimitedRotator.Pitch, MinRotator.Pitch, MaxRotator.Pitch);
+					LimitedRotator.Yaw = LimitEnabled[4].Apply<float>(LimitedRotator.Yaw, MinRotator.Yaw, MaxRotator.Yaw);
+					LimitedRotator.Roll = LimitEnabled[5].Apply<float>(LimitedRotator.Roll, MinRotator.Roll, MaxRotator.Roll);
+
+					FQuat LimitedQuat(LimitedRotator);
 
 					ValueRef.RotationX = LimitedQuat.X;
 					ValueRef.RotationY = LimitedQuat.Y;
@@ -759,27 +871,32 @@ public:
 
 			case ERigControlType::EulerTransform:
 			{
+				if(LimitEnabled.Num() < 9)
+				{
+					return;
+				}
+
 				FEulerTransform_Float& ValueRef = GetRef<FEulerTransform_Float>();
 				const FEulerTransform_Float& Min = InMinimumValue.GetRef<FEulerTransform_Float>();
 				const FEulerTransform_Float& Max = InMaximumValue.GetRef<FEulerTransform_Float>();
 
-				if (bLimitTranslation)
+				if (LimitEnabled[0].IsOn() || LimitEnabled[1].IsOn() || LimitEnabled[2].IsOn())
 				{
-					ValueRef.TranslationX = Clamp<float>(ValueRef.TranslationX, Min.TranslationX, Max.TranslationX);
-					ValueRef.TranslationY = Clamp<float>(ValueRef.TranslationY, Min.TranslationY, Max.TranslationY);
-					ValueRef.TranslationZ = Clamp<float>(ValueRef.TranslationZ, Min.TranslationZ, Max.TranslationZ);
+					ValueRef.TranslationX = LimitEnabled[0].Apply<float>(ValueRef.TranslationX, Min.TranslationX, Max.TranslationX);
+					ValueRef.TranslationY = LimitEnabled[1].Apply<float>(ValueRef.TranslationY, Min.TranslationY, Max.TranslationY);
+					ValueRef.TranslationZ = LimitEnabled[2].Apply<float>(ValueRef.TranslationZ, Min.TranslationZ, Max.TranslationZ);
 				}
-				if (bLimitRotation)
+				if (LimitEnabled[3].IsOn() || LimitEnabled[4].IsOn() || LimitEnabled[5].IsOn())
 				{
-					ValueRef.RotationPitch = Clamp<float>(ValueRef.RotationPitch, Min.RotationPitch, Max.RotationPitch);
-					ValueRef.RotationYaw = Clamp<float>(ValueRef.RotationYaw, Min.RotationYaw, Max.RotationYaw);
-					ValueRef.RotationRoll = Clamp<float>(ValueRef.RotationRoll, Min.RotationRoll, Max.RotationRoll);
+					ValueRef.RotationPitch = LimitEnabled[3].Apply<float>(ValueRef.RotationPitch, Min.RotationPitch, Max.RotationPitch);
+					ValueRef.RotationYaw = LimitEnabled[4].Apply<float>(ValueRef.RotationYaw, Min.RotationYaw, Max.RotationYaw);
+					ValueRef.RotationRoll = LimitEnabled[5].Apply<float>(ValueRef.RotationRoll, Min.RotationRoll, Max.RotationRoll);
 				}
-				if (bLimitScale)
+				if (LimitEnabled[6].IsOn() || LimitEnabled[7].IsOn() || LimitEnabled[8].IsOn())
 				{
-					ValueRef.ScaleX = Clamp<float>(ValueRef.ScaleX, Min.ScaleX, Max.ScaleX);
-					ValueRef.ScaleY = Clamp<float>(ValueRef.ScaleY, Min.ScaleY, Max.ScaleY);
-					ValueRef.ScaleZ = Clamp<float>(ValueRef.ScaleZ, Min.ScaleZ, Max.ScaleZ);
+					ValueRef.ScaleX = LimitEnabled[6].Apply<float>(ValueRef.ScaleX, Min.ScaleX, Max.ScaleX);
+					ValueRef.ScaleY = LimitEnabled[7].Apply<float>(ValueRef.ScaleY, Min.ScaleY, Max.ScaleY);
+					ValueRef.ScaleZ = LimitEnabled[8].Apply<float>(ValueRef.ScaleZ, Min.ScaleZ, Max.ScaleZ);
 				}
 				break;
 			}
diff --git a/Engine/Plugins/Animation/ControlRig/Source/ControlRig/Public/Rigs/RigHierarchyElements.h b/Engine/Plugins/Animation/ControlRig/Source/ControlRig/Public/Rigs/RigHierarchyElements.h
--- a/Engine/Plugins/Animation/ControlRig/Source/ControlRig/Public/Rigs/RigHierarchyElements.h
+++ b/Engine/Plugins/Animation/ControlRig/Source/ControlRig/Public/Rigs/RigHierarchyElements.h
@@ -880,28 +880,23 @@ struct CONTROLRIG_API FRigControlSettings
 	UPROPERTY(BlueprintReadWrite, EditAnywhere, Category = Control)
 	bool bAnimatable;
 
-	/** True if the control has to obey translation limits. */
+	/** True if the control has limits. */
 	UPROPERTY(BlueprintReadWrite, EditAnywhere, Category = Limits)
-	bool bLimitTranslation;
+	TArray<FRigControlLimitEnabled> LimitEnabled;
 
-	/** True if the control has to obey rotation limits. */
-	UPROPERTY(BlueprintReadWrite, EditAnywhere, Category = Limits)
-	bool bLimitRotation;
-
-	/** True if the control has to obey scale limits. */
+	/**
+	 * True if the limits should be drawn in debug.
+	 * For this to be enabled you need to have at least one min and max limit turned on.
+	 */
 	UPROPERTY(BlueprintReadWrite, EditAnywhere, Category = Limits)
-	bool bLimitScale;
-
-	/** True if the limits should be drawn in debug. */
-	UPROPERTY(BlueprintReadWrite, EditAnywhere, Category = Limits, meta = (EditCondition = "bLimitTranslation || bLimitRotation || bLimitScale"))
 	bool bDrawLimits;
 
 	/** The minimum limit of the control's value */
-	UPROPERTY(BlueprintReadWrite, Category = Limits, meta = (EditCondition = "bLimitTranslation || bLimitRotation || bLimitScale"))
+	UPROPERTY(BlueprintReadWrite, Category = Limits)
 	FRigControlValue MinimumValue;
 
 	/** The maximum limit of the control's value */
-	UPROPERTY(BlueprintReadWrite, Category = Limits, meta = (EditCondition = "bLimitTranslation || bLimitRotation || bLimitScale"))
+	UPROPERTY(BlueprintReadWrite, Category = Limits)
 	FRigControlValue MaximumValue;
 
 	/** Set to true if the shape is enabled in 3d */
@@ -937,7 +932,7 @@ struct CONTROLRIG_API FRigControlSettings
 	/** Applies the limits expressed by these settings to a value */
 	FORCEINLINE void ApplyLimits(FRigControlValue& InOutValue) const
 	{
-		InOutValue.ApplyLimits(bLimitTranslation, bLimitRotation, bLimitScale, ControlType, MinimumValue, MaximumValue);
+		InOutValue.ApplyLimits(LimitEnabled, ControlType, MinimumValue, MaximumValue);
 	}
 
 	/** Applies the limits expressed by these settings to a transform */
@@ -949,15 +944,6 @@ struct CONTROLRIG_API FRigControlSettings
 		InOutValue = Value.GetAsTransform(ControlType, PrimaryAxis);
 	}
 
-	FORCEINLINE bool IsValueTypeEnabled(ERigControlValueType InValueType) const
-	{
-		if(InValueType == ERigControlValueType::Minimum || InValueType == ERigControlValueType::Maximum)
-		{
-			return bLimitTranslation || bLimitRotation || bLimitScale;
-		}
-		return true;
-	}
-
 	FORCEINLINE FRigControlValue GetIdentityValue() const
 	{
 		FRigControlValue Value;
@@ -971,6 +957,8 @@ struct CONTROLRIG_API FRigControlSettings
 	{
 		return !(*this == InOther);
 	}
+
+	void SetupLimitArrayForType(bool bLimitTranslation = false, bool bLimitRotation = false, bool bLimitScale = false);
 };
 
 USTRUCT(BlueprintType)
diff --git a/Engine/Plugins/Animation/ControlRig/Source/ControlRigEditor/Private/ControlRigElementDetails.cpp b/Engine/Plugins/Animation/ControlRig/Source/ControlRigEditor/Private/ControlRigElementDetails.cpp
--- a/Engine/Plugins/Animation/ControlRig/Source/ControlRigEditor/Private/ControlRigElementDetails.cpp
+++ b/Engine/Plugins/Animation/ControlRig/Source/ControlRigEditor/Private/ControlRigElementDetails.cpp
@@ -23,7 +23,6 @@
 #include "HAL/PlatformApplicationMisc.h"
 #include "ScopedTransaction.h"
 #include "AnimationCoreLibrary.h"
-#include "HAL/PlatformApplicationMisc.h"
 
 #define LOCTEXT_NAMESPACE "ControlRigElementDetails"
 
@@ -781,6 +780,21 @@ void FRigBaseElementDetails::CustomizeDetails(IDetailLayoutBuilder& DetailBuilde
 	DetailBuilder.HideCategory(TEXT("RigElement"));
 }
 
+URigHierarchy* FRigBaseElementDetails::GetHierarchyBeingDebugged() const
+{
+	URigHierarchy* HierarchyBeingDebugged = HierarchyBeingCustomized;
+		
+	if (UControlRig* DebuggedRig = Cast<UControlRig>(BlueprintBeingCustomized->GetObjectBeingDebugged()))
+	{
+		if(!DebuggedRig->IsSetupModeEnabled())
+		{
+			HierarchyBeingDebugged = DebuggedRig->GetHierarchy();
+		}
+	}
+
+	return HierarchyBeingDebugged;
+}
+
 FRigElementKey FRigBaseElementDetails::GetElementKey() const
 {
 	check(ObjectsBeingCustomized.Num() == 1);
@@ -977,16 +991,7 @@ void FRigBaseElementDetails::RegisterSectionMappings(FPropertyEditorModule& Prop
 {
 }
 
-namespace ERigTransformElementDetailsTransform
-{
-	enum Type
-	{
-		Initial,
-		Current,
-		Offset,
-		Max
-	};
-}
+TSharedPtr<TArray<ERigTransformElementDetailsTransform::Type>> FRigTransformElementDetails::PickedTransforms;
 
 void FRigTransformElementDetails::CustomizeDetails(IDetailLayoutBuilder& DetailBuilder)
 {
@@ -999,17 +1004,16 @@ void FRigTransformElementDetails::RegisterSectionMappings(FPropertyEditorModule&
 	
 	TSharedRef<FPropertySection> TransformSection = PropertyEditorModule.FindOrCreateSection(InClass->GetFName(), "Transform", LOCTEXT("Transform", "Transform"));
 	TransformSection->AddCategory("General");
+	TransformSection->AddCategory("Value");
 	TransformSection->AddCategory("Transform");
 }
 
 void FRigTransformElementDetails::CustomizeTransform(IDetailLayoutBuilder& DetailBuilder)
 {
 	URigHierarchy* HierarchyBeingDebugged = HierarchyBeingCustomized;
-	bool IsSetupModeEnabled = false;
 	if (UControlRig* DebuggedRig = Cast<UControlRig>(BlueprintBeingCustomized->GetObjectBeingDebugged()))
 	{
-		IsSetupModeEnabled = DebuggedRig->IsSetupModeEnabled(); 
- 		if(!IsSetupModeEnabled)
+ 		if(!DebuggedRig->IsSetupModeEnabled())
 		{
 			HierarchyBeingDebugged = DebuggedRig->GetHierarchy();
 		}
@@ -1017,7 +1021,8 @@ void FRigTransformElementDetails::CustomizeTransform(IDetailLayoutBuilder& Detai
 
 	TArray<FRigElementKey> Keys = GetElementKeys();
 	Keys = HierarchyBeingDebugged->SortKeys(Keys);
-	const bool bAllControls = !IsAnyElementNotOfType(ERigElementType::Control);
+	const bool bAllControls = !IsAnyElementNotOfType(ERigElementType::Control) && !IsAnyControlOfType(ERigControlType::Bool);
+	bool bShowLimits = false;
 	TArray<ERigTransformElementDetailsTransform::Type> TransformTypes;
 	TArray<FText> ButtonLabels;
 	TArray<FText> ButtonTooltips;
@@ -1039,6 +1044,24 @@ void FRigTransformElementDetails::CustomizeTransform(IDetailLayoutBuilder& Detai
 			LOCTEXT("CurrentTooltip", "Current animation transform"),
 			LOCTEXT("OffsetTooltip", "Offset transform under the control")
 		};
+
+		bShowLimits = !IsAnyControlNotOfType(ERigControlType::EulerTransform);
+
+		if(bShowLimits)
+		{
+			TransformTypes.Append({
+				ERigTransformElementDetailsTransform::Minimum,
+				ERigTransformElementDetailsTransform::Maximum
+			});
+			ButtonLabels.Append({
+				LOCTEXT("Min", "Min"),
+				LOCTEXT("Max", "Max")
+			});
+			ButtonTooltips.Append({
+				LOCTEXT("ValueMinimumTooltip", "The minimum limit(s) for the control"),
+				LOCTEXT("ValueMaximumTooltip", "The maximum limit(s) for the control")
+			});
+		}
 	}
 	else
 	{
@@ -1061,7 +1084,7 @@ void FRigTransformElementDetails::CustomizeTransform(IDetailLayoutBuilder& Detai
 	// determine if the transforms are enabled
 	for(int32 Index = 0; Index < TransformTypes.Num(); Index++)
 	{
-		ERigTransformElementDetailsTransform::Type CurrentTransformType = (ERigTransformElementDetailsTransform::Type)Index;
+		const ERigTransformElementDetailsTransform::Type CurrentTransformType = TransformTypes[Index];
 
 		bool bIsTransformEnabled = true;
 		if (IsAnyElementOfType(ERigElementType::Control))
@@ -1104,1129 +1127,1347 @@ void FRigTransformElementDetails::CustomizeTransform(IDetailLayoutBuilder& Detai
 		bTransformsEnabled.Add(bIsTransformEnabled);
 	}
 
-	static TAttribute<TArray<ERigTransformElementDetailsTransform::Type>> PickedTransforms = 
-		TArray<ERigTransformElementDetailsTransform::Type>({ERigTransformElementDetailsTransform::Current});
+	if(!PickedTransforms.IsValid())
+	{
+		PickedTransforms = MakeShareable(new TArray<ERigTransformElementDetailsTransform::Type>({ERigTransformElementDetailsTransform::Current}));
+	}
 
 	TSharedPtr<SSegmentedControl<ERigTransformElementDetailsTransform::Type>> TransformChoiceWidget =
 		SSegmentedControl<ERigTransformElementDetailsTransform::Type>::Create(
 			TransformTypes,
 			ButtonLabels,
 			ButtonTooltips,
-			PickedTransforms
+			*PickedTransforms.Get(),
+			true,
+			SSegmentedControl<ERigTransformElementDetailsTransform::Type>::FOnValuesChanged::CreateLambda(
+				[](TArray<ERigTransformElementDetailsTransform::Type> NewSelection)
+				{
+					(*FRigTransformElementDetails::PickedTransforms.Get()) = NewSelection;
+				}
+			)
 		);
 
 	IDetailCategoryBuilder& TransformCategory = DetailBuilder.EditCategory(TEXT("Transform"), LOCTEXT("Transform", "Transform"));
 	AddChoiceWidgetRow(TransformCategory, FText::FromString(TEXT("TransformType")), TransformChoiceWidget.ToSharedRef());
-	
-	TSharedRef<UE::Math::TVector<float>> IsComponentRelative = MakeShareable(new UE::Math::TVector<float>(1.f, 1.f, 1.f));
 
 	using TransformType = FEulerTransform;
 	using NumericType = FVector::FReal;
 
 	SAdvancedTransformInputBox<TransformType>::FArguments TransformWidgetArgs = SAdvancedTransformInputBox<TransformType>::FArguments()
-		.DisplayToggle(false)
-		.DisplayRelativeWorld(true)
-		.Font(IDetailLayoutBuilder::GetDetailFont())
-		.OnGetIsComponentRelative_Lambda(
-			[IsComponentRelative](ESlateTransformComponent::Type InComponent)
-			{
-				return IsComponentRelative->operator[]((int32)InComponent) > 0.f;
-			})
-		.OnIsComponentRelativeChanged_Lambda(
-			[IsComponentRelative](ESlateTransformComponent::Type InComponent, bool bIsRelative)
-			{
-				IsComponentRelative->operator[]((int32)InComponent) = bIsRelative ? 1.f : 0.f;
-			});
-	
+	.DisplayToggle(false)
+	.DisplayRelativeWorld(true)
+	.Font(IDetailLayoutBuilder::GetDetailFont());
+
 	for(int32 Index = 0; Index < ButtonLabels.Num(); Index++)
 	{
-		ERigTransformElementDetailsTransform::Type CurrentTransformType = (ERigTransformElementDetailsTransform::Type)Index;
-
-		auto GetRelativeAbsoluteTransforms = [CurrentTransformType, Keys, HierarchyBeingDebugged](
-			const FRigElementKey& Key,
-			ERigTransformElementDetailsTransform::Type InTransformType = ERigTransformElementDetailsTransform::Max
-			) -> TPair<TransformType, TransformType>
+		const ERigTransformElementDetailsTransform::Type CurrentTransformType = TransformTypes[Index];
+		ERigControlValueType CurrentValueType = ERigControlValueType::Current;
+		switch(CurrentTransformType)
 		{
-			if(InTransformType == ERigTransformElementDetailsTransform::Max)
+			case ERigTransformElementDetailsTransform::Initial:
 			{
-				InTransformType = CurrentTransformType;
+				CurrentValueType = ERigControlValueType::Initial;
+				break;
 			}
-
-			TransformType RelativeTransform = TransformType::Identity;
-			TransformType AbsoluteTransform = TransformType::Identity;
-
-			const bool bInitial = InTransformType == ERigTransformElementDetailsTransform::Initial; 
-			if(bInitial || InTransformType == ERigTransformElementDetailsTransform::Current)
+			case ERigTransformElementDetailsTransform::Minimum:
 			{
-				RelativeTransform = HierarchyBeingDebugged->GetLocalTransform(Key, bInitial);
-				AbsoluteTransform = HierarchyBeingDebugged->GetGlobalTransform(Key, bInitial);
+				CurrentValueType = ERigControlValueType::Minimum;
+				break;
 			}
-			else
+			case ERigTransformElementDetailsTransform::Maximum:
 			{
-				if(FRigControlElement* ControlElement = HierarchyBeingDebugged->Find<FRigControlElement>(Key))
-				{
-					RelativeTransform = HierarchyBeingDebugged->GetControlOffsetTransform(ControlElement, bInitial ? ERigTransformType::InitialLocal : ERigTransformType::CurrentLocal);
-					AbsoluteTransform = HierarchyBeingDebugged->GetControlOffsetTransform(ControlElement, bInitial ? ERigTransformType::InitialGlobal : ERigTransformType::CurrentGlobal);
-				}
+				CurrentValueType = ERigControlValueType::Maximum;
+				break;
 			}
+		}
 
-			return TPair<FTransform, FTransform>(RelativeTransform, AbsoluteTransform);
-		};
-
-		auto GetCombinedTransform = [IsComponentRelative, GetRelativeAbsoluteTransforms](
-			const FRigElementKey& Key,
-			ERigTransformElementDetailsTransform::Type InTransformType = ERigTransformElementDetailsTransform::Max
-			) -> TransformType
+		TransformWidgetArgs.Visibility_Lambda([TransformChoiceWidget, Index]() -> EVisibility
 		{
-			const TPair<TransformType, TransformType> TransformPair = GetRelativeAbsoluteTransforms(Key, InTransformType);
-			const TransformType RelativeTransform = TransformPair.Key;
-			const TransformType AbsoluteTransform = TransformPair.Value;
-
-			TransformType Xfo;
-			Xfo.SetLocation((IsComponentRelative->X > 0.f) ? RelativeTransform.GetLocation() : AbsoluteTransform.GetLocation());
-			Xfo.SetRotation((IsComponentRelative->Y > 0.f) ? RelativeTransform.GetRotation() : AbsoluteTransform.GetRotation());
-			Xfo.SetScale3D((IsComponentRelative->Z > 0.f) ? RelativeTransform.GetScale3D() : AbsoluteTransform.GetScale3D());
+			return TransformChoiceWidget->HasValue((ERigTransformElementDetailsTransform::Type)Index) ? EVisibility::Visible : EVisibility::Collapsed;
+		});
 
-			return Xfo;
-		};
+		TransformWidgetArgs.IsEnabled(bTransformsEnabled[Index]);
 
-		auto GetSingleTransform = [GetRelativeAbsoluteTransforms](
-			const FRigElementKey& Key,
-			bool bIsRelative,
-			ERigTransformElementDetailsTransform::Type InTransformType = ERigTransformElementDetailsTransform::Max
-			) -> TransformType
-		{
-			const TPair<TransformType, TransformType> TransformPair = GetRelativeAbsoluteTransforms(Key, InTransformType);
-			const TransformType RelativeTransform = TransformPair.Key;
-			const TransformType AbsoluteTransform = TransformPair.Value;
-			return bIsRelative ? RelativeTransform : AbsoluteTransform;
-		};
+		CreateEulerTransformValueWidgetRow(
+			HierarchyBeingDebugged,
+			Keys,
+			TransformWidgetArgs,
+			TransformCategory,
+			ButtonLabels[Index],
+			ButtonTooltips[Index],
+			CurrentTransformType,
+			CurrentValueType);
+	}
+}
 
-		auto SetSingleTransform = [CurrentTransformType, GetRelativeAbsoluteTransforms, this, IsSetupModeEnabled, HierarchyBeingDebugged](
-			const FRigElementKey& Key,
-			TransformType InTransform,
-			bool bIsRelative)
+bool FRigTransformElementDetails::IsCurrentLocalEnabled() const
+{
+	for(TWeakObjectPtr<UDetailsViewWrapperObject> ObjectBeingCustomized : ObjectsBeingCustomized)
+	{
+		if(ObjectBeingCustomized.IsValid())
 		{
-			const bool bInitial = CurrentTransformType == ERigTransformElementDetailsTransform::Initial; 
-
-			TArray<URigHierarchy*> HierarchiesToUpdate;
-			HierarchiesToUpdate.Add(HierarchyBeingDebugged);
-			if(bInitial || IsSetupModeEnabled)
+			if(ObjectBeingCustomized->GetContent<FRigBaseElement>().GetType() == ERigElementType::Control)
 			{
-				HierarchiesToUpdate.Add(HierarchyBeingCustomized);
+				return false;
 			}
+		}
+	}
+	return true;
+}
 
-			for(URigHierarchy* HierarchyToUpdate : HierarchiesToUpdate)
-			{
-				if(bInitial || CurrentTransformType == ERigTransformElementDetailsTransform::Current)
-				{
-					if(bIsRelative)
-					{
-						HierarchyToUpdate->SetLocalTransform(Key, InTransform, bInitial, true, true);
-					}
-					else
-					{
-						HierarchyToUpdate->SetGlobalTransform(Key, InTransform, bInitial, true, true);
-					}
-				}
-				else
-				{
-					if(!bIsRelative)
-					{
-						const FTransform ParentTransform = HierarchyToUpdate->GetParentTransform(Key, bInitial);
-						InTransform = FTransform(InTransform).GetRelativeTransform(ParentTransform);
-					}
-							
-					if(FRigControlElement* ControlElement = HierarchyToUpdate->Find<FRigControlElement>(Key))
-					{
-						HierarchyToUpdate->SetControlOffsetTransform(Key, InTransform, bInitial, true, true);
-					}
-				}
-			}
-		};
+void FRigTransformElementDetails::AddChoiceWidgetRow(IDetailCategoryBuilder& InCategory, const FText& InSearchText, TSharedRef<SWidget> InWidget)
+{
+	InCategory.AddCustomRow(FText::FromString(TEXT("TransformType")))
+	.ValueContent()
+	.MinDesiredWidth(375.f)
+	.MaxDesiredWidth(375.f)
+	.HAlign(HAlign_Left)
+	[
+		SNew(SHorizontalBox)
+		+SHorizontalBox::Slot()
+		.AutoWidth()
+		.HAlign(HAlign_Left)
+		.VAlign(VAlign_Center)
+		[
+			InWidget
+		]
+	];
+}
 
-		TransformWidgetArgs.Visibility_Lambda([TransformChoiceWidget, Index]() -> EVisibility
+TSharedPtr<TArray<ERigControlValueType>> FRigControlElementDetails::PickedValueTypes;
+
+void FRigTransformElementDetails::CreateEulerTransformValueWidgetRow(
+	URigHierarchy* HierarchyBeingDebugged,
+	const TArray<FRigElementKey>& Keys,
+	SAdvancedTransformInputBox<FEulerTransform>::FArguments TransformWidgetArgs,
+	IDetailCategoryBuilder& CategoryBuilder,
+	const FText& Label,
+	const FText& Tooltip,
+	ERigTransformElementDetailsTransform::Type CurrentTransformType,
+	ERigControlValueType ValueType)
+{
+	using TransformType = FEulerTransform;
+	using NumericType = FVector::FReal;
+
+	TSharedRef<UE::Math::TVector<float>> IsComponentRelativeStorage = MakeShareable(new UE::Math::TVector<float>(1.f, 1.f, 1.f));
+	TransformWidgetArgs.OnGetIsComponentRelative_Lambda(
+		[IsComponentRelativeStorage](ESlateTransformComponent::Type InComponent)
 		{
-			return TransformChoiceWidget->HasValue((ERigTransformElementDetailsTransform::Type)Index) ? EVisibility::Visible : EVisibility::Collapsed;
+			return IsComponentRelativeStorage->operator[]((int32)InComponent) > 0.f;
+		})
+	.OnIsComponentRelativeChanged_Lambda(
+		[IsComponentRelativeStorage](ESlateTransformComponent::Type InComponent, bool bIsRelative)
+		{
+			IsComponentRelativeStorage->operator[]((int32)InComponent) = bIsRelative ? 1.f : 0.f;
 		});
 
-		TransformWidgetArgs.IsEnabled(bTransformsEnabled[Index]);
-
-		TransformWidgetArgs.OnGetNumericValue_Lambda([Keys, GetCombinedTransform](
-			ESlateTransformComponent::Type Component,
-			ESlateRotationRepresentation::Type Representation,
-			ESlateTransformSubComponent::Type SubComponent) -> TOptional<NumericType>
+	const TSharedPtr<ESlateRotationRepresentation::Type> RotationRepresentationStorage =
+		MakeShareable(new ESlateRotationRepresentation::Type(ESlateRotationRepresentation::Rotator));
+	TransformWidgetArgs.RotationRepresentation(RotationRepresentationStorage);
+	
+	auto IsComponentRelative = [TransformWidgetArgs](int32 Component) -> bool
+	{
+		if(TransformWidgetArgs._OnGetIsComponentRelative.IsBound())
 		{
-			TOptional<NumericType> FirstValue;
+			return TransformWidgetArgs._OnGetIsComponentRelative.Execute((ESlateTransformComponent::Type)Component);
+		}
+		return true;
+	};
 
-			for(int32 Index = 0; Index < Keys.Num(); Index++)
-			{
-				const FRigElementKey& Key = Keys[Index];
-				TransformType Xfo = GetCombinedTransform(Key);
+	auto ConformComponentRelative = [TransformWidgetArgs, IsComponentRelative](int32 Component)
+	{
+		if(TransformWidgetArgs._OnIsComponentRelativeChanged.IsBound())
+		{
+			bool bRelative = IsComponentRelative(Component);
+			TransformWidgetArgs._OnIsComponentRelativeChanged.Execute(ESlateTransformComponent::Location, bRelative);
+			TransformWidgetArgs._OnIsComponentRelativeChanged.Execute(ESlateTransformComponent::Rotation, bRelative);
+			TransformWidgetArgs._OnIsComponentRelativeChanged.Execute(ESlateTransformComponent::Scale, bRelative);
+		}
+	};
 
-				TOptional<NumericType> CurrentValue = SAdvancedTransformInputBox<TransformType>::GetNumericValueFromTransform(Xfo, Component, Representation, SubComponent);
-				if(!CurrentValue.IsSet())
+	switch(CurrentTransformType)
+	{
+		case ERigTransformElementDetailsTransform::Minimum:
+		case ERigTransformElementDetailsTransform::Maximum:
+		{
+			TransformWidgetArgs.AllowEditRotationRepresentation(false);
+			TransformWidgetArgs.DisplayRelativeWorld(false);
+			TransformWidgetArgs.DisplayToggle(true);
+			TransformWidgetArgs.OnGetToggleChecked_Lambda([Keys, HierarchyBeingDebugged, ValueType]
+				(
+					ESlateTransformComponent::Type Component,
+					ESlateRotationRepresentation::Type RotationRepresentation,
+					ESlateTransformSubComponent::Type SubComponent
+				) -> ECheckBoxState
 				{
-					return CurrentValue;
-				}
+					TOptional<bool> FirstValue;
 
-				if(Index == 0)
-				{
-					FirstValue = CurrentValue;
-				}
-				else
-				{
-					if(!FMath::IsNearlyEqual(FirstValue.GetValue(), CurrentValue.GetValue()))
+					for(const FRigElementKey& Key : Keys)
 					{
-						return TOptional<NumericType>();
-					}
-				}
-			}
-			
-			return FirstValue;
-		});
-
-		TransformWidgetArgs.OnNumericValueChanged_Lambda([
-			Keys,
-			this,
-			IsComponentRelative,
-			GetSingleTransform,
-			SetSingleTransform
-		](
-			ESlateTransformComponent::Type Component,
-			ESlateRotationRepresentation::Type Representation,
-			ESlateTransformSubComponent::Type SubComponent,
-			NumericType InNumericValue)
-		{
-			const bool bIsRelative = IsComponentRelative->Component((int32)Component) > 0.f;
+						if(const FRigControlElement* ControlElement = HierarchyBeingDebugged->Find<FRigControlElement>(Key))
+						{
+							TOptional<bool> Value;
 
-			FScopedTransaction Transaction(LOCTEXT("ChangeNumericValue", "Change Numeric Value"));
+							switch(ControlElement->Settings.ControlType)
+							{
+								case ERigControlType::Position:
+								case ERigControlType::Rotator:
+								case ERigControlType::Scale:
+								{
+									if(ControlElement->Settings.LimitEnabled.Num() == 3)
+									{
+										const int32 Index = ControlElement->Settings.ControlType == ERigControlType::Rotator ?
+											int32(SubComponent) - int32(ESlateTransformSubComponent::Pitch) :
+											int32(SubComponent) - int32(ESlateTransformSubComponent::X);
+
+										Value = ControlElement->Settings.LimitEnabled[Index].GetForValueType(ValueType);
+									}
+									break;
+								}
+								case ERigControlType::EulerTransform:
+								{
+									if(ControlElement->Settings.LimitEnabled.Num() == 9)
+									{
+										switch(Component)
+										{
+											case ESlateTransformComponent::Location:
+											{
+												switch(SubComponent)
+												{
+													case ESlateTransformSubComponent::X:
+													{
+														Value = ControlElement->Settings.LimitEnabled[0].GetForValueType(ValueType);
+														break;
+													}
+													case ESlateTransformSubComponent::Y:
+													{
+														Value = ControlElement->Settings.LimitEnabled[1].GetForValueType(ValueType);
+														break;
+													}
+													case ESlateTransformSubComponent::Z:
+													{
+														Value = ControlElement->Settings.LimitEnabled[2].GetForValueType(ValueType);
+														break;
+													}
+													default:
+													{
+														break;
+													}
+												}
+												break;
+											}
+											case ESlateTransformComponent::Rotation:
+											{
+												switch(SubComponent)
+												{
+													case ESlateTransformSubComponent::Pitch:
+													{
+														Value = ControlElement->Settings.LimitEnabled[3].GetForValueType(ValueType);
+														break;
+													}
+													case ESlateTransformSubComponent::Yaw:
+													{
+														Value = ControlElement->Settings.LimitEnabled[4].GetForValueType(ValueType);
+														break;
+													}
+													case ESlateTransformSubComponent::Roll:
+													{
+														Value = ControlElement->Settings.LimitEnabled[5].GetForValueType(ValueType);
+														break;
+													}
+													default:
+													{
+														break;
+													}
+												}
+												break;
+											}
+											case ESlateTransformComponent::Scale:
+											{
+												switch(SubComponent)
+												{
+													case ESlateTransformSubComponent::X:
+													{
+														Value = ControlElement->Settings.LimitEnabled[6].GetForValueType(ValueType);
+														break;
+													}
+													case ESlateTransformSubComponent::Y:
+													{
+														Value = ControlElement->Settings.LimitEnabled[7].GetForValueType(ValueType);
+														break;
+													}
+													case ESlateTransformSubComponent::Z:
+													{
+														Value = ControlElement->Settings.LimitEnabled[8].GetForValueType(ValueType);
+														break;
+													}
+													default:
+													{
+														break;
+													}
+												}
+												break;
+											}
+										}
+									}
+									break;
+								}
+							}
 
-			for(const FRigElementKey& Key : Keys)
-			{
-				TransformType Transform = GetSingleTransform(Key, bIsRelative);
-				SAdvancedTransformInputBox<TransformType>::ApplyNumericValueChange(Transform, InNumericValue, Component, Representation, SubComponent);
-				SetSingleTransform(Key, Transform, bIsRelative);
-			}
-		});
+							if(Value.IsSet())
+							{
+								if(FirstValue.IsSet())
+								{
+									if(FirstValue.GetValue() != Value.GetValue())
+									{
+										return ECheckBoxState::Undetermined;
+									}
+								}
+								else
+								{
+									FirstValue = Value;
+								}
+							}
+						}
+					}
 
-		TransformWidgetArgs.OnCopyToClipboard_Lambda([Keys, IsComponentRelative, GetSingleTransform](
-			ESlateTransformComponent::Type InComponent
+					if(!ensure(FirstValue.IsSet()))
+					{
+						return ECheckBoxState::Undetermined;
+					}
+					return FirstValue.GetValue() ? ECheckBoxState::Checked : ECheckBoxState::Unchecked;
+				});
+				
+			TransformWidgetArgs.OnToggleChanged_Lambda([ValueType, Keys, this]
+			(
+				ESlateTransformComponent::Type Component,
+				ESlateRotationRepresentation::Type RotationRepresentation,
+				ESlateTransformSubComponent::Type SubComponent,
+				ECheckBoxState CheckState
 			)
-		{
-			if(Keys.Num() == 0)
 			{
-				return;
-			}
+				if(CheckState == ECheckBoxState::Undetermined)
+				{
+					return;
+				}
 
-			// make sure that we use the same relative setting on all components when copying
-			IsComponentRelative->Y = IsComponentRelative->Z = IsComponentRelative->X;
-			const bool bIsRelative = IsComponentRelative->Y > 0.f; 
+				const bool Value = CheckState == ECheckBoxState::Checked;
 
-			const FRigElementKey& FirstKey = Keys[0];
-			TransformType Xfo = GetSingleTransform(FirstKey, bIsRelative);
+				FScopedTransaction Transaction(LOCTEXT("ChangeLimitToggle", "Change Limit Toggle"));
+				HierarchyBeingCustomized->Modify();
 
-			FString Content;
-			switch(InComponent)
-			{
-				case ESlateTransformComponent::Location:
-				{
-					const FVector Data = Xfo.GetLocation();
-					TBaseStructure<FVector>::Get()->ExportText(Content, &Data, &Data, nullptr, PPF_None, nullptr);
-					break;
-				}
-				case ESlateTransformComponent::Rotation:
-				{
-					const FRotator Data = Xfo.Rotator();
-					TBaseStructure<FRotator>::Get()->ExportText(Content, &Data, &Data, nullptr, PPF_None, nullptr);
-					break;
-				}
-				case ESlateTransformComponent::Scale:
-				{
-					const FVector Data = Xfo.GetScale3D();
-					TBaseStructure<FVector>::Get()->ExportText(Content, &Data, &Data, nullptr, PPF_None, nullptr);
-					break;
-				}
-				case ESlateTransformComponent::Max:
-				default:
+				for(const FRigElementKey& Key : Keys)
 				{
-					TBaseStructure<TransformType>::Get()->ExportText(Content, &Xfo, &Xfo, nullptr, PPF_None, nullptr);
-					break;
-				}
-			}
-
-			if(!Content.IsEmpty())
-			{
-				FPlatformApplicationMisc::ClipboardCopy(*Content);
-			}
-		});
+					if(FRigControlElement* ControlElement = HierarchyBeingCustomized->Find<FRigControlElement>(Key))
+					{
+						switch(ControlElement->Settings.ControlType)
+						{
+							case ERigControlType::Position:
+							case ERigControlType::Rotator:
+							case ERigControlType::Scale:
+							{
+								if(ControlElement->Settings.LimitEnabled.Num() == 3)
+								{
+									const int32 Index = ControlElement->Settings.ControlType == ERigControlType::Rotator ?
+										int32(SubComponent) - int32(ESlateTransformSubComponent::Pitch) :
+										int32(SubComponent) - int32(ESlateTransformSubComponent::X);
 
-		TransformWidgetArgs.OnPasteFromClipboard_Lambda([Keys, IsComponentRelative, GetSingleTransform, SetSingleTransform](
-			ESlateTransformComponent::Type InComponent
-			)
+									ControlElement->Settings.LimitEnabled[Index].SetForValueType(ValueType, Value);
+								}
+								break;
+							}
+							case ERigControlType::EulerTransform:
+							{
+								if(ControlElement->Settings.LimitEnabled.Num() == 9)
+								{
+									switch(Component)
+									{
+										case ESlateTransformComponent::Location:
+										{
+											switch(SubComponent)
+											{
+												case ESlateTransformSubComponent::X:
+												{
+													ControlElement->Settings.LimitEnabled[0].SetForValueType(ValueType, Value);
+													break;
+												}
+												case ESlateTransformSubComponent::Y:
+												{
+													ControlElement->Settings.LimitEnabled[1].SetForValueType(ValueType, Value);
+													break;
+												}
+												case ESlateTransformSubComponent::Z:
+												{
+													ControlElement->Settings.LimitEnabled[2].SetForValueType(ValueType, Value);
+													break;
+												}
+												default:
+												{
+													break;
+												}
+											}
+											break;
+										}
+										case ESlateTransformComponent::Rotation:
+										{
+											switch(SubComponent)
+											{
+												case ESlateTransformSubComponent::Pitch:
+												{
+													ControlElement->Settings.LimitEnabled[3].SetForValueType(ValueType, Value);
+													break;
+												}
+												case ESlateTransformSubComponent::Yaw:
+												{
+													ControlElement->Settings.LimitEnabled[4].SetForValueType(ValueType, Value);
+													break;
+												}
+												case ESlateTransformSubComponent::Roll:
+												{
+													ControlElement->Settings.LimitEnabled[5].SetForValueType(ValueType, Value);
+													break;
+												}
+												default:
+												{
+													break;
+												}
+											}
+											break;
+										}
+										case ESlateTransformComponent::Scale:
+										{
+											switch(SubComponent)
+											{
+												case ESlateTransformSubComponent::X:
+												{
+													ControlElement->Settings.LimitEnabled[6].SetForValueType(ValueType, Value);
+													break;
+												}
+												case ESlateTransformSubComponent::Y:
+												{
+													ControlElement->Settings.LimitEnabled[7].SetForValueType(ValueType, Value);
+													break;
+												}
+												case ESlateTransformSubComponent::Z:
+												{
+													ControlElement->Settings.LimitEnabled[8].SetForValueType(ValueType, Value);
+													break;
+												}
+												default:
+												{
+													break;
+												}
+											}
+											break;
+										}
+									}
+								}
+								break;
+							}
+						}
+						
+						HierarchyBeingCustomized->SetControlSettings(ControlElement, ControlElement->Settings, true, true, true);
+					}
+				}
+			});
+			break;
+		}
+		default:
 		{
-			if(Keys.Num() == 0)
-			{
-				return;
-			}
-			
-			
-			// make sure that we use the same relative setting on all components when pasting
-			IsComponentRelative->Y = IsComponentRelative->Z = IsComponentRelative->X;
-			const bool bIsRelative = IsComponentRelative->Y > 0.f; 
-
-			FString Content;
-			FPlatformApplicationMisc::ClipboardPaste(Content);
+			TransformWidgetArgs.AllowEditRotationRepresentation(true);
+			TransformWidgetArgs.DisplayRelativeWorld(true);
+			TransformWidgetArgs.DisplayToggle(false);
+			TransformWidgetArgs._OnGetToggleChecked.Unbind();
+			TransformWidgetArgs._OnToggleChanged.Unbind();
+			break;
+		}
+	}
 
-			if(Content.IsEmpty())
-			{
-				return;
-			}
+	auto GetRelativeAbsoluteTransforms = [CurrentTransformType, Keys, HierarchyBeingDebugged](
+		const FRigElementKey& Key,
+		ERigTransformElementDetailsTransform::Type InTransformType = ERigTransformElementDetailsTransform::Max
+		) -> TPair<TransformType, TransformType>
+	{
+		if(InTransformType == ERigTransformElementDetailsTransform::Max)
+		{
+			InTransformType = CurrentTransformType;
+		}
 
-			FScopedTransaction Transaction(LOCTEXT("PasteTransform", "Paste Transform"));
+		TransformType RelativeTransform = TransformType::Identity;
+		TransformType AbsoluteTransform = TransformType::Identity;
 
-			for(const FRigElementKey& Key : Keys)
+		const bool bInitial = InTransformType == ERigTransformElementDetailsTransform::Initial; 
+		if(bInitial || InTransformType == ERigTransformElementDetailsTransform::Current)
+		{
+			RelativeTransform = HierarchyBeingDebugged->GetLocalTransform(Key, bInitial);
+			AbsoluteTransform = HierarchyBeingDebugged->GetGlobalTransform(Key, bInitial);
+		}
+		else
+		{
+			if(FRigControlElement* ControlElement = HierarchyBeingDebugged->Find<FRigControlElement>(Key))
 			{
-				TransformType Xfo = GetSingleTransform(Key, bIsRelative);
+				const ERigControlType ControlType = ControlElement->Settings.ControlType;
 
+				if(InTransformType == ERigTransformElementDetailsTransform::Offset)
 				{
-					class FRigPasteTransformWidgetErrorPipe : public FOutputDevice
-					{
-					public:
-
-						int32 NumErrors;
-
-						FRigPasteTransformWidgetErrorPipe()
-							: FOutputDevice()
-							, NumErrors(0)
-						{
-						}
-
-						virtual void Serialize(const TCHAR* V, ELogVerbosity::Type Verbosity, const class FName& Category) override
-						{
-							UE_LOG(LogControlRig, Error, TEXT("Error Pasting to Widget: %s"), V);
-							NumErrors++;
-						}
-					};
-
-					FRigPasteTransformWidgetErrorPipe ErrorPipe;
-					
-					switch(InComponent)
+					RelativeTransform = HierarchyBeingDebugged->GetControlOffsetTransform(ControlElement, ERigTransformType::InitialLocal);
+					AbsoluteTransform = HierarchyBeingDebugged->GetControlOffsetTransform(ControlElement, ERigTransformType::InitialGlobal);
+				}
+				else if(InTransformType == ERigTransformElementDetailsTransform::Minimum)
+				{
+					switch(ControlType)
 					{
-						case ESlateTransformComponent::Location:
+						case ERigControlType::Position:
 						{
-							FVector Data = Xfo.GetLocation();
-							TBaseStructure<FVector>::Get()->ImportText(*Content, &Data, nullptr, PPF_None, &ErrorPipe, TBaseStructure<FVector>::Get()->GetName(), true);
-							Xfo.SetLocation(Data);
+							const FVector3f Data = 
+								HierarchyBeingDebugged->GetControlValue(ControlElement, ERigControlValueType::Minimum)
+								.Get<FVector3f>();
+							AbsoluteTransform = RelativeTransform = FEulerTransform(Data, FRotator::ZeroRotator, FVector3f::OneVector);
 							break;
 						}
-						case ESlateTransformComponent::Rotation:
+						case ERigControlType::Rotator:
 						{
-							FRotator Data = Xfo.Rotator();
-							TBaseStructure<FRotator>::Get()->ImportText(*Content, &Data, nullptr, PPF_None, &ErrorPipe, TBaseStructure<FRotator>::Get()->GetName(), true);
-							Xfo.SetRotator(Data);
+							const FVector3f Data = 
+								HierarchyBeingDebugged->GetControlValue(ControlElement, ERigControlValueType::Minimum)
+								.Get<FVector3f>();
+							FRotator Rotator = FRotator::MakeFromEuler(Data);
+							AbsoluteTransform = RelativeTransform = FEulerTransform(FVector3f::ZeroVector, Rotator, FVector3f::OneVector);
 							break;
 						}
-						case ESlateTransformComponent::Scale:
+						case ERigControlType::Scale:
 						{
-							FVector Data = Xfo.GetScale3D();
-							TBaseStructure<FVector>::Get()->ImportText(*Content, &Data, nullptr, PPF_None, &ErrorPipe, TBaseStructure<FVector>::Get()->GetName(), true);
-							Xfo.SetScale3D(Data);
+							const FVector3f Data = 
+								HierarchyBeingDebugged->GetControlValue(ControlElement, ERigControlValueType::Minimum)
+								.Get<FVector3f>();
+							AbsoluteTransform = RelativeTransform = FEulerTransform(FVector3f::ZeroVector, FRotator::ZeroRotator, Data);
 							break;
 						}
-						case ESlateTransformComponent::Max:
-						default:
+						case ERigControlType::EulerTransform:
 						{
-							TBaseStructure<TransformType>::Get()->ImportText(*Content, &Xfo, nullptr, PPF_None, &ErrorPipe, TBaseStructure<TransformType>::Get()->GetName(), true);
+							const FRigControlValue::FEulerTransform_Float EulerTransform = 
+								HierarchyBeingDebugged->GetControlValue(ControlElement, ERigControlValueType::Minimum)
+								.Get<FRigControlValue::FEulerTransform_Float>();
+							AbsoluteTransform = RelativeTransform = EulerTransform.ToTransform();
 							break;
 						}
 					}
-
-					if(ErrorPipe.NumErrors == 0)
+				}
+				else if(InTransformType == ERigTransformElementDetailsTransform::Maximum)
+				{
+					switch(ControlType)
 					{
-						SetSingleTransform(Key, Xfo, bIsRelative);
+						case ERigControlType::Position:
+						{
+							const FVector3f Data = 
+								HierarchyBeingDebugged->GetControlValue(ControlElement, ERigControlValueType::Maximum)
+								.Get<FVector3f>();
+							AbsoluteTransform = RelativeTransform = FEulerTransform(Data, FRotator::ZeroRotator, FVector3f::OneVector);
+							break;
+						}
+						case ERigControlType::Rotator:
+						{
+							const FVector3f Data = 
+								HierarchyBeingDebugged->GetControlValue(ControlElement, ERigControlValueType::Maximum)
+								.Get<FVector3f>();
+							FRotator Rotator = FRotator::MakeFromEuler(Data);
+							AbsoluteTransform = RelativeTransform = FEulerTransform(FVector3f::ZeroVector, Rotator, FVector3f::OneVector);
+							break;
+						}
+						case ERigControlType::Scale:
+						{
+							const FVector3f Data = 
+								HierarchyBeingDebugged->GetControlValue(ControlElement, ERigControlValueType::Maximum)
+								.Get<FVector3f>();
+							AbsoluteTransform = RelativeTransform = FEulerTransform(FVector3f::ZeroVector, FRotator::ZeroRotator, Data);
+							break;
+						}
+						case ERigControlType::EulerTransform:
+						{
+							const FRigControlValue::FEulerTransform_Float EulerTransform = 
+								HierarchyBeingDebugged->GetControlValue(ControlElement, ERigControlValueType::Maximum)
+								.Get<FRigControlValue::FEulerTransform_Float>();
+							AbsoluteTransform = RelativeTransform = EulerTransform.ToTransform();
+							break;
+						}
 					}
 				}
 			}
-		});
+		}
 
-		TransformWidgetArgs.DiffersFromDefault_Lambda([
-			CurrentTransformType,
-			Keys,
-			GetSingleTransform
-			
-		](
-			ESlateTransformComponent::Type InComponent) -> bool
+		return TPair<FTransform, FTransform>(RelativeTransform, AbsoluteTransform);
+	};
+
+	
+	auto GetCombinedTransform = [IsComponentRelative, GetRelativeAbsoluteTransforms](
+		const FRigElementKey& Key,
+		ERigTransformElementDetailsTransform::Type InTransformType = ERigTransformElementDetailsTransform::Max
+		) -> TransformType
+	{
+		const TPair<TransformType, TransformType> TransformPair = GetRelativeAbsoluteTransforms(Key, InTransformType);
+		const TransformType RelativeTransform = TransformPair.Key;
+		const TransformType AbsoluteTransform = TransformPair.Value;
+
+		TransformType Xfo;
+		Xfo.SetLocation((IsComponentRelative(0)) ? RelativeTransform.GetLocation() : AbsoluteTransform.GetLocation());
+		Xfo.SetRotation((IsComponentRelative(1)) ? RelativeTransform.GetRotation() : AbsoluteTransform.GetRotation());
+		Xfo.SetScale3D((IsComponentRelative(2)) ? RelativeTransform.GetScale3D() : AbsoluteTransform.GetScale3D());
+
+		return Xfo;
+	};
+
+	auto GetSingleTransform = [GetRelativeAbsoluteTransforms](
+		const FRigElementKey& Key,
+		bool bIsRelative,
+		ERigTransformElementDetailsTransform::Type InTransformType = ERigTransformElementDetailsTransform::Max
+		) -> TransformType
+	{
+		const TPair<TransformType, TransformType> TransformPair = GetRelativeAbsoluteTransforms(Key, InTransformType);
+		const TransformType RelativeTransform = TransformPair.Key;
+		const TransformType AbsoluteTransform = TransformPair.Value;
+		return bIsRelative ? RelativeTransform : AbsoluteTransform;
+	};
+
+	auto SetSingleTransform = [CurrentTransformType, GetRelativeAbsoluteTransforms, this, HierarchyBeingDebugged](
+		const FRigElementKey& Key,
+		TransformType InTransform,
+		bool bIsRelative,
+		bool bSetupUndoRedo)
+	{
+		const bool bCurrent = CurrentTransformType == ERigTransformElementDetailsTransform::Current; 
+		const bool bInitial = CurrentTransformType == ERigTransformElementDetailsTransform::Initial; 
+
+		bool bSetupModeEnabled = false;
+		if (UControlRig* DebuggedRig = Cast<UControlRig>(BlueprintBeingCustomized->GetObjectBeingDebugged()))
 		{
-			for(const FRigElementKey& Key : Keys)
-			{
-				const TransformType CurrentTransform = GetSingleTransform(Key, true);
-				TransformType DefaultTransform;
+			bSetupModeEnabled = DebuggedRig->IsSetupModeEnabled();
+		}
+
+		TArray<URigHierarchy*> HierarchiesToUpdate;
+		HierarchiesToUpdate.Add(HierarchyBeingDebugged);
+		if(!bCurrent || bSetupModeEnabled)
+		{
+			HierarchiesToUpdate.Add(HierarchyBeingCustomized);
+		}
 
-				switch(CurrentTransformType)
+		for(URigHierarchy* HierarchyToUpdate : HierarchiesToUpdate)
+		{
+			if(bInitial || CurrentTransformType == ERigTransformElementDetailsTransform::Current)
+			{
+				if(bIsRelative)
 				{
-					case ERigTransformElementDetailsTransform::Current:
-					{
-						DefaultTransform = GetSingleTransform(Key, true, ERigTransformElementDetailsTransform::Initial);
-						break;
-					}
-					default:
-					{
-						DefaultTransform = TransformType::Identity; 
-						break;
-					}
+					HierarchyToUpdate->SetLocalTransform(Key, InTransform, bInitial, true, bSetupUndoRedo);
 				}
-
-				switch(InComponent)
+				else
 				{
-					case ESlateTransformComponent::Location:
-					{
-						if(!DefaultTransform.GetLocation().Equals(CurrentTransform.GetLocation()))
-						{
-							return true;
-						}
-						break;
-					}
-					case ESlateTransformComponent::Rotation:
+					HierarchyToUpdate->SetGlobalTransform(Key, InTransform, bInitial, true, bSetupUndoRedo);
+				}
+			}
+			else
+			{
+				if(FRigControlElement* ControlElement = HierarchyToUpdate->Find<FRigControlElement>(Key))
+				{
+					const ERigControlType ControlType = ControlElement->Settings.ControlType;
+
+					if(CurrentTransformType == ERigTransformElementDetailsTransform::Offset)
 					{
-						if(!DefaultTransform.Rotator().Equals(CurrentTransform.Rotator()))
+						if(!bIsRelative)
 						{
-							return true;
+							const FTransform ParentTransform = HierarchyToUpdate->GetParentTransform(Key, bInitial);
+							InTransform = FTransform(InTransform).GetRelativeTransform(ParentTransform);
 						}
-						break;
+						HierarchyToUpdate->SetControlOffsetTransform(Key, InTransform, true, true, bSetupUndoRedo);
 					}
-					case ESlateTransformComponent::Scale:
+					else if(CurrentTransformType == ERigTransformElementDetailsTransform::Minimum)
 					{
-						if(!DefaultTransform.GetScale3D().Equals(CurrentTransform.GetScale3D()))
+						switch(ControlType)
 						{
-							return true;
+							case ERigControlType::Position:
+							{
+								const FRigControlValue Value = FRigControlValue::Make<FVector3f>(InTransform.GetLocation());
+								HierarchyToUpdate->SetControlValue(ControlElement, Value, ERigControlValueType::Minimum, bSetupUndoRedo, true);
+								break;
+							}
+							case ERigControlType::Rotator:
+							{
+								const FVector3f Euler = InTransform.Rotator().Euler();
+								const FRigControlValue Value = FRigControlValue::Make<FVector3f>(Euler);
+								HierarchyToUpdate->SetControlValue(ControlElement, Value, ERigControlValueType::Minimum, bSetupUndoRedo, true);
+								break;
+							}
+							case ERigControlType::Scale:
+							{
+								const FRigControlValue Value = FRigControlValue::Make<FVector3f>(InTransform.GetScale3D());
+								HierarchyToUpdate->SetControlValue(ControlElement, Value, ERigControlValueType::Minimum, bSetupUndoRedo, true);
+								break;
+							}
+							case ERigControlType::EulerTransform:
+							{
+								const FRigControlValue Value = FRigControlValue::Make<FRigControlValue::FEulerTransform_Float>(InTransform);
+								HierarchyToUpdate->SetControlValue(ControlElement, Value, ERigControlValueType::Minimum, bSetupUndoRedo, true);
+								break;
+							}
 						}
-						break;
 					}
-					default: // also no component whole transform
+					else if(CurrentTransformType == ERigTransformElementDetailsTransform::Maximum)
 					{
-						if(!DefaultTransform.GetLocation().Equals(CurrentTransform.GetLocation()) ||
-							!DefaultTransform.Rotator().Equals(CurrentTransform.Rotator()) ||
-							!DefaultTransform.GetScale3D().Equals(CurrentTransform.GetScale3D()))
+						switch(ControlType)
 						{
-							return true;
+							case ERigControlType::Position:
+							{
+								const FRigControlValue Value = FRigControlValue::Make<FVector3f>(InTransform.GetLocation());
+								HierarchyToUpdate->SetControlValue(ControlElement, Value, ERigControlValueType::Maximum, bSetupUndoRedo, true);
+								break;
+							}
+							case ERigControlType::Rotator:
+							{
+								const FVector3f Euler = InTransform.Rotator().Euler();
+								const FRigControlValue Value = FRigControlValue::Make<FVector3f>(Euler);
+								HierarchyToUpdate->SetControlValue(ControlElement, Value, ERigControlValueType::Maximum, bSetupUndoRedo, true);
+								break;
+							}
+							case ERigControlType::Scale:
+							{
+								const FRigControlValue Value = FRigControlValue::Make<FVector3f>(InTransform.GetScale3D());
+								HierarchyToUpdate->SetControlValue(ControlElement, Value, ERigControlValueType::Maximum, bSetupUndoRedo, true);
+								break;
+							}
+							case ERigControlType::EulerTransform:
+							{
+								const FRigControlValue Value = FRigControlValue::Make<FRigControlValue::FEulerTransform_Float>(InTransform);
+								HierarchyToUpdate->SetControlValue(ControlElement, Value, ERigControlValueType::Maximum, bSetupUndoRedo, true);
+								break;
+							}
 						}
-						break;
 					}
 				}
 			}
-			return false;
-		});
+		}
+	};
+
+	TransformWidgetArgs.OnGetNumericValue_Lambda([Keys, GetCombinedTransform](
+		ESlateTransformComponent::Type Component,
+		ESlateRotationRepresentation::Type Representation,
+		ESlateTransformSubComponent::Type SubComponent) -> TOptional<NumericType>
+	{
+		TOptional<NumericType> FirstValue;
+
+		for(int32 Index = 0; Index < Keys.Num(); Index++)
+		{
+			const FRigElementKey& Key = Keys[Index];
+			TransformType Xfo = GetCombinedTransform(Key);
+
+			TOptional<NumericType> CurrentValue = SAdvancedTransformInputBox<TransformType>::GetNumericValueFromTransform(Xfo, Component, Representation, SubComponent);
+			if(!CurrentValue.IsSet())
+			{
+				return CurrentValue;
+			}
+
+			if(Index == 0)
+			{
+				FirstValue = CurrentValue;
+			}
+			else
+			{
+				if(!FMath::IsNearlyEqual(FirstValue.GetValue(), CurrentValue.GetValue()))
+				{
+					return TOptional<NumericType>();
+				}
+			}
+		}
+		
+		return FirstValue;
+	});
+
+	TransformWidgetArgs.OnNumericValueChanged_Lambda(
+	[
+		Keys,
+		this,
+		IsComponentRelative,
+		GetSingleTransform,
+		SetSingleTransform,
+		CurrentTransformType
+	](
+		ESlateTransformComponent::Type Component,
+		ESlateRotationRepresentation::Type Representation,
+		ESlateTransformSubComponent::Type SubComponent,
+		NumericType InNumericValue)
+	{
+		const bool bIsRelative = IsComponentRelative((int32)Component);
+
+		for(const FRigElementKey& Key : Keys)
+		{
+			TransformType Transform = GetSingleTransform(Key, bIsRelative);
+			TransformType PreviousTransform = Transform;
+			SAdvancedTransformInputBox<TransformType>::ApplyNumericValueChange(Transform, InNumericValue, Component, Representation, SubComponent);
+
+			if(!FRigControlElementDetails::Equals(Transform, PreviousTransform))
+			{
+				if(!SliderTransaction.IsValid())
+				{
+					SliderTransaction = MakeShareable(new FScopedTransaction(NSLOCTEXT("ControlRigElementDetails", "ChangeNumericValue", "Change Numeric Value")));
+					HierarchyBeingCustomized->Modify();
+				}
+							
+				SetSingleTransform(Key, Transform, bIsRelative, false);
+			}
+		}
+	});
+
+	TransformWidgetArgs.OnNumericValueCommitted_Lambda(
+	[
+		Keys,
+		this,
+		IsComponentRelative,
+		GetSingleTransform,
+		SetSingleTransform,
+		CurrentTransformType
+	](
+		ESlateTransformComponent::Type Component,
+		ESlateRotationRepresentation::Type Representation,
+		ESlateTransformSubComponent::Type SubComponent,
+		NumericType InNumericValue,
+		ETextCommit::Type InCommitType)
+	{
+		const bool bIsRelative = IsComponentRelative((int32)Component);
 
-		TransformWidgetArgs.OnResetToDefault_Lambda([CurrentTransformType, Keys, GetSingleTransform, SetSingleTransform](
-			ESlateTransformComponent::Type InComponent)
 		{
-			FScopedTransaction Transaction(LOCTEXT("ResetTransformToDefault", "Reset Transform to Default"));
+			FScopedTransaction Transaction(LOCTEXT("ChangeNumericValue", "Change Numeric Value"));
+			if(!SliderTransaction.IsValid())
+			{
+				HierarchyBeingCustomized->Modify();
+			}
 			
 			for(const FRigElementKey& Key : Keys)
 			{
-				TransformType CurrentTransform = GetSingleTransform(Key, true);
-				TransformType DefaultTransform;
+				TransformType Transform = GetSingleTransform(Key, bIsRelative);
+				SAdvancedTransformInputBox<TransformType>::ApplyNumericValueChange(Transform, InNumericValue, Component, Representation, SubComponent);
+				SetSingleTransform(Key, Transform, bIsRelative, true);
+			}
+		}
+
+		SliderTransaction.Reset();
+	});
+
+	TransformWidgetArgs.OnCopyToClipboard_Lambda([Keys, IsComponentRelative, ConformComponentRelative, GetSingleTransform](
+		ESlateTransformComponent::Type InComponent
+		)
+	{
+		if(Keys.Num() == 0)
+		{
+			return;
+		}
+
+		// make sure that we use the same relative setting on all components when copying
+		ConformComponentRelative(0);
+		const bool bIsRelative = IsComponentRelative(0); 
+
+		const FRigElementKey& FirstKey = Keys[0];
+		TransformType Xfo = GetSingleTransform(FirstKey, bIsRelative);
+
+		FString Content;
+		switch(InComponent)
+		{
+			case ESlateTransformComponent::Location:
+			{
+				const FVector Data = Xfo.GetLocation();
+				TBaseStructure<FVector>::Get()->ExportText(Content, &Data, &Data, nullptr, PPF_None, nullptr);
+				break;
+			}
+			case ESlateTransformComponent::Rotation:
+			{
+				const FRotator Data = Xfo.Rotator();
+				TBaseStructure<FRotator>::Get()->ExportText(Content, &Data, &Data, nullptr, PPF_None, nullptr);
+				break;
+			}
+			case ESlateTransformComponent::Scale:
+			{
+				const FVector Data = Xfo.GetScale3D();
+				TBaseStructure<FVector>::Get()->ExportText(Content, &Data, &Data, nullptr, PPF_None, nullptr);
+				break;
+			}
+			case ESlateTransformComponent::Max:
+			default:
+			{
+				TBaseStructure<TransformType>::Get()->ExportText(Content, &Xfo, &Xfo, nullptr, PPF_None, nullptr);
+				break;
+			}
+		}
+
+		if(!Content.IsEmpty())
+		{
+			FPlatformApplicationMisc::ClipboardCopy(*Content);
+		}
+	});
+
+	TransformWidgetArgs.OnPasteFromClipboard_Lambda([this, Keys, IsComponentRelative, ConformComponentRelative, GetSingleTransform, SetSingleTransform](
+		ESlateTransformComponent::Type InComponent
+		)
+	{
+		if(Keys.Num() == 0)
+		{
+			return;
+		}
+		
+		
+		// make sure that we use the same relative setting on all components when pasting
+		ConformComponentRelative(0);
+		const bool bIsRelative = IsComponentRelative(0); 
+
+		FString Content;
+		FPlatformApplicationMisc::ClipboardPaste(Content);
+
+		if(Content.IsEmpty())
+		{
+			return;
+		}
+
+		FScopedTransaction Transaction(LOCTEXT("PasteTransform", "Paste Transform"));
+		HierarchyBeingCustomized->Modify();
+
+		for(const FRigElementKey& Key : Keys)
+		{
+			TransformType Xfo = GetSingleTransform(Key, bIsRelative);
 
-				switch(CurrentTransformType)
+			{
+				class FRigPasteTransformWidgetErrorPipe : public FOutputDevice
 				{
-					case ERigTransformElementDetailsTransform::Current:
+				public:
+
+					int32 NumErrors;
+
+					FRigPasteTransformWidgetErrorPipe()
+						: FOutputDevice()
+						, NumErrors(0)
 					{
-						DefaultTransform = GetSingleTransform(Key, true, ERigTransformElementDetailsTransform::Initial);
-						break;
 					}
-					default:
+
+					virtual void Serialize(const TCHAR* V, ELogVerbosity::Type Verbosity, const class FName& Category) override
 					{
-						DefaultTransform = TransformType::Identity; 
-						break;
+						UE_LOG(LogControlRig, Error, TEXT("Error Pasting to Widget: %s"), V);
+						NumErrors++;
 					}
-				}
+				};
 
+				FRigPasteTransformWidgetErrorPipe ErrorPipe;
+				
 				switch(InComponent)
 				{
 					case ESlateTransformComponent::Location:
 					{
-						CurrentTransform.SetLocation(DefaultTransform.GetLocation());
+						FVector Data = Xfo.GetLocation();
+						TBaseStructure<FVector>::Get()->ImportText(*Content, &Data, nullptr, PPF_None, &ErrorPipe, TBaseStructure<FVector>::Get()->GetName(), true);
+						Xfo.SetLocation(Data);
 						break;
 					}
 					case ESlateTransformComponent::Rotation:
 					{
-						CurrentTransform.SetRotator(DefaultTransform.Rotator());
+						FRotator Data = Xfo.Rotator();
+						TBaseStructure<FRotator>::Get()->ImportText(*Content, &Data, nullptr, PPF_None, &ErrorPipe, TBaseStructure<FRotator>::Get()->GetName(), true);
+						Xfo.SetRotator(Data);
 						break;
 					}
 					case ESlateTransformComponent::Scale:
 					{
-						CurrentTransform.SetScale3D(DefaultTransform.GetScale3D());
+						FVector Data = Xfo.GetScale3D();
+						TBaseStructure<FVector>::Get()->ImportText(*Content, &Data, nullptr, PPF_None, &ErrorPipe, TBaseStructure<FVector>::Get()->GetName(), true);
+						Xfo.SetScale3D(Data);
 						break;
 					}
-					default: // whole transform / max component
+					case ESlateTransformComponent::Max:
+					default:
 					{
-						CurrentTransform = DefaultTransform;
+						TBaseStructure<TransformType>::Get()->ImportText(*Content, &Xfo, nullptr, PPF_None, &ErrorPipe, TBaseStructure<TransformType>::Get()->GetName(), true);
 						break;
 					}
 				}
 
-				SetSingleTransform(Key, CurrentTransform, true);
+				if(ErrorPipe.NumErrors == 0)
+				{
+					SetSingleTransform(Key, Xfo, bIsRelative, true);
+				}
 			}
-		});
-
-		SAdvancedTransformInputBox<TransformType>::ConstructGroupedTransformRows(
-			TransformCategory, 
-			ButtonLabels[Index], 
-			ButtonTooltips[Index], 
-			TransformWidgetArgs);
-	}
-}
+		}
+	});
 
-bool FRigTransformElementDetails::IsCurrentLocalEnabled() const
-{
-	for(TWeakObjectPtr<UDetailsViewWrapperObject> ObjectBeingCustomized : ObjectsBeingCustomized)
+	TransformWidgetArgs.DiffersFromDefault_Lambda([
+		CurrentTransformType,
+		Keys,
+		GetSingleTransform
+		
+	](
+		ESlateTransformComponent::Type InComponent) -> bool
 	{
-		if(ObjectBeingCustomized.IsValid())
+		for(const FRigElementKey& Key : Keys)
 		{
-			if(ObjectBeingCustomized->GetContent<FRigBaseElement>().GetType() == ERigElementType::Control)
+			const TransformType CurrentTransform = GetSingleTransform(Key, true);
+			TransformType DefaultTransform;
+
+			switch(CurrentTransformType)
 			{
-				return false;
+				case ERigTransformElementDetailsTransform::Current:
+				{
+					DefaultTransform = GetSingleTransform(Key, true, ERigTransformElementDetailsTransform::Initial);
+					break;
+				}
+				default:
+				{
+					DefaultTransform = TransformType::Identity; 
+					break;
+				}
 			}
-		}
-	}
-	return true;
-}
 
-void FRigTransformElementDetails::AddChoiceWidgetRow(IDetailCategoryBuilder& InCategory, const FText& InSearchText, TSharedRef<SWidget> InWidget)
-{
-	InCategory.AddCustomRow(FText::FromString(TEXT("TransformType")))
-	.ValueContent()
-	.MinDesiredWidth(375.f)
-	.MaxDesiredWidth(375.f)
-	.HAlign(HAlign_Left)
-	[
-		SNew(SHorizontalBox)
-		+SHorizontalBox::Slot()
-		.AutoWidth()
-		.HAlign(HAlign_Left)
-		.VAlign(VAlign_Center)
-		[
-			InWidget
-		]
-	];
-}
+			switch(InComponent)
+			{
+				case ESlateTransformComponent::Location:
+				{
+					if(!DefaultTransform.GetLocation().Equals(CurrentTransform.GetLocation()))
+					{
+						return true;
+					}
+					break;
+				}
+				case ESlateTransformComponent::Rotation:
+				{
+					if(!DefaultTransform.Rotator().Equals(CurrentTransform.Rotator()))
+					{
+						return true;
+					}
+					break;
+				}
+				case ESlateTransformComponent::Scale:
+				{
+					if(!DefaultTransform.GetScale3D().Equals(CurrentTransform.GetScale3D()))
+					{
+						return true;
+					}
+					break;
+				}
+				default: // also no component whole transform
+				{
+					if(!DefaultTransform.GetLocation().Equals(CurrentTransform.GetLocation()) ||
+						!DefaultTransform.Rotator().Equals(CurrentTransform.Rotator()) ||
+						!DefaultTransform.GetScale3D().Equals(CurrentTransform.GetScale3D()))
+					{
+						return true;
+					}
+					break;
+				}
+			}
+		}
+		return false;
+	});
 
-void FRigControlElementDetails_SetupBoolValueWidget(IDetailCategoryBuilder& InCategory, const TSharedRef<IPropertyUtilities>& InPropertyUtilities, ERigControlValueType InValueType, const FRigElementKey& InKey, URigHierarchy* InHierarchy)
-{
-	UEnum* ControlTypeEnum = StaticEnum<ERigControlType>();
-	UEnum* ValueTypeEnum = StaticEnum<ERigControlValueType>();
+	TransformWidgetArgs.OnResetToDefault_Lambda([this, CurrentTransformType, Keys, GetSingleTransform, SetSingleTransform](
+		ESlateTransformComponent::Type InComponent)
+	{
+		FScopedTransaction Transaction(LOCTEXT("ResetTransformToDefault", "Reset Transform to Default"));
+		HierarchyBeingCustomized->Modify();
 
-	const FString ValueTypeName = ValueTypeEnum->GetDisplayNameTextByValue((int64)InValueType).ToString();
-	const FText PropertyLabel = FText::FromString(FString::Printf(TEXT("%s Value"), *ValueTypeName));
-	TWeakObjectPtr<URigHierarchy> HierarchyPtr = InHierarchy;
+		for(const FRigElementKey& Key : Keys)
+		{
+			TransformType CurrentTransform = GetSingleTransform(Key, true);
+			TransformType DefaultTransform;
 
-	InCategory.AddCustomRow(PropertyLabel)
-	.NameContent()
-	[
-	    SNew(STextBlock)
-		.Text(PropertyLabel)
-		.Font(IDetailLayoutBuilder::GetDetailFont())
-	]
-	.ValueContent()
-	[
-	    SNew(SVerticalBox)
-		+ SVerticalBox::Slot()
-	    [
-	        SNew(SCheckBox)
-	        .IsChecked_Lambda([HierarchyPtr, InKey, InValueType]() -> ECheckBoxState
-	        {
-	        	if(HierarchyPtr.IsValid())
-	        	{
-	        		if(FRigControlElement* ControlElement = HierarchyPtr->Find<FRigControlElement>(InKey))
-	        		{
-                        bool Value = HierarchyPtr->GetControlValue(ControlElement, InValueType).Get<bool>();
-						return Value ? ECheckBoxState::Checked : ECheckBoxState::Unchecked;
-                    }
-	        	}
-	        	return ECheckBoxState::Unchecked;
-	        })
-	        .OnCheckStateChanged_Lambda([HierarchyPtr, InKey, InValueType](ECheckBoxState NewState)
-	        {
-	        	if(HierarchyPtr.IsValid())
-	        	{
-	        		if(FRigControlElement* ControlElement = HierarchyPtr->Find<FRigControlElement>(InKey))
-	        		{
-                        const FRigControlValue Value = FRigControlValue::Make<bool>(NewState == ECheckBoxState::Checked);
-						HierarchyPtr->SetControlValue(ControlElement->GetKey(), Value, InValueType, true);
-	        			if(InValueType == ERigControlValueType::Initial)
-	        			{
-                            if(UControlRigBlueprint* Blueprint = RigElementDetails_GetBlueprintFromHierarchy(HierarchyPtr.Get()))
-                            {
-                                Blueprint->Hierarchy->SetControlValue(ControlElement->GetKey(), Value, InValueType, true);
-                            }
-                        }
-		            }
-	        	}
-	        })
-	    ]
-	]
-	.IsEnabled(TAttribute<bool>::Create(TAttribute<bool>::FGetter::CreateLambda([HierarchyPtr, InKey, InValueType]()->bool
-    {
-		if(HierarchyPtr.IsValid())
-		{
-			if(FRigControlElement* ControlElement = HierarchyPtr->Find<FRigControlElement>(InKey))
+			switch(CurrentTransformType)
 			{
-				return ControlElement->Settings.IsValueTypeEnabled(InValueType);
+				case ERigTransformElementDetailsTransform::Current:
+				{
+					DefaultTransform = GetSingleTransform(Key, true, ERigTransformElementDetailsTransform::Initial);
+					break;
+				}
+				default:
+				{
+					DefaultTransform = TransformType::Identity; 
+					break;
+				}
 			}
-		}
-		return false;
-    })));
-}
-
-void FRigControlElementDetails_SetupIntegerValueWidget(IDetailCategoryBuilder& InCategory, const TSharedRef<IPropertyUtilities>& InPropertyUtilities, ERigControlValueType InValueType, const FRigElementKey& InKey, URigHierarchy* InHierarchy)
-{
-	FRigControlElement* ControlElement = InHierarchy->Find<FRigControlElement>(InKey);
-	if(ControlElement == nullptr)
-	{
-		return;
-	}
-
-	UEnum* ControlTypeEnum = StaticEnum<ERigControlType>();
-	UEnum* ValueTypeEnum = StaticEnum<ERigControlValueType>();
-
-	const FString ValueTypeName = ValueTypeEnum->GetDisplayNameTextByValue((int64)InValueType).ToString();
-	const FText PropertyLabel = FText::FromString(FString::Printf(TEXT("%s Value"), *ValueTypeName));
-	TWeakObjectPtr<URigHierarchy> HierarchyPtr = InHierarchy;
-
-	const TAttribute<bool> EnabledAttribute = TAttribute<bool>::Create(TAttribute<bool>::FGetter::CreateLambda([HierarchyPtr, InKey, InValueType]()->bool
-    {
-		if(HierarchyPtr.IsValid())
-		{
-            if(FRigControlElement* ControlElement = HierarchyPtr->Find<FRigControlElement>(InKey))
-            {
-                return ControlElement->Settings.IsValueTypeEnabled(InValueType);
-            }
-        }
-		return false;
-    }));
-
-	const TAttribute<EVisibility> VisibilityAttribute = TAttribute<EVisibility>::Create(TAttribute<EVisibility>::FGetter::CreateLambda([EnabledAttribute]()->EVisibility
-	{
-		return EnabledAttribute.Get() ? EVisibility::Visible : EVisibility::Hidden;
-	}));
-
-	if (ControlElement->Settings.ControlEnum)
-	{
-		InCategory.AddCustomRow(PropertyLabel)
-		.Visibility(VisibilityAttribute)
-        .NameContent()
-        [
-            SNew(STextBlock)
-            .Text(PropertyLabel)
-            .Font(IDetailLayoutBuilder::GetDetailFont())
-        ]
-        .ValueContent()
-        .MinDesiredWidth(125.0f * 3.0f) // copied from FComponentTransformDetails
-        .MaxDesiredWidth(125.0f * 3.0f)
-        [
-            SNew(SVerticalBox)
-            + SVerticalBox::Slot()
-            [
-                SNew(SEnumComboBox, ControlElement->Settings.ControlEnum)
-                .CurrentValue_Lambda([HierarchyPtr, InKey, InValueType]() -> int32
-                {
-                	if(HierarchyPtr.IsValid())
-                	{
-                		if(FRigControlElement* ControlElement = HierarchyPtr->Find<FRigControlElement>(InKey))
-                		{
-                            return HierarchyPtr->GetControlValue(ControlElement, InValueType).Get<int32>();
-                        }
-                	}
-                	return 0;
-                })
-                .OnEnumSelectionChanged_Lambda([HierarchyPtr, InKey, InValueType](int32 NewSelection, ESelectInfo::Type)
-                {
-                	if(HierarchyPtr.IsValid())
-                	{
-                		if(FRigControlElement* ControlElement = HierarchyPtr->Find<FRigControlElement>(InKey))
-                		{
-                            const FRigControlValue Value = FRigControlValue::Make<int32>(NewSelection);
-                            HierarchyPtr->SetControlValue(ControlElement->GetKey(), Value, InValueType, true);
-                			if(InValueType == ERigControlValueType::Initial)
-                			{
-                                if(UControlRigBlueprint* Blueprint = RigElementDetails_GetBlueprintFromHierarchy(HierarchyPtr.Get()))
-                                {
-                                    Blueprint->Hierarchy->SetControlValue(ControlElement->GetKey(), Value, InValueType, true);
-                                }
-                            }
-	                    }
-                	}
-                })
-                .Font(FEditorStyle::GetFontStyle(TEXT("MenuItem.Font")))
-            ]
-        ]
-		.IsEnabled(EnabledAttribute);
-	}
-	else
-	{
-		InCategory.AddCustomRow(PropertyLabel)
-		.Visibility(VisibilityAttribute)
-        .NameContent()
-        [
-            SNew(STextBlock)
-            .Text(PropertyLabel)
-            .Font(IDetailLayoutBuilder::GetDetailFont())
-        ]
-        .ValueContent()
-        .MinDesiredWidth(125.0f * 3.0f) // copied from FComponentTransformDetails
-        .MaxDesiredWidth(125.0f * 3.0f)
-        [
-            SNew(SVerticalBox)
-            + SVerticalBox::Slot()
-            [
-                SNew(SNumericEntryBox<int32>)
-                .Font(FEditorStyle::GetFontStyle(TEXT("MenuItem.Font")))
-                .AllowSpin(InValueType == ERigControlValueType::Current || InValueType == ERigControlValueType::Initial)
-                .MinSliderValue_Lambda([HierarchyPtr, InKey, InValueType]() -> TOptional<int32>
-				{
-				    if(InValueType == ERigControlValueType::Current || InValueType == ERigControlValueType::Initial)
-				    {
-				    	if(HierarchyPtr.IsValid())
-				    	{
-				    		if(FRigControlElement* ControlElement = HierarchyPtr->Find<FRigControlElement>(InKey))
-				    		{
-                                return ControlElement->Settings.MinimumValue.Get<int32>();
-                            }
-				    	}
-				    }
-				    return TOptional<int32>();
-				})
-				.MaxSliderValue_Lambda([HierarchyPtr, InKey, InValueType]() -> TOptional<int32>
-				{
-				    if(InValueType == ERigControlValueType::Current || InValueType == ERigControlValueType::Initial)
-				    {
-				    	if(HierarchyPtr.IsValid())
-				    	{
-				    		if(FRigControlElement* ControlElement = HierarchyPtr->Find<FRigControlElement>(InKey))
-				    		{
-                                return ControlElement->Settings.MaximumValue.Get<int32>();
-                            }
-				    	}
-				    }
-				    return TOptional<int32>();
-				})
-                .Value_Lambda([HierarchyPtr, InKey, InValueType]() -> int32
-                {
-                	if(HierarchyPtr.IsValid())
-                	{
-                		if(FRigControlElement* ControlElement = HierarchyPtr->Find<FRigControlElement>(InKey))
-                		{
-                            return HierarchyPtr->GetControlValue(ControlElement, InValueType).Get<int32>();
-                        }
-                	}
-                	return 0;
-                })
-                .OnValueChanged_Lambda([HierarchyPtr, InKey, InValueType](TOptional<int32> InNewSelection)
-                {
-                	if(InNewSelection.IsSet())
-                	{
-                		if(HierarchyPtr.IsValid())
-                		{
-                			if(FRigControlElement* ControlElement = HierarchyPtr->Find<FRigControlElement>(InKey))
-                			{
-                                const FRigControlValue Value = FRigControlValue::Make<int32>(InNewSelection.GetValue());
-								HierarchyPtr->SetControlValue(ControlElement->GetKey(), Value, InValueType, true);
-                				if(InValueType == ERigControlValueType::Initial)
-                				{
-                                    if(UControlRigBlueprint* Blueprint = RigElementDetails_GetBlueprintFromHierarchy(HierarchyPtr.Get()))
-                                    {
-                                        Blueprint->Hierarchy->SetControlValue(ControlElement->GetKey(), Value, InValueType, true);
-                                    }
-                                }
-                            }
-                		}
-                	}
-                })
-            ]
-        ]
-		.IsEnabled(EnabledAttribute);
-	}
-}
-
-void FRigControlElementDetails_SetupFloatValueWidget(IDetailCategoryBuilder& InCategory, const TSharedRef<IPropertyUtilities>& InPropertyUtilities, ERigControlValueType InValueType, const FRigElementKey& InKey, URigHierarchy* InHierarchy)
-{
-	UEnum* ControlTypeEnum = StaticEnum<ERigControlType>();
-	UEnum* ValueTypeEnum = StaticEnum<ERigControlValueType>();
-
-	const FString ValueTypeName = ValueTypeEnum->GetDisplayNameTextByValue((int64)InValueType).ToString();
-	const FText PropertyLabel = FText::FromString(FString::Printf(TEXT("%s Value"), *ValueTypeName));
-	TWeakObjectPtr<URigHierarchy> HierarchyPtr = InHierarchy;
-
-	const TAttribute<bool> EnabledAttribute = TAttribute<bool>::Create(TAttribute<bool>::FGetter::CreateLambda([HierarchyPtr, InKey, InValueType]()->bool
-    {
-        if(HierarchyPtr.IsValid())
-        {
-            if(FRigControlElement* ControlElement = HierarchyPtr->Find<FRigControlElement>(InKey))
-            {
-                return ControlElement->Settings.IsValueTypeEnabled(InValueType);
-            }
-        }
-        return false;
-    }));
-
-	const TAttribute<EVisibility> VisibilityAttribute = TAttribute<EVisibility>::Create(TAttribute<EVisibility>::FGetter::CreateLambda([EnabledAttribute]()->EVisibility
-    {
-        return EnabledAttribute.Get() ? EVisibility::Visible : EVisibility::Hidden;
-    }));
-
-	InCategory.AddCustomRow(PropertyLabel)
-	.Visibility(VisibilityAttribute)
-    .NameContent()
-    [
-        SNew(STextBlock)
-        .Text(PropertyLabel)
-        .Font(IDetailLayoutBuilder::GetDetailFont())
-    ]
-    .ValueContent()
-    .MinDesiredWidth(125.0f * 3.0f) // copied from FComponentTransformDetails
-    .MaxDesiredWidth(125.0f * 3.0f)
-    [
-        SNew(SVerticalBox)
-        + SVerticalBox::Slot()
-        [
-            SNew(SNumericEntryBox<float>)
-            .Font(FEditorStyle::GetFontStyle(TEXT("MenuItem.Font")))
-            .AllowSpin(InValueType == ERigControlValueType::Current || InValueType == ERigControlValueType::Initial)
-            .Value_Lambda([HierarchyPtr, InKey, InValueType]() -> float
-            {
-            	if(HierarchyPtr.IsValid())
-            	{
-            		if(FRigControlElement* ControlElement = HierarchyPtr->Find<FRigControlElement>(InKey))
-            		{
-                        return HierarchyPtr->GetControlValue(ControlElement, InValueType).Get<float>();
-                    }
-            	}
-            	return 0.f;
-            })
-            .MinSliderValue_Lambda([HierarchyPtr, InKey, InValueType]() -> TOptional<float>
-            {
-                if(InValueType == ERigControlValueType::Current || InValueType == ERigControlValueType::Initial)
-                {
-                	if(HierarchyPtr.IsValid())
-                	{
-                		if(FRigControlElement* ControlElement = HierarchyPtr->Find<FRigControlElement>(InKey))
-                		{
-                            return ControlElement->Settings.MinimumValue.Get<float>();
-                        }
-                	}
-                }
-                return TOptional<float>();
-            })
-            .MaxSliderValue_Lambda([HierarchyPtr, InKey, InValueType]() -> TOptional<float>
-            {
-                if(InValueType == ERigControlValueType::Current || InValueType == ERigControlValueType::Initial)
-                {
-                	if(HierarchyPtr.IsValid())
-                	{
-                		if(FRigControlElement* ControlElement = HierarchyPtr->Find<FRigControlElement>(InKey))
-                		{
-                            return ControlElement->Settings.MaximumValue.Get<float>();
-                        }
-                	}
-                }
-                return TOptional<float>();
-            })
-            .OnValueChanged_Lambda([HierarchyPtr, InKey, InValueType](TOptional<float> InNewSelection)
-            {
-            	if(InNewSelection.IsSet())
-            	{
-            		if(HierarchyPtr.IsValid())
-            		{
-            			if(FRigControlElement* ControlElement = HierarchyPtr->Find<FRigControlElement>(InKey))
-            			{
-                            const FRigControlValue Value = FRigControlValue::Make<float>(InNewSelection.GetValue());
-                            HierarchyPtr->SetControlValue(ControlElement->GetKey(), Value, InValueType, true);
-            				if(InValueType == ERigControlValueType::Initial)
-            				{
-            					if(UControlRigBlueprint* Blueprint = RigElementDetails_GetBlueprintFromHierarchy(HierarchyPtr.Get()))
-            					{
-            						Blueprint->Hierarchy->SetControlValue(ControlElement->GetKey(), Value, InValueType, true);
-            					}
-            				}
-                        }
-            		}
-            	}
-            })
-        ]
-    ]
-	.IsEnabled(EnabledAttribute);
-}
-
-template<typename T>
-T FRigControlElementDetails_ExtractValue(const FRigControlValue& InValue)
-{
-	return InValue.Get<T>();
-}
-
-template<typename T>
-FRigControlValue FRigControlElementDetails_PackageValue(const T& InValue)
-{
-	return FRigControlValue::Make<T>(InValue);
-}
 
-template<>
-FVector2D FRigControlElementDetails_ExtractValue(const FRigControlValue& InValue)
-{
-	const FVector3f TempValue = InValue.Get<FVector3f>();
-	return FVector2D(TempValue.X, TempValue.Y);
-}
-
-template<>
-FRigControlValue FRigControlElementDetails_PackageValue(const FVector2D& InValue)
-{
-	const FVector3f TempValue(InValue.X, InValue.Y, 0.f);
-	return FRigControlValue::Make<FVector3f>(TempValue);
-}
-
-template<>
-FVector FRigControlElementDetails_ExtractValue(const FRigControlValue& InValue)
-{
-	return InValue.Get<FVector3f>();
-}
-
-template<>
-FRigControlValue FRigControlElementDetails_PackageValue(const FVector& InValue)
-{
-	return FRigControlValue::Make<FVector3f>(InValue);
-}
-
-template<>
-FRotator FRigControlElementDetails_ExtractValue(const FRigControlValue& InValue)
-{
-	return FRotator::MakeFromEuler(InValue.Get<FVector3f>());
-}
-
-template<>
-FRigControlValue FRigControlElementDetails_PackageValue(const FRotator& InValue)
-{
-	return FRigControlValue::Make<FVector3f>(InValue.Euler());
-}
-
-template<>
-FTransform FRigControlElementDetails_ExtractValue(const FRigControlValue& InValue)
-{
-	return InValue.Get<FRigControlValue::FTransform_Float>().ToTransform();
-}
-
-template<>
-FRigControlValue FRigControlElementDetails_PackageValue(const FTransform& InValue)
-{
-	return FRigControlValue::Make<FRigControlValue::FTransform_Float>(InValue);
-}
+			switch(InComponent)
+			{
+				case ESlateTransformComponent::Location:
+				{
+					CurrentTransform.SetLocation(DefaultTransform.GetLocation());
+					break;
+				}
+				case ESlateTransformComponent::Rotation:
+				{
+					CurrentTransform.SetRotator(DefaultTransform.Rotator());
+					break;
+				}
+				case ESlateTransformComponent::Scale:
+				{
+					CurrentTransform.SetScale3D(DefaultTransform.GetScale3D());
+					break;
+				}
+				default: // whole transform / max component
+				{
+					CurrentTransform = DefaultTransform;
+					break;
+				}
+			}
 
-template<>
-FTransformNoScale FRigControlElementDetails_ExtractValue(const FRigControlValue& InValue)
-{
-	return InValue.Get<FRigControlValue::FTransformNoScale_Float>().ToTransform();
-}
+			SetSingleTransform(Key, CurrentTransform, true, true);
+		}
+	});
 
-template<>
-FRigControlValue FRigControlElementDetails_PackageValue(const FTransformNoScale& InValue)
-{
-	return FRigControlValue::Make<FRigControlValue::FTransformNoScale_Float>(InValue);
+	SAdvancedTransformInputBox<TransformType>::ConstructGroupedTransformRows(
+		CategoryBuilder, 
+		Label, 
+		Tooltip, 
+		TransformWidgetArgs);
 }
 
-template<>
-FEulerTransform FRigControlElementDetails_ExtractValue(const FRigControlValue& InValue)
+void FRigBoneElementDetails::CustomizeDetails(IDetailLayoutBuilder& DetailBuilder)
 {
-	return InValue.Get<FRigControlValue::FEulerTransform_Float>().ToTransform();
+	FRigTransformElementDetails::CustomizeDetails(DetailBuilder);
+	CustomizeTransform(DetailBuilder);
 }
 
-template<>
-FRigControlValue FRigControlElementDetails_PackageValue(const FEulerTransform& InValue)
-{
-	return FRigControlValue::Make<FRigControlValue::FEulerTransform_Float>(InValue);
-}
+TArray<TSharedPtr<FString>> FRigControlElementDetails::ControlTypeList;
 
-template<typename T>
-void FRigControlElementDetails_SetupStructValueWidget(IDetailCategoryBuilder& InCategory, const TSharedRef<IPropertyUtilities>& InPropertyUtilities, ERigControlValueType InValueType, const FRigElementKey& InKey, URigHierarchy* InHierarchy)
+void FRigControlElementDetails::CustomizeDetails(IDetailLayoutBuilder& DetailBuilder)
 {
-	UEnum* ControlTypeEnum = StaticEnum<ERigControlType>();
-	UEnum* ValueTypeEnum = StaticEnum<ERigControlValueType>();
-
-	const FString ValueTypeName = ValueTypeEnum->GetDisplayNameTextByValue((int64)InValueType).ToString();
-	const FText PropertyLabel = FText::FromString(FString::Printf(TEXT("%s Value"), *ValueTypeName));
-	const UScriptStruct* ValueStruct = TBaseStructure<T>::Get();
-
-	const TSharedPtr<FStructOnScope> StructToDisplay = MakeShareable(new FStructOnScope(ValueStruct));
+	FRigTransformElementDetails::CustomizeDetails(DetailBuilder);
 
-	TWeakObjectPtr<URigHierarchy> HierarchyPtr = InHierarchy;
+	if (HierarchyBeingCustomized == nullptr)
+	{
+		return;
+	}
 
-	const TAttribute<EVisibility> VisibilityAttribute = TAttribute<EVisibility>::Create(TAttribute<EVisibility>::FGetter::CreateLambda([HierarchyPtr, InKey, InValueType, StructToDisplay, ValueStruct]()->EVisibility
-    {
-		if(HierarchyPtr.IsValid())
-		{
-			if(FRigControlElement* ControlElement = HierarchyPtr->Find<FRigControlElement>(InKey))
-			{
-				// update the struct with the current control value
-				uint8* StructMemory = StructToDisplay->GetStructMemory();
-				const FRigControlValue& ControlValue = HierarchyPtr->GetControlValue(ControlElement, InValueType);
-				T ExtractedValue = FRigControlElementDetails_ExtractValue<T>(ControlValue);
-				ValueStruct->CopyScriptStruct(StructToDisplay->GetStructMemory(), &ExtractedValue, 1);
-				
-				if(ControlElement->Settings.IsValueTypeEnabled(InValueType))
-                {
-                    return EVisibility::Visible;
-                }
-            }
-		}
-		return EVisibility::Hidden;
-    }));
-
-	IDetailPropertyRow* Row = InCategory.AddExternalStructure(StructToDisplay.ToSharedRef());
-	Row->DisplayName(PropertyLabel);
-	Row->ShouldAutoExpand(true);
-	Row->Visibility(VisibilityAttribute);
-
-	TSharedPtr<SWidget> NameWidget, ValueWidget;
-	Row->GetDefaultWidgets(NameWidget, ValueWidget);
+	ControlElements.Reset();
+	ObjectPerControl.Reset();;
 	
-	const FSimpleDelegate OnStructContentsChangedDelegate = FSimpleDelegate::CreateLambda([HierarchyPtr, InKey, StructToDisplay, InValueType]()
+	for(TWeakObjectPtr<UDetailsViewWrapperObject> ObjectBeingCustomized : ObjectsBeingCustomized)
 	{
-		if(HierarchyPtr.IsValid())
+		if(ObjectBeingCustomized.IsValid())
 		{
-            const FRigControlValue ControlValue = FRigControlElementDetails_PackageValue(*(T*)StructToDisplay->GetStructMemory());
-			HierarchyPtr->SetControlValue(InKey, ControlValue, InValueType, true, true);
-			if(InValueType == ERigControlValueType::Initial)
+			if(ObjectBeingCustomized->IsChildOf<FRigControlElement>())
 			{
-                if(UControlRigBlueprint* Blueprint = RigElementDetails_GetBlueprintFromHierarchy(HierarchyPtr.Get()))
-                {
-                    Blueprint->Hierarchy->SetControlValue(InKey, ControlValue, InValueType, true);
-                }
-            }
+				ControlElements.Add(ObjectBeingCustomized->GetContent<FRigControlElement>());
+				ObjectPerControl.Add(ObjectBeingCustomized.Get());
+			}
 		}
-	});
+	}
 
-	TSharedPtr<IPropertyHandle> Handle = Row->GetPropertyHandle();
-	Handle->SetOnPropertyValueChanged(OnStructContentsChangedDelegate);
-	Handle->SetOnChildPropertyValueChanged(OnStructContentsChangedDelegate);
+	CustomizeControl(DetailBuilder);
+	CustomizeValue(DetailBuilder);
+	CustomizeTransform(DetailBuilder);
+	CustomizeShape(DetailBuilder);
 }
 
-void FRigControlElementDetails_SetupValueWidget(IDetailCategoryBuilder& InCategory, const TSharedRef<IPropertyUtilities>& InPropertyUtilities, ERigControlValueType InValueType, const FRigElementKey& InKey, URigHierarchy* InHierarchy)
+void FRigControlElementDetails::CustomizeValue(IDetailLayoutBuilder& DetailBuilder)
 {
-	const FRigControlElement* ControlElement = InHierarchy->Find<FRigControlElement>(InKey);
-	if(ControlElement == nullptr)
+	if(ControlElements.Num() == 0)
 	{
 		return;
 	}
-	
-	switch(ControlElement->Settings.ControlType)
+
+	// only show this section if all controls are the same type
+	const ERigControlType ControlType = ControlElements[0].Settings.ControlType;
+	for(const FRigControlElement& ControlElement : ControlElements)
 	{
-		case ERigControlType::Bool:
-		{
-			if((InValueType == ERigControlValueType::Minimum) || (InValueType == ERigControlValueType::Maximum))
-			{
-				return;
-			}
-			FRigControlElementDetails_SetupBoolValueWidget(InCategory, InPropertyUtilities, InValueType, InKey, InHierarchy);
-			break;
-		}
-		case ERigControlType::Integer:
-		{
-			FRigControlElementDetails_SetupIntegerValueWidget(InCategory, InPropertyUtilities, InValueType, InKey, InHierarchy);
-			break;
-		}
-		case ERigControlType::Float:
-		{
-			FRigControlElementDetails_SetupFloatValueWidget(InCategory, InPropertyUtilities, InValueType, InKey, InHierarchy);
-			break;
-		}
-		case ERigControlType::Vector2D:
-		{
-			FRigControlElementDetails_SetupStructValueWidget<FVector2D>(InCategory, InPropertyUtilities, InValueType, InKey, InHierarchy);
-			break;
-		}
-		case ERigControlType::Position:
-		case ERigControlType::Scale:
-		{
-			FRigControlElementDetails_SetupStructValueWidget<FVector>(InCategory, InPropertyUtilities, InValueType, InKey, InHierarchy);
-			break;
-		}
-		case ERigControlType::Rotator:
-		{
-			FRigControlElementDetails_SetupStructValueWidget<FRotator>(InCategory, InPropertyUtilities, InValueType, InKey, InHierarchy);
-			break;
-		}
-		case ERigControlType::TransformNoScale:
-		{
-			FRigControlElementDetails_SetupStructValueWidget<FTransformNoScale>(InCategory, InPropertyUtilities, InValueType, InKey, InHierarchy);
-			break;
-		}
-		case ERigControlType::EulerTransform:
-		{
-			FRigControlElementDetails_SetupStructValueWidget<FEulerTransform>(InCategory, InPropertyUtilities, InValueType, InKey, InHierarchy);
-			break;
-		}
-		case ERigControlType::Transform:
+		if(ControlElement.Settings.ControlType != ControlType)
 		{
-			FRigControlElementDetails_SetupStructValueWidget<FTransform>(InCategory, InPropertyUtilities, InValueType, InKey, InHierarchy);
-			break;
-		}
-		default:
-		{
-			ensure(false);
-			break;
+			return;
 		}
 	}
-}
 
-void FRigBoneElementDetails::CustomizeDetails(IDetailLayoutBuilder& DetailBuilder)
-{
-	FRigTransformElementDetails::CustomizeDetails(DetailBuilder);
-	CustomizeTransform(DetailBuilder);
-}
+	// transforms don't show their value here - instead they are shown in the transform section
+	if(ControlType == ERigControlType::EulerTransform ||
+		ControlType == ERigControlType::Transform ||
+		ControlType == ERigControlType::TransformNoScale)
+	{
+		return;
+	}
+	
+	TArray<FText> Labels = {
+		LOCTEXT("Initial", "Initial"),
+		LOCTEXT("Current", "Current")
+	};
+	TArray<FText> Tooltips = {
+		LOCTEXT("ValueInitialTooltip", "The initial animation value of the control"),
+		LOCTEXT("ValueCurrentTooltip", "The current animation value of the control")
+	};
+	TArray<ERigControlValueType> ValueTypes = {
+		ERigControlValueType::Initial,
+		ERigControlValueType::Current
+	};
 
-TArray<TSharedPtr<FString>> FRigControlElementDetails::ControlTypeList;
+	// bool doesn't have limits,
+	// transform types already got filtered out earlier.
+	// integers with enums don't have limits either
+	if(ControlType != ERigControlType::Bool &&
+		(ControlType != ERigControlType::Integer || ControlElements[0].Settings.ControlEnum.IsNull()))
+	{
+		Labels.Append({
+			LOCTEXT("Min", "Min"),
+			LOCTEXT("Max", "Max")
+		});
+		Tooltips.Append({
+			LOCTEXT("ValueMinimumTooltip", "The minimum limit(s) for the control"),
+			LOCTEXT("ValueMaximumTooltip", "The maximum limit(s) for the control")
+		});
+		ValueTypes.Append({
+			ERigControlValueType::Minimum,
+			ERigControlValueType::Maximum
+		});
+	}
+	
+	IDetailCategoryBuilder& ValueCategory = DetailBuilder.EditCategory(TEXT("Value"), LOCTEXT("Value", "Value"));
 
-void FRigControlElementDetails::CustomizeDetails(IDetailLayoutBuilder& DetailBuilder)
-{
-	FRigTransformElementDetails::CustomizeDetails(DetailBuilder);
+	if(!PickedValueTypes.IsValid())
+	{
+		PickedValueTypes = MakeShareable(new TArray<ERigControlValueType>({ERigControlValueType::Current}));
+	}
 
-	ShapeNameList.Reset();
-	if (BlueprintBeingCustomized)
+	TSharedPtr<SSegmentedControl<ERigControlValueType>> ValueTypeChoiceWidget =
+		SSegmentedControl<ERigControlValueType>::Create(
+			ValueTypes,
+			Labels,
+			Tooltips,
+			*PickedValueTypes.Get(),
+			true,
+			SSegmentedControl<ERigControlValueType>::FOnValuesChanged::CreateLambda(
+				[](TArray<ERigControlValueType> NewSelection)
+				{
+					(*FRigControlElementDetails::PickedValueTypes.Get()) = NewSelection;
+				}
+			)
+		);
+
+	AddChoiceWidgetRow(ValueCategory, FText::FromString(TEXT("ValueType")), ValueTypeChoiceWidget.ToSharedRef());
+
+	using TransformType = FEulerTransform;
+	using NumericType = FVector::FReal;
+
+	TSharedRef<UE::Math::TVector<float>> IsComponentRelative = MakeShareable(new UE::Math::TVector<float>(1.f, 1.f, 1.f));
+
+	for(int32 Index=0; Index < ValueTypes.Num(); Index++)
 	{
-		bool bUseNameSpace = BlueprintBeingCustomized->ShapeLibraries.Num() > 1;
-		for(TSoftObjectPtr<UControlRigShapeLibrary>& ShapeLibrary : BlueprintBeingCustomized->ShapeLibraries)
+		const ERigControlValueType ValueType = ValueTypes[Index];
+		
+		const TAttribute<EVisibility> VisibilityAttribute =
+			TAttribute<EVisibility>::CreateLambda([ValueType, ValueTypeChoiceWidget]()-> EVisibility
+			{
+				return ValueTypeChoiceWidget->HasValue(ValueType) ? EVisibility::Visible : EVisibility::Collapsed; 
+			});
+		
+		switch(ControlType)
 		{
-			if (!ShapeLibrary.IsValid())
+			case ERigControlType::Bool:
 			{
-				ShapeLibrary.LoadSynchronous();
+				CreateBoolValueWidgetRow(ValueCategory, Labels[Index], Tooltips[Index], ValueType, VisibilityAttribute);
+				break;
 			}
-			if (ShapeLibrary.IsValid())
+			case ERigControlType::Float:
 			{
-				const FString NameSpace = bUseNameSpace ? ShapeLibrary->GetName() + TEXT(".") : FString();
-				ShapeNameList.Add(MakeShared<FString>(NameSpace + ShapeLibrary->DefaultShape.ShapeName.ToString()));
-				for (const FControlRigShapeDefinition& Shape : ShapeLibrary->Shapes)
+				CreateFloatValueWidgetRow(ValueCategory, Labels[Index], Tooltips[Index], ValueType, VisibilityAttribute);
+				break;
+			}
+			case ERigControlType::Integer:
+			{
+				bool bIsEnum = false;
+				const TArray<FRigElementKey> Keys = GetElementKeys();
+				for(const FRigElementKey& Key : Keys)
 				{
-					ShapeNameList.Add(MakeShared<FString>(NameSpace + Shape.ShapeName.ToString()));
+					if(const FRigControlElement* ControlElement = HierarchyBeingCustomized->Find<FRigControlElement>(Key))
+					{
+						if(!ControlElement->Settings.ControlEnum.IsNull())
+						{
+							bIsEnum = true;
+							break;
+						}
+					}
+				}
+
+				if(bIsEnum)
+				{
+					CreateEnumValueWidgetRow(ValueCategory, Labels[Index], Tooltips[Index], ValueType, VisibilityAttribute);
+				}
+				else
+				{
+					CreateIntegerValueWidgetRow(ValueCategory, Labels[Index], Tooltips[Index], ValueType, VisibilityAttribute);
+				}
+				break;
+			}
+			case ERigControlType::Vector2D:
+			{
+				CreateVector2DValueWidgetRow(ValueCategory, Labels[Index], Tooltips[Index], ValueType, VisibilityAttribute);
+				break;
+			}
+			case ERigControlType::Position:
+			case ERigControlType::Rotator:
+			case ERigControlType::Scale:
+			{
+				ERigTransformElementDetailsTransform::Type CurrentTransformType = ERigTransformElementDetailsTransform::Current;
+				switch(ValueType)
+				{
+					case ERigControlValueType::Initial:
+					{
+						CurrentTransformType = ERigTransformElementDetailsTransform::Initial;
+						break;
+					}
+					case ERigControlValueType::Minimum:
+					{
+						CurrentTransformType = ERigTransformElementDetailsTransform::Minimum;
+						break;
+					}
+					case ERigControlValueType::Maximum:
+					{
+						CurrentTransformType = ERigTransformElementDetailsTransform::Maximum;
+						break;
+					}
 				}
+
+				SAdvancedTransformInputBox<TransformType>::FArguments TransformWidgetArgs = SAdvancedTransformInputBox<TransformType>::FArguments()
+				.DisplayToggle(false)
+				.DisplayRelativeWorld(true)
+				.Font(IDetailLayoutBuilder::GetDetailFont())
+				.AllowEditRotationRepresentation(false)
+				.OnGetIsComponentRelative_Lambda(
+					[IsComponentRelative](ESlateTransformComponent::Type InComponent)
+					{
+						return IsComponentRelative->operator[]((int32)InComponent) > 0.f;
+					})
+				.OnIsComponentRelativeChanged_Lambda(
+					[IsComponentRelative](ESlateTransformComponent::Type InComponent, bool bIsRelative)
+					{
+						IsComponentRelative->operator[]((int32)InComponent) = bIsRelative ? 1.f : 0.f;
+					});
+
+				TransformWidgetArgs.Visibility(VisibilityAttribute);
+
+				TransformWidgetArgs.ConstructLocation(ControlType == ERigControlType::Position);
+				TransformWidgetArgs.ConstructRotation(ControlType == ERigControlType::Rotator);
+				TransformWidgetArgs.ConstructScale(ControlType == ERigControlType::Scale);
+
+				CreateEulerTransformValueWidgetRow(
+					GetHierarchyBeingDebugged(),
+					GetElementKeys(),
+					TransformWidgetArgs,
+					ValueCategory,
+					Labels[Index],
+					Tooltips[Index],
+					CurrentTransformType,
+					ValueType);
+				break;
+			}
+			default:
+			{
+				break;
 			}
 		}
 	}
+}
 
-	if (HierarchyBeingCustomized == nullptr)
-	{
-		return;
-	}
+void FRigControlElementDetails::CustomizeControl(IDetailLayoutBuilder& DetailBuilder)
+{
+	const TSharedPtr<IPropertyHandle> SettingsHandle = DetailBuilder.GetProperty(TEXT("Settings"));
+	DetailBuilder.HideProperty(SettingsHandle);
 
 	IDetailCategoryBuilder& ControlCategory = DetailBuilder.EditCategory(TEXT("Control"), LOCTEXT("Control", "Control"));
-	IDetailCategoryBuilder& ShapeCategory = DetailBuilder.EditCategory(TEXT("Shape"), LOCTEXT("Shape", "Shape"));
-	IDetailCategoryBuilder& LimitsCategory = DetailBuilder.EditCategory(TEXT("Limits"), LOCTEXT("Limits", "Limits"));
-
-	// the transform category should show up after the control, shape and limits categories
-	CustomizeTransform(DetailBuilder);
 
-	const TSharedPtr<IPropertyHandle> SettingsHandle = DetailBuilder.GetProperty(TEXT("Settings"));
 	const TSharedPtr<IPropertyHandle> DisplayNameHandle = SettingsHandle->GetChildHandle(TEXT("DisplayName"));
-
-	DetailBuilder.HideProperty(SettingsHandle);
-
 	ControlCategory.AddCustomRow(LOCTEXT("DisplayName", "Display Name"))
 	.NameContent()
 	[
@@ -2263,6 +2504,8 @@ void FRigControlElementDetails::CustomizeDetails(IDetailLayoutBuilder& DetailBui
 
 			if (this->HierarchyBeingCustomized && ControlElementsInHierarchy.Num() > 0)
 			{
+				this->HierarchyBeingCustomized->Modify();
+				
 				for(int32 ControlIndex = 0; ControlIndex< ControlElementsInView.Num(); ControlIndex++)
 				{
 					const FRigControlElement& ViewElement = ControlElementsInView[ControlIndex];
@@ -2271,9 +2514,7 @@ void FRigControlElementDetails::CustomizeDetails(IDetailLayoutBuilder& DetailBui
 					FRigControlValue ValueToSet;
 
 					ControlElement->Settings.ControlType = ViewElement.Settings.ControlType;
-					ControlElement->Settings.bLimitTranslation = false;
-					ControlElement->Settings.bLimitRotation = false;
-					ControlElement->Settings.bLimitScale = false;
+					ControlElement->Settings.LimitEnabled.Reset();
 
 					switch (ControlElement->Settings.ControlType)
 					{
@@ -2285,7 +2526,7 @@ void FRigControlElementDetails::CustomizeDetails(IDetailLayoutBuilder& DetailBui
 						case ERigControlType::Float:
 						{
 							ValueToSet = FRigControlValue::Make<float>(0.f);
-							ControlElement->Settings.bLimitTranslation = true;
+							ControlElement->Settings.SetupLimitArrayForType(true);
 							ControlElement->Settings.MinimumValue = FRigControlValue::Make<float>(0.f);
 							ControlElement->Settings.MaximumValue = FRigControlValue::Make<float>(100.f);
 							break;
@@ -2293,7 +2534,7 @@ void FRigControlElementDetails::CustomizeDetails(IDetailLayoutBuilder& DetailBui
 						case ERigControlType::Integer:
 						{
 							ValueToSet = FRigControlValue::Make<int32>(0);
-							ControlElement->Settings.bLimitTranslation = true;
+							ControlElement->Settings.SetupLimitArrayForType(true);
 							ControlElement->Settings.MinimumValue = FRigControlValue::Make<int32>(0);
 							ControlElement->Settings.MaximumValue = FRigControlValue::Make<int32>(100);
 							break;
@@ -2301,7 +2542,7 @@ void FRigControlElementDetails::CustomizeDetails(IDetailLayoutBuilder& DetailBui
 						case ERigControlType::Vector2D:
 						{
 							ValueToSet = FRigControlValue::Make<FVector2D>(FVector2D::ZeroVector);
-							ControlElement->Settings.bLimitTranslation = true;
+							ControlElement->Settings.SetupLimitArrayForType(true);
 							ControlElement->Settings.MinimumValue = FRigControlValue::Make<FVector2D>(FVector2D::ZeroVector);
 							ControlElement->Settings.MaximumValue = FRigControlValue::Make<FVector2D>(FVector2D(100.f, 100.f));
 							break;
@@ -2309,6 +2550,7 @@ void FRigControlElementDetails::CustomizeDetails(IDetailLayoutBuilder& DetailBui
 						case ERigControlType::Position:
 						{
 							ValueToSet = FRigControlValue::Make<FVector>(FVector::ZeroVector);
+							ControlElement->Settings.SetupLimitArrayForType(false);
 							ControlElement->Settings.MinimumValue = FRigControlValue::Make<FVector>(-FVector::OneVector);
 							ControlElement->Settings.MaximumValue = FRigControlValue::Make<FVector>(FVector::OneVector);
 							break;
@@ -2316,6 +2558,7 @@ void FRigControlElementDetails::CustomizeDetails(IDetailLayoutBuilder& DetailBui
 						case ERigControlType::Scale:
 						{
 							ValueToSet = FRigControlValue::Make<FVector>(FVector::OneVector);
+							ControlElement->Settings.SetupLimitArrayForType(false);
 							ControlElement->Settings.MinimumValue = FRigControlValue::Make<FVector>(FVector::ZeroVector);
 							ControlElement->Settings.MaximumValue = FRigControlValue::Make<FVector>(FVector::OneVector);
 							break;
@@ -2323,6 +2566,7 @@ void FRigControlElementDetails::CustomizeDetails(IDetailLayoutBuilder& DetailBui
 						case ERigControlType::Rotator:
 						{
 							ValueToSet = FRigControlValue::Make<FRotator>(FRotator::ZeroRotator);
+							ControlElement->Settings.SetupLimitArrayForType(false, false);
 							ControlElement->Settings.MinimumValue = FRigControlValue::Make<FRotator>(FRotator::ZeroRotator);
 							ControlElement->Settings.MaximumValue = FRigControlValue::Make<FRotator>(FRotator(180.f, 180.f, 180.f));
 							break;
@@ -2330,6 +2574,7 @@ void FRigControlElementDetails::CustomizeDetails(IDetailLayoutBuilder& DetailBui
 						case ERigControlType::Transform:
 						{
 							ValueToSet = FRigControlValue::Make<FTransform>(FTransform::Identity);
+							ControlElement->Settings.SetupLimitArrayForType(false, false, false);
 							ControlElement->Settings.MinimumValue = ValueToSet;
 							ControlElement->Settings.MaximumValue = ValueToSet;
 							break;
@@ -2338,6 +2583,7 @@ void FRigControlElementDetails::CustomizeDetails(IDetailLayoutBuilder& DetailBui
 						{
 							FTransformNoScale Identity = FTransform::Identity;
 							ValueToSet = FRigControlValue::Make<FTransformNoScale>(Identity);
+							ControlElement->Settings.SetupLimitArrayForType(false, false, false);
 							ControlElement->Settings.MinimumValue = ValueToSet;
 							ControlElement->Settings.MaximumValue = ValueToSet;
 							break;
@@ -2346,6 +2592,7 @@ void FRigControlElementDetails::CustomizeDetails(IDetailLayoutBuilder& DetailBui
 						{
 							FEulerTransform Identity = FEulerTransform::Identity;
 							ValueToSet = FRigControlValue::Make<FEulerTransform>(Identity);
+							ControlElement->Settings.SetupLimitArrayForType(false, false, false);
 							ControlElement->Settings.MinimumValue = ValueToSet;
 							ControlElement->Settings.MaximumValue = ValueToSet;
 							break;
@@ -2357,7 +2604,7 @@ void FRigControlElementDetails::CustomizeDetails(IDetailLayoutBuilder& DetailBui
 						}
 					}
 
-					this->HierarchyBeingCustomized->Notify(ERigHierarchyNotification::ControlSettingChanged, ControlElement);
+					this->HierarchyBeingCustomized->SetControlSettings(ControlElement, ControlElement->Settings, true, true, true);
 					this->HierarchyBeingCustomized->SetControlValue(ControlElement, ValueToSet, ERigControlValueType::Initial, true, false, true);
 					this->HierarchyBeingCustomized->SetControlValue(ControlElement, ValueToSet, ERigControlValueType::Current, true, false, true);
 
@@ -2368,239 +2615,551 @@ void FRigControlElementDetails::CustomizeDetails(IDetailLayoutBuilder& DetailBui
 						if(FRigControlElement* OtherControlElement = this->BlueprintBeingCustomized->Hierarchy->Find<FRigControlElement>(ControlElement->GetKey()))
 						{
 							OtherControlElement->Settings = ControlElement->Settings;
-							this->BlueprintBeingCustomized->Hierarchy->Notify(ERigHierarchyNotification::ControlSettingChanged, ControlElement);
-							this->BlueprintBeingCustomized->Hierarchy->SetControlValue(ControlElement, ValueToSet, ERigControlValueType::Initial, true);
-							this->BlueprintBeingCustomized->Hierarchy->SetControlValue(ControlElement, ValueToSet, ERigControlValueType::Current, true);
+							this->BlueprintBeingCustomized->Hierarchy->SetControlSettings(OtherControlElement, OtherControlElement->Settings, true, true, true);
+							this->BlueprintBeingCustomized->Hierarchy->SetControlValue(OtherControlElement, ValueToSet, ERigControlValueType::Initial, true);
+							this->BlueprintBeingCustomized->Hierarchy->SetControlValue(OtherControlElement, ValueToSet, ERigControlValueType::Current, true);
+						}
+					}
+				}
+				
+				PropertyUtilities->ForceRefresh();
+			}
+		}
+	));
+
+	ControlCategory.AddProperty(ControlTypeHandle.ToSharedRef());
+
+	if(!(IsAnyControlNotOfType(ERigControlType::Integer) &&
+		IsAnyControlNotOfType(ERigControlType::Float) &&
+		IsAnyControlNotOfType(ERigControlType::Vector2D)))
+	{
+		const TSharedPtr<IPropertyHandle> PrimaryAxisHandle = SettingsHandle->GetChildHandle(TEXT("PrimaryAxis"));
+		ControlCategory.AddProperty(PrimaryAxisHandle.ToSharedRef()).DisplayName(FText::FromString(TEXT("Primary Axis")));
+	}
+
+	if(IsAnyControlOfType(ERigControlType::Integer))
+	{
+		const TSharedPtr<IPropertyHandle> ControlEnumHandle = SettingsHandle->GetChildHandle(TEXT("ControlEnum"));
+		ControlCategory.AddProperty(ControlEnumHandle.ToSharedRef()).DisplayName(FText::FromString(TEXT("Control Enum")));
+
+		ControlEnumHandle->SetOnPropertyValueChanged(FSimpleDelegate::CreateLambda(
+			[this, PropertyUtilities]()
+			{
+				PropertyUtilities->ForceRefresh();
+
+				if (this->HierarchyBeingCustomized != nullptr)
+				{
+					for(int32 ControlIndex = 0; ControlIndex < ControlElements.Num(); ControlIndex++)
+					{
+						const FRigControlElement ControlInView = ControlElements[ControlIndex];
+						FRigControlElement* ControlBeingCustomized = HierarchyBeingCustomized->Find<FRigControlElement>(ControlInView.GetKey());
+						
+						const UEnum* ControlEnum = ControlBeingCustomized->Settings.ControlEnum;
+						if (ControlEnum != nullptr)
+						{
+							int32 Maximum = (int32)ControlEnum->GetMaxEnumValue() - 1;
+							ControlBeingCustomized->Settings.MinimumValue.Set<int32>(0);
+							ControlBeingCustomized->Settings.MaximumValue.Set<int32>(Maximum);
+							ControlBeingCustomized->Settings.LimitEnabled.Reset();
+							ControlBeingCustomized->Settings.LimitEnabled.Add(true);
+							HierarchyBeingCustomized->SetControlSettings(ControlBeingCustomized, ControlBeingCustomized->Settings, true, true, true);
+
+							FRigControlValue InitialValue = HierarchyBeingCustomized->GetControlValue(ControlBeingCustomized, ERigControlValueType::Initial);
+							FRigControlValue CurrentValue = HierarchyBeingCustomized->GetControlValue(ControlBeingCustomized, ERigControlValueType::Current);
+
+							ControlBeingCustomized->Settings.ApplyLimits(InitialValue);
+							ControlBeingCustomized->Settings.ApplyLimits(CurrentValue);
+							HierarchyBeingCustomized->SetControlValue(ControlBeingCustomized, InitialValue, ERigControlValueType::Initial, false, false, true);
+							HierarchyBeingCustomized->SetControlValue(ControlBeingCustomized, CurrentValue, ERigControlValueType::Current, false, false, true);
+
+							if (UControlRig* DebuggedRig = Cast<UControlRig>(BlueprintBeingCustomized->GetObjectBeingDebugged()))
+							{
+								URigHierarchy* DebuggedHierarchy = DebuggedRig->GetHierarchy();
+								if(FRigControlElement* DebuggedControlElement = DebuggedHierarchy->Find<FRigControlElement>(ControlBeingCustomized->GetKey()))
+								{
+									DebuggedControlElement->Settings.MinimumValue.Set<int32>(0);
+                                    DebuggedControlElement->Settings.MaximumValue.Set<int32>(Maximum);
+									DebuggedHierarchy->SetControlSettings(DebuggedControlElement, DebuggedControlElement->Settings, true, true, true);
+
+                                    DebuggedHierarchy->SetControlValue(DebuggedControlElement, InitialValue, ERigControlValueType::Initial);
+                                    DebuggedHierarchy->SetControlValue(DebuggedControlElement, CurrentValue, ERigControlValueType::Current);
+								}
+							}
 						}
+
+						ObjectPerControl[ControlIndex]->SetContent<FRigControlElement>(*ControlBeingCustomized);
 					}
 				}
-				
-				PropertyUtilities->ForceRefresh();
+			}
+		));
+		
+	}
+
+	ControlCategory.AddProperty(SettingsHandle->GetChildHandle(TEXT("bAnimatable")).ToSharedRef());
+
+	const TSharedPtr<IPropertyHandle> CustomizationHandle = SettingsHandle->GetChildHandle(TEXT("Customization"));
+	const TSharedPtr<IPropertyHandle> AvailableSpacesHandle = CustomizationHandle->GetChildHandle(TEXT("AvailableSpaces"));
+	ControlCategory.AddProperty(AvailableSpacesHandle.ToSharedRef());
+
+	TArray<FRigElementKey> Keys = GetElementKeys();
+	URigHierarchy* HierarchyBeingDebugged = GetHierarchyBeingDebugged();
+
+	const TSharedPtr<IPropertyHandle> DrawLimitsHandle = SettingsHandle->GetChildHandle(TEXT("bDrawLimits"));
+	
+	ControlCategory
+	.AddProperty(DrawLimitsHandle.ToSharedRef()).DisplayName(FText::FromString(TEXT("Draw Limits")))
+	.IsEnabled(TAttribute<bool>::CreateLambda([Keys, HierarchyBeingDebugged]() -> bool
+	{
+		for(const FRigElementKey& Key : Keys)
+		{
+			if(const FRigControlElement* ControlElement = HierarchyBeingDebugged->Find<FRigControlElement>(Key))
+			{
+				if(ControlElement->Settings.LimitEnabled.Contains(FRigControlLimitEnabled(true, true)))
+				{
+					return true;
+				}
+			}
+		}
+		return false;
+	}));
+}
+
+void FRigControlElementDetails::CustomizeShape(IDetailLayoutBuilder& DetailBuilder)
+{
+	// bools don't have shapes
+	if(IsAnyControlOfType(ERigControlType::Bool))
+	{
+		return;
+	}
+
+	TSharedPtr<IPropertyHandle> ShapeHandle = DetailBuilder.GetProperty(TEXT("Shape"));
+	TSharedPtr<IPropertyHandle> InitialHandle = ShapeHandle->GetChildHandle(TEXT("Initial"));
+	TSharedPtr<IPropertyHandle> LocalHandle = InitialHandle->GetChildHandle(TEXT("Local"));
+	ShapeTransformHandle = LocalHandle->GetChildHandle(TEXT("Transform"));
+	
+	ShapeNameList.Reset();
+	if (BlueprintBeingCustomized)
+	{
+		const bool bUseNameSpace = BlueprintBeingCustomized->ShapeLibraries.Num() > 1;
+		for(TSoftObjectPtr<UControlRigShapeLibrary>& ShapeLibrary : BlueprintBeingCustomized->ShapeLibraries)
+		{
+			if (!ShapeLibrary.IsValid())
+			{
+				ShapeLibrary.LoadSynchronous();
+			}
+			if (ShapeLibrary.IsValid())
+			{
+				const FString NameSpace = bUseNameSpace ? ShapeLibrary->GetName() + TEXT(".") : FString();
+				ShapeNameList.Add(MakeShared<FString>(NameSpace + ShapeLibrary->DefaultShape.ShapeName.ToString()));
+				for (const FControlRigShapeDefinition& Shape : ShapeLibrary->Shapes)
+				{
+					ShapeNameList.Add(MakeShared<FString>(NameSpace + Shape.ShapeName.ToString()));
+				}
 			}
 		}
-	));
+	}
 
-	ControlCategory.AddProperty(ControlTypeHandle.ToSharedRef());
-	ControlCategory.AddProperty(SettingsHandle->GetChildHandle(TEXT("bAnimatable")).ToSharedRef());
+	IDetailCategoryBuilder& ShapeCategory = DetailBuilder.EditCategory(TEXT("Shape"), LOCTEXT("Shape", "Shape"));
+
+	const TSharedPtr<IPropertyHandle> SettingsHandle = DetailBuilder.GetProperty(TEXT("Settings"));
+
+	ShapeCategory.AddProperty(SettingsHandle->GetChildHandle(TEXT("bShapeEnabled")).ToSharedRef())
+	.DisplayName(FText::FromString(TEXT("Enabled")));
+	ShapeCategory.AddProperty(SettingsHandle->GetChildHandle(TEXT("bShapeVisible")).ToSharedRef())
+	.DisplayName(FText::FromString(TEXT("Visible")));
+
+	IDetailGroup& ShapePropertiesGroup = ShapeCategory.AddGroup(TEXT("Shape Properties"), LOCTEXT("ShapeProperties", "Shape Properties"));
+	ShapePropertiesGroup.HeaderRow().NameContent()
+	[
+		SNew(STextBlock)
+		.Font(IDetailLayoutBuilder::GetDetailFont())
+		.Text(LOCTEXT("ShapeProperties", "Shape Properties"))
+		.ToolTipText(LOCTEXT("ShapeProperties", "Customize the properties of the shape"))
+	]
+	.CopyAction(FUIAction(
+		FExecuteAction::CreateSP(this, &FRigControlElementDetails::OnCopyShapeProperties)))
+	.PasteAction(FUIAction(
+		FExecuteAction::CreateSP(this, &FRigControlElementDetails::OnPasteShapeProperties)));
+	
+	// setup shape transform
+	using TransformType = FEulerTransform;
+	using NumericType = FVector::FReal;
+	
+	SAdvancedTransformInputBox<TransformType>::FArguments TransformWidgetArgs = SAdvancedTransformInputBox<TransformType>::FArguments()
+	.DisplayToggle(false)
+	.DisplayRelativeWorld(false)
+	.Font(IDetailLayoutBuilder::GetDetailFont());
+
+	TArray<FRigElementKey> Keys = GetElementKeys();
+	URigHierarchy* HierarchyBeingDebugged = GetHierarchyBeingDebugged();
+	Keys = HierarchyBeingDebugged->SortKeys(Keys);
 
-	// any but bool controls show the offset + shape + limits
-	const bool bNeedsShapeProperties = IsAnyControlNotOfType(ERigControlType::Bool);
-	if (bNeedsShapeProperties)
+	auto GetShapeTransform = [HierarchyBeingDebugged](
+		const FRigElementKey& Key
+		) -> TransformType
 	{
-		// setup offset
+		if(FRigControlElement* ControlElement = HierarchyBeingDebugged->Find<FRigControlElement>(Key))
 		{
-			const TSharedPtr<IPropertyHandle> OffsetHandle = DetailBuilder.GetProperty(TEXT("Offset"));
-			const TSharedPtr<IPropertyHandle> InitialHandle = OffsetHandle->GetChildHandle(TEXT("Initial"));
-			const TSharedPtr<IPropertyHandle> LocalHandle = InitialHandle->GetChildHandle(TEXT("Local"));
-			const TSharedPtr<IPropertyHandle> TransformHandle = LocalHandle->GetChildHandle(TEXT("Transform"));
-			ControlCategory.AddProperty(TransformHandle.ToSharedRef()).DisplayName(FText::FromString(TEXT("Offset Transform")));
+			return HierarchyBeingDebugged->GetControlShapeTransform(ControlElement, ERigTransformType::InitialLocal);
 		}
-	}
+		return TransformType::Identity;
+	};
 
-	if(IsAnyControlOfType(ERigControlType::Float) ||
-		IsAnyControlOfType(ERigControlType::Integer) ||
-		IsAnyControlOfType(ERigControlType::Vector2D) ||
-		IsAnyControlOfType(ERigControlType::Position) ||
-		IsAnyControlOfType(ERigControlType::Transform) ||
-		IsAnyControlOfType(ERigControlType::TransformNoScale) ||
-		IsAnyControlOfType(ERigControlType::EulerTransform))
+	auto SetShapeTransform = [this](
+		const FRigElementKey& Key,
+		const TransformType& InTransform,
+		bool bSetupUndo
+		)
 	{
-		const TSharedPtr<IPropertyHandle> LimitHandle = SettingsHandle->GetChildHandle(TEXT("bLimitTranslation"));
-		LimitsCategory.AddProperty(LimitHandle.ToSharedRef()).DisplayName(FText::FromString(TEXT("Limit Translation")));
-	}
+		if(FRigControlElement* ControlElement = HierarchyBeingCustomized->Find<FRigControlElement>(Key))
+		{
+			HierarchyBeingCustomized->SetControlShapeTransform(ControlElement, InTransform, ERigTransformType::InitialLocal, bSetupUndo, true, bSetupUndo);
+		}
+	};
 
-	if(IsAnyControlOfType(ERigControlType::Rotator) ||
-		IsAnyControlOfType(ERigControlType::Transform) ||
-		IsAnyControlOfType(ERigControlType::TransformNoScale) ||
-		IsAnyControlOfType(ERigControlType::EulerTransform))
+	TransformWidgetArgs.OnGetNumericValue_Lambda([Keys, HierarchyBeingDebugged, GetShapeTransform](
+		ESlateTransformComponent::Type Component,
+		ESlateRotationRepresentation::Type Representation,
+		ESlateTransformSubComponent::Type SubComponent) -> TOptional<NumericType>
 	{
-		const TSharedPtr<IPropertyHandle> LimitHandle = SettingsHandle->GetChildHandle(TEXT("bLimitRotation"));
-		LimitsCategory.AddProperty(LimitHandle.ToSharedRef()).DisplayName(FText::FromString(TEXT("Limit Rotation")));
-	}
+		TOptional<NumericType> FirstValue;
 
-	if(IsAnyControlOfType(ERigControlType::Scale) ||
-		IsAnyControlOfType(ERigControlType::Transform) ||
-		IsAnyControlOfType(ERigControlType::EulerTransform))
-	{
-		const TSharedPtr<IPropertyHandle> LimitHandle = SettingsHandle->GetChildHandle(TEXT("bLimitScale"));
-		LimitsCategory.AddProperty(LimitHandle.ToSharedRef()).DisplayName(FText::FromString(TEXT("Limit Scale")));
-	}
+		for(int32 Index = 0; Index < Keys.Num(); Index++)
+		{
+			const FRigElementKey& Key = Keys[Index];
+			TransformType Xfo = GetShapeTransform(Key);
 
-	if(!(IsAnyControlNotOfType(ERigControlType::Integer) &&
-		IsAnyControlNotOfType(ERigControlType::Float) &&
-		IsAnyControlNotOfType(ERigControlType::Vector2D)))
-	{
-		const TSharedPtr<IPropertyHandle> PrimaryAxisHandle = SettingsHandle->GetChildHandle(TEXT("PrimaryAxis"));
-		ControlCategory.AddProperty(PrimaryAxisHandle.ToSharedRef()).DisplayName(FText::FromString(TEXT("Primary Axis")));
-	}
+			TOptional<NumericType> CurrentValue = SAdvancedTransformInputBox<TransformType>::GetNumericValueFromTransform(Xfo, Component, Representation, SubComponent);
+			if(!CurrentValue.IsSet())
+			{
+				return CurrentValue;
+			}
 
-	if (bNeedsShapeProperties)
-	{
-		const TSharedPtr<IPropertyHandle> DrawLimitsHandle = SettingsHandle->GetChildHandle(TEXT("bDrawLimits"));
-		LimitsCategory.AddProperty(DrawLimitsHandle.ToSharedRef()).DisplayName(FText::FromString(TEXT("Draw Limits")));
-	}
+			if(Index == 0)
+			{
+				FirstValue = CurrentValue;
+			}
+			else
+			{
+				if(!FMath::IsNearlyEqual(FirstValue.GetValue(), CurrentValue.GetValue()))
+				{
+					return TOptional<NumericType>();
+				}
+			}
+		}
+		
+		return FirstValue;
+	});
 
-	TArray<FRigControlElement> ControlElements;
-	TArray<UDetailsViewWrapperObject*> ObjectPerControl;
-	for(TWeakObjectPtr<UDetailsViewWrapperObject> ObjectBeingCustomized : ObjectsBeingCustomized)
-	{
-		if(ObjectBeingCustomized.IsValid())
-		{
-			if(ObjectBeingCustomized->IsChildOf<FRigControlElement>())
+	TransformWidgetArgs.OnNumericValueChanged_Lambda(
+	[
+		Keys,
+		this,
+		GetShapeTransform,
+		SetShapeTransform
+	](
+		ESlateTransformComponent::Type Component,
+		ESlateRotationRepresentation::Type Representation,
+		ESlateTransformSubComponent::Type SubComponent,
+		NumericType InNumericValue)
+	{
+		for(const FRigElementKey& Key : Keys)
+		{
+			TransformType Transform = GetShapeTransform(Key);
+			TransformType PreviousTransform = Transform;
+			SAdvancedTransformInputBox<TransformType>::ApplyNumericValueChange(Transform, InNumericValue, Component, Representation, SubComponent);
+
+			if(!FRigControlElementDetails::Equals(Transform, PreviousTransform))
 			{
-				ControlElements.Add(ObjectBeingCustomized->GetContent<FRigControlElement>());
-				ObjectPerControl.Add(ObjectBeingCustomized.Get());
+				if(!SliderTransaction.IsValid())
+				{
+					SliderTransaction = MakeShareable(new FScopedTransaction(NSLOCTEXT("ControlRigElementDetails", "ChangeNumericValue", "Change Numeric Value")));
+					HierarchyBeingCustomized->Modify();
+				}
+				SetShapeTransform(Key, Transform, false);
 			}
 		}
-	}
+	});
 
-	// only setup value widgets if there is only ony control selected
-	if(ControlElements.Num() == 1)
+	TransformWidgetArgs.OnNumericValueCommitted_Lambda(
+	[
+		Keys,
+		this,
+		GetShapeTransform,
+		SetShapeTransform
+	](
+		ESlateTransformComponent::Type Component,
+		ESlateRotationRepresentation::Type Representation,
+		ESlateTransformSubComponent::Type SubComponent,
+		NumericType InNumericValue,
+		ETextCommit::Type InCommitType)
 	{
-		URigHierarchy* HierarchyBeingDebugged = HierarchyBeingCustomized;
-		
-		if (UControlRig* DebuggedRig = Cast<UControlRig>(BlueprintBeingCustomized->GetObjectBeingDebugged()))
 		{
-			if(!DebuggedRig->IsSetupModeEnabled())
+			FScopedTransaction Transaction(LOCTEXT("ChangeNumericValue", "Change Numeric Value"));
+			this->HierarchyBeingCustomized->Modify();
+
+			for(const FRigElementKey& Key : Keys)
 			{
-				HierarchyBeingDebugged = DebuggedRig->GetHierarchy();
+				TransformType Transform = GetShapeTransform(Key);
+				TransformType PreviousTransform = Transform;
+				SAdvancedTransformInputBox<TransformType>::ApplyNumericValueChange(Transform, InNumericValue, Component, Representation, SubComponent);
+				if(!FRigControlElementDetails::Equals(Transform, PreviousTransform))
+				{
+					SetShapeTransform(Key, Transform, true);
+				}
 			}
 		}
-		
-		FRigControlElementDetails_SetupValueWidget(ControlCategory, PropertyUtilities, ERigControlValueType::Current, ControlElements[0].GetKey(), HierarchyBeingDebugged);
+		SliderTransaction.Reset();
+	});
 
-		switch (ControlElements[0].Settings.ControlType)
+	TransformWidgetArgs.OnCopyToClipboard_Lambda([Keys, GetShapeTransform](
+		ESlateTransformComponent::Type InComponent
+		)
+	{
+		if(Keys.Num() == 0)
 		{
-			case ERigControlType::Bool:
-			case ERigControlType::Float:
-			case ERigControlType::Integer:
-			case ERigControlType::Vector2D:
+			return;
+		}
+
+		const FRigElementKey& FirstKey = Keys[0];
+		TransformType Xfo = GetShapeTransform(FirstKey);
+
+		FString Content;
+		switch(InComponent)
+		{
+			case ESlateTransformComponent::Location:
+			{
+				const FVector Data = Xfo.GetLocation();
+				TBaseStructure<FVector>::Get()->ExportText(Content, &Data, &Data, nullptr, PPF_None, nullptr);
+				break;
+			}
+			case ESlateTransformComponent::Rotation:
+			{
+				const FRotator Data = Xfo.Rotator();
+				TBaseStructure<FRotator>::Get()->ExportText(Content, &Data, &Data, nullptr, PPF_None, nullptr);
+				break;
+			}
+			case ESlateTransformComponent::Scale:
 			{
-				FRigControlElementDetails_SetupValueWidget(ControlCategory, PropertyUtilities, ERigControlValueType::Initial,ControlElements[0].GetKey(), HierarchyBeingCustomized);
+				const FVector Data = Xfo.GetScale3D();
+				TBaseStructure<FVector>::Get()->ExportText(Content, &Data, &Data, nullptr, PPF_None, nullptr);
 				break;
 			}
+			case ESlateTransformComponent::Max:
 			default:
 			{
+				TBaseStructure<TransformType>::Get()->ExportText(Content, &Xfo, &Xfo, nullptr, PPF_None, nullptr);
 				break;
 			}
 		}
-		
-		
-		FRigControlElementDetails_SetupValueWidget(LimitsCategory, PropertyUtilities, ERigControlValueType::Minimum, ControlElements[0].GetKey(), HierarchyBeingCustomized);
-		FRigControlElementDetails_SetupValueWidget(LimitsCategory, PropertyUtilities, ERigControlValueType::Maximum, ControlElements[0].GetKey(), HierarchyBeingCustomized);
-	}
 
-	if (bNeedsShapeProperties)
+		if(!Content.IsEmpty())
+		{
+			FPlatformApplicationMisc::ClipboardCopy(*Content);
+		}
+	});
+
+	TransformWidgetArgs.OnPasteFromClipboard_Lambda([Keys, GetShapeTransform, SetShapeTransform, this](
+		ESlateTransformComponent::Type InComponent
+		)
 	{
-		ShapeCategory.AddProperty(SettingsHandle->GetChildHandle(TEXT("bShapeEnabled")).ToSharedRef())
-		.DisplayName(FText::FromString(TEXT("Enabled")));
-		ShapeCategory.AddProperty(SettingsHandle->GetChildHandle(TEXT("bShapeVisible")).ToSharedRef())
-		.DisplayName(FText::FromString(TEXT("Visible")));
+		if(Keys.Num() == 0)
+		{
+			return;
+		}
 
-		IDetailGroup& ShapeProperitesGroup = ShapeCategory.AddGroup(TEXT("Shape Properties"), LOCTEXT("ShapeProperties", "Shape Properties"));
-		ShapeProperitesGroup.HeaderRow().NameContent()
-		[
-			SNew(STextBlock)
-			.Font(IDetailLayoutBuilder::GetDetailFont())
-			.Text(LOCTEXT("ShapeProperties", "Shape Properties"))
-			.ToolTipText(LOCTEXT("ShapeProperties", "Customize the properties of the shape"))
-		]
-		.CopyAction(FUIAction(
-			FExecuteAction::CreateSP(this, &FRigControlElementDetails::OnCopyShapeProperties)))
-		.PasteAction(FUIAction(
-			FExecuteAction::CreateSP(this, &FRigControlElementDetails::OnPasteShapeProperties)));
-		
+		FString Content;
+		FPlatformApplicationMisc::ClipboardPaste(Content);
+
+		if(Content.IsEmpty())
+		{
+			return;
+		}
+
+		FScopedTransaction Transaction(LOCTEXT("PasteTransform", "Paste Transform"));
+		HierarchyBeingCustomized->Modify();
+
+		for(const FRigElementKey& Key : Keys)
 		{
-			// setup shape transform
+			TransformType Xfo = GetShapeTransform(Key);
+
 			{
-				TSharedPtr<IPropertyHandle> ShapeHandle = DetailBuilder.GetProperty(TEXT("Shape"));
-				TSharedPtr<IPropertyHandle> InitialHandle = ShapeHandle->GetChildHandle(TEXT("Initial"));
-				TSharedPtr<IPropertyHandle> LocalHandle = InitialHandle->GetChildHandle(TEXT("Local"));
-				ShapeTransformHandle = LocalHandle->GetChildHandle(TEXT("Transform"));
-				ShapeProperitesGroup.AddPropertyRow(ShapeTransformHandle.ToSharedRef())
-				.IsEnabled(TAttribute<bool>::CreateSP(this, &FRigControlElementDetails::IsShapeEnabled));
-			}
+				class FRigPasteTransformWidgetErrorPipe : public FOutputDevice
+				{
+				public:
 
-			ShapeNameHandle = SettingsHandle->GetChildHandle(TEXT("ShapeName"));
-			ShapeProperitesGroup.AddPropertyRow(ShapeNameHandle.ToSharedRef()).CustomWidget()
-			.NameContent()
-			[
-				SNew(STextBlock)
-				.Text(FText::FromString(TEXT("Shape")))
-				.Font(IDetailLayoutBuilder::GetDetailFont())
-				.IsEnabled(this, &FRigControlElementDetails::IsShapeEnabled)
-			]
-			.ValueContent()
-			[
-				SNew(SControlRigShapeNameList, ControlElements, BlueprintBeingCustomized)
-				.OnGetNameListContent(this, &FRigControlElementDetails::GetShapeNameList)
-				.IsEnabled(this, &FRigControlElementDetails::IsShapeEnabled)
-			];
+					int32 NumErrors;
+
+					FRigPasteTransformWidgetErrorPipe()
+						: FOutputDevice()
+						, NumErrors(0)
+					{
+					}
+
+					virtual void Serialize(const TCHAR* V, ELogVerbosity::Type Verbosity, const class FName& Category) override
+					{
+						UE_LOG(LogControlRig, Error, TEXT("Error Pasting to Widget: %s"), V);
+						NumErrors++;
+					}
+				};
+
+				FRigPasteTransformWidgetErrorPipe ErrorPipe;
+				
+				switch(InComponent)
+				{
+					case ESlateTransformComponent::Location:
+					{
+						FVector Data = Xfo.GetLocation();
+						TBaseStructure<FVector>::Get()->ImportText(*Content, &Data, nullptr, PPF_None, &ErrorPipe, TBaseStructure<FVector>::Get()->GetName(), true);
+						Xfo.SetLocation(Data);
+						break;
+					}
+					case ESlateTransformComponent::Rotation:
+					{
+						FRotator Data = Xfo.Rotator();
+						TBaseStructure<FRotator>::Get()->ImportText(*Content, &Data, nullptr, PPF_None, &ErrorPipe, TBaseStructure<FRotator>::Get()->GetName(), true);
+						Xfo.SetRotator(Data);
+						break;
+					}
+					case ESlateTransformComponent::Scale:
+					{
+						FVector Data = Xfo.GetScale3D();
+						TBaseStructure<FVector>::Get()->ImportText(*Content, &Data, nullptr, PPF_None, &ErrorPipe, TBaseStructure<FVector>::Get()->GetName(), true);
+						Xfo.SetScale3D(Data);
+						break;
+					}
+					case ESlateTransformComponent::Max:
+					default:
+					{
+						TBaseStructure<TransformType>::Get()->ImportText(*Content, &Xfo, nullptr, PPF_None, &ErrorPipe, TBaseStructure<TransformType>::Get()->GetName(), true);
+						break;
+					}
+				}
 
-			ShapeColorHandle = SettingsHandle->GetChildHandle(TEXT("ShapeColor"));
-			ShapeProperitesGroup.AddPropertyRow(ShapeColorHandle.ToSharedRef())
-			.DisplayName(FText::FromString(TEXT("Color")));
+				if(ErrorPipe.NumErrors == 0)
+				{
+					SetShapeTransform(Key, Xfo, true);
+				}
+			}
 		}
-	}
+	});
 
-	if(IsAnyControlOfType(ERigControlType::Integer))
+	TransformWidgetArgs.DiffersFromDefault_Lambda([
+		Keys,
+		GetShapeTransform
+	](
+		ESlateTransformComponent::Type InComponent) -> bool
 	{
-		const TSharedPtr<IPropertyHandle> ControlEnumHandle = SettingsHandle->GetChildHandle(TEXT("ControlEnum"));
-		ControlCategory.AddProperty(ControlEnumHandle.ToSharedRef()).DisplayName(FText::FromString(TEXT("Control Enum")));
+		for(const FRigElementKey& Key : Keys)
+		{
+			const TransformType CurrentTransform = GetShapeTransform(Key);
+			const TransformType DefaultTransform = TransformType::Identity;
 
-		ControlEnumHandle->SetOnPropertyValueChanged(FSimpleDelegate::CreateLambda(
-			[this, PropertyUtilities, ControlElements, ObjectPerControl]()
+			switch(InComponent)
 			{
-				PropertyUtilities->ForceRefresh();
-
-				if (this->HierarchyBeingCustomized != nullptr)
+				case ESlateTransformComponent::Location:
 				{
-					for(int32 ControlIndex = 0; ControlIndex < ControlElements.Num(); ControlIndex++)
+					if(!DefaultTransform.GetLocation().Equals(CurrentTransform.GetLocation()))
 					{
-						const FRigControlElement ControlInView = ControlElements[ControlIndex];
-						FRigControlElement* ControlBeingCustomized = HierarchyBeingCustomized->Find<FRigControlElement>(ControlInView.GetKey());
-						
-						const UEnum* ControlEnum = ControlBeingCustomized->Settings.ControlEnum;
-						if (ControlEnum != nullptr)
-						{
-							int32 Maximum = (int32)ControlEnum->GetMaxEnumValue() - 1;
-							ControlBeingCustomized->Settings.MinimumValue.Set<int32>(0);
-							ControlBeingCustomized->Settings.MaximumValue.Set<int32>(Maximum);
-							HierarchyBeingCustomized->Notify(ERigHierarchyNotification::ControlSettingChanged, ControlBeingCustomized);
+						return true;
+					}
+					break;
+				}
+				case ESlateTransformComponent::Rotation:
+				{
+					if(!DefaultTransform.Rotator().Equals(CurrentTransform.Rotator()))
+					{
+						return true;
+					}
+					break;
+				}
+				case ESlateTransformComponent::Scale:
+				{
+					if(!DefaultTransform.GetScale3D().Equals(CurrentTransform.GetScale3D()))
+					{
+						return true;
+					}
+					break;
+				}
+				default: // also no component whole transform
+				{
+					if(!DefaultTransform.GetLocation().Equals(CurrentTransform.GetLocation()) ||
+						!DefaultTransform.Rotator().Equals(CurrentTransform.Rotator()) ||
+						!DefaultTransform.GetScale3D().Equals(CurrentTransform.GetScale3D()))
+					{
+						return true;
+					}
+					break;
+				}
+			}
+		}
+		return false;
+	});
 
-							FRigControlValue InitialValue = HierarchyBeingCustomized->GetControlValue(ControlBeingCustomized, ERigControlValueType::Initial);
-							FRigControlValue CurrentValue = HierarchyBeingCustomized->GetControlValue(ControlBeingCustomized, ERigControlValueType::Current);
+	TransformWidgetArgs.OnResetToDefault_Lambda([Keys, GetShapeTransform, SetShapeTransform, this](
+		ESlateTransformComponent::Type InComponent)
+	{
+		FScopedTransaction Transaction(LOCTEXT("ResetTransformToDefault", "Reset Transform to Default"));
+		HierarchyBeingCustomized->Modify();
 
-							ControlBeingCustomized->Settings.ApplyLimits(InitialValue);
-							ControlBeingCustomized->Settings.ApplyLimits(CurrentValue);
-							HierarchyBeingCustomized->SetControlValue(ControlBeingCustomized, InitialValue, ERigControlValueType::Initial, false, false, true);
-							HierarchyBeingCustomized->SetControlValue(ControlBeingCustomized, CurrentValue, ERigControlValueType::Current, false, false, true);
+		for(const FRigElementKey& Key : Keys)
+		{
+			TransformType CurrentTransform = GetShapeTransform(Key);
+			const TransformType DefaultTransform = TransformType::Identity; 
 
-							if (UControlRig* DebuggedRig = Cast<UControlRig>(BlueprintBeingCustomized->GetObjectBeingDebugged()))
-							{
-								URigHierarchy* DebuggedHierarchy = DebuggedRig->GetHierarchy();
-								if(FRigControlElement* DebuggedControlElement = DebuggedHierarchy->Find<FRigControlElement>(ControlBeingCustomized->GetKey()))
-								{
-									DebuggedControlElement->Settings.MinimumValue.Set<int32>(0);
-                                    DebuggedControlElement->Settings.MaximumValue.Set<int32>(Maximum);
-                                    DebuggedHierarchy->Notify(ERigHierarchyNotification::ControlSettingChanged, DebuggedControlElement);
+			switch(InComponent)
+			{
+				case ESlateTransformComponent::Location:
+				{
+					CurrentTransform.SetLocation(DefaultTransform.GetLocation());
+					break;
+				}
+				case ESlateTransformComponent::Rotation:
+				{
+					CurrentTransform.SetRotator(DefaultTransform.Rotator());
+					break;
+				}
+				case ESlateTransformComponent::Scale:
+				{
+					CurrentTransform.SetScale3D(DefaultTransform.GetScale3D());
+					break;
+				}
+				default: // whole transform / max component
+				{
+					CurrentTransform = DefaultTransform;
+					break;
+				}
+			}
 
-                                    DebuggedHierarchy->SetControlValue(DebuggedControlElement, InitialValue, ERigControlValueType::Initial);
-                                    DebuggedHierarchy->SetControlValue(DebuggedControlElement, CurrentValue, ERigControlValueType::Current);
-								}
-							}
-						}
+			SetShapeTransform(Key, CurrentTransform, true);
+		}
+	});
+
+	SAdvancedTransformInputBox<TransformType>::ConstructGroupedTransformRows(
+		ShapeCategory, 
+		LOCTEXT("ShapeTransform", "Shape Transform"), 
+		LOCTEXT("ShapeTransformTooltip", "The relative transform of the shape under the control"),
+		TransformWidgetArgs);
+
+	ShapeNameHandle = SettingsHandle->GetChildHandle(TEXT("ShapeName"));
+	ShapePropertiesGroup.AddPropertyRow(ShapeNameHandle.ToSharedRef()).CustomWidget()
+	.NameContent()
+	[
+		SNew(STextBlock)
+		.Text(FText::FromString(TEXT("Shape")))
+		.Font(IDetailLayoutBuilder::GetDetailFont())
+		.IsEnabled(this, &FRigControlElementDetails::IsShapeEnabled)
+	]
+	.ValueContent()
+	[
+		SNew(SControlRigShapeNameList, ControlElements, BlueprintBeingCustomized)
+		.OnGetNameListContent(this, &FRigControlElementDetails::GetShapeNameList)
+		.IsEnabled(this, &FRigControlElementDetails::IsShapeEnabled)
+	];
 
-						ObjectPerControl[ControlIndex]->SetContent<FRigControlElement>(*ControlBeingCustomized);
-					}
-				}
-			}
-		));
-		
-	}
-	
-	IDetailCategoryBuilder& AnimationCategory = DetailBuilder.EditCategory(TEXT("Animation"), LOCTEXT("Animation", "Animation"));
-	AnimationCategory.AddProperty(SettingsHandle->GetChildHandle(TEXT("Customization"))->GetChildHandle(TEXT("AvailableSpaces")).ToSharedRef());
+	ShapeColorHandle = SettingsHandle->GetChildHandle(TEXT("ShapeColor"));
+	ShapePropertiesGroup.AddPropertyRow(ShapeColorHandle.ToSharedRef())
+	.DisplayName(FText::FromString(TEXT("Color")));
 }
 
 void FRigControlElementDetails::RegisterSectionMappings(FPropertyEditorModule& PropertyEditorModule, UClass* InClass)
@@ -2610,15 +3169,11 @@ void FRigControlElementDetails::RegisterSectionMappings(FPropertyEditorModule& P
 	TSharedRef<FPropertySection> ControlSection = PropertyEditorModule.FindOrCreateSection(InClass->GetFName(), "Control", LOCTEXT("Control", "Control"));
 	ControlSection->AddCategory("General");
 	ControlSection->AddCategory("Control");
-	ControlSection->AddCategory("Animation");
+	ControlSection->AddCategory("Value");
 
 	TSharedRef<FPropertySection> ShapeSection = PropertyEditorModule.FindOrCreateSection(InClass->GetFName(), "Shape", LOCTEXT("Shape", "Shape"));
 	ShapeSection->AddCategory("General");
 	ShapeSection->AddCategory("Shape");
-
-	TSharedRef<FPropertySection> LimitsSection = PropertyEditorModule.FindOrCreateSection(InClass->GetFName(), "Limits", LOCTEXT("Limits", "Limits"));
-	LimitsSection->AddCategory("General");
-	LimitsSection->AddCategory("Limits");
 }
 
 bool FRigControlElementDetails::IsShapeEnabled() const
@@ -2640,37 +3195,6 @@ bool FRigControlElementDetails::IsShapeEnabled() const
 	return false;
 }
 
-bool FRigControlElementDetails::IsEnabled(ERigControlValueType InValueType) const
-{
-	for(TWeakObjectPtr<UDetailsViewWrapperObject> ObjectBeingCustomized : ObjectsBeingCustomized)
-	{
-		if(ObjectBeingCustomized.IsValid())
-		{
-			if(ObjectBeingCustomized->IsChildOf<FRigControlElement>())
-			{
-				const FRigControlElement ControlElement = ObjectBeingCustomized->GetContent<FRigControlElement>();
-				
-				switch (InValueType)
-				{
-					case ERigControlValueType::Minimum:
-					case ERigControlValueType::Maximum:
-					{
-						if(ControlElement.Settings.bLimitTranslation || ControlElement.Settings.bLimitRotation || ControlElement.Settings.bLimitScale)
-						{
-							return true;
-						}
-					}
-					default:
-					{
-						break;
-					}
-				}
-			}
-		}
-	}
-	return true;
-}
-
 const TArray<TSharedPtr<FString>>& FRigControlElementDetails::GetShapeNameList() const
 {
 	return ShapeNameList;
@@ -2734,7 +3258,7 @@ void FRigControlElementDetails::SetDisplayName(const FText& InNewText, ETextComm
 
 				if(GetHierarchy())
 				{
-					GetHierarchy()->SetControlSettings(ControlElement.GetKey(), ControlElement.Settings, false, false, true);
+					GetHierarchy()->SetControlSettings(ControlElement.GetKey(), ControlElement.Settings, true, true, true);
 				}
 			}
 		}
@@ -2827,6 +3351,504 @@ void FRigControlElementDetails::OnPasteShapeProperties()
 	}
 }
 
+void FRigControlElementDetails::CreateBoolValueWidgetRow(
+	IDetailCategoryBuilder& CategoryBuilder,
+	const FText& Label,
+	const FText& Tooltip,
+	ERigControlValueType ValueType,
+	TAttribute<EVisibility> Visibility)
+{
+	const bool bCurrent = ValueType == ERigControlValueType::Current;
+	const bool bInitial = ValueType == ERigControlValueType::Initial;
+
+	TArray<FRigElementKey> Keys = GetElementKeys();
+	URigHierarchy* HierarchyBeingDebugged = GetHierarchyBeingDebugged();
+	URigHierarchy* HierarchyToChange = bCurrent ? HierarchyBeingDebugged : HierarchyBeingCustomized;
+	Keys = HierarchyBeingDebugged->SortKeys(Keys);
+
+	const static TCHAR* TrueText = TEXT("True");
+	const static TCHAR* FalseText = TEXT("False");
+
+	CategoryBuilder.AddCustomRow(Label)
+	.Visibility(Visibility)
+	.NameContent()
+	[
+		SNew(STextBlock)
+		.Text(Label)
+		.ToolTipText(Tooltip)
+		.Font(IDetailLayoutBuilder::GetDetailFont())
+	]
+	.ValueContent()
+	[
+		SNew(SCheckBox)
+		.IsChecked_Lambda([ValueType, Keys, HierarchyBeingDebugged]() -> ECheckBoxState
+		{
+			const bool FirstValue = HierarchyBeingDebugged->GetControlValue<bool>(Keys[0], ValueType);
+			for(int32 Index = 1; Index < Keys.Num(); Index++)
+			{
+				const bool SecondValue = HierarchyBeingDebugged->GetControlValue<bool>(Keys[Index], ValueType);
+				if(FirstValue != SecondValue)
+				{
+					return ECheckBoxState::Undetermined;
+				}
+			}
+			return FirstValue ? ECheckBoxState::Checked : ECheckBoxState::Unchecked;
+		})
+		.OnCheckStateChanged_Lambda([ValueType, Keys, HierarchyToChange](ECheckBoxState NewState)
+		{
+			if(NewState == ECheckBoxState::Undetermined)
+			{
+				return;
+			}
+
+			const bool Value = NewState == ECheckBoxState::Checked;
+			FScopedTransaction Transaction(LOCTEXT("ChangeValue", "Change Value"));
+			HierarchyToChange->Modify();
+			for(const FRigElementKey& Key : Keys)
+			{
+				HierarchyToChange->SetControlValue(Key, FRigControlValue::Make<bool>(Value), ValueType, true, true); 
+			}
+		})
+	]
+	.CopyAction(FUIAction(
+	FExecuteAction::CreateLambda([ValueType, Keys, HierarchyBeingDebugged]()
+		{
+			const bool FirstValue = HierarchyBeingDebugged->GetControlValue<bool>(Keys[0], ValueType);
+			FPlatformApplicationMisc::ClipboardCopy(FirstValue ? TrueText : FalseText);
+		}),
+		FCanExecuteAction())
+	)
+	.PasteAction(FUIAction(
+		FExecuteAction::CreateLambda([ValueType, Keys, HierarchyToChange]()
+		{
+			FString Content;
+			FPlatformApplicationMisc::ClipboardPaste(Content);
+
+			const bool Value = FToBoolHelper::FromCStringWide(*Content);
+			FScopedTransaction Transaction(LOCTEXT("ChangeValue", "Change Value"));
+			HierarchyToChange->Modify();
+			for(const FRigElementKey& Key : Keys)
+			{
+				HierarchyToChange->SetControlValue(Key, FRigControlValue::Make<bool>(Value), ValueType, true, true); 
+			}
+		}),
+		FCanExecuteAction())
+	)
+	.OverrideResetToDefault(FResetToDefaultOverride::Create(
+		TAttribute<bool>::CreateLambda([ValueType, Keys, HierarchyBeingDebugged]() -> bool
+		{
+			const bool FirstValue = HierarchyBeingDebugged->GetControlValue<bool>(Keys[0], ValueType);
+			const bool ReferenceValue = ValueType == ERigControlValueType::Initial ? false :
+				HierarchyBeingDebugged->GetControlValue<bool>(Keys[0], ERigControlValueType::Initial);
+
+			return FirstValue != ReferenceValue;
+		}),
+		FSimpleDelegate::CreateLambda([ValueType, Keys, HierarchyToChange]()
+		{
+			FScopedTransaction Transaction(LOCTEXT("ResetValueToDefault", "Reset Value To Default"));
+			HierarchyToChange->Modify();
+			for(const FRigElementKey& Key : Keys)
+			{
+				const bool ReferenceValue = ValueType == ERigControlValueType::Initial ? false :
+					HierarchyToChange->GetControlValue<bool>(Keys[0], ERigControlValueType::Initial);
+				HierarchyToChange->SetControlValue(Key, FRigControlValue::Make<bool>(ReferenceValue), ValueType, true, true); 
+			}
+		})
+	));
+}
+
+void FRigControlElementDetails::CreateFloatValueWidgetRow(
+	IDetailCategoryBuilder& CategoryBuilder,
+	const FText& Label,
+	const FText& Tooltip,
+	ERigControlValueType ValueType,
+	TAttribute<EVisibility> Visibility)
+{
+	CreateNumericValueWidgetRow<float>(CategoryBuilder, Label, Tooltip, ValueType, Visibility);
+}
+
+void FRigControlElementDetails::CreateIntegerValueWidgetRow(
+	IDetailCategoryBuilder& CategoryBuilder,
+	const FText& Label,
+	const FText& Tooltip,
+	ERigControlValueType ValueType,
+	TAttribute<EVisibility> Visibility)
+{
+	CreateNumericValueWidgetRow<int32>(CategoryBuilder, Label, Tooltip, ValueType, Visibility);
+}
+
+void FRigControlElementDetails::CreateEnumValueWidgetRow(IDetailCategoryBuilder& CategoryBuilder, const FText& Label,
+	const FText& Tooltip, ERigControlValueType ValueType, TAttribute<EVisibility> Visibility)
+{
+	const bool bCurrent = ValueType == ERigControlValueType::Current;
+	const bool bInitial = ValueType == ERigControlValueType::Initial;
+	
+	TArray<FRigElementKey> Keys = GetElementKeys();
+	URigHierarchy* HierarchyBeingDebugged = GetHierarchyBeingDebugged();
+	URigHierarchy* HierarchyToChange = bCurrent ? HierarchyBeingDebugged : HierarchyBeingCustomized;
+	Keys = HierarchyBeingDebugged->SortKeys(Keys);
+
+	UEnum* Enum = nullptr;
+	for(const FRigElementKey& Key : Keys)
+	{
+		if(const FRigControlElement* ControlElement = HierarchyBeingCustomized->Find<FRigControlElement>(Key))
+		{
+			Enum = ControlElement->Settings.ControlEnum.Get();
+			if(Enum)
+			{
+				break;
+			}
+		}
+	}
+
+	check(Enum != nullptr);
+
+	CategoryBuilder.AddCustomRow(Label)
+	.Visibility(Visibility)
+	.NameContent()
+	[
+		SNew(STextBlock)
+		.Text(Label)
+		.ToolTipText(Tooltip)
+		.Font(IDetailLayoutBuilder::GetDetailFont())
+	]
+	.ValueContent()
+	[
+		SNew(SEnumComboBox, Enum)
+		.CurrentValue_Lambda([ValueType, Keys, HierarchyBeingDebugged]() -> int32
+		{
+			const int32 FirstValue = HierarchyBeingDebugged->GetControlValue<int32>(Keys[0], ValueType);
+			for(int32 Index = 1; Index < Keys.Num(); Index++)
+			{
+				const int32 SecondValue = HierarchyBeingDebugged->GetControlValue<int32>(Keys[Index], ValueType);
+				if(FirstValue != SecondValue)
+				{
+					return INDEX_NONE;
+				}
+			}
+			return FirstValue;
+		})
+		.OnEnumSelectionChanged_Lambda([ValueType, Keys, HierarchyToChange](int32 NewSelection, ESelectInfo::Type)
+		{
+			if(NewSelection == INDEX_NONE)
+			{
+				return;
+			}
+
+			FScopedTransaction Transaction(LOCTEXT("ChangeValue", "Change Value"));
+			HierarchyToChange->Modify();
+			for(const FRigElementKey& Key : Keys)
+			{
+				HierarchyToChange->SetControlValue(Key, FRigControlValue::Make<int32>(NewSelection), ValueType, true, true); 
+			}
+		})
+		.Font(FEditorStyle::GetFontStyle(TEXT("MenuItem.Font")))
+	]
+	.CopyAction(FUIAction(
+	FExecuteAction::CreateLambda([ValueType, Keys, HierarchyBeingDebugged]()
+		{
+			const int32 FirstValue = HierarchyBeingDebugged->GetControlValue<int32>(Keys[0], ValueType);
+			FPlatformApplicationMisc::ClipboardCopy(*FString::FromInt(FirstValue));
+		}),
+		FCanExecuteAction())
+	)
+	.PasteAction(FUIAction(
+		FExecuteAction::CreateLambda([ValueType, Keys, HierarchyToChange]()
+		{
+			FString Content;
+			FPlatformApplicationMisc::ClipboardPaste(Content);
+			if(!Content.IsNumeric())
+			{
+				return;
+			}
+
+			const int32 Value = FCString::Atoi(*Content);
+			FScopedTransaction Transaction(LOCTEXT("ChangeValue", "Change Value"));
+			HierarchyToChange->Modify();
+
+			for(const FRigElementKey& Key : Keys)
+			{
+				HierarchyToChange->SetControlValue(Key, FRigControlValue::Make<int32>(Value), ValueType, true, true); 
+			}
+		}),
+		FCanExecuteAction())
+	)
+	.OverrideResetToDefault(FResetToDefaultOverride::Create(
+		TAttribute<bool>::CreateLambda([ValueType, Keys, HierarchyBeingDebugged]() -> bool
+		{
+			const int32 FirstValue = HierarchyBeingDebugged->GetControlValue<int32>(Keys[0], ValueType);
+			const int32 ReferenceValue = ValueType == ERigControlValueType::Initial ? 0 :
+				HierarchyBeingDebugged->GetControlValue<int32>(Keys[0], ERigControlValueType::Initial);
+
+			return FirstValue != ReferenceValue;
+		}),
+		FSimpleDelegate::CreateLambda([ValueType, Keys, HierarchyToChange]()
+		{
+			FScopedTransaction Transaction(LOCTEXT("ResetValueToDefault", "Reset Value To Default"));
+			HierarchyToChange->Modify();
+			for(const FRigElementKey& Key : Keys)
+			{
+				const int32 ReferenceValue = ValueType == ERigControlValueType::Initial ? 0 :
+					HierarchyToChange->GetControlValue<int32>(Keys[0], ERigControlValueType::Initial);
+				HierarchyToChange->SetControlValue(Key, FRigControlValue::Make<int32>(ReferenceValue), ValueType, true, true); 
+			}
+		})
+	));
+}
+
+void FRigControlElementDetails::CreateVector2DValueWidgetRow(
+	IDetailCategoryBuilder& CategoryBuilder, 
+	const FText& Label, 
+	const FText& Tooltip, 
+	ERigControlValueType ValueType,
+	TAttribute<EVisibility> Visibility)
+{
+	const bool bCurrent = ValueType == ERigControlValueType::Current;
+	const bool bInitial = ValueType == ERigControlValueType::Initial;
+	const bool bShowToggle = (ValueType == ERigControlValueType::Minimum) || (ValueType == ERigControlValueType::Maximum);
+	
+	TArray<FRigElementKey> Keys = GetElementKeys();
+	URigHierarchy* HierarchyBeingDebugged = GetHierarchyBeingDebugged();
+	URigHierarchy* HierarchyToChange = bCurrent ? HierarchyBeingDebugged : HierarchyBeingCustomized;
+	Keys = HierarchyBeingDebugged->SortKeys(Keys);
+
+	using SNumericVector2DInputBox = SNumericVectorInputBox<float, FVector2f, 2>;
+	TSharedPtr<SNumericVector2DInputBox> VectorInputBox;
+	
+	FDetailWidgetRow& WidgetRow = CategoryBuilder.AddCustomRow(Label);
+	TAttribute<ECheckBoxState> ToggleXChecked, ToggleYChecked;
+	FOnCheckStateChanged OnToggleXChanged, OnToggleYChanged;
+
+	if(bShowToggle)
+	{
+		auto ToggleChecked = [ValueType, Keys, HierarchyBeingDebugged](int32 Index) -> ECheckBoxState
+		{
+			TOptional<bool> FirstValue;
+
+			for(const FRigElementKey& Key : Keys)
+			{
+				if(const FRigControlElement* ControlElement = HierarchyBeingDebugged->Find<FRigControlElement>(Key))
+				{
+					if(ControlElement->Settings.LimitEnabled.Num() == 2)
+					{
+						const bool Value = ControlElement->Settings.LimitEnabled[Index].GetForValueType(ValueType);
+						if(FirstValue.IsSet())
+						{
+							if(FirstValue.GetValue() != Value)
+							{
+								return ECheckBoxState::Undetermined;
+							}
+						}
+						else
+						{
+							FirstValue = Value;
+						}
+					}
+				}
+			}
+
+			if(!ensure(FirstValue.IsSet()))
+			{
+				return ECheckBoxState::Undetermined;
+			}
+
+			return FirstValue.GetValue() ? ECheckBoxState::Checked : ECheckBoxState::Unchecked;
+		};
+		
+		ToggleXChecked = TAttribute<ECheckBoxState>::CreateLambda([ToggleChecked]() -> ECheckBoxState
+		{
+			return ToggleChecked(0);
+		});
+
+		ToggleYChecked = TAttribute<ECheckBoxState>::CreateLambda([ToggleChecked]() -> ECheckBoxState
+		{
+			return ToggleChecked(1);
+		});
+
+		auto OnToggleChanged = [ValueType, Keys, HierarchyToChange](ECheckBoxState InValue, int32 Index)
+		{
+			if(InValue == ECheckBoxState::Undetermined)
+			{
+				return;
+			}
+					
+			FScopedTransaction Transaction(LOCTEXT("ChangeLimitToggle", "Change Limit Toggle"));
+			HierarchyToChange->Modify();
+
+			for(const FRigElementKey& Key : Keys)
+			{
+				if(FRigControlElement* ControlElement = HierarchyToChange->Find<FRigControlElement>(Key))
+				{
+					if(ControlElement->Settings.LimitEnabled.Num() == 2)
+					{
+						ControlElement->Settings.LimitEnabled[Index].SetForValueType(ValueType, InValue == ECheckBoxState::Checked);
+						HierarchyToChange->SetControlSettings(ControlElement, ControlElement->Settings, true, true, true);
+					}
+				}
+			}
+		};
+
+		OnToggleXChanged = FOnCheckStateChanged::CreateLambda([OnToggleChanged](ECheckBoxState InValue)
+		{
+			OnToggleChanged(InValue, 0);
+		});
+
+		OnToggleYChanged = FOnCheckStateChanged::CreateLambda([OnToggleChanged](ECheckBoxState InValue)
+		{
+			OnToggleChanged(InValue, 1);
+		});
+	}
+
+	auto GetValue = [ValueType, Keys, HierarchyBeingDebugged](int32 Component) -> TOptional<float>
+	{
+		const float FirstValue = HierarchyBeingDebugged->GetControlValue<FVector2f>(Keys[0], ValueType).Component(Component);
+		for(int32 Index = 1; Index < Keys.Num(); Index++)
+		{
+			const float SecondValue = HierarchyBeingDebugged->GetControlValue<FVector2f>(Keys[Index], ValueType).Component(Component);
+			if(FirstValue != SecondValue)
+			{
+				return TOptional<float>();
+			}
+		}
+		return FirstValue;
+	};
+
+	auto OnValueChanged = [ValueType, Keys, HierarchyBeingDebugged, HierarchyToChange, this]
+		(TOptional<float> InValue, ETextCommit::Type InCommitType, bool bSetupUndo, int32 Component)
+		{
+			if(!InValue.IsSet())
+			{
+				return;
+			}
+
+			const float Value = InValue.GetValue();
+		
+			for(const FRigElementKey& Key : Keys)
+			{
+				FVector2f Vector = HierarchyBeingDebugged->GetControlValue<FVector2f>(Key, ValueType);
+				if(!FMath::IsNearlyEqual(Vector.Component(Component), Value))
+				{
+					if(!SliderTransaction.IsValid())
+					{
+						SliderTransaction = MakeShareable(new FScopedTransaction(NSLOCTEXT("ControlRigElementDetails", "ChangeValue", "Change Value")));
+						HierarchyToChange->Modify();
+					}
+					Vector.Component(Component) = Value;
+					HierarchyToChange->SetControlValue(Key, FRigControlValue::Make<FVector2f>(Vector), ValueType, bSetupUndo, bSetupUndo);
+				};
+			}
+
+			if(bSetupUndo)
+			{
+				SliderTransaction.Reset();
+			}
+	};
+
+	WidgetRow
+	.Visibility(Visibility)
+	.NameContent()
+	[
+		SNew(STextBlock)
+		.Text(Label)
+		.ToolTipText(Tooltip)
+		.Font(IDetailLayoutBuilder::GetDetailFont())
+	]
+	.ValueContent()
+	[
+		SAssignNew(VectorInputBox, SNumericVector2DInputBox)
+        .Font(FEditorStyle::GetFontStyle(TEXT("MenuItem.Font")))
+        .AllowSpin(ValueType == ERigControlValueType::Current || ValueType == ERigControlValueType::Initial)
+		.SpinDelta(0.01f)
+		.X_Lambda([GetValue]() -> TOptional<float>
+        {
+			return GetValue(0);
+        })
+        .Y_Lambda([GetValue]() -> TOptional<float>
+		{
+			return GetValue(1);
+		})
+		.OnXChanged_Lambda([OnValueChanged](TOptional<float> InValue)
+		{
+			OnValueChanged(InValue, ETextCommit::Default, false, 0);
+		})
+		.OnYChanged_Lambda([OnValueChanged](TOptional<float> InValue)
+		{
+			OnValueChanged(InValue, ETextCommit::Default, false, 1);
+		})
+		.OnXCommitted_Lambda([OnValueChanged](TOptional<float> InValue, ETextCommit::Type InCommitType)
+		{
+			OnValueChanged(InValue, InCommitType, true, 0);
+		})
+		.OnYCommitted_Lambda([OnValueChanged](TOptional<float> InValue, ETextCommit::Type InCommitType)
+		{
+			OnValueChanged(InValue, InCommitType, true, 1);
+		})
+		 .DisplayToggle(bShowToggle)
+		 .ToggleXChecked(ToggleXChecked)
+		 .ToggleYChecked(ToggleYChecked)
+		 .OnToggleXChanged(OnToggleXChanged)
+		 .OnToggleYChanged(OnToggleYChanged)
+	]
+	.CopyAction(FUIAction(
+	FExecuteAction::CreateLambda([ValueType, Keys, HierarchyBeingDebugged]()
+		{
+			const FVector2f Data = HierarchyBeingDebugged->GetControlValue<FVector2f>(Keys[0], ValueType);
+			FString Content = Data.ToString();
+			FPlatformApplicationMisc::ClipboardCopy(*Content);
+		}),
+		FCanExecuteAction())
+	)
+	.PasteAction(FUIAction(
+		FExecuteAction::CreateLambda([ValueType, Keys, HierarchyToChange]()
+		{
+			FString Content;
+			FPlatformApplicationMisc::ClipboardPaste(Content);
+			if(Content.IsEmpty())
+			{
+				return;
+			}
+
+			FVector2f Data = FVector2f::ZeroVector;
+			Data.InitFromString(Content);
+
+			FScopedTransaction Transaction(NSLOCTEXT("ControlRigElementDetails", "ChangeValue", "Change Value"));
+			HierarchyToChange->Modify();
+			
+			for(const FRigElementKey& Key : Keys)
+			{
+				HierarchyToChange->SetControlValue(Key, FRigControlValue::Make<FVector2f>(Data), ValueType, true, true); 
+			}
+		}),
+		FCanExecuteAction())
+	);
+
+	if(ValueType == ERigControlValueType::Current || ValueType == ERigControlValueType::Initial)
+	{
+		WidgetRow.OverrideResetToDefault(FResetToDefaultOverride::Create(
+			TAttribute<bool>::CreateLambda([ValueType, Keys, HierarchyBeingDebugged]() -> bool
+			{
+				const FVector2f FirstValue = HierarchyBeingDebugged->GetControlValue<FVector2f>(Keys[0], ValueType);
+				const FVector2f ReferenceValue = ValueType == ERigControlValueType::Initial ? FVector2f::ZeroVector :
+					HierarchyBeingDebugged->GetControlValue<FVector2f>(Keys[0], ERigControlValueType::Initial);
+
+				return !(FirstValue - ReferenceValue).IsNearlyZero();
+			}),
+			FSimpleDelegate::CreateLambda([ValueType, Keys, HierarchyToChange]()
+			{
+				FScopedTransaction Transaction(LOCTEXT("ResetValueToDefault", "Reset Value To Default"));
+				HierarchyToChange->Modify();
+				
+				for(const FRigElementKey& Key : Keys)
+				{
+					const FVector2f ReferenceValue = ValueType == ERigControlValueType::Initial ? FVector2f::ZeroVector :
+						HierarchyToChange->GetControlValue<FVector2f>(Keys[0], ERigControlValueType::Initial);
+					HierarchyToChange->SetControlValue(Key, FRigControlValue::Make<FVector2f>(ReferenceValue), ValueType, true, true); 
+				}
+			})
+		));
+	}
+}
+
 void FRigNullElementDetails::CustomizeDetails(IDetailLayoutBuilder& DetailBuilder)
 {
 	FRigTransformElementDetails::CustomizeDetails(DetailBuilder);
diff --git a/Engine/Plugins/Animation/ControlRig/Source/ControlRigEditor/Private/ControlRigElementDetails.h b/Engine/Plugins/Animation/ControlRig/Source/ControlRigEditor/Private/ControlRigElementDetails.h
--- a/Engine/Plugins/Animation/ControlRig/Source/ControlRigEditor/Private/ControlRigElementDetails.h
+++ b/Engine/Plugins/Animation/ControlRig/Source/ControlRigEditor/Private/ControlRigElementDetails.h
@@ -22,6 +22,8 @@
 #include "DetailWidgetRow.h"
 #include "IDetailGroup.h"
 #include "SAdvancedTransformInputBox.h"
+#include "HAL/PlatformApplicationMisc.h"
+#include "Internationalization/FastDecimalFormat.h"
 
 class IPropertyHandle;
 
@@ -134,6 +136,7 @@ public:
 	virtual void CustomizeDetails(IDetailLayoutBuilder& DetailBuilder) override;
 
 	URigHierarchy* GetHierarchy() const { return HierarchyBeingCustomized; }
+	URigHierarchy* GetHierarchyBeingDebugged() const;
 	FRigElementKey GetElementKey() const;
 	FText GetName() const;
 	void SetName(const FText& InNewText, ETextCommit::Type InCommitType);
@@ -190,6 +193,19 @@ protected:
 	TArray<TWeakObjectPtr<UDetailsViewWrapperObject>> ObjectsBeingCustomized;
 };
 
+namespace ERigTransformElementDetailsTransform
+{
+	enum Type
+	{
+		Initial,
+		Current,
+		Offset,
+		Minimum,
+		Maximum,
+		Max
+	};
+}
+
 class FRigTransformElementDetails : public FRigBaseElementDetails
 {
 public:
@@ -207,6 +223,26 @@ protected:
 	bool IsCurrentLocalEnabled() const;
 
 	void AddChoiceWidgetRow(IDetailCategoryBuilder& InCategory, const FText& InSearchText, TSharedRef<SWidget> InWidget);
+
+protected:
+
+	void CreateEulerTransformValueWidgetRow(
+		URigHierarchy* HierarchyBeingDebugged,
+		const TArray<FRigElementKey>& Keys,
+		SAdvancedTransformInputBox<FEulerTransform>::FArguments TransformWidgetArgs,
+		IDetailCategoryBuilder& CategoryBuilder, 
+		const FText& Label, 
+		const FText& Tooltip,
+		ERigTransformElementDetailsTransform::Type CurrentTransformType,
+		ERigControlValueType ValueType);
+
+private:
+
+	static TSharedPtr<TArray<ERigTransformElementDetailsTransform::Type>> PickedTransforms;
+
+protected:
+
+	TSharedPtr<FScopedTransaction> SliderTransaction;
 };
 
 class FRigBoneElementDetails : public FRigTransformElementDetails
@@ -235,12 +271,14 @@ public:
 
 	/** IDetailCustomization interface */
 	virtual void CustomizeDetails(IDetailLayoutBuilder& DetailBuilder) override;
+	void CustomizeValue(IDetailLayoutBuilder& DetailBuilder);
+	void CustomizeControl(IDetailLayoutBuilder& DetailBuilder);
+	void CustomizeShape(IDetailLayoutBuilder& DetailBuilder);
 
 	/** FRigBaseElementDetails interface */
 	virtual void RegisterSectionMappings(FPropertyEditorModule& PropertyEditorModule, UClass* InClass) override;
 
 	bool IsShapeEnabled() const;
-	bool IsEnabled(ERigControlValueType InValueType) const;
 
 	const TArray<TSharedPtr<FString>>& GetShapeNameList() const;
 	const TArray<TSharedPtr<FString>>& GetControlTypeList() const;
@@ -251,7 +289,290 @@ public:
 	void OnCopyShapeProperties();
 	void OnPasteShapeProperties();
 
+	void CreateBoolValueWidgetRow(
+		IDetailCategoryBuilder& CategoryBuilder, 
+		const FText& Label, 
+		const FText& Tooltip, 
+		ERigControlValueType ValueType,
+		TAttribute<EVisibility> Visibility = EVisibility::Visible);
+
+	void CreateFloatValueWidgetRow(
+		IDetailCategoryBuilder& CategoryBuilder, 
+		const FText& Label, 
+		const FText& Tooltip, 
+		ERigControlValueType ValueType,
+		TAttribute<EVisibility> Visibility = EVisibility::Visible);
+
+	void CreateIntegerValueWidgetRow(
+		IDetailCategoryBuilder& CategoryBuilder, 
+		const FText& Label, 
+		const FText& Tooltip, 
+		ERigControlValueType ValueType,
+		TAttribute<EVisibility> Visibility = EVisibility::Visible);
+
+	void CreateEnumValueWidgetRow(
+		IDetailCategoryBuilder& CategoryBuilder, 
+		const FText& Label, 
+		const FText& Tooltip, 
+		ERigControlValueType ValueType,
+		TAttribute<EVisibility> Visibility = EVisibility::Visible);
+
+	void CreateVector2DValueWidgetRow(
+		IDetailCategoryBuilder& CategoryBuilder, 
+		const FText& Label, 
+		const FText& Tooltip, 
+		ERigControlValueType ValueType,
+		TAttribute<EVisibility> Visibility = EVisibility::Visible);
+
+
+	// this is a template since we specialize it further down for
+	// a 'nearly equals' implementation for floats and math types.
+	template<typename T>
+	static bool Equals(const T& A, const T& B)
+	{
+		return A == B;
+	}
+
 private:
+
+	template<typename T>
+	void CreateNumericValueWidgetRow(
+		IDetailCategoryBuilder& CategoryBuilder, 
+		const FText& Label, 
+		const FText& Tooltip, 
+		ERigControlValueType ValueType,
+		TAttribute<EVisibility> Visibility = EVisibility::Visible)
+	{
+		const bool bCurrent = ValueType == ERigControlValueType::Current;
+		const bool bInitial = ValueType == ERigControlValueType::Initial;
+		const bool bShowToggle = (ValueType == ERigControlValueType::Minimum) || (ValueType == ERigControlValueType::Maximum);
+		
+		TArray<FRigElementKey> Keys = GetElementKeys();
+		URigHierarchy* HierarchyBeingDebugged = GetHierarchyBeingDebugged();
+		URigHierarchy* HierarchyToChange = bCurrent ? HierarchyBeingDebugged : HierarchyBeingCustomized;
+		Keys = HierarchyBeingDebugged->SortKeys(Keys);
+
+		TSharedPtr<SNumericEntryBox<T>> NumericEntryBox;
+		
+		FDetailWidgetRow& WidgetRow = CategoryBuilder.AddCustomRow(Label);
+		TAttribute<ECheckBoxState> ToggleChecked;
+		FOnCheckStateChanged OnToggleChanged;
+
+		if(bShowToggle)
+		{
+			ToggleChecked = TAttribute<ECheckBoxState>::CreateLambda(
+				[ValueType, Keys, HierarchyBeingDebugged]() -> ECheckBoxState
+				{
+					TOptional<bool> FirstValue;
+
+					for(const FRigElementKey& Key : Keys)
+					{
+						if(const FRigControlElement* ControlElement = HierarchyBeingDebugged->Find<FRigControlElement>(Key))
+						{
+							if(ControlElement->Settings.LimitEnabled.Num() == 1)
+							{
+								const bool Value = ControlElement->Settings.LimitEnabled[0].GetForValueType(ValueType);
+								if(FirstValue.IsSet())
+								{
+									if(FirstValue.GetValue() != Value)
+									{
+										return ECheckBoxState::Undetermined;
+									}
+								}
+								else
+								{
+									FirstValue = Value;
+								}
+							}
+						}
+					}
+
+					if(!ensure(FirstValue.IsSet()))
+					{
+						return ECheckBoxState::Undetermined;
+					}
+
+					return FirstValue.GetValue() ? ECheckBoxState::Checked : ECheckBoxState::Unchecked;
+				});
+
+			OnToggleChanged = FOnCheckStateChanged::CreateLambda(
+				[ValueType, Keys, HierarchyToChange](ECheckBoxState InValue)
+				{
+					if(InValue == ECheckBoxState::Undetermined)
+					{
+						return;
+					}
+
+					FScopedTransaction Transaction(NSLOCTEXT("ControlRigElementDetails", "ChangeLimitToggle", "Change Limit Toggle"));
+					for(const FRigElementKey& Key : Keys)
+					{
+						if(FRigControlElement* ControlElement = HierarchyToChange->Find<FRigControlElement>(Key))
+						{
+							if(ControlElement->Settings.LimitEnabled.Num() == 1)
+							{
+								HierarchyToChange->Modify();
+								ControlElement->Settings.LimitEnabled[0].SetForValueType(ValueType, InValue == ECheckBoxState::Checked);
+								HierarchyToChange->SetControlSettings(ControlElement, ControlElement->Settings, true, true, true);
+							}
+						}
+					}
+				});
+		}
+
+		auto OnValueChanged = [ValueType, Keys, HierarchyToChange, this]
+			(TOptional<T> InValue, ETextCommit::Type InCommitType, bool bSetupUndo)
+			{
+				if(!InValue.IsSet())
+				{
+					return;
+				}
+
+				const T Value = InValue.GetValue();
+				
+				for(const FRigElementKey& Key : Keys)
+				{
+					const T PreviousValue = HierarchyToChange->GetControlValue(Key, ValueType).Get<T>();
+					if(!Equals(PreviousValue, Value))
+					{
+						if(!SliderTransaction.IsValid())
+						{
+							SliderTransaction = MakeShareable(new FScopedTransaction(NSLOCTEXT("ControlRigElementDetails", "ChangeValue", "Change Value")));
+							HierarchyToChange->Modify();
+						}
+						HierarchyToChange->SetControlValue(Key, FRigControlValue::Make<T>(Value), ValueType, bSetupUndo, bSetupUndo);
+					}
+				}
+
+				if(bSetupUndo)
+				{
+					SliderTransaction.Reset();
+				}
+			};
+
+		WidgetRow
+		.Visibility(Visibility)
+		.NameContent()
+		[
+			SNew(STextBlock)
+			.Text(Label)
+			.ToolTipText(Tooltip)
+			.Font(IDetailLayoutBuilder::GetDetailFont())
+		]
+		.ValueContent()
+		[
+			SAssignNew(NumericEntryBox, SNumericEntryBox<T>)
+	        .Font(FEditorStyle::GetFontStyle(TEXT("MenuItem.Font")))
+	        .AllowSpin(ValueType == ERigControlValueType::Current || ValueType == ERigControlValueType::Initial)
+	        .LinearDeltaSensitivity(1)
+			.Delta(0.01f)
+	        .Value_Lambda([ValueType, Keys, HierarchyBeingDebugged]() -> TOptional<T>
+	        {
+		        const T FirstValue = HierarchyBeingDebugged->GetControlValue<T>(Keys[0], ValueType);
+				for(int32 Index = 1; Index < Keys.Num(); Index++)
+				{
+					const T SecondValue = HierarchyBeingDebugged->GetControlValue<T>(Keys[Index], ValueType);
+					if(FirstValue != SecondValue)
+					{
+						return TOptional<T>();
+					}
+				}
+				return FirstValue;
+	        })
+	        .OnValueChanged_Lambda([ValueType, Keys, HierarchyToChange, OnValueChanged](TOptional<T> InValue)
+	        {
+        		OnValueChanged(InValue, ETextCommit::Default, false);
+	        })
+	        .OnValueCommitted_Lambda([ValueType, Keys, HierarchyToChange, OnValueChanged](TOptional<T> InValue, ETextCommit::Type InCommitType)
+	        {
+        		OnValueChanged(InValue, InCommitType, true);
+	        })
+	        .MinSliderValue_Lambda([ValueType, Keys, HierarchyBeingDebugged]() -> TOptional<T>
+			 {
+				 if(ValueType == ERigControlValueType::Current || ValueType == ERigControlValueType::Initial)
+				 {
+			 		return HierarchyBeingDebugged->GetControlValue<T>(Keys[0], ERigControlValueType::Minimum);
+				 }
+				 return TOptional<T>();
+			 })
+			 .MaxSliderValue_Lambda([ValueType, Keys, HierarchyBeingDebugged]() -> TOptional<T>
+			 {
+		 		if(ValueType == ERigControlValueType::Current || ValueType == ERigControlValueType::Initial)
+		 		{
+					 return HierarchyBeingDebugged->GetControlValue<T>(Keys[0], ERigControlValueType::Maximum);
+				 }
+				 return TOptional<T>();
+			 })
+			 .DisplayToggle(bShowToggle)
+			 .ToggleChecked(ToggleChecked)
+			 .OnToggleChanged(OnToggleChanged)
+		]
+		.CopyAction(FUIAction(
+		FExecuteAction::CreateLambda([ValueType, Keys, HierarchyBeingDebugged]()
+			{
+				const T FirstValue = HierarchyBeingDebugged->GetControlValue<T>(Keys[0], ValueType);
+				const FString Content = FastDecimalFormat::NumberToString(
+					FirstValue,
+					FastDecimalFormat::GetCultureAgnosticFormattingRules(),
+					FNumberFormattingOptions());
+				FPlatformApplicationMisc::ClipboardCopy(*Content);
+			}),
+			FCanExecuteAction())
+		)
+		.PasteAction(FUIAction(
+			FExecuteAction::CreateLambda([ValueType, Keys, HierarchyToChange]()
+			{
+				FString Content;
+				FPlatformApplicationMisc::ClipboardPaste(Content);
+				if(Content.IsEmpty())
+				{
+					return;
+				}
+
+				T Value = T(0);
+				if(FastDecimalFormat::StringToNumber(
+					*Content,
+					Content.Len(),
+					FastDecimalFormat::GetCultureAgnosticFormattingRules(),
+					FNumberParsingOptions(),
+					Value))
+				{
+					FScopedTransaction Transaction(NSLOCTEXT("ControlRigElementDetails", "ChangeValue", "Change Value"));
+					for(const FRigElementKey& Key : Keys)
+					{
+						HierarchyToChange->Modify();
+						HierarchyToChange->SetControlValue(Key, FRigControlValue::Make<T>(Value), ValueType, true, true); 
+					}
+				}
+			}),
+			FCanExecuteAction())
+		);
+
+		if(ValueType == ERigControlValueType::Current || ValueType == ERigControlValueType::Initial)
+		{
+			WidgetRow.OverrideResetToDefault(FResetToDefaultOverride::Create(
+				TAttribute<bool>::CreateLambda([ValueType, Keys, HierarchyBeingDebugged]() -> bool
+				{
+					const T FirstValue = HierarchyBeingDebugged->GetControlValue<T>(Keys[0], ValueType);
+					const T ReferenceValue = ValueType == ERigControlValueType::Initial ? T(0) :
+						HierarchyBeingDebugged->GetControlValue<T>(Keys[0], ERigControlValueType::Initial);
+
+					return !FRigControlElementDetails::Equals(FirstValue, ReferenceValue);
+				}),
+				FSimpleDelegate::CreateLambda([ValueType, Keys, HierarchyToChange]()
+				{
+					FScopedTransaction Transaction(NSLOCTEXT("ControlRigElementDetails", "ResetValueToDefault", "Reset Value To Default"));
+					for(const FRigElementKey& Key : Keys)
+					{
+						const T ReferenceValue = ValueType == ERigControlValueType::Initial ? T(0) :
+							HierarchyToChange->GetControlValue<T>(Keys[0], ERigControlValueType::Initial);
+						HierarchyToChange->Modify();
+						HierarchyToChange->SetControlValue(Key, FRigControlValue::Make<T>(ReferenceValue), ValueType, true, true); 
+					}
+				})
+			));
+		}
+	}
+
 	TArray<TSharedPtr<FString>> ShapeNameList;
 	static TArray<TSharedPtr<FString>> ControlTypeList;
 	TSharedPtr<FRigInfluenceEntryModifier> InfluenceModifier;
@@ -260,8 +581,43 @@ private:
 	TSharedPtr<IPropertyHandle> ShapeNameHandle;
 	TSharedPtr<IPropertyHandle> ShapeColorHandle;
 	TSharedPtr<IPropertyHandle> ShapeTransformHandle;
+
+	TArray<FRigControlElement> ControlElements;
+	TArray<UDetailsViewWrapperObject*> ObjectPerControl;
+
+	static TSharedPtr<TArray<ERigControlValueType>> PickedValueTypes;
 };
 
+template<>
+inline bool FRigControlElementDetails::Equals<float>(const float& A, const float& B)
+{
+	return FMath::IsNearlyEqual(A, B);
+}
+
+template<>
+inline bool FRigControlElementDetails::Equals<double>(const double& A, const double& B)
+{
+	return FMath::IsNearlyEqual(A, B);
+}
+
+template<>
+inline bool FRigControlElementDetails::Equals<FVector>(const FVector& A, const FVector& B)
+{
+	return (A - B).IsNearlyZero();
+}
+
+template<>
+inline bool FRigControlElementDetails::Equals<FRotator>(const FRotator& A, const FRotator& B)
+{
+	return (A - B).IsNearlyZero();
+}
+
+template<>
+inline bool FRigControlElementDetails::Equals<FEulerTransform>(const FEulerTransform& A, const FEulerTransform& B)
+{
+	return Equals(A.Location, B.Location) && Equals(A.Rotation, B.Rotation) && Equals(A.Scale, B.Scale);
+}
+
 class FRigNullElementDetails : public FRigTransformElementDetails
 {
 public:
