Control Rig Details Customization LWC fix

diff --git a/Engine/Plugins/Animation/ControlRig/Source/ControlRigEditor/Private/ControlRigElementDetails.cpp b/Engine/Plugins/Animation/ControlRig/Source/ControlRigEditor/Private/ControlRigElementDetails.cpp
--- a/Engine/Plugins/Animation/ControlRig/Source/ControlRigEditor/Private/ControlRigElementDetails.cpp
+++ b/Engine/Plugins/Animation/ControlRig/Source/ControlRigEditor/Private/ControlRigElementDetails.cpp
@@ -3701,10 +3701,10 @@ void FRigControlElementDetails::CreateVector2DValueWidgetRow(
 
 	auto GetValue = [ValueType, Keys, HierarchyBeingDebugged](int32 Component) -> TOptional<float>
 	{
-		const float FirstValue = HierarchyBeingDebugged->GetControlValue<FVector2f>(Keys[0], ValueType).Component(Component);
+		const float FirstValue = HierarchyBeingDebugged->GetControlValue<FVector3f>(Keys[0], ValueType).Component(Component);
 		for(int32 Index = 1; Index < Keys.Num(); Index++)
 		{
-			const float SecondValue = HierarchyBeingDebugged->GetControlValue<FVector2f>(Keys[Index], ValueType).Component(Component);
+			const float SecondValue = HierarchyBeingDebugged->GetControlValue<FVector3f>(Keys[Index], ValueType).Component(Component);
 			if(FirstValue != SecondValue)
 			{
 				return TOptional<float>();
@@ -3725,7 +3725,7 @@ void FRigControlElementDetails::CreateVector2DValueWidgetRow(
 		
 			for(const FRigElementKey& Key : Keys)
 			{
-				FVector2f Vector = HierarchyBeingDebugged->GetControlValue<FVector2f>(Key, ValueType);
+				FVector3f Vector = HierarchyBeingDebugged->GetControlValue<FVector3f>(Key, ValueType);
 				if(!FMath::IsNearlyEqual(Vector.Component(Component), Value))
 				{
 					if(!SliderTransaction.IsValid())
@@ -3734,7 +3734,7 @@ void FRigControlElementDetails::CreateVector2DValueWidgetRow(
 						HierarchyToChange->Modify();
 					}
 					Vector.Component(Component) = Value;
-					HierarchyToChange->SetControlValue(Key, FRigControlValue::Make<FVector2f>(Vector), ValueType, bSetupUndo, bSetupUndo);
+					HierarchyToChange->SetControlValue(Key, FRigControlValue::Make<FVector3f>(Vector), ValueType, bSetupUndo, bSetupUndo);
 				};
 			}
 
@@ -3792,7 +3792,8 @@ void FRigControlElementDetails::CreateVector2DValueWidgetRow(
 	.CopyAction(FUIAction(
 	FExecuteAction::CreateLambda([ValueType, Keys, HierarchyBeingDebugged]()
 		{
-			const FVector2f Data = HierarchyBeingDebugged->GetControlValue<FVector2f>(Keys[0], ValueType);
+			const FVector3f Data3 = HierarchyBeingDebugged->GetControlValue<FVector3f>(Keys[0], ValueType);
+			const FVector2f Data(Data3.X, Data3.Y);
 			FString Content = Data.ToString();
 			FPlatformApplicationMisc::ClipboardCopy(*Content);
 		}),
@@ -3811,12 +3812,14 @@ void FRigControlElementDetails::CreateVector2DValueWidgetRow(
 			FVector2f Data = FVector2f::ZeroVector;
 			Data.InitFromString(Content);
 
+			FVector3f Data3(Data.X, Data.Y, 0);
+
 			FScopedTransaction Transaction(NSLOCTEXT("ControlRigElementDetails", "ChangeValue", "Change Value"));
 			HierarchyToChange->Modify();
 			
 			for(const FRigElementKey& Key : Keys)
 			{
-				HierarchyToChange->SetControlValue(Key, FRigControlValue::Make<FVector2f>(Data), ValueType, true, true); 
+				HierarchyToChange->SetControlValue(Key, FRigControlValue::Make<FVector3f>(Data3), ValueType, true, true); 
 			}
 		}),
 		FCanExecuteAction())
@@ -3827,9 +3830,9 @@ void FRigControlElementDetails::CreateVector2DValueWidgetRow(
 		WidgetRow.OverrideResetToDefault(FResetToDefaultOverride::Create(
 			TAttribute<bool>::CreateLambda([ValueType, Keys, HierarchyBeingDebugged]() -> bool
 			{
-				const FVector2f FirstValue = HierarchyBeingDebugged->GetControlValue<FVector2f>(Keys[0], ValueType);
-				const FVector2f ReferenceValue = ValueType == ERigControlValueType::Initial ? FVector2f::ZeroVector :
-					HierarchyBeingDebugged->GetControlValue<FVector2f>(Keys[0], ERigControlValueType::Initial);
+				const FVector3f FirstValue = HierarchyBeingDebugged->GetControlValue<FVector3f>(Keys[0], ValueType);
+				const FVector3f ReferenceValue = ValueType == ERigControlValueType::Initial ? FVector3f::ZeroVector :
+					HierarchyBeingDebugged->GetControlValue<FVector3f>(Keys[0], ERigControlValueType::Initial);
 
 				return !(FirstValue - ReferenceValue).IsNearlyZero();
 			}),
@@ -3840,9 +3843,9 @@ void FRigControlElementDetails::CreateVector2DValueWidgetRow(
 				
 				for(const FRigElementKey& Key : Keys)
 				{
-					const FVector2f ReferenceValue = ValueType == ERigControlValueType::Initial ? FVector2f::ZeroVector :
-						HierarchyToChange->GetControlValue<FVector2f>(Keys[0], ERigControlValueType::Initial);
-					HierarchyToChange->SetControlValue(Key, FRigControlValue::Make<FVector2f>(ReferenceValue), ValueType, true, true); 
+					const FVector3f ReferenceValue = ValueType == ERigControlValueType::Initial ? FVector3f::ZeroVector :
+						HierarchyToChange->GetControlValue<FVector3f>(Keys[0], ERigControlValueType::Initial);
+					HierarchyToChange->SetControlValue(Key, FRigControlValue::Make<FVector3f>(ReferenceValue), ValueType, true, true); 
 				}
 			})
 		));
