Control Rig: Fix for non-unity builds.

diff --git a/Engine/Plugins/Animation/ControlRig/Source/ControlRigEditor/Private/ControlRigElementDetails.cpp b/Engine/Plugins/Animation/ControlRig/Source/ControlRigEditor/Private/ControlRigElementDetails.cpp
--- a/Engine/Plugins/Animation/ControlRig/Source/ControlRigEditor/Private/ControlRigElementDetails.cpp
+++ b/Engine/Plugins/Animation/ControlRig/Source/ControlRigEditor/Private/ControlRigElementDetails.cpp
@@ -6,10 +6,7 @@
 #include "Widgets/Text/STextBlock.h"
 #include "Widgets/Input/SEditableTextBox.h"
 #include "Widgets/Input/SVectorInputBox.h"
-#include "Widgets/Input/SRotatorInputBox.h"
 #include "Widgets/Input/SCheckBox.h"
-#include "Widgets/Input/SNumericEntryBox.h"
-#include "Widgets/Input/SComboBox.h"
 #include "Widgets/Input/SButton.h"
 #include "Editor/SControlRigGizmoNameList.h"
 #include "ControlRigBlueprint.h"
@@ -21,8 +18,6 @@
 #include "RigVMModel/RigVMNode.h"
 #include "Graph/SControlRigGraphPinVariableBinding.h"
 #include "HAL/PlatformApplicationMisc.h"
-#include "ScopedTransaction.h"
-#include "AnimationCoreLibrary.h"
 
 #define LOCTEXT_NAMESPACE "ControlRigElementDetails"
 
diff --git a/Engine/Plugins/Animation/ControlRig/Source/ControlRigEditor/Private/ControlRigElementDetails.h b/Engine/Plugins/Animation/ControlRig/Source/ControlRigEditor/Private/ControlRigElementDetails.h
--- a/Engine/Plugins/Animation/ControlRig/Source/ControlRigEditor/Private/ControlRigElementDetails.h
+++ b/Engine/Plugins/Animation/ControlRig/Source/ControlRigEditor/Private/ControlRigElementDetails.h
@@ -10,20 +10,18 @@
 #include "ControlRig.h"
 #include "ControlRigBlueprint.h"
 #include "DetailsViewWrapperObject.h"
-#include "Graph/ControlRigGraph.h"
 #include "Graph/SControlRigGraphPinNameListValueWidget.h"
 #include "Styling/SlateTypes.h"
 #include "IPropertyUtilities.h"
 #include "SSearchableComboBox.h"
-#include "DetailsViewWrapperObject.h"
 #include "Widgets/Input/SSegmentedControl.h"
 #include "DetailCategoryBuilder.h"
 #include "DetailLayoutBuilder.h"
 #include "DetailWidgetRow.h"
-#include "IDetailGroup.h"
 #include "SAdvancedTransformInputBox.h"
 #include "HAL/PlatformApplicationMisc.h"
 #include "Internationalization/FastDecimalFormat.h"
+#include "ScopedTransaction.h"
 
 class IPropertyHandle;
 
