Control Rig Details: Move to centralized settings storage for per widget configuration

diff --git a/Engine/Plugins/Animation/ControlRig/Source/ControlRigEditor/Private/ControlRigElementDetails.cpp b/Engine/Plugins/Animation/ControlRig/Source/ControlRigEditor/Private/ControlRigElementDetails.cpp
--- a/Engine/Plugins/Animation/ControlRig/Source/ControlRigEditor/Private/ControlRigElementDetails.cpp
+++ b/Engine/Plugins/Animation/ControlRig/Source/ControlRigEditor/Private/ControlRigElementDetails.cpp
@@ -23,6 +23,40 @@
 
 static const FText ControlRigDetailsMultipleValues = LOCTEXT("MultipleValues", "Multiple Values");
 
+struct FRigElementTransformWidgetSettings
+{
+	FRigElementTransformWidgetSettings()
+	: RotationRepresentation(MakeShareable(new ESlateRotationRepresentation::Type(ESlateRotationRepresentation::Rotator)))
+	, IsComponentRelative(MakeShareable(new UE::Math::TVector<float>(1.f, 1.f, 1.f)))
+	, IsScaleLocked(TSharedPtr<bool>(new bool(false)))
+	{
+	}
+
+	TSharedPtr<ESlateRotationRepresentation::Type> RotationRepresentation;
+	TSharedRef<UE::Math::TVector<float>> IsComponentRelative;
+	TSharedPtr<bool> IsScaleLocked;
+
+	static FRigElementTransformWidgetSettings& FindOrAdd(
+		ERigControlValueType InValueType,
+		ERigTransformElementDetailsTransform::Type InTransformType,
+		const SAdvancedTransformInputBox<FEulerTransform>::FArguments& WidgetArgs)
+	{
+		uint32 Hash = GetTypeHash(WidgetArgs._ConstructLocation);
+		Hash = HashCombine(Hash, GetTypeHash(WidgetArgs._ConstructRotation));
+		Hash = HashCombine(Hash, GetTypeHash(WidgetArgs._ConstructScale));
+		Hash = HashCombine(Hash, GetTypeHash(WidgetArgs._AllowEditRotationRepresentation));
+		Hash = HashCombine(Hash, GetTypeHash(WidgetArgs._DisplayScaleLock));
+		Hash = HashCombine(Hash, GetTypeHash(InValueType));
+		Hash = HashCombine(Hash, GetTypeHash(InTransformType));
+		return sSettings.FindOrAdd(Hash);
+	}
+
+	static TMap<uint32, FRigElementTransformWidgetSettings> sSettings;
+};
+
+TMap<uint32, FRigElementTransformWidgetSettings> FRigElementTransformWidgetSettings::sSettings;
+
+
 namespace FRigElementKeyDetailsDefs
 {
 	// Active foreground pin alpha
@@ -1145,10 +1179,8 @@ void FRigTransformElementDetails::CustomizeTransform(IDetailLayoutBuilder& Detai
 	IDetailCategoryBuilder& TransformCategory = DetailBuilder.EditCategory(TEXT("Transform"), LOCTEXT("Transform", "Transform"));
 	AddChoiceWidgetRow(TransformCategory, FText::FromString(TEXT("TransformType")), TransformChoiceWidget.ToSharedRef());
 
-	using TransformType = FEulerTransform;
-	using NumericType = FVector::FReal;
 
-	SAdvancedTransformInputBox<TransformType>::FArguments TransformWidgetArgs = SAdvancedTransformInputBox<TransformType>::FArguments()
+	SAdvancedTransformInputBox<FEulerTransform>::FArguments TransformWidgetArgs = SAdvancedTransformInputBox<FEulerTransform>::FArguments()
 	.DisplayToggle(false)
 	.DisplayRelativeWorld(true)
 	.Font(IDetailLayoutBuilder::GetDetailFont());
@@ -1241,10 +1273,9 @@ void FRigTransformElementDetails::CreateEulerTransformValueWidgetRow(
 	ERigTransformElementDetailsTransform::Type CurrentTransformType,
 	ERigControlValueType ValueType)
 {
-	using TransformType = FEulerTransform;
-	using NumericType = FVector::FReal;
-
-	TSharedRef<UE::Math::TVector<float>> IsComponentRelativeStorage = MakeShareable(new UE::Math::TVector<float>(1.f, 1.f, 1.f));
+	const FRigElementTransformWidgetSettings& Settings = FRigElementTransformWidgetSettings::FindOrAdd(ValueType, CurrentTransformType, TransformWidgetArgs);
+	TSharedRef<UE::Math::TVector<float>> IsComponentRelativeStorage = Settings.IsComponentRelative;
+	
 	TransformWidgetArgs.OnGetIsComponentRelative_Lambda(
 		[IsComponentRelativeStorage](ESlateTransformComponent::Type InComponent)
 		{
@@ -1256,8 +1287,7 @@ void FRigTransformElementDetails::CreateEulerTransformValueWidgetRow(
 			IsComponentRelativeStorage->operator[]((int32)InComponent) = bIsRelative ? 1.f : 0.f;
 		});
 
-	const TSharedPtr<ESlateRotationRepresentation::Type> RotationRepresentationStorage =
-		MakeShareable(new ESlateRotationRepresentation::Type(ESlateRotationRepresentation::Rotator));
+	const TSharedPtr<ESlateRotationRepresentation::Type> RotationRepresentationStorage = Settings.RotationRepresentation;
 	TransformWidgetArgs.RotationRepresentation(RotationRepresentationStorage);
 	
 	auto IsComponentRelative = [TransformWidgetArgs](int32 Component) -> bool
@@ -1280,6 +1310,8 @@ void FRigTransformElementDetails::CreateEulerTransformValueWidgetRow(
 		}
 	};
 
+	TransformWidgetArgs.IsScaleLocked(Settings.IsScaleLocked);
+
 	switch(CurrentTransformType)
 	{
 		case ERigTransformElementDetailsTransform::Minimum:
@@ -1581,15 +1613,15 @@ void FRigTransformElementDetails::CreateEulerTransformValueWidgetRow(
 	auto GetRelativeAbsoluteTransforms = [CurrentTransformType, Keys, HierarchyBeingDebugged](
 		const FRigElementKey& Key,
 		ERigTransformElementDetailsTransform::Type InTransformType = ERigTransformElementDetailsTransform::Max
-		) -> TPair<TransformType, TransformType>
+		) -> TPair<FEulerTransform, FEulerTransform>
 	{
 		if(InTransformType == ERigTransformElementDetailsTransform::Max)
 		{
 			InTransformType = CurrentTransformType;
 		}
 
-		TransformType RelativeTransform = TransformType::Identity;
-		TransformType AbsoluteTransform = TransformType::Identity;
+		FEulerTransform RelativeTransform = FEulerTransform::Identity;
+		FEulerTransform AbsoluteTransform = FEulerTransform::Identity;
 
 		const bool bInitial = InTransformType == ERigTransformElementDetailsTransform::Initial; 
 		if(bInitial || InTransformType == ERigTransformElementDetailsTransform::Current)
@@ -1696,13 +1728,13 @@ void FRigTransformElementDetails::CreateEulerTransformValueWidgetRow(
 	auto GetCombinedTransform = [IsComponentRelative, GetRelativeAbsoluteTransforms](
 		const FRigElementKey& Key,
 		ERigTransformElementDetailsTransform::Type InTransformType = ERigTransformElementDetailsTransform::Max
-		) -> TransformType
+		) -> FEulerTransform
 	{
-		const TPair<TransformType, TransformType> TransformPair = GetRelativeAbsoluteTransforms(Key, InTransformType);
-		const TransformType RelativeTransform = TransformPair.Key;
-		const TransformType AbsoluteTransform = TransformPair.Value;
+		const TPair<FEulerTransform, FEulerTransform> TransformPair = GetRelativeAbsoluteTransforms(Key, InTransformType);
+		const FEulerTransform RelativeTransform = TransformPair.Key;
+		const FEulerTransform AbsoluteTransform = TransformPair.Value;
 
-		TransformType Xfo;
+		FEulerTransform Xfo;
 		Xfo.SetLocation((IsComponentRelative(0)) ? RelativeTransform.GetLocation() : AbsoluteTransform.GetLocation());
 		Xfo.SetRotation((IsComponentRelative(1)) ? RelativeTransform.GetRotation() : AbsoluteTransform.GetRotation());
 		Xfo.SetScale3D((IsComponentRelative(2)) ? RelativeTransform.GetScale3D() : AbsoluteTransform.GetScale3D());
@@ -1714,17 +1746,17 @@ void FRigTransformElementDetails::CreateEulerTransformValueWidgetRow(
 		const FRigElementKey& Key,
 		bool bIsRelative,
 		ERigTransformElementDetailsTransform::Type InTransformType = ERigTransformElementDetailsTransform::Max
-		) -> TransformType
+		) -> FEulerTransform
 	{
-		const TPair<TransformType, TransformType> TransformPair = GetRelativeAbsoluteTransforms(Key, InTransformType);
-		const TransformType RelativeTransform = TransformPair.Key;
-		const TransformType AbsoluteTransform = TransformPair.Value;
+		const TPair<FEulerTransform, FEulerTransform> TransformPair = GetRelativeAbsoluteTransforms(Key, InTransformType);
+		const FEulerTransform RelativeTransform = TransformPair.Key;
+		const FEulerTransform AbsoluteTransform = TransformPair.Value;
 		return bIsRelative ? RelativeTransform : AbsoluteTransform;
 	};
 
 	auto SetSingleTransform = [CurrentTransformType, GetRelativeAbsoluteTransforms, this, HierarchyBeingDebugged](
 		const FRigElementKey& Key,
-		TransformType InTransform,
+		FEulerTransform InTransform,
 		bool bIsRelative,
 		bool bSetupUndoRedo)
 	{
@@ -1842,16 +1874,16 @@ void FRigTransformElementDetails::CreateEulerTransformValueWidgetRow(
 	TransformWidgetArgs.OnGetNumericValue_Lambda([Keys, GetCombinedTransform](
 		ESlateTransformComponent::Type Component,
 		ESlateRotationRepresentation::Type Representation,
-		ESlateTransformSubComponent::Type SubComponent) -> TOptional<NumericType>
+		ESlateTransformSubComponent::Type SubComponent) -> TOptional<FVector::FReal>
 	{
-		TOptional<NumericType> FirstValue;
+		TOptional<FVector::FReal> FirstValue;
 
 		for(int32 Index = 0; Index < Keys.Num(); Index++)
 		{
 			const FRigElementKey& Key = Keys[Index];
-			TransformType Xfo = GetCombinedTransform(Key);
+			FEulerTransform Xfo = GetCombinedTransform(Key);
 
-			TOptional<NumericType> CurrentValue = SAdvancedTransformInputBox<TransformType>::GetNumericValueFromTransform(Xfo, Component, Representation, SubComponent);
+			TOptional<FVector::FReal> CurrentValue = SAdvancedTransformInputBox<FEulerTransform>::GetNumericValueFromTransform(Xfo, Component, Representation, SubComponent);
 			if(!CurrentValue.IsSet())
 			{
 				return CurrentValue;
@@ -1865,7 +1897,7 @@ void FRigTransformElementDetails::CreateEulerTransformValueWidgetRow(
 			{
 				if(!FMath::IsNearlyEqual(FirstValue.GetValue(), CurrentValue.GetValue()))
 				{
-					return TOptional<NumericType>();
+					return TOptional<FVector::FReal>();
 				}
 			}
 		}
@@ -1885,15 +1917,15 @@ void FRigTransformElementDetails::CreateEulerTransformValueWidgetRow(
 		ESlateTransformComponent::Type Component,
 		ESlateRotationRepresentation::Type Representation,
 		ESlateTransformSubComponent::Type SubComponent,
-		NumericType InNumericValue)
+		FVector::FReal InNumericValue)
 	{
 		const bool bIsRelative = IsComponentRelative((int32)Component);
 
 		for(const FRigElementKey& Key : Keys)
 		{
-			TransformType Transform = GetSingleTransform(Key, bIsRelative);
-			TransformType PreviousTransform = Transform;
-			SAdvancedTransformInputBox<TransformType>::ApplyNumericValueChange(Transform, InNumericValue, Component, Representation, SubComponent);
+			FEulerTransform Transform = GetSingleTransform(Key, bIsRelative);
+			FEulerTransform PreviousTransform = Transform;
+			SAdvancedTransformInputBox<FEulerTransform>::ApplyNumericValueChange(Transform, InNumericValue, Component, Representation, SubComponent);
 
 			if(!FRigControlElementDetails::Equals(Transform, PreviousTransform))
 			{
@@ -1920,7 +1952,7 @@ void FRigTransformElementDetails::CreateEulerTransformValueWidgetRow(
 		ESlateTransformComponent::Type Component,
 		ESlateRotationRepresentation::Type Representation,
 		ESlateTransformSubComponent::Type SubComponent,
-		NumericType InNumericValue,
+		FVector::FReal InNumericValue,
 		ETextCommit::Type InCommitType)
 	{
 		const bool bIsRelative = IsComponentRelative((int32)Component);
@@ -1934,8 +1966,8 @@ void FRigTransformElementDetails::CreateEulerTransformValueWidgetRow(
 			
 			for(const FRigElementKey& Key : Keys)
 			{
-				TransformType Transform = GetSingleTransform(Key, bIsRelative);
-				SAdvancedTransformInputBox<TransformType>::ApplyNumericValueChange(Transform, InNumericValue, Component, Representation, SubComponent);
+				FEulerTransform Transform = GetSingleTransform(Key, bIsRelative);
+				SAdvancedTransformInputBox<FEulerTransform>::ApplyNumericValueChange(Transform, InNumericValue, Component, Representation, SubComponent);
 				SetSingleTransform(Key, Transform, bIsRelative, true);
 			}
 		}
@@ -1957,7 +1989,7 @@ void FRigTransformElementDetails::CreateEulerTransformValueWidgetRow(
 		const bool bIsRelative = IsComponentRelative(0); 
 
 		const FRigElementKey& FirstKey = Keys[0];
-		TransformType Xfo = GetSingleTransform(FirstKey, bIsRelative);
+		FEulerTransform Xfo = GetSingleTransform(FirstKey, bIsRelative);
 
 		FString Content;
 		switch(InComponent)
@@ -1983,7 +2015,7 @@ void FRigTransformElementDetails::CreateEulerTransformValueWidgetRow(
 			case ESlateTransformComponent::Max:
 			default:
 			{
-				TBaseStructure<TransformType>::Get()->ExportText(Content, &Xfo, &Xfo, nullptr, PPF_None, nullptr);
+				TBaseStructure<FEulerTransform>::Get()->ExportText(Content, &Xfo, &Xfo, nullptr, PPF_None, nullptr);
 				break;
 			}
 		}
@@ -2021,8 +2053,7 @@ void FRigTransformElementDetails::CreateEulerTransformValueWidgetRow(
 
 		for(const FRigElementKey& Key : Keys)
 		{
-			TransformType Xfo = GetSingleTransform(Key, bIsRelative);
-
+			FEulerTransform Xfo = GetSingleTransform(Key, bIsRelative);
 			{
 				class FRigPasteTransformWidgetErrorPipe : public FOutputDevice
 				{
@@ -2071,7 +2102,7 @@ void FRigTransformElementDetails::CreateEulerTransformValueWidgetRow(
 					case ESlateTransformComponent::Max:
 					default:
 					{
-						TBaseStructure<TransformType>::Get()->ImportText(*Content, &Xfo, nullptr, PPF_None, &ErrorPipe, TBaseStructure<TransformType>::Get()->GetName(), true);
+						TBaseStructure<FEulerTransform>::Get()->ImportText(*Content, &Xfo, nullptr, PPF_None, &ErrorPipe, TBaseStructure<FEulerTransform>::Get()->GetName(), true);
 						break;
 					}
 				}
@@ -2094,8 +2125,8 @@ void FRigTransformElementDetails::CreateEulerTransformValueWidgetRow(
 	{
 		for(const FRigElementKey& Key : Keys)
 		{
-			const TransformType CurrentTransform = GetSingleTransform(Key, true);
-			TransformType DefaultTransform;
+			const FEulerTransform CurrentTransform = GetSingleTransform(Key, true);
+			FEulerTransform DefaultTransform;
 
 			switch(CurrentTransformType)
 			{
@@ -2106,7 +2137,7 @@ void FRigTransformElementDetails::CreateEulerTransformValueWidgetRow(
 				}
 				default:
 				{
-					DefaultTransform = TransformType::Identity; 
+					DefaultTransform = FEulerTransform::Identity; 
 					break;
 				}
 			}
@@ -2160,8 +2191,8 @@ void FRigTransformElementDetails::CreateEulerTransformValueWidgetRow(
 
 		for(const FRigElementKey& Key : Keys)
 		{
-			TransformType CurrentTransform = GetSingleTransform(Key, true);
-			TransformType DefaultTransform;
+			FEulerTransform CurrentTransform = GetSingleTransform(Key, true);
+			FEulerTransform DefaultTransform;
 
 			switch(CurrentTransformType)
 			{
@@ -2172,7 +2203,7 @@ void FRigTransformElementDetails::CreateEulerTransformValueWidgetRow(
 				}
 				default:
 				{
-					DefaultTransform = TransformType::Identity; 
+					DefaultTransform = FEulerTransform::Identity; 
 					break;
 				}
 			}
@@ -2205,7 +2236,7 @@ void FRigTransformElementDetails::CreateEulerTransformValueWidgetRow(
 		}
 	});
 
-	SAdvancedTransformInputBox<TransformType>::ConstructGroupedTransformRows(
+	SAdvancedTransformInputBox<FEulerTransform>::ConstructGroupedTransformRows(
 		CategoryBuilder, 
 		Label, 
 		Tooltip, 
@@ -2332,8 +2363,6 @@ void FRigControlElementDetails::CustomizeValue(IDetailLayoutBuilder& DetailBuild
 
 	AddChoiceWidgetRow(ValueCategory, FText::FromString(TEXT("ValueType")), ValueTypeChoiceWidget.ToSharedRef());
 
-	using TransformType = FEulerTransform;
-	using NumericType = FVector::FReal;
 
 	TSharedRef<UE::Math::TVector<float>> IsComponentRelative = MakeShareable(new UE::Math::TVector<float>(1.f, 1.f, 1.f));
 
@@ -2414,7 +2443,7 @@ void FRigControlElementDetails::CustomizeValue(IDetailLayoutBuilder& DetailBuild
 					}
 				}
 
-				SAdvancedTransformInputBox<TransformType>::FArguments TransformWidgetArgs = SAdvancedTransformInputBox<TransformType>::FArguments()
+				SAdvancedTransformInputBox<FEulerTransform>::FArguments TransformWidgetArgs = SAdvancedTransformInputBox<FEulerTransform>::FArguments()
 				.DisplayToggle(false)
 				.DisplayRelativeWorld(true)
 				.Font(IDetailLayoutBuilder::GetDetailFont())
@@ -2777,10 +2806,7 @@ void FRigControlElementDetails::CustomizeShape(IDetailLayoutBuilder& DetailBuild
 		FExecuteAction::CreateSP(this, &FRigControlElementDetails::OnPasteShapeProperties)));
 	
 	// setup shape transform
-	using TransformType = FEulerTransform;
-	using NumericType = FVector::FReal;
-	
-	SAdvancedTransformInputBox<TransformType>::FArguments TransformWidgetArgs = SAdvancedTransformInputBox<TransformType>::FArguments()
+	SAdvancedTransformInputBox<FEulerTransform>::FArguments TransformWidgetArgs = SAdvancedTransformInputBox<FEulerTransform>::FArguments()
 	.DisplayToggle(false)
 	.DisplayRelativeWorld(false)
 	.Font(IDetailLayoutBuilder::GetDetailFont());
@@ -2791,18 +2817,18 @@ void FRigControlElementDetails::CustomizeShape(IDetailLayoutBuilder& DetailBuild
 
 	auto GetShapeTransform = [HierarchyBeingDebugged](
 		const FRigElementKey& Key
-		) -> TransformType
+		) -> FEulerTransform
 	{
 		if(FRigControlElement* ControlElement = HierarchyBeingDebugged->Find<FRigControlElement>(Key))
 		{
 			return HierarchyBeingDebugged->GetControlShapeTransform(ControlElement, ERigTransformType::InitialLocal);
 		}
-		return TransformType::Identity;
+		return FEulerTransform::Identity;
 	};
 
 	auto SetShapeTransform = [this](
 		const FRigElementKey& Key,
-		const TransformType& InTransform,
+		const FEulerTransform& InTransform,
 		bool bSetupUndo
 		)
 	{
@@ -2815,16 +2841,16 @@ void FRigControlElementDetails::CustomizeShape(IDetailLayoutBuilder& DetailBuild
 	TransformWidgetArgs.OnGetNumericValue_Lambda([Keys, HierarchyBeingDebugged, GetShapeTransform](
 		ESlateTransformComponent::Type Component,
 		ESlateRotationRepresentation::Type Representation,
-		ESlateTransformSubComponent::Type SubComponent) -> TOptional<NumericType>
+		ESlateTransformSubComponent::Type SubComponent) -> TOptional<FVector::FReal>
 	{
-		TOptional<NumericType> FirstValue;
+		TOptional<FVector::FReal> FirstValue;
 
 		for(int32 Index = 0; Index < Keys.Num(); Index++)
 		{
 			const FRigElementKey& Key = Keys[Index];
-			TransformType Xfo = GetShapeTransform(Key);
+			FEulerTransform Xfo = GetShapeTransform(Key);
 
-			TOptional<NumericType> CurrentValue = SAdvancedTransformInputBox<TransformType>::GetNumericValueFromTransform(Xfo, Component, Representation, SubComponent);
+			TOptional<FVector::FReal> CurrentValue = SAdvancedTransformInputBox<FEulerTransform>::GetNumericValueFromTransform(Xfo, Component, Representation, SubComponent);
 			if(!CurrentValue.IsSet())
 			{
 				return CurrentValue;
@@ -2838,7 +2864,7 @@ void FRigControlElementDetails::CustomizeShape(IDetailLayoutBuilder& DetailBuild
 			{
 				if(!FMath::IsNearlyEqual(FirstValue.GetValue(), CurrentValue.GetValue()))
 				{
-					return TOptional<NumericType>();
+					return TOptional<FVector::FReal>();
 				}
 			}
 		}
@@ -2856,13 +2882,13 @@ void FRigControlElementDetails::CustomizeShape(IDetailLayoutBuilder& DetailBuild
 		ESlateTransformComponent::Type Component,
 		ESlateRotationRepresentation::Type Representation,
 		ESlateTransformSubComponent::Type SubComponent,
-		NumericType InNumericValue)
+		FVector::FReal InNumericValue)
 	{
 		for(const FRigElementKey& Key : Keys)
 		{
-			TransformType Transform = GetShapeTransform(Key);
-			TransformType PreviousTransform = Transform;
-			SAdvancedTransformInputBox<TransformType>::ApplyNumericValueChange(Transform, InNumericValue, Component, Representation, SubComponent);
+			FEulerTransform Transform = GetShapeTransform(Key);
+			FEulerTransform PreviousTransform = Transform;
+			SAdvancedTransformInputBox<FEulerTransform>::ApplyNumericValueChange(Transform, InNumericValue, Component, Representation, SubComponent);
 
 			if(!FRigControlElementDetails::Equals(Transform, PreviousTransform))
 			{
@@ -2886,7 +2912,7 @@ void FRigControlElementDetails::CustomizeShape(IDetailLayoutBuilder& DetailBuild
 		ESlateTransformComponent::Type Component,
 		ESlateRotationRepresentation::Type Representation,
 		ESlateTransformSubComponent::Type SubComponent,
-		NumericType InNumericValue,
+		FVector::FReal InNumericValue,
 		ETextCommit::Type InCommitType)
 	{
 		{
@@ -2895,9 +2921,9 @@ void FRigControlElementDetails::CustomizeShape(IDetailLayoutBuilder& DetailBuild
 
 			for(const FRigElementKey& Key : Keys)
 			{
-				TransformType Transform = GetShapeTransform(Key);
-				TransformType PreviousTransform = Transform;
-				SAdvancedTransformInputBox<TransformType>::ApplyNumericValueChange(Transform, InNumericValue, Component, Representation, SubComponent);
+				FEulerTransform Transform = GetShapeTransform(Key);
+				FEulerTransform PreviousTransform = Transform;
+				SAdvancedTransformInputBox<FEulerTransform>::ApplyNumericValueChange(Transform, InNumericValue, Component, Representation, SubComponent);
 				if(!FRigControlElementDetails::Equals(Transform, PreviousTransform))
 				{
 					SetShapeTransform(Key, Transform, true);
@@ -2917,7 +2943,7 @@ void FRigControlElementDetails::CustomizeShape(IDetailLayoutBuilder& DetailBuild
 		}
 
 		const FRigElementKey& FirstKey = Keys[0];
-		TransformType Xfo = GetShapeTransform(FirstKey);
+		FEulerTransform Xfo = GetShapeTransform(FirstKey);
 
 		FString Content;
 		switch(InComponent)
@@ -2943,7 +2969,7 @@ void FRigControlElementDetails::CustomizeShape(IDetailLayoutBuilder& DetailBuild
 			case ESlateTransformComponent::Max:
 			default:
 			{
-				TBaseStructure<TransformType>::Get()->ExportText(Content, &Xfo, &Xfo, nullptr, PPF_None, nullptr);
+				TBaseStructure<FEulerTransform>::Get()->ExportText(Content, &Xfo, &Xfo, nullptr, PPF_None, nullptr);
 				break;
 			}
 		}
@@ -2976,8 +3002,7 @@ void FRigControlElementDetails::CustomizeShape(IDetailLayoutBuilder& DetailBuild
 
 		for(const FRigElementKey& Key : Keys)
 		{
-			TransformType Xfo = GetShapeTransform(Key);
-
+			FEulerTransform Xfo = GetShapeTransform(Key);
 			{
 				class FRigPasteTransformWidgetErrorPipe : public FOutputDevice
 				{
@@ -3026,7 +3051,7 @@ void FRigControlElementDetails::CustomizeShape(IDetailLayoutBuilder& DetailBuild
 					case ESlateTransformComponent::Max:
 					default:
 					{
-						TBaseStructure<TransformType>::Get()->ImportText(*Content, &Xfo, nullptr, PPF_None, &ErrorPipe, TBaseStructure<TransformType>::Get()->GetName(), true);
+						TBaseStructure<FEulerTransform>::Get()->ImportText(*Content, &Xfo, nullptr, PPF_None, &ErrorPipe, TBaseStructure<FEulerTransform>::Get()->GetName(), true);
 						break;
 					}
 				}
@@ -3047,8 +3072,8 @@ void FRigControlElementDetails::CustomizeShape(IDetailLayoutBuilder& DetailBuild
 	{
 		for(const FRigElementKey& Key : Keys)
 		{
-			const TransformType CurrentTransform = GetShapeTransform(Key);
-			const TransformType DefaultTransform = TransformType::Identity;
+			const FEulerTransform CurrentTransform = GetShapeTransform(Key);
+			static const FEulerTransform DefaultTransform = FEulerTransform::Identity;
 
 			switch(InComponent)
 			{
@@ -3099,8 +3124,8 @@ void FRigControlElementDetails::CustomizeShape(IDetailLayoutBuilder& DetailBuild
 
 		for(const FRigElementKey& Key : Keys)
 		{
-			TransformType CurrentTransform = GetShapeTransform(Key);
-			const TransformType DefaultTransform = TransformType::Identity; 
+			FEulerTransform CurrentTransform = GetShapeTransform(Key);
+			static const FEulerTransform DefaultTransform = FEulerTransform::Identity; 
 
 			switch(InComponent)
 			{
@@ -3130,7 +3155,7 @@ void FRigControlElementDetails::CustomizeShape(IDetailLayoutBuilder& DetailBuild
 		}
 	});
 
-	SAdvancedTransformInputBox<TransformType>::ConstructGroupedTransformRows(
+	SAdvancedTransformInputBox<FEulerTransform>::ConstructGroupedTransformRows(
 		ShapeCategory, 
 		LOCTEXT("ShapeTransform", "Shape Transform"), 
 		LOCTEXT("ShapeTransformTooltip", "The relative transform of the shape under the control"),
