Control Rig: Support undo for edit shape color

diff --git a/Engine/Plugins/Animation/ControlRig/Source/ControlRig/Private/Rigs/RigHierarchy.cpp b/Engine/Plugins/Animation/ControlRig/Source/ControlRig/Private/Rigs/RigHierarchy.cpp
--- a/Engine/Plugins/Animation/ControlRig/Source/ControlRig/Private/Rigs/RigHierarchy.cpp
+++ b/Engine/Plugins/Animation/ControlRig/Source/ControlRig/Private/Rigs/RigHierarchy.cpp
@@ -2603,6 +2603,11 @@ void URigHierarchy::SetControlSettings(FRigControlElement* InControlElement, FRi
 		return;
 	}
 
+	if(bSetupUndo && !HasAnyFlags(RF_Transient))
+	{
+		Modify();
+	}
+
 	InControlElement->Settings = InSettings;
 	Notify(ERigHierarchyNotification::ControlSettingChanged, InControlElement);
 	
