Control Rig: Improve argument layour order changes

diff --git a/Engine/Plugins/Animation/ControlRig/Source/ControlRigEditor/Private/ControlRigGraphDetails.cpp b/Engine/Plugins/Animation/ControlRig/Source/ControlRigEditor/Private/ControlRigGraphDetails.cpp
--- a/Engine/Plugins/Animation/ControlRig/Source/ControlRigEditor/Private/ControlRigGraphDetails.cpp
+++ b/Engine/Plugins/Animation/ControlRig/Source/ControlRigEditor/Private/ControlRigGraphDetails.cpp
@@ -261,7 +261,7 @@ void FControlRigArgumentLayout::GenerateHeaderRowContent(FDetailWidgetRow& NodeR
 				SNew(SButton)
 				.ButtonStyle(FAppStyle::Get(), TEXT("SimpleButton"))
 				.ContentPadding(0)
-				.IsEnabled(!IsPinEditingReadOnly())
+				.IsEnabled_Raw(this, &FControlRigArgumentLayout::CanArgumentBeMoved, true)
 				.OnClicked(this, &FControlRigArgumentLayout::OnArgMoveUp)
 				.ToolTipText(LOCTEXT("FunctionArgDetailsArgMoveUpTooltip", "Move this parameter up in the list."))
 				[
@@ -277,7 +277,7 @@ void FControlRigArgumentLayout::GenerateHeaderRowContent(FDetailWidgetRow& NodeR
 				SNew(SButton)
 				.ButtonStyle(FAppStyle::Get(), TEXT("SimpleButton"))
 				.ContentPadding(0)
-				.IsEnabled(!IsPinEditingReadOnly())
+				.IsEnabled_Raw(this, &FControlRigArgumentLayout::CanArgumentBeMoved, false)
 				.OnClicked(this, &FControlRigArgumentLayout::OnArgMoveDown)
 				.ToolTipText(LOCTEXT("FunctionArgDetailsArgMoveDownTooltip", "Move this parameter down in the list."))
 				[
@@ -411,13 +411,60 @@ bool FControlRigArgumentLayout::ShouldPinBeReadOnly(bool bIsEditingPinType/* = f
 
 bool FControlRigArgumentLayout::IsPinEditingReadOnly(bool bIsEditingPinType/* = false*/) const
 {
-	/*
+	return false;
+}
+
+bool FControlRigArgumentLayout::CanArgumentBeMoved(bool bMoveUp) const
+{
+	if(IsPinEditingReadOnly())
+	{
+		return false;
+	}
 	if (PinPtr.IsValid())
 	{
-		return PinPtr.Get()->IsExecuteContext();
+		URigVMPin* Pin = PinPtr.Get();
+		if(Pin->IsExecuteContext())
+		{
+			return false;
+		}
+
+		if(URigVMNode* Node = Pin->GetNode())
+		{
+			auto IsInput = [](URigVMPin* InPin) -> bool
+			{
+				return InPin->GetDirection() == ERigVMPinDirection::Input ||
+					InPin->GetDirection() == ERigVMPinDirection::Visible;
+			};
+
+			const bool bLookForInput = IsInput(Pin);
+
+			if(bMoveUp)
+			{
+				// if this is the first pin of its type
+				for(int32 Index = 0; Index < Node->GetPins().Num(); Index++)
+				{
+					URigVMPin* OtherPin = Node->GetPins()[Index];
+					if(IsInput(OtherPin) == bLookForInput)
+					{
+						return OtherPin != Pin;
+					}
+				}
+			}
+			else if(!bMoveUp)
+			{
+				// if this is the last pin of its type
+				for(int32 Index = Node->GetPins().Num() - 1; Index >= 0; Index--)
+				{
+					URigVMPin* OtherPin = Node->GetPins()[Index];
+					if(IsInput(OtherPin) == bLookForInput)
+					{
+						return OtherPin != Pin;
+					}
+				}
+			}
+		}
 	}
-	*/
-	return false;
+	return true;
 }
 
 FText FControlRigArgumentLayout::OnGetArgNameText() const
@@ -923,7 +970,28 @@ FReply FControlRigGraphDetails::OnAddNewInputClicked()
 					URigVMPin* LastPin = LibraryNode->GetPins().Last();
 					if (!LastPin->IsExecuteContext())
 					{
-						ArgumentName = LastPin->GetFName();
+						// strip off any tailing number from for example Argument_2
+						FString StrippedArgumentName = LastPin->GetName();
+						FString LastChars = StrippedArgumentName.Right(1);
+						StrippedArgumentName.LeftChopInline(1);
+						while(LastChars.IsNumeric() && !StrippedArgumentName.IsEmpty())
+						{
+							LastChars = StrippedArgumentName.Right(1);
+							StrippedArgumentName.LeftChopInline(1);
+
+							if(LastChars.StartsWith(TEXT("_")))
+							{
+								LastChars.Reset();
+								break;
+							}
+						}
+
+						StrippedArgumentName = StrippedArgumentName + LastChars;
+						if(!StrippedArgumentName.IsEmpty())
+						{
+							ArgumentName = *StrippedArgumentName;
+						}
+						
 						CPPType = LastPin->GetCPPType();
 						if (LastPin->GetCPPTypeObject())
 						{
diff --git a/Engine/Plugins/Animation/ControlRig/Source/ControlRigEditor/Private/ControlRigGraphDetails.h b/Engine/Plugins/Animation/ControlRig/Source/ControlRigEditor/Private/ControlRigGraphDetails.h
--- a/Engine/Plugins/Animation/ControlRig/Source/ControlRigEditor/Private/ControlRigGraphDetails.h
+++ b/Engine/Plugins/Animation/ControlRig/Source/ControlRigEditor/Private/ControlRigGraphDetails.h
@@ -82,6 +82,9 @@ private:
 	/** Determines if editing the pins on the node should be read only */
 	bool IsPinEditingReadOnly(bool bIsEditingPinType = false) const;
 
+	/** Determines if an argument can be moved up or down */
+	bool CanArgumentBeMoved(bool bMoveUp) const;
+
 	/** Callbacks for all the functionality for modifying arguments */
 	void OnRemoveClicked();
 	FReply OnArgMoveUp();
