Control Rig: Change undetermined string for float control value fields

diff --git a/Engine/Plugins/Animation/ControlRig/Source/ControlRigEditor/Private/ControlRigElementDetails.h b/Engine/Plugins/Animation/ControlRig/Source/ControlRigEditor/Private/ControlRigElementDetails.h
--- a/Engine/Plugins/Animation/ControlRig/Source/ControlRigEditor/Private/ControlRigElementDetails.h
+++ b/Engine/Plugins/Animation/ControlRig/Source/ControlRigEditor/Private/ControlRigElementDetails.h
@@ -503,6 +503,7 @@ private:
 			 .DisplayToggle(bShowToggle)
 			 .ToggleChecked(ToggleChecked)
 			 .OnToggleChanged(OnToggleChanged)
+			 .UndeterminedString(NSLOCTEXT("FRigControlElementDetails", "MultipleValues", "Multiple Values"))
 		]
 		.CopyAction(FUIAction(
 		FExecuteAction::CreateLambda([ValueType, Keys, HierarchyBeingDebugged]()
