Control Rig: Fix crash when control rig has no VM initialized

diff --git a/Engine/Plugins/Animation/ControlRig/Source/ControlRig/Private/ControlRig.cpp b/Engine/Plugins/Animation/ControlRig/Source/ControlRig/Private/ControlRig.cpp
--- a/Engine/Plugins/Animation/ControlRig/Source/ControlRig/Private/ControlRig.cpp
+++ b/Engine/Plugins/Animation/ControlRig/Source/ControlRig/Private/ControlRig.cpp
@@ -496,10 +496,14 @@ void UControlRig::InstantiateVMFromCDO()
 			// which should be unlikely
 			VM->CopyFrom(CDO->VM, !IsInGameThread(), true);
 		}
-		else
+		else if (VM)
 		{
 			VM->Reset();
 		}
+		else
+		{
+			ensure(false);
+		}
 	}
 
 	RequestInit();
diff --git a/Engine/Plugins/Animation/ControlRig/Source/ControlRigDeveloper/Private/ControlRigBlueprint.cpp b/Engine/Plugins/Animation/ControlRig/Source/ControlRigDeveloper/Private/ControlRigBlueprint.cpp
--- a/Engine/Plugins/Animation/ControlRig/Source/ControlRigDeveloper/Private/ControlRigBlueprint.cpp
+++ b/Engine/Plugins/Animation/ControlRig/Source/ControlRigDeveloper/Private/ControlRigBlueprint.cpp
@@ -715,6 +715,12 @@ void UControlRigBlueprint::RecompileVM()
 		{
 			if (UControlRig* InstanceRig = Cast<UControlRig>(Instance))
 			{
+				// No objects should be created during load, so PostInitInstanceIfRequired, which creates a new VM and
+				// DynamicHierarchy, should not be called during load
+				if (!InstanceRig->HasAllFlags(RF_NeedPostLoad))
+				{
+					InstanceRig->PostInitInstanceIfRequired();
+				}
 				InstanceRig->InstantiateVMFromCDO();
 			}
 		}
