Control Rig: Support Undo/Redo for DisplayName

diff --git a/Engine/Plugins/Animation/ControlRig/Source/ControlRigEditor/Private/ControlRigElementDetails.cpp b/Engine/Plugins/Animation/ControlRig/Source/ControlRigEditor/Private/ControlRigElementDetails.cpp
--- a/Engine/Plugins/Animation/ControlRig/Source/ControlRigEditor/Private/ControlRigElementDetails.cpp
+++ b/Engine/Plugins/Animation/ControlRig/Source/ControlRigEditor/Private/ControlRigElementDetails.cpp
@@ -3276,9 +3276,10 @@ void FRigControlElementDetails::SetDisplayName(const FText& InNewText, ETextComm
 				ControlElement.Settings.DisplayName = DisplayName;
 				ObjectBeingCustomized->SetContent<FRigControlElement>(ControlElement);
 
-				if(GetHierarchy())
+				if(HierarchyBeingCustomized)
 				{
-					GetHierarchy()->SetControlSettings(ControlElement.GetKey(), ControlElement.Settings, true, true, true);
+					FScopedTransaction Transaction(LOCTEXT("SetDisplayName", "SetDisplayName"));
+					HierarchyBeingCustomized->SetControlSettings(ControlElement.GetKey(), ControlElement.Settings, true, true, true);
 				}
 			}
 		}
