Fix for static analysis warnings. Removed duplicate Limit tests

diff --git a/Engine/Plugins/Animation/ControlRig/Source/ControlRig/Public/Rigs/RigHierarchyDefines.h b/Engine/Plugins/Animation/ControlRig/Source/ControlRig/Public/Rigs/RigHierarchyDefines.h
--- a/Engine/Plugins/Animation/ControlRig/Source/ControlRig/Public/Rigs/RigHierarchyDefines.h
+++ b/Engine/Plugins/Animation/ControlRig/Source/ControlRig/Public/Rigs/RigHierarchyDefines.h
@@ -743,7 +743,7 @@ public:
 				{
 					return;
 				}
-				if (LimitEnabled[0].IsOn() || LimitEnabled[2].IsOn() || LimitEnabled[2].IsOn())
+				if (LimitEnabled[0].IsOn() || LimitEnabled[2].IsOn())
 				{
 					FVector3f& ValueRef = GetRef<FVector3f>();
 					const FVector3f& Min = InMinimumValue.GetRef<FVector3f>();
@@ -760,7 +760,7 @@ public:
 				{
 					return;
 				}
-				if (LimitEnabled[0].IsOn() || LimitEnabled[2].IsOn() || LimitEnabled[2].IsOn())
+				if (LimitEnabled[0].IsOn() || LimitEnabled[2].IsOn())
 				{
 					FVector3f& ValueRef = GetRef<FVector3f>();
 					const FVector3f& Min = InMinimumValue.GetRef<FVector3f>();
@@ -777,7 +777,7 @@ public:
 				{
 					return;
 				}
-				if (LimitEnabled[0].IsOn() || LimitEnabled[2].IsOn() || LimitEnabled[2].IsOn())
+				if (LimitEnabled[0].IsOn() || LimitEnabled[2].IsOn())
 				{
 					FVector3f& ValueRef = GetRef<FVector3f>();
 					const FVector3f& Min = InMinimumValue.GetRef<FVector3f>();
@@ -799,7 +799,7 @@ public:
 				const FTransform Min = InMinimumValue.GetRef<FTransform_Float>().ToTransform();
 				const FTransform Max = InMaximumValue.GetRef<FTransform_Float>().ToTransform();
 
-				if (LimitEnabled[0].IsOn() || LimitEnabled[2].IsOn() || LimitEnabled[2].IsOn())
+				if (LimitEnabled[0].IsOn() || LimitEnabled[2].IsOn())
 				{
 					ValueRef.TranslationX = LimitEnabled[0].Apply<float>(ValueRef.TranslationX, Min.GetLocation().X, Max.GetLocation().X);
 					ValueRef.TranslationY = LimitEnabled[1].Apply<float>(ValueRef.TranslationY, Min.GetLocation().Y, Max.GetLocation().Y);
