Control Rig: Fix pin structs not showing binding in node

diff --git a/Engine/Plugins/Animation/ControlRig/Source/ControlRigEditor/Private/Graph/SControlRigGraphNode.cpp b/Engine/Plugins/Animation/ControlRig/Source/ControlRigEditor/Private/Graph/SControlRigGraphNode.cpp
--- a/Engine/Plugins/Animation/ControlRig/Source/ControlRigEditor/Private/Graph/SControlRigGraphNode.cpp
+++ b/Engine/Plugins/Animation/ControlRig/Source/ControlRigEditor/Private/Graph/SControlRigGraphNode.cpp
@@ -181,7 +181,7 @@ void SControlRigGraphNode::Construct( const FArguments& InArgs )
 		PinInfo.ModelPinPath = ModelPin->GetPinPath();
 		PinInfo.bAutoHeight = false;
 		
-		const bool bAskSchemaForEdition = RigSchema && ModelPin->IsStruct();
+		const bool bAskSchemaForEdition = RigSchema && ModelPin->IsStruct() && !ModelPin->IsBoundToVariable();
 		PinInfo.bHideInputWidget = (!ModelPin->IsBoundToVariable()) && PinInfo.bIsContainer;
 		if (!PinInfo.bHideInputWidget)
 		{
