Control Rig: Adapt label for contextual menu entries

diff --git a/Engine/Plugins/Animation/ControlRig/Source/ControlRigEditor/Private/ControlRigEditorModule.cpp b/Engine/Plugins/Animation/ControlRig/Source/ControlRigEditor/Private/ControlRigEditorModule.cpp
--- a/Engine/Plugins/Animation/ControlRig/Source/ControlRigEditor/Private/ControlRigEditorModule.cpp
+++ b/Engine/Plugins/Animation/ControlRig/Source/ControlRigEditor/Private/ControlRigEditorModule.cpp
@@ -1806,7 +1806,7 @@ void FControlRigEditorModule::GetContextMenuActions(const UControlRigGraphSchema
 					Section.AddMenuEntry(
 						"SelectRigElements",
 						LOCTEXT("SelectRigElements", "Select Rig Elements"),
-						LOCTEXT("SelectRigElements_Tooltip", "Selects the bone, controls or spaces associated with this node."),
+						LOCTEXT("SelectRigElements_Tooltip", "Selects the bone, controls or nulls associated with this node."),
 						FSlateIcon(),
 						FUIAction(FExecuteAction::CreateLambda([RigBlueprint, RigElementsToSelect]() {
 
diff --git a/Engine/Plugins/Animation/ControlRig/Source/ControlRigEditor/Private/ControlRigHierarchyCommands.cpp b/Engine/Plugins/Animation/ControlRig/Source/ControlRigEditor/Private/ControlRigHierarchyCommands.cpp
--- a/Engine/Plugins/Animation/ControlRig/Source/ControlRigEditor/Private/ControlRigHierarchyCommands.cpp
+++ b/Engine/Plugins/Animation/ControlRig/Source/ControlRigEditor/Private/ControlRigHierarchyCommands.cpp
@@ -17,8 +17,8 @@ void FControlRigHierarchyCommands::RegisterCommands()
 	UI_COMMAND(PasteItems, "Paste", "Paste the selected items.", EUserInterfaceActionType::Button, FInputChord(EKeys::V, EModifierKey::Control));
 	UI_COMMAND(PasteLocalTransforms, "Paste Local Transform", "Paste the local transforms.", EUserInterfaceActionType::Button, FInputChord());
 	UI_COMMAND(PasteGlobalTransforms, "Paste Global Transform", "Paste the global transforms.", EUserInterfaceActionType::Button, FInputChord(EKeys::V, EModifierKey::Control | EModifierKey::Shift));
-	UI_COMMAND(SetInitialTransformFromClosestBone, "Set Initial Transform from Closest Bone", "Find the Closest Bone to Initial Transform of the Selected Bones/Spaces or To Offset Transform of Selected Controls", EUserInterfaceActionType::Button, FInputChord());
-	UI_COMMAND(SetInitialTransformFromCurrentTransform, "Set Initial Transform from Current", "Save the Current Transform To Initial Transform of the Selected Bones/Spaces or To Offset Transform of Selected Controls", EUserInterfaceActionType::Button, FInputChord());
+	UI_COMMAND(SetInitialTransformFromClosestBone, "Set Initial Transform from Closest Bone", "Find the Closest Bone to Initial Transform of the Selected Bones/Nulls or To Offset Transform of Selected Controls", EUserInterfaceActionType::Button, FInputChord());
+	UI_COMMAND(SetInitialTransformFromCurrentTransform, "Set Initial Transform from Current", "Save the Current Transform To Initial Transform of the Selected Bones/Nulls or To Offset Transform of Selected Controls", EUserInterfaceActionType::Button, FInputChord());
 	UI_COMMAND(SetShapeTransformFromCurrent, "Set Shape Transform From Current", "Transfer the Current Local Transform of the Control to its Shape Transform, the Control's Transform Resets", EUserInterfaceActionType::Button, FInputChord());
 	UI_COMMAND(ResetTransform, "Reset Transform", "Reset the Transform", EUserInterfaceActionType::Button, FInputChord(EKeys::G, EModifierKey::Control));
 	UI_COMMAND(ResetAllTransforms, "Reset All Transforms", "Resets all Transforms", EUserInterfaceActionType::Button, FInputChord(EKeys::G, EModifierKey::Control | EModifierKey::Shift));
diff --git a/Engine/Plugins/Animation/ControlRig/Source/ControlRigEditor/Private/Editor/SRigHierarchy.cpp b/Engine/Plugins/Animation/ControlRig/Source/ControlRigEditor/Private/Editor/SRigHierarchy.cpp
--- a/Engine/Plugins/Animation/ControlRig/Source/ControlRigEditor/Private/Editor/SRigHierarchy.cpp
+++ b/Engine/Plugins/Animation/ControlRig/Source/ControlRigEditor/Private/Editor/SRigHierarchy.cpp
@@ -1104,8 +1104,41 @@ void SRigHierarchy::CreateContextMenu() const
 					FToolMenuSection& TransformsSection = InMenu->AddSection(TEXT("Transforms"), LOCTEXT("TransformsHeader", "Transforms"));
 					TransformsSection.AddMenuEntry(Commands.ResetTransform);
 					TransformsSection.AddMenuEntry(Commands.ResetAllTransforms);
-					TransformsSection.AddMenuEntry(Commands.SetInitialTransformFromCurrentTransform);
-					TransformsSection.AddMenuEntry(Commands.SetInitialTransformFromClosestBone);
+
+					{
+						static const FString InitialKeyword = TEXT("Initial");
+						static const FString OffsetKeyword = TEXT("Offset");
+						static const FString InitialOffsetKeyword = TEXT("Initial / Offset");
+						
+						const FString* Keyword = &InitialKeyword;
+						TArray<ERigElementType> SelectedTypes;;
+						TArray<TSharedPtr<FRigTreeElement>> SelectedItems = RigHierarchyPanel->TreeView->GetSelectedItems();
+						for(const TSharedPtr<FRigTreeElement>& SelectedItem : SelectedItems)
+						{
+							if(SelectedItem.IsValid())
+							{
+								SelectedTypes.AddUnique(SelectedItem->Key.Type);
+							}
+						}
+						if(SelectedTypes.Contains(ERigElementType::Control))
+						{
+							// since it is unique this means it is only controls
+							if(SelectedTypes.Num() == 1)
+							{
+								Keyword = &OffsetKeyword;
+							}
+							else
+							{
+								Keyword = &InitialOffsetKeyword;
+							}
+						}
+
+						const FText FromCurrentLabel = FText::Format(LOCTEXT("SetTransformFromCurrentTransform", "Set {0} Transform from Current"), FText::FromString(*Keyword));
+						const FText FromClosestBoneLabel = FText::Format(LOCTEXT("SetTransformFromClosestBone", "Set {0} Transform from Closest Bone"), FText::FromString(*Keyword));
+						TransformsSection.AddMenuEntry(Commands.SetInitialTransformFromCurrentTransform, FromCurrentLabel);
+						TransformsSection.AddMenuEntry(Commands.SetInitialTransformFromClosestBone, FromClosestBoneLabel);
+					}
+					
 					TransformsSection.AddMenuEntry(Commands.SetShapeTransformFromCurrent);
 					TransformsSection.AddMenuEntry(Commands.Unparent);
 
