Sequencer: Control Rig: Handle Renaming Controls.
In Control Rig, needed to make the RenameMap get copied when the hierarchy is copied and be persistent so it's saved if a rename happens when sequencer isn't loaded.
In Sequencer we handle this via the ReplaceControlRig function which is called when a Control Rig is compiled or when it is loaded. Removed function to try to get the rename on the fly but this doesn't work since the Control Rig Editor holds a different Control Rig than what's in Sequencer. Also added BP/Python function so users can script out renames as needed if the rig changes.
#rrb max.chen,helge.mathee

diff --git a/Engine/Plugins/Animation/ControlRig/Source/ControlRig/Private/Rigs/RigHierarchy.cpp b/Engine/Plugins/Animation/ControlRig/Source/ControlRig/Private/Rigs/RigHierarchy.cpp
--- a/Engine/Plugins/Animation/ControlRig/Source/ControlRig/Private/Rigs/RigHierarchy.cpp
+++ b/Engine/Plugins/Animation/ControlRig/Source/ControlRig/Private/Rigs/RigHierarchy.cpp
@@ -459,6 +459,12 @@ void URigHierarchy::CopyHierarchy(URigHierarchy* InHierarchy)
 		Target->CopyFrom(this, Source, InHierarchy);
 	}
 
+	for (const TPair<FRigElementKey, FRigElementKey>& NameMapPair : InHierarchy->PreviousNameMap)
+	{
+		PreviousNameMap.FindOrAdd(NameMapPair.Key) = NameMapPair.Value;
+	}
+
+
 	TopologyVersion = InHierarchy->GetTopologyVersion();
 	UpdateAllCachedChildren();
 	
diff --git a/Engine/Plugins/Animation/ControlRig/Source/ControlRig/Private/Sequencer/MovieSceneControlRigParameterSection.cpp b/Engine/Plugins/Animation/ControlRig/Source/ControlRig/Private/Sequencer/MovieSceneControlRigParameterSection.cpp
--- a/Engine/Plugins/Animation/ControlRig/Source/ControlRig/Private/Sequencer/MovieSceneControlRigParameterSection.cpp
+++ b/Engine/Plugins/Animation/ControlRig/Source/ControlRig/Private/Sequencer/MovieSceneControlRigParameterSection.cpp
@@ -712,6 +712,97 @@ UMovieSceneControlRigParameterSection::UMovieSceneControlRigParameterSection() :
 #endif
 }
 
+bool UMovieSceneControlRigParameterSection::RenameParameterName(const FName& OldParameterName, const FName& NewParameterName)
+{
+	bool bWasReplaced = false;
+	Modify();
+	for (FScalarParameterNameAndCurve& ScalarParameterNameAndCurve : ScalarParameterNamesAndCurves)
+	{
+		if (ScalarParameterNameAndCurve.ParameterName == OldParameterName)
+		{
+			ScalarParameterNameAndCurve.ParameterName = NewParameterName;
+			bWasReplaced = true;
+			break;
+		}
+	}
+
+	for (FBoolParameterNameAndCurve& BoolParameterNameAndCurve : BoolParameterNamesAndCurves)
+	{
+		if (BoolParameterNameAndCurve.ParameterName == OldParameterName)
+		{
+			BoolParameterNameAndCurve.ParameterName = NewParameterName;
+			bWasReplaced = true;
+			break;
+		}
+	}
+
+	for (FEnumParameterNameAndCurve& EnumParameterNameAndCurve : EnumParameterNamesAndCurves)
+	{
+		if (EnumParameterNameAndCurve.ParameterName == OldParameterName)
+		{
+			EnumParameterNameAndCurve.ParameterName = NewParameterName;
+			bWasReplaced = true;
+			break;
+		}
+	}
+
+	for (FIntegerParameterNameAndCurve& IntegerParameterNameAndCurve : IntegerParameterNamesAndCurves)
+	{
+		if (IntegerParameterNameAndCurve.ParameterName == OldParameterName)
+		{
+			IntegerParameterNameAndCurve.ParameterName = NewParameterName;
+			bWasReplaced = true;
+			break;
+		}
+	}
+
+	for (FVector2DParameterNameAndCurves& Vector2DParameterNameAndCurve : Vector2DParameterNamesAndCurves)
+	{
+		if (Vector2DParameterNameAndCurve.ParameterName == OldParameterName)
+		{
+			Vector2DParameterNameAndCurve.ParameterName = NewParameterName;
+			bWasReplaced = true;
+			break;
+		}
+	}
+
+	for (FVectorParameterNameAndCurves& VectorParameterNameAndCurve : VectorParameterNamesAndCurves)
+	{
+		if (VectorParameterNameAndCurve.ParameterName == OldParameterName)
+		{
+			VectorParameterNameAndCurve.ParameterName = NewParameterName;
+			bWasReplaced = true;
+			break;
+		}
+	}
+
+	for (FColorParameterNameAndCurves& ColorParameterNameAndCurve : ColorParameterNamesAndCurves)
+	{
+		if (ColorParameterNameAndCurve.ParameterName == OldParameterName)
+		{
+			ColorParameterNameAndCurve.ParameterName = NewParameterName;
+			bWasReplaced = true;
+			break;
+		}
+				
+	}
+
+	for (FTransformParameterNameAndCurves& TransformParameterNamesAndCurve : TransformParameterNamesAndCurves)
+	{
+		if (TransformParameterNamesAndCurve.ParameterName == OldParameterName)
+		{
+			TransformParameterNamesAndCurve.ParameterName = NewParameterName;
+			bWasReplaced = true;
+			break;
+		}
+	}
+	if (bWasReplaced)
+	{
+		ReconstructChannelProxy();
+	}
+	return bWasReplaced;
+}
+
 void UMovieSceneControlRigParameterSection::SetBlendType(EMovieSceneBlendType InBlendType)
 {
 	if (GetSupportedBlendTypes().Contains(InBlendType))
@@ -2028,6 +2119,12 @@ void UMovieSceneControlRigParameterSection::RecreateWithThisControlRig(UControlR
 			continue;
 		}
 
+		FName PreviousName = ControlRig->GetHierarchy()->GetPreviousName(ControlElement->GetKey());
+		if (PreviousName != NAME_None && PreviousName != ControlElement->GetKey().Name)
+		{
+			RenameParameterName(PreviousName, ControlElement->GetKey().Name);
+		}
+
 		switch (ControlElement->Settings.ControlType)
 		{
 		case ERigControlType::Float:
diff --git a/Engine/Plugins/Animation/ControlRig/Source/ControlRig/Private/Sequencer/MovieSceneControlRigParameterTrack.cpp b/Engine/Plugins/Animation/ControlRig/Source/ControlRig/Private/Sequencer/MovieSceneControlRigParameterTrack.cpp
--- a/Engine/Plugins/Animation/ControlRig/Source/ControlRig/Private/Sequencer/MovieSceneControlRigParameterTrack.cpp
+++ b/Engine/Plugins/Animation/ControlRig/Source/ControlRig/Private/Sequencer/MovieSceneControlRigParameterTrack.cpp
@@ -446,8 +446,21 @@ void UMovieSceneControlRigParameterTrack::PostEditImport()
 	ReconstructControlRig();
 }
 
+void UMovieSceneControlRigParameterTrack::RenameParameterName(const FName& OldParameterName, const FName& NewParameterName)
+{
+	if (OldParameterName != NewParameterName)
+	{
+		for (UMovieSceneSection* Section : Sections)
+		{
+			if (UMovieSceneControlRigParameterSection* CRSection = Cast<UMovieSceneControlRigParameterSection>(Section))
+			{
+				CRSection->RenameParameterName(OldParameterName, NewParameterName);
+			}
+		}
+	}
+}
 
-CONTROLRIG_API void UMovieSceneControlRigParameterTrack::ReplaceControlRig(UControlRig* NewControlRig, bool RecreateChannels)
+void UMovieSceneControlRigParameterTrack::ReplaceControlRig(UControlRig* NewControlRig, bool RecreateChannels)
 {
 	ControlRig = NewControlRig;
 	if (ControlRig->GetOuter() != this)
diff --git a/Engine/Plugins/Animation/ControlRig/Source/ControlRig/Public/Rigs/RigHierarchy.h b/Engine/Plugins/Animation/ControlRig/Source/ControlRig/Public/Rigs/RigHierarchy.h
--- a/Engine/Plugins/Animation/ControlRig/Source/ControlRig/Public/Rigs/RigHierarchy.h
+++ b/Engine/Plugins/Animation/ControlRig/Source/ControlRig/Public/Rigs/RigHierarchy.h
@@ -2676,6 +2676,9 @@ private:
 	TObjectPtr<URigHierarchyController> HierarchyController;
 
 	TMap<FRigElementKey, FRigElementKey> PreviousParentMap;
+
+	/*We save this so Sequencer can remap this after load*/
+	UPROPERTY()
 	TMap<FRigElementKey, FRigElementKey> PreviousNameMap;
 
 	int32 ResetPoseHash;
diff --git a/Engine/Plugins/Animation/ControlRig/Source/ControlRig/Public/Sequencer/MovieSceneControlRigParameterSection.h b/Engine/Plugins/Animation/ControlRig/Source/ControlRig/Public/Sequencer/MovieSceneControlRigParameterSection.h
--- a/Engine/Plugins/Animation/ControlRig/Source/ControlRig/Public/Sequencer/MovieSceneControlRigParameterSection.h
+++ b/Engine/Plugins/Animation/ControlRig/Source/ControlRig/Public/Sequencer/MovieSceneControlRigParameterSection.h
@@ -199,6 +199,8 @@ public:
 	FName FindControlNameFromSpaceChannel(const FMovieSceneControlRigSpaceChannel* SpaceChannel) const;
 	
 	FSpaceChannelAddedEvent& SpaceChannelAdded() { return OnSpaceChannelAdded; }
+
+	bool RenameParameterName(const FName& OldParameterName, const FName& NewParameterName);
 private:
 
 	FSpaceChannelAddedEvent OnSpaceChannelAdded;
diff --git a/Engine/Plugins/Animation/ControlRig/Source/ControlRig/Public/Sequencer/MovieSceneControlRigParameterTrack.h b/Engine/Plugins/Animation/ControlRig/Source/ControlRig/Public/Sequencer/MovieSceneControlRigParameterTrack.h
--- a/Engine/Plugins/Animation/ControlRig/Source/ControlRig/Public/Sequencer/MovieSceneControlRigParameterTrack.h
+++ b/Engine/Plugins/Animation/ControlRig/Source/ControlRig/Public/Sequencer/MovieSceneControlRigParameterTrack.h
@@ -61,6 +61,8 @@ public:
 
 	CONTROLRIG_API void ReplaceControlRig(UControlRig* NewControlRig, bool RecreateChannels);
 
+	CONTROLRIG_API void RenameParameterName(const FName& OldParameterName, const FName& NewParameterName);
+
 public:
 	CONTROLRIG_API UControlRig* GetControlRig() const { return ControlRig; }
 
diff --git a/Engine/Plugins/Animation/ControlRig/Source/ControlRigEditor/Private/ControlRigSequencerEditorLibrary.cpp b/Engine/Plugins/Animation/ControlRig/Source/ControlRigEditor/Private/ControlRigSequencerEditorLibrary.cpp
--- a/Engine/Plugins/Animation/ControlRig/Source/ControlRigEditor/Private/ControlRigSequencerEditorLibrary.cpp
+++ b/Engine/Plugins/Animation/ControlRig/Source/ControlRigEditor/Private/ControlRigSequencerEditorLibrary.cpp
@@ -2265,4 +2265,54 @@ bool UControlRigSequencerEditorLibrary::BakeControlRigSpace(ULevelSequence* InSe
 	return bValid;
 }
 
+bool UControlRigSequencerEditorLibrary::RenameControlRigControlChannels(ULevelSequence* LevelSequence, UControlRig* ControlRig, const TArray<FName>& InOldControlNames, const TArray<FName>& InNewControlNames)
+{
+	if (LevelSequence == nullptr || ControlRig == nullptr)
+	{
+		UE_LOG(LogControlRig, Error, TEXT("LevelSequence and Control Rig must be valid"));
+		return false;
+	}
+	if (InOldControlNames.Num() != InNewControlNames.Num() || InOldControlNames.Num() < 1)
+	{
+		UE_LOG(LogControlRig, Error, TEXT("Old and New Control Name arrays don't match in length"));
+		return false;
+	}
+	for (const FName& NewName : InNewControlNames)
+	{
+		if (ControlRig->FindControl(NewName) == nullptr)
+		{
+			const FString StringName = NewName.ToString();
+			UE_LOG(LogControlRig, Error, TEXT("Missing Control Name %s"), *(StringName));
+			return false;
+		}
+	}
+	bool bValid = false;
+
+
+	UMovieScene* MovieScene = LevelSequence->GetMovieScene();
+	if (MovieScene)
+	{
+		const TArray<FMovieSceneBinding>& Bindings = MovieScene->GetBindings();
+		for (const FMovieSceneBinding& Binding : Bindings)
+		{
+			TArray<UMovieSceneTrack*> Tracks = MovieScene->FindTracks(UMovieSceneControlRigParameterTrack::StaticClass(), Binding.GetObjectGuid(), NAME_None);
+			for (UMovieSceneTrack* AnyOleTrack : Tracks)
+			{
+				UMovieSceneControlRigParameterTrack* Track = Cast<UMovieSceneControlRigParameterTrack>(AnyOleTrack);
+				if (Track && Track->GetControlRig() == ControlRig)
+				{
+					Track->Modify();
+					bValid = true;
+					for (int32 Index = 0; Index < InOldControlNames.Num(); ++Index)
+					{
+						Track->RenameParameterName(InOldControlNames[Index], InNewControlNames[Index]);
+					}
+				}
+			}
+		}
+	}
+	
+	return bValid;
+}
+
 #undef LOCTEXT_NAMESPACE
diff --git a/Engine/Plugins/Animation/ControlRig/Source/ControlRigEditor/Private/Sequencer/ControlRigParameterTrackEditor.cpp b/Engine/Plugins/Animation/ControlRig/Source/ControlRigEditor/Private/Sequencer/ControlRigParameterTrackEditor.cpp
--- a/Engine/Plugins/Animation/ControlRig/Source/ControlRigEditor/Private/Sequencer/ControlRigParameterTrackEditor.cpp
+++ b/Engine/Plugins/Animation/ControlRig/Source/ControlRigEditor/Private/Sequencer/ControlRigParameterTrackEditor.cpp
@@ -257,6 +257,7 @@ FControlRigParameterTrackEditor::FControlRigParameterTrackEditor(TSharedRef<ISeq
 								UnbindControlRig(OldControlRig);
 								if (*NewControlRig)
 								{
+									Track->Modify();
 									Track->ReplaceControlRig(*NewControlRig, OldControlRig->GetClass() != (*NewControlRig)->GetClass());
 									BindControlRig(*NewControlRig);
 
@@ -2349,40 +2350,6 @@ void FControlRigParameterTrackEditor::HandleOnInitialized(UControlRig* ControlRi
 	}
 }
 
-void FControlRigParameterTrackEditor::HandleHierarchyModified(ERigHierarchyNotification InNotification, URigHierarchy* InHierarchy, const FRigBaseElement* InElement)
-{
-	//Only handle renamed events, the other ones like ERigHierarchyNotification::ElementRemoved will get handle when we recompile.
-	switch (InNotification)
-	{
-
-	case ERigHierarchyNotification::ElementRenamed:
-	{
-		if (const FRigControlElement* ControlElement = Cast<FRigControlElement>(InElement))
-		{
-			if (!GetSequencer().IsValid() || !GetSequencer()->IsAllowedToChange())
-			{
-				return;
-			}
-			UMovieScene* MovieScene = GetSequencer()->GetFocusedMovieSceneSequence()->GetMovieScene();
-			const TArray<FMovieSceneBinding>& Bindings = MovieScene->GetBindings();
-			for (const FMovieSceneBinding& Binding : Bindings)
-			{
-				UMovieSceneControlRigParameterTrack* Track = Cast<UMovieSceneControlRigParameterTrack>(MovieScene->FindTrack(UMovieSceneControlRigParameterTrack::StaticClass(), Binding.GetObjectGuid(), NAME_None));
-				if (Track && Track->GetControlRig() && Track->GetControlRig()->GetHierarchy() == InHierarchy)
-				{
-
-				}
-			}
-
-		}
-		break;
-	}
-	default:
-		break;
-	}
-	
-}
-
 void FControlRigParameterTrackEditor::HandleControlModified(UControlRig* ControlRig, FRigControlElement* ControlElement, const FRigControlModifiedContext& Context)
 {
 	if (!GetSequencer().IsValid() || !GetSequencer()->IsAllowedToChange() || Context.SetKey == EControlRigSetKey::Never)
diff --git a/Engine/Plugins/Animation/ControlRig/Source/ControlRigEditor/Private/Sequencer/ControlRigParameterTrackEditor.h b/Engine/Plugins/Animation/ControlRig/Source/ControlRigEditor/Private/Sequencer/ControlRigParameterTrackEditor.h
--- a/Engine/Plugins/Animation/ControlRig/Source/ControlRigEditor/Private/Sequencer/ControlRigParameterTrackEditor.h
+++ b/Engine/Plugins/Animation/ControlRig/Source/ControlRigEditor/Private/Sequencer/ControlRigParameterTrackEditor.h
@@ -126,7 +126,6 @@ private:
 
 	/** Control Rig Delegates*/
 	void HandleControlModified(UControlRig* Subject, FRigControlElement* ControlElement, const FRigControlModifiedContext& Context);
-	void HandleHierarchyModified(ERigHierarchyNotification InNotification, URigHierarchy* InHierarchy, const FRigBaseElement* InElement);
 	void HandleControlSelected(UControlRig* Subject, FRigControlElement* ControlElement, bool bSelected);
 	void HandleOnInitialized(UControlRig* Subject, const EControlRigState InState, const FName& InEventName);
 
diff --git a/Engine/Plugins/Animation/ControlRig/Source/ControlRigEditor/Public/ControlRigSequencerEditorLibrary.h b/Engine/Plugins/Animation/ControlRig/Source/ControlRigEditor/Public/ControlRigSequencerEditorLibrary.h
--- a/Engine/Plugins/Animation/ControlRig/Source/ControlRigEditor/Public/ControlRigSequencerEditorLibrary.h
+++ b/Engine/Plugins/Animation/ControlRig/Source/ControlRigEditor/Public/ControlRigSequencerEditorLibrary.h
@@ -834,7 +834,7 @@ public:
 	UFUNCTION(BlueprintCallable, Category = "Editor Scripting | Sequencer Tools | Control Rig")
 	static bool DeleteControlRigSpace(ULevelSequence* InSequence, UControlRig* InControlRig, FName InControlName,  FFrameNumber InTime, ESequenceTimeUnit TimeUnit = ESequenceTimeUnit::DisplayRate);
 
-	/**
+	/** Bake specified Control Rig Controls to a specified Space based upon the current settings
 	* @param InSequence Sequence to bake
 	* @param InControlRig ControlRig to bake
 	* @param InControlNames The name of the Controls to bake
@@ -845,5 +845,15 @@ public:
 	static bool BakeControlRigSpace(ULevelSequence* InSequence, UControlRig* InControlRig, TArray<FName>& InControlNames, FRigSpacePickerBakeSettings InSettings, ESequenceTimeUnit TimeUnit = ESequenceTimeUnit::DisplayRate);
 	
 	
+	/** Rename the Control Rig Channels in Sequencer to the specified new control names, which should be present on the Control Rig
+	* @param InSequence Sequence to rename controls
+	* @param InControlRig ControlRig to rename controls
+	* @param InOldControlNames The name of the old Control Rig Control Channels to change. Will be replaced by the corresponding name in the InNewControlNames array
+	* @param InNewControlNames  The name of the new Control Rig Channels 
+	* @return Return true if the function succeeds, false if it doesn't which can happen if the name arrays don't match in size or any of the new Control Names aren't valid
+	*/
+	UFUNCTION(BlueprintCallable, Category = "Editor Scripting | Sequencer Tools | Control Rig")
+	bool RenameControlRigControlChannels(ULevelSequence* InSequence, UControlRig* InControlRig, const TArray<FName>& InOldControlNames, const TArray<FName>& InNewControlNames);
+
 };
 
