Control Rig: Drag items onto non-allowed items to stop reparenting

diff --git a/Engine/Plugins/Animation/ControlRig/Source/ControlRigEditor/Private/Editor/SRigHierarchy.cpp b/Engine/Plugins/Animation/ControlRig/Source/ControlRigEditor/Private/Editor/SRigHierarchy.cpp
--- a/Engine/Plugins/Animation/ControlRig/Source/ControlRigEditor/Private/Editor/SRigHierarchy.cpp
+++ b/Engine/Plugins/Animation/ControlRig/Source/ControlRigEditor/Private/Editor/SRigHierarchy.cpp
@@ -42,6 +42,7 @@
 #include "ControlRigContextMenuContext.h"
 #include "SRigSpacePickerWidget.h"
 #include "Settings/ControlRigSettings.h"
+#include "Slate/Private/Widgets/Views/SListPanel.h"
 
 #define LOCTEXT_NAMESPACE "SRigHierarchy"
 
@@ -2132,6 +2133,13 @@ bool SRigHierarchy::HandleVerifyNameChanged(const FRigElementKey& OldKey, const
 
 FReply SRigHierarchy::OnDrop(const FGeometry& MyGeometry, const FDragDropEvent& DragDropEvent)
 {
+	// only allow drops onto empty space of the widget (when there's no target item under the mouse)
+	TSharedPtr<FRigTreeElement> ItemAtMouse = TreeView->FindItemAtPosition(DragDropEvent.GetScreenSpacePosition());
+	if(ItemAtMouse.IsValid())
+	{
+		return FReply::Unhandled();
+	}
+	
 	return OnAcceptDrop(DragDropEvent, EItemDropZone::OntoItem, nullptr);
 }
 
diff --git a/Engine/Plugins/Animation/ControlRig/Source/ControlRigEditor/Private/Editor/SRigHierarchyTreeView.cpp b/Engine/Plugins/Animation/ControlRig/Source/ControlRigEditor/Private/Editor/SRigHierarchyTreeView.cpp
--- a/Engine/Plugins/Animation/ControlRig/Source/ControlRigEditor/Private/Editor/SRigHierarchyTreeView.cpp
+++ b/Engine/Plugins/Animation/ControlRig/Source/ControlRigEditor/Private/Editor/SRigHierarchyTreeView.cpp
@@ -7,6 +7,7 @@
 #include "Widgets/Images/SImage.h"
 #include "Widgets/Input/SSearchBox.h"
 #include "Widgets/Text/STextBlock.h"
+#include "Slate/Private/Widgets/Views/SListPanel.h"
 #include "PropertyCustomizationHelpers.h"
 #include "Framework/Application/SlateApplication.h"
 #include "Editor/EditorEngine.h"
@@ -617,6 +618,23 @@ TArray<FRigElementKey> SRigHierarchyTreeView::GetSelectedKeys() const
 	return Keys;
 }
 
+TSharedPtr<FRigTreeElement> SRigHierarchyTreeView::FindItemAtPosition(FVector2D InScreenSpacePosition) const
+{
+	if(ItemsPanel.IsValid() && ItemsSource != nullptr)
+	{
+		const FGeometry MyGeometry = ItemsPanel->GetCachedGeometry();
+		FArrangedChildren ArrangedChildren(EVisibility::Visible);
+		ItemsPanel->ArrangeChildren(MyGeometry, ArrangedChildren, true);
+
+		const int32 Index = ItemsPanel->FindChildUnderPosition(ArrangedChildren, InScreenSpacePosition); 
+		if(ItemsSource->IsValidIndex(Index))
+		{
+			return ItemsSource->operator[](Index);
+		}
+	}
+	return TSharedPtr<FRigTreeElement>();
+}
+
 bool SRigHierarchyItem::OnVerifyNameChanged(const FText& InText, FText& OutErrorMessage)
 {
 	const FString NewName = InText.ToString();
diff --git a/Engine/Plugins/Animation/ControlRig/Source/ControlRigEditor/Private/Editor/SRigHierarchyTreeView.h b/Engine/Plugins/Animation/ControlRig/Source/ControlRigEditor/Private/Editor/SRigHierarchyTreeView.h
--- a/Engine/Plugins/Animation/ControlRig/Source/ControlRigEditor/Private/Editor/SRigHierarchyTreeView.h
+++ b/Engine/Plugins/Animation/ControlRig/Source/ControlRigEditor/Private/Editor/SRigHierarchyTreeView.h
@@ -264,6 +264,8 @@ public:
 	const TArray<TSharedPtr<FRigTreeElement>>& GetRootElements() const { return RootElements; }
 	FRigTreeDelegates& GetRigTreeDelegates() { return Delegates; }
 
+	TSharedPtr<FRigTreeElement> FindItemAtPosition(FVector2D InScreenSpacePosition) const;
+
 private:
 
 	/** A temporary snapshot of the SparseItemInfos in STreeView, used during RefreshTreeView() */
