Control Rig: Rename test space switching and hotkey

diff --git a/Engine/Plugins/Animation/ControlRig/Source/ControlRigEditor/Private/ControlRigHierarchyCommands.cpp b/Engine/Plugins/Animation/ControlRig/Source/ControlRigEditor/Private/ControlRigHierarchyCommands.cpp
--- a/Engine/Plugins/Animation/ControlRig/Source/ControlRigEditor/Private/ControlRigHierarchyCommands.cpp
+++ b/Engine/Plugins/Animation/ControlRig/Source/ControlRigEditor/Private/ControlRigHierarchyCommands.cpp
@@ -37,7 +37,7 @@ void FControlRigHierarchyCommands::RegisterCommands()
 	UI_COMMAND(ShowReferences, "Show References", "Whether to show or hide references", EUserInterfaceActionType::ToggleButton, FInputChord());
 	UI_COMMAND(ShowDynamicHierarchy, "Dynamic Hierarchy", "Whether to show the static or dynamic hierarchy", EUserInterfaceActionType::ToggleButton, FInputChord());
 	UI_COMMAND(ToggleControlShapeTransformEdit, "Toggle Shape Transform Edit", "Toggle Editing Selected Control's Shape Transform", EUserInterfaceActionType::ToggleButton, FInputChord(EKeys::Period, EModifierKey::Control)); 
-	UI_COMMAND(TestSpaceSwitching, "Test Space Switching", "Try out space switching on the control", EUserInterfaceActionType::Button, FInputChord(EKeys::Q, EModifierKey::Control)); 
+	UI_COMMAND(SpaceSwitching, "Space Switching", "Space switching on the control", EUserInterfaceActionType::Button, FInputChord(EKeys::Tab)); 
 	UI_COMMAND(ShowIconColors, "Show Icon Colors", "Whether to tint the icons with the element color", EUserInterfaceActionType::ToggleButton, FInputChord());
 }
 
diff --git a/Engine/Plugins/Animation/ControlRig/Source/ControlRigEditor/Private/ControlRigHierarchyCommands.h b/Engine/Plugins/Animation/ControlRig/Source/ControlRigEditor/Private/ControlRigHierarchyCommands.h
--- a/Engine/Plugins/Animation/ControlRig/Source/ControlRigEditor/Private/ControlRigHierarchyCommands.h
+++ b/Engine/Plugins/Animation/ControlRig/Source/ControlRigEditor/Private/ControlRigHierarchyCommands.h
@@ -111,8 +111,8 @@ public:
 	/** Toggle Shape Transform Edit*/
 	TSharedPtr< FUICommandInfo > ToggleControlShapeTransformEdit;
 
-	/** Test out space switching as it would look like for animator */
-	TSharedPtr< FUICommandInfo > TestSpaceSwitching;
+	/** Space switch as it would look like for animator */
+	TSharedPtr< FUICommandInfo > SpaceSwitching;
 
 	/** Whether to tint the icons with the element color */
 	TSharedPtr< FUICommandInfo > ShowIconColors;
diff --git a/Engine/Plugins/Animation/ControlRig/Source/ControlRigEditor/Private/Editor/SRigHierarchy.cpp b/Engine/Plugins/Animation/ControlRig/Source/ControlRigEditor/Private/Editor/SRigHierarchy.cpp
--- a/Engine/Plugins/Animation/ControlRig/Source/ControlRigEditor/Private/Editor/SRigHierarchy.cpp
+++ b/Engine/Plugins/Animation/ControlRig/Source/ControlRigEditor/Private/Editor/SRigHierarchy.cpp
@@ -445,7 +445,7 @@ void SRigHierarchy::BindCommands()
 		}));
 
 	CommandList->MapAction(
-		Commands.TestSpaceSwitching,
+		Commands.SpaceSwitching,
 		FExecuteAction::CreateSP(this, &SRigHierarchy::HandleTestSpaceSwitching),
 		FCanExecuteAction::CreateSP(this, &SRigHierarchy::IsControlSelected));
 
@@ -1092,7 +1092,7 @@ void SRigHierarchy::CreateContextMenu() const
 						}
 						else if(RigHierarchyPanel->IsControlSelected())
 						{
-							InteractionSection.AddMenuEntry(Commands.TestSpaceSwitching);
+							InteractionSection.AddMenuEntry(Commands.SpaceSwitching);
 						}
 					}
 					
