Control Rig: When reloading assets make sure to init the new instances

diff --git a/Engine/Plugins/Animation/ControlRig/Source/ControlRigEditor/Private/EditMode/ControlRigEditMode.cpp b/Engine/Plugins/Animation/ControlRig/Source/ControlRigEditor/Private/EditMode/ControlRigEditMode.cpp
--- a/Engine/Plugins/Animation/ControlRig/Source/ControlRigEditor/Private/EditMode/ControlRigEditMode.cpp
+++ b/Engine/Plugins/Animation/ControlRig/Source/ControlRigEditor/Private/EditMode/ControlRigEditMode.cpp
@@ -3168,6 +3168,7 @@ void FControlRigEditMode::AddControlRig(UControlRig* InControlRig)
 {
 	RuntimeControlRigs.AddUnique(InControlRig);
 
+	InControlRig->PostInitInstanceIfRequired();
 	InControlRig->GetHierarchy()->OnModified().RemoveAll(this);
 	InControlRig->GetHierarchy()->OnModified().AddSP(this, &FControlRigEditMode::OnHierarchyModified);
 
