Control Rig: Fix variable values not being initialized with default

diff --git a/Engine/Plugins/Animation/ControlRig/Source/ControlRig/Private/ControlRig.cpp b/Engine/Plugins/Animation/ControlRig/Source/ControlRig/Private/ControlRig.cpp
--- a/Engine/Plugins/Animation/ControlRig/Source/ControlRig/Private/ControlRig.cpp
+++ b/Engine/Plugins/Animation/ControlRig/Source/ControlRig/Private/ControlRig.cpp
@@ -509,12 +509,33 @@ void UControlRig::InstantiateVMFromCDO()
 	RequestInit();
 }
 
+void UControlRig::CopyExternalVariableDefaultValuesFromCDO()
+{
+	if (!HasAnyFlags(RF_ClassDefaultObject))
+	{
+		UControlRig* CDO = GetClass()->GetDefaultObject<UControlRig>();
+		TArray<FRigVMExternalVariable> CurrentVariables = GetExternalVariablesImpl(false);
+		TArray<FRigVMExternalVariable> CDOVariables = CDO->GetExternalVariablesImpl(false);
+		if (ensure(CurrentVariables.Num() == CDOVariables.Num()))
+		{
+			for (int32 i=0; i<CurrentVariables.Num(); ++i)
+			{
+				FRigVMExternalVariable& Variable = CurrentVariables[i];
+				FRigVMExternalVariable& CDOVariable = CDOVariables[i];
+				Variable.Property->CopyCompleteValue(Variable.Memory, CDOVariable.Memory);
+			}
+		}
+	}
+}
+
 void UControlRig::Execute(const EControlRigState InState, const FName& InEventName)
 {
 	if(!CanExecute())
 	{
 		return;
 	}
+
+	ensure(!HasAnyFlags(RF_ClassDefaultObject));
 	
 	DECLARE_SCOPE_HIERARCHICAL_COUNTER_FUNC()
 	QUICK_SCOPE_CYCLE_COUNTER(STAT_ControlRig_Execute);
diff --git a/Engine/Plugins/Animation/ControlRig/Source/ControlRig/Public/ControlRig.h b/Engine/Plugins/Animation/ControlRig/Source/ControlRig/Public/ControlRig.h
--- a/Engine/Plugins/Animation/ControlRig/Source/ControlRig/Public/ControlRig.h
+++ b/Engine/Plugins/Animation/ControlRig/Source/ControlRig/Public/ControlRig.h
@@ -511,6 +511,9 @@ private:
 	/** Copy the VM from the default object */
 	void InstantiateVMFromCDO();
 	
+	/** Copy the default values of external variables from the default object */
+	void CopyExternalVariableDefaultValuesFromCDO();
+	
 	/** Broadcasts a notification whenever the controlrig's memory is initialized. */
 	FControlRigExecuteEvent InitializedEvent;
 
diff --git a/Engine/Plugins/Animation/ControlRig/Source/ControlRigDeveloper/Private/ControlRigBlueprint.cpp b/Engine/Plugins/Animation/ControlRig/Source/ControlRigDeveloper/Private/ControlRigBlueprint.cpp
--- a/Engine/Plugins/Animation/ControlRig/Source/ControlRigDeveloper/Private/ControlRigBlueprint.cpp
+++ b/Engine/Plugins/Animation/ControlRig/Source/ControlRigDeveloper/Private/ControlRigBlueprint.cpp
@@ -706,9 +706,6 @@ void UControlRigBlueprint::RecompileVM()
 			return;
 		}
 
-		CDO->Execute(EControlRigState::Init, FRigUnit_BeginExecution::EventName); // need to clarify if we actually need this
-		Statistics = CDO->VM->GetStatistics();
-
 		TArray<UObject*> ArchetypeInstances;
 		CDO->GetArchetypeInstances(ArchetypeInstances);
 		for (UObject* Instance : ArchetypeInstances)
@@ -722,6 +719,7 @@ void UControlRigBlueprint::RecompileVM()
 					InstanceRig->PostInitInstanceIfRequired();
 				}
 				InstanceRig->InstantiateVMFromCDO();
+				InstanceRig->CopyExternalVariableDefaultValuesFromCDO();
 			}
 		}
 
diff --git a/Engine/Plugins/Animation/ControlRig/Source/ControlRigDeveloper/Public/ControlRigBlueprint.h b/Engine/Plugins/Animation/ControlRig/Source/ControlRigDeveloper/Public/ControlRigBlueprint.h
--- a/Engine/Plugins/Animation/ControlRig/Source/ControlRigDeveloper/Public/ControlRigBlueprint.h
+++ b/Engine/Plugins/Animation/ControlRig/Source/ControlRigDeveloper/Public/ControlRigBlueprint.h
@@ -422,8 +422,8 @@ public:
 	const FControlRigShapeDefinition* GetControlShapeByName(const FName& InName) const;
 #endif
 
-	UPROPERTY(transient, DuplicateTransient, VisibleAnywhere, Category = "VM", meta = (DisplayName = "VM Statistics", DisplayAfter = "VMCompileSettings"))
-	FRigVMStatistics Statistics;
+	UPROPERTY(transient, DuplicateTransient, meta = (DisplayName = "VM Statistics", DisplayAfter = "VMCompileSettings"))
+	FRigVMStatistics Statistics_DEPRECATED;
 
 	UPROPERTY(EditAnywhere, Category = "Drawing")
 	FControlRigDrawContainer DrawContainer;
diff --git a/Engine/Source/Runtime/RigVM/Public/RigVMCore/RigVM.h b/Engine/Source/Runtime/RigVM/Public/RigVMCore/RigVM.h
--- a/Engine/Source/Runtime/RigVM/Public/RigVMCore/RigVM.h
+++ b/Engine/Source/Runtime/RigVM/Public/RigVMCore/RigVM.h
@@ -682,7 +682,7 @@ public:
 #endif
 
 	// returns the statistics information
-	UFUNCTION(BlueprintPure, Category = "RigVM")
+	UFUNCTION(BlueprintPure, Category = "RigVM", meta=(DeprecatedFunction))
 	FRigVMStatistics GetStatistics() const
 	{
 		FRigVMStatistics Statistics;
