Control Rig: Deprecate ParentSwitchConstraint node

diff --git a/Engine/Plugins/Animation/ControlRig/Source/ControlRig/Private/Units/Hierarchy/RigUnit_ParentSwitchConstraint.h b/Engine/Plugins/Animation/ControlRig/Source/ControlRig/Private/Units/Hierarchy/RigUnit_ParentSwitchConstraint.h
--- a/Engine/Plugins/Animation/ControlRig/Source/ControlRig/Private/Units/Hierarchy/RigUnit_ParentSwitchConstraint.h
+++ b/Engine/Plugins/Animation/ControlRig/Source/ControlRig/Private/Units/Hierarchy/RigUnit_ParentSwitchConstraint.h
@@ -86,7 +86,7 @@ struct CONTROLRIG_API FRigUnit_ParentSwitchConstraint : public FRigUnitMutable
  * The Parent Switch Constraint is used to have an item follow one of multiple parents,
  * and allowing to switch between the parent in question.
  */
-USTRUCT(meta=(DisplayName="Parent Switch Constraint", Category="Constraints", DocumentationPolicy="Strict", Keywords = "SpaceSwitch"))
+USTRUCT(meta=(DisplayName="Parent Switch Constraint", Category="Constraints", DocumentationPolicy="Strict", Keywords = "SpaceSwitch", Deprecated = "5.0"))
 struct CONTROLRIG_API FRigUnit_ParentSwitchConstraintArray : public FRigUnitMutable
 {
 	GENERATED_BODY()
