ControlRig: Minor change, remove unused orphaned pins during graph refresh to avoid user confusion

diff --git a/Engine/Plugins/Animation/ControlRig/Source/ControlRigDeveloper/Private/ControlRigBlueprint.cpp b/Engine/Plugins/Animation/ControlRig/Source/ControlRigDeveloper/Private/ControlRigBlueprint.cpp
--- a/Engine/Plugins/Animation/ControlRig/Source/ControlRigDeveloper/Private/ControlRigBlueprint.cpp
+++ b/Engine/Plugins/Animation/ControlRig/Source/ControlRigDeveloper/Private/ControlRigBlueprint.cpp
@@ -800,6 +800,19 @@ void UControlRigBlueprint::RefreshAllModels()
 		URigVMController* Controller = GetOrCreateController(GraphToDetach);
 		Controller->ReattachLinksToPinObjects(true /* follow redirectors */, nullptr, false, true);
 	}
+
+	
+	TArray<URigVMGraph*> GraphsToClean = GetAllModels();
+	
+	for(int32 GraphIndex=0; GraphIndex<GraphsToClean.Num(); GraphIndex++)
+	{
+		URigVMGraph* GraphToClean = GraphsToClean[GraphIndex];
+		URigVMController* Controller = GetOrCreateController(GraphToClean);
+		for(URigVMNode* ModelNode : GraphToClean->GetNodes())
+		{
+			Controller->RemoveUnusedOrphanedPins(ModelNode, false);
+		}
+	}
 }
 
 void UControlRigBlueprint::HandleReportFromCompiler(EMessageSeverity::Type InSeverity, UObject* InSubject, const FString& InMessage)
