Sequencer:Control Rig: Need to make sure MaskIndex is valid outside editor.

diff --git a/Engine/Plugins/Animation/ControlRig/Source/ControlRig/Private/Sequencer/MovieSceneControlRigParameterSection.cpp b/Engine/Plugins/Animation/ControlRig/Source/ControlRig/Private/Sequencer/MovieSceneControlRigParameterSection.cpp
--- a/Engine/Plugins/Animation/ControlRig/Source/ControlRig/Private/Sequencer/MovieSceneControlRigParameterSection.cpp
+++ b/Engine/Plugins/Animation/ControlRig/Source/ControlRig/Private/Sequencer/MovieSceneControlRigParameterSection.cpp
@@ -1766,7 +1766,7 @@ void UMovieSceneControlRigParameterSection::ReconstructChannelProxy()
 					{
 						if (ControlElement->GetName() == Scalar.ParameterName)
 						{
-							ControlChannelMap.Add(Scalar.ParameterName, FChannelMapInfo(ControlIndex, TotalIndex,FloatChannelIndex));
+							ControlChannelMap.Add(Scalar.ParameterName, FChannelMapInfo(ControlIndex, TotalIndex,FloatChannelIndex,INDEX_NONE, NAME_None,MaskIndex));
 							Channels.Add(Scalar.ParameterCurve);
 							FloatChannelIndex += 1;
 							TotalIndex += 1;
@@ -1782,7 +1782,7 @@ void UMovieSceneControlRigParameterSection::ReconstructChannelProxy()
 					{
 						if (ControlElement->GetName() == Bool.ParameterName)
 						{
-							ControlChannelMap.Add(Bool.ParameterName, FChannelMapInfo(ControlIndex, TotalIndex, BoolChannelIndex));
+							ControlChannelMap.Add(Bool.ParameterName, FChannelMapInfo(ControlIndex, TotalIndex, BoolChannelIndex, INDEX_NONE, NAME_None, MaskIndex));
 							Channels.Add(Bool.ParameterCurve);
 							BoolChannelIndex += 1;
 							TotalIndex += 1;
@@ -1800,7 +1800,7 @@ void UMovieSceneControlRigParameterSection::ReconstructChannelProxy()
 						{
 							if (ControlElement->GetName() == Enum.ParameterName)
 							{
-								ControlChannelMap.Add(Enum.ParameterName, FChannelMapInfo(ControlIndex, TotalIndex, EnumChannelIndex));
+								ControlChannelMap.Add(Enum.ParameterName, FChannelMapInfo(ControlIndex, TotalIndex, EnumChannelIndex, INDEX_NONE, NAME_None, MaskIndex));
 								Channels.Add(Enum.ParameterCurve);
 								EnumChannelIndex += 1;
 								TotalIndex += 1;
@@ -1815,7 +1815,7 @@ void UMovieSceneControlRigParameterSection::ReconstructChannelProxy()
 						{
 							if (ControlElement->GetName() == Integer.ParameterName)
 							{
-								ControlChannelMap.Add(Integer.ParameterName, FChannelMapInfo(ControlIndex, TotalIndex, IntegerChannelIndex));
+								ControlChannelMap.Add(Integer.ParameterName, FChannelMapInfo(ControlIndex, TotalIndex, IntegerChannelIndex, INDEX_NONE, NAME_None, MaskIndex));
 								Channels.Add(Integer.ParameterCurve);
 								IntegerChannelIndex += 1;
 								TotalIndex += 1;
@@ -1832,7 +1832,7 @@ void UMovieSceneControlRigParameterSection::ReconstructChannelProxy()
 					{
 						if (ControlElement->GetName() == Vector2D.ParameterName)
 						{
-							ControlChannelMap.Add(Vector2D.ParameterName, FChannelMapInfo(ControlIndex, TotalIndex, FloatChannelIndex));
+							ControlChannelMap.Add(Vector2D.ParameterName, FChannelMapInfo(ControlIndex, TotalIndex, FloatChannelIndex, INDEX_NONE, NAME_None, MaskIndex));
 							Channels.Add(Vector2D.XCurve);
 							Channels.Add(Vector2D.YCurve);
 							FloatChannelIndex += 2;
@@ -1851,7 +1851,7 @@ void UMovieSceneControlRigParameterSection::ReconstructChannelProxy()
 					{
 						if (ControlElement->GetName() == Vector.ParameterName)
 						{
-							ControlChannelMap.Add(Vector.ParameterName, FChannelMapInfo(ControlIndex, TotalIndex, FloatChannelIndex));
+							ControlChannelMap.Add(Vector.ParameterName, FChannelMapInfo(ControlIndex, TotalIndex, FloatChannelIndex, INDEX_NONE, NAME_None, MaskIndex));
 							bool bDoSpaceChannel = true;
 							if (bDoSpaceChannel)
 							{
@@ -1897,7 +1897,7 @@ void UMovieSceneControlRigParameterSection::ReconstructChannelProxy()
 					{
 						if (ControlElement->GetName() == Transform.ParameterName)
 						{
-							ControlChannelMap.Add(Transform.ParameterName, FChannelMapInfo(ControlIndex, TotalIndex,FloatChannelIndex));
+							ControlChannelMap.Add(Transform.ParameterName, FChannelMapInfo(ControlIndex, TotalIndex,FloatChannelIndex, INDEX_NONE, NAME_None, MaskIndex));
 						
 							bool bDoSpaceChannel = true;
 							if (bDoSpaceChannel)
