Control Rig: Treat limits as range and allow inverted limits

diff --git a/Engine/Plugins/Animation/ControlRig/Source/ControlRig/Public/Rigs/RigHierarchyDefines.h b/Engine/Plugins/Animation/ControlRig/Source/ControlRig/Public/Rigs/RigHierarchyDefines.h
--- a/Engine/Plugins/Animation/ControlRig/Source/ControlRig/Public/Rigs/RigHierarchyDefines.h
+++ b/Engine/Plugins/Animation/ControlRig/Source/ControlRig/Public/Rigs/RigHierarchyDefines.h
@@ -177,7 +177,14 @@ struct CONTROLRIG_API FRigControlLimitEnabled
 
 		if(bMinimum && bMaximum)
 		{
-			return FMath::Clamp<T>(InValue, InMinimum, InMaximum);
+			if(InMinimum < InMaximum)
+			{
+				return FMath::Clamp<T>(InValue, InMinimum, InMaximum);
+			}
+			else
+			{
+				return FMath::Clamp<T>(InValue, InMaximum, InMinimum);
+			}
 		}
 
 		if(bMinimum)
@@ -680,16 +687,6 @@ public:
 		}
 	}
 
-	template<typename T>
-	FORCEINLINE static T Clamp(const T Value, const T Minimum, const T Maximum)
-	{
-		if (Minimum < Maximum)
-		{
-			return FMath::Clamp<T>(Value, Minimum, Maximum);
-		}
-		return FMath::Clamp<T>(Value, Maximum, Minimum);
-	}
-
 	FORCEINLINE void ApplyLimits(
 		const TArray<FRigControlLimitEnabled>& LimitEnabled,
 		ERigControlType InControlType,
@@ -708,7 +705,7 @@ public:
 				if (LimitEnabled[0].IsOn())
 				{
 					float& ValueRef = GetRef<float>();
-					ValueRef = Clamp<float>(ValueRef, InMinimumValue.Get<float>(), InMaximumValue.Get<float>());
+					ValueRef = LimitEnabled[0].Apply<float>(ValueRef, InMinimumValue.Get<float>(), InMaximumValue.Get<float>());
 				}
 				break;
 			}
@@ -717,7 +714,7 @@ public:
 				if (LimitEnabled[0].IsOn())
 				{
 					int32& ValueRef = GetRef<int32>();
-					ValueRef = Clamp<int32>(ValueRef, InMinimumValue.Get<int32>(), InMaximumValue.Get<int32>());
+					ValueRef = LimitEnabled[0].Apply<int32>(ValueRef, InMinimumValue.Get<int32>(), InMaximumValue.Get<int32>());
 				}
 				break;
 			}
