Control Rig: Fix rotator values not displayed in a consistent way through the control rig editor

diff --git a/Engine/Plugins/Animation/ControlRig/Source/ControlRigDeveloper/Private/Graph/ControlRigGraphNode.cpp b/Engine/Plugins/Animation/ControlRig/Source/ControlRigDeveloper/Private/Graph/ControlRigGraphNode.cpp
--- a/Engine/Plugins/Animation/ControlRig/Source/ControlRigDeveloper/Private/Graph/ControlRigGraphNode.cpp
+++ b/Engine/Plugins/Animation/ControlRig/Source/ControlRigDeveloper/Private/Graph/ControlRigGraphNode.cpp
@@ -121,7 +121,7 @@ FText UControlRigGraphNode::GetNodeTitle(ENodeTitleType::Type TitleType) const
 						{
 							if(NewVariable.VarName == VariableName)
 							{
-								const FString DefaultValue = NewVariable.DefaultValue;
+								FString DefaultValue = NewVariable.DefaultValue;
 								if(DefaultValue.IsEmpty())
 								{
 									static const FString VariableString = TEXT("Variable");
@@ -129,6 +129,15 @@ FText UControlRigGraphNode::GetNodeTitle(ENodeTitleType::Type TitleType) const
 								}
 								else
 								{
+									// Change the order of values in rotators so that they match the pin order
+									if (!NewVariable.VarType.IsContainer() && NewVariable.VarType.PinSubCategoryObject == TBaseStructure<FRotator>::Get())
+									{
+										TArray<FString> Values;
+										DefaultValue.ParseIntoArray(Values, TEXT(","));
+										Values.Swap(0, 1);
+										Values.Swap(0, 2);
+										DefaultValue = FString::Join(Values, TEXT(","));										
+									}
 									SubTitle = FString::Printf(TEXT("Default %s"), *DefaultValue);
 								}
 								break;
diff --git a/Engine/Plugins/Animation/ControlRig/Source/ControlRigEditor/Private/Graph/SControlRigGraphNode.cpp b/Engine/Plugins/Animation/ControlRig/Source/ControlRigEditor/Private/Graph/SControlRigGraphNode.cpp
--- a/Engine/Plugins/Animation/ControlRig/Source/ControlRigEditor/Private/Graph/SControlRigGraphNode.cpp
+++ b/Engine/Plugins/Animation/ControlRig/Source/ControlRigEditor/Private/Graph/SControlRigGraphNode.cpp
@@ -107,9 +107,20 @@ void SControlRigGraphNode::Construct( const FArguments& InArgs )
 		static void VisitPinRecursively(URigVMPin* InPin, TArray<URigVMPin*>& OutPins)
 		{
 			OutPins.Add(InPin);
-			for (URigVMPin* SubPin : InPin->GetSubPins())
+
+			if (InPin->GetCPPType() == TEXT("FRotator"))
 			{
-				VisitPinRecursively(SubPin, OutPins);
+				TArray<URigVMPin*> SubPins = InPin->GetSubPins();
+				OutPins.Add(SubPins[2]);	
+				OutPins.Add(SubPins[0]);	
+				OutPins.Add(SubPins[1]);	
+			}
+			else
+			{				
+				for (URigVMPin* SubPin : InPin->GetSubPins())
+				{
+					VisitPinRecursively(SubPin, OutPins);
+				}
 			}
 		}
 	};
@@ -968,7 +979,7 @@ void SControlRigGraphNode::GetNodeInfoPopups(FNodeInfoContext* Context, TArray<F
 #else
 								URigVMMemoryStorage* Memory = ActiveObject->GetVM()->GetDebugMemory();
 								// We mark PPF_ExternalEditor so that default values are also printed
-								const FString DebugValue = Memory->GetDataAsStringSafe(WatchOperand->GetRegisterIndex(), PPF_ExternalEditor);
+								const FString DebugValue = Memory->GetDataAsStringSafe(WatchOperand->GetRegisterIndex(), PPF_ExternalEditor | STRUCT_ExportTextItemNative);
 								if(!DebugValue.IsEmpty())
 								{
 									DefaultValues = URigVMPin::SplitDefaultValue(DebugValue);
@@ -978,13 +989,32 @@ void SControlRigGraphNode::GetNodeInfoPopups(FNodeInfoContext* Context, TArray<F
 						}
 
 						FString WatchText;
-
 						if (DefaultValues.Num() == 1)
 						{
-							WatchText = DefaultValues[0];
+							// Fixing the order of values in the rotator to match the order in the pins (x, y, z)
+							if (ModelPin->GetCPPType() == TEXT("FRotator"))
+							{
+								TArray<FString> Values;
+								FString TrimmedText = DefaultValues[0].LeftChop(1).RightChop(1); //Remove ()
+								TrimmedText.ParseIntoArray(Values, TEXT(","));
+								Values.Swap(0, 1);
+								Values.Swap(0, 2);
+								WatchText = FString(TEXT("(")) + FString::Join(Values, TEXT(",")) + FString(TEXT(")"));
+								UScriptStruct* RotatorStruct = TBaseStructure<FRotator>::Get();
+								for (TFieldIterator<FProperty> It(RotatorStruct); It; ++It)
+								{
+									FName PropertyName = It->GetFName();
+									WatchText.ReplaceInline(*PropertyName.ToString(), *It->GetDisplayNameText().ToString());
+								}
+							}
+							else
+							{								
+								WatchText = DefaultValues[0];
+							}
 						}
 						else if (DefaultValues.Num() > 1)
 						{
+							// todo: Fix order of values in rotators within other structures
 							WatchText = FString::Printf(TEXT("%s"), *FString::Join(DefaultValues, TEXT("\n")));
 						}
 
