Control Rig: Sequencer: Python, add functions to get default space keys, add track to binding proxy and remove delete function that wasn't needed,.
#trivial

diff --git a/Engine/Plugins/Animation/ControlRig/Source/ControlRigEditor/Private/ControlRigSequencerEditorLibrary.cpp b/Engine/Plugins/Animation/ControlRig/Source/ControlRigEditor/Private/ControlRigSequencerEditorLibrary.cpp
--- a/Engine/Plugins/Animation/ControlRig/Source/ControlRigEditor/Private/ControlRigSequencerEditorLibrary.cpp
+++ b/Engine/Plugins/Animation/ControlRig/Source/ControlRigEditor/Private/ControlRigSequencerEditorLibrary.cpp
@@ -69,6 +69,7 @@ TArray<FControlRigSequencerBindingProxy> UControlRigSequencerEditorLibrary::GetC
 					if (Track && Track->GetControlRig())
 					{
 						FControlRigSequencerBindingProxy BindingProxy;
+						BindingProxy.Track = Track;
 						BindingProxy.ControlRig = Track->GetControlRig();
 						BindingProxy.Proxy.BindingID = Binding.GetObjectGuid();
 						BindingProxy.Proxy.Sequence = LevelSequence;
@@ -2188,43 +2189,6 @@ bool UControlRigSequencerEditorLibrary::SetControlRigSpace(ULevelSequence* Level
 	return bValid;
 }
 
-bool UControlRigSequencerEditorLibrary::DeleteControlRigSpace(ULevelSequence* LevelSequence, UControlRig* ControlRig, FName ControlName, FFrameNumber InTime, ESequenceTimeUnit TimeUnit)
-{
-	TWeakPtr<ISequencer> WeakSequencer = GetSequencerFromAsset(LevelSequence);
-	bool bValid = false;
-
-	if (WeakSequencer.IsValid() && ControlRig && ControlName != NAME_None)
-	{
-		if (FRigControlElement* Element = ControlRig->FindControl(ControlName))
-		{
-			TSharedPtr<ISequencer>  Sequencer = WeakSequencer.Pin();
-			FScopedTransaction Transaction(LOCTEXT("KeyControlRigSpace", "Key Control Rig Space"));
-			FSpaceChannelAndSection SpaceChannelAndSection = FControlRigSpaceChannelHelpers::FindSpaceChannelAndSectionForControl(ControlRig, ControlName, Sequencer.Get(), false /*bCreateIfNeeded*/);
-			if (SpaceChannelAndSection.SpaceChannel)
-			{
-				if (TimeUnit == ESequenceTimeUnit::DisplayRate)
-				{
-					InTime = FFrameRate::TransformTime(FFrameTime(InTime, 0), LevelSequence->GetMovieScene()->GetDisplayRate(), LevelSequence->GetMovieScene()->GetTickResolution()).RoundToFrame();
-				}
-				UMovieSceneControlRigParameterSection* ParamSection = Cast<UMovieSceneControlRigParameterSection>(SpaceChannelAndSection.SectionToKey);
-				FControlRigSpaceChannelHelpers::SequencerSpaceChannelKeyDeleted(ControlRig, Sequencer.Get(), ControlName, SpaceChannelAndSection.SpaceChannel, ParamSection, InTime);
-				bValid = true;
-			}
-			else
-			{
-				UE_LOG(LogControlRig, Error, TEXT("Can not find Space Channel"));
-				return false;
-			}
-		}
-		else
-		{
-			UE_LOG(LogControlRig, Error, TEXT("Can not find Control with that Name"));
-			return false;
-		}
-	}
-	return bValid;
-}
-
 bool UControlRigSequencerEditorLibrary::BakeControlRigSpace(ULevelSequence* InSequence, UControlRig* InControlRig, TArray<FName>& InControlNames, FRigSpacePickerBakeSettings InSettings, ESequenceTimeUnit TimeUnit)
 {
 	TWeakPtr<ISequencer> WeakSequencer = GetSequencerFromAsset(InSequence);
@@ -2315,4 +2279,14 @@ bool UControlRigSequencerEditorLibrary::RenameControlRigControlChannels(ULevelSe
 	return bValid;
 }
 
+FRigElementKey UControlRigSequencerEditorLibrary::GetDefaultParentKey()
+{
+	return URigHierarchy::GetDefaultParentKey();
+}
+
+FRigElementKey UControlRigSequencerEditorLibrary::GetWorldSpaceReferenceKey()
+{
+	return URigHierarchy::GetWorldSpaceReferenceKey();
+}
+
 #undef LOCTEXT_NAMESPACE
diff --git a/Engine/Plugins/Animation/ControlRig/Source/ControlRigEditor/Public/ControlRigSequencerEditorLibrary.h b/Engine/Plugins/Animation/ControlRig/Source/ControlRigEditor/Public/ControlRigSequencerEditorLibrary.h
--- a/Engine/Plugins/Animation/ControlRig/Source/ControlRigEditor/Public/ControlRigSequencerEditorLibrary.h
+++ b/Engine/Plugins/Animation/ControlRig/Source/ControlRigEditor/Public/ControlRigSequencerEditorLibrary.h
@@ -28,11 +28,13 @@ struct FControlRigSequencerBindingProxy
 
 	FControlRigSequencerBindingProxy()
 		: ControlRig(nullptr)
+		, Track(nullptr)
 	{}
 
-	FControlRigSequencerBindingProxy(const FSequencerBindingProxy& InProxy, UControlRig* InControlRig)
+	FControlRigSequencerBindingProxy(const FSequencerBindingProxy& InProxy, UControlRig* InControlRig, UMovieSceneControlRigParameterTrack* InTrack)
 		: Proxy(InProxy)
 		, ControlRig(InControlRig)
+		, Track(InTrack)
 	{}
 
 	UPROPERTY(BlueprintReadOnly, Category = ControlRig)
@@ -40,6 +42,9 @@ struct FControlRigSequencerBindingProxy
 
 	UPROPERTY(BlueprintReadOnly, Category = ControlRig)
 	TObjectPtr<UControlRig> ControlRig;
+
+	UPROPERTY(BlueprintReadOnly, Category = ControlRig)
+	TObjectPtr<UMovieSceneControlRigParameterTrack> Track;
 };
 
 /**
@@ -808,6 +813,20 @@ public:
 	UFUNCTION(BlueprintCallable, Category = "Editor Scripting | Sequencer Tools | Control Rig")
 	static bool CollapseControlRigAnimLayers(ULevelSequence* InSequence,UMovieSceneControlRigParameterTrack* InTrack, bool bKeyReduce = false, float Tolerance = 0.001f);
 
+	/*
+	 * Get the default parent key, can be used a parent space.
+	 *
+	*/
+	UFUNCTION(BlueprintCallable, Category = "Editor Scripting | Sequencer Tools | Control Rig")
+	static FRigElementKey GetDefaultParentKey();
+
+	/*
+	 * Get the default world space key, can be used a world space.
+	 *
+	*/
+	UFUNCTION(BlueprintCallable, Category = "Editor Scripting | Sequencer Tools | Control Rig")
+	static FRigElementKey GetWorldSpaceReferenceKey();
+
 	/*
 	 * Set the a key for the Control Rig Space for the Control at the specified time. If space is the same as the current no key witll be set.
 	 *
@@ -821,19 +840,6 @@ public:
 	UFUNCTION(BlueprintCallable, Category = "Editor Scripting | Sequencer Tools | Control Rig")
 	static bool SetControlRigSpace(ULevelSequence* InSequence, UControlRig* InControlRig, FName InControlName, const FRigElementKey& InSpaceKey, FFrameNumber InTime,  ESequenceTimeUnit TimeUnit = ESequenceTimeUnit::DisplayRate);
 
-	/*Set the a key for the Control Rig Space for the Control at the specified time.If space is the same as the current no key witll be set.
-	*
-	* @param InSequence Sequence to set the space
-	* @param InControlRig ControlRig with the Control
-	* @param InControlName The name of the Control
-	* @param InSpaceKey  The new space for the Control
-	* @param InTime Time to change the space.
-	* @param TimeUnit Unit for the InTime, either in display rate or tick resolution
-	* 
-	*/
-	UFUNCTION(BlueprintCallable, Category = "Editor Scripting | Sequencer Tools | Control Rig")
-	static bool DeleteControlRigSpace(ULevelSequence* InSequence, UControlRig* InControlRig, FName InControlName,  FFrameNumber InTime, ESequenceTimeUnit TimeUnit = ESequenceTimeUnit::DisplayRate);
-
 	/** Bake specified Control Rig Controls to a specified Space based upon the current settings
 	* @param InSequence Sequence to bake
 	* @param InControlRig ControlRig to bake
