Control Rig: Draw expander arrows with auto width

diff --git a/Engine/Plugins/Animation/ControlRig/Source/ControlRigEditor/Private/Graph/SControlRigGraphNode.cpp b/Engine/Plugins/Animation/ControlRig/Source/ControlRigEditor/Private/Graph/SControlRigGraphNode.cpp
--- a/Engine/Plugins/Animation/ControlRig/Source/ControlRigEditor/Private/Graph/SControlRigGraphNode.cpp
+++ b/Engine/Plugins/Animation/ControlRig/Source/ControlRigEditor/Private/Graph/SControlRigGraphNode.cpp
@@ -285,6 +285,7 @@ void SControlRigGraphNode::Construct( const FArguments& InArgs )
 					// Add the expander arrow
 					FullPinHorizontalRowWidget->InsertSlot(bPinWidgetForExpanderLeft ? 1 : FullPinHorizontalRowWidget->GetChildren()->Num() - 1)
 					.Padding(LeftIndentation, 0.f, RightIndentation, 0.f)
+					.AutoWidth()
 					[
 						SNew(SButton)
 						.ButtonStyle(FAppStyle::Get(), TEXT("SimpleButton"))
@@ -1256,7 +1257,7 @@ EVisibility SControlRigGraphNode::GetPinVisibility(int32 InPinInfoIndex) const
 		if(ParentPinIndex != INDEX_NONE)
 		{
 			const EVisibility ParentPinVisibility = GetPinVisibility(ParentPinIndex); 
-			if(GetPinVisibility(ParentPinIndex) != EVisibility::Visible)
+			if(ParentPinVisibility != EVisibility::Visible)
 			{
 				return ParentPinVisibility;
 			}
