Control Rig: Fix Preview Animation Controller does not play on Control Rigs post import hierarchy

diff --git a/Engine/Plugins/Animation/ControlRig/Source/ControlRigEditor/Private/Editor/ControlRigEditor.cpp b/Engine/Plugins/Animation/ControlRig/Source/ControlRigEditor/Private/Editor/ControlRigEditor.cpp
--- a/Engine/Plugins/Animation/ControlRig/Source/ControlRigEditor/Private/Editor/ControlRigEditor.cpp
+++ b/Engine/Plugins/Animation/ControlRig/Source/ControlRigEditor/Private/Editor/ControlRigEditor.cpp
@@ -4037,6 +4037,15 @@ void FControlRigEditor::RebindToSkeletalMeshComponent()
 	}
 }
 
+void FControlRigEditor::UpdateMeshInAnimInstance(USkeletalMesh* InNewSkeletalMesh)
+{
+	UControlRigSkeletalMeshComponent* EditorSkelComp = Cast<UControlRigSkeletalMeshComponent>(GetPersonaToolkit()->GetPreviewScene()->GetPreviewMeshComponent());
+	if (EditorSkelComp)
+	{
+		EditorSkelComp->GetAnimInstance()->CurrentSkeleton = InNewSkeletalMesh->GetSkeleton();
+	}
+}
+
 void FControlRigEditor::UpdateStaleWatchedPins()
 {
 	UControlRigBlueprint* ControlRigBP = GetControlRigBlueprint();
diff --git a/Engine/Plugins/Animation/ControlRig/Source/ControlRigEditor/Private/Editor/ControlRigEditor.h b/Engine/Plugins/Animation/ControlRig/Source/ControlRigEditor/Private/Editor/ControlRigEditor.h
--- a/Engine/Plugins/Animation/ControlRig/Source/ControlRigEditor/Private/Editor/ControlRigEditor.h
+++ b/Engine/Plugins/Animation/ControlRig/Source/ControlRigEditor/Private/Editor/ControlRigEditor.h
@@ -331,6 +331,9 @@ private:
 	/** Rebind our anim instance to the preview's skeletal mesh component */
 	void RebindToSkeletalMeshComponent();
 
+	/** Update the skeletal mesh componens so that the anim instance has the correct skeletal mesh */
+	void UpdateMeshInAnimInstance(USkeletalMesh* InNewSkeletalMesh);
+
 	/** Update stale watch pins */
 	void UpdateStaleWatchedPins();
 
diff --git a/Engine/Plugins/Animation/ControlRig/Source/ControlRigEditor/Private/Editor/SRigHierarchy.cpp b/Engine/Plugins/Animation/ControlRig/Source/ControlRigEditor/Private/Editor/SRigHierarchy.cpp
--- a/Engine/Plugins/Animation/ControlRig/Source/ControlRigEditor/Private/Editor/SRigHierarchy.cpp
+++ b/Engine/Plugins/Animation/ControlRig/Source/ControlRigEditor/Private/Editor/SRigHierarchy.cpp
@@ -1354,6 +1354,7 @@ void SRigHierarchy::ImportHierarchy(const FAssetData& InAssetData)
 		ControlRigEditor.IsValid() && 
 		Mesh != nullptr)
 	{
+		ControlRigEditor.Pin()->UpdateMeshInAnimInstance(Mesh);
 		ControlRigEditor.Pin()->GetPersonaToolkit()->SetPreviewMesh(Mesh, true);
 	}
 
