Control Rig: Permanently delete graph nodes from the panel to avoid double undo / redo

diff --git a/Engine/Plugins/Animation/ControlRig/Source/ControlRigDeveloper/Private/Graph/ControlRigGraph.cpp b/Engine/Plugins/Animation/ControlRig/Source/ControlRigDeveloper/Private/Graph/ControlRigGraph.cpp
--- a/Engine/Plugins/Animation/ControlRig/Source/ControlRigDeveloper/Private/Graph/ControlRigGraph.cpp
+++ b/Engine/Plugins/Animation/ControlRig/Source/ControlRigDeveloper/Private/Graph/ControlRigGraph.cpp
@@ -395,6 +395,8 @@ void UControlRigGraph::HandleModifiedEvent(ERigVMGraphNotifType InNotifType, URi
 				UEdGraphNode* EdNode = FindNodeForModelNodeName(ModelNode->GetFName(), false);
 				if (EdNode)
 				{
+					static UObject* NewOuter = GetTransientPackage();
+					
 					// Rename the soon to be deleted object to a unique name, so that other objects can use
 					// the old name
 					{
@@ -404,10 +406,10 @@ void UControlRigGraph::HandleModifiedEvent(ERigVMGraphNotifType InNotifType, URi
 						do
 						{
 							DeletedName = FString::Printf(TEXT("%s_Deleted_%d"), *EdNode->GetName(), Index++); 
-							ExistingObject = StaticFindObject(/*Class=*/ NULL, EdNode->GetOuter(), *DeletedName, true);						
+							ExistingObject = StaticFindObject(/*Class=*/ NULL, NewOuter, *DeletedName, true);						
 						}
 						while (ExistingObject);
-						EdNode->Rename(*DeletedName);
+						EdNode->Rename(*DeletedName, NewOuter, REN_ForceNoResetLoaders | REN_DoNotDirty | REN_DontCreateRedirectors | REN_NonTransactional);
 					}
 					RemoveNode(EdNode, true);
 					NotifyGraphChanged();
diff --git a/Engine/Plugins/Animation/ControlRig/Source/ControlRigDeveloper/Private/Graph/ControlRigGraphSchema.cpp b/Engine/Plugins/Animation/ControlRig/Source/ControlRigDeveloper/Private/Graph/ControlRigGraphSchema.cpp
--- a/Engine/Plugins/Animation/ControlRig/Source/ControlRigDeveloper/Private/Graph/ControlRigGraphSchema.cpp
+++ b/Engine/Plugins/Animation/ControlRig/Source/ControlRigDeveloper/Private/Graph/ControlRigGraphSchema.cpp
@@ -1664,7 +1664,11 @@ bool UControlRigGraphSchema::SafeDeleteNodeFromGraph(UEdGraph* Graph, UEdGraphNo
 {
 	if (UControlRigGraphNode* RigNode = Cast<UControlRigGraphNode>(Node))
 	{
-		return RigNode->GetController()->RemoveNode(RigNode->GetModelNode(), true, false, true);
+		if (GEditor)
+		{
+			GEditor->CancelTransaction(0);
+		}
+		return RigNode->GetController()->RemoveNode(RigNode->GetModelNode(), true, true, true);
 	}
 	return false;
 }
