Control Rig: Fix name data base hash issue

The previous hash was discarding the internal suffix number, which causes Control_1 == Control_2

diff --git a/Engine/Plugins/Animation/ControlRig/Source/ControlRig/Private/Rigs/RigNameCache.cpp b/Engine/Plugins/Animation/ControlRig/Source/ControlRig/Private/Rigs/RigNameCache.cpp
--- a/Engine/Plugins/Animation/ControlRig/Source/ControlRig/Private/Rigs/RigNameCache.cpp
+++ b/Engine/Plugins/Animation/ControlRig/Source/ControlRig/Private/Rigs/RigNameCache.cpp
@@ -15,8 +15,8 @@ FRigNameOp FRigNameOp::Concat(const FName& InA, const FName& InB)
 {
 	FRigNameOp Op;
 	Op.Type = ERigNameOp::Concat;
-	Op.A = InA.GetComparisonIndex().ToUnstableInt();
-	Op.B = InB.GetComparisonIndex().ToUnstableInt();
+	Op.A = GetTypeHash(InA);
+	Op.B = GetTypeHash(InB);
 	return Op;
 }
 
@@ -24,7 +24,7 @@ FRigNameOp FRigNameOp::Left(const FName& InA, const uint32 InCount)
 {
 	FRigNameOp Op;
 	Op.Type = ERigNameOp::Left;
-	Op.A = InA.GetComparisonIndex().ToUnstableInt();
+	Op.A = GetTypeHash(InA);
 	Op.B = InCount;
 	return Op;
 }
@@ -33,7 +33,7 @@ FRigNameOp FRigNameOp::Right(const FName& InA, const uint32 InCount)
 {
 	FRigNameOp Op;
 	Op.Type = ERigNameOp::Right;
-	Op.A = InA.GetComparisonIndex().ToUnstableInt();
+	Op.A = GetTypeHash(InA);
 	Op.B = InCount;
 	return Op;
 }
@@ -42,7 +42,7 @@ FRigNameOp FRigNameOp::LeftChop(const FName& InA, const uint32 InCount)
 {
 	FRigNameOp Op;
 	Op.Type = ERigNameOp::LeftChop;
-	Op.A = InA.GetComparisonIndex().ToUnstableInt();
+	Op.A = GetTypeHash(InA);
 	Op.B = InCount;
 	return Op;
 }
@@ -51,7 +51,7 @@ FRigNameOp FRigNameOp::RightChop(const FName& InA, const uint32 InCount)
 {
 	FRigNameOp Op;
 	Op.Type = ERigNameOp::RightChop;
-	Op.A = InA.GetComparisonIndex().ToUnstableInt();
+	Op.A = GetTypeHash(InA);
 	Op.B = InCount;
 	return Op;
 }
@@ -60,9 +60,9 @@ FRigNameOp FRigNameOp::Replace(const FName& InA, const FName& InB, const FName&
 {
 	FRigNameOp Op;
 	Op.Type = InSearchCase == ESearchCase::CaseSensitive ? ERigNameOp::ReplaceCase : ERigNameOp::ReplaceNoCase;
-	Op.A = InA.GetComparisonIndex().ToUnstableInt();
-	Op.B = InB.GetComparisonIndex().ToUnstableInt();
-	Op.C = InC.GetComparisonIndex().ToUnstableInt();
+	Op.A = GetTypeHash(InA);
+	Op.B = GetTypeHash(InB);
+	Op.C = GetTypeHash(InC);
 	return Op;
 }
 
@@ -70,8 +70,8 @@ FRigNameOp FRigNameOp::EndsWith(const FName& InA, const FName& InB, const ESearc
 {
 	FRigNameOp Op;
 	Op.Type = InSearchCase == ESearchCase::CaseSensitive ? ERigNameOp::EndsWithCase : ERigNameOp::EndsWithNoCase;
-	Op.A = InA.GetComparisonIndex().ToUnstableInt();
-	Op.B = InB.GetComparisonIndex().ToUnstableInt();
+	Op.A = GetTypeHash(InA);
+	Op.B = GetTypeHash(InB);
 	return Op;
 }
 
@@ -79,8 +79,8 @@ FRigNameOp FRigNameOp::StartsWith(const FName& InA, const FName& InB, const ESea
 {
 	FRigNameOp Op;
 	Op.Type = InSearchCase == ESearchCase::CaseSensitive ? ERigNameOp::StartsWithCase : ERigNameOp::StartsWithNoCase;
-	Op.A = InA.GetComparisonIndex().ToUnstableInt();
-	Op.B = InB.GetComparisonIndex().ToUnstableInt();
+	Op.A = GetTypeHash(InA);
+	Op.B = GetTypeHash(InB);
 	return Op;
 }
 
@@ -88,8 +88,8 @@ FRigNameOp FRigNameOp::Contains(const FName& InA, const FName& InB, const ESearc
 {
 	FRigNameOp Op;
 	Op.Type = InSearchCase == ESearchCase::CaseSensitive ? ERigNameOp::ContainsCase : ERigNameOp::ContainsNoCase;
-	Op.A = InA.GetComparisonIndex().ToUnstableInt();
-	Op.B = InB.GetComparisonIndex().ToUnstableInt();
+	Op.A = GetTypeHash(InA);
+	Op.B = GetTypeHash(InB);
 	return Op;
 }
 
diff --git a/Engine/Plugins/Animation/ControlRig/Source/ControlRig/Public/Rigs/RigNameCache.h b/Engine/Plugins/Animation/ControlRig/Source/ControlRig/Public/Rigs/RigNameCache.h
--- a/Engine/Plugins/Animation/ControlRig/Source/ControlRig/Public/Rigs/RigNameCache.h
+++ b/Engine/Plugins/Animation/ControlRig/Source/ControlRig/Public/Rigs/RigNameCache.h
@@ -33,6 +33,11 @@ public:
 	, Type(ERigNameOp::None)
 	{}
 
+	FORCEINLINE static uint32 GetTypeHash(const FName& InName)
+	{
+		return HashCombine(InName.GetComparisonIndex().ToUnstableInt(), uint32(InName.GetNumber()));
+	}
+	
 	static FRigNameOp Concat(const FName& InA, const FName& InB);
 	static FRigNameOp Left(const FName& InA, const uint32 InCount);
 	static FRigNameOp Right(const FName& InA, const uint32 InCount);
