Control Rig: Propage rename changes and maintain selection during rename

diff --git a/Engine/Plugins/Animation/ControlRig/Source/ControlRig/Private/Rigs/RigHierarchy.cpp b/Engine/Plugins/Animation/ControlRig/Source/ControlRig/Private/Rigs/RigHierarchy.cpp
--- a/Engine/Plugins/Animation/ControlRig/Source/ControlRig/Private/Rigs/RigHierarchy.cpp
+++ b/Engine/Plugins/Animation/ControlRig/Source/ControlRig/Private/Rigs/RigHierarchy.cpp
@@ -2088,6 +2088,42 @@ void URigHierarchy::Notify(ERigHierarchyNotification InNotifType, const FRigBase
 		return;
 	}
 	ModifiedEvent.Broadcast(InNotifType, this, InElement);
+
+#if WITH_EDITOR
+
+	// certain events needs to be forwarded to the listening hierarchies.
+	// this mainly has to do with topological change within the hierarchy.
+	switch (InNotifType)
+	{
+		case ERigHierarchyNotification::ElementAdded:
+		case ERigHierarchyNotification::ElementRemoved:
+		case ERigHierarchyNotification::ElementRenamed:
+		case ERigHierarchyNotification::ParentChanged:
+		case ERigHierarchyNotification::ParentWeightsChanged:
+		{
+			ensure(InElement != nullptr);
+				
+			for(FRigHierarchyListener& Listener : ListeningHierarchies)
+			{
+				URigHierarchy* ListeningHierarchy = Listener.Hierarchy.Get();
+				if (ListeningHierarchy)
+				{			
+					if(const FRigBaseElement* ListeningElement = ListeningHierarchy->Find( InElement->GetKey()))
+					{
+						ListeningHierarchy->Notify(InNotifType, ListeningElement);
+					}
+				}
+			}
+
+			break;
+		}
+		default:
+		{
+			break;
+		}
+	}
+
+#endif
 }
 
 FTransform URigHierarchy::GetTransform(FRigTransformElement* InTransformElement,
diff --git a/Engine/Plugins/Animation/ControlRig/Source/ControlRig/Private/Rigs/RigHierarchyController.cpp b/Engine/Plugins/Animation/ControlRig/Source/ControlRig/Private/Rigs/RigHierarchyController.cpp
--- a/Engine/Plugins/Animation/ControlRig/Source/ControlRig/Private/Rigs/RigHierarchyController.cpp
+++ b/Engine/Plugins/Animation/ControlRig/Source/ControlRig/Private/Rigs/RigHierarchyController.cpp
@@ -1751,7 +1751,7 @@ bool URigHierarchyController::RemoveElement(FRigBaseElement* InElement)
 	return NumElementsRemoved == 1;
 }
 
-FRigElementKey URigHierarchyController::RenameElement(FRigElementKey InElement, FName InName, bool bSetupUndo, bool bPrintPythonCommand)
+FRigElementKey URigHierarchyController::RenameElement(FRigElementKey InElement, FName InName, bool bSetupUndo, bool bPrintPythonCommand, bool bClearSelection)
 {
 	if(!IsValid())
 	{
@@ -1774,7 +1774,7 @@ FRigElementKey URigHierarchyController::RenameElement(FRigElementKey InElement,
 	}
 #endif
 
-	const bool bRenamed = RenameElement(Element, InName);
+	const bool bRenamed = RenameElement(Element, InName, bClearSelection);
 
 #if WITH_EDITOR
 	if(!bRenamed && TransactionPtr.IsValid())
@@ -1783,7 +1783,7 @@ FRigElementKey URigHierarchyController::RenameElement(FRigElementKey InElement,
 	}
 	TransactionPtr.Reset();
 
-	if (bRenamed)
+	if (bRenamed && bClearSelection)
 	{
 		ClearSelection();
 	}
@@ -1804,7 +1804,7 @@ FRigElementKey URigHierarchyController::RenameElement(FRigElementKey InElement,
 	return bRenamed ? Element->GetKey() : FRigElementKey();
 }
 
-bool URigHierarchyController::RenameElement(FRigBaseElement* InElement, const FName &InName)
+bool URigHierarchyController::RenameElement(FRigBaseElement* InElement, const FName &InName, bool bClearSelection)
 {
 	if(InElement == nullptr)
 	{
@@ -1820,6 +1820,7 @@ bool URigHierarchyController::RenameElement(FRigBaseElement* InElement, const FN
 
 	// deselect the key that no longer exists
 	// no need to trigger a reselect since we always clear selection after rename
+	const bool bWasSelected = Hierarchy->IsSelected(InElement); 
 	if (Hierarchy->IsSelected(InElement))
 	{
 		DeselectElement(OldKey);
@@ -1870,6 +1871,11 @@ bool URigHierarchyController::RenameElement(FRigBaseElement* InElement, const FN
 	Hierarchy->IncrementTopologyVersion();
 	Notify(ERigHierarchyNotification::ElementRenamed, InElement);
 
+	if (!bClearSelection && bWasSelected)
+	{
+		SelectElement(InElement->GetKey(), true);
+	}
+
 	return true;
 }
 
diff --git a/Engine/Plugins/Animation/ControlRig/Source/ControlRig/Public/Rigs/RigHierarchyController.h b/Engine/Plugins/Animation/ControlRig/Source/ControlRig/Public/Rigs/RigHierarchyController.h
--- a/Engine/Plugins/Animation/ControlRig/Source/ControlRig/Public/Rigs/RigHierarchyController.h
+++ b/Engine/Plugins/Animation/ControlRig/Source/ControlRig/Public/Rigs/RigHierarchyController.h
@@ -358,10 +358,11 @@ public:
 	 * @param InName The new name to set for the element
 	 * @param bSetupUndo If set to true the stack will record the change for undo / redo
 	 * @param bPrintPythonCommand If set to true a python command equivalent to this call will be printed out
+	 * @param bClearSelection True if the selection should be cleared after a rename
 	 * @return Returns the new element key used for the element
 	 */
 	UFUNCTION(BlueprintCallable, Category = URigHierarchyController)
-    FRigElementKey RenameElement(FRigElementKey InElement, FName InName, bool bSetupUndo = false, bool bPrintPythonCommand = false);
+    FRigElementKey RenameElement(FRigElementKey InElement, FName InName, bool bSetupUndo = false, bool bPrintPythonCommand = false, bool bClearSelection = true);
 
 	/**
 	 * Adds a new parent to an element. For elements that allow only one parent the parent will be replaced (Same as ::SetParent).
@@ -553,9 +554,10 @@ private:
 	 * Renames an existing element in the hierarchy
 	 * @param InElement The element to rename
 	 * @param InName The new name to set for the element
+	 * @param bClearSelection True if the selection should be cleared after a rename
 	 * @return Returns true if successful.
 	 */
-    bool RenameElement(FRigBaseElement* InElement, const FName &InName);
+    bool RenameElement(FRigBaseElement* InElement, const FName &InName, bool bClearSelection = true);
 
 	/**
 	 * Removes an existing parent from an element in the hierarchy. For elements that allow only one parent the element will be unparented (same as ::RemoveAllParents)
diff --git a/Engine/Plugins/Animation/ControlRig/Source/ControlRigEditor/Private/Editor/SRigHierarchy.cpp b/Engine/Plugins/Animation/ControlRig/Source/ControlRigEditor/Private/Editor/SRigHierarchy.cpp
--- a/Engine/Plugins/Animation/ControlRig/Source/ControlRigEditor/Private/Editor/SRigHierarchy.cpp
+++ b/Engine/Plugins/Animation/ControlRig/Source/ControlRigEditor/Private/Editor/SRigHierarchy.cpp
@@ -2096,7 +2096,7 @@ FName SRigHierarchy::HandleRenameElement(const FRigElementKey& OldKey, const FSt
 		const FName SanitizedName = *SanitizedNameStr;
 		FName ResultingName = NAME_None;
 
-		ResultingName = Controller->RenameElement(OldKey, SanitizedName, true, true).Name;
+		ResultingName = Controller->RenameElement(OldKey, SanitizedName, true, true, false).Name;
 		ControlRigBlueprint->PropagateHierarchyFromBPToInstances();
 		return ResultingName;
 	}
