Control Rig: Fix crash when naming an exposed input of a collapsed node when clashes with an output

diff --git a/Engine/Plugins/Animation/ControlRig/Source/ControlRigDeveloper/Private/Graph/ControlRigGraphSchema.cpp b/Engine/Plugins/Animation/ControlRig/Source/ControlRigDeveloper/Private/Graph/ControlRigGraphSchema.cpp
--- a/Engine/Plugins/Animation/ControlRig/Source/ControlRigDeveloper/Private/Graph/ControlRigGraphSchema.cpp
+++ b/Engine/Plugins/Animation/ControlRig/Source/ControlRigDeveloper/Private/Graph/ControlRigGraphSchema.cpp
@@ -67,6 +67,11 @@ FControlRigLocalVariableNameValidator::FControlRigLocalVariableNameValidator(con
 		{
 			Names.Add(InputArgument.Name.ToString());
 		}
+
+		for (const FRigVMGraphVariableDescription& OutputArgument : Graph->GetOutputArguments())
+		{
+			Names.Add(OutputArgument.Name.ToString());
+		}
 	}
 }
 
diff --git a/Engine/Plugins/Animation/ControlRig/Source/ControlRigEditor/Private/ControlRigGraphDetails.cpp b/Engine/Plugins/Animation/ControlRig/Source/ControlRigEditor/Private/ControlRigGraphDetails.cpp
--- a/Engine/Plugins/Animation/ControlRig/Source/ControlRigEditor/Private/ControlRigGraphDetails.cpp
+++ b/Engine/Plugins/Animation/ControlRig/Source/ControlRigEditor/Private/ControlRigGraphDetails.cpp
@@ -113,12 +113,15 @@ void FControlRigArgumentGroupLayout::HandleModifiedEvent(ERigVMGraphNotifType In
 	switch (InNotifType)
 	{
 		case ERigVMGraphNotifType::PinAdded:
+		case ERigVMGraphNotifType::PinRenamed:
 		case ERigVMGraphNotifType::PinRemoved:
 		case ERigVMGraphNotifType::PinIndexChanged:
 		case ERigVMGraphNotifType::PinTypeChanged:
 		{
 			URigVMPin* Pin = CastChecked<URigVMPin>(InSubject);
-			if (Pin->GetNode() == LibraryNode)
+			if (Pin->GetNode() == LibraryNode ||
+				(Pin->GetNode()->IsA<URigVMFunctionEntryNode>() && Pin->GetNode()->GetOuter() == Graph) ||
+				(Pin->GetNode()->IsA<URigVMFunctionReturnNode>() && Pin->GetNode()->GetOuter() == Graph))
 			{
 				OnRebuildChildren.ExecuteIfBound();
 			}
@@ -483,16 +486,19 @@ FText FControlRigArgumentLayout::OnGetArgToolTipText() const
 
 void FControlRigArgumentLayout::OnArgNameTextCommitted(const FText& NewText, ETextCommit::Type InTextCommit)
 {
-	if (!NewText.IsEmpty() && PinPtr.IsValid() && ControlRigBlueprintPtr.IsValid() && !ShouldPinBeReadOnly())
+	if (InTextCommit == ETextCommit::OnEnter)
 	{
-		URigVMPin* Pin = PinPtr.Get();
-		UControlRigBlueprint* Blueprint = ControlRigBlueprintPtr.Get();
-		if (URigVMLibraryNode* LibraryNode = Cast<URigVMLibraryNode>(Pin->GetNode()))
+		if (!NewText.IsEmpty() && PinPtr.IsValid() && ControlRigBlueprintPtr.IsValid() && !ShouldPinBeReadOnly())
 		{
-			if (URigVMController* Controller = Blueprint->GetController(LibraryNode->GetContainedGraph()))
+			URigVMPin* Pin = PinPtr.Get();
+			UControlRigBlueprint* Blueprint = ControlRigBlueprintPtr.Get();
+			if (URigVMLibraryNode* LibraryNode = Cast<URigVMLibraryNode>(Pin->GetNode()))
 			{
-				const FString& NewName = NewText.ToString();
-				Controller->RenameExposedPin(Pin->GetFName(), *NewName, true, true);
+				if (URigVMController* Controller = Blueprint->GetController(LibraryNode->GetContainedGraph()))
+				{
+					const FString& NewName = NewText.ToString();
+					Controller->RenameExposedPin(Pin->GetFName(), *NewName, true, true);
+				}
 			}
 		}
 	}
diff --git a/Engine/Source/Developer/RigVMDeveloper/Private/RigVMModel/RigVMController.cpp b/Engine/Source/Developer/RigVMDeveloper/Private/RigVMModel/RigVMController.cpp
--- a/Engine/Source/Developer/RigVMDeveloper/Private/RigVMModel/RigVMController.cpp
+++ b/Engine/Source/Developer/RigVMDeveloper/Private/RigVMModel/RigVMController.cpp
@@ -8113,8 +8113,28 @@ FName URigVMController::AddExposedPin(const FName& InPinName, ERigVMPinDirection
 		SetPinDefaultValue(Pin, InDefaultValue, true, bSetupUndoRedo, false);
 	}
 
-	RefreshFunctionPins(Graph->GetEntryNode(), true);
-	RefreshFunctionPins(Graph->GetReturnNode(), true);
+	URigVMFunctionEntryNode* EntryNode = Graph->GetEntryNode();
+	if (!EntryNode)
+	{
+		EntryNode = NewObject<URigVMFunctionEntryNode>(Graph, TEXT("Entry"));
+		Graph->Nodes.Add(EntryNode);
+		RefreshFunctionPins(EntryNode, false);
+
+		Notify(ERigVMGraphNotifType::NodeAdded, EntryNode);
+	}
+
+	URigVMFunctionReturnNode* ReturnNode = Graph->GetReturnNode();
+	if (!ReturnNode)
+	{
+		ReturnNode = NewObject<URigVMFunctionReturnNode>(Graph, TEXT("Return"));
+		Graph->Nodes.Add(ReturnNode);
+		RefreshFunctionPins(ReturnNode, false);
+
+		Notify(ERigVMGraphNotifType::NodeAdded, ReturnNode);
+	}
+
+	RefreshFunctionPins(EntryNode, true);
+	RefreshFunctionPins(ReturnNode, true);
 	RefreshFunctionReferences(LibraryNode, bSetupUndoRedo);
 
 	if (bSetupUndoRedo)
diff --git a/Engine/Source/Developer/RigVMDeveloper/Private/RigVMModel/RigVMGraph.cpp b/Engine/Source/Developer/RigVMDeveloper/Private/RigVMModel/RigVMGraph.cpp
--- a/Engine/Source/Developer/RigVMDeveloper/Private/RigVMModel/RigVMGraph.cpp
+++ b/Engine/Source/Developer/RigVMDeveloper/Private/RigVMModel/RigVMGraph.cpp
@@ -323,6 +323,23 @@ TArray<FRigVMGraphVariableDescription> URigVMGraph::GetInputArguments() const
 	return Inputs;
 }
 
+TArray<FRigVMGraphVariableDescription> URigVMGraph::GetOutputArguments() const
+{
+	TArray<FRigVMGraphVariableDescription> Outputs;
+	if (URigVMFunctionReturnNode* ReturnNode = GetReturnNode())
+	{
+		for (URigVMPin* Pin : ReturnNode->GetPins())
+		{			
+			FRigVMGraphVariableDescription Description;
+			Description.Name = Pin->GetFName();
+			Description.CPPType = Pin->GetCPPType();
+			Description.CPPTypeObject = Pin->GetCPPTypeObject();
+			Outputs.Add(Description);			
+		}
+	}
+	return Outputs;
+}
+
 FRigVMGraphModifiedEvent& URigVMGraph::OnModified()
 {
 	return ModifiedEvent;
diff --git a/Engine/Source/Developer/RigVMDeveloper/Public/RigVMModel/RigVMGraph.h b/Engine/Source/Developer/RigVMDeveloper/Public/RigVMModel/RigVMGraph.h
--- a/Engine/Source/Developer/RigVMDeveloper/Public/RigVMModel/RigVMGraph.h
+++ b/Engine/Source/Developer/RigVMDeveloper/Public/RigVMModel/RigVMGraph.h
@@ -128,6 +128,10 @@ public:
 	// Returns the input arguments of this graph
 	UFUNCTION(BlueprintCallable, Category = RigVMGraph)
 	TArray<FRigVMGraphVariableDescription> GetInputArguments() const;
+	
+	// Returns the output arguments of this graph
+    UFUNCTION(BlueprintCallable, Category = RigVMGraph)
+    TArray<FRigVMGraphVariableDescription> GetOutputArguments() const;
 
 	// Returns the modified event, which can be used to 
 	// subscribe to changes happening within the Graph.
