Sequencer: Add Move/Delete Space Key Python/BP support based on latest feedback from Greg

diff --git a/Engine/Plugins/Animation/ControlRig/Source/ControlRigEditor/Private/ControlRigSequencerEditorLibrary.cpp b/Engine/Plugins/Animation/ControlRig/Source/ControlRigEditor/Private/ControlRigSequencerEditorLibrary.cpp
--- a/Engine/Plugins/Animation/ControlRig/Source/ControlRigEditor/Private/ControlRigSequencerEditorLibrary.cpp
+++ b/Engine/Plugins/Animation/ControlRig/Source/ControlRigEditor/Private/ControlRigSequencerEditorLibrary.cpp
@@ -2229,6 +2229,122 @@ bool UControlRigSequencerEditorLibrary::BakeControlRigSpace(ULevelSequence* InSe
 	return bValid;
 }
 
+bool UControlRigSequencerEditorLibrary::DeleteControlRigSpace(ULevelSequence* LevelSequence, UControlRig* ControlRig, FName ControlName, FFrameNumber InTime, ESequenceTimeUnit TimeUnit)
+{
+	TWeakPtr<ISequencer> WeakSequencer = GetSequencerFromAsset(LevelSequence);
+	bool bValid = false;
+
+	if (WeakSequencer.IsValid() && ControlRig && ControlName != NAME_None)
+	{
+		if (FRigControlElement* Element = ControlRig->FindControl(ControlName))
+		{
+			TSharedPtr<ISequencer>  Sequencer = WeakSequencer.Pin();
+			FSpaceChannelAndSection SpaceChannelAndSection = FControlRigSpaceChannelHelpers::FindSpaceChannelAndSectionForControl(ControlRig, ControlName, Sequencer.Get(), false /*bCreateIfNeeded*/);
+			if (SpaceChannelAndSection.SpaceChannel)
+			{
+
+				if (TimeUnit == ESequenceTimeUnit::DisplayRate)
+				{
+					InTime = FFrameRate::TransformTime(FFrameTime(InTime, 0), LevelSequence->GetMovieScene()->GetDisplayRate(), LevelSequence->GetMovieScene()->GetTickResolution()).RoundToFrame();
+				}
+				UMovieSceneControlRigParameterSection* ParamSection = Cast<UMovieSceneControlRigParameterSection>(SpaceChannelAndSection.SectionToKey);
+
+
+				TArray<FFrameNumber> OurKeyTimes;
+				TArray<FKeyHandle> OurKeyHandles;
+				TRange<FFrameNumber> CurrentFrameRange;
+				CurrentFrameRange.SetLowerBound(TRangeBound<FFrameNumber>(InTime));
+				CurrentFrameRange.SetUpperBound(TRangeBound<FFrameNumber>(InTime));
+				TMovieSceneChannelData<FMovieSceneControlRigSpaceBaseKey> ChannelInterface = SpaceChannelAndSection.SpaceChannel->GetData();
+				ChannelInterface.GetKeys(CurrentFrameRange, &OurKeyTimes, &OurKeyHandles);
+				if (OurKeyHandles.Num() > 0)
+				{
+					FScopedTransaction DeleteKeysTransaction(LOCTEXT("DeleteSpaceChannelKeys_Transaction", "Delete Space Channel Keys"));
+					ParamSection->TryModify();
+					SpaceChannelAndSection.SpaceChannel->DeleteKeys(OurKeyHandles);
+					WeakSequencer.Pin()->NotifyMovieSceneDataChanged(EMovieSceneDataChangeType::TrackValueChanged);
+					bValid = true;
+				}
+				else
+				{
+					UE_LOG(LogControlRig, Error, TEXT("No Keys To Delete At that Time"));
+					return false;
+				}
+			}
+			else
+			{
+				UE_LOG(LogControlRig, Error, TEXT("Can not find Space Channel"));
+				return false;
+			}
+		}
+		else
+		{
+			UE_LOG(LogControlRig, Error, TEXT("Can not find Control with that Name"));
+			return false;
+		}
+	}
+	return bValid;
+}
+
+bool UControlRigSequencerEditorLibrary::MoveControlRigSpace(ULevelSequence* LevelSequence, UControlRig* ControlRig, FName ControlName, FFrameNumber InTime, FFrameNumber InNewTime, ESequenceTimeUnit TimeUnit)
+{
+	TWeakPtr<ISequencer> WeakSequencer = GetSequencerFromAsset(LevelSequence);
+	bool bValid = false;
+
+	if (WeakSequencer.IsValid() && ControlRig && ControlName != NAME_None)
+	{
+		if (FRigControlElement* Element = ControlRig->FindControl(ControlName))
+		{
+			TSharedPtr<ISequencer>  Sequencer = WeakSequencer.Pin();
+			FSpaceChannelAndSection SpaceChannelAndSection = FControlRigSpaceChannelHelpers::FindSpaceChannelAndSectionForControl(ControlRig, ControlName, Sequencer.Get(), false /*bCreateIfNeeded*/);
+			if (SpaceChannelAndSection.SpaceChannel)
+			{
+
+				if (TimeUnit == ESequenceTimeUnit::DisplayRate)
+				{
+					InTime = FFrameRate::TransformTime(FFrameTime(InTime, 0), LevelSequence->GetMovieScene()->GetDisplayRate(), LevelSequence->GetMovieScene()->GetTickResolution()).RoundToFrame();
+				}
+				UMovieSceneControlRigParameterSection* ParamSection = Cast<UMovieSceneControlRigParameterSection>(SpaceChannelAndSection.SectionToKey);
+
+
+				TArray<FFrameNumber> OurKeyTimes;
+				TArray<FKeyHandle> OurKeyHandles;
+				TRange<FFrameNumber> CurrentFrameRange;
+				CurrentFrameRange.SetLowerBound(TRangeBound<FFrameNumber>(InTime));
+				CurrentFrameRange.SetUpperBound(TRangeBound<FFrameNumber>(InTime));
+				TMovieSceneChannelData<FMovieSceneControlRigSpaceBaseKey> ChannelInterface = SpaceChannelAndSection.SpaceChannel->GetData();
+				ChannelInterface.GetKeys(CurrentFrameRange, &OurKeyTimes, &OurKeyHandles);
+				if (OurKeyHandles.Num() > 0)
+				{
+					FScopedTransaction MoveKeys(LOCTEXT("MoveSpaceChannelKeys_Transaction", "Move Space Channel Keys"));
+					ParamSection->TryModify();
+					TArray<FFrameNumber> NewKeyTimes;
+					NewKeyTimes.Add(InNewTime);
+					SpaceChannelAndSection.SpaceChannel->SetKeyTimes(OurKeyHandles, NewKeyTimes);
+					WeakSequencer.Pin()->NotifyMovieSceneDataChanged(EMovieSceneDataChangeType::TrackValueChanged);
+					bValid = true;
+				}
+				else
+				{
+					UE_LOG(LogControlRig, Error, TEXT("No Key To Move At that Time"));
+					return false;
+				}
+			}
+			else
+			{
+				UE_LOG(LogControlRig, Error, TEXT("Can not find Space Channel"));
+				return false;
+			}
+		}
+		else
+		{
+			UE_LOG(LogControlRig, Error, TEXT("Can not find Control with that Name"));
+			return false;
+		}
+	}
+	return bValid;
+}
+
 bool UControlRigSequencerEditorLibrary::RenameControlRigControlChannels(ULevelSequence* LevelSequence, UControlRig* ControlRig, const TArray<FName>& InOldControlNames, const TArray<FName>& InNewControlNames)
 {
 	if (LevelSequence == nullptr || ControlRig == nullptr)
diff --git a/Engine/Plugins/Animation/ControlRig/Source/ControlRigEditor/Public/ControlRigSequencerEditorLibrary.h b/Engine/Plugins/Animation/ControlRig/Source/ControlRigEditor/Public/ControlRigSequencerEditorLibrary.h
--- a/Engine/Plugins/Animation/ControlRig/Source/ControlRigEditor/Public/ControlRigSequencerEditorLibrary.h
+++ b/Engine/Plugins/Animation/ControlRig/Source/ControlRigEditor/Public/ControlRigSequencerEditorLibrary.h
@@ -850,7 +850,33 @@ public:
 	UFUNCTION(BlueprintCallable, Category = "Editor Scripting | Sequencer Tools | Control Rig")
 	static bool BakeControlRigSpace(ULevelSequence* InSequence, UControlRig* InControlRig, const TArray<FName>& InControlNames, FRigSpacePickerBakeSettings InSettings, ESequenceTimeUnit TimeUnit = ESequenceTimeUnit::DisplayRate);
 	
-	
+	/** Delete the Control Rig Space Key for the Control at the specified time. This will delete any attached Control Rig keys at this time and will perform any needed compensation to the new space.
+	*
+	* @param InSequence Sequence to set the space
+	* @param InControlRig ControlRig with the Control
+	* @param InControlName The name of the Control
+	* @param InTime Time to delete the space.
+	* @param TimeUnit Unit for the InTime, either in display rate or tick resolution
+	* @return Will return false if function fails,  for example if there is no key at this time it will fail.
+	*/
+	UFUNCTION(BlueprintCallable, Category = "Editor Scripting | Sequencer Tools | Control Rig")
+	static bool DeleteControlRigSpace(ULevelSequence* InSequence, UControlRig* InControlRig, FName InControlName, FFrameNumber InTime, ESequenceTimeUnit TimeUnit = ESequenceTimeUnit::DisplayRate);
+
+	/** Move the Control Rig Space Key for the Control at the specified time to the new time. This will also move any Control Rig keys at this space switch boundary.
+	*
+	* @param InSequence Sequence to set the space
+	* @param InControlRig ControlRig with the Control
+	* @param InControlName The name of the Control
+	* @param InTime Original time of the space key
+	* @param InNewTime New time for the space key
+	* @param TimeUnit Unit for the time params, either in display rate or tick resolution
+	* @return Will return false if function fails, for example if there is no key at this time it will fail, or if the new time is invalid it could fail also
+	*/
+	UFUNCTION(BlueprintCallable, Category = "Editor Scripting | Sequencer Tools | Control Rig")
+	static bool MoveControlRigSpace(ULevelSequence* InSequence, UControlRig* InControlRig, FName InControlName, FFrameNumber InTime, FFrameNumber InNewTime, ESequenceTimeUnit TimeUnit = ESequenceTimeUnit::DisplayRate);
+
+
+
 	/** Rename the Control Rig Channels in Sequencer to the specified new control names, which should be present on the Control Rig
 	* @param InSequence Sequence to rename controls
 	* @param InControlRig ControlRig to rename controls
