Control Rig: Fix static analysis issue with RigHierarchy.cpp

diff --git a/Engine/Plugins/Animation/ControlRig/Source/ControlRig/Private/Rigs/RigHierarchy.cpp b/Engine/Plugins/Animation/ControlRig/Source/ControlRig/Private/Rigs/RigHierarchy.cpp
--- a/Engine/Plugins/Animation/ControlRig/Source/ControlRig/Private/Rigs/RigHierarchy.cpp
+++ b/Engine/Plugins/Animation/ControlRig/Source/ControlRig/Private/Rigs/RigHierarchy.cpp
@@ -2101,20 +2101,20 @@ void URigHierarchy::Notify(ERigHierarchyNotification InNotifType, const FRigBase
 		case ERigHierarchyNotification::ParentChanged:
 		case ERigHierarchyNotification::ParentWeightsChanged:
 		{
-			ensure(InElement != nullptr);
-				
-			for(FRigHierarchyListener& Listener : ListeningHierarchies)
+			if (ensure(InElement != nullptr))
 			{
-				URigHierarchy* ListeningHierarchy = Listener.Hierarchy.Get();
-				if (ListeningHierarchy)
-				{			
-					if(const FRigBaseElement* ListeningElement = ListeningHierarchy->Find( InElement->GetKey()))
-					{
-						ListeningHierarchy->Notify(InNotifType, ListeningElement);
+				for(FRigHierarchyListener& Listener : ListeningHierarchies)
+				{
+					URigHierarchy* ListeningHierarchy = Listener.Hierarchy.Get();
+					if (ListeningHierarchy)
+					{			
+						if(const FRigBaseElement* ListeningElement = ListeningHierarchy->Find( InElement->GetKey()))
+						{
+							ListeningHierarchy->Notify(InNotifType, ListeningElement);
+						}
 					}
 				}
 			}
-
 			break;
 		}
 		default:
