Control Rig: Relay the loading of gizmo libraries to OnEndLoadPackage

diff --git a/Engine/Plugins/Animation/ControlRig/Source/ControlRigDeveloper/Private/ControlRigBlueprint.cpp b/Engine/Plugins/Animation/ControlRig/Source/ControlRigDeveloper/Private/ControlRigBlueprint.cpp
--- a/Engine/Plugins/Animation/ControlRig/Source/ControlRigDeveloper/Private/ControlRigBlueprint.cpp
+++ b/Engine/Plugins/Animation/ControlRig/Source/ControlRigDeveloper/Private/ControlRigBlueprint.cpp
@@ -536,53 +536,29 @@ void UControlRigBlueprint::PostLoad()
 
 		if(GizmoLibrary_DEPRECATED.IsValid())
 		{
-			ShapeLibraries.Add(GizmoLibrary_DEPRECATED.LoadSynchronous());
-			GizmoLibrary_DEPRECATED.Reset();
+			ShapeLibrariesToLoadOnPackageLoaded.Add(GizmoLibrary_DEPRECATED->GetPathName());
 		}
 		else
 		{
-			ShapeLibraries.Add(LoadObject<UControlRigShapeLibrary>(nullptr, TEXT("/ControlRig/Controls/DefaultGizmoLibrary.DefaultGizmoLibrary")));
+			static const FString DefaultGizmoLibraryPath = TEXT("/ControlRig/Controls/DefaultGizmoLibrary.DefaultGizmoLibrary");
+			ShapeLibrariesToLoadOnPackageLoaded.Add(DefaultGizmoLibraryPath);
 		}
 
-		// also walk over all controls and check if any of them were using the "default" gizmo
-		Hierarchy->ForEach<FRigControlElement>([](FRigControlElement* ControlElement) -> bool
-		{
-			static FName PreviousDefault = TEXT("Gizmo");
-			static FName NewDefault = FControlRigShapeDefinition().ShapeName;
-			if(ControlElement->Settings.ShapeName == PreviousDefault)
-			{
-				ControlElement->Settings.ShapeName = NewDefault;
-			}
-			return true;
-		});
-
 		UControlRigBlueprintGeneratedClass* RigClass = GetControlRigBlueprintGeneratedClass();
 		UControlRig* CDO = Cast<UControlRig>(RigClass->GetDefaultObject(false /* create if needed */));
 
-		CDO->ShapeLibraries = ShapeLibraries;
-		CDO->GizmoLibrary_DEPRECATED.Reset();
-		if(CDO->GetHierarchy())
-		{
-			CDO->GetHierarchy()->CopyHierarchy(Hierarchy);
-			CDO->Initialize(true);
-		}
-
 		TArray<UObject*> ArchetypeInstances;
 		CDO->GetArchetypeInstances(ArchetypeInstances);
+		ArchetypeInstances.Insert(CDO, 0);
+
 		for (UObject* Instance : ArchetypeInstances)
 		{
 			if (UControlRig* InstanceRig = Cast<UControlRig>(Instance))
 			{
-				InstanceRig->ShapeLibraries = ShapeLibraries;
+				InstanceRig->ShapeLibraries.Reset();
 				InstanceRig->GizmoLibrary_DEPRECATED.Reset();
-				if(InstanceRig->GetHierarchy())
-				{
-					InstanceRig->GetHierarchy()->CopyHierarchy(Hierarchy);
-					InstanceRig->Initialize(true);
-				}
 			}
 		}
-
 	}
 
 #if WITH_EDITOR
@@ -613,6 +589,31 @@ void UControlRigBlueprint::HandlePackageDone(TConstArrayView<UPackage*> InPackag
 	}
 
 	FCoreUObjectDelegates::OnEndLoadPackage.RemoveAll(this);
+
+	if (ShapeLibrariesToLoadOnPackageLoaded.Num() > 0)
+	{
+		for(const FString& ShapeLibraryToLoadOnPackageLoaded : ShapeLibrariesToLoadOnPackageLoaded)
+		{
+			ShapeLibraries.Add(LoadObject<UControlRigShapeLibrary>(nullptr, *ShapeLibraryToLoadOnPackageLoaded));
+		}
+
+		UControlRigBlueprintGeneratedClass* RigClass = GetControlRigBlueprintGeneratedClass();
+		UControlRig* CDO = Cast<UControlRig>(RigClass->GetDefaultObject(false /* create if needed */));
+
+		TArray<UObject*> ArchetypeInstances;
+		CDO->GetArchetypeInstances(ArchetypeInstances);
+		ArchetypeInstances.Insert(CDO, 0);
+
+		for (UObject* Instance : ArchetypeInstances)
+		{
+			if (UControlRig* InstanceRig = Cast<UControlRig>(Instance))
+			{
+				InstanceRig->ShapeLibraries = ShapeLibraries;
+			}
+		}
+
+		ShapeLibrariesToLoadOnPackageLoaded.Reset();
+	}
 	
 	PropagateHierarchyFromBPToInstances();
 	RecompileVM();
diff --git a/Engine/Plugins/Animation/ControlRig/Source/ControlRigDeveloper/Public/ControlRigBlueprint.h b/Engine/Plugins/Animation/ControlRig/Source/ControlRigDeveloper/Public/ControlRigBlueprint.h
--- a/Engine/Plugins/Animation/ControlRig/Source/ControlRigDeveloper/Public/ControlRigBlueprint.h
+++ b/Engine/Plugins/Animation/ControlRig/Source/ControlRigDeveloper/Public/ControlRigBlueprint.h
@@ -685,6 +685,9 @@ public:
 	/** Recomputes the instruction breakpoints given the node breakpoints in the blueprint */
 	void RefreshControlRigBreakpoints();
 
+	/** Shape libraries to load during package load completed */ 
+	TArray<FString> ShapeLibrariesToLoadOnPackageLoaded;
+
 #endif
 
 private:
