Control Rig: Root VM / Hierarchy for Native Control Rigs

diff --git a/Engine/Plugins/Animation/ControlRig/Source/ControlRig/Private/ControlRig.cpp b/Engine/Plugins/Animation/ControlRig/Source/ControlRig/Private/ControlRig.cpp
--- a/Engine/Plugins/Animation/ControlRig/Source/ControlRig/Private/ControlRig.cpp
+++ b/Engine/Plugins/Animation/ControlRig/Source/ControlRig/Private/ControlRig.cpp
@@ -2816,6 +2816,16 @@ void UControlRig::PostInitInstance(UControlRig* InCDO)
 		VM->CopyFrom(InCDO->GetVM());
 		DynamicHierarchy->CopyHierarchy(InCDO->GetHierarchy());
 	}
+	else // we are the CDO
+	{
+		// for default objects we need to check if the CDO is rooted. specialized Control Rigs
+		// such as the FK control rig may not have a root since they are part of a C++ package.
+		if(!IsRooted() && GetClass()->IsNative())
+		{
+			VM->AddToRoot();
+			DynamicHierarchy->AddToRoot();
+		}
+	}
 }
 
 uint32 UControlRig::GetHashForInitializeVMSnapShot()
