Control Rig: refactor hierarchy cycle checking to avoid infinite recursion + added unit test for it

diff --git a/Engine/Plugins/Animation/ControlRig/Source/ControlRig/Private/Rigs/RigHierarchy.cpp b/Engine/Plugins/Animation/ControlRig/Source/ControlRig/Private/Rigs/RigHierarchy.cpp
--- a/Engine/Plugins/Animation/ControlRig/Source/ControlRig/Private/Rigs/RigHierarchy.cpp
+++ b/Engine/Plugins/Animation/ControlRig/Source/ControlRig/Private/Rigs/RigHierarchy.cpp
@@ -3008,11 +3008,18 @@ FRigElementKey URigHierarchy::GetPreviousParent(const FRigElementKey& InKey) con
 
 bool URigHierarchy::IsParentedTo(FRigBaseElement* InChild, FRigBaseElement* InParent, const TElementDependencyMap& InDependencyMap) const
 {
-	return IsDependentOn(InChild, InParent, InDependencyMap);
+	TArray<bool> ElementsVisited;
+	return IsDependentOn(InChild, InParent, ElementsVisited, InDependencyMap);
 }
 
-bool URigHierarchy::IsDependentOn(FRigBaseElement* InDependent, FRigBaseElement* InDependency, const TElementDependencyMap& InDependencyMap) const
+bool URigHierarchy::IsDependentOn(FRigBaseElement* InDependent, FRigBaseElement* InDependency, TArray<bool>& InElementsVisited, const TElementDependencyMap& InDependencyMap) const
 {
+	if (InElementsVisited.Num() != Elements.Num())
+	{
+		InElementsVisited.Reset();
+		InElementsVisited.AddZeroed(Elements.Num());
+	}
+
 	if((InDependent == nullptr) || (InDependency == nullptr))
 	{
 		return false;
@@ -3023,30 +3030,31 @@ bool URigHierarchy::IsDependentOn(FRigBaseElement* InDependent, FRigBaseElement*
 		return true;
 	}
 
-	if(const FRigSingleParentElement* SingleParentElement = Cast<FRigSingleParentElement>(InDependent))
+	const int32 DependentElementIndex = InDependent->GetIndex();
+
+	if (!InElementsVisited.IsValidIndex(DependentElementIndex))
 	{
-		if(SingleParentElement->ParentElement == InDependency)
-		{
-			return true;
-		}
-		if(IsDependentOn(SingleParentElement->ParentElement, InDependency, InDependencyMap))
-		{
-			return true;
-		}
+		return false;
+	}
+	
+	if (InElementsVisited[DependentElementIndex])
+	{
+		return false;
 	}
 
-	if(const FRigMultiParentElement* MultiParentElement = Cast<FRigMultiParentElement>(InDependent))
+	InElementsVisited[DependentElementIndex] = true;
+
+	// collect all possible parents of the dependent
+	TArray<FRigBaseElement*> DependentParents;
+	if(const FRigSingleParentElement* SingleParentElement = Cast<FRigSingleParentElement>(InDependent))
+	{
+		 DependentParents.AddUnique(SingleParentElement->ParentElement);
+	}
+	else if(const FRigMultiParentElement* MultiParentElement = Cast<FRigMultiParentElement>(InDependent))
 	{
 		for(const FRigElementParentConstraint& ParentConstraint : MultiParentElement->ParentConstraints)
 		{
-			if(ParentConstraint.ParentElement == InDependency)
-			{
-				return true;
-			}
-			if(IsDependentOn(ParentConstraint.ParentElement, InDependency, InDependencyMap))
-			{
-				return true;
-			}
+			 DependentParents.AddUnique(ParentConstraint.ParentElement);
 		}
 	}
 
@@ -3054,14 +3062,18 @@ bool URigHierarchy::IsDependentOn(FRigBaseElement* InDependent, FRigBaseElement*
 	if(const TArray<int32>* DependentIndicesPtr = InDependencyMap.Find(InDependent->GetIndex()))
 	{
 		const TArray<int32>& DependentIndices = *DependentIndicesPtr;
-		for(int32 DependentIndex : DependentIndices)
+		for(const int32 DependentIndex : DependentIndices)
 		{
 			ensure(Elements.IsValidIndex(DependentIndex));
-			
-			if(IsDependentOn(Elements[DependentIndex], InDependency, InDependencyMap))
-			{
-				return true;
-			}
+			DependentParents.AddUnique(Elements[DependentIndex]);
+		}
+	}
+
+	for (FRigBaseElement* DependentParent :  DependentParents)
+	{
+		if(IsDependentOn(DependentParent, InDependency, InElementsVisited, InDependencyMap))
+		{
+			return true;
 		}
 	}
 
diff --git a/Engine/Plugins/Animation/ControlRig/Source/ControlRig/Public/Rigs/RigHierarchy.h b/Engine/Plugins/Animation/ControlRig/Source/ControlRig/Public/Rigs/RigHierarchy.h
--- a/Engine/Plugins/Animation/ControlRig/Source/ControlRig/Public/Rigs/RigHierarchy.h
+++ b/Engine/Plugins/Animation/ControlRig/Source/ControlRig/Public/Rigs/RigHierarchy.h
@@ -2449,10 +2449,11 @@ public:
 	 * Returns true if an element is affected to another element
 	 * @param InDependent The dependent element to check for a dependency
 	 * @param InDependency The dependency element to check for
+	 * @param InElementsVisited An array to keep track of whether an element is visited to avoid infinite recursion
 	 * @param InDependencyMap An additional map of dependencies to respect
 	 * @return True if the child is parented to the parent
 	 */
-	bool IsDependentOn(FRigBaseElement* InDependent, FRigBaseElement* InDependency, const TElementDependencyMap& InDependencyMap = TElementDependencyMap()) const;
+	bool IsDependentOn(FRigBaseElement* InDependent, FRigBaseElement* InDependency, TArray<bool>& InElementsVisited, const TElementDependencyMap& InDependencyMap = TElementDependencyMap()) const;
 
 	/**
 	 * Returns a reference to the suspend notifications flag
