Sequencer: Space Switching:  We remove duplicate spaces, but when doing so need to make sure to update counting index.

diff --git a/Engine/Plugins/Animation/ControlRig/Source/ControlRig/Private/Sequencer/MovieSceneControlRigSpaceChannel.cpp b/Engine/Plugins/Animation/ControlRig/Source/ControlRig/Private/Sequencer/MovieSceneControlRigSpaceChannel.cpp
--- a/Engine/Plugins/Animation/ControlRig/Source/ControlRig/Private/Sequencer/MovieSceneControlRigSpaceChannel.cpp
+++ b/Engine/Plugins/Animation/ControlRig/Source/ControlRig/Private/Sequencer/MovieSceneControlRigSpaceChannel.cpp
@@ -133,6 +133,7 @@ TArray <FSpaceRange> FMovieSceneControlRigSpaceChannel::FindSpaceIntervals()
 				NextIndex = Index + 1;
 				if (Range.Key == KeyValues[NextIndex])
 				{
+					Index++;
 					ToDelete.Add(NextIndex);
 				}
 				else
@@ -158,6 +159,7 @@ TArray <FSpaceRange> FMovieSceneControlRigSpaceChannel::FindSpaceIntervals()
 		}
 
 		GetData().DeleteKeys(DeleteKeyHandles);
+
 	}
 	return Ranges;
 }
diff --git a/Engine/Plugins/Animation/ControlRig/Source/ControlRigEditor/Private/Sequencer/ControlRigSpaceChannelEditors.cpp b/Engine/Plugins/Animation/ControlRig/Source/ControlRigEditor/Private/Sequencer/ControlRigSpaceChannelEditors.cpp
--- a/Engine/Plugins/Animation/ControlRig/Source/ControlRigEditor/Private/Sequencer/ControlRigSpaceChannelEditors.cpp
+++ b/Engine/Plugins/Animation/ControlRig/Source/ControlRigEditor/Private/Sequencer/ControlRigSpaceChannelEditors.cpp
@@ -442,6 +442,7 @@ FKeyHandle FControlRigSpaceChannelHelpers::SequencerKeyControlRigSpaceChannel(UC
 			ControlRig->SetControlGlobalTransform(ControlKey.Name, ControlWorldTransforms[FramesIndex], true, Context);
 			FramesIndex++;
 		}
+		Channel->FindSpaceIntervals();
 		Sequencer->NotifyMovieSceneDataChanged(EMovieSceneDataChangeType::MovieSceneStructureItemAdded);
 
 		Channel->GetUniqueSpaceList(&AfterKeys);
