Control Rig: Dynamic hierarchy nodes to send warnings rather than errors

diff --git a/Engine/Plugins/Animation/ControlRig/Source/ControlRig/Private/Units/Execution/RigUnit_DynamicHierarchy.cpp b/Engine/Plugins/Animation/ControlRig/Source/ControlRig/Private/Units/Execution/RigUnit_DynamicHierarchy.cpp
--- a/Engine/Plugins/Animation/ControlRig/Source/ControlRig/Private/Units/Execution/RigUnit_DynamicHierarchy.cpp
+++ b/Engine/Plugins/Animation/ControlRig/Source/ControlRig/Private/Units/Execution/RigUnit_DynamicHierarchy.cpp
@@ -14,14 +14,14 @@ FRigUnit_AddParent_Execute()
 	FRigTransformElement* ChildElement = ExecuteContext.Hierarchy->Find<FRigTransformElement>(Child);
 	if(ChildElement == nullptr)
 	{
-		UE_CONTROLRIG_RIGUNIT_REPORT_ERROR(TEXT("Child item %s does not exist."), *Child.ToString())
+		UE_CONTROLRIG_RIGUNIT_REPORT_WARNING(TEXT("Child item %s does not exist."), *Child.ToString())
 		return;
 	}
 
 	FRigTransformElement* ParentElement = ExecuteContext.Hierarchy->Find<FRigTransformElement>(Parent);
 	if(ParentElement == nullptr)
 	{
-		UE_CONTROLRIG_RIGUNIT_REPORT_ERROR(TEXT("Parent item %s does not exist."), *Parent.ToString())
+		UE_CONTROLRIG_RIGUNIT_REPORT_WARNING(TEXT("Parent item %s does not exist."), *Parent.ToString())
 		return;
 	}
 
@@ -41,12 +41,12 @@ FRigUnit_SwitchParent_Execute()
 	FRigTransformElement* ChildElement = ExecuteContext.Hierarchy->Find<FRigTransformElement>(Child);
 	if(ChildElement == nullptr)
 	{
-		UE_CONTROLRIG_RIGUNIT_REPORT_ERROR(TEXT("Child item %s does not exist."), *Child.ToString())
+		UE_CONTROLRIG_RIGUNIT_REPORT_WARNING(TEXT("Child item %s does not exist."), *Child.ToString())
 		return;
 	}
 	if(!ChildElement->IsA<FRigMultiParentElement>())
 	{
-		UE_CONTROLRIG_RIGUNIT_REPORT_ERROR(TEXT("Child item %s cannot be space switched (only Nulls and Controls can)."), *Child.ToString())
+		UE_CONTROLRIG_RIGUNIT_REPORT_WARNING(TEXT("Child item %s cannot be space switched (only Nulls and Controls can)."), *Child.ToString())
 		return;
 	}
 
@@ -57,13 +57,13 @@ FRigUnit_SwitchParent_Execute()
 		ParentElement = ExecuteContext.Hierarchy->Find<FRigTransformElement>(Parent);
 		if(ParentElement == nullptr)
 		{
-			UE_CONTROLRIG_RIGUNIT_REPORT_ERROR(TEXT("Parent item %s does not exist."), *Parent.ToString())
+			UE_CONTROLRIG_RIGUNIT_REPORT_WARNING(TEXT("Parent item %s does not exist."), *Parent.ToString())
 			return;
 		}
 
 		if(!ParentElement->IsA<FRigTransformElement>())
 		{
-			UE_CONTROLRIG_RIGUNIT_REPORT_ERROR(TEXT("Parent item %s does not have a transform."), *Parent.ToString())
+			UE_CONTROLRIG_RIGUNIT_REPORT_WARNING(TEXT("Parent item %s does not have a transform."), *Parent.ToString())
 			return;
 		}
 	}
@@ -106,7 +106,7 @@ FRigUnit_SwitchParent_Execute()
 			{
 				if(!FailureReason.IsEmpty())
 				{
-					UE_CONTROLRIG_RIGUNIT_REPORT_ERROR(TEXT("%s"), *FailureReason);
+					UE_CONTROLRIG_RIGUNIT_REPORT_WARNING(TEXT("%s"), *FailureReason);
 				}
 				return;
 			}
@@ -132,7 +132,7 @@ FRigUnit_HierarchyGetParentWeightsArray_Execute()
 	FRigBaseElement* ChildElement = Context.Hierarchy->Find(Child);
 	if(ChildElement == nullptr)
 	{
-		UE_CONTROLRIG_RIGUNIT_REPORT_ERROR(TEXT("Item %s does not exist."), *Child.ToString())
+		UE_CONTROLRIG_RIGUNIT_REPORT_WARNING(TEXT("Item %s does not exist."), *Child.ToString())
 		return;
 	}
 	
@@ -150,13 +150,13 @@ FRigUnit_HierarchySetParentWeights_Execute()
 	FRigBaseElement* ChildElement = Context.Hierarchy->Find(Child);
 	if(ChildElement == nullptr)
 	{
-		UE_CONTROLRIG_RIGUNIT_REPORT_ERROR(TEXT("Item %s does not exist."), *Child.ToString())
+		UE_CONTROLRIG_RIGUNIT_REPORT_WARNING(TEXT("Item %s does not exist."), *Child.ToString())
 		return;
 	}
 
 	if(Weights.Num() != ExecuteContext.Hierarchy->GetNumberOfParents(ChildElement))
 	{
-		UE_CONTROLRIG_RIGUNIT_REPORT_ERROR(TEXT("Provided incorrect number of weights(%d) for %s - expected %d."), Weights.Num(), *Child.ToString(), ExecuteContext.Hierarchy->GetNumberOfParents(Child))
+		UE_CONTROLRIG_RIGUNIT_REPORT_WARNING(TEXT("Provided incorrect number of weights(%d) for %s - expected %d."), Weights.Num(), *Child.ToString(), ExecuteContext.Hierarchy->GetNumberOfParents(Child))
 		return;
 	}
 
