Control Rig: Add shape libraries to object reference collection

diff --git a/Engine/Plugins/Animation/ControlRig/Source/ControlRigDeveloper/Private/ControlRigBlueprint.cpp b/Engine/Plugins/Animation/ControlRig/Source/ControlRigDeveloper/Private/ControlRigBlueprint.cpp
--- a/Engine/Plugins/Animation/ControlRig/Source/ControlRigDeveloper/Private/ControlRigBlueprint.cpp
+++ b/Engine/Plugins/Animation/ControlRig/Source/ControlRigDeveloper/Private/ControlRigBlueprint.cpp
@@ -315,6 +315,15 @@ void UControlRigBlueprint::Serialize(FArchive& Ar)
 		{
 			Ar << ReferencedBlueprints;
 		}
+
+		for(const TSoftObjectPtr<UControlRigShapeLibrary>& ShapeLibraryPtr : ShapeLibraries)
+		{
+			if(ShapeLibraryPtr.IsValid())
+			{
+				UControlRigShapeLibrary* ShapeLibrary = ShapeLibraryPtr.Get();
+				Ar << ShapeLibrary;
+			}
+		}
 	}
 }
 
