Control Rig: Fix crash when creating rotator in function

diff --git a/Engine/Plugins/Animation/ControlRig/Source/ControlRigEditor/Private/Graph/SControlRigGraphNode.cpp b/Engine/Plugins/Animation/ControlRig/Source/ControlRigEditor/Private/Graph/SControlRigGraphNode.cpp
--- a/Engine/Plugins/Animation/ControlRig/Source/ControlRigEditor/Private/Graph/SControlRigGraphNode.cpp
+++ b/Engine/Plugins/Animation/ControlRig/Source/ControlRigEditor/Private/Graph/SControlRigGraphNode.cpp
@@ -111,9 +111,12 @@ void SControlRigGraphNode::Construct( const FArguments& InArgs )
 			if (InPin->GetCPPType() == TEXT("FRotator"))
 			{
 				TArray<URigVMPin*> SubPins = InPin->GetSubPins();
-				OutPins.Add(SubPins[2]);	
-				OutPins.Add(SubPins[0]);	
-				OutPins.Add(SubPins[1]);	
+				if (SubPins.Num() == 3)
+				{
+					OutPins.Add(SubPins[2]);	
+					OutPins.Add(SubPins[0]);	
+					OutPins.Add(SubPins[1]);
+				}	
 			}
 			else
 			{				
@@ -998,8 +1001,11 @@ void SControlRigGraphNode::GetNodeInfoPopups(FNodeInfoContext* Context, TArray<F
 								TArray<FString> Values;
 								FString TrimmedText = DefaultValues[0].LeftChop(1).RightChop(1); //Remove ()
 								TrimmedText.ParseIntoArray(Values, TEXT(","));
-								Values.Swap(0, 1);
-								Values.Swap(0, 2);
+								if (Values.Num() == 3)
+								{
+									Values.Swap(0, 1);
+									Values.Swap(0, 2);
+								}
 								WatchText = FString(TEXT("(")) + FString::Join(Values, TEXT(",")) + FString(TEXT(")"));
 								UScriptStruct* RotatorStruct = TBaseStructure<FRotator>::Get();
 								for (TFieldIterator<FProperty> It(RotatorStruct); It; ++It)
