Control Rig: React to control changes in tree view

diff --git a/Engine/Plugins/Animation/ControlRig/Source/ControlRigEditor/Private/Editor/SRigHierarchy.cpp b/Engine/Plugins/Animation/ControlRig/Source/ControlRigEditor/Private/Editor/SRigHierarchy.cpp
--- a/Engine/Plugins/Animation/ControlRig/Source/ControlRigEditor/Private/Editor/SRigHierarchy.cpp
+++ b/Engine/Plugins/Animation/ControlRig/Source/ControlRigEditor/Private/Editor/SRigHierarchy.cpp
@@ -748,6 +748,26 @@ void SRigHierarchy::OnHierarchyModified(ERigHierarchyNotification InNotif, URigH
 			}
 			break;
 		}
+		case ERigHierarchyNotification::ControlSettingChanged:
+		{
+			// update color and other settings of the item
+			if(InElement && InElement->GetType() == ERigElementType::Control)
+			{
+				for (int32 RootIndex = 0; RootIndex < TreeView->RootElements.Num(); ++RootIndex)
+				{
+					TSharedPtr<FRigTreeElement> TreeElement = TreeView->FindElement(InElement->GetKey(), TreeView->RootElements[RootIndex]);
+					if (TreeElement.IsValid())
+					{
+						if(URigHierarchy* Hierarchy = GetDebuggedHierarchy())
+						{
+							const FRigTreeDisplaySettings& Settings = TreeView->GetRigTreeDelegates().GetDisplaySettings();
+							TreeElement->RefreshDisplaySettings(Hierarchy, Settings);
+						}
+					}
+				}
+			}
+			break;
+		}
 		default:
 		{
 			break;
diff --git a/Engine/Plugins/Animation/ControlRig/Source/ControlRigEditor/Private/Editor/SRigHierarchyTreeView.cpp b/Engine/Plugins/Animation/ControlRig/Source/ControlRigEditor/Private/Editor/SRigHierarchyTreeView.cpp
--- a/Engine/Plugins/Animation/ControlRig/Source/ControlRigEditor/Private/Editor/SRigHierarchyTreeView.cpp
+++ b/Engine/Plugins/Animation/ControlRig/Source/ControlRigEditor/Private/Editor/SRigHierarchyTreeView.cpp
@@ -41,6 +41,9 @@ FRigTreeElement::FRigTreeElement(const FRigElementKey& InKey, TWeakPtr<SRigHiera
 			{
 				bIsTransient = ControlElement->Settings.bIsTransientControl;
 			}
+
+			const FRigTreeDisplaySettings& Settings = InTreeView.Pin()->GetRigTreeDelegates().GetDisplaySettings();
+			RefreshDisplaySettings(Hierarchy, Settings);
 		}
 	}
 }
@@ -59,6 +62,23 @@ void FRigTreeElement::RequestRename()
 	}
 }
 
+void FRigTreeElement::RefreshDisplaySettings(const URigHierarchy* InHierarchy, const FRigTreeDisplaySettings& InSettings)
+{
+	const TPair<const FSlateBrush*, FSlateColor> Result = SRigHierarchyItem::GetBrushForElementType(InHierarchy, Key);
+
+	IconBrush = Result.Key;
+	IconColor = Result.Value;
+	if(IconColor.IsColorSpecified() && InSettings.bShowIconColors)
+	{
+		IconColor = FilterResult == ERigTreeFilterResult::Shown ? Result.Value : FSlateColor(Result.Value.GetSpecifiedColor() * 0.5f);
+	}
+	else
+	{
+		IconColor = FilterResult == ERigTreeFilterResult::Shown ? FSlateColor::UseForeground() : FSlateColor(FLinearColor::Gray * 0.5f);
+	}
+	TextColor = FilterResult == ERigTreeFilterResult::Shown ? FSlateColor::UseForeground() : FSlateColor(FLinearColor::Gray * 0.5f);
+}
+
 //////////////////////////////////////////////////////////////
 /// SRigHierarchyItem
 ///////////////////////////////////////////////////////////
@@ -88,19 +108,6 @@ void SRigHierarchyItem::Construct(const FArguments& InArgs, const TSharedRef<STa
 
 	TSharedPtr< SInlineEditableTextBlock > InlineWidget;
 
-	TPair<const FSlateBrush*, FSlateColor> Result = GetBrushForElementType(Delegates.GetHierarchy(), InRigTreeElement->Key);
-	const FSlateBrush* Brush = Result.Key;
-	FSlateColor IconColor = Result.Value;
-	if(IconColor.IsColorSpecified() && InSettings.bShowIconColors)
-	{
-		IconColor = InRigTreeElement->FilterResult == ERigTreeFilterResult::Shown ? Result.Value : FSlateColor(Result.Value.GetSpecifiedColor() * 0.5f);
-	}
-	else
-	{
-		IconColor = InRigTreeElement->FilterResult == ERigTreeFilterResult::Shown ? FSlateColor::UseForeground() : FSlateColor(FLinearColor::Gray * 0.5f);
-	}
-	const FSlateColor TextColor = InRigTreeElement->FilterResult == ERigTreeFilterResult::Shown ? FSlateColor::UseForeground() : FSlateColor(FLinearColor::Gray * 0.5f);
-
 	STableRow<TSharedPtr<FRigTreeElement>>::Construct(
 		STableRow<TSharedPtr<FRigTreeElement>>::FArguments()
 		.OnDragDetected(Delegates.OnDragDetected)
@@ -118,8 +125,22 @@ void SRigHierarchyItem::Construct(const FArguments& InArgs, const TSharedRef<STa
 			.Padding(FMargin(0.f, 0.f, 3.f, 0.f))
 			[
 				SNew(SImage)
-				.Image(Brush)
-				.ColorAndOpacity(IconColor)
+				.Image_Lambda([this]() -> const FSlateBrush*
+				{
+					if(WeakRigTreeElement.IsValid())
+					{
+						return WeakRigTreeElement.Pin()->IconBrush;
+					}
+					return nullptr;
+				})
+				.ColorAndOpacity_Lambda([this]()
+				{
+					if(WeakRigTreeElement.IsValid())
+					{
+						return WeakRigTreeElement.Pin()->IconColor;
+					}
+					return FSlateColor::UseForeground();
+				})
 			]
 			+ SHorizontalBox::Slot()
 			.AutoWidth()
@@ -130,7 +151,14 @@ void SRigHierarchyItem::Construct(const FArguments& InArgs, const TSharedRef<STa
 				.OnVerifyTextChanged(this, &SRigHierarchyItem::OnVerifyNameChanged)
 				.OnTextCommitted(this, &SRigHierarchyItem::OnNameCommitted)
 				.MultiLine(false)
-				.ColorAndOpacity(TextColor)
+				.ColorAndOpacity_Lambda([this]()
+				{
+					if(WeakRigTreeElement.IsValid())
+					{
+						return WeakRigTreeElement.Pin()->TextColor;
+					}
+					return FSlateColor::UseForeground();
+				})
 			]
 		], OwnerTable);
 
diff --git a/Engine/Plugins/Animation/ControlRig/Source/ControlRigEditor/Private/Editor/SRigHierarchyTreeView.h b/Engine/Plugins/Animation/ControlRig/Source/ControlRigEditor/Private/Editor/SRigHierarchyTreeView.h
--- a/Engine/Plugins/Animation/ControlRig/Source/ControlRigEditor/Private/Editor/SRigHierarchyTreeView.h
+++ b/Engine/Plugins/Animation/ControlRig/Source/ControlRigEditor/Private/Editor/SRigHierarchyTreeView.h
@@ -178,12 +178,23 @@ public:
 
 	void RequestRename();
 
+	void RefreshDisplaySettings(const URigHierarchy* InHierarchy, const FRigTreeDisplaySettings& InSettings);
+
 	/** Delegate for when the context menu requests a rename */
 	DECLARE_DELEGATE(FOnRenameRequested);
 	FOnRenameRequested OnRenameRequested;
 
 	/** The current filter result */
 	ERigTreeFilterResult FilterResult;
+
+	/** The brush to use when rendering an icon */
+	const FSlateBrush* IconBrush;;
+
+	/** The color to use when rendering an icon */
+	FSlateColor IconColor;
+
+	/** The color to use when rendering the label text */
+	FSlateColor TextColor;
 };
 
 class SRigHierarchyItem : public STableRow<TSharedPtr<FRigTreeElement>>
