Control Rig: Fix crash when selecting many bones in setup mode

diff --git a/Engine/Plugins/Animation/ControlRig/Source/ControlRigEditor/Private/Editor/SRigHierarchy.cpp b/Engine/Plugins/Animation/ControlRig/Source/ControlRigEditor/Private/Editor/SRigHierarchy.cpp
--- a/Engine/Plugins/Animation/ControlRig/Source/ControlRigEditor/Private/Editor/SRigHierarchy.cpp
+++ b/Engine/Plugins/Animation/ControlRig/Source/ControlRigEditor/Private/Editor/SRigHierarchy.cpp
@@ -2302,12 +2302,12 @@ void SRigHierarchy::HandleControlBoneOrSpaceTransform()
 	}
 
 	TArray<FRigElementKey> SelectedKeys = GetSelectedKeys();
-	for (const FRigElementKey& SelectedKey : SelectedKeys)
+	if (SelectedKeys.Num() == 1)
 	{
-		if (SelectedKey.Type == ERigElementType::Bone ||
-			SelectedKey.Type == ERigElementType::Null)
+		if (SelectedKeys[0].Type == ERigElementType::Bone ||
+			SelectedKeys[0].Type == ERigElementType::Null)
 		{
-			Blueprint->AddTransientControl(SelectedKey);
+			Blueprint->AddTransientControl(SelectedKeys[0]);
 			return;
 		}
 	}
