Control Rig: New Pos/Rot/Scale Constraint nodes that use local space offset. Previously it was the global space offset that was preserved

diff --git a/Engine/Plugins/Animation/ControlRig/Source/ControlRig/Private/Units/Highlevel/Hierarchy/RigUnit_TransformConstraint.cpp b/Engine/Plugins/Animation/ControlRig/Source/ControlRig/Private/Units/Highlevel/Hierarchy/RigUnit_TransformConstraint.cpp
--- a/Engine/Plugins/Animation/ControlRig/Source/ControlRig/Private/Units/Highlevel/Hierarchy/RigUnit_TransformConstraint.cpp
+++ b/Engine/Plugins/Animation/ControlRig/Source/ControlRig/Private/Units/Highlevel/Hierarchy/RigUnit_TransformConstraint.cpp
@@ -576,6 +576,230 @@
 
 #endif
 
+FRigUnit_PositionConstraintLocalSpaceOffset_Execute()
+{
+	if (Weight < KINDA_SMALL_NUMBER)
+	{
+		return;
+	}
+	
+	URigHierarchy* Hierarchy = ExecuteContext.Hierarchy;	
+	if (Hierarchy)
+	{
+		if (!Child.IsValid())
+		{
+			return;
+		}
+		
+		FTransform ChildInitialLocalTransform = Hierarchy->GetInitialLocalTransform(Child);
+		
+		// Controls need to be handled a bit differently
+		if (Child.Type == ERigElementType::Control)
+		{
+			if (FRigControlElement* ChildAsControlElement = Hierarchy->Find<FRigControlElement>(Child))
+			{
+				FTransform AdditionalOffsetTransform = Hierarchy->GetControlOffsetTransform(ChildAsControlElement, ERigTransformType::InitialLocal);
+				// Control's local(parent) space position = local * offset
+				ChildInitialLocalTransform *= AdditionalOffsetTransform;
+			}
+		}
+
+		float OverallWeight = 0;
+		for (const FConstraintParent& Parent : Parents)
+		{
+			const float ClampedWeight = FMath::Max(Parent.Weight, 0.f);
+
+			if (ClampedWeight < KINDA_SMALL_NUMBER)
+			{
+				continue;
+			}
+
+			if (!Parent.Item.IsValid())
+			{
+				continue;
+			}
+
+			OverallWeight += ClampedWeight;
+		}
+
+		const float WeightNormalizer = 1.0f / OverallWeight;
+		
+		if (OverallWeight > KINDA_SMALL_NUMBER)
+		{
+			FVector OffsetPosition = FVector::ZeroVector;
+			if (bMaintainOffset)
+			{
+				FVector MixedInitialGlobalPosition = FVector::ZeroVector;
+
+				for (const FConstraintParent& Parent : Parents)
+				{
+					const float ClampedWeight = FMath::Max(Parent.Weight, 0.f);
+
+					if (ClampedWeight < KINDA_SMALL_NUMBER)
+					{
+						continue;
+					}
+
+					if (!Parent.Item.IsValid())
+					{
+						continue;
+					}
+
+					const float NormalizedWeight = ClampedWeight * WeightNormalizer;
+
+					FTransform ParentInitialGlobalTransform = Hierarchy->GetGlobalTransform(Parent.Item, true);
+
+					MixedInitialGlobalPosition += ParentInitialGlobalTransform.GetLocation() * NormalizedWeight;
+				}
+				
+				FTransform ChildParentInitialGlobalTransform = Hierarchy->GetParentTransform(Child, true);
+				FVector MixedInitialLocalPosition = ChildParentInitialGlobalTransform.Inverse().TransformPosition(MixedInitialGlobalPosition);
+
+				// use initial transforms to calculate the offset
+				// offset is applied in local space, at the end of constraint solve
+				OffsetPosition = ChildInitialLocalTransform.GetLocation() - MixedInitialLocalPosition;
+			}
+			
+			FVector MixedGlobalPosition = FVector::ZeroVector;
+
+			for (const FConstraintParent& Parent : Parents)
+			{
+				const float ClampedWeight = FMath::Max(Parent.Weight, 0.f);
+
+				if (ClampedWeight < KINDA_SMALL_NUMBER)
+				{
+					continue;
+				}
+
+				if (!Parent.Item.IsValid())
+				{
+					continue;
+				}
+
+				const float NormalizedWeight = ClampedWeight * WeightNormalizer;
+
+				FTransform ParentCurrentGlobalTransform = Hierarchy->GetGlobalTransform(Parent.Item, false);
+				MixedGlobalPosition += ParentCurrentGlobalTransform.GetLocation() * NormalizedWeight;
+			}
+			
+			// handle filtering, performed in local space
+			FTransform ChildParentGlobalTransform = Hierarchy->GetParentTransform(Child, false);
+			FVector MixedPosition = ChildParentGlobalTransform.Inverse().TransformPosition(MixedGlobalPosition);
+			
+			if (bMaintainOffset)
+			{
+				MixedPosition = MixedPosition + OffsetPosition;	
+			}
+			
+			FTransform ChildCurrentLocalTransform = Hierarchy->GetLocalTransform(Child);
+			
+			// Controls need to be handled a bit differently
+			FTransform AdditionalOffsetTransform = FTransform::Identity;
+			
+			if (Child.Type == ERigElementType::Control)
+			{
+				if (FRigControlElement* ChildAsControlElement = Hierarchy->Find<FRigControlElement>(Child))
+				{
+					AdditionalOffsetTransform = Hierarchy->GetControlOffsetTransform(ChildAsControlElement, ERigTransformType::CurrentLocal);
+					// Control's local(parent) space position = local * offset
+					ChildCurrentLocalTransform *= AdditionalOffsetTransform;
+				}
+			}
+			
+			FVector ChildPosition = ChildCurrentLocalTransform.GetTranslation();
+
+			FVector FilteredPosition;
+			FilteredPosition.X = Filter.bX ? MixedPosition.X : ChildPosition.X;
+			FilteredPosition.Y = Filter.bY ? MixedPosition.Y : ChildPosition.Y;
+			FilteredPosition.Z = Filter.bZ ? MixedPosition.Z : ChildPosition.Z;
+
+			FTransform FilteredMixedLocalTransform = ChildCurrentLocalTransform;
+
+			FilteredMixedLocalTransform.SetTranslation(FilteredPosition); 
+
+			FTransform FinalLocalTransform = FilteredMixedLocalTransform;
+			
+			if (Weight < 1.0f - KINDA_SMALL_NUMBER)
+			{
+				FinalLocalTransform = FControlRigMathLibrary::LerpTransform(ChildCurrentLocalTransform, FinalLocalTransform, Weight);
+			}
+			
+			if (Child.Type == ERigElementType::Control)
+			{
+				// need to convert back to offset space for the actual control value
+				FinalLocalTransform = FinalLocalTransform.GetRelativeTransform(AdditionalOffsetTransform);
+				FinalLocalTransform.NormalizeRotation();
+			}
+
+			Hierarchy->SetLocalTransform(Child, FinalLocalTransform);
+		}
+	}
+}
+
+IMPLEMENT_RIGUNIT_AUTOMATION_TEST(FRigUnit_PositionConstraintLocalSpaceOffset)
+{
+	// multi-parent test
+	{
+		// use euler rotation here to match other software's rotation representation more easily
+    	EEulerRotationOrder Order = EEulerRotationOrder::XZY;
+    	const FRigElementKey Child = Controller->AddBone(TEXT("Child"), FRigElementKey(), FTransform(FVector(0.f, 0.f, 0.f)), true, ERigBoneType::User);
+    	const FRigElementKey Parent1 = Controller->AddBone(TEXT("Parent1"), FRigElementKey(), FTransform(FVector(20.f, 20.f, 20.f)), true, ERigBoneType::User);
+    	const FRigElementKey Parent2 = Controller->AddBone(TEXT("Parent2"), FRigElementKey(), FTransform(FVector(40.f, 40.f, 40.f)), true, ERigBoneType::User);
+    	const FRigElementKey Parent3 = Controller->AddBone(TEXT("Parent3"), FRigElementKey(), FTransform(FVector(60.f, 60.f, 60.f)), true, ERigBoneType::User);
+    	
+    	Unit.ExecuteContext.Hierarchy = Hierarchy;
+    	Unit.Child = Child;
+    
+    	Unit.Parents.Add(FConstraintParent(Parent1, 1.0));
+    	Unit.Parents.Add(FConstraintParent(Parent2, 1.0));
+    	Unit.Parents.Add(FConstraintParent(Parent3, 1.0));
+    
+    	Hierarchy->ResetPoseToInitial(ERigElementType::Bone);
+    	Unit.bMaintainOffset = false;
+    	
+    	Execute();
+    	AddErrorIfFalse(Hierarchy->GetGlobalTransform(0).GetTranslation().Equals(FVector(40.f, 40.f, 40.f)), TEXT("unexpected translation for maintain offset off"));
+    
+    	
+    	Hierarchy->ResetPoseToInitial(ERigElementType::Bone);
+    	Hierarchy->SetGlobalTransform(2, FTransform(FVector(-40.f, -40.f, -40.f)));
+    	Unit.bMaintainOffset = true;
+    	
+    	Execute();
+    	AddErrorIfFalse(Hierarchy->GetGlobalTransform(0).GetTranslation().Equals(FVector(-26.67f, -26.67f, -26.67f), 0.01f),
+    					TEXT("unexpected translation for maintain offset on"));
+
+	}
+
+	// local offset test
+	{
+		Unit.Parents.Reset();
+		
+		Controller->GetHierarchy()->Reset();
+		const FRigElementKey Thigh = Controller->AddBone(TEXT("Thigh"), FRigElementKey(), FTransform(FVector(0.f, 0.f, 0.f)), true, ERigBoneType::User);
+		const FRigElementKey Knee = Controller->AddBone(TEXT("Knee"), FRigElementKey(TEXT("Thigh"), ERigElementType::Bone), FTransform(FVector(10.f, 0.f, -50.f)), true, ERigBoneType::User);
+		const FRigElementKey Corrective = Controller->AddBone(TEXT("Corrective"), FRigElementKey(TEXT("Knee"), ERigElementType::Bone), FTransform(FVector(-50.f, 0.f, 0.f)), false, ERigBoneType::User);
+
+		Unit.ExecuteContext.Hierarchy = Hierarchy;
+		Unit.Child = Corrective;
+    
+		Unit.Parents.Add(FConstraintParent(Knee, 1.0));
+		Unit.Parents.Add(FConstraintParent(Thigh, 1.0));
+		Unit.bMaintainOffset = true;
+
+		// use euler rotation here to match other software's rotation representation more easily
+		EEulerRotationOrder Order = EEulerRotationOrder::XZY;	
+		Hierarchy->SetLocalTransform(1,
+			FTransform(AnimationCore::QuatFromEuler( FVector(0, 45, 0), Order), 
+			FVector(10.f,0.f,-50.f)));
+		
+		Execute();
+		AddErrorIfFalse(Hierarchy->GetGlobalTransform(2).GetTranslation().Equals(FVector(-44.5f, 0.f, -10.86f), 0.1f), TEXT("unexpected translation for local space offset test"));	
+	}
+	
+	return true;
+}
+
 FRigUnit_RotationConstraint_Execute()
 {
 	if (Weight < KINDA_SMALL_NUMBER)
@@ -807,37 +1031,8 @@
 }
 #endif
 
-FRigUnit_ScaleConstraint_Execute()
+FRigUnit_RotationConstraintLocalSpaceOffset_Execute()
 {
-	TFunction<FVector (const FVector&)> GetNonZeroScale([RigVMExecuteContext, Context](const FVector& InScale)
-	{
-		FVector NonZeroScale = InScale;
-       
-       	bool bDetectedZeroScale = false;
-       	if (FMath::Abs(NonZeroScale.X) < KINDA_SMALL_NUMBER)
-       	{
-       		NonZeroScale.X = KINDA_SMALL_NUMBER * FMath::Sign(NonZeroScale.X);
-       		bDetectedZeroScale = true;
-       	}
-       	if (FMath::Abs(NonZeroScale.Y) < KINDA_SMALL_NUMBER)
-       	{
-       		NonZeroScale.Y = KINDA_SMALL_NUMBER * FMath::Sign(NonZeroScale.Y);
-       		bDetectedZeroScale = true;
-       	}
-       	if (FMath::Abs(NonZeroScale.Z) < KINDA_SMALL_NUMBER)
-       	{
-       		NonZeroScale.Z = KINDA_SMALL_NUMBER * FMath::Sign(NonZeroScale.Z);
-       		bDetectedZeroScale = true;
-       	}
-       
-       	if (bDetectedZeroScale)
-       	{
-       		UE_CONTROLRIG_RIGUNIT_REPORT_WARNING(TEXT("Scale value: (%f, %f, %f) contains value too close to 0 to use with scale constraint."), InScale.X, InScale.Y, InScale.Z);		
-       	}
-       	
-       	return NonZeroScale;	
-	});
-
 	if (Weight < KINDA_SMALL_NUMBER)
 	{
 		return;
@@ -850,13 +1045,24 @@
 		{
 			return;
 		}
+
+		FTransform ChildInitialLocalTransform = Hierarchy->GetInitialLocalTransform(Child);
+
+		// Controls need to be handled a bit differently
+		if (Child.Type == ERigElementType::Control)
+		{
+			if (FRigControlElement* ChildAsControlElement = Hierarchy->Find<FRigControlElement>(Child))
+			{
+				FTransform AdditionalOffsetTransform = Hierarchy->GetControlOffsetTransform(ChildAsControlElement, ERigTransformType::InitialLocal);
+				// Control's local(parent) space position = local * offset
+				ChildInitialLocalTransform *= AdditionalOffsetTransform;
+			}
+		}
 		
-		FTransform ChildInitialGlobalTransform = Hierarchy->GetInitialGlobalTransform(Child);
-
-		FVector::FReal OverallWeight = 0;
+		float OverallWeight = 0;
 		for (const FConstraintParent& Parent : Parents)
 		{
-			const FVector::FReal ClampedWeight = FMath::Max<FVector::FReal>(Parent.Weight, 0.f);
+			const float ClampedWeight = FMath::Max(Parent.Weight, 0.f);
 
 			if (ClampedWeight < KINDA_SMALL_NUMBER)
 			{
@@ -871,15 +1077,88 @@
 			OverallWeight += ClampedWeight;
 		}
 
+		const float WeightNormalizer = 1.0f / OverallWeight;
+		
 		if (OverallWeight > KINDA_SMALL_NUMBER)
 		{
-			const FVector::FReal WeightNormalizer = 1.0f / OverallWeight;
+			FQuat OffsetRotation = FQuat::Identity;
+			
+			if (bMaintainOffset)
+			{
+				FQuat MixedInitialGlobalRotation = FQuat(0,0,0,0);
 
-			FVector MixedGlobalScale = FVector::OneVector;
+				for (const FConstraintParent& Parent : Parents)
+				{
+					const float ClampedWeight = FMath::Max(Parent.Weight, 0.f);
 
+					if (ClampedWeight < KINDA_SMALL_NUMBER)
+					{
+						continue;
+					}
+
+					if (!Parent.Item.IsValid())
+					{
+						continue;
+					}
+
+					const float NormalizedWeight = ClampedWeight * WeightNormalizer;
+
+					FTransform ParentInitialGlobalTransform = Hierarchy->GetGlobalTransform(Parent.Item, true);
+					FQuat ParentInitialGlobalRotation = ParentInitialGlobalTransform.GetRotation();
+
+					// deal with different interpolation types
+					if (AdvancedSettings.InterpolationType == EConstraintInterpType::Average)
+					{
+						// component-wise average
+						FQuat WeightedInitialParentRotation = ParentInitialGlobalRotation * NormalizedWeight; 
+						
+						// To ensure the 'shortest route', we make sure the dot product between the accumulator and the incoming child Quat is positive.
+						if ((WeightedInitialParentRotation | MixedInitialGlobalRotation) < 0.f)
+						{
+							MixedInitialGlobalRotation -= WeightedInitialParentRotation;	
+						}
+						else
+						{
+							MixedInitialGlobalRotation += WeightedInitialParentRotation;	
+						}
+					}
+					else if (AdvancedSettings.InterpolationType == EConstraintInterpType::Shortest)
+					{
+						if (MixedInitialGlobalRotation == FQuat(0.0f, 0.0f, 0.0f, 0.0f))
+						{
+							MixedInitialGlobalRotation = ParentInitialGlobalRotation;
+						}
+						else
+						{
+							MixedInitialGlobalRotation = FQuat::Slerp(MixedInitialGlobalRotation, ParentInitialGlobalRotation, NormalizedWeight);
+						}
+					}
+					else
+					{
+						// invalid interpolation type
+						ensure(false);
+						FTransform ChildInitialGlobalTransform = Hierarchy->GetGlobalTransform(Child, true);
+						MixedInitialGlobalRotation = ChildInitialGlobalTransform.GetRotation();
+						break;
+					}
+				}
+
+				MixedInitialGlobalRotation.Normalize();
+				
+				// calculate the offset rotation that rotates the initial blended result
+				// back to the initial orientation of the child
+				FTransform ChildParentInitialGlobalTransform = Hierarchy->GetParentTransform(Child, true);
+				FQuat MixedInitialLocalRotation = ChildParentInitialGlobalTransform.GetRotation().Inverse() * MixedInitialGlobalRotation;
+				FQuat ChildInitialLocalRotation = ChildInitialLocalTransform.GetRotation();
+					
+				OffsetRotation = MixedInitialLocalRotation.Inverse() * ChildInitialLocalRotation;
+			}
+			
+			FQuat MixedGlobalRotation = FQuat(0,0,0,0);
+
 			for (const FConstraintParent& Parent : Parents)
 			{
-				const FVector::FReal ClampedWeight = FMath::Max<FVector::FReal>(Parent.Weight, 0.f);
+				const float ClampedWeight = FMath::Max(Parent.Weight, 0.f);
 
 				if (ClampedWeight < KINDA_SMALL_NUMBER)
 				{
@@ -891,6 +1170,317 @@
 					continue;
 				}
 
+				const float NormalizedWeight = ClampedWeight * WeightNormalizer;
+
+				FTransform ParentCurrentGlobalTransform = Hierarchy->GetGlobalTransform(Parent.Item, false);
+				FQuat ParentCurrentGlobalRotation = ParentCurrentGlobalTransform.GetRotation();
+
+				// deal with different interpolation types
+				if (AdvancedSettings.InterpolationType == EConstraintInterpType::Average)
+				{
+					// component-wise average
+					FQuat WeightedParentRotation = ParentCurrentGlobalRotation * NormalizedWeight; 
+                    
+                    // To ensure the 'shortest route', we make sure the dot product between the accumulator and the incoming child Quat is positive.
+					if ((WeightedParentRotation | MixedGlobalRotation) < 0.f)
+					{
+						MixedGlobalRotation -= WeightedParentRotation;	
+					}
+					else
+					{
+						MixedGlobalRotation += WeightedParentRotation;	
+					}
+				}
+				else if (AdvancedSettings.InterpolationType == EConstraintInterpType::Shortest)
+				{
+					if (MixedGlobalRotation == FQuat(0.0f, 0.0f, 0.0f, 0.0f))
+					{
+						MixedGlobalRotation = ParentCurrentGlobalRotation;
+					}
+					else
+					{
+						MixedGlobalRotation = FQuat::Slerp(MixedGlobalRotation, ParentCurrentGlobalRotation, NormalizedWeight);
+					}
+				}
+				else
+				{
+					// invalid interpolation type
+					ensure(false);
+					FTransform ChildCurrentGlobalTransform = Hierarchy->GetGlobalTransform(Child, false);
+					MixedGlobalRotation = ChildCurrentGlobalTransform.GetRotation();
+					break;
+				}	
+			}
+				
+			MixedGlobalRotation.Normalize();
+			
+			// handle filtering, performed in local space
+			FTransform ChildParentGlobalTransform = Hierarchy->GetParentTransform(Child, false);
+			FQuat MixedLocalRotation = ChildParentGlobalTransform.GetRotation().Inverse() * MixedGlobalRotation;
+
+			if (bMaintainOffset)
+			{
+				MixedLocalRotation = MixedLocalRotation * OffsetRotation;
+			}
+			
+			FVector MixedEulerRotation = AnimationCore::EulerFromQuat(MixedLocalRotation, AdvancedSettings.RotationOrderForFilter);
+			
+			FTransform ChildCurrentLocalTransform = Hierarchy->GetLocalTransform(Child, false);
+			
+			// Controls need to be handled a bit differently
+			FTransform AdditionalOffsetTransform = FTransform::Identity;
+
+			if (Child.Type == ERigElementType::Control)
+			{
+				if (FRigControlElement* ChildAsControlElement = Hierarchy->Find<FRigControlElement>(Child))
+				{
+					AdditionalOffsetTransform = Hierarchy->GetControlOffsetTransform(ChildAsControlElement, ERigTransformType::CurrentLocal);
+					// Control's local(parent) space = local * offset
+					ChildCurrentLocalTransform *= AdditionalOffsetTransform;
+				}
+			}
+
+			FQuat ChildRotation = ChildCurrentLocalTransform.GetRotation();
+			
+			FVector ChildEulerRotation = AnimationCore::EulerFromQuat(ChildRotation, AdvancedSettings.RotationOrderForFilter);	
+			
+			FVector FilteredEulerRotation;
+			FilteredEulerRotation.X = Filter.bX ? MixedEulerRotation.X : ChildEulerRotation.X;
+			FilteredEulerRotation.Y = Filter.bY ? MixedEulerRotation.Y : ChildEulerRotation.Y;
+			FilteredEulerRotation.Z = Filter.bZ ? MixedEulerRotation.Z : ChildEulerRotation.Z; 
+
+			FTransform FilteredMixedLocalTransform = ChildCurrentLocalTransform;
+
+			FilteredMixedLocalTransform.SetRotation(AnimationCore::QuatFromEuler(FilteredEulerRotation, AdvancedSettings.RotationOrderForFilter));
+
+			FTransform FinalLocalTransform = FilteredMixedLocalTransform;
+
+			if (Weight < 1.0f - KINDA_SMALL_NUMBER)
+			{
+				FinalLocalTransform = FControlRigMathLibrary::LerpTransform(ChildCurrentLocalTransform, FinalLocalTransform, Weight);
+			}
+
+			if (Child.Type == ERigElementType::Control)
+			{
+				// need to convert back to offset space for the actual control value
+				FinalLocalTransform = FinalLocalTransform.GetRelativeTransform(AdditionalOffsetTransform);
+				FinalLocalTransform.NormalizeRotation();
+			}
+			
+			Hierarchy->SetLocalTransform(Child, FinalLocalTransform);
+		}
+	} 	
+}
+
+#if WITH_DEV_AUTOMATION_TESTS
+#include "Units/RigUnitTest.h"
+
+IMPLEMENT_RIGUNIT_AUTOMATION_TEST(FRigUnit_RotationConstraintLocalSpaceOffset)
+{
+	// the rotation constraint is expected to behave similarly to parent constraint with translation filter turned off.
+
+	// multi parent test
+	{
+		// use euler rotation here to match other software's rotation representation more easily
+		EEulerRotationOrder Order = EEulerRotationOrder::XZY;
+		const FRigElementKey Child = Controller->AddBone(TEXT("Child"), FRigElementKey(), FTransform( AnimationCore::QuatFromEuler( FVector(-10, -10, -10), Order), FVector(0.f, 0.f, 0.f)), true, ERigBoneType::User);
+		const FRigElementKey Parent1 = Controller->AddBone(TEXT("Parent1"), FRigElementKey(), FTransform( AnimationCore::QuatFromEuler( FVector(30, -30, -30), Order), FVector(20.f, 20.f, 20.f)), true, ERigBoneType::User);
+		const FRigElementKey Parent2 = Controller->AddBone(TEXT("Parent2"), FRigElementKey(), FTransform( AnimationCore::QuatFromEuler( FVector(-40, -40, 40), Order), FVector(40.f, 40.f, 40.f)), true, ERigBoneType::User);
+		const FRigElementKey Parent3 = Controller->AddBone(TEXT("Parent3"), FRigElementKey(), FTransform( AnimationCore::QuatFromEuler( FVector(-50, 50, -50), Order), FVector(60.f, 60.f, 60.f)), true, ERigBoneType::User);
+		const FRigElementKey Parent4 = Controller->AddBone(TEXT("Parent4"), FRigElementKey(), FTransform( AnimationCore::QuatFromEuler( FVector(60, 60, 60), Order), FVector(80.f, 80.f, 80.f)), true, ERigBoneType::User);
+		
+		Unit.ExecuteContext.Hierarchy = Hierarchy;
+		Unit.Child = Child;
+
+		Unit.Parents.Add(FConstraintParent(Parent1, 1.0));
+		Unit.Parents.Add(FConstraintParent(Parent2, 1.0));
+		Unit.Parents.Add(FConstraintParent(Parent3, 1.0));
+		Unit.Parents.Add(FConstraintParent(Parent4, 1.0));
+
+		Hierarchy->ResetPoseToInitial(ERigElementType::Bone);
+		Unit.bMaintainOffset = false;
+		Unit.AdvancedSettings.InterpolationType = EConstraintInterpType::Average;
+		
+		Execute();
+
+		FQuat Result = Hierarchy->GetGlobalTransform(0).GetRotation();
+		FQuat Expected = AnimationCore::QuatFromEuler( FVector(-0.853f, 15.189f, -0.572f), Order);
+		AddErrorIfFalse(Result.Equals(Expected, 0.001f), TEXT("unexpected rotation for average interpolation type"));
+		
+		Hierarchy->ResetPoseToInitial(ERigElementType::Bone);
+		Unit.bMaintainOffset = false;
+		Unit.AdvancedSettings.InterpolationType = EConstraintInterpType::Shortest;
+		
+		Execute();
+		
+		Result = Hierarchy->GetGlobalTransform(0).GetRotation();
+		Expected = AnimationCore::QuatFromEuler( FVector(16.74f, 8.865f, -5.562f), Order);
+		AddErrorIfFalse(Result.Equals(Expected, 0.001f), TEXT("unexpected rotation for shortest interpolation type"));
+
+		
+		Hierarchy->ResetPoseToInitial(ERigElementType::Bone);
+		Hierarchy->SetGlobalTransform(2, FTransform(AnimationCore::QuatFromEuler( FVector(100, 100, -100), Order), FVector(-40.f, -40.f, -40.f)));
+		Unit.bMaintainOffset = true;
+		Unit.AdvancedSettings.InterpolationType = EConstraintInterpType::Average;
+		
+		Execute();
+		
+		Result = Hierarchy->GetGlobalTransform(0).GetRotation();
+		Expected = AnimationCore::QuatFromEuler( FVector(-0.61f, 8.524f, -45.755f), Order);
+		AddErrorIfFalse(Result.Equals(Expected, 0.1f), TEXT("unexpected rotation for maintain offset and average interpolation type"));
+
+		
+		Hierarchy->ResetPoseToInitial(ERigElementType::Bone);
+		Hierarchy->SetGlobalTransform(2, FTransform(AnimationCore::QuatFromEuler( FVector(100.0f, 100.0f, -100.0f), Order), FVector(-40.f, -40.f, -40.f)));
+		Unit.bMaintainOffset = true;
+		Unit.AdvancedSettings.InterpolationType = EConstraintInterpType::Shortest;
+		
+		Execute();
+		
+		Result = Hierarchy->GetGlobalTransform(0).GetRotation();
+		Expected = AnimationCore::QuatFromEuler( FVector(-9.143, -4.997f, -38.156f), Order);
+		AddErrorIfFalse(Result.Equals(Expected, 0.1f), TEXT("unexpected rotation for maintain offset and shortest interpolation type"));
+	}
+
+	// local space offset test
+	{
+		Unit.Parents.Reset();
+		
+		Controller->GetHierarchy()->Reset();
+		
+		// use euler rotation here to match other software's rotation representation more easily
+		EEulerRotationOrder Order = EEulerRotationOrder::XZY;
+		
+		const FRigElementKey Root = Controller->AddBone(TEXT("Root"), FRigElementKey(), FTransform(FVector(0.f, 0.f, 0.f)), true, ERigBoneType::User);
+		const FRigElementKey Parent1 = Controller->AddBone(TEXT("Parent1"), FRigElementKey(TEXT("Root"), ERigElementType::Bone),
+			FTransform(AnimationCore::QuatFromEuler(FVector(90,0,0), Order), FVector(0.f, 0.f, 50.f)), false, ERigBoneType::User);
+		const FRigElementKey Parent2 = Controller->AddBone(TEXT("Parent2"), FRigElementKey(TEXT("Root"), ERigElementType::Bone),
+			FTransform(AnimationCore::QuatFromEuler(FVector(0,0,-45), Order), FVector(0.f, 100.f, 0.f)), false, ERigBoneType::User);
+		const FRigElementKey Child = Controller->AddBone(TEXT("Child"), FRigElementKey(TEXT("Parent2"), ERigElementType::Bone),
+			FTransform(AnimationCore::QuatFromEuler(FVector(0,0,-45), Order), FVector(0.f, 0.f, 50.f)), false, ERigBoneType::User);
+
+		Unit.ExecuteContext.Hierarchy = Hierarchy;
+		Unit.Child = Child;
+    
+		Unit.Parents.Add(FConstraintParent(Parent1, 1.0));
+		Unit.Parents.Add(FConstraintParent(Parent2, 1.0));
+		
+		{
+			Unit.bMaintainOffset = false;
+			Execute();
+			FQuat Rotation = Hierarchy->GetLocalTransform(3).GetRotation();
+			FVector EulerRotation = AnimationCore::EulerFromQuat(Rotation, Order);
+			AddErrorIfFalse(EulerRotation.Equals(FVector(41.12f, 10.18f, 22.18f), 0.1f), TEXT("unexpected rotation for local space offset test 1"));
+		}
+		
+		{
+			Unit.bMaintainOffset = true;
+			Execute();
+			FQuat Rotation = Hierarchy->GetLocalTransform(3).GetRotation();
+			FVector EulerRotation = AnimationCore::EulerFromQuat(Rotation, Order);
+			AddErrorIfFalse(EulerRotation.Equals(FVector(0.f, 0.f, -45.f), 0.1f), TEXT("unexpected rotation for local space offset test 2"));
+		}
+
+		{
+			Hierarchy->SetLocalTransform(1,
+				FTransform(AnimationCore::QuatFromEuler( FVector(120, 40, 20), Order), 
+				FVector(0.f,0.f,50.f)));
+			
+			Unit.bMaintainOffset = true;
+			Execute();
+			FQuat Rotation = Hierarchy->GetLocalTransform(3).GetRotation();
+			FVector EulerRotation = AnimationCore::EulerFromQuat(Rotation, Order);
+			AddErrorIfFalse(EulerRotation.Equals(FVector(0.583f, 30.506f, -53.959f), 0.1f), TEXT("unexpected rotation for local space offset test 3"));
+		}
+	}
+	return true;
+}
+#endif
+
+FRigUnit_ScaleConstraint_Execute()
+{
+	TFunction<FVector (const FVector&)> GetNonZeroScale([RigVMExecuteContext, Context](const FVector& InScale)
+	{
+		FVector NonZeroScale = InScale;
+       
+       	bool bDetectedZeroScale = false;
+       	if (FMath::Abs(NonZeroScale.X) < KINDA_SMALL_NUMBER)
+       	{
+       		NonZeroScale.X = KINDA_SMALL_NUMBER * FMath::Sign(NonZeroScale.X);
+       		bDetectedZeroScale = true;
+       	}
+       	if (FMath::Abs(NonZeroScale.Y) < KINDA_SMALL_NUMBER)
+       	{
+       		NonZeroScale.Y = KINDA_SMALL_NUMBER * FMath::Sign(NonZeroScale.Y);
+       		bDetectedZeroScale = true;
+       	}
+       	if (FMath::Abs(NonZeroScale.Z) < KINDA_SMALL_NUMBER)
+       	{
+       		NonZeroScale.Z = KINDA_SMALL_NUMBER * FMath::Sign(NonZeroScale.Z);
+       		bDetectedZeroScale = true;
+       	}
+       
+       	if (bDetectedZeroScale)
+       	{
+       		UE_CONTROLRIG_RIGUNIT_REPORT_WARNING(TEXT("Scale value: (%f, %f, %f) contains value too close to 0 to use with scale constraint."), InScale.X, InScale.Y, InScale.Z);		
+       	}
+       	
+       	return NonZeroScale;	
+	});
+
+	if (Weight < KINDA_SMALL_NUMBER)
+	{
+		return;
+	}
+	
+	URigHierarchy* Hierarchy = ExecuteContext.Hierarchy;	
+	if (Hierarchy)
+	{
+		if (!Child.IsValid())
+		{
+			return;
+		}
+		
+		FTransform ChildInitialGlobalTransform = Hierarchy->GetInitialGlobalTransform(Child);
+
+		FVector::FReal OverallWeight = 0;
+		for (const FConstraintParent& Parent : Parents)
+		{
+			const FVector::FReal ClampedWeight = FMath::Max<FVector::FReal>(Parent.Weight, 0.f);
+
+			if (ClampedWeight < KINDA_SMALL_NUMBER)
+			{
+				continue;
+			}
+
+			if (!Parent.Item.IsValid())
+			{
+				continue;
+			}
+
+			OverallWeight += ClampedWeight;
+		}
+
+		if (OverallWeight > KINDA_SMALL_NUMBER)
+		{
+			const FVector::FReal WeightNormalizer = 1.0f / OverallWeight;
+
+			FVector MixedGlobalScale = FVector::OneVector;
+
+			for (const FConstraintParent& Parent : Parents)
+			{
+				const FVector::FReal ClampedWeight = FMath::Max<FVector::FReal>(Parent.Weight, 0.f);
+
+				if (ClampedWeight < KINDA_SMALL_NUMBER)
+				{
+					continue;
+				}
+
+				if (!Parent.Item.IsValid())
+				{
+					continue;
+				}
+
 				const FVector::FReal NormalizedWeight = ClampedWeight * WeightNormalizer;
 
 				FVector OffsetScale = FVector::OneVector;
@@ -1000,3 +1590,301 @@
 }
 
 #endif
+
+FRigUnit_ScaleConstraintLocalSpaceOffset_Execute()
+{
+	TFunction<FVector (const FVector&)> GetNonZeroScale([RigVMExecuteContext, Context](const FVector& InScale)
+	{
+		FVector NonZeroScale = InScale;
+       
+       	bool bDetectedZeroScale = false;
+       	if (FMath::Abs(NonZeroScale.X) < KINDA_SMALL_NUMBER)
+       	{
+       		NonZeroScale.X = KINDA_SMALL_NUMBER * FMath::Sign(NonZeroScale.X);
+       		bDetectedZeroScale = true;
+       	}
+       	if (FMath::Abs(NonZeroScale.Y) < KINDA_SMALL_NUMBER)
+       	{
+       		NonZeroScale.Y = KINDA_SMALL_NUMBER * FMath::Sign(NonZeroScale.Y);
+       		bDetectedZeroScale = true;
+       	}
+       	if (FMath::Abs(NonZeroScale.Z) < KINDA_SMALL_NUMBER)
+       	{
+       		NonZeroScale.Z = KINDA_SMALL_NUMBER * FMath::Sign(NonZeroScale.Z);
+       		bDetectedZeroScale = true;
+       	}
+       
+       	if (bDetectedZeroScale)
+       	{
+       		UE_CONTROLRIG_RIGUNIT_REPORT_WARNING(TEXT("Scale value: (%f, %f, %f) contains value too close to 0 to use with scale constraint."), InScale.X, InScale.Y, InScale.Z);		
+       	}
+       	
+       	return NonZeroScale;	
+	});
+
+	if (Weight < KINDA_SMALL_NUMBER)
+	{
+		return;
+	}
+	
+	URigHierarchy* Hierarchy = ExecuteContext.Hierarchy;	
+	if (Hierarchy)
+	{
+		if (!Child.IsValid())
+		{
+			return;
+		}
+		
+		FTransform ChildInitialLocalTransform = Hierarchy->GetInitialLocalTransform(Child);
+		
+		// Controls need to be handled a bit differently
+		if (Child.Type == ERigElementType::Control)
+		{
+			if (FRigControlElement* ChildAsControlElement = Hierarchy->Find<FRigControlElement>(Child))
+			{
+				FTransform AdditionalOffsetTransform = Hierarchy->GetControlOffsetTransform(ChildAsControlElement, ERigTransformType::InitialLocal);
+				// Control's local(parent) space position = local * offset
+				ChildInitialLocalTransform *= AdditionalOffsetTransform;
+			}
+		}
+		FVector::FReal OverallWeight = 0;
+		for (const FConstraintParent& Parent : Parents)
+		{
+			const FVector::FReal ClampedWeight = FMath::Max<FVector::FReal>(Parent.Weight, 0.f);
+
+			if (ClampedWeight < KINDA_SMALL_NUMBER)
+			{
+				continue;
+			}
+
+			if (!Parent.Item.IsValid())
+			{
+				continue;
+			}
+
+			OverallWeight += ClampedWeight;
+		}
+
+		const FVector::FReal WeightNormalizer = 1.0f / OverallWeight;
+		
+		if (OverallWeight > KINDA_SMALL_NUMBER)
+		{
+			FVector OffsetScale = FVector::OneVector;
+			if (bMaintainOffset)
+			{
+				FVector MixedInitialGlobalScale = FVector::OneVector;
+
+				for (const FConstraintParent& Parent : Parents)
+				{
+					const FVector::FReal ClampedWeight = FMath::Max<FVector::FReal>(Parent.Weight, 0.f);
+
+					if (ClampedWeight < KINDA_SMALL_NUMBER)
+					{
+						continue;
+					}
+
+					if (!Parent.Item.IsValid())
+					{
+						continue;
+					}
+
+					const FVector::FReal NormalizedWeight = ClampedWeight * WeightNormalizer;
+
+					FTransform ParentInitialGlobalTransform = Hierarchy->GetGlobalTransform(Parent.Item, true);
+					FVector ParentInitialGlobalScale = ParentInitialGlobalTransform.GetScale3D();
+					
+					FVector WeightedInitialParentGlobalScale;
+					WeightedInitialParentGlobalScale.X = FMath::Pow(ParentInitialGlobalScale.X, NormalizedWeight);
+					WeightedInitialParentGlobalScale.Y = FMath::Pow(ParentInitialGlobalScale.Y, NormalizedWeight);
+					WeightedInitialParentGlobalScale.Z = FMath::Pow(ParentInitialGlobalScale.Z, NormalizedWeight);
+
+					MixedInitialGlobalScale *= WeightedInitialParentGlobalScale;
+				}
+				
+				// handle filtering, performed in local space
+				FTransform ChildParentInitialGlobalTransform = Hierarchy->GetParentTransform(Child, true);
+				FVector ChildParentInitialGlobalScale = ChildParentInitialGlobalTransform.GetScale3D();
+				FVector MixedInitialLocalScale = MixedInitialGlobalScale / GetNonZeroScale(ChildParentInitialGlobalScale);
+			
+				 OffsetScale = ChildInitialLocalTransform.GetScale3D() / GetNonZeroScale(MixedInitialLocalScale);
+			}
+
+			
+			FVector MixedGlobalScale = FVector::OneVector;
+
+			for (const FConstraintParent& Parent : Parents)
+			{
+				const FVector::FReal ClampedWeight = FMath::Max<FVector::FReal>(Parent.Weight, 0.f);
+
+				if (ClampedWeight < KINDA_SMALL_NUMBER)
+				{
+					continue;
+				}
+
+				if (!Parent.Item.IsValid())
+				{
+					continue;
+				}
+
+				const FVector::FReal NormalizedWeight = ClampedWeight * WeightNormalizer;
+
+				FTransform ParentCurrentGlobalTransform = Hierarchy->GetGlobalTransform(Parent.Item, false);
+				FVector ParentCurrentGlobalScale = ParentCurrentGlobalTransform.GetScale3D();
+
+				FVector WeightedParentGlobalScale;
+				WeightedParentGlobalScale.X = FMath::Pow(ParentCurrentGlobalScale.X, NormalizedWeight);
+				WeightedParentGlobalScale.Y = FMath::Pow(ParentCurrentGlobalScale.Y, NormalizedWeight);
+				WeightedParentGlobalScale.Z = FMath::Pow(ParentCurrentGlobalScale.Z, NormalizedWeight);
+
+				MixedGlobalScale *= WeightedParentGlobalScale;
+			}
+
+			// handle filtering, performed in local space
+			FTransform ChildParentGlobalTransform = Hierarchy->GetParentTransform(Child, false);
+			FVector ChildParentGlobalScale = ChildParentGlobalTransform.GetScale3D();
+			FVector MixedLocalScale = MixedGlobalScale / GetNonZeroScale(ChildParentGlobalScale);
+
+			if (bMaintainOffset)
+			{
+				MixedLocalScale = MixedLocalScale * OffsetScale;
+			}
+			
+			FTransform ChildCurrentLocalTransform = Hierarchy->GetLocalTransform(Child, false);
+
+			// Controls need to be handled a bit differently
+			FTransform AdditionalOffsetTransform = FTransform::Identity;
+
+			if (Child.Type == ERigElementType::Control)
+			{
+				if (FRigControlElement* ChildAsControlElement = Hierarchy->Find<FRigControlElement>(Child))
+				{
+					AdditionalOffsetTransform = Hierarchy->GetControlOffsetTransform(ChildAsControlElement, ERigTransformType::CurrentLocal);
+					// Control's local(parent) space = local * offset
+					ChildCurrentLocalTransform *= AdditionalOffsetTransform;
+				}
+			}
+			
+			FVector ChildLocalScale = ChildCurrentLocalTransform.GetScale3D();
+
+			FVector FilteredLocalScale;
+			FilteredLocalScale.X = Filter.bX ? MixedLocalScale.X : ChildLocalScale.X;
+			FilteredLocalScale.Y = Filter.bY ? MixedLocalScale.Y : ChildLocalScale.Y;
+			FilteredLocalScale.Z = Filter.bZ ? MixedLocalScale.Z : ChildLocalScale.Z;
+
+			FTransform FilteredMixedLocalTransform = ChildCurrentLocalTransform;
+
+			FilteredMixedLocalTransform.SetScale3D(FilteredLocalScale);
+
+			FTransform FinalLocalTransform = FilteredMixedLocalTransform;
+
+			if (Weight < 1.0f - KINDA_SMALL_NUMBER)
+			{
+				FinalLocalTransform = FControlRigMathLibrary::LerpTransform(ChildCurrentLocalTransform, FinalLocalTransform, Weight);
+			}
+			
+			if (Child.Type == ERigElementType::Control)
+			{
+				// need to convert back to offset space for the actual control value
+				FinalLocalTransform = FinalLocalTransform.GetRelativeTransform(AdditionalOffsetTransform);
+				FinalLocalTransform.NormalizeRotation();
+			}
+
+			Hierarchy->SetLocalTransform(Child, FinalLocalTransform);
+		}
+	}
+}
+
+#if WITH_DEV_AUTOMATION_TESTS
+#include "Units/RigUnitTest.h"
+
+IMPLEMENT_RIGUNIT_AUTOMATION_TEST(FRigUnit_ScaleConstraintLocalSpaceOffset)
+{
+	// multi parent test
+	{
+		// use euler rotation here to match other software's rotation representation more easily
+		EEulerRotationOrder Order = EEulerRotationOrder::XZY;
+		const FRigElementKey Child = Controller->AddBone(TEXT("Child"), FRigElementKey(), FTransform(FQuat::Identity, FVector::ZeroVector, FVector::OneVector), true, ERigBoneType::User);
+		const FRigElementKey Parent1 = Controller->AddBone(TEXT("Parent1"), FRigElementKey(), FTransform(FQuat::Identity, FVector::ZeroVector, FVector(4,4,4)), true, ERigBoneType::User);
+		const FRigElementKey Parent2 = Controller->AddBone(TEXT("Parent2"), FRigElementKey(), FTransform(FQuat::Identity, FVector::ZeroVector, FVector(1,1,1)), true, ERigBoneType::User);
+	
+		Unit.ExecuteContext.Hierarchy = Hierarchy;
+		Unit.Child = Child;
+
+		Unit.Parents.Add(FConstraintParent(Parent1, 1.0));
+		Unit.Parents.Add(FConstraintParent(Parent2, 1.0));
+
+		Hierarchy->ResetPoseToInitial(ERigElementType::Bone);
+		Unit.bMaintainOffset = false;
+	
+		Execute();
+		AddErrorIfFalse(Hierarchy->GetGlobalTransform(0).GetScale3D().Equals(FVector(2,2,2)), TEXT("unexpected scale for maintain offset off"));
+	
+		Hierarchy->ResetPoseToInitial(ERigElementType::Bone);
+		Hierarchy->SetGlobalTransform(2, FTransform(FQuat::Identity, FVector::ZeroVector, FVector(0.5, 0.5, 0.5)));
+		Unit.bMaintainOffset = true;
+	
+		Execute();
+		AddErrorIfFalse(Hierarchy->GetGlobalTransform(0).GetScale3D().Equals(FVector(0.707, 0.707f, 0.707f), 0.001f),
+						TEXT("unexpected scale for maintain offset on"));	
+	}
+
+	// local offset test
+	{
+		Unit.Parents.Reset();
+		
+		Controller->GetHierarchy()->Reset();
+		
+		// use euler rotation here to match other software's rotation representation more easily
+		EEulerRotationOrder Order = EEulerRotationOrder::XZY;
+		
+		const FRigElementKey Root = Controller->AddBone(TEXT("Root"), FRigElementKey(), FTransform(FVector(0.f, 0.f, 0.f)), true, ERigBoneType::User);
+		const FRigElementKey Parent1 = Controller->AddBone(TEXT("Parent1"), FRigElementKey(TEXT("Root"), ERigElementType::Bone),
+			FTransform(AnimationCore::QuatFromEuler(FVector(0,0,0), Order), FVector(0.f, 0.f, 50.f), FVector(4,5,6)),
+			false, ERigBoneType::User);
+		const FRigElementKey Parent2 = Controller->AddBone(TEXT("Parent2"), FRigElementKey(TEXT("Root"), ERigElementType::Bone),
+			FTransform(AnimationCore::QuatFromEuler(FVector(0,0,0), Order), FVector(0.f, 100.f, 0.f), FVector(1,2,3)),
+			false, ERigBoneType::User);
+		const FRigElementKey Child = Controller->AddBone(TEXT("Child"), FRigElementKey(TEXT("Parent2"), ERigElementType::Bone),
+			FTransform(AnimationCore::QuatFromEuler(FVector(0,0,0), Order), FVector(0.f, 0.f, 50.f)),
+			false, ERigBoneType::User);
+
+		Unit.ExecuteContext.Hierarchy = Hierarchy;
+		Unit.Child = Child;
+    
+		Unit.Parents.Add(FConstraintParent(Parent1, 1.0));
+		Unit.Parents.Add(FConstraintParent(Parent2, 1.0));
+		
+		{
+			Unit.bMaintainOffset = false;
+			Execute();
+			FVector Scale = Hierarchy->GetLocalTransform(3).GetScale3D();
+			AddErrorIfFalse(Scale.Equals(FVector(2.f, 1.58f, 1.414f), 0.1f), TEXT("unexpected scale for local space offset test 1"));
+		}
+		
+		{
+			Unit.bMaintainOffset = true;
+			Execute();
+			FVector Scale = Hierarchy->GetLocalTransform(3).GetScale3D();
+			AddErrorIfFalse(Scale.Equals(FVector(1.f, 1.f, 1.f), 0.1f), TEXT("unexpected scale for local space offset test 2"));
+		}
+
+		{
+			Hierarchy->SetLocalTransform(1,
+				FTransform(AnimationCore::QuatFromEuler( FVector(0, 0, 0), Order), 
+				FVector(0.f,0.f,50.f),
+				FVector(4,10,6)));
+			
+			Unit.bMaintainOffset = true;
+			Execute();
+			FVector Scale = Hierarchy->GetLocalTransform(3).GetScale3D();
+			AddErrorIfFalse(Scale.Equals(FVector(1.0f, 1.414f, 1.0f), 0.1f), TEXT("unexpected scale for local space offset test 3"));
+		}
+		
+	}
+	
+
+	
+	return true;
+}
+
+#endif
diff --git a/Engine/Plugins/Animation/ControlRig/Source/ControlRig/Private/Units/Highlevel/Hierarchy/RigUnit_TransformConstraint.h b/Engine/Plugins/Animation/ControlRig/Source/ControlRig/Private/Units/Highlevel/Hierarchy/RigUnit_TransformConstraint.h
--- a/Engine/Plugins/Animation/ControlRig/Source/ControlRig/Private/Units/Highlevel/Hierarchy/RigUnit_TransformConstraint.h
+++ b/Engine/Plugins/Animation/ControlRig/Source/ControlRig/Private/Units/Highlevel/Hierarchy/RigUnit_TransformConstraint.h
@@ -271,13 +271,13 @@ struct FRigUnit_ParentConstraint : public FRigUnit_HighlevelBaseMutable
 	FRigUnit_ParentConstraint_AdvancedSettings AdvancedSettings;
 
 	UPROPERTY(EditAnywhere, Category = "Constraint", meta = (Input))
-	float Weight;	
+	float Weight;
 };
 
 /**
 * Constrains an item's position to multiple items' positions 
 */
-USTRUCT(meta=(DisplayName="Position Constraint", Category="Constraints", Keywords = "Parent,Translation"))
+USTRUCT(meta=(DisplayName="Position Constraint", Category="Constraints", Keywords = "Parent,Translation", Deprecated = "5.0"))
 struct FRigUnit_PositionConstraint : public FRigUnit_HighlevelBaseMutable
 {
 	GENERATED_BODY()
@@ -309,6 +309,42 @@ struct FRigUnit_PositionConstraint : public FRigUnit_HighlevelBaseMutable
 	float Weight;
 };
 
+/**
+* Constrains an item's position to multiple items' positions 
+*/
+USTRUCT(meta=(DisplayName="Position Constraint", Category="Constraints", Keywords = "Parent,Translation"))
+struct FRigUnit_PositionConstraintLocalSpaceOffset : public FRigUnit_HighlevelBaseMutable
+{
+	GENERATED_BODY()
+
+	FRigUnit_PositionConstraintLocalSpaceOffset()
+		: Child(FRigElementKey(NAME_None, ERigElementType::Bone))
+		, bMaintainOffset(true)
+		, Weight(1.0f)
+	{
+		Parents.Add(FConstraintParent());
+	}
+	
+	RIGVM_METHOD()
+	virtual void Execute(const FRigUnitContext& Context) override;
+
+	UPROPERTY(EditAnywhere, Category = "Constraint", meta = (Input, ExpandByDefault))
+	FRigElementKey Child;
+	
+	UPROPERTY(EditAnywhere, Category = "Constraint", meta = (Input))
+	bool bMaintainOffset;
+
+	UPROPERTY(EditAnywhere, Category = "Constraint", meta = (Input)) 
+	FFilterOptionPerAxis Filter;
+
+	UPROPERTY(EditAnywhere, Category = "Constraint", meta = (Input, ExpandByDefault, DefaultArraySize = 1))
+	TArray<FConstraintParent> Parents;
+
+	UPROPERTY(EditAnywhere, Category = "Constraint", meta = (Input))
+	float Weight;
+};
+
+
 USTRUCT()
 struct FRigUnit_RotationConstraint_AdvancedSettings
 {
@@ -335,7 +371,7 @@ struct FRigUnit_RotationConstraint_AdvancedSettings
 /**
 * Constrains an item's rotation to multiple items' rotations 
 */
-USTRUCT(meta=(DisplayName="Rotation Constraint", Category="Constraints", Keywords = "Parent,Orientation,Orient,Rotate"))
+USTRUCT(meta=(DisplayName="Rotation Constraint", Category="Constraints", Keywords = "Parent,Orientation,Orient,Rotate", Deprecated = "5.0"))
 struct FRigUnit_RotationConstraint : public FRigUnit_HighlevelBaseMutable
 {
 	GENERATED_BODY()
@@ -370,10 +406,48 @@ struct FRigUnit_RotationConstraint : public FRigUnit_HighlevelBaseMutable
 	float Weight;
 };
 
+/**
+* Constrains an item's rotation to multiple items' rotations 
+*/
+USTRUCT(meta=(DisplayName="Rotation Constraint", Category="Constraints", Keywords = "Parent,Orientation,Orient,Rotate"))
+struct FRigUnit_RotationConstraintLocalSpaceOffset : public FRigUnit_HighlevelBaseMutable
+{
+	GENERATED_BODY()
+
+	FRigUnit_RotationConstraintLocalSpaceOffset()
+		: Child(FRigElementKey(NAME_None, ERigElementType::Bone))
+		, bMaintainOffset(true)
+		, Weight(1.0f)
+	{
+		Parents.Add(FConstraintParent());
+	}
+	
+	RIGVM_METHOD()
+	virtual void Execute(const FRigUnitContext& Context) override;
+
+	UPROPERTY(EditAnywhere, Category = "Constraint", meta = (Input, ExpandByDefault))
+	FRigElementKey Child;
+	
+	UPROPERTY(EditAnywhere, Category = "Constraint", meta = (Input))
+	bool bMaintainOffset;
+
+	UPROPERTY(EditAnywhere, Category = "Constraint", meta = (Input)) 
+	FFilterOptionPerAxis Filter;
+
+	UPROPERTY(EditAnywhere, Category = "Constraint", meta = (Input, ExpandByDefault, DefaultArraySize = 1))
+	TArray<FConstraintParent> Parents;
+	
+	UPROPERTY(EditAnywhere, Category = "Constraint", meta = (Input)) 
+	FRigUnit_RotationConstraint_AdvancedSettings AdvancedSettings;
+
+	UPROPERTY(EditAnywhere, Category = "Constraint", meta = (Input))
+	float Weight;
+};
+
 /**
 * Constrains an item's scale to multiple items' scales
 */
-USTRUCT(meta=(DisplayName="Scale Constraint", Category="Constraints", Keywords = "Parent, Size"))
+USTRUCT(meta=(DisplayName="Scale Constraint", Category="Constraints", Keywords = "Parent, Size", Deprecated = "5.0"))
 struct FRigUnit_ScaleConstraint : public FRigUnit_HighlevelBaseMutable
 {
 	GENERATED_BODY()
@@ -404,3 +478,38 @@ struct FRigUnit_ScaleConstraint : public FRigUnit_HighlevelBaseMutable
 	UPROPERTY(EditAnywhere, Category = "Constraint", meta = (Input))
 	float Weight;
 };
+
+/**
+* Constrains an item's scale to multiple items' scales
+*/
+USTRUCT(meta=(DisplayName="Scale Constraint", Category="Constraints", Keywords = "Parent, Size"))
+struct FRigUnit_ScaleConstraintLocalSpaceOffset : public FRigUnit_HighlevelBaseMutable
+{
+	GENERATED_BODY()
+
+	FRigUnit_ScaleConstraintLocalSpaceOffset()
+		: Child(FRigElementKey(NAME_None, ERigElementType::Bone))
+		, bMaintainOffset(true)
+		, Weight(1.0f)
+	{
+		Parents.Add(FConstraintParent());
+	}
+	
+	RIGVM_METHOD()
+	virtual void Execute(const FRigUnitContext& Context) override;
+
+	UPROPERTY(EditAnywhere, Category = "Constraint", meta = (Input, ExpandByDefault))
+	FRigElementKey Child;
+	
+	UPROPERTY(EditAnywhere, Category = "Constraint", meta = (Input))
+	bool bMaintainOffset;
+
+	UPROPERTY(EditAnywhere, Category = "Constraint", meta = (Input)) 
+	FFilterOptionPerAxis Filter;
+
+	UPROPERTY(EditAnywhere, Category = "Constraint", meta = (Input, ExpandByDefault, DefaultArraySize = 1))
+	TArray<FConstraintParent> Parents;
+
+	UPROPERTY(EditAnywhere, Category = "Constraint", meta = (Input))
+	float Weight;
+};
