Control Rig: Follow up to 18790810, fix a compilation error

diff --git a/Engine/Plugins/Animation/ControlRig/Source/ControlRig/Private/Units/Highlevel/Hierarchy/RigUnit_TransformConstraint.cpp b/Engine/Plugins/Animation/ControlRig/Source/ControlRig/Private/Units/Highlevel/Hierarchy/RigUnit_TransformConstraint.cpp
--- a/Engine/Plugins/Animation/ControlRig/Source/ControlRig/Private/Units/Highlevel/Hierarchy/RigUnit_TransformConstraint.cpp
+++ b/Engine/Plugins/Animation/ControlRig/Source/ControlRig/Private/Units/Highlevel/Hierarchy/RigUnit_TransformConstraint.cpp
@@ -736,6 +736,9 @@ FRigUnit_PositionConstraintLocalSpaceOffset_Execute()
 	}
 }
 
+#if WITH_DEV_AUTOMATION_TESTS
+#include "Units/RigUnitTest.h"
+
 IMPLEMENT_RIGUNIT_AUTOMATION_TEST(FRigUnit_PositionConstraintLocalSpaceOffset)
 {
 	// multi-parent test
@@ -799,6 +802,7 @@ IMPLEMENT_RIGUNIT_AUTOMATION_TEST(FRigUnit_PositionConstraintLocalSpaceOffset)
 	
 	return true;
 }
+#endif
 
 FRigUnit_RotationConstraint_Execute()
 {
