Sequencer: Control Rig: Mirrored poses take into scale now.

diff --git a/Engine/Plugins/Animation/ControlRig/Source/ControlRig/Private/Tools/ControlRigPose.cpp b/Engine/Plugins/Animation/ControlRig/Source/ControlRig/Private/Tools/ControlRigPose.cpp
--- a/Engine/Plugins/Animation/ControlRig/Source/ControlRig/Private/Tools/ControlRigPose.cpp
+++ b/Engine/Plugins/Animation/ControlRig/Source/ControlRig/Private/Tools/ControlRigPose.cpp
@@ -54,18 +54,17 @@ void FControlRigControlPose::PastePose(UControlRig* ControlRig, bool bDoKey, boo
 
 }
 
-void FControlRigControlPose::SetControlMirrorTransform(bool bDoLocal, UControlRig* ControlRig, const FName& Name, bool bIsMatched, const FVector& GlobalTranslation, const FQuat& GlobalRotation, const FVector& LocalTranslation, const FQuat& LocalRotation, bool bNotify, const  FRigControlModifiedContext&Context,bool bSetupUndo)
+void FControlRigControlPose::SetControlMirrorTransform(bool bDoLocal, UControlRig* ControlRig, const FName& Name, bool bIsMatched, 
+	const FTransform& GlobalTransform, const FTransform& LocalTransform, bool bNotify, const  FRigControlModifiedContext&Context,bool bSetupUndo)
 {
 	if (bDoLocal || bIsMatched)
 	{
-		FTransform NewLocalTransform(LocalRotation, LocalTranslation);
-		ControlRig->SetControlLocalTransform(Name, NewLocalTransform, bNotify,Context,bSetupUndo);
+		ControlRig->SetControlLocalTransform(Name, LocalTransform, bNotify,Context,bSetupUndo);
 
 	}
 	else
 	{
-		FTransform NewGlobalTransform(GlobalRotation, GlobalTranslation);
-		ControlRig->SetControlGlobalTransform(Name, NewGlobalTransform, bNotify,Context,bSetupUndo);
+		ControlRig->SetControlGlobalTransform(Name, GlobalTransform, bNotify,Context,bSetupUndo);
 	}	
 }
 
@@ -127,13 +126,11 @@ void FControlRigControlPose::PastePoseInternal(UControlRig* ControlRig, bool bDo
 				}
 				else
 				{
-					FVector GlobalTranslation;
-					FQuat GlobalRotation;
-					FVector LocalTranslation;
-					FQuat LocalRotation;
+					FTransform GlobalTransform;
+					FTransform LocalTransform;
 					bool bIsMatched = MirrorTable.IsMatched(CopyRigControl->Name);
-					MirrorTable.GetMirrorTransform(*CopyRigControl, bDoLocal,bIsMatched, GlobalTranslation, GlobalRotation, LocalTranslation, LocalRotation);
-					SetControlMirrorTransform(bDoLocal,ControlRig, ControlElement->GetName(), bIsMatched, GlobalTranslation, GlobalRotation, LocalTranslation,LocalRotation, true,Context,bSetupUndo);
+					MirrorTable.GetMirrorTransform(*CopyRigControl, bDoLocal,bIsMatched,GlobalTransform,LocalTransform);
+					SetControlMirrorTransform(bDoLocal,ControlRig, ControlElement->GetName(), bIsMatched, GlobalTransform,LocalTransform,true,Context,bSetupUndo);
 				}				
 				break;
 			}
@@ -258,21 +255,27 @@ void FControlRigControlPose::BlendWithInitialPoses(FControlRigControlPose& Initi
 					}
 					else
 					{
-						FVector GlobalTranslation;
-						FQuat GlobalRotation;
-						FVector LocalTranslation;
-						FQuat LocalRotation;
+						FTransform GlobalTransform;
+						FTransform LocalTransform;
 						bool bIsMatched = MirrorTable.IsMatched(CopyRigControl->Name);
-						MirrorTable.GetMirrorTransform(*CopyRigControl,bDoLocal, bIsMatched, GlobalTranslation, GlobalRotation, LocalTranslation, LocalRotation);
+						MirrorTable.GetMirrorTransform(*CopyRigControl, bDoLocal, bIsMatched, GlobalTransform, LocalTransform);
 						FVector InitialTranslation = InitialFound->GlobalTransform.GetTranslation();
 						FQuat InitialGlobalRotation = InitialFound->GlobalTransform.GetRotation();
+						FVector InitialGlobalScale = InitialFound->GlobalTransform.GetScale3D();
+
 						FVector InitialLocalTranslation = InitialFound->LocalTransform.GetTranslation();
 						FQuat InitialLocationRotation = InitialFound->LocalTransform.GetRotation();
-						GlobalTranslation = FMath::Lerp(InitialTranslation, GlobalTranslation, BlendValue);
-						GlobalRotation = FQuat::Slerp(InitialGlobalRotation, GlobalRotation, BlendValue); //doing slerp here not fast lerp, can be slow this is for content creation
-						LocalTranslation = FMath::Lerp(InitialLocalTranslation, LocalTranslation, BlendValue);
-						LocalRotation = FQuat::Slerp(InitialLocationRotation, LocalRotation, BlendValue); //doing slerp here not fast lerp, can be slow this is for content creation
-						SetControlMirrorTransform(bDoLocal,ControlRig, ControlElement->GetName(), bIsMatched, GlobalTranslation, GlobalRotation,LocalTranslation, LocalRotation, bDoKey,Context,bSetupUndo);							
+						FVector InitialLocalScale = InitialFound->LocalTransform.GetScale3D();
+
+						GlobalTransform.SetTranslation(FMath::Lerp(InitialTranslation, GlobalTransform.GetTranslation(), BlendValue));
+						GlobalTransform.SetRotation(FQuat::Slerp(InitialGlobalRotation, GlobalTransform.GetRotation(), BlendValue)); //doing slerp here not fast lerp, can be slow this is for content creation
+						GlobalTransform.SetScale3D(FMath::Lerp(InitialGlobalScale, GlobalTransform.GetScale3D(), BlendValue));
+
+						LocalTransform.SetTranslation(FMath::Lerp(InitialLocalTranslation, LocalTransform.GetTranslation(), BlendValue));
+						LocalTransform.SetRotation(FQuat::Slerp(InitialLocationRotation, LocalTransform.GetRotation(), BlendValue)); //doing slerp here not fast lerp, can be slow this is for content creation
+						LocalTransform.SetScale3D(FMath::Lerp(InitialLocalScale, LocalTransform.GetScale3D(), BlendValue));
+
+						SetControlMirrorTransform(bDoLocal,ControlRig, ControlElement->GetName(), bIsMatched, GlobalTransform,LocalTransform,bDoKey,Context,bSetupUndo);							
 					}
 					if (bHasNonDefaultParent == true)
 					{
diff --git a/Engine/Plugins/Animation/ControlRig/Source/ControlRig/Private/Tools/ControlRigPoseMirrorTable.cpp b/Engine/Plugins/Animation/ControlRig/Source/ControlRig/Private/Tools/ControlRigPoseMirrorTable.cpp
--- a/Engine/Plugins/Animation/ControlRig/Source/ControlRig/Private/Tools/ControlRigPoseMirrorTable.cpp
+++ b/Engine/Plugins/Animation/ControlRig/Source/ControlRig/Private/Tools/ControlRigPoseMirrorTable.cpp
@@ -71,32 +71,28 @@ bool FControlRigPoseMirrorTable::IsMatched(const FName& Name) const
 }
 
 //Now returns mirrored global and local unmirrored
-void FControlRigPoseMirrorTable::GetMirrorTransform(const FRigControlCopy& ControlCopy, bool bDoLocal, bool bIsMatched, FVector& OutGlobalTranslation,
-	FQuat& OutGlobalRotation, FVector& OutLocalTranslation, FQuat& OutLocalRotation) const
+void FControlRigPoseMirrorTable::GetMirrorTransform(const FRigControlCopy& ControlCopy, bool bDoLocal, bool bIsMatched, FTransform& OutGlobalTransform,
+	FTransform& OutLocalTransform) const
 {
 	const UControlRigPoseMirrorSettings* Settings = GetDefault<UControlRigPoseMirrorSettings>();
 	FTransform GlobalTransform = ControlCopy.GlobalTransform;
-	FQuat GlobalRotation = GlobalTransform.GetRotation();
-
+	OutGlobalTransform = GlobalTransform;
 	FTransform LocalTransform = ControlCopy.LocalTransform;
-	OutLocalRotation = LocalTransform.GetRotation();
-	OutLocalTranslation = LocalTransform.GetTranslation();
-	OutGlobalTranslation = GlobalTransform.GetTranslation();
-	OutGlobalRotation = GlobalTransform.GetRotation();
+	OutLocalTransform = LocalTransform;
 	if (Settings)
 	{
-		if (!bIsMatched  && (OutGlobalRotation.IsIdentity() == false || OutLocalRotation.IsIdentity() == false))
+		if (!bIsMatched  && (OutGlobalTransform.GetRotation().IsIdentity() == false || OutLocalTransform.GetRotation().IsIdentity() == false))
 		{
 			FRigMirrorSettings MirrorSettings;
 			MirrorSettings.MirrorAxis = Settings->MirrorAxis;
 			MirrorSettings.AxisToFlip = Settings->AxisToFlip;
 			FTransform NewTransform = MirrorSettings.MirrorTransform(GlobalTransform);
-			OutGlobalTranslation = NewTransform.GetTranslation();
-			OutGlobalRotation = NewTransform.GetRotation();
+			OutGlobalTransform.SetTranslation(NewTransform.GetTranslation());
+			OutGlobalTransform.SetRotation(NewTransform.GetRotation());
 
 			FTransform NewLocalTransform = MirrorSettings.MirrorTransform(LocalTransform);
-			OutLocalTranslation = NewLocalTransform.GetTranslation();
-			OutLocalRotation = NewLocalTransform.GetRotation();
+			OutLocalTransform.SetTranslation(NewLocalTransform.GetTranslation());
+			OutLocalTransform.SetRotation(NewLocalTransform.GetRotation());
 			return;
 		}
 	}
diff --git a/Engine/Plugins/Animation/ControlRig/Source/ControlRig/Public/Tools/ControlRigPose.h b/Engine/Plugins/Animation/ControlRig/Source/ControlRig/Public/Tools/ControlRigPose.h
--- a/Engine/Plugins/Animation/ControlRig/Source/ControlRig/Public/Tools/ControlRigPose.h
+++ b/Engine/Plugins/Animation/ControlRig/Source/ControlRig/Public/Tools/ControlRigPose.h
@@ -92,8 +92,8 @@ struct CONTROLRIG_API FControlRigControlPose
 
 	void SavePose(UControlRig* ControlRig, bool bUseAll);
 	void PastePose(UControlRig* ControlRig, bool bDoKey, bool bDoMirror);
-	void SetControlMirrorTransform(bool bDoLocalSpace, UControlRig* ControlRig, const FName& Name, bool bIsMatched, const FVector& GlobalTranslation, const FQuat& GlobalRotation,const FVector& LocalTranslation,
-		const FQuat& LocalRotation, bool bNotify, const  FRigControlModifiedContext& Context, bool bSetupUndo);
+	void SetControlMirrorTransform(bool bDoLocalSpace, UControlRig* ControlRig, const FName& Name, bool bIsMatched,
+		const FTransform& GlobalTrnnsform, const FTransform& LocalTransform, bool bNotify, const  FRigControlModifiedContext& Context, bool bSetupUndo);
 	void PastePoseInternal(UControlRig* ControlRig, bool bDoKey, bool bDoMirror, const TArray<FRigControlCopy>& ControlsToPaste);
 	void BlendWithInitialPoses(FControlRigControlPose& InitialPose, UControlRig* ControlRig, bool bDoKey, bool bDoMirror, float BlendValue);
 
diff --git a/Engine/Plugins/Animation/ControlRig/Source/ControlRig/Public/Tools/ControlRigPoseMirrorTable.h b/Engine/Plugins/Animation/ControlRig/Source/ControlRig/Public/Tools/ControlRigPoseMirrorTable.h
--- a/Engine/Plugins/Animation/ControlRig/Source/ControlRig/Public/Tools/ControlRigPoseMirrorTable.h
+++ b/Engine/Plugins/Animation/ControlRig/Source/ControlRig/Public/Tools/ControlRigPoseMirrorTable.h
@@ -31,7 +31,7 @@ public:
 	bool IsMatched(const FName& ControlName) const;
 
 	/*Return the Mirrored Global(In Control Rig Space) Translation and Mirrored Local Rotation*/
-	void GetMirrorTransform(const FRigControlCopy& ControlCopy, bool bDoLocal, bool bIsMatched,FVector& OutGlobalTranslation, FQuat& OutGlobalRotation, FVector& OutLocalTranslation,FQuat& OutLocalRotation) const;
+	void GetMirrorTransform(const FRigControlCopy& ControlCopy, bool bDoLocal, bool bIsMatched,FTransform& OutGlobalTransform, FTransform& OutLocalTransform) const;
 
 private:
 
