Control Rig: change the SpacePicker window to be a popup window so that it automatically closes when user clicks somewhere else.

The old approach causes infinite recusion because RequestDestroyWindow would trigger WindowDeactivate, and the OnDeactivated delegate we use would then invoke closedialog() (namely, RequestDestroyWindow) again. It surprisingly worked on Windows, but did not work on Mac, which makes sense.

Current approach: if a space is picked, closedialog is triggered to close the window proactively. if user clicks outside of the window, normal popup window destruction path is invoked.

diff --git a/Engine/Plugins/Animation/ControlRig/Source/ControlRigEditor/Private/Editor/SRigSpacePickerWidget.cpp b/Engine/Plugins/Animation/ControlRig/Source/ControlRigEditor/Private/Editor/SRigSpacePickerWidget.cpp
--- a/Engine/Plugins/Animation/ControlRig/Source/ControlRigEditor/Private/Editor/SRigSpacePickerWidget.cpp
+++ b/Engine/Plugins/Animation/ControlRig/Source/ControlRigEditor/Private/Editor/SRigSpacePickerWidget.cpp
@@ -228,22 +228,6 @@ void SRigSpacePickerWidget::SetControls(
 
 class SRigSpaceDialogWindow : public SWindow
 {
-public:
-
-	virtual bool OnIsActiveChanged(const FWindowActivateEvent& ActivateEvent) override
-	{
-		if(ActivateEvent.GetActivationType() == FWindowActivateEvent::EA_Deactivate)
-		{
-			DeactivatedDelegate.ExecuteIfBound();
-		}
-		return SWindow::OnIsActiveChanged(ActivateEvent);
-	}
-
-	FSimpleDelegate& OnDeactivated() { return DeactivatedDelegate; }
-
-private:
-
-	FSimpleDelegate DeactivatedDelegate;
 }; 
 
 FReply SRigSpacePickerWidget::OpenDialog(bool bModal)
@@ -256,6 +240,7 @@ FReply SRigSpacePickerWidget::OpenDialog(bool bModal)
 	.Title( LOCTEXT("SRigSpacePickerWidgetPickSpace", "Pick a new space") )
 	.CreateTitleBar(false)
 	.Type(EWindowType::Menu)
+	.IsPopupWindow(true) // the window automatically closes when user clicks outside of it
 	.SizingRule( ESizingRule::Autosized )
 	.ScreenPosition(CursorPos)
 	.FocusWhenFirstShown(true)
@@ -265,7 +250,6 @@ FReply SRigSpacePickerWidget::OpenDialog(bool bModal)
 	];
 	
 	Window->SetWidgetToFocusOnActivate(AsShared());
-	Window->OnDeactivated().BindSP(this, &SRigSpacePickerWidget::CloseDialog);
 	
 	DialogWindow = Window;
 
