Exposing more of control rig to Blueprints/scripting so we can better build up tools around it.

diff --git a/Engine/Plugins/Animation/ControlRig/Source/ControlRig/Private/ControlRig.cpp b/Engine/Plugins/Animation/ControlRig/Source/ControlRig/Private/ControlRig.cpp
--- a/Engine/Plugins/Animation/ControlRig/Source/ControlRig/Private/ControlRig.cpp
+++ b/Engine/Plugins/Animation/ControlRig/Source/ControlRig/Private/ControlRig.cpp
@@ -431,6 +431,11 @@ TArray<FName> UControlRig::GetSupportedEvents() const
 	return TArray<FName>();
 }
 
+AActor* UControlRig::GetHostingActor() const
+{
+	return ObjectBinding ? ObjectBinding->GetHostingActor() : nullptr;
+}
+
 #if WITH_EDITOR
 FText UControlRig::GetCategory() const
 {
@@ -1714,6 +1719,8 @@ void UControlRig::HandleHierarchyModified(ERigHierarchyNotification InNotificati
 			{
 				const bool bSelected = InNotification == ERigHierarchyNotification::ElementSelected;
 				ControlSelected().Broadcast(this, ControlElement, bSelected);
+
+				OnControlSelected_BP.Broadcast(this, *ControlElement, bSelected);
 			}
 			break;
 		}
diff --git a/Engine/Plugins/Animation/ControlRig/Source/ControlRig/Public/ControlRig.h b/Engine/Plugins/Animation/ControlRig/Source/ControlRig/Public/ControlRig.h
--- a/Engine/Plugins/Animation/ControlRig/Source/ControlRig/Public/ControlRig.h
+++ b/Engine/Plugins/Animation/ControlRig/Source/ControlRig/Public/ControlRig.h
@@ -5,6 +5,7 @@
 #include "CoreMinimal.h"
 #include "UObject/ObjectMacros.h"
 #include "UObject/ScriptMacros.h"
+#include "UObject/SparseDelegate.h"
 #include "Engine/EngineBaseTypes.h"
 #include "Templates/SubclassOf.h"
 #include "ControlRigDefines.h"
@@ -36,6 +37,7 @@ class IControlRigObjectBinding;
 class UScriptStruct;
 class USkeletalMesh;
 class USkeletalMeshComponent;
+class AActor;
 
 struct FReferenceSkeleton;
 struct FRigUnit;
@@ -62,6 +64,8 @@ public:
 
 	/** Bindable event for external objects to be notified that a Control is Selected */
 	DECLARE_EVENT_ThreeParams(UControlRig, FControlSelectedEvent, UControlRig*, FRigControlElement*, bool);
+	// To support Blueprints/scripting, we need a different delegate type (a 'Dynamic' delegate) which supports looser style UFunction binding (using names).
+	DECLARE_DYNAMIC_MULTICAST_SPARSE_DELEGATE_ThreeParams(FOnControlSelectedBP, UControlRig, OnControlSelected_BP, UControlRig*, Rig, const FRigControlElement&, Control, bool, bSelected);
 
 #if WITH_EDITOR
 	/** Bindable event for external objects to be notified that a control rig is fully end-loaded*/
@@ -221,6 +225,10 @@ public:
 		return ObjectBinding;
 	}
 
+	/** Find the actor the rig is bound to, if any */
+	UFUNCTION(BlueprintPure, Category = "Control Rig")
+	AActor* GetHostingActor() const;
+
 	virtual FString GetName() const
 	{
 		FString ObjectName = (GetClass()->GetName());
@@ -699,6 +707,9 @@ protected:
 	FControlModifiedEvent OnControlModified;
 	FControlSelectedEvent OnControlSelected;
 
+	UPROPERTY(BlueprintAssignable, Category = ControlRig, meta=(DisplayName="OnControlSelected"))
+	FOnControlSelectedBP OnControlSelected_BP;
+
 	TArray<FRigElementKey> QueuedModifiedControls;
 
 private:
