Sequencer: Space Switching Curve Editor Space view was getting lost on undo since the ChannelHandle would become invaid, need to keep track of OnDestroy callback so we can recreate it like the other curve models do.

diff --git a/Engine/Plugins/Animation/ControlRig/Source/ControlRigEditor/Private/ControlRigSpaceChannelCurveModel.cpp b/Engine/Plugins/Animation/ControlRig/Source/ControlRigEditor/Private/ControlRigSpaceChannelCurveModel.cpp
--- a/Engine/Plugins/Animation/ControlRig/Source/ControlRigEditor/Private/ControlRigSpaceChannelCurveModel.cpp
+++ b/Engine/Plugins/Animation/ControlRig/Source/ControlRigEditor/Private/ControlRigSpaceChannelCurveModel.cpp
@@ -20,7 +20,13 @@ ECurveEditorViewID FControlRigSpaceChannelCurveModel::ViewID = ECurveEditorViewI
 
 FControlRigSpaceChannelCurveModel::FControlRigSpaceChannelCurveModel(TMovieSceneChannelHandle<FMovieSceneControlRigSpaceChannel> InChannel, UMovieSceneSection* OwningSection, TWeakPtr<ISequencer> InWeakSequencer)
 {
-	ChannelHandle = InChannel;
+	FMovieSceneChannelProxy* NewChannelProxy = &OwningSection->GetChannelProxy();
+	ChannelHandle = NewChannelProxy->MakeHandle<FMovieSceneControlRigSpaceChannel>(InChannel.GetChannelIndex());
+	if (FMovieSceneChannelProxy* ChannelProxy = InChannel.GetChannelProxy())
+	{
+		OnDestroyHandle = NewChannelProxy->OnDestroy.AddRaw(this, &FControlRigSpaceChannelCurveModel::FixupCurve);
+	}
+
 	WeakSection = OwningSection;
 	WeakSequencer = InWeakSequencer;
 	SupportedViews = ViewID;
@@ -28,6 +34,23 @@ FControlRigSpaceChannelCurveModel::FControlRigSpaceChannelCurveModel(TMovieScene
 	Color = FSequencerSectionPainter::BlendColor(OwningSection->GetTypedOuter<UMovieSceneTrack>()->GetColorTint());
 }
 
+FControlRigSpaceChannelCurveModel::~FControlRigSpaceChannelCurveModel()
+{
+	if (FMovieSceneChannelProxy* ChannelProxy = ChannelHandle.GetChannelProxy())
+	{
+		ChannelProxy->OnDestroy.Remove(OnDestroyHandle);
+	}
+}
+
+void FControlRigSpaceChannelCurveModel::FixupCurve()
+{
+	if (UMovieSceneSection* Section = WeakSection.Get())
+	{
+		FMovieSceneChannelProxy* NewChannelProxy = &Section->GetChannelProxy();
+		ChannelHandle = NewChannelProxy->MakeHandle<FMovieSceneControlRigSpaceChannel>(ChannelHandle.GetChannelIndex());
+		OnDestroyHandle = NewChannelProxy->OnDestroy.AddRaw(this, &FControlRigSpaceChannelCurveModel::FixupCurve);
+	}
+}
 TArray<FKeyBarCurveModel::FBarRange> FControlRigSpaceChannelCurveModel::FindRanges()
 {
 	FMovieSceneControlRigSpaceChannel* Channel = ChannelHandle.Get();
@@ -119,7 +142,7 @@ void FControlRigSpaceChannelCurveModel::GetKeys(const FCurveEditor& CurveEditor,
 {
 	FMovieSceneControlRigSpaceChannel* Channel = ChannelHandle.Get();
 	UMovieSceneSection*      Section = WeakSection.Get();
-
+	
 	if (Channel && Section)
 	{
 		FFrameRate TickResolution = Section->GetTypedOuter<UMovieScene>()->GetTickResolution();
diff --git a/Engine/Plugins/Animation/ControlRig/Source/ControlRigEditor/Public/ControlRigSpaceChannelCurveModel.h b/Engine/Plugins/Animation/ControlRig/Source/ControlRigEditor/Public/ControlRigSpaceChannelCurveModel.h
--- a/Engine/Plugins/Animation/ControlRig/Source/ControlRigEditor/Public/ControlRigSpaceChannelCurveModel.h
+++ b/Engine/Plugins/Animation/ControlRig/Source/ControlRigEditor/Public/ControlRigSpaceChannelCurveModel.h
@@ -19,7 +19,7 @@ public:
 	static ECurveEditorViewID ViewID;
 
 	FControlRigSpaceChannelCurveModel(TMovieSceneChannelHandle<FMovieSceneControlRigSpaceChannel> InChannel, UMovieSceneSection* InOwningSection, TWeakPtr<ISequencer> InWeakSequencer);
-
+	~FControlRigSpaceChannelCurveModel();
 	//FCurveModel
 	virtual const void* GetCurve() const override;
 
@@ -51,8 +51,12 @@ public:
 	virtual void BuildContextMenu(const FCurveEditor& CurveEditor,FMenuBuilder& MenuBuilder, TOptional<FCurvePointHandle> ClickedPoint);
 	virtual TArray<FKeyBarCurveModel::FBarRange> FindRanges();
 
+
+private:
+	void FixupCurve();
 private:
 	TMovieSceneChannelHandle<FMovieSceneControlRigSpaceChannel> ChannelHandle;
 	TWeakObjectPtr<UMovieSceneSection> WeakSection;
 	TWeakPtr<ISequencer> WeakSequencer;
+	FDelegateHandle OnDestroyHandle;
 };
\ No newline at end of file
