Control Rig: offer additional transform as part of gizmo actors

diff --git a/Engine/Plugins/Animation/ControlRig/Source/ControlRigEditor/Private/EditMode/ControlRigEditMode.cpp b/Engine/Plugins/Animation/ControlRig/Source/ControlRigEditor/Private/EditMode/ControlRigEditMode.cpp
--- a/Engine/Plugins/Animation/ControlRig/Source/ControlRigEditor/Private/EditMode/ControlRigEditMode.cpp
+++ b/Engine/Plugins/Animation/ControlRig/Source/ControlRigEditor/Private/EditMode/ControlRigEditMode.cpp
@@ -694,7 +694,12 @@ void FControlRigEditMode::Render(const FSceneView* View, FViewport* Viewport, FP
 
 	bool bRender = !Settings->bHideControlShapes;
 
-	FTransform ComponentTransform = GetHostingSceneComponentTransform();
+	FTransform ComponentTransform = FTransform::Identity;
+	if(IsInLevelEditor())
+	{
+		ComponentTransform = GetHostingSceneComponentTransform();
+	}
+	
 	if (bRender)
 	{
 		for (AControlRigShapeActor* Actor : ShapeActors)
@@ -3362,7 +3367,12 @@ void FControlRigEditMode::OnPoseInitialized()
 
 void FControlRigEditMode::PostPoseUpdate()
 {
-	FTransform ComponentTransform = GetHostingSceneComponentTransform();
+	FTransform ComponentTransform = FTransform::Identity;
+	if(IsInLevelEditor())
+	{
+		ComponentTransform = GetHostingSceneComponentTransform();
+	}
+	
 	for(AControlRigShapeActor* ShapeActor : ShapeActors)
 	{
 		TickControlShape(ShapeActor, ComponentTransform);
diff --git a/Engine/Plugins/Animation/ControlRig/Source/ControlRigEditor/Private/Editor/ControlRigEditor.cpp b/Engine/Plugins/Animation/ControlRig/Source/ControlRigEditor/Private/Editor/ControlRigEditor.cpp
--- a/Engine/Plugins/Animation/ControlRig/Source/ControlRigEditor/Private/Editor/ControlRigEditor.cpp
+++ b/Engine/Plugins/Animation/ControlRig/Source/ControlRigEditor/Private/Editor/ControlRigEditor.cpp
@@ -1288,12 +1288,9 @@ void FControlRigEditor::HandleSetObjectBeingDebugged(UObject* InObject)
 
 		DebuggedControlRig->GetHierarchy()->OnModified().AddSP(this, &FControlRigEditor::OnHierarchyModified_AnyThread);
 
-		if(!bIsExternalControlRig)
+		if(EditorSkelComp)
 		{
-			if(EditorSkelComp)
-			{
-				EditorSkelComp->SetComponentToWorld(FTransform::Identity);
-			}
+			EditorSkelComp->SetComponentToWorld(FTransform::Identity);
 		}
 	}
 	else
@@ -3032,21 +3029,6 @@ void FControlRigEditor::Tick(float DeltaTime)
 			bDrawHierarchyBones = true;
 		}
 
-		if(UControlRig* DebuggedControlRig = Cast<UControlRig>(GetBlueprintObj()->GetObjectBeingDebugged()))
-		{
-			const bool bIsExternalControlRig = DebuggedControlRig != ControlRig;			
-			if(bIsExternalControlRig)
-			{
-				if(UControlRigSkeletalMeshComponent* EditorSkelComp = Cast<UControlRigSkeletalMeshComponent>(GetPersonaToolkit()->GetPreviewScene()->GetPreviewMeshComponent()))
-				{
-					if(DebuggedControlRig->OuterSceneComponent.IsValid())
-					{
-						EditorSkelComp->SetComponentToWorld(DebuggedControlRig->OuterSceneComponent->GetComponentToWorld());
-					}
-				}
-			}
-		}
-
 		if (LastDebuggedRig != GetCustomDebugObjectLabel(Blueprint->GetObjectBeingDebugged()))
 		{
 			TArray<FCustomDebugObject> DebugList;
