Control Rig: Performance optimization when adding/breaking link on reroute

diff --git a/Engine/Plugins/Animation/ControlRig/Source/ControlRigDeveloper/Private/Graph/ControlRigGraph.cpp b/Engine/Plugins/Animation/ControlRig/Source/ControlRigDeveloper/Private/Graph/ControlRigGraph.cpp
--- a/Engine/Plugins/Animation/ControlRig/Source/ControlRigDeveloper/Private/Graph/ControlRigGraph.cpp
+++ b/Engine/Plugins/Animation/ControlRig/Source/ControlRigDeveloper/Private/Graph/ControlRigGraph.cpp
@@ -9,6 +9,7 @@
 #include "ControlRigObjectVersion.h"
 #include "Units/RigUnit.h"
 #include "EdGraphNode_Comment.h"
+#include "GraphEditAction.h"
 #include "ControlRig/Private/Units/Execution/RigUnit_BeginExecution.h"
 #include "RigVMModel/Nodes/RigVMLibraryNode.h"
 #include "RigVMModel/Nodes/RigVMFunctionEntryNode.h"
@@ -460,7 +461,8 @@ void UControlRigGraph::HandleModifiedEvent(ERigVMGraphNotifType InNotifType, URi
 						{
 							RigNode->Pins[PinIndex]->bHidden = !ModelNode->GetShowsAsFullNode();
 						}
-						NotifyGraphChanged();
+						NotifyGraphChanged(FEdGraphEditAction(EEdGraphActionType::GRAPHACTION_RemoveNode, EdNode->GetGraph(), EdNode, true));
+						NotifyGraphChanged(FEdGraphEditAction(EEdGraphActionType::GRAPHACTION_AddNode, EdNode->GetGraph(), EdNode, true));
 					}
 				}
 			}
