Control Rig: Always show user defined struct in the pin type drop down list. Even if the struct contains unsupported members, it is impractical to load and check all of them when the drop down menu is invoked.

diff --git a/Engine/Plugins/Animation/ControlRig/Source/ControlRigDeveloper/Private/Graph/ControlRigGraphSchema.cpp b/Engine/Plugins/Animation/ControlRig/Source/ControlRigDeveloper/Private/Graph/ControlRigGraphSchema.cpp
--- a/Engine/Plugins/Animation/ControlRig/Source/ControlRigDeveloper/Private/Graph/ControlRigGraphSchema.cpp
+++ b/Engine/Plugins/Animation/ControlRig/Source/ControlRigDeveloper/Private/Graph/ControlRigGraphSchema.cpp
@@ -922,6 +922,15 @@ bool UControlRigGraphSchema::SupportsPinType(TWeakPtr<const FEdGraphSchemaAction
 			}
 			return SupportsPinType(ScriptStruct);
 		}
+		else if (PinType.PinSubCategoryObject == UUserDefinedStruct::StaticClass())
+		{
+			// if a user defined struct hasn't been loaded yet,
+			// its PinSubCategoryObject equals UUserDefinedStruct::StaticClass()
+			// and since it is not practical to load every user defined struct to check if they only contain supported types,
+			// we always return true so that they at least show up in the drop down menu
+			// if they contain members of unsupported type, the spawned node will generate error
+			return true;
+		}
 	}
 
 	if (PinType.PinCategory == UEdGraphSchema_K2::PC_Byte)
