Control Rig: Fix crash when Blueprint Diff and control rig editor are opened at the same time

diff --git a/Engine/Plugins/Animation/ControlRig/Source/ControlRigEditor/Private/ControlRigBlueprintActions.cpp b/Engine/Plugins/Animation/ControlRig/Source/ControlRigEditor/Private/ControlRigBlueprintActions.cpp
--- a/Engine/Plugins/Animation/ControlRig/Source/ControlRigEditor/Private/ControlRigBlueprintActions.cpp
+++ b/Engine/Plugins/Animation/ControlRig/Source/ControlRigEditor/Private/ControlRigBlueprintActions.cpp
@@ -29,6 +29,7 @@
 #include "Editor.h"
 #include "EditorModeManager.h"
 #include "MovieSceneToolsProjectSettings.h"
+#include "SBlueprintDiff.h"
 
 #define LOCTEXT_NAMESPACE "ControlRigBlueprintActions"
 
@@ -51,12 +52,16 @@ void FControlRigBlueprintActions::OpenAssetEditor( const TArray<UObject*>& InObj
 		if (UControlRigBlueprint* ControlRigBlueprint = Cast<UControlRigBlueprint>(*ObjIt))
 		{
 			const bool bBringToFrontIfOpen = true;
-			if (IAssetEditorInstance* EditorInstance = GEditor->GetEditorSubsystem<UAssetEditorSubsystem>()->FindEditorForAsset(ControlRigBlueprint, bBringToFrontIfOpen))
+			UAssetEditorSubsystem* AssetEditorSubsystem = GEditor->GetEditorSubsystem<UAssetEditorSubsystem>();
+			if (IAssetEditorInstance* EditorInstance = AssetEditorSubsystem->FindEditorForAsset(ControlRigBlueprint, bBringToFrontIfOpen))
 			{
 				EditorInstance->FocusWindow(ControlRigBlueprint);
 			}
 			else
-			{
+			{				
+				// If any other editors are opened (for example, a BlueprintDiff window), close them 
+				AssetEditorSubsystem->CloseAllEditorsForAsset(ControlRigBlueprint);				
+				
 				IControlRigEditorModule& ControlRigEditorModule = FModuleManager::LoadModuleChecked<IControlRigEditorModule>("ControlRigEditor");
 				ControlRigEditorModule.CreateControlRigEditor(Mode, EditWithinLevelEditor, ControlRigBlueprint);
 			}
@@ -80,6 +85,49 @@ TSharedPtr<SWidget> FControlRigBlueprintActions::GetThumbnailOverlay(const FAsse
 		];
 }
 
+void FControlRigBlueprintActions::PerformAssetDiff(UObject* OldAsset, UObject* NewAsset, const FRevisionInfo& OldRevision, const FRevisionInfo& NewRevision) const
+{
+	UBlueprint* OldBlueprint = CastChecked<UBlueprint>(OldAsset);
+	UBlueprint* NewBlueprint = CastChecked<UBlueprint>(NewAsset);
+
+	UAssetEditorSubsystem* AssetEditorSubsystem = GEditor->GetEditorSubsystem<UAssetEditorSubsystem>();
+	for (IAssetEditorInstance* Editor : AssetEditorSubsystem->FindEditorsForAsset(OldAsset))
+	{
+		const EAppReturnType::Type Answer = FMessageDialog::Open( EAppMsgType::YesNo,
+				FText::FromString(FString::Printf(TEXT("Opening a diff window will close the control rig editor. %s.\nAre you sure?"),  *OldBlueprint->GetName() )));
+		if(Answer == EAppReturnType::No)
+		{
+		   return;
+		}
+	}
+	for (IAssetEditorInstance* Editor : AssetEditorSubsystem->FindEditorsForAsset(NewAsset))
+	{
+		const EAppReturnType::Type Answer = FMessageDialog::Open( EAppMsgType::YesNo,
+				FText::FromString(FString::Printf(TEXT("Opening a diff window will close the control rig editor. %s.\nAre you sure?"),  *NewBlueprint->GetName() )));
+		if(Answer == EAppReturnType::No)
+		{
+			return;
+		}
+	}
+	
+	AssetEditorSubsystem->CloseAllEditorsForAsset(OldAsset);
+	AssetEditorSubsystem->CloseAllEditorsForAsset(NewAsset);
+
+	// sometimes we're comparing different revisions of one single asset (other 
+	// times we're comparing two completely separate assets altogether)
+	bool bIsSingleAsset = (NewBlueprint->GetName() == OldBlueprint->GetName());
+
+	FText WindowTitle = LOCTEXT("ControlRigBlueprintDiff", "Control Rig Blueprint Diff");
+	// if we're diffing one asset against itself 
+	if (bIsSingleAsset)
+	{
+		// identify the assumed single asset in the window's title
+		WindowTitle = FText::Format(LOCTEXT("Control Rig Diff", "{0} - Control Rig Diff"), FText::FromString(NewBlueprint->GetName()));
+	}
+
+	SBlueprintDiff::CreateDiffWindow(WindowTitle, OldBlueprint, NewBlueprint, OldRevision, NewRevision);
+}
+
 void FControlRigBlueprintActions::ExtendSketalMeshToolMenu()
 {
 	TArray<UToolMenu*> MenusToExtend;
diff --git a/Engine/Plugins/Animation/ControlRig/Source/ControlRigEditor/Private/ControlRigBlueprintActions.h b/Engine/Plugins/Animation/ControlRig/Source/ControlRigEditor/Private/ControlRigBlueprintActions.h
--- a/Engine/Plugins/Animation/ControlRig/Source/ControlRigEditor/Private/ControlRigBlueprintActions.h
+++ b/Engine/Plugins/Animation/ControlRig/Source/ControlRigEditor/Private/ControlRigBlueprintActions.h
@@ -22,6 +22,7 @@ public:
 	virtual void OpenAssetEditor( const TArray<UObject*>& InObjects, TSharedPtr<class IToolkitHost> EditWithinLevelEditor = TSharedPtr<IToolkitHost>() ) override;
 	virtual uint32 GetCategories() override { return EAssetTypeCategories::Animation; }
 	virtual TSharedPtr<SWidget> GetThumbnailOverlay(const FAssetData& AssetData) const override;
+	virtual void PerformAssetDiff(UObject* Asset1, UObject* Asset2, const struct FRevisionInfo& OldRevision, const struct FRevisionInfo& NewRevision) const override;
 
 	// FAssetTypeActions_Blueprint interface
 	virtual UFactory* GetFactoryForBlueprintType(UBlueprint* InBlueprint) const override;
