Control Rig: Fix compilation error in ControlRigBlueprintActions.cpp

diff --git a/Engine/Plugins/Animation/ControlRig/Source/ControlRigEditor/Private/ControlRigBlueprintActions.cpp b/Engine/Plugins/Animation/ControlRig/Source/ControlRigEditor/Private/ControlRigBlueprintActions.cpp
--- a/Engine/Plugins/Animation/ControlRig/Source/ControlRigEditor/Private/ControlRigBlueprintActions.cpp
+++ b/Engine/Plugins/Animation/ControlRig/Source/ControlRigEditor/Private/ControlRigBlueprintActions.cpp
@@ -30,6 +30,7 @@
 #include "EditorModeManager.h"
 #include "MovieSceneToolsProjectSettings.h"
 #include "SBlueprintDiff.h"
+#include "Misc/MessageDialog.h"
 
 #define LOCTEXT_NAMESPACE "ControlRigBlueprintActions"
 
