Sequencer: Clean up duplicate LOCTEXT and remove unnecessary NSLOCTEXT("Sequencer") since its already in a namespace.

diff --git a/Engine/Plugins/Animation/ControlRig/Source/ControlRigEditor/Private/EditMode/SControlRigBaseListWidget.cpp b/Engine/Plugins/Animation/ControlRig/Source/ControlRigEditor/Private/EditMode/SControlRigBaseListWidget.cpp
--- a/Engine/Plugins/Animation/ControlRig/Source/ControlRigEditor/Private/EditMode/SControlRigBaseListWidget.cpp
+++ b/Engine/Plugins/Animation/ControlRig/Source/ControlRigEditor/Private/EditMode/SControlRigBaseListWidget.cpp
@@ -204,13 +204,13 @@ void FCreateControlAssetRigDialog::GetControlAssetParams(FControlRigAssetType Ty
 	switch (Type)
 	{
 	case FControlRigAssetType::ControlRigPose:
-		TitleText = LOCTEXT("CreateControlAssetRig", "Create Control Rig Pose");
+		TitleText = LOCTEXT("CreateControlAssetRigPose", "Create Control Rig Pose");
 		break;
 	case FControlRigAssetType::ControlRigAnimation:
-		TitleText = LOCTEXT("CreateControlAssetRig", "Create Control Rig Animation");
+		TitleText = LOCTEXT("CreateControlAssetRigAnimation", "Create Control Rig Animation");
 		break;
 	case FControlRigAssetType::ControlRigSelectionSet:
-		TitleText = LOCTEXT("CreateControlAssetRig", "Create Control Rig Selection Set");
+		TitleText = LOCTEXT("CreateControlAssetRigSelectionSet", "Create Control Rig Selection Set");
 		break;
 	};
 
@@ -1108,7 +1108,7 @@ TSharedPtr<SWidget> SControlRigBaseListWidget::OnGetAssetContextMenu(const TArra
 				FUIAction Action = FUIAction(FExecuteAction::CreateRaw((this), &SControlRigBaseListWidget::ExecutePasteMirrorPose, PoseAsset),
 					FCanExecuteAction::CreateRaw(this, &SControlRigBaseListWidget::CanExecutePasteMirrorPose, PoseAsset));
 				const FText Label = LOCTEXT("PasteMirrorPoseButton", "Paste Mirror Pose");
-				const FText ToolTipText = LOCTEXT("PastePoseButtonTooltip", "Paste the mirror of the pose");
+				const FText ToolTipText = LOCTEXT("PasteMirrorPoseButtonTooltip", "Paste the mirror of the pose");
 				MenuBuilder.AddMenuEntry(Label, ToolTipText, FSlateIcon(), Action);
 			}
 			MenuBuilder.EndSection();
diff --git a/Engine/Plugins/Animation/ControlRig/Source/ControlRigEditor/Private/EditMode/SControlRigTweenWidget.cpp b/Engine/Plugins/Animation/ControlRig/Source/ControlRigEditor/Private/EditMode/SControlRigTweenWidget.cpp
--- a/Engine/Plugins/Animation/ControlRig/Source/ControlRigEditor/Private/EditMode/SControlRigTweenWidget.cpp
+++ b/Engine/Plugins/Animation/ControlRig/Source/ControlRigEditor/Private/EditMode/SControlRigTweenWidget.cpp
@@ -88,7 +88,7 @@ void SControlRigTweenWidget::OnBeginSliderMovement()
 	{
 		bIsBlending = true;
 		bSliderStartedTransaction = true;
-		GEditor->BeginTransaction(LOCTEXT("TweenTransaction", "Tween It"));
+		GEditor->BeginTransaction(LOCTEXT("TweenTransaction", "Tween"));
 		SetupControls();
 	}
 }
diff --git a/Engine/Plugins/Animation/ControlRig/Source/ControlRigEditor/Private/Sequencer/ControlRigParameterTrackEditor.cpp b/Engine/Plugins/Animation/ControlRig/Source/ControlRigEditor/Private/Sequencer/ControlRigParameterTrackEditor.cpp
--- a/Engine/Plugins/Animation/ControlRig/Source/ControlRigEditor/Private/Sequencer/ControlRigParameterTrackEditor.cpp
+++ b/Engine/Plugins/Animation/ControlRig/Source/ControlRigEditor/Private/Sequencer/ControlRigParameterTrackEditor.cpp
@@ -517,8 +517,8 @@ void FControlRigParameterTrackEditor::BuildObjectBindingContextMenu(FMenuBuilder
 					EUserInterfaceActionType::Button);
 
 				MenuBuilder.AddMenuEntry(
-					NSLOCTEXT("Sequencer", "FilterAssetBySkeleton", "Filter Asset By Skeleton"),
-					NSLOCTEXT("Sequencer", "FilterAssetBySkeletonTooltip", "Filters Control Rig assets to match current skeleton"),
+					LOCTEXT("FilterAssetBySkeleton", "Filter Asset By Skeleton"),
+					LOCTEXT("FilterAssetBySkeletonTooltip", "Filters Control Rig assets to match current skeleton"),
 					FSlateIcon(),
 					FUIAction(
 						FExecuteAction::CreateSP(this, &FControlRigParameterTrackEditor::ToggleFilterAssetBySkeleton),
@@ -1068,8 +1068,8 @@ void FControlRigParameterTrackEditor::BuildObjectBindingTrackMenu(FMenuBuilder&
 			if (UControlRigComponent* ControlRigComponent = BoundActor->FindComponentByClass<UControlRigComponent>())
 			{
 				MenuBuilder.AddMenuEntry(
-					LOCTEXT("AddControlRig", "Animation ControlRig"),
-					NSLOCTEXT("Sequencer", "AddControlRigTooltip", "Adds an animation Control Rig track"),
+					LOCTEXT("AddControlRigTrack", "Add Control Rig Track"),
+					LOCTEXT("AddControlRigTrackTooltip", "Adds an animation Control Rig track"),
 					FSlateIcon(),
 					FUIAction(
 						FExecuteAction::CreateSP(this, &FControlRigParameterTrackEditor::AddControlRigFromComponent, ObjectBindings[0]),
@@ -1101,7 +1101,7 @@ void FControlRigParameterTrackEditor::HandleAddTrackSubMenu(FMenuBuilder& MenuBu
 {
 	MenuBuilder.AddMenuEntry(
 		LOCTEXT("AddFKControlRig", "FK Control Rig"),
-		NSLOCTEXT("Sequencer", "AddFKControlRigTooltip", "Adds an FK Control Rig track"),
+		LOCTEXT("AddFKControlRigTooltip", "Adds an FK Control Rig track"),
 		FSlateIcon(),
 		FUIAction(
 			FExecuteAction::CreateSP(this, &FControlRigParameterTrackEditor::AddFKControlRig, ObjectBindings),
@@ -1110,8 +1110,8 @@ void FControlRigParameterTrackEditor::HandleAddTrackSubMenu(FMenuBuilder& MenuBu
 	);
 
 	MenuBuilder.AddMenuEntry(
-		NSLOCTEXT("Sequencer", "FilterAssetBySkeleton", "Filter Asset By Skeleton"),
-		NSLOCTEXT("Sequencer", "FilterAssetBySkeletonTooltip", "Filters Control Rig assets to match current skeleton"),
+		LOCTEXT("FilterAssetBySkeleton", "Filter Asset By Skeleton"),
+		LOCTEXT("FilterAssetBySkeletonTooltip", "Filters Control Rig assets to match current skeleton"),
 		FSlateIcon(),
 		FUIAction(
 			FExecuteAction::CreateSP(this, &FControlRigParameterTrackEditor::ToggleFilterAssetBySkeleton),
@@ -1122,8 +1122,8 @@ void FControlRigParameterTrackEditor::HandleAddTrackSubMenu(FMenuBuilder& MenuBu
 		EUserInterfaceActionType::ToggleButton);
 
 	MenuBuilder.AddMenuEntry(
-		NSLOCTEXT("Sequencer", "FilterAssetByAnimatableControls", "Filter Asset By Animatable Controls"),
-		NSLOCTEXT("Sequencer", "FilterAssetByAnimatableControlsTooltip", "Filters Control Rig assets to only show those with Animatable Controls"),
+		LOCTEXT("FilterAssetByAnimatableControls", "Filter Asset By Animatable Controls"),
+		LOCTEXT("FilterAssetByAnimatableControlsTooltip", "Filters Control Rig assets to only show those with Animatable Controls"),
 		FSlateIcon(),
 		FUIAction(
 			FExecuteAction::CreateSP(this, &FControlRigParameterTrackEditor::ToggleFilterAssetByAnimatableControls),
@@ -1135,7 +1135,7 @@ void FControlRigParameterTrackEditor::HandleAddTrackSubMenu(FMenuBuilder& MenuBu
 
 	MenuBuilder.AddSubMenu(
 		LOCTEXT("AddAssetControlRig", "Asset-Based Control Rig"),
-		NSLOCTEXT("Sequencer", "AddAsetControlRigTooltip", "Adds an asset based Control Rig track"),
+		LOCTEXT("AddAsetControlRigTooltip", "Adds an asset based Control Rig track"),
 		FNewMenuDelegate::CreateRaw(this, &FControlRigParameterTrackEditor::HandleAddControlRigSubMenu, ObjectBindings, Track)
 	);
 }
@@ -1291,7 +1291,7 @@ void FControlRigParameterTrackEditor::AddControlRig(UClass* InClass, UObject* Bo
 	{
 		UMovieSceneSequence* OwnerSequence = GetSequencer()->GetFocusedMovieSceneSequence();
 		UMovieScene* OwnerMovieScene = OwnerSequence->GetMovieScene();
-		FScopedTransaction AddControlRigTrackTransaction(LOCTEXT("AddControlRigTrack_Transaction", "Add Control Rig Track"));
+		FScopedTransaction AddControlRigTrackTransaction(LOCTEXT("AddControlRigTrack", "Add Control Rig Track"));
 
 		OwnerSequence->Modify();
 		OwnerMovieScene->Modify();
@@ -2945,11 +2945,11 @@ void FControlRigParameterTrackEditor::BuildTrackContextMenu(FMenuBuilder& MenuBu
 
 	TArray<FFBXNodeAndChannels>* NodeAndChannels = Track->GetNodeAndChannelMappings(SectionToKey);
 
-	MenuBuilder.BeginSection("Import To Control Rig", NSLOCTEXT("Sequencer", "ImportToControlRig", "Import To Control Rig"));
+	MenuBuilder.BeginSection("Import To Control Rig", LOCTEXT("ImportToControlRig", "Import To Control Rig"));
 	{
 		MenuBuilder.AddMenuEntry(
-			NSLOCTEXT("Sequencer", "ImportControlRigFBX", "Import Control Rig FBX"),
-			NSLOCTEXT("Sequencer", "ImportControlRigFBXTooltip", "Import Control Rig FBX"),
+			LOCTEXT("ImportControlRigFBX", "Import Control Rig FBX"),
+			LOCTEXT("ImportControlRigFBXTooltip", "Import Control Rig FBX"),
 			FSlateIcon(),
 			FUIAction(
 				FExecuteAction::CreateRaw(this, &FControlRigParameterTrackEditor::ImportFBX, Track, SectionToKey, NodeAndChannels)));
@@ -2960,19 +2960,19 @@ void FControlRigParameterTrackEditor::BuildTrackContextMenu(FMenuBuilder& MenuBu
 
 	if (UFKControlRig* AutoRig = Cast<UFKControlRig>(Track->GetControlRig()))
 	{
-		MenuBuilder.BeginSection("FK Control Rig", NSLOCTEXT("Sequencer", "FKControlRig", "FK Control Rig"));
+		MenuBuilder.BeginSection("FK Control Rig", LOCTEXT("FKControlRig", "FK Control Rig"));
 		{
 
 			MenuBuilder.AddMenuEntry(
-				NSLOCTEXT("Sequencer", "SelectBonesToAnimate", "Select Bones Or Curves To Animate"),
-				NSLOCTEXT("Sequencer", "SelectBonesToAnimateToolTip", "Select which bones or curves you want to directly animate"),
+				LOCTEXT("SelectBonesToAnimate", "Select Bones Or Curves To Animate"),
+				LOCTEXT("SelectBonesToAnimateToolTip", "Select which bones or curves you want to directly animate"),
 				FSlateIcon(),
 				FUIAction(
 					FExecuteAction::CreateRaw(this, &FControlRigParameterTrackEditor::SelectFKBonesToAnimate, AutoRig,Track)));
 
 			MenuBuilder.AddMenuEntry(
-				NSLOCTEXT("Sequencer", "FKRigApplyMode", "Additive"),
-				NSLOCTEXT("Sequencer", "FKRigApplyModeToolTip", "Toggles the apply mode between Replace and Additive"),
+				LOCTEXT("FKRigApplyMode", "Additive"),
+				LOCTEXT("FKRigApplyModeToolTip", "Toggles the apply mode between Replace and Additive"),
 				FSlateIcon(),
 				FUIAction(
 					FExecuteAction::CreateRaw(this, &FControlRigParameterTrackEditor::ToggleFKControlRig, Track, AutoRig),
@@ -3012,7 +3012,7 @@ bool FControlRigParameterTrackEditor::HandleAssetAdded(UObject* Asset, const FGu
 		return false;
 	}
 
-	const FScopedTransaction Transaction(LOCTEXT("AddControlRig", "Add Control Rig"));
+	const FScopedTransaction Transaction(LOCTEXT("AddControlRigAsset", "Add Control Rig"));
 
 	// Spawn a skeletal mesh actor with the preview mesh
 	ASkeletalMeshActor* SkeletalMeshActor = GCurrentLevelEditingViewportClient->GetWorld()->SpawnActor<ASkeletalMeshActor>();
@@ -3337,7 +3337,7 @@ void FControlRigParameterTrackEditor::SelectFKBonesToAnimate(UFKControlRig* Auto
 {
 	if (AutoRig)
 	{
-		const FText TitleText = NSLOCTEXT("Sequencer", "SelectBonesOrCurvesToAnimate", "Select Bones Or Curves To Animate");
+		const FText TitleText = LOCTEXT("SelectBonesOrCurvesToAnimate", "Select Bones Or Curves To Animate");
 
 		// Create the window to choose our options
 		TSharedRef<SWindow> Window = SNew(SWindow)
@@ -3384,7 +3384,7 @@ bool FControlRigParameterTrackEditor::CollapseAllLayers(TSharedPtr<ISequencer>&
 			UE_LOG(LogControlRigEditor, Log, TEXT("Section wrong type or not first when collapsing layers"));
 			return false;
 		}
-		FScopedTransaction Transaction(LOCTEXT("CollapseAllSections", "CollapseAllSections"));
+		FScopedTransaction Transaction(LOCTEXT("CollapseAllSections", "Collapse All Sections"));
 		ParameterSection->Modify();
 		UControlRig* ControlRig = ParameterSection->GetControlRig();
 		TRange<FFrameNumber> Range = SequencerPtr->GetFocusedMovieSceneSequence()->GetMovieScene()->GetPlaybackRange();
diff --git a/Engine/Plugins/MovieScene/MatineeToLevelSequence/Source/MatineeToLevelSequence/Private/MatineeCameraShakeToNewCameraShakeConverter.cpp b/Engine/Plugins/MovieScene/MatineeToLevelSequence/Source/MatineeToLevelSequence/Private/MatineeCameraShakeToNewCameraShakeConverter.cpp
--- a/Engine/Plugins/MovieScene/MatineeToLevelSequence/Source/MatineeToLevelSequence/Private/MatineeCameraShakeToNewCameraShakeConverter.cpp
+++ b/Engine/Plugins/MovieScene/MatineeToLevelSequence/Source/MatineeToLevelSequence/Private/MatineeCameraShakeToNewCameraShakeConverter.cpp
@@ -118,7 +118,7 @@ public:
 			.Padding(10.f)
 			[
 				SNew(SCheckBox)
-				.ToolTipText(LOCTEXT("FilterAnimCameraShakes", "Only show camera shakes with a Matinee camera animation"))
+				.ToolTipText(LOCTEXT("FilterAnimCameraShakesTooltip", "Only show camera shakes with a Matinee camera animation"))
 				.OnCheckStateChanged_Lambda([this](ECheckBoxState State)
 						{
 							bFilterNonAnimCameraShakes = !bFilterNonAnimCameraShakes; 
