Sequencer: Control Rig: When dragging in a Control Rig use the opened Level Sequence, don't search for an actor this can cause a crash if one doesn't exist. Also with this JIRA due to the order of operations we may search for a widget on the SequencerAnimToolsEdit mode which will override the Control Rig Edit Mode and not put the Widget in the correct location.

diff --git a/Engine/Plugins/Animation/ControlRig/Source/ControlRigEditor/Private/ControlRigBlueprintActions.cpp b/Engine/Plugins/Animation/ControlRig/Source/ControlRigEditor/Private/ControlRigBlueprintActions.cpp
--- a/Engine/Plugins/Animation/ControlRig/Source/ControlRigEditor/Private/ControlRigBlueprintActions.cpp
+++ b/Engine/Plugins/Animation/ControlRig/Source/ControlRigEditor/Private/ControlRigBlueprintActions.cpp
@@ -31,6 +31,7 @@
 #include "MovieSceneToolsProjectSettings.h"
 #include "SBlueprintDiff.h"
 #include "Misc/MessageDialog.h"
+#include "LevelSequenceEditorBlueprintLibrary.h"
 
 #define LOCTEXT_NAMESPACE "ControlRigBlueprintActions"
 
@@ -289,19 +290,12 @@ void FControlRigBlueprintActions::OnSpawnedSkeletalMeshActorChanged(UObject* InO
 	}
 	UClass* ControlRigClass = RigBlueprint->GeneratedClass;
 
-	// find a level sequence in the world
-	ALevelSequenceActor* LevelSequenceActor = nullptr;
-	for (ALevelSequenceActor* ExistingLevelSequenceActor : TActorRange<ALevelSequenceActor>(MeshActor->GetWorld()))
+	// find a level sequence in the world, if can't find that, create one
+	ULevelSequence* Sequence = ULevelSequenceEditorBlueprintLibrary::GetCurrentLevelSequence();
+	if (Sequence == nullptr)
 	{
-		LevelSequenceActor = ExistingLevelSequenceActor;
-		break;
-	}
-    //The creation of the FSequencer that happens below is not transactional so we need to make sure we doen't transaction there.
-	GEditor->CancelTransaction(0);
+		ALevelSequenceActor* LevelSequenceActor = nullptr;
 
-	// todo: show a dialog for picking an existing one or creating a new one
-	if (LevelSequenceActor == nullptr)
-	{
 		FString SequenceName = FString::Printf(TEXT("%s_Take1"), *InAsset->GetName());
 		FString PackagePath = TEXT("/Game");
 
@@ -311,7 +305,7 @@ void FControlRigBlueprintActions::OnSpawnedSkeletalMeshActorChanged(UObject* InO
 		AssetToolsModule.Get().CreateUniqueAssetName(PackagePath / SequenceName, TEXT(""), UniquePackageName, UniqueAssetName);
 
 		UPackage* Package = CreatePackage(*UniquePackageName);
-		ULevelSequence* Sequence = NewObject<ULevelSequence>(Package, *UniqueAssetName, RF_Public | RF_Standalone);
+		Sequence = NewObject<ULevelSequence>(Package, *UniqueAssetName, RF_Public | RF_Standalone);
 		Sequence->Initialize(); //creates movie scene
 		Sequence->MarkPackageDirty();
 
@@ -335,8 +329,7 @@ void FControlRigBlueprintActions::OnSpawnedSkeletalMeshActorChanged(UObject* InO
 		LevelSequenceActor = CastChecked<ALevelSequenceActor>(NewActor);
 		LevelSequenceActor->SetSequence(Sequence);
 	}
-
-	ULevelSequence* Sequence = LevelSequenceActor->GetSequence();
+ 
 	if (Sequence == nullptr)
 	{
 		return;
diff --git a/Engine/Plugins/Editor/SequencerAnimTools/Source/SequencerAnimTools/Public/SequencerToolsEditMode.h b/Engine/Plugins/Editor/SequencerAnimTools/Source/SequencerAnimTools/Public/SequencerToolsEditMode.h
--- a/Engine/Plugins/Editor/SequencerAnimTools/Source/SequencerAnimTools/Public/SequencerToolsEditMode.h
+++ b/Engine/Plugins/Editor/SequencerAnimTools/Source/SequencerAnimTools/Public/SequencerToolsEditMode.h
@@ -27,7 +27,7 @@ public:
 	virtual bool UsesToolkits() const  override { return false; }
 	virtual bool InputKey(FEditorViewportClient* InViewportClient, FViewport* InViewport, FKey InKey, EInputEvent InEvent) override;
 	virtual bool UsesPropertyWidgets() const override { return true; }
-
+	virtual bool UsesTransformWidget() const { return false; }
 	virtual bool IsCompatibleWith(FEditorModeID OtherModeID) const override;
 
 };
