Control Rig: Control Rig Control Actor had a hard reference to the ControlRig that's owned by Sequencer, though transient will cause an ensure if the LevelSequence is deleted, fix is to just change it to a soft reference.

diff --git a/Engine/Plugins/Animation/ControlRig/Source/ControlRig/Public/ControlRigControlActor.h b/Engine/Plugins/Animation/ControlRig/Source/ControlRig/Public/ControlRigControlActor.h
--- a/Engine/Plugins/Animation/ControlRig/Source/ControlRig/Public/ControlRigControlActor.h
+++ b/Engine/Plugins/Animation/ControlRig/Source/ControlRig/Public/ControlRigControlActor.h
@@ -68,7 +68,7 @@ private:
 	TObjectPtr<class USceneComponent> ActorRootComponent;
 
 	UPROPERTY(transient)
-	TObjectPtr<UControlRig>  ControlRig;
+	TSoftObjectPtr<UControlRig>  ControlRig;
 
 	UPROPERTY(transient)
 	TArray<FName> ControlNames;
