Control Rig: LOC duplicates

diff --git a/Engine/Plugins/Animation/ControlRig/Source/ControlRigEditor/Private/ControlRigElementDetails.cpp b/Engine/Plugins/Animation/ControlRig/Source/ControlRigEditor/Private/ControlRigElementDetails.cpp
--- a/Engine/Plugins/Animation/ControlRig/Source/ControlRigEditor/Private/ControlRigElementDetails.cpp
+++ b/Engine/Plugins/Animation/ControlRig/Source/ControlRigEditor/Private/ControlRigElementDetails.cpp
@@ -2798,7 +2798,7 @@ void FRigControlElementDetails::CustomizeShape(IDetailLayoutBuilder& DetailBuild
 		SNew(STextBlock)
 		.Font(IDetailLayoutBuilder::GetDetailFont())
 		.Text(LOCTEXT("ShapeProperties", "Shape Properties"))
-		.ToolTipText(LOCTEXT("ShapeProperties", "Customize the properties of the shape"))
+		.ToolTipText(LOCTEXT("ShapePropertiesTooltip", "Customize the properties of the shape"))
 	]
 	.CopyAction(FUIAction(
 		FExecuteAction::CreateSP(this, &FRigControlElementDetails::OnCopyShapeProperties)))
diff --git a/Engine/Plugins/Animation/ControlRig/Source/ControlRigEditor/Private/Editor/ControlRigEditor.cpp b/Engine/Plugins/Animation/ControlRig/Source/ControlRigEditor/Private/Editor/ControlRigEditor.cpp
--- a/Engine/Plugins/Animation/ControlRig/Source/ControlRigEditor/Private/Editor/ControlRigEditor.cpp
+++ b/Engine/Plugins/Animation/ControlRig/Source/ControlRigEditor/Private/Editor/ControlRigEditor.cpp
@@ -3285,11 +3285,11 @@ void FControlRigEditor::HandleViewportCreated(const TSharedRef<class IPersonaVie
 			{
 				FFormatNamedArguments Args;
 				Args.Add(TEXT("HotKey"), HotKeyText);
-				return FText::Format(LOCTEXT("ControlRigBPViewportShapeTransformEditNotification", "Currently Manipulating Shape Transform - Press {HotKey} to Exit"), Args);
+				return FText::Format(LOCTEXT("ControlRigBPViewportShapeTransformEditNotificationPress", "Currently Manipulating Shape Transform - Press {HotKey} to Exit"), Args);
 			}
 		}
 		
-		return LOCTEXT("ControlRigBPViewportShapeTransformEditNotification", "Currently Manipulating Shape Transform - Assign a Hotkey and Use It to Exit");
+		return LOCTEXT("ControlRigBPViewportShapeTransformEditNotificationAssign", "Currently Manipulating Shape Transform - Assign a Hotkey and Use It to Exit");
 	};
 
 	auto GetChangingShapeTransformTextVisibility = [this]()
diff --git a/Engine/Plugins/Animation/ControlRig/Source/ControlRigEditor/Private/Editor/SRigHierarchy.cpp b/Engine/Plugins/Animation/ControlRig/Source/ControlRigEditor/Private/Editor/SRigHierarchy.cpp
--- a/Engine/Plugins/Animation/ControlRig/Source/ControlRigEditor/Private/Editor/SRigHierarchy.cpp
+++ b/Engine/Plugins/Animation/ControlRig/Source/ControlRigEditor/Private/Editor/SRigHierarchy.cpp
@@ -1813,7 +1813,7 @@ void SRigHierarchy::HandlePasteItems()
 		FString Content;
 		FPlatformApplicationMisc::ClipboardPaste(Content);
 
-		FScopedTransaction Transaction(LOCTEXT("HierarchyTreePaste", "Pasted rig elements."));
+		FScopedTransaction Transaction(LOCTEXT("HierarchyTreePastedRigElements", "Pasted rig elements."));
 
 		URigHierarchyController* Controller = Hierarchy->GetController(true);
 		check(Controller);
@@ -1866,7 +1866,7 @@ void SRigHierarchy::HandlePasteTransforms(ERigTransformType::Type InTransformTyp
 		FString Content;
 		FPlatformApplicationMisc::ClipboardPaste(Content);
 
-		FScopedTransaction Transaction(LOCTEXT("HierarchyTreePaste", "Pasted transforms."));
+		FScopedTransaction Transaction(LOCTEXT("HierarchyTreePastedTransform", "Pasted transforms."));
 
 		SRigHierarchyPasteTransformsErrorPipe ErrorPipe;
 		FRigHierarchyCopyPasteContent Data;
diff --git a/Engine/Plugins/Animation/ControlRig/Source/ControlRigEditor/Private/Graph/SControlRigGraphPinNameList.cpp b/Engine/Plugins/Animation/ControlRig/Source/ControlRigEditor/Private/Graph/SControlRigGraphPinNameList.cpp
--- a/Engine/Plugins/Animation/ControlRig/Source/ControlRigEditor/Private/Graph/SControlRigGraphPinNameList.cpp
+++ b/Engine/Plugins/Animation/ControlRig/Source/ControlRigEditor/Private/Graph/SControlRigGraphPinNameList.cpp
@@ -99,7 +99,7 @@ TSharedRef<SWidget>	SControlRigGraphPinNameList::GetDefaultValueWidget()
 				.ButtonColorAndOpacity( this, &SControlRigGraphPinNameList::OnGetWidgetBackground )
 				.OnClicked(this, &SControlRigGraphPinNameList::OnBrowseClicked)
 				.ContentPadding(0)
-				.ToolTipText(NSLOCTEXT("GraphEditor", "ObjectGraphPin_Browse_Tooltip", "Browse"))
+				.ToolTipText(NSLOCTEXT("SControlRigGraphPinNameList", "ObjectGraphPin_Browse_Tooltip", "Browse"))
 				[
 					SNew(SImage)
 					.ColorAndOpacity( this, &SControlRigGraphPinNameList::OnGetWidgetForeground )
