Control Rig: Add sub objects of a CDO of a native uclass to root (such as fk control rig, additive control rig). Otherwise, GC will get rid of them right after they are created

In general all objects created before a certain point (constructor) are considered part of the CDO not considered for GC and we simply assume they are good and will never introduce new references. But in this case, sub objects are added to the CDO after constructor and as a result need special treatment such that GC don't get rid of them.

For Native Class, it is fine to keep its CDO always alive, so it is also fine to keep the subobjects always alive by adding them to root as well

diff --git a/Engine/Plugins/Animation/ControlRig/Source/ControlRig/Private/ControlRig.cpp b/Engine/Plugins/Animation/ControlRig/Source/ControlRig/Private/ControlRig.cpp
--- a/Engine/Plugins/Animation/ControlRig/Source/ControlRig/Private/ControlRig.cpp
+++ b/Engine/Plugins/Animation/ControlRig/Source/ControlRig/Private/ControlRig.cpp
@@ -2827,7 +2827,13 @@ void UControlRig::PostInitInstance(UControlRig* InCDO)
 	{
 		// for default objects we need to check if the CDO is rooted. specialized Control Rigs
 		// such as the FK control rig may not have a root since they are part of a C++ package.
-		if(!IsRooted() && GetClass()->IsNative())
+
+		// since the sub objects are created after the constructor
+		// GC won't consider them part of the CDO, even if they have the sub object flags
+		// so even if CDO is rooted and references these sub objects, 
+		// it is not enough to keep them alive.
+		// Hence, we have to add them to root here.
+		if(GetClass()->IsNative())
 		{
 			VM->AddToRoot();
 			DynamicHierarchy->AddToRoot();
