Control Rig: Support modifybone to work without a preview mesh

diff --git a/Engine/Plugins/Animation/ControlRig/Source/ControlRigEditor/Private/Editor/ControlRigEditor.cpp b/Engine/Plugins/Animation/ControlRig/Source/ControlRigEditor/Private/Editor/ControlRigEditor.cpp
--- a/Engine/Plugins/Animation/ControlRig/Source/ControlRigEditor/Private/Editor/ControlRigEditor.cpp
+++ b/Engine/Plugins/Animation/ControlRig/Source/ControlRigEditor/Private/Editor/ControlRigEditor.cpp
@@ -3023,6 +3023,19 @@ void FControlRigEditor::Tick(float DeltaTime)
 				// reset transforms here to prevent additive transforms from accumulating to INF
 				ControlRig->GetHierarchy()->ResetPoseToInitial(ERigElementType::Bone);
 			}
+
+			if (PreviewInstance)
+			{
+				// since we don't have a preview mesh the anim instance cannot deal with the modify bone
+				// functionality. we need to perform this manually to ensure the pose is kept.
+				const TArray<FAnimNode_ModifyBone>& BoneControllers = PreviewInstance->GetBoneControllers();
+				for(const FAnimNode_ModifyBone& ModifyBone : BoneControllers)
+				{
+					const FRigElementKey BoneKey(ModifyBone.BoneToModify.BoneName, ERigElementType::Bone);
+					const FTransform BoneTransform(ModifyBone.Rotation, ModifyBone.Translation, ModifyBone.Scale);
+					ControlRig->GetHierarchy()->SetLocalTransform(BoneKey, BoneTransform);
+				}
+			}
 			
 			ControlRig->SetDeltaTime(DeltaTime);
 			ControlRig->Evaluate_AnyThread();
diff --git a/Engine/Source/Editor/AnimGraph/Private/AnimPreviewInstance.cpp b/Engine/Source/Editor/AnimGraph/Private/AnimPreviewInstance.cpp
--- a/Engine/Source/Editor/AnimGraph/Private/AnimPreviewInstance.cpp
+++ b/Engine/Source/Editor/AnimGraph/Private/AnimPreviewInstance.cpp
@@ -523,6 +523,11 @@ void UAnimPreviewInstance::ResetModifiedBone(bool bCurveController/*=false*/)
 	GetProxyOnGameThread<FAnimPreviewInstanceProxy>().ResetModifiedBone(bCurveController);
 }
 
+const TArray<FAnimNode_ModifyBone>& UAnimPreviewInstance::GetBoneControllers()
+{
+	return GetProxyOnGameThread<FAnimPreviewInstanceProxy>().GetBoneControllers();
+}
+
 #if WITH_EDITOR	
 
 void UAnimPreviewInstance::SetKey()
diff --git a/Engine/Source/Editor/AnimGraph/Public/AnimPreviewInstance.h b/Engine/Source/Editor/AnimGraph/Public/AnimPreviewInstance.h
--- a/Engine/Source/Editor/AnimGraph/Public/AnimPreviewInstance.h
+++ b/Engine/Source/Editor/AnimGraph/Public/AnimPreviewInstance.h
@@ -258,6 +258,11 @@ public:
 	 */
 	void ResetModifiedBone(bool bCurveController=false);
 
+	/**
+	 * Returns all currently active bone controllers on this instance's proxy
+	 */
+	const TArray<FAnimNode_ModifyBone>& GetBoneControllers();
+
 	void SetForceRetargetBasePose(bool ForceRetargetBasePose);
 
 	bool GetForceRetargetBasePose() const;
