Sequencer: Allow adding an animation track to the top level actor binding if the root is a skeletal mesh component. Previously, restrictions were put in so that animation tracks would only be allowed on specific skeletal mesh component bindings if the actor had multiple skeletal mesh components. This loosens that restriction to make it more convenient to add animation to the root.

follow up to CL 15750937

diff --git a/Engine/Plugins/Animation/ControlRig/Source/ControlRigEditor/Private/Sequencer/ControlRigParameterTrackEditor.cpp b/Engine/Plugins/Animation/ControlRig/Source/ControlRigEditor/Private/Sequencer/ControlRigParameterTrackEditor.cpp
--- a/Engine/Plugins/Animation/ControlRig/Source/ControlRigEditor/Private/Sequencer/ControlRigParameterTrackEditor.cpp
+++ b/Engine/Plugins/Animation/ControlRig/Source/ControlRigEditor/Private/Sequencer/ControlRigParameterTrackEditor.cpp
@@ -89,6 +89,11 @@ static USkeletalMeshComponent* AcquireSkeletalMeshFromObject(UObject* BoundObjec
 {
 	if (AActor* Actor = Cast<AActor>(BoundObject))
 	{
+		if (USkeletalMeshComponent* SkeletalMeshComponent = Cast<USkeletalMeshComponent>(Actor->GetRootComponent()))
+		{
+			return SkeletalMeshComponent;
+		}
+
 		TArray<USkeletalMeshComponent*> SkeletalMeshComponents;
 		Actor->GetComponents(SkeletalMeshComponents);
 
@@ -127,6 +132,11 @@ static USkeleton* AcquireSkeletonFromObjectGuid(const FGuid& Guid, UObject** Obj
 	*Object = BoundObject;
 	if (AActor* Actor = Cast<AActor>(BoundObject))
 	{
+		if (USkeletalMeshComponent* SkeletalMeshComponent = Cast<USkeletalMeshComponent>(Actor->GetRootComponent()))
+		{
+			return GetSkeletonFromComponent(SkeletalMeshComponent);
+		}
+
 		TArray<USkeletalMeshComponent*> SkeletalMeshComponents;
 		Actor->GetComponents(SkeletalMeshComponents);
 		if (SkeletalMeshComponents.Num() == 1)
@@ -138,6 +148,11 @@ static USkeleton* AcquireSkeletonFromObjectGuid(const FGuid& Guid, UObject** Obj
 		AActor* ActorCDO = Cast<AActor>(Actor->GetClass()->GetDefaultObject());
 		if (ActorCDO)
 		{
+			if (USkeletalMeshComponent* SkeletalMeshComponent = Cast<USkeletalMeshComponent>(ActorCDO->GetRootComponent()))
+			{
+				return GetSkeletonFromComponent(SkeletalMeshComponent);
+			}
+
 			ActorCDO->GetComponents(SkeletalMeshComponents);
 			if (SkeletalMeshComponents.Num() == 1)
 			{
diff --git a/Engine/Source/Editor/MovieSceneTools/Private/TrackEditors/SkeletalAnimationTrackEditor.cpp b/Engine/Source/Editor/MovieSceneTools/Private/TrackEditors/SkeletalAnimationTrackEditor.cpp
--- a/Engine/Source/Editor/MovieSceneTools/Private/TrackEditors/SkeletalAnimationTrackEditor.cpp
+++ b/Engine/Source/Editor/MovieSceneTools/Private/TrackEditors/SkeletalAnimationTrackEditor.cpp
@@ -300,6 +300,11 @@ USkeletalMeshComponent* AcquireSkeletalMeshFromObjectGuid(const FGuid& Guid, TSh
 
 	if (AActor* Actor = Cast<AActor>(BoundObject))
 	{
+		if (USkeletalMeshComponent* SkeletalMeshComponent = Cast<USkeletalMeshComponent>(Actor->GetRootComponent()))
+		{
+			return SkeletalMeshComponent;
+		}
+
 		TArray<USkeletalMeshComponent*> SkeletalMeshComponents;
 		Actor->GetComponents(SkeletalMeshComponents);
 		
@@ -349,6 +354,12 @@ TArray<USkeletalMeshComponent*> AcquireSkeletalMeshComponentsFromObjectGuid(cons
 
 	if (Actor)
 	{
+		if (USkeletalMeshComponent* SkeletalMeshComponent = Cast<USkeletalMeshComponent>(Actor->GetRootComponent()))
+		{
+			SkeletalMeshComponents.Add(SkeletalMeshComponent);
+			return SkeletalMeshComponents;
+		}
+
 		Actor->GetComponents(SkeletalMeshComponents);
 		if (SkeletalMeshComponents.Num())
 		{
@@ -358,6 +369,12 @@ TArray<USkeletalMeshComponent*> AcquireSkeletalMeshComponentsFromObjectGuid(cons
 		AActor* ActorCDO = Cast<AActor>(Actor->GetClass()->GetDefaultObject());
 		if (ActorCDO)
 		{
+			if (USkeletalMeshComponent* SkeletalMeshComponent = Cast<USkeletalMeshComponent>(ActorCDO->GetRootComponent()))
+			{
+				SkeletalMeshComponents.Add(SkeletalMeshComponent);
+				return SkeletalMeshComponents;
+			}
+
 			ActorCDO->GetComponents(SkeletalMeshComponents);
 			if (SkeletalMeshComponents.Num())
 			{
