Control Rig: Clear snapshots of initialized memory

diff --git a/Engine/Plugins/Animation/ControlRig/Source/ControlRig/Private/ControlRig.cpp b/Engine/Plugins/Animation/ControlRig/Source/ControlRig/Private/ControlRig.cpp
--- a/Engine/Plugins/Animation/ControlRig/Source/ControlRig/Private/ControlRig.cpp
+++ b/Engine/Plugins/Animation/ControlRig/Source/ControlRig/Private/ControlRig.cpp
@@ -122,6 +122,12 @@ void UControlRig::BeginDestroy()
 	}
 #endif
 
+	// on destruction clear out the initialized snapshots
+	if (HasAnyFlags(RF_ClassDefaultObject))
+	{
+		InitializedVMSnapshots.Reset();
+	}
+
 	TRACE_OBJECT_LIFETIME_END(this);
 }
 
@@ -2826,6 +2832,9 @@ void UControlRig::PostInitInstance(UControlRig* InCDO)
 			VM->AddToRoot();
 			DynamicHierarchy->AddToRoot();
 		}
+
+		// Clear the initialized VM snapshots
+		InitializedVMSnapshots.Reset();
 	}
 }
 
