Control Rig: Fix anim blueprint control rig node not updating when changing real variable

diff --git a/Engine/Plugins/Animation/ControlRig/Source/ControlRig/Private/AnimNode_ControlRig.cpp b/Engine/Plugins/Animation/ControlRig/Source/ControlRig/Private/AnimNode_ControlRig.cpp
--- a/Engine/Plugins/Animation/ControlRig/Source/ControlRig/Private/AnimNode_ControlRig.cpp
+++ b/Engine/Plugins/Animation/ControlRig/Source/ControlRig/Private/AnimNode_ControlRig.cpp
@@ -556,6 +556,11 @@ void FAnimNode_ControlRig::PropagateInputProperties(const UObject* InSourceInsta
 					const float Value = *(const float*)SrcPtr;
 					Variable.SetValue<float>(Value);
 				}
+				else if (CastField<FDoubleProperty>(CallerProperty) != nullptr && Variable.TypeName == TEXT("double"))
+				{
+					const double Value = *(const double*)SrcPtr;
+					Variable.SetValue<double>(Value);
+				}
 				else if (CastField<FIntProperty>(CallerProperty) != nullptr && Variable.TypeName == TEXT("int32"))
 				{
 					const int32 Value = *(const int32*)SrcPtr;
@@ -578,6 +583,10 @@ void FAnimNode_ControlRig::PropagateInputProperties(const UObject* InSourceInsta
 						StructProperty->Struct->CopyScriptStruct(Variable.Memory, SrcPtr, 1);
 					}
 				}
+				else
+				{
+					ensureMsgf(false, TEXT("Property %s type %s not recognized"), *CallerProperty->GetName(), *CallerProperty->GetCPPType());
+				}
 			}
 		}
 	}
