Sequencer: Control Rig:  Don't force scales to be zero when we create channels, which can happen when a space channel is added. We already set the default to zero when new sections are created, so code not needed. Also don't allow creation of space channels on parented controls(floats/ints/bools/enums).

diff --git a/Engine/Plugins/Animation/ControlRig/Source/ControlRig/Private/Sequencer/MovieSceneControlRigParameterSection.cpp b/Engine/Plugins/Animation/ControlRig/Source/ControlRig/Private/Sequencer/MovieSceneControlRigParameterSection.cpp
--- a/Engine/Plugins/Animation/ControlRig/Source/ControlRig/Private/Sequencer/MovieSceneControlRigParameterSection.cpp
+++ b/Engine/Plugins/Animation/ControlRig/Source/ControlRig/Private/Sequencer/MovieSceneControlRigParameterSection.cpp
@@ -1241,6 +1241,19 @@ void UMovieSceneControlRigParameterSection::AddTransformParameter(FName InParame
 	}
 }
 
+// only allow creation of space channels onto non-parented Controls
+bool UMovieSceneControlRigParameterSection::CanCreateSpaceChannel(FName InControlName) const
+{
+	if (const FChannelMapInfo* ChannelInfo = ControlChannelMap.Find(InControlName))
+	{
+		if (ChannelInfo->ParentControlIndex == INDEX_NONE)
+		{
+			return true;
+		}
+	}
+	return false;
+}
+
 void UMovieSceneControlRigParameterSection::AddSpaceChannel(FName InControlName, bool bReconstructChannel)
 {
 	//only add it if it's the first section since we can't blend them
@@ -1249,7 +1262,8 @@ void UMovieSceneControlRigParameterSection::AddSpaceChannel(FName InControlName,
 		const TArray<UMovieSceneSection*>& Sections = Track->GetAllSections();
 		if (Sections[0] == this)
 		{
-			if (!HasSpaceChannel(InControlName))
+			
+			if (CanCreateSpaceChannel(InControlName) && !HasSpaceChannel(InControlName))
 			{
 				SpaceChannels.Add(FSpaceControlNameAndChannel(InControlName));
 				if (OnSpaceChannelAdded.IsBound())
@@ -1275,7 +1289,6 @@ const TArray< FSpaceControlNameAndChannel>& UMovieSceneControlRigParameterSectio
 	return SpaceChannels;
 }
 
-
 void UMovieSceneControlRigParameterSection::ReconstructChannelProxy()
 {
 	FMovieSceneChannelProxyData Channels;
@@ -1606,25 +1619,6 @@ void UMovieSceneControlRigParameterSection::ReconstructChannelProxy()
 					{
 						if (ControlElement->GetName() == Vector.ParameterName)
 						{
-							if (ControlElement->Settings.ControlType == ERigControlType::Scale)
-							{
-								//only set defaults when not loading or undo
-								if (HasAnyFlags(RF_NeedLoad | RF_NeedPostLoad | RF_NeedPostLoadSubobjects) == false && bIsInUndo == false)
-								{
-									if (BlendType == EMovieSceneBlendType::Additive)
-									{
-										Vector.XCurve.SetDefault(0.0f);
-										Vector.YCurve.SetDefault(0.0f);
-										Vector.ZCurve.SetDefault(0.0f);
-									}
-									else
-									{
-										Vector.XCurve.SetDefault(1.0f);
-										Vector.YCurve.SetDefault(1.0f);
-										Vector.ZCurve.SetDefault(1.0f);
-									}
-								}
-							}
 							ControlChannelMap.Add(Vector.ParameterName, FChannelMapInfo(ControlIndex, TotalIndex, FloatChannelIndex, INDEX_NONE, NAME_None, MaskIndex, CategoryIndex));
 							ControlIndex++;
 							if (bEnabled)
@@ -1721,22 +1715,7 @@ void UMovieSceneControlRigParameterSection::ReconstructChannelProxy()
 							if (ControlElement->Settings.ControlType == ERigControlType::Transform ||
 								ControlElement->Settings.ControlType == ERigControlType::EulerTransform)
 							{
-								//only set defaults when not loading or undo
-								if (HasAnyFlags(RF_NeedLoad | RF_NeedPostLoad | RF_NeedPostLoadSubobjects) == false && bIsInUndo == false)
-								{
-									if (BlendType == EMovieSceneBlendType::Additive)
-									{
-										Transform.Scale[0].SetDefault(0.0f);
-										Transform.Scale[1].SetDefault(0.0f);
-										Transform.Scale[2].SetDefault(0.0f);
-									}
-									else
-									{
-										Transform.Scale[0].SetDefault(1.0f);
-										Transform.Scale[1].SetDefault(1.0f);
-										Transform.Scale[2].SetDefault(1.0f);
-									}
-								}
+								
 								Channels.Add(Transform.Scale[0], EditorData.MetaData[6], EditorData.ExternalValues[6]);
 								Channels.Add(Transform.Scale[1], EditorData.MetaData[7], EditorData.ExternalValues[7]);
 								Channels.Add(Transform.Scale[2], EditorData.MetaData[8], EditorData.ExternalValues[8]);
diff --git a/Engine/Plugins/Animation/ControlRig/Source/ControlRig/Public/Sequencer/MovieSceneControlRigParameterSection.h b/Engine/Plugins/Animation/ControlRig/Source/ControlRig/Public/Sequencer/MovieSceneControlRigParameterSection.h
--- a/Engine/Plugins/Animation/ControlRig/Source/ControlRig/Public/Sequencer/MovieSceneControlRigParameterSection.h
+++ b/Engine/Plugins/Animation/ControlRig/Source/ControlRig/Public/Sequencer/MovieSceneControlRigParameterSection.h
@@ -428,6 +428,9 @@ public:
 	/** Remove All Keys, but maybe not space keys if bIncludeSpaceKeys is false */
 	void RemoveAllKeys(bool bIncludeSpaceKeys);
 
+	/** Whether or not create a space channel for a particular control */
+	bool CanCreateSpaceChannel(FName InControlName) const;
+
 public:
 	/**
 	* Access the interrogation key for control rig data 
