Control Rig: Fix ensure when undo a functions local variable deletion

diff --git a/Engine/Plugins/Animation/ControlRig/Source/ControlRigDeveloper/Private/Graph/ControlRigGraph.cpp b/Engine/Plugins/Animation/ControlRig/Source/ControlRigDeveloper/Private/Graph/ControlRigGraph.cpp
--- a/Engine/Plugins/Animation/ControlRig/Source/ControlRigDeveloper/Private/Graph/ControlRigGraph.cpp
+++ b/Engine/Plugins/Animation/ControlRig/Source/ControlRigDeveloper/Private/Graph/ControlRigGraph.cpp
@@ -412,8 +412,14 @@ void UControlRigGraph::HandleModifiedEvent(ERigVMGraphNotifType InNotifType, URi
 						while (ExistingObject);
 						EdNode->Rename(*DeletedName, NewOuter, REN_ForceNoResetLoaders | REN_DoNotDirty | REN_DontCreateRedirectors | REN_NonTransactional);
 					}
-					RemoveNode(EdNode, true);
-					NotifyGraphChanged();
+
+					// Remove the node and notify of the removal
+					{
+						Nodes.Remove(EdNode);
+						EdNode->BreakAllNodeLinks();
+						
+						NotifyGraphChanged();
+					}
 				}
 			}
 			break;
@@ -849,22 +855,28 @@ UEdGraphNode* UControlRigGraph::FindNodeForModelNodeName(const FName& InModelNod
 		{
 			if (RigNode->ModelNodePath == InModelNodePath)
 			{
-				if (bCacheIfRequired)
+				if (RigNode->GetOuter() == this)
 				{
-					ModelNodePathToEdNode.Add(InModelNodeName, EdNode);
+					if (bCacheIfRequired)
+					{
+						ModelNodePathToEdNode.Add(InModelNodeName, EdNode);
+					}
+					return EdNode;
 				}
-				return EdNode;
 			}
 		}
-		else
+		else if (EdNode)
 		{
 			if (EdNode->GetFName() == InModelNodeName)
 			{
-				if (bCacheIfRequired)
+				if (EdNode->GetOuter() == this)
 				{
-					ModelNodePathToEdNode.Add(InModelNodeName, EdNode);
+					if (bCacheIfRequired)
+					{
+						ModelNodePathToEdNode.Add(InModelNodeName, EdNode);
+					}
+					return EdNode;
 				}
-				return EdNode;
 			}
 		}
 	}
