Control Rig: refactor details panel to show any node, support vector, rotator and transform widgets

diff --git a/Engine/Plugins/Animation/ControlRig/Source/ControlRigEditor/Private/ControlRigEditorModule.cpp b/Engine/Plugins/Animation/ControlRig/Source/ControlRigEditor/Private/ControlRigEditorModule.cpp
--- a/Engine/Plugins/Animation/ControlRig/Source/ControlRigEditor/Private/ControlRigEditorModule.cpp
+++ b/Engine/Plugins/Animation/ControlRig/Source/ControlRigEditor/Private/ControlRigEditorModule.cpp
@@ -1362,7 +1362,7 @@ void FControlRigEditorModule::GetContextMenuActions(const UControlRigGraphSchema
 								TSharedRef<SControlRigVariableBinding> VariableBindingWidget =
 									SNew(SControlRigVariableBinding)
 									.Blueprint(RigBlueprint)
-									.ModelPin(ModelPin)
+									.ModelPins({ModelPin})
 									.CanRemoveBinding(false);
 
 								VariablesSection.AddEntry(FToolMenuEntry::InitWidget("BindPinToVariableWidget", VariableBindingWidget, FText(), true));
diff --git a/Engine/Plugins/Animation/ControlRig/Source/ControlRigEditor/Private/ControlRigElementDetails.cpp b/Engine/Plugins/Animation/ControlRig/Source/ControlRigEditor/Private/ControlRigElementDetails.cpp
--- a/Engine/Plugins/Animation/ControlRig/Source/ControlRigEditor/Private/ControlRigElementDetails.cpp
+++ b/Engine/Plugins/Animation/ControlRig/Source/ControlRigEditor/Private/ControlRigElementDetails.cpp
@@ -77,11 +77,14 @@ void RigElementDetails_GetCustomizedInfo(TSharedRef<IPropertyHandle> InStructPro
 	{
 		if (Object->IsA<UControlRigBlueprint>())
 		{
-			OutBlueprint = Cast<UControlRigBlueprint>(Object);
-			if (OutBlueprint)
-			{
-				break;
-			}
+			OutBlueprint = CastChecked<UControlRigBlueprint>(Object);
+			break;
+		}
+
+		OutBlueprint = Object->GetTypedOuter<UControlRigBlueprint>();
+		if(OutBlueprint)
+		{
+			break;
 		}
 	}
 
@@ -430,245 +433,6 @@ FReply FRigElementKeyDetails::OnSelectInHierarchyClicked()
 	return FReply::Handled();
 }
 
-void FRigUnitDetails::CustomizeDetails(IDetailLayoutBuilder& DetailBuilder)
-{
-	TArray<TSharedPtr<FStructOnScope>> StructsBeingCustomized;
-	DetailBuilder.GetStructsBeingCustomized(StructsBeingCustomized);
-	if (StructsBeingCustomized.Num() == 0)
-	{
-		return;
-	}
-
-	TSharedPtr<FStructOnScope> StructBeingCustomized = StructsBeingCustomized[0];
-
-	BlueprintBeingCustomized = nullptr;
-	if (UPackage* Package = StructBeingCustomized->GetPackage())
-	{
-		TArray<UObject*> SubObjects;
-		Package->GetDefaultSubobjects(SubObjects);
-
-		for (UObject* SubObject : SubObjects)
-		{
-			if (UControlRig* Rig = Cast<UControlRig>(SubObject))
-			{
-				BlueprintBeingCustomized = Cast<UControlRigBlueprint>(Rig->GetClass()->ClassGeneratedBy);
-				if (BlueprintBeingCustomized)
-				{
-					break;
-				}
-			}
-		}
-	}
-
-	if (BlueprintBeingCustomized == nullptr)
-	{
-		return;
-	}
-
-	GraphBeingCustomized = nullptr;
-	for (UEdGraph* Graph : BlueprintBeingCustomized->UbergraphPages)
-	{
-		GraphBeingCustomized = Cast<UControlRigGraph>(Graph);
-		if (GraphBeingCustomized)
-		{
-			break;
-		}
-	}
-
-	if (GraphBeingCustomized == nullptr)
-	{
-		return;
-	}
-
-	URigVMGraph* Model = GraphBeingCustomized->GetModel();
-	if(Model == nullptr)
-	{
-		return;
-	}
-
-	const TArray<FName> SelectedNodeNames = Model->GetSelectNodes();
-	if(SelectedNodeNames.Num() == 0)
-	{
-		return;
-	}
-
-	URigVMNode* ModelNode = Model->FindNodeByName(SelectedNodeNames[0]);
-	if(ModelNode == nullptr)
-	{
-		return;
-	}
-
-	UScriptStruct* ScriptStruct = Cast<UScriptStruct>((UStruct*)StructBeingCustomized->GetStruct());
-	check(ScriptStruct);
-
-	IDetailCategoryBuilder& CategoryBuilder = DetailBuilder.EditCategory(*ScriptStruct->GetDisplayNameText().ToString());
-
-	for (TFieldIterator<FProperty> PropertyIt(ScriptStruct); PropertyIt; ++PropertyIt)
-	{
-		FProperty* Property = *PropertyIt;
-		TSharedPtr<IPropertyHandle> PropertyHandle = DetailBuilder.GetProperty(Property->GetFName(), ScriptStruct);
-		if (!PropertyHandle->IsValidHandle())
-		{
-			continue;
-		}
-		DetailBuilder.HideProperty(PropertyHandle);
-
-		URigVMPin* ModelPin = ModelNode->FindPin(Property->GetName());
-		if(ModelPin == nullptr)
-		{
-			continue;
-		}
-
-		if(ModelPin->IsBoundToVariable())
-		{
-			CategoryBuilder.AddCustomRow(FText::FromString(Property->GetName()))
-			.NameContent()
-			[
-				PropertyHandle->CreatePropertyNameWidget()
-			]
-			.ValueContent()
-			[
-				SNew(SControlRigVariableBinding)
-					.ModelPin(ModelPin)
-					.Blueprint(BlueprintBeingCustomized)
-			];
-
-			continue;
-		}
-
-		if (FNameProperty* NameProperty = CastField<FNameProperty>(Property))
-		{
-			FString CustomWidgetName = NameProperty->GetMetaData(TEXT("CustomWidget"));
-			if (!CustomWidgetName.IsEmpty())
-			{
-				const TArray<TSharedPtr<FString>>* NameList = nullptr;
-				if (CustomWidgetName == TEXT("BoneName"))
-				{
-					NameList = GraphBeingCustomized->GetBoneNameList();
-				}
-				else if (CustomWidgetName == TEXT("ControlName"))
-				{
-					NameList = GraphBeingCustomized->GetControlNameList();
-				}
-				else if (CustomWidgetName == TEXT("SpaceName"))
-				{
-					NameList = GraphBeingCustomized->GetNullNameList();
-				}
-				else if (CustomWidgetName == TEXT("CurveName"))
-				{
-					NameList = GraphBeingCustomized->GetCurveNameList();
-				}
-
-				if (NameList)
-				{
-					TSharedPtr<SControlRigGraphPinNameListValueWidget> NameListWidget;
-
-					CategoryBuilder.AddCustomRow(FText::FromString(Property->GetName()))
-						.NameContent()
-						[
-							PropertyHandle->CreatePropertyNameWidget()
-						]
-						.ValueContent()
-						[
-							SAssignNew(NameListWidget, SControlRigGraphPinNameListValueWidget)
-							.OptionsSource(NameList)
-							.OnGenerateWidget(this, &FRigUnitDetails::MakeNameListItemWidget)
-							.OnSelectionChanged(this, &FRigUnitDetails::OnNameListChanged, StructBeingCustomized, NameProperty, DetailBuilder.GetPropertyUtilities())
-							.OnComboBoxOpening(this, &FRigUnitDetails::OnNameListComboBox, StructBeingCustomized, NameProperty, NameList)
-							.InitiallySelectedItem(GetCurrentlySelectedItem(StructBeingCustomized, NameProperty, NameList))
-							.Content()
-							[
-								SNew(STextBlock)
-								.Text(this, &FRigUnitDetails::GetNameListText, StructBeingCustomized, NameProperty)
-							]
-						];
-
-					NameListWidgets.Add(Property->GetFName(), NameListWidget);
-				}
-				else
-				{
-					CategoryBuilder.AddCustomRow(FText::FromString(Property->GetName()))
-						.NameContent()
-						[
-							PropertyHandle->CreatePropertyNameWidget()
-						];
-				}
-				continue;
-			}
-		}
-		else if (FStructProperty* StructProperty = CastField<FStructProperty>(Property))
-		{
-			const FSimpleDelegate OnStructContentsChangedDelegate = FSimpleDelegate::CreateSP(this, &FRigUnitDetails::OnStructContentsChanged, Property, DetailBuilder.GetPropertyUtilities());
-			PropertyHandle->SetOnPropertyValueChanged(OnStructContentsChangedDelegate);
-			PropertyHandle->SetOnChildPropertyValueChanged(OnStructContentsChangedDelegate);
-		}
-
-		CategoryBuilder.AddProperty(PropertyHandle);
-	}
-}
-
-TSharedRef<SWidget> FRigUnitDetails::MakeNameListItemWidget(TSharedPtr<FString> InItem)
-{
-	return 	SNew(STextBlock).Text(FText::FromString(*InItem));// .Font(FEditorStyle::GetFontStyle(TEXT("PropertyWindow.NormalFont")));
-}
-
-FText FRigUnitDetails::GetNameListText(TSharedPtr<FStructOnScope> InStructOnScope, FNameProperty* InProperty) const
-{
-	if (FName* Value = InProperty->ContainerPtrToValuePtr<FName>(InStructOnScope->GetStructMemory()))
-	{
-		return FText::FromName(*Value);
-	}
-	return FText();
-}
-
-TSharedPtr<FString> FRigUnitDetails::GetCurrentlySelectedItem(TSharedPtr<FStructOnScope> InStructOnScope, FNameProperty* InProperty, const TArray<TSharedPtr<FString>>* InNameList) const
-{
-	FString CurrentItem = GetNameListText(InStructOnScope, InProperty).ToString();
-	for (const TSharedPtr<FString>& Item : *InNameList)
-	{
-		if (Item->Equals(CurrentItem))
-		{
-			return Item;
-		}
-	}
-
-	return TSharedPtr<FString>();
-}
-
-
-void FRigUnitDetails::SetNameListText(const FText& NewTypeInValue, ETextCommit::Type /*CommitInfo*/, TSharedPtr<FStructOnScope> InStructOnScope, FNameProperty* InProperty, TSharedRef<IPropertyUtilities> PropertyUtilities)
-{
-	if (FName* Value = InProperty->ContainerPtrToValuePtr<FName>(InStructOnScope->GetStructMemory()))
-	{
-		*Value = *NewTypeInValue.ToString();
-
-		FPropertyChangedEvent ChangeEvent(InProperty, EPropertyChangeType::ValueSet);
-		PropertyUtilities->NotifyFinishedChangingProperties(ChangeEvent);
-	}
-}
-
-void FRigUnitDetails::OnNameListChanged(TSharedPtr<FString> NewSelection, ESelectInfo::Type SelectInfo, TSharedPtr<FStructOnScope> InStructOnScope, FNameProperty* InProperty, TSharedRef<IPropertyUtilities> PropertyUtilities)
-{
-	if (SelectInfo != ESelectInfo::Direct)
-	{
-		FString NewValue = *NewSelection.Get();
-		SetNameListText(FText::FromString(NewValue), ETextCommit::OnEnter, InStructOnScope, InProperty, PropertyUtilities);
-	}
-}
-
-void FRigUnitDetails::OnNameListComboBox(TSharedPtr<FStructOnScope> InStructOnScope, FNameProperty* InProperty, const TArray<TSharedPtr<FString>>* InNameList)
-{
-	TSharedPtr<SControlRigGraphPinNameListValueWidget> Widget = NameListWidgets.FindChecked(InProperty->GetFName());
-	const TSharedPtr<FString> CurrentlySelected = GetCurrentlySelectedItem(InStructOnScope, InProperty, InNameList);
-	Widget->SetSelectedItem(CurrentlySelected);
-}
-
-void FRigUnitDetails::OnStructContentsChanged(FProperty* InProperty, const TSharedRef<IPropertyUtilities> PropertyUtilities)
-{
-	const FPropertyChangedEvent ChangeEvent(InProperty, EPropertyChangeType::ValueSet);
-	PropertyUtilities->NotifyFinishedChangingProperties(ChangeEvent);
-}
-
 void FRigComputedTransformDetails::CustomizeHeader(TSharedRef<IPropertyHandle> InStructPropertyHandle, FDetailWidgetRow& HeaderRow, IPropertyTypeCustomizationUtils& StructCustomizationUtils)
 {
 	BlueprintBeingCustomized = nullptr;
diff --git a/Engine/Plugins/Animation/ControlRig/Source/ControlRigEditor/Private/ControlRigElementDetails.h b/Engine/Plugins/Animation/ControlRig/Source/ControlRigEditor/Private/ControlRigElementDetails.h
--- a/Engine/Plugins/Animation/ControlRig/Source/ControlRigEditor/Private/ControlRigElementDetails.h
+++ b/Engine/Plugins/Animation/ControlRig/Source/ControlRigEditor/Private/ControlRigElementDetails.h
@@ -63,33 +63,6 @@ protected:
 	TSharedPtr<SSearchableComboBox> SearchableComboBox;
 };
 
-class FRigUnitDetails : public IDetailCustomization
-{
-public:
-
-	static TSharedRef<IDetailCustomization> MakeInstance()
-	{
-		return MakeShareable(new FRigUnitDetails);
-	}
-
-	/** IDetailCustomization interface */
-	virtual void CustomizeDetails(IDetailLayoutBuilder& DetailBuilder) override;
-
-protected:
-
-	TSharedRef<SWidget> MakeNameListItemWidget(TSharedPtr<FString> InItem);
-	FText GetNameListText(TSharedPtr<FStructOnScope> InStructOnScope, FNameProperty* InProperty) const;
-	TSharedPtr<FString> GetCurrentlySelectedItem(TSharedPtr<FStructOnScope> InStructOnScope, FNameProperty* InProperty, const TArray<TSharedPtr<FString>>* InNameList) const;
-	void SetNameListText(const FText& NewTypeInValue, ETextCommit::Type /*CommitInfo*/, TSharedPtr<FStructOnScope> InStructOnScope, FNameProperty* InProperty, TSharedRef<IPropertyUtilities> PropertyUtilities);
-	void OnNameListChanged(TSharedPtr<FString> NewSelection, ESelectInfo::Type SelectInfo, TSharedPtr<FStructOnScope> InStructOnScope, FNameProperty* InProperty, TSharedRef<IPropertyUtilities> PropertyUtilities);
-	void OnNameListComboBox(TSharedPtr<FStructOnScope> InStructOnScope, FNameProperty* InProperty, const TArray<TSharedPtr<FString>>* InNameList);
-	void OnStructContentsChanged(FProperty* InProperty, const TSharedRef<IPropertyUtilities> PropertyUtilities);
-
-	UControlRigBlueprint* BlueprintBeingCustomized;
-	UControlRigGraph* GraphBeingCustomized;
-	TMap<FName, TSharedPtr<SControlRigGraphPinNameListValueWidget>> NameListWidgets;
-};
-
 UENUM()
 enum class ERigElementDetailsTransformComponent : uint8
 {
diff --git a/Engine/Plugins/Animation/ControlRig/Source/ControlRigEditor/Private/ControlRigGraphDetails.cpp b/Engine/Plugins/Animation/ControlRig/Source/ControlRigEditor/Private/ControlRigGraphDetails.cpp
--- a/Engine/Plugins/Animation/ControlRig/Source/ControlRigEditor/Private/ControlRigGraphDetails.cpp
+++ b/Engine/Plugins/Animation/ControlRig/Source/ControlRigEditor/Private/ControlRigGraphDetails.cpp
@@ -27,6 +27,8 @@
 
 #define LOCTEXT_NAMESPACE "ControlRigGraphDetails"
 
+static const FText ControlRigGraphDetailsMultipleValues = LOCTEXT("MultipleValues", "Multiple Values");
+
 FControlRigArgumentGroupLayout::FControlRigArgumentGroupLayout(
 	URigVMGraph* InGraph, 
 	UControlRigBlueprint* InBlueprint, 
@@ -565,7 +567,8 @@ FControlRigArgumentDefaultNode::FControlRigArgumentDefaultNode(
 					FOnGraphChanged::FDelegate::CreateRaw(this, &FControlRigArgumentDefaultNode::OnGraphChanged)
 				);
 			}
-		}
+		}	
+
 	}
 }
 
@@ -1272,6 +1275,11 @@ TSharedRef<ITableRow> FControlRigGraphDetails::HandleGenerateRowAccessSpecifier(
         ];
 }
 
+FControlRigWrappedNodeDetails::FControlRigWrappedNodeDetails()
+: BlueprintBeingCustomized(nullptr)
+{
+}
+
 TSharedRef<IDetailCustomization> FControlRigWrappedNodeDetails::MakeInstance()
 {
 	return MakeShareable(new FControlRigWrappedNodeDetails);
@@ -1279,93 +1287,305 @@ TSharedRef<IDetailCustomization> FControlRigWrappedNodeDetails::MakeInstance()
 
 void FControlRigWrappedNodeDetails::CustomizeDetails(IDetailLayoutBuilder& DetailLayout)
 {
-	TArray<TWeakObjectPtr<UObject>> Objects;
-	DetailLayout.GetObjectsBeingCustomized(Objects);
-
-	TArray<UDetailsViewWrapperObject*> Wrappers;
-	TArray<URigVMNode*> Nodes;
-	for(const TWeakObjectPtr<UObject>& Object : Objects)
+	TArray<TWeakObjectPtr<UObject>> DetailObjects;
+	DetailLayout.GetObjectsBeingCustomized(DetailObjects);
+	if (DetailObjects.Num() == 0)
 	{
-		UDetailsViewWrapperObject* Wrapper = CastChecked<UDetailsViewWrapperObject>(Object);
-		Wrappers.Add(Wrapper);
+		return;
+	}
 
-		URigVMNode* Node = Wrapper->GetTypedOuter<URigVMNode>();
-		if(Node == nullptr)
+	for(TWeakObjectPtr<UObject> DetailObject : DetailObjects)
+	{
+		UDetailsViewWrapperObject* WrapperObject = CastChecked<UDetailsViewWrapperObject>(DetailObject.Get());
+		if(BlueprintBeingCustomized == nullptr)
 		{
-			return;
+			BlueprintBeingCustomized = WrapperObject->GetTypedOuter<UControlRigBlueprint>();
 		}
-		Nodes.Add(Node);
-	}
 
-	UDetailsViewWrapperObject* FirstWrapper = Wrappers[0];
-	URigVMNode* FirstNode = Nodes[0];
-
-	TArray<FName> Categories;
-	DetailLayout.GetCategoryNames(Categories);
+		ObjectsBeingCustomized.Add(WrapperObject);
+		NodesBeingCustomized.Add(CastChecked<URigVMNode>(WrapperObject->GetOuter()));
+	}
 
-	if(Categories.IsEmpty())
+	if (BlueprintBeingCustomized == nullptr || ObjectsBeingCustomized.IsEmpty() || NodesBeingCustomized.IsEmpty())
 	{
 		return;
 	}
-	
-	IDetailCategoryBuilder& DefaultsCategory = DetailLayout.EditCategory(Categories[0]);
 
-	UControlRigBlueprint* Blueprint = FirstNode->GetTypedOuter<UControlRigBlueprint>();
-	if(Blueprint == nullptr)
+	UClass* WrapperClass = ObjectsBeingCustomized[0]->GetClass();
+
+	// now loop over all of the properties and display them
+	TArray<TSharedPtr<IPropertyHandle>> PropertiesToVisit;
+	for (TFieldIterator<FProperty> PropertyIt(WrapperClass); PropertyIt; ++PropertyIt)
 	{
-		return;
+		FProperty* Property = *PropertyIt;
+		TSharedPtr<IPropertyHandle> PropertyHandle = DetailLayout.GetProperty(Property->GetFName(), WrapperClass);
+		if (!PropertyHandle->IsValidHandle())
+		{
+			continue;
+		}
+		PropertiesToVisit.Add(PropertyHandle);
+
+		// check if any / all pins are bound to a variable
+		int32 PinsBoundToVariable = 0;
+		TArray<URigVMPin*> ModelPins;
+		for(TWeakObjectPtr<URigVMNode> Node : NodesBeingCustomized)
+		{
+			if(URigVMPin* ModelPin = Node->FindPin(Property->GetName()))
+			{
+				ModelPins.Add(ModelPin);
+				PinsBoundToVariable += ModelPin->IsBoundToVariable() ? 1 : 0;
+			}
+		}
+
+		if(PinsBoundToVariable > 0)
+		{
+			if(PinsBoundToVariable == ModelPins.Num())
+			{
+				if(IDetailPropertyRow* Row = DetailLayout.EditDefaultProperty(PropertyHandle))
+				{
+					Row->CustomWidget()
+					.NameContent()
+					[
+						PropertyHandle->CreatePropertyNameWidget()
+					]
+					.ValueContent()
+					[
+					SNew(SControlRigVariableBinding)
+						.ModelPins(ModelPins)
+						.Blueprint(BlueprintBeingCustomized)
+					];
+				}
+			}
+			else // in this case some pins are bound, and some are not - we'll hide the input value widget
+			{
+				if(IDetailPropertyRow* Row = DetailLayout.EditDefaultProperty(PropertyHandle))
+				{
+					Row->CustomWidget()
+					.NameContent()
+					[
+						PropertyHandle->CreatePropertyNameWidget()
+					];
+				}
+			}
+
+			continue;
+		}
+		
+		if (FNameProperty* NameProperty = CastField<FNameProperty>(Property))
+        {
+        	FString CustomWidgetName = NameProperty->GetMetaData(TEXT("CustomWidget"));
+        	if (!CustomWidgetName.IsEmpty())
+        	{
+        		UControlRigGraph* GraphBeingCustomized = Cast<UControlRigGraph>(
+        			BlueprintBeingCustomized->GetEdGraph(NodesBeingCustomized[0]->GetGraph()));
+        		ensure(GraphBeingCustomized);
+        		
+        		const TArray<TSharedPtr<FString>>* NameList = nullptr;
+        		if (CustomWidgetName == TEXT("BoneName"))
+        		{
+        			NameList = GraphBeingCustomized->GetBoneNameList();
+        		}
+        		else if (CustomWidgetName == TEXT("ControlName"))
+        		{
+        			NameList = GraphBeingCustomized->GetControlNameList();
+        		}
+        		else if (CustomWidgetName == TEXT("SpaceName"))
+        		{
+        			NameList = GraphBeingCustomized->GetNullNameList();
+        		}
+        		else if (CustomWidgetName == TEXT("CurveName"))
+        		{
+        			NameList = GraphBeingCustomized->GetCurveNameList();
+        		}
+
+        		if (NameList)
+        		{
+        			TSharedPtr<SControlRigGraphPinNameListValueWidget> NameListWidget;
+
+        			if(IDetailPropertyRow* Row = DetailLayout.EditDefaultProperty(PropertyHandle))
+        			{
+        				Row->CustomWidget()
+						.NameContent()
+						[
+							PropertyHandle->CreatePropertyNameWidget()
+						]
+						.ValueContent()
+						[
+							SAssignNew(NameListWidget, SControlRigGraphPinNameListValueWidget)
+							.OptionsSource(NameList)
+							.OnGenerateWidget(this, &FControlRigWrappedNodeDetails::MakeNameListItemWidget)
+							.OnSelectionChanged(this, &FControlRigWrappedNodeDetails::OnNameListChanged, NameProperty, DetailLayout.GetPropertyUtilities())
+							.OnComboBoxOpening(this, &FControlRigWrappedNodeDetails::OnNameListComboBox, NameProperty, NameList)
+							.InitiallySelectedItem(GetCurrentlySelectedItem(NameProperty, NameList))
+							.Content()
+							[
+								SNew(STextBlock)
+								.Text(this, &FControlRigWrappedNodeDetails::GetNameListText, NameProperty)
+								.ColorAndOpacity_Lambda([this, NameProperty]() -> FSlateColor
+								{
+									static FText NoneText = LOCTEXT("None", "None"); 
+									if(GetNameListText(NameProperty).EqualToCaseIgnored(NoneText))
+									{
+										return FSlateColor(FLinearColor::Red);
+									}
+									return FSlateColor::UseForeground();
+								})
+							]
+        				];
+        			}        			
+        			NameListWidgets.Add(Property->GetFName(), NameListWidget);
+        		}
+        		else
+        		{
+        			if(IDetailPropertyRow* Row = DetailLayout.EditDefaultProperty(PropertyHandle))
+        			{
+        				Row->CustomWidget()
+						.NameContent()
+						[
+							PropertyHandle->CreatePropertyNameWidget()
+						];
+        			}
+        		}
+        		continue;
+        	}
+        }
 	}
 
-	for(URigVMPin* Pin : FirstNode->GetPins())
+	// now loop over all handles and determine expansion states of the corresponding pins
+	for (int32 Index = 0; Index < PropertiesToVisit.Num(); Index++)
 	{
-		bool bHasAnyInputLink = false;
-		for(URigVMNode* EachNode : Nodes)
+		TSharedPtr<IPropertyHandle> PropertyHandle = PropertiesToVisit[Index];
+		FProperty* Property = PropertyHandle->GetProperty();
+
+		// certain properties we don't look at for expansion states
+		if(FStructProperty* StructProperty = CastField<FStructProperty>(Property))
+		{
+			if(StructProperty->Struct == TBaseStructure<FVector>::Get() ||
+				StructProperty->Struct == TBaseStructure<FVector2D>::Get() ||
+				StructProperty->Struct == TBaseStructure<FRotator>::Get() ||
+				StructProperty->Struct == TBaseStructure<FQuat>::Get())
+			{
+				continue;
+			}
+		}
+
+		bool bFound = false;
+		const FString PinPath = PropertyHandle->GeneratePathToProperty();
+		for(TWeakObjectPtr<URigVMNode> Node : NodesBeingCustomized)
 		{
-			if(URigVMPin* EachPin = EachNode->FindPin(Pin->GetName()))
+			if(URigVMPin* Pin = Node->FindPin(PinPath))
 			{
-				if(EachPin->GetSourceLinks(true).Num() > 0)
+				bFound = true;
+				
+				if(Pin->IsExpanded())
 				{
-					bHasAnyInputLink = true;
+					if(IDetailPropertyRow* Row = DetailLayout.EditDefaultProperty(PropertyHandle))
+					{
+						Row->ShouldAutoExpand(true);
+					}
 					break;
 				}
 			}
 		}
 
-		TSharedPtr<IPropertyHandle> PinHandle = DetailLayout.GetProperty(Pin->GetFName());
-		if(PinHandle.IsValid())
+		if(!bFound)
 		{
-			static const TCHAR DisabledFormat[] = TEXT("%s\n\nNote: Editing disabled since pin has an input link.");
+			continue;
+		}
 
-			if (Pin->IsBoundToVariable())
+		uint32 NumChildren = 0;
+		PropertyHandle->GetNumChildren(NumChildren);
+		for(uint32 ChildIndex = 0; ChildIndex < NumChildren; ChildIndex++)
+		{
+			PropertiesToVisit.Add(PropertyHandle->GetChildHandle(ChildIndex));
+		}
+	}
+
+	CustomizeLiveValues(DetailLayout);
+}
+
+TSharedRef<SWidget> FControlRigWrappedNodeDetails::MakeNameListItemWidget(TSharedPtr<FString> InItem)
+{
+	return 	SNew(STextBlock).Text(FText::FromString(*InItem));// .Font(FEditorStyle::GetFontStyle(TEXT("PropertyWindow.NormalFont")));
+}
+
+FText FControlRigWrappedNodeDetails::GetNameListText(FNameProperty* InProperty) const
+{
+	FText FirstText;
+	for(TWeakObjectPtr<UDetailsViewWrapperObject> ObjectBeingCustomized : ObjectsBeingCustomized)
+	{
+		if (FName* Value = InProperty->ContainerPtrToValuePtr<FName>(ObjectBeingCustomized.Get()))
+		{
+			FText Text = FText::FromName(*Value);
+			if(FirstText.IsEmpty())
 			{
-				DefaultsCategory.AddCustomRow(FText::FromName(Pin->GetDisplayName()))
-				.NameContent()
-				[
-					PinHandle->CreatePropertyNameWidget()
-				]
-				.ValueContent()
-				[
-					SNew(SControlRigVariableBinding)
-						.ModelPin(Pin)
-						.Blueprint(Blueprint)
-				];
+				FirstText = Text;
 			}
-			else
+			else if(!FirstText.EqualTo(Text))
 			{
-				DefaultsCategory.AddProperty(PinHandle)
-				.DisplayName(FText::FromName(Pin->GetDisplayName()))
-				.IsEnabled(!bHasAnyInputLink)
-				.ToolTip(FText::FromString(FString::Printf(DisabledFormat, *Pin->GetToolTipText().ToString())));				
+				return ControlRigGraphDetailsMultipleValues;
 			}
 		}
 	}
+	return FirstText;
+}
 
-	if(Objects.Num() > 1)
+TSharedPtr<FString> FControlRigWrappedNodeDetails::GetCurrentlySelectedItem(FNameProperty* InProperty, const TArray<TSharedPtr<FString>>* InNameList) const
+{
+	FString CurrentItem = GetNameListText(InProperty).ToString();
+	for (const TSharedPtr<FString>& Item : *InNameList)
 	{
-		return;
+		if (Item->Equals(CurrentItem))
+		{
+			return Item;
+		}
 	}
+	return TSharedPtr<FString>();
+}
 
-	UControlRig* DebuggedRig = Cast<UControlRig>(Blueprint->GetObjectBeingDebugged());
+
+void FControlRigWrappedNodeDetails::SetNameListText(const FText& NewTypeInValue, ETextCommit::Type, FNameProperty* InProperty, TSharedRef<IPropertyUtilities> PropertyUtilities)
+{
+	URigVMGraph* Graph = NodesBeingCustomized[0]->GetGraph();
+	URigVMController* Controller = BlueprintBeingCustomized->GetController(Graph);
+
+	Controller->OpenUndoBracket(FString::Printf(TEXT("Set %s"), *InProperty->GetName()));
+	
+	for(TWeakObjectPtr<URigVMNode> Node : NodesBeingCustomized)
+	{
+		if(URigVMPin* Pin = Node->FindPin(InProperty->GetName()))
+		{
+			Controller->SetPinDefaultValue(Pin->GetPinPath(), NewTypeInValue.ToString(), false, true, false, true);
+		}
+	}
+
+	Controller->CloseUndoBracket();
+}
+
+void FControlRigWrappedNodeDetails::OnNameListChanged(TSharedPtr<FString> NewSelection, ESelectInfo::Type SelectInfo, FNameProperty* InProperty, TSharedRef<IPropertyUtilities> PropertyUtilities)
+{
+	if (SelectInfo != ESelectInfo::Direct)
+	{
+		FString NewValue = *NewSelection.Get();
+		SetNameListText(FText::FromString(NewValue), ETextCommit::OnEnter, InProperty, PropertyUtilities);
+	}
+}
+
+void FControlRigWrappedNodeDetails::OnNameListComboBox(FNameProperty* InProperty, const TArray<TSharedPtr<FString>>* InNameList)
+{
+	TSharedPtr<SControlRigGraphPinNameListValueWidget> Widget = NameListWidgets.FindChecked(InProperty->GetFName());
+	const TSharedPtr<FString> CurrentlySelected = GetCurrentlySelectedItem(InProperty, InNameList);
+	Widget->SetSelectedItem(CurrentlySelected);
+}
+
+void FControlRigWrappedNodeDetails::CustomizeLiveValues(IDetailLayoutBuilder& DetailLayout)
+{
+	if(ObjectsBeingCustomized.Num() != 1)
+	{
+		return;
+	}
+	
+	UControlRig* DebuggedRig = Cast<UControlRig>(BlueprintBeingCustomized->GetObjectBeingDebugged());
 	if(DebuggedRig == nullptr)
 	{
 		return;
@@ -1377,7 +1597,10 @@ void FControlRigWrappedNodeDetails::CustomizeDetails(IDetailLayoutBuilder& Detai
 		return;
 	}
 
-	TSharedPtr<FRigVMParserAST> AST = FirstNode->GetGraph()->GetRuntimeAST(Blueprint->VMCompileSettings.ASTSettings, false);
+	UDetailsViewWrapperObject* FirstWrapper = ObjectsBeingCustomized[0].Get();
+	URigVMNode* FirstNode = NodesBeingCustomized[0].Get();
+
+	TSharedPtr<FRigVMParserAST> AST = FirstNode->GetGraph()->GetRuntimeAST(BlueprintBeingCustomized->VMCompileSettings.ASTSettings, false);
 	if(!AST.IsValid())
 	{
 		return;
@@ -1404,7 +1627,7 @@ void FControlRigWrappedNodeDetails::CustomizeDetails(IDetailLayoutBuilder& Detai
 		}
 		
 		URigVMPin* SourcePin = Pin;
-		if(Blueprint->VMCompileSettings.ASTSettings.bFoldAssignments)
+		if(BlueprintBeingCustomized->VMCompileSettings.ASTSettings.bFoldAssignments)
 		{
 			do
 			{
@@ -1474,7 +1697,7 @@ void FControlRigWrappedNodeDetails::CustomizeDetails(IDetailLayoutBuilder& Detai
 			const FRigVMVarExprAST* VarExpr = Expression->To<FRigVMVarExprAST>();
 
 			FString PinHash = URigVMCompiler::GetPinHash(SourcePin, VarExpr, false);
-			const FRigVMOperand* Operand = Blueprint->PinToOperandMap.Find(PinHash);
+			const FRigVMOperand* Operand = BlueprintBeingCustomized->PinToOperandMap.Find(PinHash);
 			if(Operand)
 			{
 				if(Operand->GetRegisterOffset() != INDEX_NONE)
@@ -1569,4 +1792,97 @@ void FControlRigWrappedNodeDetails::CustomizeDetails(IDetailLayoutBuilder& Detai
 	}
 }
 
+FControlRigGraphMathTypeDetails::FControlRigGraphMathTypeDetails()
+: ScriptStruct(nullptr)
+, BlueprintBeingCustomized(nullptr)
+, GraphBeingCustomized(nullptr)
+, bEnabled(true)
+{
+}
+
+void FControlRigGraphMathTypeDetails::CustomizeHeader(TSharedRef<IPropertyHandle> InPropertyHandle,
+                                                           FDetailWidgetRow& HeaderRow, IPropertyTypeCustomizationUtils& StructCustomizationUtils)
+{
+	TArray<UObject*> Objects;
+	InPropertyHandle->GetOuterObjects(Objects);
+
+	for (UObject* Object : Objects)
+	{
+		ObjectsBeingCustomized.Add(Object);
+
+		if(BlueprintBeingCustomized == nullptr)
+		{
+			BlueprintBeingCustomized = Object->GetTypedOuter<UControlRigBlueprint>();
+		}
+
+		if(GraphBeingCustomized == nullptr)
+		{
+			GraphBeingCustomized = Object->GetTypedOuter<URigVMGraph>();
+		}
+	}
+
+	ObjectBeingCustomizedView = TArrayView<const UObject* const>(ObjectsBeingCustomized.GetData(), ObjectsBeingCustomized.Num());
+
+	FProperty* Property = InPropertyHandle->GetProperty();
+	const FStructProperty* StructProperty = CastField<FStructProperty>(Property);
+	ScriptStruct = StructProperty->Struct;
+
+	TSharedPtr<class IPropertyHandle> ChainHandle = InPropertyHandle;
+	do
+	{
+		PropertyChain.AddHead(ChainHandle->GetProperty());
+		ChainHandle = ChainHandle->GetParentHandle();
+
+		if(ChainHandle.IsValid())
+		{
+			if(ChainHandle->GetProperty() == nullptr)
+			{
+				break;
+			}
+		}
+	}
+	while (ChainHandle.IsValid());
+	PropertyChain.SetActiveMemberPropertyNode(PropertyChain.GetTail()->GetValue());
+
+	bEnabled = !PropertyChain.GetHead()->GetValue()->HasAnyPropertyFlags(CPF_EditConst);
+}
+
+void FControlRigGraphMathTypeDetails::CustomizeChildren(TSharedRef<IPropertyHandle> InPropertyHandle,
+	IDetailChildrenBuilder& StructBuilder, IPropertyTypeCustomizationUtils& StructCustomizationUtils)
+{
+	if(!InPropertyHandle->IsValidHandle())
+	{
+		return;
+	}
+
+	if(ScriptStruct == TBaseStructure<FVector>::Get())
+	{
+		CustomizeVector<FVector, 3>(InPropertyHandle, StructBuilder, StructCustomizationUtils);
+	}
+	else if(ScriptStruct == TBaseStructure<FVector2D>::Get())
+	{
+		CustomizeVector<FVector2D, 2>(InPropertyHandle, StructBuilder, StructCustomizationUtils);
+	}
+	else if(ScriptStruct == TBaseStructure<FVector4>::Get())
+	{
+		CustomizeVector<FVector4, 4>(InPropertyHandle, StructBuilder, StructCustomizationUtils);
+	}
+	else if(ScriptStruct == TBaseStructure<FRotator>::Get())
+	{
+		CustomizeRotation<FRotator>(InPropertyHandle, StructBuilder, StructCustomizationUtils);
+	}
+	else if(ScriptStruct == TBaseStructure<FQuat>::Get())
+	{
+		CustomizeRotation<FQuat>(InPropertyHandle, StructBuilder, StructCustomizationUtils);
+	}
+	else if(ScriptStruct == TBaseStructure<FTransform>::Get())
+	{
+		CustomizeTransform<FTransform>(InPropertyHandle, StructBuilder, StructCustomizationUtils);
+	}
+	else if(ScriptStruct == TBaseStructure<FEulerTransform>::Get())
+	{
+		CustomizeTransform<FEulerTransform>(InPropertyHandle, StructBuilder, StructCustomizationUtils);
+	}
+}
+
 #undef LOCTEXT_NAMESPACE
diff --git a/Engine/Plugins/Animation/ControlRig/Source/ControlRigEditor/Private/ControlRigGraphDetails.h b/Engine/Plugins/Animation/ControlRig/Source/ControlRigEditor/Private/ControlRigGraphDetails.h
--- a/Engine/Plugins/Animation/ControlRig/Source/ControlRigEditor/Private/ControlRigGraphDetails.h
+++ b/Engine/Plugins/Animation/ControlRig/Source/ControlRigEditor/Private/ControlRigGraphDetails.h
@@ -13,7 +13,10 @@
 #include "Graph/SControlRigGraphNode.h"
 #include "Widgets/Colors/SColorBlock.h"
 #include "DetailsViewWrapperObject.h"
+#include "IDetailPropertyExtensionHandler.h"
 #include "Graph/ControlRigGraphSchema.h"
+#include "SAdvancedTransformInputBox.h"
+#include "PropertyEditor/Private/PropertyNode.h"
 
 class IDetailLayoutBuilder;
 
@@ -217,10 +220,643 @@ private:
 class FControlRigWrappedNodeDetails : public IDetailCustomization
 {
 public:
+	
+	FControlRigWrappedNodeDetails();
 
 	/** Makes a new instance of this detail layout class for a specific detail view requesting it */
 	static TSharedRef<IDetailCustomization> MakeInstance();
 
 	// IDetailCustomization interface
 	virtual void CustomizeDetails(IDetailLayoutBuilder& DetailLayout) override;
+
+	TSharedRef<SWidget> MakeNameListItemWidget(TSharedPtr<FString> InItem);
+	FText GetNameListText(FNameProperty* InProperty) const;
+	TSharedPtr<FString> GetCurrentlySelectedItem(FNameProperty* InProperty, const TArray<TSharedPtr<FString>>* InNameList) const;
+	void SetNameListText(const FText& NewTypeInValue, ETextCommit::Type, FNameProperty* InProperty, TSharedRef<IPropertyUtilities> PropertyUtilities);
+	void OnNameListChanged(TSharedPtr<FString> NewSelection, ESelectInfo::Type SelectInfo, FNameProperty* InProperty, TSharedRef<IPropertyUtilities> PropertyUtilities);
+	void OnNameListComboBox(FNameProperty* InProperty, const TArray<TSharedPtr<FString>>* InNameList);
+	void CustomizeLiveValues(IDetailLayoutBuilder& DetailLayout);
+
+	UControlRigBlueprint* BlueprintBeingCustomized;
+	TArray<TWeakObjectPtr<UDetailsViewWrapperObject>> ObjectsBeingCustomized;
+	TArray<TWeakObjectPtr<URigVMNode>> NodesBeingCustomized;
+	TMap<FName, TSharedPtr<SControlRigGraphPinNameListValueWidget>> NameListWidgets;
+};
+
+/** Customization for editing a Control Rig node */
+class FControlRigGraphMathTypeDetails : public IPropertyTypeCustomization
+{
+public:
+
+	FControlRigGraphMathTypeDetails();
+	
+	static TSharedRef<IPropertyTypeCustomization> MakeInstance()
+	{
+		return MakeShareable(new FControlRigGraphMathTypeDetails);
+	}
+
+	/** IPropertyTypeCustomization interface */
+	virtual void CustomizeHeader(TSharedRef<class IPropertyHandle> InPropertyHandle, class FDetailWidgetRow& HeaderRow, IPropertyTypeCustomizationUtils& StructCustomizationUtils) override;
+	virtual void CustomizeChildren(TSharedRef<class IPropertyHandle> InPropertyHandle, class IDetailChildrenBuilder& StructBuilder, IPropertyTypeCustomizationUtils& StructCustomizationUtils) override;
+
+protected:
+
+	// extracts the value for a nested property (for Example Settings.WorldTransform) from an outer owner
+	template<typename ValueType>
+	FORCEINLINE ValueType& ContainerUObjectToValueRef(UObject* InOwner) const
+	{
+		FEditPropertyChain::TDoubleLinkedListNode* PropertyNode = PropertyChain.GetHead();
+		uint8* MemoryPtr = (uint8*)InOwner;
+		do
+		{
+			MemoryPtr = PropertyNode->GetValue()->ContainerPtrToValuePtr<uint8>(MemoryPtr);
+			PropertyNode = PropertyNode->GetNextNode();
+		}
+		while (PropertyNode);
+
+		return *(ValueType*)MemoryPtr;
+	}
+
+	// specializations for FEulerTransform and FRotator at the end of this file
+	template<typename ValueType>
+	static bool IsQuaternionBasedRotation() { return true; }
+
+	// returns the numeric value of a vector component (or empty optional for varying values)
+	template<typename VectorType, typename NumericType>
+	TOptional<NumericType> GetVectorComponent(TSharedRef<class IPropertyHandle> InPropertyHandle, int32 InComponent) 
+	{
+		TOptional<NumericType> Result;
+		for(UObject* Object : ObjectsBeingCustomized)
+		{
+			if(InPropertyHandle->IsValidHandle())
+			{
+				const VectorType& Vector = ContainerUObjectToValueRef<VectorType>(Object);
+				NumericType Component = Vector[InComponent];
+				if(Result.IsSet())
+				{
+					if(!FMath::IsNearlyEqual(Result.GetValue(), Component))
+					{
+						return TOptional<NumericType>();
+					}
+				}
+				else
+				{
+					Result = Component;
+				}
+			}
+		}
+		return Result;
+	};
+
+	// called when a numeric value of a vector component is changed
+	template<typename VectorType, typename NumericType>
+	void OnVectorComponentChanged(TSharedRef<class IPropertyHandle> InPropertyHandle, int32 InComponent, NumericType InValue, bool bIsCommit, ETextCommit::Type InCommitType = ETextCommit::Default)
+	{
+		FPropertyChangedEvent PropertyChangedEvent(InPropertyHandle->GetProperty(), bIsCommit ? EPropertyChangeType::ValueSet : EPropertyChangeType::Interactive, ObjectsBeingCustomized);
+		FPropertyChangedChainEvent PropertyChangedChainEvent(PropertyChain, PropertyChangedEvent);
+
+		URigVMController* Controller = nullptr;
+		if(BlueprintBeingCustomized && GraphBeingCustomized)
+		{
+			Controller = BlueprintBeingCustomized->GetController(GraphBeingCustomized);
+			if(bIsCommit)
+			{
+				Controller->OpenUndoBracket(FString::Printf(TEXT("Set %s"), *InPropertyHandle->GetProperty()->GetName()));
+			}
+		}
+
+		for(int32 Index = 0; Index < ObjectsBeingCustomized.Num(); Index++)
+		{
+			UObject* Object = ObjectsBeingCustomized[Index];
+			if(InPropertyHandle->IsValidHandle())
+			{
+				VectorType& Vector = ContainerUObjectToValueRef<VectorType>(Object);
+				VectorType PreviousVector = Vector;
+				Vector[InComponent] = InValue;
+					
+				if(!PreviousVector.Equals(Vector))
+				{
+					Object->PostEditChangeChainProperty(PropertyChangedChainEvent);
+					InPropertyHandle->NotifyPostChange(PropertyChangedEvent.ChangeType);
+				}
+			}
+		}
+
+		if(Controller && bIsCommit)
+		{
+			Controller->CloseUndoBracket();
+		}
+	};
+
+	// specializations for FVector and FVector4 at the end of this file
+	template<typename VectorType>
+	FORCEINLINE void ExtendVectorArgs(TSharedRef<class IPropertyHandle> InPropertyHandle, void* ArgumentsPtr) {}
+
+	template<typename VectorType, int32 NumberOfComponents>
+	FORCEINLINE void CustomizeVector(TSharedRef<class IPropertyHandle> InPropertyHandle, class IDetailChildrenBuilder& StructBuilder, IPropertyTypeCustomizationUtils& StructCustomizationUtils)
+	{
+		typedef typename VectorType::FReal NumericType;
+		typedef SNumericVectorInputBox<NumericType, VectorType, NumberOfComponents> SLocalVectorInputBox;
+
+		typename SLocalVectorInputBox::FArguments Args;
+		Args.Font(IDetailLayoutBuilder::GetDetailFont());
+		Args.IsEnabled(bEnabled);
+		Args.AllowSpin(true);
+		Args.SpinDelta(0.01f);
+		Args.bColorAxisLabels(true);
+		Args.X_Lambda([this, InPropertyHandle]()
+		{
+			return GetVectorComponent<VectorType, NumericType>(InPropertyHandle, 0);
+		});
+		Args.OnXChanged_Lambda([this, InPropertyHandle](NumericType Value)
+		{
+			OnVectorComponentChanged<VectorType, NumericType>(InPropertyHandle, 0, Value, false);
+		});
+		Args.OnXCommitted_Lambda([this, InPropertyHandle](NumericType Value, ETextCommit::Type CommitType)
+		{
+			OnVectorComponentChanged<VectorType, NumericType>(InPropertyHandle, 0, Value, true, CommitType);
+		});
+		Args.Y_Lambda([this, InPropertyHandle]()
+		{
+			return GetVectorComponent<VectorType, NumericType>(InPropertyHandle, 1);
+		});
+		Args.OnYChanged_Lambda([this, InPropertyHandle](NumericType Value)
+		{
+			OnVectorComponentChanged<VectorType, NumericType>(InPropertyHandle, 1, Value, false);
+		});
+		Args.OnYCommitted_Lambda([this, InPropertyHandle](NumericType Value, ETextCommit::Type CommitType)
+		{
+			OnVectorComponentChanged<VectorType, NumericType>(InPropertyHandle, 1, Value, true, CommitType);
+		});
+
+		ExtendVectorArgs<VectorType>(InPropertyHandle, &Args);
+
+		StructBuilder.AddProperty(InPropertyHandle).CustomWidget()
+		.IsEnabled(bEnabled)
+		.NameContent()
+		[
+			InPropertyHandle->CreatePropertyNameWidget()
+		]
+		.ValueContent()
+		.MinDesiredWidth(375.f)
+		.MaxDesiredWidth(375.f)
+		.HAlign(HAlign_Left)
+		[
+			SArgumentNew(Args, SLocalVectorInputBox)
+		];
+	}
+
+	// returns the rotation for rotator or quaternions (or empty optional for varying values)
+	template<typename RotationType>
+	TOptional<RotationType> GetRotation(TSharedRef<class IPropertyHandle> InPropertyHandle)
+	{
+		TOptional<RotationType> Result;
+		for(UObject* Object : ObjectsBeingCustomized)
+		{
+			if(InPropertyHandle->IsValidHandle())
+			{
+				const RotationType& Rotation = ContainerUObjectToValueRef<RotationType>(Object);
+				if(Result.IsSet())
+				{
+					if(!Rotation.Equals(Result.GetValue()))
+					{
+						return TOptional<RotationType>();
+					}
+				}
+				else
+				{
+					Result = Rotation;
+				}
+			}
+		}
+		return Result;
+	};
+
+	// called when a rotation value is changed / committed
+	template<typename RotationType>
+	void OnRotationChanged(TSharedRef<class IPropertyHandle> InPropertyHandle, RotationType InValue, bool bIsCommit, ETextCommit::Type InCommitType = ETextCommit::Default)
+	{
+		FPropertyChangedEvent PropertyChangedEvent(InPropertyHandle->GetProperty(), bIsCommit ? EPropertyChangeType::ValueSet : EPropertyChangeType::Interactive, ObjectsBeingCustomized);
+		FPropertyChangedChainEvent PropertyChangedChainEvent(PropertyChain, PropertyChangedEvent);
+
+		URigVMController* Controller = nullptr;
+		if(BlueprintBeingCustomized && GraphBeingCustomized)
+		{
+			Controller = BlueprintBeingCustomized->GetController(GraphBeingCustomized);
+			if(bIsCommit)
+			{
+				Controller->OpenUndoBracket(FString::Printf(TEXT("Set %s"), *InPropertyHandle->GetProperty()->GetName()));
+			}
+		}
+
+		for(int32 Index = 0; Index < ObjectsBeingCustomized.Num(); Index++)
+		{
+			UObject* Object = ObjectsBeingCustomized[Index];
+			if(InPropertyHandle->IsValidHandle())
+			{
+				RotationType& Rotation = ContainerUObjectToValueRef<RotationType>(Object);
+				RotationType PreviousRotation = Rotation;
+				Rotation = InValue;
+					
+				if(!PreviousRotation.Equals(Rotation))
+				{
+					Object->PostEditChangeChainProperty(PropertyChangedChainEvent);
+					InPropertyHandle->NotifyPostChange(PropertyChangedEvent.ChangeType);
+				}
+			}
+		}
+
+		if(Controller && bIsCommit)
+		{
+			Controller->CloseUndoBracket();
+		}
+	};
+
+	// specializations for FRotator and FQuat at the end of this file
+	template<typename RotationType>
+	FORCEINLINE void ExtendRotationArgs(TSharedRef<class IPropertyHandle> InPropertyHandle, void* ArgumentsPtr) {}
+
+	// add the widget for a rotation (rotator or quat)
+	template<typename RotationType>
+	FORCEINLINE void CustomizeRotation(TSharedRef<class IPropertyHandle> InPropertyHandle, class IDetailChildrenBuilder& StructBuilder, IPropertyTypeCustomizationUtils& StructCustomizationUtils)
+	{
+		typedef typename RotationType::FReal NumericType;
+		typedef SAdvancedRotationInputBox<NumericType> SLocalRotationInputBox;
+		typename SLocalRotationInputBox::FArguments Args;
+		Args.Font(IDetailLayoutBuilder::GetDetailFont());
+		Args.IsEnabled(bEnabled);
+		Args.AllowSpin(true);
+		Args.bColorAxisLabels(true);
+
+		ExtendRotationArgs<RotationType>(InPropertyHandle, &Args);
+		
+		StructBuilder.AddProperty(InPropertyHandle).CustomWidget()
+		.IsEnabled(bEnabled)
+		.NameContent()
+		[
+			
+			InPropertyHandle->CreatePropertyNameWidget()
+		]
+		.ValueContent()
+		.MinDesiredWidth(375.f)
+		.MaxDesiredWidth(375.f)
+		.HAlign(HAlign_Left)
+		[
+			SArgumentNew(Args, SLocalRotationInputBox)
+		];
+	}
+
+	// add the widget for a transform / euler transform
+	template<typename TransformType>
+	FORCEINLINE void CustomizeTransform(TSharedRef<class IPropertyHandle> InPropertyHandle, class IDetailChildrenBuilder& StructBuilder, IPropertyTypeCustomizationUtils& StructCustomizationUtils)
+	{
+		typedef typename TransformType::FReal FReal;
+		typename SAdvancedTransformInputBox<TransformType>::FArguments WidgetArgs;
+		WidgetArgs.IsEnabled(bEnabled);
+		WidgetArgs.AllowEditRotationRepresentation(true);
+		WidgetArgs.UseQuaternionForRotation(IsQuaternionBasedRotation<TransformType>());
+
+		TransformType DefaultValue = ContainerUObjectToValueRef<TransformType>(ObjectsBeingCustomized[0]->GetClass()->GetDefaultObject());
+
+		WidgetArgs.DiffersFromDefault_Lambda([this, InPropertyHandle, DefaultValue](ESlateTransformComponent::Type InTransformComponent) -> bool
+		{
+			for(UObject* Object : ObjectsBeingCustomized)
+			{
+				if(InPropertyHandle->IsValidHandle())
+				{
+					const TransformType& Transform = ContainerUObjectToValueRef<TransformType>(Object);
+
+					switch(InTransformComponent)
+					{
+						case ESlateTransformComponent::Location:
+						{
+							if(!Transform.GetLocation().Equals(DefaultValue.GetLocation()))
+							{
+								return true;
+							}
+							break;
+						}
+						case ESlateTransformComponent::Rotation:
+						{
+							if(!Transform.Rotator().Equals(DefaultValue.Rotator()))
+							{
+								return true;
+							}
+							break;
+						}
+						case ESlateTransformComponent::Scale:
+						{
+							if(!Transform.GetScale3D().Equals(DefaultValue.GetScale3D()))
+							{
+								return true;
+							}
+							break;
+						}
+						default:
+						{
+							break;
+						}
+					}
+				}
+			}
+			return false;
+		});
+
+		WidgetArgs.OnGetNumericValue_Lambda([this, InPropertyHandle](
+			ESlateTransformComponent::Type InTransformComponent,
+			ESlateRotationRepresentation::Type InRotationRepresentation,
+			ESlateTransformSubComponent::Type InTransformSubComponent) -> TOptional<FReal>
+		{
+			TOptional<FReal> Result;
+			for(UObject* Object : ObjectsBeingCustomized)
+			{
+				if(InPropertyHandle->IsValidHandle())
+				{
+					const TransformType& Transform = ContainerUObjectToValueRef<TransformType>(Object);
+					
+					TOptional<FReal> Value = SAdvancedTransformInputBox<TransformType>::GetNumericValueFromTransform(
+						Transform,
+						InTransformComponent,
+						InRotationRepresentation,
+						InTransformSubComponent
+						);
+
+					if(Value.IsSet())
+					{
+						if(Result.IsSet())
+						{
+							if(!FMath::IsNearlyEqual(Result.GetValue(), Value.GetValue()))
+							{
+								return TOptional<FReal>();
+							}
+						}
+						else
+						{
+							Result = Value;
+						}
+					}
+				}
+			}
+			return Result;
+		});
+		
+
+		auto OnNumericValueChanged = [this, InPropertyHandle](
+			ESlateTransformComponent::Type InTransformComponent, 
+			ESlateRotationRepresentation::Type InRotationRepresentation, 
+			ESlateTransformSubComponent::Type InSubComponent,
+			FReal InValue,
+			bool bIsCommit,
+			ETextCommit::Type InCommitType = ETextCommit::Default)
+		{
+			FPropertyChangedEvent PropertyChangedEvent(InPropertyHandle->GetProperty(), bIsCommit ? EPropertyChangeType::ValueSet : EPropertyChangeType::Interactive, ObjectsBeingCustomized);
+			FPropertyChangedChainEvent PropertyChangedChainEvent(PropertyChain, PropertyChangedEvent);
+
+			URigVMController* Controller = nullptr;
+			if(BlueprintBeingCustomized && GraphBeingCustomized)
+			{
+				Controller = BlueprintBeingCustomized->GetController(GraphBeingCustomized);
+				if(bIsCommit)
+				{
+					Controller->OpenUndoBracket(FString::Printf(TEXT("Set %s"), *InPropertyHandle->GetProperty()->GetName()));
+				}
+			}
+
+			for(int32 Index = 0; Index < ObjectsBeingCustomized.Num(); Index++)
+			{
+				UObject* Object = ObjectsBeingCustomized[Index];
+				if(InPropertyHandle->IsValidHandle())
+				{
+					TransformType& Transform = ContainerUObjectToValueRef<TransformType>(Object);
+					TransformType PreviousTransform = Transform;
+					
+					SAdvancedTransformInputBox<TransformType>::ApplyNumericValueChange(
+						Transform,
+						InValue,
+						InTransformComponent,
+						InRotationRepresentation,
+						InSubComponent);
+
+					if(!PreviousTransform.Equals(Transform))
+					{
+						Object->PostEditChangeChainProperty(PropertyChangedChainEvent);
+						InPropertyHandle->NotifyPostChange(PropertyChangedEvent.ChangeType);
+					}
+				}
+			}
+
+			if(Controller && bIsCommit)
+			{
+				Controller->CloseUndoBracket();
+			}
+		};
+
+		WidgetArgs.OnNumericValueChanged_Lambda([OnNumericValueChanged](
+			ESlateTransformComponent::Type InTransformComponent, 
+			ESlateRotationRepresentation::Type InRotationRepresentation, 
+			ESlateTransformSubComponent::Type InSubComponent,
+			FReal InValue)
+		{
+			OnNumericValueChanged(InTransformComponent, InRotationRepresentation, InSubComponent, InValue, false);
+		});
+
+		WidgetArgs.OnNumericValueCommitted_Lambda([OnNumericValueChanged](
+			ESlateTransformComponent::Type InTransformComponent, 
+			ESlateRotationRepresentation::Type InRotationRepresentation, 
+			ESlateTransformSubComponent::Type InSubComponent,
+			FReal InValue, 
+			ETextCommit::Type InCommitType)
+		{
+			OnNumericValueChanged(InTransformComponent, InRotationRepresentation, InSubComponent, InValue, true, InCommitType);
+		});
+
+		WidgetArgs.OnResetToDefault_Lambda([this, DefaultValue, InPropertyHandle](ESlateTransformComponent::Type InTransformComponent)
+		{
+			URigVMController* Controller = nullptr;
+			if(BlueprintBeingCustomized && GraphBeingCustomized)
+			{
+				Controller = BlueprintBeingCustomized->GetController(GraphBeingCustomized);
+				if(Controller)
+				{
+					Controller->OpenUndoBracket(FString::Printf(TEXT("Reset %s to Default"), *InPropertyHandle->GetProperty()->GetName()));
+				}
+			}
+
+			FPropertyChangedEvent PropertyChangedEvent(InPropertyHandle->GetProperty(), EPropertyChangeType::ValueSet, ObjectsBeingCustomized);
+			FPropertyChangedChainEvent PropertyChangedChainEvent(PropertyChain, PropertyChangedEvent);
+			
+			for(int32 Index = 0; Index < ObjectsBeingCustomized.Num(); Index++)
+			{
+				UObject* Object = ObjectsBeingCustomized[Index];
+				if(InPropertyHandle->IsValidHandle())
+				{
+					TransformType& Transform = ContainerUObjectToValueRef<TransformType>(Object);
+					TransformType PreviousTransform = Transform;
+
+					switch(InTransformComponent)
+					{
+						case ESlateTransformComponent::Location:
+						{
+							Transform.SetLocation(DefaultValue.GetLocation());
+							break;
+						}
+						case ESlateTransformComponent::Rotation:
+						{
+							Transform.SetRotation(DefaultValue.GetRotation());
+							break;
+						}
+						case ESlateTransformComponent::Scale:
+						{
+							Transform.SetScale3D(DefaultValue.GetScale3D());
+							break;
+						}
+						case ESlateTransformComponent::Max:
+						default:
+						{
+							Transform.SetLocation(DefaultValue.GetLocation());
+							break;
+						}
+					}
+
+					
+					if(!PreviousTransform.Equals(Transform))
+					{
+						Object->PostEditChangeChainProperty(PropertyChangedChainEvent);
+						InPropertyHandle->NotifyPostChange(PropertyChangedEvent.ChangeType);
+					}
+				}
+			}
+
+			if(Controller)
+			{
+				Controller->CloseUndoBracket();
+			}
+		});
+
+		SAdvancedTransformInputBox<TransformType>::ConstructGroupedTransformRows(
+			StructBuilder,
+			InPropertyHandle->GetPropertyDisplayName(),
+			InPropertyHandle->GetToolTipText(),
+			WidgetArgs);
+	}
+		
+	UScriptStruct* ScriptStruct;
+	UControlRigBlueprint* BlueprintBeingCustomized;
+	URigVMGraph* GraphBeingCustomized;
+	TArray<UObject*> ObjectsBeingCustomized;
+	TArrayView<const UObject* const> ObjectBeingCustomizedView;
+	FEditPropertyChain PropertyChain;
+	bool bEnabled;
 };
+
+template<>
+FORCEINLINE bool FControlRigGraphMathTypeDetails::IsQuaternionBasedRotation<FEulerTransform>() { return false; }
+
+template<>
+FORCEINLINE bool FControlRigGraphMathTypeDetails::IsQuaternionBasedRotation<FRotator>() { return false; }
+
+template<>
+FORCEINLINE void FControlRigGraphMathTypeDetails::ExtendVectorArgs<FVector>(TSharedRef<class IPropertyHandle> InPropertyHandle, void* ArgumentsPtr)
+{
+	using VectorType = FVector;
+	typedef typename VectorType::FReal NumericType;
+	typedef SNumericVectorInputBox<NumericType, VectorType, 3> SLocalVectorInputBox;
+
+	typename SLocalVectorInputBox::FArguments& Args = *(typename SLocalVectorInputBox::FArguments*)ArgumentsPtr; 
+	Args
+	.Z_Lambda([this, InPropertyHandle]()
+	{
+		return GetVectorComponent<VectorType, NumericType>(InPropertyHandle, 2);
+	})
+	.OnZChanged_Lambda([this, InPropertyHandle](NumericType Value)
+	{
+		OnVectorComponentChanged<VectorType, NumericType>(InPropertyHandle, 2, Value, false);
+	})
+	.OnZCommitted_Lambda([this, InPropertyHandle](NumericType Value, ETextCommit::Type CommitType)
+	{
+		OnVectorComponentChanged<VectorType, NumericType>(InPropertyHandle, 2, Value, true, CommitType);
+	});
+}
+
+template<>
+FORCEINLINE void FControlRigGraphMathTypeDetails::ExtendVectorArgs<FVector4>(TSharedRef<class IPropertyHandle> InPropertyHandle, void* ArgumentsPtr)
+{
+	using VectorType = FVector4;
+	typedef typename VectorType::FReal NumericType;
+	typedef SNumericVectorInputBox<NumericType, VectorType, 4> SLocalVectorInputBox;
+
+	typename SLocalVectorInputBox::FArguments& Args = *(typename SLocalVectorInputBox::FArguments*)ArgumentsPtr; 
+	Args
+	.Z_Lambda([this, InPropertyHandle]()
+	{
+		return GetVectorComponent<VectorType, NumericType>(InPropertyHandle, 2);
+	})
+	.OnZChanged_Lambda([this, InPropertyHandle](NumericType Value)
+	{
+		OnVectorComponentChanged<VectorType, NumericType>(InPropertyHandle, 2, Value, false);
+	})
+	.OnZCommitted_Lambda([this, InPropertyHandle](NumericType Value, ETextCommit::Type CommitType)
+	{
+		OnVectorComponentChanged<VectorType, NumericType>(InPropertyHandle, 2, Value, true, CommitType);
+	})
+	.W_Lambda([this, InPropertyHandle]()
+	{
+		return GetVectorComponent<VectorType, NumericType>(InPropertyHandle, 3);
+	})
+	.OnWChanged_Lambda([this, InPropertyHandle](NumericType Value)
+	{
+		OnVectorComponentChanged<VectorType, NumericType>(InPropertyHandle, 3, Value, false);
+	})
+	.OnWCommitted_Lambda([this, InPropertyHandle](NumericType Value, ETextCommit::Type CommitType)
+	{
+		OnVectorComponentChanged<VectorType, NumericType>(InPropertyHandle, 3, Value, true, CommitType);
+	});
+}
+
+template<>
+FORCEINLINE void FControlRigGraphMathTypeDetails::ExtendRotationArgs<FQuat>(TSharedRef<class IPropertyHandle> InPropertyHandle, void* ArgumentsPtr)
+{
+	using RotationType = FQuat;
+	typedef typename RotationType::FReal NumericType;
+	typedef SAdvancedRotationInputBox<NumericType> SLocalRotationInputBox;
+	typename SLocalRotationInputBox::FArguments& Args = *(typename SLocalRotationInputBox::FArguments*)ArgumentsPtr; 
+
+	Args.Quaternion_Lambda([this, InPropertyHandle]() -> TOptional<RotationType>
+	{
+		return GetRotation<RotationType>(InPropertyHandle);
+	});
+
+	Args.OnQuaternionChanged_Lambda([this, InPropertyHandle](RotationType InValue)
+	{
+		OnRotationChanged<RotationType>(InPropertyHandle, InValue, false);
+	});
+
+	Args.OnQuaternionCommitted_Lambda([this, InPropertyHandle](RotationType InValue, ETextCommit::Type InCommitType)
+	{
+		OnRotationChanged<RotationType>(InPropertyHandle, InValue, true, InCommitType);
+	});
+}
+
+template<>
+FORCEINLINE void FControlRigGraphMathTypeDetails::ExtendRotationArgs<FRotator>(TSharedRef<class IPropertyHandle> InPropertyHandle, void* ArgumentsPtr)
+{
+	using RotationType = FRotator;
+	typedef typename RotationType::FReal NumericType;
+	typedef SAdvancedRotationInputBox<NumericType> SLocalRotationInputBox;
+	typename SLocalRotationInputBox::FArguments& Args = *(typename SLocalRotationInputBox::FArguments*)ArgumentsPtr; 
+
+	Args.Rotator_Lambda([this, InPropertyHandle]() -> TOptional<RotationType>
+	{
+		return GetRotation<RotationType>(InPropertyHandle);
+	});
+
+	Args.OnRotatorChanged_Lambda([this, InPropertyHandle](RotationType InValue)
+	{
+		OnRotationChanged<RotationType>(InPropertyHandle, InValue, false);
+	});
+
+	Args.OnRotatorCommitted_Lambda([this, InPropertyHandle](RotationType InValue, ETextCommit::Type InCommitType)
+	{
+		OnRotationChanged<RotationType>(InPropertyHandle, InValue, true, InCommitType);
+	});
+}
\ No newline at end of file
diff --git a/Engine/Plugins/Animation/ControlRig/Source/ControlRigEditor/Private/Editor/ControlRigEditor.cpp b/Engine/Plugins/Animation/ControlRig/Source/ControlRigEditor/Private/Editor/ControlRigEditor.cpp
--- a/Engine/Plugins/Animation/ControlRig/Source/ControlRigEditor/Private/Editor/ControlRigEditor.cpp
+++ b/Engine/Plugins/Animation/ControlRig/Source/ControlRigEditor/Private/Editor/ControlRigEditor.cpp
@@ -547,6 +547,25 @@ void FControlRigEditor::InitControlRigEditor(const EToolkitMode::Type Mode, cons
 	FString BlueprintName = InControlRigBlueprint->GetPathName();
 	RigVMPythonUtils::PrintPythonContext(BlueprintName);
 #endif
+
+	TArray<UScriptStruct*> StructsToCustomize = {
+		TBaseStructure<FVector>::Get(),
+		TBaseStructure<FVector2D>::Get(),
+		TBaseStructure<FVector4>::Get(),
+		TBaseStructure<FRotator>::Get(),
+		TBaseStructure<FQuat>::Get(),
+		TBaseStructure<FTransform>::Get(),
+		TBaseStructure<FEulerTransform>::Get(),
+	};
+
+	for(UScriptStruct* StructToCustomize : StructsToCustomize)
+	{
+		Inspector->GetPropertyView()->RegisterInstancedCustomPropertyTypeLayout(StructToCustomize->GetFName(),
+			FOnGetPropertyTypeCustomizationInstance::CreateLambda([=]()
+			{
+				return FControlRigGraphMathTypeDetails::MakeInstance();
+			}));
+	}
 }
 
 void FControlRigEditor::BindCommands()
@@ -1423,41 +1442,18 @@ void FControlRigEditor::SetDetailObjects(const TArray<UObject*>& InObjects, bool
 
 	TArray<UObject*> FilteredObjects;
 
+	TArray<URigVMNode*> ModelNodes;
 	for(UObject* InObject : InObjects)
 	{
-		if (URigVMUnitNode* UnitNode = Cast<URigVMUnitNode>(InObject))
+		if (URigVMNode* ModelNode = Cast<URigVMNode>(InObject))
 		{
-			UScriptStruct* NodeDetailStruct = UnitNode->GetScriptStruct();
-			const FName NodeDetailName = UnitNode->GetFName();
-
-			if (NodeDetailStruct)
-			{
-				// create a struct matching this node
-				TSharedPtr<FStructOnScope> StructInstance  = UnitNode->ConstructStructInstance(false);
-
-				// check if we know the dynamic class already
-				bool bClassCreated = UDetailsViewWrapperObject::GetClassForStruct(NodeDetailStruct, false) == nullptr;
-
-				// create the wrapper object
-				UDetailsViewWrapperObject* WrapperObject = UDetailsViewWrapperObject::MakeInstance(
-					NodeDetailStruct, StructInstance->GetStructMemory(), UnitNode);
-				WrapperObject->GetWrappedPropertyChangedChainEvent().AddSP(this, &FControlRigEditor::OnWrappedPropertyChangedChainEvent);
-				WrapperObject->AddToRoot();
-
-				if(bClassCreated)
-				{
-					UClass* WrapperClass = UDetailsViewWrapperObject::GetClassForStruct(NodeDetailStruct, false);
-					check(WrapperClass);
-
-					FPropertyEditorModule& EditModule = FModuleManager::Get().GetModuleChecked<FPropertyEditorModule>("PropertyEditor");
-					EditModule.RegisterCustomClassLayout(WrapperClass->GetFName(), FOnGetDetailCustomizationInstance::CreateStatic(&FControlRigWrappedNodeDetails::MakeInstance));
-				}
-				
-				FilteredObjects.Add(WrapperObject);
-				continue;
-			}
+			ModelNodes.Add(ModelNode);
 		}
-		else if (URigVMLibraryNode* LibraryNode = Cast<URigVMLibraryNode>(InObject))
+	}
+	
+	for(UObject* InObject : InObjects)
+	{
+		if (URigVMLibraryNode* LibraryNode = Cast<URigVMLibraryNode>(InObject))
 		{
 			UEdGraph* EdGraph = GetControlRigBlueprint()->GetEdGraph(LibraryNode->GetContainedGraph());
 			if (EdGraph)
@@ -1477,14 +1473,33 @@ void FControlRigEditor::SetDetailObjects(const TArray<UObject*>& InObjects, bool
 		}
 		else if (URigVMNode* ModelNode = Cast<URigVMNode>(InObject))
 		{
-			UControlRigGraph* RigGraph = Cast<UControlRigGraph>(GetFocusedGraph());
-			if (UEdGraphNode* EdNode = RigGraph->FindNodeForModelNodeName(ModelNode->GetFName()))
+			// check if we know the dynamic class already
+			bool bClassCreated = UDetailsViewWrapperObject::GetClassForNodes(ModelNodes, false) == nullptr;
+
+			// create the wrapper object
+			if(UDetailsViewWrapperObject* WrapperObject = UDetailsViewWrapperObject::MakeInstance(ModelNodes, ModelNode, ModelNode))
 			{
-				FilteredObjects.AddUnique(EdNode);
+				WrapperObject->GetWrappedPropertyChangedChainEvent().AddSP(this, &FControlRigEditor::OnWrappedPropertyChangedChainEvent);
+				WrapperObject->AddToRoot();
+
+				if(bClassCreated)
+				{
+					UClass* WrapperClass = UDetailsViewWrapperObject::GetClassForNodes(ModelNodes, false);
+					check(WrapperClass);
+
+					FPropertyEditorModule& EditModule = FModuleManager::Get().GetModuleChecked<FPropertyEditorModule>("PropertyEditor");
+					EditModule.RegisterCustomClassLayout(WrapperClass->GetFName(), FOnGetDetailCustomizationInstance::CreateStatic(&FControlRigWrappedNodeDetails::MakeInstance));
+				}
+
+				// todo: use transform widget for transforms
+				// todo: use rotation widget for rotations
+				
+				FilteredObjects.Add(WrapperObject);
 				continue;
 			}
 		}
 
+
 		FilteredObjects.Add(InObject);
 	}
 
@@ -4305,150 +4320,205 @@ void FControlRigEditor::OnWrappedPropertyChangedChainEvent(UDetailsViewWrapperOb
 	UControlRigBlueprint* ControlRigBP = GetControlRigBlueprint();
 
 	FString PropertyPath = InPropertyPath;
-	if(InWrapperObject->GetWrappedStruct()->IsChildOf(FRigBaseElement::StaticStruct()))
+	if(InWrapperObject->GetWrappedStruct())
 	{
-		check(InWrapperObject->GetWrappedStruct() == WrapperObjects[0]->GetWrappedStruct());
-
-		URigHierarchy* Hierarchy = CastChecked<URigHierarchy>(InWrapperObject->GetOuter());
-		const FRigBaseElement WrappedElement = InWrapperObject->GetContent<FRigBaseElement>();
-		const FRigBaseElement FirstWrappedElement = WrapperObjects[0]->GetContent<FRigBaseElement>();
-		const FRigElementKey& Key = WrappedElement.GetKey();
-		if(!Hierarchy->Contains(Key))
+		if(InWrapperObject->GetWrappedStruct()->IsChildOf(FRigBaseElement::StaticStruct()))
 		{
-			return;
-		}
+			check(InWrapperObject->GetWrappedStruct() == WrapperObjects[0]->GetWrappedStruct());
 
-		const FString PoseString = TEXT("Pose->");
-		const FString OffsetString = TEXT("Offset->");
-		const FString ShapeString = TEXT("Shape->");
-		const FString SettingsString = TEXT("Settings->");
+			URigHierarchy* Hierarchy = CastChecked<URigHierarchy>(InWrapperObject->GetOuter());
+			const FRigBaseElement WrappedElement = InWrapperObject->GetContent<FRigBaseElement>();
+			const FRigBaseElement FirstWrappedElement = WrapperObjects[0]->GetContent<FRigBaseElement>();
+			const FRigElementKey& Key = WrappedElement.GetKey();
+			if(!Hierarchy->Contains(Key))
+			{
+				return;
+			}
 
-		struct Local
-		{
-			static ERigTransformType::Type GetTransformTypeFromPath(FString& PropertyPath)
+			const FString PoseString = TEXT("Pose->");
+			const FString OffsetString = TEXT("Offset->");
+			const FString ShapeString = TEXT("Shape->");
+			const FString SettingsString = TEXT("Settings->");
+
+			struct Local
 			{
-				const FString InitialString = TEXT("Initial->");
-				const FString CurrentString = TEXT("Current->");
-				const FString GlobalString = TEXT("Global->");
-				const FString LocalString = TEXT("Local->");
+				static ERigTransformType::Type GetTransformTypeFromPath(FString& PropertyPath)
+				{
+					const FString InitialString = TEXT("Initial->");
+					const FString CurrentString = TEXT("Current->");
+					const FString GlobalString = TEXT("Global->");
+					const FString LocalString = TEXT("Local->");
 
-				ERigTransformType::Type TransformType = ERigTransformType::CurrentLocal;
+					ERigTransformType::Type TransformType = ERigTransformType::CurrentLocal;
 
-				if(PropertyPath.RemoveFromStart(InitialString))
+					if(PropertyPath.RemoveFromStart(InitialString))
+					{
+						TransformType = ERigTransformType::MakeInitial(TransformType);
+					}
+					else
+					{
+						check(PropertyPath.RemoveFromStart(CurrentString));
+						TransformType = ERigTransformType::MakeCurrent(TransformType);
+					}
+
+					if(PropertyPath.RemoveFromStart(GlobalString))
+					{
+						TransformType = ERigTransformType::MakeGlobal(TransformType);
+					}
+					else
+					{
+						check(PropertyPath.RemoveFromStart(LocalString));
+						TransformType = ERigTransformType::MakeLocal(TransformType);
+					}
+
+					return TransformType;
+				}
+			};
+
+			bool bIsInitial = false;
+			if(PropertyPath.RemoveFromStart(PoseString))
+			{
+				const ERigTransformType::Type TransformType = Local::GetTransformTypeFromPath(PropertyPath);
+				bIsInitial = bIsInitial || ERigTransformType::IsInitial(TransformType);
+				
+				if(ERigTransformType::IsInitial(TransformType) || bSetupModeEnabled)
 				{
-					TransformType = ERigTransformType::MakeInitial(TransformType);
+					Hierarchy = ControlRigBP->Hierarchy;
 				}
-				else
+
+				FRigTransformElement* TransformElement = Hierarchy->Find<FRigTransformElement>(WrappedElement.GetKey());
+				if(TransformElement == nullptr)
 				{
-					check(PropertyPath.RemoveFromStart(CurrentString));
-					TransformType = ERigTransformType::MakeCurrent(TransformType);
+					return;
 				}
 
-				if(PropertyPath.RemoveFromStart(GlobalString))
+				const FTransform Transform = InWrapperObject->GetContent<FRigTransformElement>().Pose.Get(TransformType);
+
+				if(ERigTransformType::IsLocal(TransformType) && TransformElement->IsA<FRigControlElement>())
 				{
-					TransformType = ERigTransformType::MakeGlobal(TransformType);
+					FRigControlElement* ControlElement = Cast<FRigControlElement>(TransformElement);
+							
+					FRigControlValue Value;
+					Value.SetFromTransform(Transform, ControlElement->Settings.ControlType, ControlElement->Settings.PrimaryAxis);
+							
+					if(ERigTransformType::IsInitial(TransformType) || bSetupModeEnabled)
+					{
+						Hierarchy->SetControlValue(ControlElement, Value, ERigControlValueType::Initial, true, true, true);
+					}
+					Hierarchy->SetControlValue(ControlElement, Value, ERigControlValueType::Current, true, true, true);
 				}
 				else
 				{
-					check(PropertyPath.RemoveFromStart(LocalString));
-					TransformType = ERigTransformType::MakeLocal(TransformType);
+					Hierarchy->SetTransform(TransformElement, Transform, TransformType, true, true, false, true);
 				}
-
-				return TransformType;
 			}
-		};
-
-		bool bIsInitial = false;
-		if(PropertyPath.RemoveFromStart(PoseString))
-		{
-			const ERigTransformType::Type TransformType = Local::GetTransformTypeFromPath(PropertyPath);
-			bIsInitial = bIsInitial || ERigTransformType::IsInitial(TransformType);
-			
-			if(ERigTransformType::IsInitial(TransformType) || bSetupModeEnabled)
+			else if(PropertyPath.RemoveFromStart(OffsetString))
 			{
-				Hierarchy = ControlRigBP->Hierarchy;
-			}
+				FRigControlElement* ControlElement = ControlRigBP->Hierarchy->Find<FRigControlElement>(WrappedElement.GetKey());
+				if(ControlElement == nullptr)
+				{
+					return;
+				}
+				
+				ERigTransformType::Type TransformType = Local::GetTransformTypeFromPath(PropertyPath);
+				bIsInitial = bIsInitial || ERigTransformType::IsInitial(TransformType);
 
-			FRigTransformElement* TransformElement = Hierarchy->Find<FRigTransformElement>(WrappedElement.GetKey());
-			if(TransformElement == nullptr)
-			{
-				return;
+				const FTransform Transform = WrapperObjects[0]->GetContent<FRigControlElement>().Offset.Get(TransformType);
+				
+				ControlRigBP->Hierarchy->SetControlOffsetTransform(ControlElement, Transform, ERigTransformType::MakeInitial(TransformType), true, true, false, true);
 			}
+			else if(PropertyPath.RemoveFromStart(ShapeString))
+			{
+				FRigControlElement* ControlElement = ControlRigBP->Hierarchy->Find<FRigControlElement>(WrappedElement.GetKey());
+				if(ControlElement == nullptr)
+				{
+					return;
+				}
 
-			const FTransform Transform = InWrapperObject->GetContent<FRigTransformElement>().Pose.Get(TransformType);
+				ERigTransformType::Type TransformType = Local::GetTransformTypeFromPath(PropertyPath);
+				bIsInitial = bIsInitial || ERigTransformType::IsInitial(TransformType);
 
-			if(ERigTransformType::IsLocal(TransformType) && TransformElement->IsA<FRigControlElement>())
+				const FTransform Transform = WrapperObjects[0]->GetContent<FRigControlElement>().Shape.Get(TransformType);
+				
+				ControlRigBP->Hierarchy->SetControlShapeTransform(ControlElement, Transform, ERigTransformType::MakeInitial(TransformType), true, false, true);
+			}
+			else if(PropertyPath.RemoveFromStart(SettingsString))
 			{
-				FRigControlElement* ControlElement = Cast<FRigControlElement>(TransformElement);
-						
-				FRigControlValue Value;
-				Value.SetFromTransform(Transform, ControlElement->Settings.ControlType, ControlElement->Settings.PrimaryAxis);
-						
-				if(ERigTransformType::IsInitial(TransformType) || bSetupModeEnabled)
+				const FRigControlSettings Settings  = InWrapperObject->GetContent<FRigControlElement>().Settings;
+
+				FRigControlElement* ControlElement = ControlRigBP->Hierarchy->Find<FRigControlElement>(WrappedElement.GetKey());
+				if(ControlElement == nullptr)
 				{
-					Hierarchy->SetControlValue(ControlElement, Value, ERigControlValueType::Initial, true, true, true);
+					return;
 				}
-				Hierarchy->SetControlValue(ControlElement, Value, ERigControlValueType::Current, true, true, true);
+
+				ControlRigBP->Hierarchy->SetControlSettings(ControlElement, Settings, true, false, true);
 			}
-			else
+
+			if(bSetupModeEnabled || bIsInitial)
 			{
-				Hierarchy->SetTransform(TransformElement, Transform, TransformType, true, true, false, true);
+				ControlRigBP->PropagatePoseFromBPToInstances();
+				ControlRigBP->Modify();
+				ControlRigBP->MarkPackageDirty();
 			}
 		}
-		else if(PropertyPath.RemoveFromStart(OffsetString))
+		else if(InWrapperObject->GetWrappedStruct()->IsChildOf(FRigVMGraphVariableDescription::StaticStruct()))
 		{
-			FRigControlElement* ControlElement = ControlRigBP->Hierarchy->Find<FRigControlElement>(WrappedElement.GetKey());
-			if(ControlElement == nullptr)
-			{
-				return;
-			}
+			check(InWrapperObject->GetWrappedStruct() == WrapperObjects[0]->GetWrappedStruct());
 			
-			ERigTransformType::Type TransformType = Local::GetTransformTypeFromPath(PropertyPath);
-			bIsInitial = bIsInitial || ERigTransformType::IsInitial(TransformType);
-
-			const FTransform Transform = WrapperObjects[0]->GetContent<FRigControlElement>().Offset.Get(TransformType);
-			
-			ControlRigBP->Hierarchy->SetControlOffsetTransform(ControlElement, Transform, ERigTransformType::MakeInitial(TransformType), true, true, false, true);
-		}
-		else if(PropertyPath.RemoveFromStart(ShapeString))
-		{
-			FRigControlElement* ControlElement = ControlRigBP->Hierarchy->Find<FRigControlElement>(WrappedElement.GetKey());
-			if(ControlElement == nullptr)
+			const FRigVMGraphVariableDescription VariableDescription = InWrapperObject->GetContent<FRigVMGraphVariableDescription>();
+			URigVMGraph* Graph = CastChecked<URigVMGraph>(InWrapperObject->GetOuter());
+			URigVMController* Controller = ControlRigBP->GetController(Graph);
+			if (PropertyPath == TEXT("Name") && MyBlueprintWidget.IsValid())
 			{
-				return;
+				if (FEdGraphSchemaAction_BlueprintVariableBase* VariableAcion = MyBlueprintWidget->SelectionAsBlueprintVariable())
+				{
+					const FName OldVariableName = VariableAcion->GetVariableName();
+					if (!OldVariableName.IsNone())
+					{
+						for (FRigVMGraphVariableDescription& Variable : Graph->GetLocalVariables())
+						{
+							if (Variable.Name == OldVariableName)
+							{
+								Controller->RenameLocalVariable(OldVariableName, VariableDescription.Name);
+								break;
+							}
+						}
+					}
+				}
+				RefreshMyBlueprint();
+				GetControlRigBlueprint()->RequestAutoVMRecompilation();
 			}
-
-			ERigTransformType::Type TransformType = Local::GetTransformTypeFromPath(PropertyPath);
-			bIsInitial = bIsInitial || ERigTransformType::IsInitial(TransformType);
-
-			const FTransform Transform = WrapperObjects[0]->GetContent<FRigControlElement>().Shape.Get(TransformType);
-			
-			ControlRigBP->Hierarchy->SetControlShapeTransform(ControlElement, Transform, ERigTransformType::MakeInitial(TransformType), true, false, true);
-		}
-		else if(PropertyPath.RemoveFromStart(SettingsString))
-		{
-			const FRigControlSettings Settings  = InWrapperObject->GetContent<FRigControlElement>().Settings;
-
-			FRigControlElement* ControlElement = ControlRigBP->Hierarchy->Find<FRigControlElement>(WrappedElement.GetKey());
-			if(ControlElement == nullptr)
-			{
-				return;
+			else if (PropertyPath == TEXT("CPPType") || PropertyPath == TEXT("CPPTypeObject"))
+			{			
+				for (FRigVMGraphVariableDescription& Variable : Graph->GetLocalVariables())
+				{
+					if (Variable.Name == VariableDescription.Name)
+					{
+						Controller->SetLocalVariableType(Variable.Name, VariableDescription.CPPType, VariableDescription.CPPTypeObject);
+						break;
+					}
+				}
+				GetControlRigBlueprint()->RequestAutoVMRecompilation();
 			}
+			else if (PropertyPath == TEXT("DefaultValue"))
+			{			
+				for (FRigVMGraphVariableDescription& Variable : Graph->GetLocalVariables())
+				{
+					if (Variable.Name == VariableDescription.Name)
+					{
+						Controller->SetLocalVariableDefaultValue(Variable.Name, VariableDescription.DefaultValue, true, true, false);
+						break;
+					}
+				}
 
-			ControlRigBP->Hierarchy->SetControlSettings(ControlElement, Settings, true, false, true);
-		}
-
-		if(bSetupModeEnabled || bIsInitial)
-		{
-			ControlRigBP->PropagatePoseFromBPToInstances();
-			ControlRigBP->Modify();
-			ControlRigBP->MarkPackageDirty();
+				// Do not recompile now! That destroys the object that is currently being displayed (the literal memory storage), and can cause a crash.
+				// The user has to manually trigger the recompilation.
+			}		
 		}
 	}
-	else if(InWrapperObject->GetWrappedStruct()->IsChildOf(FRigUnit::StaticStruct()))
+	else if(!InWrapperObject->GetWrappedNodeNotation().IsEmpty())
 	{
-		check(InWrapperObject->GetWrappedStruct() == WrapperObjects[0]->GetWrappedStruct());
-		
 		FName RootPinName = InPropertyChangedChainEvent.PropertyChain.GetHead()->GetValue()->GetFName();
 		FProperty* TargetProperty = WrapperObjects[0]->GetClass()->FindPropertyByName(RootPinName);
 		uint8* FirstPropertyStorage = TargetProperty->ContainerPtrToValuePtr<uint8>(WrapperObjects[0].Get());
@@ -4456,67 +4526,15 @@ void FControlRigEditor::OnWrappedPropertyChangedChainEvent(UDetailsViewWrapperOb
 		URigVMNode* Node = CastChecked<URigVMNode>(InWrapperObject->GetOuter());
 
 		FString DefaultValue = FRigVMStruct::ExportToFullyQualifiedText(TargetProperty, FirstPropertyStorage);
+		URigVMController* Controller = GetControlRigBlueprint()->GetController(Node->GetGraph());
 
 		if (!DefaultValue.IsEmpty())
 		{
 			FString PinPath = FString::Printf(TEXT("%s.%s"), *Node->GetName(), *RootPinName.ToString());
-			GetFocusedController()->SetPinDefaultValue(PinPath, DefaultValue, true, true, false, true);
+			const bool bInteractive = InPropertyChangedChainEvent.ChangeType == EPropertyChangeType::Interactive;
+			Controller->SetPinDefaultValue(PinPath, DefaultValue, true, !bInteractive, true, !bInteractive);
 		}
 	}
-	else if(InWrapperObject->GetWrappedStruct()->IsChildOf(FRigVMGraphVariableDescription::StaticStruct()))
-	{
-		check(InWrapperObject->GetWrappedStruct() == WrapperObjects[0]->GetWrappedStruct());
-		
-		const FRigVMGraphVariableDescription VariableDescription = InWrapperObject->GetContent<FRigVMGraphVariableDescription>();
-		URigVMGraph* Graph = CastChecked<URigVMGraph>(InWrapperObject->GetOuter());
-		URigVMController* Controller = ControlRigBP->GetController(Graph);
-		if (PropertyPath == TEXT("Name") && MyBlueprintWidget.IsValid())
-		{
-			if (FEdGraphSchemaAction_BlueprintVariableBase* VariableAcion = MyBlueprintWidget->SelectionAsBlueprintVariable())
-			{
-				const FName OldVariableName = VariableAcion->GetVariableName();
-				if (!OldVariableName.IsNone())
-				{
-					for (FRigVMGraphVariableDescription& Variable : Graph->GetLocalVariables())
-					{
-						if (Variable.Name == OldVariableName)
-						{
-							Controller->RenameLocalVariable(OldVariableName, VariableDescription.Name);
-							break;
-						}
-					}
-				}
-			}
-			RefreshMyBlueprint();
-			GetControlRigBlueprint()->RequestAutoVMRecompilation();
-		}
-		else if (PropertyPath == TEXT("CPPType") || PropertyPath == TEXT("CPPTypeObject"))
-		{			
-			for (FRigVMGraphVariableDescription& Variable : Graph->GetLocalVariables())
-			{
-				if (Variable.Name == VariableDescription.Name)
-				{
-					Controller->SetLocalVariableType(Variable.Name, VariableDescription.CPPType, VariableDescription.CPPTypeObject);
-					break;
-				}
-			}
-			GetControlRigBlueprint()->RequestAutoVMRecompilation();
-		}
-		else if (PropertyPath == TEXT("DefaultValue"))
-		{			
-			for (FRigVMGraphVariableDescription& Variable : Graph->GetLocalVariables())
-			{
-				if (Variable.Name == VariableDescription.Name)
-				{
-					Controller->SetLocalVariableDefaultValue(Variable.Name, VariableDescription.DefaultValue, true, true, false);
-					break;
-				}
-			}
-
-			// Do not recompile now! That destroys the object that is currently being displayed (the literal memory storage), and can cause a crash.
-			// The user has to manually trigger the recompilation.
-		}		
-	}
 }
 
 void FControlRigEditor::OnRequestLocalizeFunctionDialog(URigVMLibraryNode* InFunction, UControlRigBlueprint* InTargetBlueprint, bool bForce)
diff --git a/Engine/Plugins/Animation/ControlRig/Source/ControlRigEditor/Private/Editor/DetailsViewWrapperObject.cpp b/Engine/Plugins/Animation/ControlRig/Source/ControlRigEditor/Private/Editor/DetailsViewWrapperObject.cpp
--- a/Engine/Plugins/Animation/ControlRig/Source/ControlRigEditor/Private/Editor/DetailsViewWrapperObject.cpp
+++ b/Engine/Plugins/Animation/ControlRig/Source/ControlRigEditor/Private/Editor/DetailsViewWrapperObject.cpp
@@ -5,20 +5,26 @@
 #include "ControlRigElementDetails.h"
 #include "ControlRigLocalVariableDetails.h"
 #include "Modules/ModuleManager.h"
+#include "ControlRig.h"
+#include "Algo/Sort.h"
 
 #if WITH_EDITOR
 #include "PropertyEditorModule.h"
 #endif
 
-TMap<UScriptStruct*, UClass*> UDetailsViewWrapperObject::StructToClass;
-TMap<UClass*, UScriptStruct*> UDetailsViewWrapperObject::ClassToStruct;
+TMap<UDetailsViewWrapperObject::FPerClassInfo, UClass*> UDetailsViewWrapperObject::InfoToClass;
+TMap<UClass*, UDetailsViewWrapperObject::FPerClassInfo> UDetailsViewWrapperObject::ClassToInfo;
+
+UDetailsViewWrapperObject::UDetailsViewWrapperObject()
+	: bIsSettingValue(false)
+{
+}
 
 UClass* UDetailsViewWrapperObject::GetClassForStruct(UScriptStruct* InStruct, bool bCreateIfNeeded)
 {
 	check(InStruct != nullptr);
 
-	UClass** ExistingClass = StructToClass.Find(InStruct);
-	if(ExistingClass)
+	if(UClass** ExistingClass = InfoToClass.Find(InStruct))
 	{
 		return *ExistingClass;
 	}
@@ -109,15 +115,7 @@ UClass* UDetailsViewWrapperObject::GetClassForStruct(UScriptStruct* InStruct, bo
 	}
 
 #if WITH_EDITOR
-	if(InStruct->IsChildOf(FRigUnit::StaticStruct()))
-	{
-		FPropertyEditorModule& PropertyEditorModule = FModuleManager::LoadModuleChecked<FPropertyEditorModule>("PropertyEditor");
-		if (!PropertyEditorModule.GetClassNameToDetailLayoutNameMap().Contains(WrapperClassName))
-		{
-			PropertyEditorModule.RegisterCustomClassLayout(WrapperClassName, FOnGetDetailCustomizationInstance::CreateStatic(&FRigUnitDetails::MakeInstance));
-		}
-	}
-	else if(InStruct->IsChildOf(FRigBaseElement::StaticStruct()))
+	if(InStruct->IsChildOf(FRigBaseElement::StaticStruct()))
 	{
 		FPropertyEditorModule& PropertyEditorModule = FModuleManager::LoadModuleChecked<FPropertyEditorModule>("PropertyEditor");
 		if (!PropertyEditorModule.GetClassNameToDetailLayoutNameMap().Contains(WrapperClassName))
@@ -153,8 +151,8 @@ UClass* UDetailsViewWrapperObject::GetClassForStruct(UScriptStruct* InStruct, bo
 	// Similar to FConfigPropertyHelperDetails::CustomizeDetails, this is required for GC to work properly
 	WrapperClass->AssembleReferenceTokenStream();
 
-	StructToClass.Add(InStruct, WrapperClass);
-	ClassToStruct.Add(WrapperClass, InStruct);
+	InfoToClass.Add(InStruct, WrapperClass);
+	ClassToInfo.Add(WrapperClass, InStruct);
 
 	UObject* CDO = WrapperClass->GetDefaultObject(true);
 	CDO->AddToRoot();
@@ -185,7 +183,292 @@ UDetailsViewWrapperObject* UDetailsViewWrapperObject::MakeInstance(UScriptStruct
 
 UScriptStruct* UDetailsViewWrapperObject::GetWrappedStruct() const
 {
-	return ClassToStruct.FindChecked(GetClass());
+	return ClassToInfo.FindChecked(GetClass()).ScriptStruct;
+}
+
+UClass* UDetailsViewWrapperObject::GetClassForNodes(TArray<URigVMNode*> InNodes, bool bCreateIfNeeded)
+{
+	if(InNodes.IsEmpty())
+	{
+		return nullptr;
+	}
+
+	static TArray<ERigVMPinDirection> PinDirectionsToCheck = {
+		ERigVMPinDirection::IO,
+		ERigVMPinDirection::Input,
+		ERigVMPinDirection::Visible
+	};
+
+	// determine if all nodes are unit nodes and match their script struct
+	TArray<UScriptStruct*> UnitStructs;
+	for(URigVMNode* Node : InNodes)
+	{
+		if(URigVMUnitNode* UnitNode = Cast<URigVMUnitNode>(Node))
+		{
+			UnitStructs.AddUnique(UnitNode->GetScriptStruct());
+		}
+	}
+
+	// find the matching pins on all passed in nodes
+	TArray<URigVMPin*> PinsToInspect;
+
+	for(const ERigVMPinDirection PinDirectionToCheck : PinDirectionsToCheck)
+	{
+		for(URigVMPin* Pin : InNodes[0]->GetPins())
+		{
+			if(Pin->IsExecuteContext())
+			{
+				continue;
+			}
+			
+			if(Pin->GetDirection() != PinDirectionToCheck)
+			{
+				continue;
+			}
+			
+			bool bInspectPin = true;
+			
+			for(int32 NodeIndex=1; NodeIndex<InNodes.Num(); NodeIndex++)
+			{
+				URigVMPin* OtherPin = InNodes[NodeIndex]->FindPin(Pin->GetName());
+				if(OtherPin == nullptr)
+				{
+					bInspectPin = false;
+					break;
+				}
+
+				if(Pin->GetCPPType() != OtherPin->GetCPPType())
+				{
+					bInspectPin = false;
+					break;
+				}
+			}
+
+			if(bInspectPin)
+			{
+				PinsToInspect.Add(Pin);
+			}
+		}
+	}
+
+	TArray<FString> PinNotations;
+	for(URigVMPin* Pin : PinsToInspect)
+	{
+		PinNotations.Add(FString::Printf(TEXT("%s %s"), *Pin->GetCPPType(), *Pin->GetName()));
+	}
+
+	FString Notation = FString::Join(PinNotations, TEXT(", "));
+	if(UnitStructs.Num() > 0)
+	{
+		// sort the structs to ensure we get the same notation each time
+		Algo::Sort(UnitStructs, [](UScriptStruct* A, UScriptStruct* B) -> bool
+		{
+			return A->GetStructCPPName() > B->GetStructCPPName();
+		});
+
+		TArray<FString> StructNames;
+		for(UScriptStruct* UnitStruct : UnitStructs)
+		{
+			StructNames.Add(UnitStruct->GetStructCPPName());
+		}
+		Notation = FString::Printf(TEXT("%s(%s)"), *FString::Join(StructNames, TEXT("|")), *Notation);
+	}
+
+	const FPerClassInfo PerClassInfo(Notation);
+	if(UClass** ExistingClass = InfoToClass.Find(PerClassInfo))
+	{
+		return *ExistingClass;
+	}
+
+	if(!bCreateIfNeeded)
+	{
+		return nullptr;
+	}
+
+	UClass* SuperClass = UDetailsViewWrapperObject::StaticClass();
+	const int32 HashForNotation = (int32)GetTypeHash(PerClassInfo);
+	const FName WrapperClassName(FString::Printf(TEXT("RigNode%d_WrapperObject"), HashForNotation));
+
+	UClass* WrapperClass = NewObject<UClass>(
+		GetTransientPackage(),
+		WrapperClassName,
+		RF_Public | RF_Transient
+		 );
+
+	// make sure this doesn't get garbage collected
+	WrapperClass->AddToRoot();
+
+	// Eviscerate the class.
+	WrapperClass->PurgeClass(false);
+	WrapperClass->PropertyLink = SuperClass->PropertyLink;
+
+	WrapperClass->SetSuperStruct(SuperClass);
+	WrapperClass->ClassWithin = UObject::StaticClass();
+	WrapperClass->ClassConfigName = SuperClass->ClassConfigName;
+	WrapperClass->ClassFlags |= CLASS_NotPlaceable | CLASS_Hidden;
+
+	WrapperClass->SetMetaData(TEXT("DisplayName"), TEXT("Rig Node"));
+
+	// create properties - one for each pin to inspect
+	FField** LinkToProperty = &WrapperClass->ChildProperties;
+
+	for(URigVMPin* Pin : PinsToInspect)
+	{
+		static FString BoolString = TEXT("bool");
+		static FString Int32String = TEXT("int32");
+		static FString IntString = TEXT("int");
+		static FString FloatString = TEXT("float");
+		static FString DoubleString = TEXT("double");
+		static FString StringString = TEXT("FString");
+		static FString NameString = TEXT("FName");
+
+		FProperty* Property = nullptr;
+		FProperty** ElementProperty = &Property;
+		FFieldVariant PropertyOwner = WrapperClass;
+
+		const FString CPPType = Pin->GetCPPType();
+		FString BaseCPPType = CPPType;
+		if(RigVMTypeUtils::IsArrayType(CPPType))
+		{
+			BaseCPPType = RigVMTypeUtils::BaseTypeFromArrayType(CPPType);
+			
+			// create an array property as a container for the tail
+			FArrayProperty* ArrayProperty = new FArrayProperty(PropertyOwner, Pin->GetFName(), RF_Public);
+			Property = ArrayProperty;
+			ElementProperty = &ArrayProperty->Inner;
+			PropertyOwner = ArrayProperty;
+		}
+
+		if(BaseCPPType.Equals(BoolString, ESearchCase::IgnoreCase))
+		{
+			(*ElementProperty) = new FBoolProperty(PropertyOwner, Pin->GetFName(), RF_Public);;
+		}
+		else if(BaseCPPType.Equals(Int32String, ESearchCase::IgnoreCase) ||
+			BaseCPPType.Equals(IntString, ESearchCase::IgnoreCase))
+		{
+			(*ElementProperty) = new FIntProperty(PropertyOwner, Pin->GetFName(), RF_Public);;
+		}
+		else if(BaseCPPType.Equals(FloatString, ESearchCase::IgnoreCase))
+		{
+			(*ElementProperty) = new FFloatProperty(PropertyOwner, Pin->GetFName(), RF_Public);;
+		}
+		else if(BaseCPPType.Equals(DoubleString, ESearchCase::IgnoreCase))
+		{
+			(*ElementProperty) = new FDoubleProperty(PropertyOwner, Pin->GetFName(), RF_Public);;
+		}
+		else if(BaseCPPType.Equals(StringString, ESearchCase::IgnoreCase))
+		{
+			(*ElementProperty) = new FStrProperty(PropertyOwner, Pin->GetFName(), RF_Public);;
+		}
+		else if(BaseCPPType.Equals(NameString, ESearchCase::IgnoreCase))
+		{
+			(*ElementProperty) = new FNameProperty(PropertyOwner, Pin->GetFName(), RF_Public);;
+		}
+		else if(Pin->GetCPPTypeObject())
+		{
+			if(UEnum* Enum = Cast<UEnum>(Pin->GetCPPTypeObject()))
+			{
+				FByteProperty* ByteProperty = new FByteProperty(PropertyOwner, Pin->GetFName(), RF_Public);
+				ByteProperty->Enum = Enum;
+				(*ElementProperty) = ByteProperty;
+			}
+			else if(UClass* Class = Cast<UClass>(Pin->GetCPPTypeObject()))
+			{
+				FObjectProperty* ObjectProperty = new FObjectProperty(PropertyOwner, Pin->GetFName(), RF_Public);
+				ObjectProperty->PropertyClass = Class; 
+				(*ElementProperty) = ObjectProperty;
+			}
+			else if(UScriptStruct* ScriptStruct = Cast<UScriptStruct>(Pin->GetCPPTypeObject()))
+			{
+				FStructProperty* StructProperty = new FStructProperty(PropertyOwner, Pin->GetFName(), RF_Public);
+				StructProperty->Struct = ScriptStruct;
+				(*ElementProperty) = StructProperty;
+			}
+		}
+
+		if(Property == nullptr)
+		{
+			return nullptr;
+		}
+		Property->SetPropertyFlags(CPF_Edit);
+
+		for(UScriptStruct* UnitStruct : UnitStructs)
+		{
+			if(const FProperty* UnitProperty = UnitStruct->FindPropertyByName(Property->GetFName()))
+			{
+				if(Property->SameType(UnitProperty))
+				{
+					FField::CopyMetaData(UnitProperty, Property);
+					break;
+				}
+			}
+		}
+
+		Property->SetMetaData(TEXT("Category"), TEXT("Nodes"));
+
+		*LinkToProperty = Property;
+		LinkToProperty = &(*LinkToProperty)->Next;
+	}
+
+	// Update the class
+	WrapperClass->Bind();
+	WrapperClass->StaticLink(true);
+	
+	// Similar to FConfigPropertyHelperDetails::CustomizeDetails, this is required for GC to work properly
+	WrapperClass->AssembleReferenceTokenStream();
+
+	InfoToClass.Add(PerClassInfo, WrapperClass);
+	ClassToInfo.Add(WrapperClass, PerClassInfo);
+
+	UDetailsViewWrapperObject* CDO = Cast<UDetailsViewWrapperObject>(WrapperClass->GetDefaultObject(true));
+	CDO->AddToRoot();
+
+	// setup the CDO's defaults to match the units
+	if(UnitStructs.Num() > 0)
+	{
+		UScriptStruct* UnitStruct = UnitStructs[0];
+		FStructOnScope DefaultStructScope(UnitStruct);
+		const FRigUnit* DefaultStruct = (const FRigUnit*)DefaultStructScope.GetStructMemory();
+
+		for (TFieldIterator<FProperty> UnitProperty(UnitStruct); UnitProperty; ++UnitProperty)
+		{
+			if(FProperty* Property = WrapperClass->FindPropertyByName(UnitProperty->GetFName()))
+			{
+				if(Property->SameType(*UnitProperty))
+				{
+					const uint8* SourceMemory = UnitProperty->ContainerPtrToValuePtr<uint8>(DefaultStruct);
+					uint8* TargetMemory = Property->ContainerPtrToValuePtr<uint8>(CDO);
+					Property->CopyCompleteValue(TargetMemory, SourceMemory);
+				}
+			}
+		}
+	}
+
+	return WrapperClass;	
+}
+
+UDetailsViewWrapperObject* UDetailsViewWrapperObject::MakeInstance(TArray<URigVMNode*> InNodes, URigVMNode* InSubject, UObject* InOuter)
+{
+	InOuter = InOuter == nullptr ? GetTransientPackage() : InOuter;
+	
+	UClass* WrapperClass = GetClassForNodes(InNodes);
+	if(WrapperClass == nullptr)
+	{
+		return nullptr;
+	}
+
+	UDetailsViewWrapperObject* Instance = NewObject<UDetailsViewWrapperObject>(InOuter, WrapperClass, NAME_None, RF_Public | RF_Transient | RF_TextExportTransient | RF_DuplicateTransient);
+	Instance->SetContent(InSubject);
+
+	InSubject->GetGraph()->OnModified().AddUObject(Instance, &UDetailsViewWrapperObject::HandleModifiedEvent);
+	
+	return Instance;
+
+}
+
+FString UDetailsViewWrapperObject::GetWrappedNodeNotation() const
+{
+	return ClassToInfo.FindChecked(GetClass()).Notation;
 }
 
 void UDetailsViewWrapperObject::SetContent(const uint8* InStructMemory, const UStruct* InStruct)
@@ -198,6 +481,14 @@ void UDetailsViewWrapperObject::GetContent(uint8* OutStructMemory, const UStruct
 	CopyPropertiesForUnrelatedStructs(OutStructMemory, InStruct, (const uint8*)this, GetClass());
 }
 
+void UDetailsViewWrapperObject::SetContent(URigVMNode* InNode)
+{
+	for(URigVMPin* Pin : InNode->GetPins())
+	{
+		SetContentForPin(Pin);
+	}
+}
+
 void UDetailsViewWrapperObject::PostEditChangeChainProperty(FPropertyChangedChainEvent& PropertyChangedEvent)
 {
 	Super::PostEditChangeChainProperty(PropertyChangedEvent);
@@ -220,6 +511,7 @@ void UDetailsViewWrapperObject::PostEditChangeChainProperty(FPropertyChangedChai
 	}
 	while(PropertyNode);
 
+	TGuardValue<bool> ValueSettingGuard(bIsSettingValue, true);
 	WrappedPropertyChangedChainEvent.Broadcast(this, PropertyPath, PropertyChangedEvent);
 }
 
@@ -246,3 +538,64 @@ void UDetailsViewWrapperObject::CopyPropertiesForUnrelatedStructs(uint8* InTarge
 		TargetProperty->CopyCompleteValue(TargetPropertyMemory, SourcePropertyMemory);
 	}
 }
+
+void UDetailsViewWrapperObject::HandleModifiedEvent(ERigVMGraphNotifType InNotifType, URigVMGraph* InGraph,
+	UObject* InSubject)
+{
+	if(InNotifType != ERigVMGraphNotifType::PinDefaultValueChanged)
+	{
+		return;
+	}
+	
+	URigVMPin* Pin = CastChecked<URigVMPin>(InSubject);
+	if(Pin->GetNode() != GetOuter())
+	{
+		return;
+	}
+
+	SetContentForPin(Pin->GetRootPin());
+}
+
+void UDetailsViewWrapperObject::SetContentForPin(URigVMPin* InPin)
+{
+	if(bIsSettingValue)
+	{
+		return;
+	}
+	TGuardValue<bool> SetValueGuard(bIsSettingValue, true);
+	
+	ensure(InPin->IsRootPin());
+
+	const FString DefaultValue = InPin->GetDefaultValue();
+	if(DefaultValue.IsEmpty())
+	{
+		return;
+	}
+
+	class FDetailsViewWrapperObjectImportErrorContext : public FOutputDevice
+	{
+	public:
+
+		int32 NumErrors;
+
+		FDetailsViewWrapperObjectImportErrorContext()
+			: FOutputDevice()
+			, NumErrors(0)
+		{
+		}
+
+		virtual void Serialize(const TCHAR* V, ELogVerbosity::Type Verbosity, const class FName& Category) override
+		{
+			UE_LOG(LogControlRig, Error, TEXT("Error Importing To Hierarchy: %s"), V);
+			NumErrors++;
+		}
+	};
+
+	// it's possible that the pin has been filtered out when creating the class
+	if(const FProperty* Property = GetClass()->FindPropertyByName(InPin->GetFName()))
+	{
+		FDetailsViewWrapperObjectImportErrorContext ErrorPipe;
+		uint8* DataPtr = Property->ContainerPtrToValuePtr<uint8>(this);
+		Property->ImportText(*DefaultValue, DataPtr, EPropertyPortFlags::PPF_None, nullptr, &ErrorPipe);		
+	}
+}
diff --git a/Engine/Plugins/Animation/ControlRig/Source/ControlRigEditor/Private/Editor/DetailsViewWrapperObject.h b/Engine/Plugins/Animation/ControlRig/Source/ControlRigEditor/Private/Editor/DetailsViewWrapperObject.h
--- a/Engine/Plugins/Animation/ControlRig/Source/ControlRigEditor/Private/Editor/DetailsViewWrapperObject.h
+++ b/Engine/Plugins/Animation/ControlRig/Source/ControlRigEditor/Private/Editor/DetailsViewWrapperObject.h
@@ -3,6 +3,8 @@
 #pragma once
 
 #include "CoreMinimal.h"
+#include "RigVMModel/RigVMGraph.h"
+#include "RigVMModel/RigVMNode.h"
 #include "DetailsViewWrapperObject.generated.h"
 
 class UDetailsViewWrapperObject;
@@ -15,10 +17,18 @@ class UDetailsViewWrapperObject : public UObject
 public:
 	GENERATED_BODY()
 
+	UDetailsViewWrapperObject();
+
+	// Creating wrappers from a given struct
 	static UClass* GetClassForStruct(UScriptStruct* InStruct, bool bCreateIfNeeded = true);
 	static UDetailsViewWrapperObject* MakeInstance(UScriptStruct* InStruct, uint8* InStructMemory, UObject* InOuter = nullptr);
 	UScriptStruct* GetWrappedStruct() const;
 
+	// Creating wrappers from a RigVMNode
+	static UClass* GetClassForNodes(TArray<URigVMNode*> InNodes, bool bCreateIfNeeded = true);
+	static UDetailsViewWrapperObject* MakeInstance(TArray<URigVMNode*> InNodes, URigVMNode* InSubject, UObject* InOuter = nullptr);
+	FString GetWrappedNodeNotation() const;
+	
 	bool IsChildOf(const UStruct* InStruct) const
 	{
 		return GetWrappedStruct()->IsChildOf(InStruct);
@@ -32,6 +42,7 @@ public:
 
 	void SetContent(const uint8* InStructMemory, const UStruct* InStruct);
 	void GetContent(uint8* OutStructMemory, const UStruct* InStruct) const;
+	void SetContent(URigVMNode* InNode);
 
 	template<typename T>
 	T GetContent() const
@@ -57,9 +68,48 @@ public:
 private:
 
 	static void CopyPropertiesForUnrelatedStructs(uint8* InTargetMemory, const UStruct* InTargetStruct, const uint8* InSourceMemory, const UStruct* InSourceStruct);
+	void HandleModifiedEvent(ERigVMGraphNotifType InNotifType, URigVMGraph* InGraph, UObject* InSubject);
+	void SetContentForPin(URigVMPin* InPin);
+
+	struct FPerClassInfo
+	{
+		FString Notation;
+		UScriptStruct* ScriptStruct;
+		
+		FPerClassInfo()
+			: Notation()
+			, ScriptStruct(nullptr)
+		{}
+
+		FPerClassInfo(UScriptStruct* InScriptStruct)
+		: Notation()
+		, ScriptStruct(InScriptStruct)
+		{}
+
+		FPerClassInfo(const FString& InNotation)
+		: Notation(InNotation)
+		, ScriptStruct(nullptr)
+		{}
+
+		friend FORCEINLINE uint32 GetTypeHash(const FPerClassInfo& Info)
+		{
+			return HashCombine(GetTypeHash(Info.Notation), GetTypeHash(Info.ScriptStruct));
+		}
+
+		FORCEINLINE bool operator ==(const FPerClassInfo& Other) const
+		{
+			return Notation == Other.Notation && ScriptStruct == Other.ScriptStruct;
+		}
+
+		FORCEINLINE bool operator !=(const FPerClassInfo& Other) const
+		{
+			return Notation != Other.Notation || ScriptStruct != Other.ScriptStruct;
+		}
+	};
 	
-	static TMap<UScriptStruct*, UClass*> StructToClass;
-	static TMap<UClass*, UScriptStruct*> ClassToStruct;
+	static TMap<FPerClassInfo, UClass*> InfoToClass;
+	static TMap<UClass*, FPerClassInfo> ClassToInfo;
+	bool bIsSettingValue;
 	
 	FWrappedPropertyChangedChainEvent WrappedPropertyChangedChainEvent;
 };
diff --git a/Engine/Plugins/Animation/ControlRig/Source/ControlRigEditor/Private/Graph/ControlRigGraphPanelPinFactory.cpp b/Engine/Plugins/Animation/ControlRig/Source/ControlRigEditor/Private/Graph/ControlRigGraphPanelPinFactory.cpp
--- a/Engine/Plugins/Animation/ControlRig/Source/ControlRigEditor/Private/Graph/ControlRigGraphPanelPinFactory.cpp
+++ b/Engine/Plugins/Animation/ControlRig/Source/ControlRigEditor/Private/Graph/ControlRigGraphPanelPinFactory.cpp
@@ -43,7 +43,7 @@ TSharedPtr<SGraphPin> FControlRigGraphPanelPinFactory::CreatePin(UEdGraphPin* In
 					if (UControlRigBlueprint* Blueprint = Cast<UControlRigBlueprint>(RigGraph->GetOuter()))
 					{
 						return SNew(SControlRigGraphPinVariableBinding, InPin)
-							.ModelPin(ModelPin)
+							.ModelPins({ModelPin})
 							.Blueprint(Blueprint);
 					}
 				}
diff --git a/Engine/Plugins/Animation/ControlRig/Source/ControlRigEditor/Private/Graph/SControlRigGraphNode.cpp b/Engine/Plugins/Animation/ControlRig/Source/ControlRigEditor/Private/Graph/SControlRigGraphNode.cpp
--- a/Engine/Plugins/Animation/ControlRig/Source/ControlRigEditor/Private/Graph/SControlRigGraphNode.cpp
+++ b/Engine/Plugins/Animation/ControlRig/Source/ControlRigEditor/Private/Graph/SControlRigGraphNode.cpp
@@ -205,7 +205,7 @@ void SControlRigGraphNode::Construct( const FArguments& InArgs )
 				PinInfo.bHideInputWidget = !bIsStructEditable;
 				PinInfo.bAutoHeight = bIsStructEditable;
 			}
-			else if(PinInfo.bHasChildren)
+			else if(PinInfo.bHasChildren && !ModelPin->IsBoundToVariable())
 			{
 				PinInfo.bHideInputWidget = true;
 			}
diff --git a/Engine/Plugins/Animation/ControlRig/Source/ControlRigEditor/Private/Graph/SControlRigGraphPinVariableBinding.cpp b/Engine/Plugins/Animation/ControlRig/Source/ControlRigEditor/Private/Graph/SControlRigGraphPinVariableBinding.cpp
--- a/Engine/Plugins/Animation/ControlRig/Source/ControlRigEditor/Private/Graph/SControlRigGraphPinVariableBinding.cpp
+++ b/Engine/Plugins/Animation/ControlRig/Source/ControlRigEditor/Private/Graph/SControlRigGraphPinVariableBinding.cpp
@@ -16,9 +16,11 @@
 
 #define LOCTEXT_NAMESPACE "SControlRigGraphPinVariableBinding"
 
+static const FText ControlRigVariableBindingMultipleValues = LOCTEXT("MultipleValues", "Multiple Values");
+
 void SControlRigVariableBinding::Construct(const FArguments& InArgs)
 {
-	this->ModelPin = InArgs._ModelPin;
+	this->ModelPins = InArgs._ModelPins;
 	this->FunctionReferenceNode = InArgs._FunctionReferenceNode;
 	this->InnerVariableName = InArgs._InnerVariableName;
 	this->Blueprint = InArgs._Blueprint;
@@ -56,13 +58,30 @@ void SControlRigVariableBinding::Construct(const FArguments& InArgs)
 	];
 }
 
-FText SControlRigVariableBinding::GetBindingText() const
+FText SControlRigVariableBinding::GetBindingText(URigVMPin* ModelPin) const
 {
 	if (ModelPin)
 	{
 		const FString VariablePath = ModelPin->GetBoundVariablePath();
 		return FText::FromString(VariablePath);
 	}
+	return FText();
+}
+
+FText SControlRigVariableBinding::GetBindingText() const
+{
+	if (ModelPins.Num() > 0)
+	{
+		const FText FirstText = GetBindingText(ModelPins[0]);
+		for(int32 Index = 1; Index < ModelPins.Num(); Index++)
+		{
+			if(!GetBindingText(ModelPins[Index]).EqualTo(FirstText))
+			{
+				return ControlRigVariableBindingMultipleValues;
+			}
+		}
+		return FirstText;
+	}
 	else if(FunctionReferenceNode && !InnerVariableName.IsNone())
 	{
 		const FName BoundVariable = FunctionReferenceNode->GetOuterVariableName(InnerVariableName);
@@ -71,7 +90,6 @@ FText SControlRigVariableBinding::GetBindingText() const
 			return FText::FromName(BoundVariable);
 		}
 	}
-
 	return FText();
 }
 
@@ -88,9 +106,9 @@ FLinearColor SControlRigVariableBinding::GetBindingColor() const
 		const UControlRigGraphSchema* Schema = GetDefault<UControlRigGraphSchema>();
 		FName BoundVariable(NAME_None);
 
-		if(ModelPin)
+		if(ModelPins.Num() > 0)
 		{
-			BoundVariable = *ModelPin->GetBoundVariableName();
+			BoundVariable = *ModelPins[0]->GetBoundVariableName();
 		}
 		else if(FunctionReferenceNode && !InnerVariableName.IsNone())
 		{
@@ -109,9 +127,9 @@ FLinearColor SControlRigVariableBinding::GetBindingColor() const
 			}
 		}
 
-		if (ModelPin)
+		if (ModelPins.Num() > 0)
 		{
-			URigVMGraph* Model = ModelPin->GetGraph();
+			URigVMGraph* Model = ModelPins[0]->GetGraph();
 			if(Model == nullptr)
 			{
 				return  FLinearColor::Red;
@@ -147,9 +165,9 @@ bool SControlRigVariableBinding::OnCanBindProperty(FProperty* InProperty) const
 	if (InProperty)
 	{
 		const FRigVMExternalVariable ExternalVariable = FRigVMExternalVariable::Make(InProperty, nullptr);
-		if(ModelPin)
+		if(ModelPins.Num() > 0)
 		{
-			return ModelPin->CanBeBoundToVariable(ExternalVariable);
+			return ModelPins[0]->CanBeBoundToVariable(ExternalVariable);
 		}
 		else if(FunctionReferenceNode && !InnerVariableName.IsNone())
 		{
@@ -203,9 +221,12 @@ void SControlRigVariableBinding::OnAddBinding(FName InPropertyName, const TArray
 			Parts.Add(ChainElement.Field.GetName());
 		}
 
-		if(ModelPin)
+		if(ModelPins.Num() > 0)
 		{
-			Blueprint->GetController(ModelPin->GetGraph())->BindPinToVariable(ModelPin->GetPinPath(), FString::Join(Parts, TEXT(".")), true /* undo */, true /* python */);
+			for(URigVMPin* ModelPin : ModelPins)
+			{
+				Blueprint->GetController(ModelPin->GetGraph())->BindPinToVariable(ModelPin->GetPinPath(), FString::Join(Parts, TEXT(".")), true /* undo */, true /* python */);
+			}
 		}
 		else if(FunctionReferenceNode && !InnerVariableName.IsNone())
 		{
@@ -224,9 +245,12 @@ void SControlRigVariableBinding::OnRemoveBinding(FName InPropertyName)
 {
 	if (Blueprint)
 	{
-		if(ModelPin)
+		if(ModelPins.Num() > 0)
 		{
-			Blueprint->GetController(ModelPin->GetGraph())->UnbindPinFromVariable(ModelPin->GetPinPath(), true /* undo */, true /* python */);
+			for(URigVMPin* ModelPin : ModelPins)
+			{
+				Blueprint->GetController(ModelPin->GetGraph())->UnbindPinFromVariable(ModelPin->GetPinPath(), true /* undo */, true /* python */);
+			}
 		}
 		else if(FunctionReferenceNode && !InnerVariableName.IsNone())
 		{
@@ -237,12 +261,12 @@ void SControlRigVariableBinding::OnRemoveBinding(FName InPropertyName)
 
 void SControlRigVariableBinding::FillLocalVariableMenu(FMenuBuilder& MenuBuilder)
 {
-	if(ModelPin == nullptr)
+	if(ModelPins.Num() == 0)
 	{
 		return;
 	}
 
-	URigVMGraph* Model = ModelPin->GetGraph();
+	URigVMGraph* Model = ModelPins[0]->GetGraph();
 	if(Model == nullptr)
 	{
 		return;
@@ -258,7 +282,7 @@ void SControlRigVariableBinding::FillLocalVariableMenu(FMenuBuilder& MenuBuilder
 			continue;
 		}
 
-		if(!ModelPin->CanBeBoundToVariable(ExternalVariable))
+		if(!ModelPins[0]->CanBeBoundToVariable(ExternalVariable))
 		{
 			continue;
 		}
@@ -284,7 +308,7 @@ void SControlRigVariableBinding::FillLocalVariableMenu(FMenuBuilder& MenuBuilder
 				continue;
 			}
 
-			if(!ModelPin->CanBeBoundToVariable(ExternalVariable))
+			if(!ModelPins[0]->CanBeBoundToVariable(ExternalVariable))
 			{
 				continue;
 			}
@@ -324,31 +348,34 @@ void SControlRigVariableBinding::FillLocalVariableMenu(FMenuBuilder& MenuBuilder
 
 void SControlRigVariableBinding::HandleBindToLocalVariable(FRigVMGraphVariableDescription InLocalVariable)
 {
-	if((ModelPin == nullptr) || (Blueprint == nullptr))
+	if(ModelPins.IsEmpty() || (Blueprint == nullptr))
 	{
 		return;
 	}
 
-	URigVMGraph* Model = ModelPin->GetGraph();
-	if(Model == nullptr)
+	for(URigVMPin* ModelPin : ModelPins)
 	{
-		return;
-	}
+		URigVMGraph* Model = ModelPin->GetGraph();
+		if(Model == nullptr)
+		{
+			continue;
+		}
 
-	URigVMController* Controller = Blueprint->GetOrCreateController(Model);
-	if(Controller == nullptr)
-	{
-		return;
-	}
+		URigVMController* Controller = Blueprint->GetOrCreateController(Model);
+		if(Controller == nullptr)
+		{
+			continue;
+		}
 
-	Controller->BindPinToVariable(ModelPin->GetPinPath(), InLocalVariable.Name.ToString(), true, true);
+		Controller->BindPinToVariable(ModelPin->GetPinPath(), InLocalVariable.Name.ToString(), true, true);
+	}
 }
 
 /////////////////////////////////////////////////////////////////////////////////////////////////////////////
 
 void SControlRigGraphPinVariableBinding::Construct(const FArguments& InArgs, UEdGraphPin* InGraphPinObj)
 {
-	this->ModelPin = InArgs._ModelPin;
+	this->ModelPins = InArgs._ModelPins;
 	this->Blueprint = InArgs._Blueprint;
 
 	SGraphPin::Construct(SGraphPin::FArguments(), InGraphPinObj);
@@ -358,7 +385,7 @@ TSharedRef<SWidget>	SControlRigGraphPinVariableBinding::GetDefaultValueWidget()
 {
 	return SNew(SControlRigVariableBinding)
 		.Blueprint(Blueprint)
-		.ModelPin(ModelPin);
+		.ModelPins(ModelPins);
 }
 
 #undef LOCTEXT_NAMESPACE
\ No newline at end of file
diff --git a/Engine/Plugins/Animation/ControlRig/Source/ControlRigEditor/Private/Graph/SControlRigGraphPinVariableBinding.h b/Engine/Plugins/Animation/ControlRig/Source/ControlRigEditor/Private/Graph/SControlRigGraphPinVariableBinding.h
--- a/Engine/Plugins/Animation/ControlRig/Source/ControlRigEditor/Private/Graph/SControlRigGraphPinVariableBinding.h
+++ b/Engine/Plugins/Animation/ControlRig/Source/ControlRigEditor/Private/Graph/SControlRigGraphPinVariableBinding.h
@@ -16,14 +16,14 @@ class SControlRigVariableBinding : public SCompoundWidget
 public:
 
 	SLATE_BEGIN_ARGS(SControlRigVariableBinding)
-	: _ModelPin(nullptr)
+	: _ModelPins()
     , _FunctionReferenceNode(nullptr)
     , _InnerVariableName(NAME_None)
     , _Blueprint(nullptr)
 	, _CanRemoveBinding(true)
 	{}
 
-		SLATE_ARGUMENT(URigVMPin*, ModelPin)
+		SLATE_ARGUMENT(TArray<URigVMPin*>, ModelPins)
 		SLATE_ARGUMENT(URigVMFunctionReferenceNode*, FunctionReferenceNode)
 		SLATE_ARGUMENT(FName, InnerVariableName)
 		SLATE_ARGUMENT(UControlRigBlueprint*, Blueprint)
@@ -35,6 +35,7 @@ public:
 
 protected:
 
+	FText GetBindingText(URigVMPin* ModelPin) const;
 	FText GetBindingText() const;
 	const FSlateBrush* GetBindingImage() const;
 	FLinearColor GetBindingColor() const;
@@ -46,7 +47,7 @@ protected:
 	void FillLocalVariableMenu( FMenuBuilder& MenuBuilder );
 	void HandleBindToLocalVariable(FRigVMGraphVariableDescription InLocalVariable);
 
-	URigVMPin* ModelPin;
+	TArray<URigVMPin*> ModelPins;
 	URigVMFunctionReferenceNode* FunctionReferenceNode;
 	FName InnerVariableName;
 	UControlRigBlueprint* Blueprint;
@@ -60,7 +61,7 @@ public:
 
 	SLATE_BEGIN_ARGS(SControlRigGraphPinVariableBinding){}
 
-		SLATE_ARGUMENT(URigVMPin*, ModelPin)
+		SLATE_ARGUMENT(TArray<URigVMPin*>, ModelPins)
 		SLATE_ARGUMENT(UControlRigBlueprint*, Blueprint)
 
 	SLATE_END_ARGS()
@@ -73,6 +74,6 @@ protected:
 	virtual TSharedRef<SWidget>	GetDefaultValueWidget() override;
 	//~ End SGraphPin Interface
 
-	URigVMPin* ModelPin;
+	TArray<URigVMPin*> ModelPins;
 	UControlRigBlueprint* Blueprint;
 };
diff --git a/Engine/Source/Developer/AnimationWidgets/Public/SAdvancedTransformInputBox.h b/Engine/Source/Developer/AnimationWidgets/Public/SAdvancedTransformInputBox.h
--- a/Engine/Source/Developer/AnimationWidgets/Public/SAdvancedTransformInputBox.h
+++ b/Engine/Source/Developer/AnimationWidgets/Public/SAdvancedTransformInputBox.h
@@ -15,6 +15,7 @@
 
 #include "DetailCategoryBuilder.h"
 #include "DetailLayoutBuilder.h"
+#include "IDetailChildrenBuilder.h"
 #include "DetailWidgetRow.h"
 #include "IDetailGroup.h"
 #include "Framework/MultiBox/MultiBoxBuilder.h"
@@ -1093,9 +1094,20 @@ public:
 	}
 
 #if WITH_EDITOR
-	
+
+	static IDetailGroup& ConstructDetailGroup(IDetailCategoryBuilder& InBuilder, FName GroupName, const FText& LocalizedDisplayName)
+	{
+		return InBuilder.AddGroup(GroupName, LocalizedDisplayName, false, true);
+	}
+
+	static IDetailGroup& ConstructDetailGroup(IDetailChildrenBuilder& InBuilder, FName GroupName, const FText& LocalizedDisplayName)
+	{
+		return InBuilder.AddGroup(GroupName, LocalizedDisplayName);
+	}
+
+	template<typename BuilderType = IDetailCategoryBuilder>
 	static void ConstructGroupedTransformRows(
-		IDetailCategoryBuilder& InCategory,
+		BuilderType& InBuilder,
 		const FText& InLabel,
 		const FText& InTooltip,
 		typename SAdvancedTransformInputBox<TransformType>::FArguments WidgetArgs
@@ -1164,7 +1176,7 @@ public:
 			}			
 		};
 		
-		IDetailGroup& Group = InCategory.AddGroup(*InLabel.ToString(), InLabel, false, true);
+		IDetailGroup& Group = ConstructDetailGroup(InBuilder, *InLabel.ToString(), InLabel);
 		FDetailWidgetRow& HeaderRow = Group.HeaderRow();
 		ConstructComponentWidgetRow(HeaderRow, ESlateTransformComponent::Max);
 
diff --git a/Engine/Source/Developer/RigVMDeveloper/Private/RigVMModel/RigVMControllerActions.cpp b/Engine/Source/Developer/RigVMDeveloper/Private/RigVMModel/RigVMControllerActions.cpp
--- a/Engine/Source/Developer/RigVMDeveloper/Private/RigVMModel/RigVMControllerActions.cpp
+++ b/Engine/Source/Developer/RigVMDeveloper/Private/RigVMModel/RigVMControllerActions.cpp
@@ -53,6 +53,10 @@ bool URigVMActionStack::OpenUndoBracket(const FString& InTitle)
 bool URigVMActionStack::CloseUndoBracket()
 {
 	ensure(BracketActions.Num() > 0);
+	if(BracketActions.Last()->IsEmpty())
+	{
+		return CancelUndoBracket();
+	}
 	FRigVMBaseAction* Action = BracketActions.Pop();
 	EndAction(*Action);
 	delete(Action);
diff --git a/Engine/Source/Developer/RigVMDeveloper/Public/RigVMModel/RigVMControllerActions.h b/Engine/Source/Developer/RigVMDeveloper/Public/RigVMModel/RigVMControllerActions.h
--- a/Engine/Source/Developer/RigVMDeveloper/Public/RigVMModel/RigVMControllerActions.h
+++ b/Engine/Source/Developer/RigVMDeveloper/Public/RigVMModel/RigVMControllerActions.h
@@ -91,13 +91,16 @@ struct FRigVMBaseAction
 
 	// Default constructor
 	FRigVMBaseAction()
-		 : Title(TEXT("ACtion"))
+		 : Title(TEXT("Action"))
 	{
 	}
 
 	// Default destructor
 	virtual ~FRigVMBaseAction() {};
 
+	// Returns true if this action is empty has no effect
+	virtual bool IsEmpty() const { return SubActions.IsEmpty(); }
+
 	// Returns the title of the action - used for the Edit menu's undo / redo
 	virtual FString GetTitle() const { return Title; }
 
diff --git a/Engine/Source/Runtime/AnimationCore/Public/EulerTransform.h b/Engine/Source/Runtime/AnimationCore/Public/EulerTransform.h
--- a/Engine/Source/Runtime/AnimationCore/Public/EulerTransform.h
+++ b/Engine/Source/Runtime/AnimationCore/Public/EulerTransform.h
@@ -21,6 +21,8 @@ struct ANIMATIONCORE_API FEulerTransform
 {
 	GENERATED_BODY()
 
+	typedef FVector::FReal FReal;
+
 	/**
 	 * The identity transformation (Rotation = FRotator::ZeroRotator, Translation = FVector::ZeroVector, Scale = (1,1,1)).
 	 */
@@ -92,6 +94,14 @@ struct ANIMATIONCORE_API FEulerTransform
 		Scale = InTransform.GetScale3D();
 	}
 
+	// Test if all components of the transforms are equal, within a tolerance.
+	FORCEINLINE bool Equals(const FEulerTransform& Other, FReal Tolerance = KINDA_SMALL_NUMBER) const
+	{
+		return Location.Equals(Other.Location, Tolerance) &&
+			Rotation.Equals(Other.Rotation, Tolerance) &&
+			Scale.Equals(Other.Scale, Tolerance);
+	}
+
 	FORCEINLINE const FVector& GetLocation() const { return Location; }
 	FORCEINLINE FQuat GetRotation() const { return Rotation.Quaternion(); }
 	FORCEINLINE const FRotator& Rotator() const { return Rotation; }
