Control Rig: React to changes in member variables in the right callback

diff --git a/Engine/Plugins/Animation/ControlRig/Source/ControlRigEditor/Private/Editor/ControlRigEditor.cpp b/Engine/Plugins/Animation/ControlRig/Source/ControlRigEditor/Private/Editor/ControlRigEditor.cpp
--- a/Engine/Plugins/Animation/ControlRig/Source/ControlRigEditor/Private/Editor/ControlRigEditor.cpp
+++ b/Engine/Plugins/Animation/ControlRig/Source/ControlRigEditor/Private/Editor/ControlRigEditor.cpp
@@ -4266,6 +4266,24 @@ void FControlRigEditor::NotifyPreChange(FProperty* PropertyAboutToChange)
 void FControlRigEditor::NotifyPostChange(const FPropertyChangedEvent& PropertyChangedEvent, FProperty* PropertyThatChanged)
 {
 	FBlueprintEditor::NotifyPostChange(PropertyChangedEvent, PropertyThatChanged);
+	
+	// we need to listen to changes for variables on the blueprint here since
+	// OnFinishedChangingProperties is called only for top level property changes.
+	// changes on a lower level property like transform under a user defined struct
+	// only go through this.
+	UControlRigBlueprint* ControlRigBP = GetControlRigBlueprint();
+	if(ControlRig && ControlRigBP)
+	{
+		const FName VarName = PropertyChangedEvent.MemberProperty->GetFName();
+		for(FBPVariableDescription& NewVariable : ControlRigBP->NewVariables)
+		{
+			if(NewVariable.VarName == VarName)
+			{
+				UpdateDefaultValueForVariable(NewVariable, false);
+				break;
+			}
+		}
+	}
 }
 
 void FControlRigEditor::OnFinishedChangingProperties(const FPropertyChangedEvent& PropertyChangedEvent)
@@ -4294,20 +4312,6 @@ void FControlRigEditor::OnFinishedChangingProperties(const FPropertyChangedEvent
 			return;
 		}
 	}
-
-	// this might be a property value for a variable
-	if(ControlRig && ControlRigBP)
-	{
-		const FName VarName = PropertyChangedEvent.MemberProperty->GetFName();
-		for(FBPVariableDescription& NewVariable : ControlRigBP->NewVariables)
-		{
-			if(NewVariable.VarName == VarName)
-			{
-				UpdateDefaultValueForVariable(NewVariable, false);
-				break;
-			}
-		}
-	}
 }
 
 void FControlRigEditor::OnWrappedPropertyChangedChainEvent(UDetailsViewWrapperObject* InWrapperObject, const FString& InPropertyPath, FPropertyChangedChainEvent& InPropertyChangedChainEvent)
