ControlRig: Hierarchy to comply with FMemory::Alloc and FMemory::Free

diff --git a/Engine/Plugins/Animation/ControlRig/Source/ControlRig/Private/Rigs/RigHierarchy.cpp b/Engine/Plugins/Animation/ControlRig/Source/ControlRig/Private/Rigs/RigHierarchy.cpp
--- a/Engine/Plugins/Animation/ControlRig/Source/ControlRig/Private/Rigs/RigHierarchy.cpp
+++ b/Engine/Plugins/Animation/ControlRig/Source/ControlRig/Private/Rigs/RigHierarchy.cpp
@@ -3357,12 +3357,7 @@ FRigBaseElement* URigHierarchy::MakeElement(ERigElementType InElementType, int32
 			{
 				*OutStructureSize = sizeof(FRigBoneElement);
 			}
-			FRigBoneElement* Elements = (FRigBoneElement*)FMemory::Malloc(sizeof(FRigBoneElement) * InCount);
-			for(int32 Index=0;Index<InCount;Index++)
-			{
-				new(&Elements[Index]) FRigBoneElement(); 
-			} 
-			Element = Elements;
+			Element = NewElement<FRigBoneElement>(InCount);
 			break;
 		}
 		case ERigElementType::Null:
@@ -3371,12 +3366,7 @@ FRigBaseElement* URigHierarchy::MakeElement(ERigElementType InElementType, int32
 			{
 				*OutStructureSize = sizeof(FRigNullElement);
 			}
-			FRigNullElement* Elements = (FRigNullElement*)FMemory::Malloc(sizeof(FRigNullElement) * InCount);
-			for(int32 Index=0;Index<InCount;Index++)
-			{
-				new(&Elements[Index]) FRigNullElement(); 
-			} 
-			Element = Elements;
+			Element = NewElement<FRigNullElement>(InCount);
 			break;
 		}
 		case ERigElementType::Control:
@@ -3385,12 +3375,7 @@ FRigBaseElement* URigHierarchy::MakeElement(ERigElementType InElementType, int32
 			{
 				*OutStructureSize = sizeof(FRigControlElement);
 			}
-			FRigControlElement* Elements = (FRigControlElement*)FMemory::Malloc(sizeof(FRigControlElement) * InCount);
-			for(int32 Index=0;Index<InCount;Index++)
-			{
-				new(&Elements[Index]) FRigControlElement(); 
-			} 
-			Element = Elements;
+			Element = NewElement<FRigControlElement>(InCount);
 			break;
 		}
 		case ERigElementType::Curve:
@@ -3399,12 +3384,7 @@ FRigBaseElement* URigHierarchy::MakeElement(ERigElementType InElementType, int32
 			{
 				*OutStructureSize = sizeof(FRigCurveElement);
 			}
-			FRigCurveElement* Elements = (FRigCurveElement*)FMemory::Malloc(sizeof(FRigCurveElement) * InCount);
-			for(int32 Index=0;Index<InCount;Index++)
-			{
-				new(&Elements[Index]) FRigCurveElement(); 
-			} 
-			Element = Elements;
+			Element = NewElement<FRigCurveElement>(InCount);
 			break;
 		}
 		case ERigElementType::RigidBody:
@@ -3413,12 +3393,7 @@ FRigBaseElement* URigHierarchy::MakeElement(ERigElementType InElementType, int32
 			{
 				*OutStructureSize = sizeof(FRigRigidBodyElement);
 			}
-			FRigRigidBodyElement* Elements = (FRigRigidBodyElement*)FMemory::Malloc(sizeof(FRigRigidBodyElement) * InCount);
-			for(int32 Index=0;Index<InCount;Index++)
-			{
-				new(&Elements[Index]) FRigRigidBodyElement(); 
-			} 
-			Element = Elements;
+			Element = NewElement<FRigRigidBodyElement>(InCount);
 			break;
 		}
 		case ERigElementType::Reference:
@@ -3427,12 +3402,7 @@ FRigBaseElement* URigHierarchy::MakeElement(ERigElementType InElementType, int32
 			{
 				*OutStructureSize = sizeof(FRigReferenceElement);
 			}
-			FRigReferenceElement* Elements = (FRigReferenceElement*)FMemory::Malloc(sizeof(FRigReferenceElement) * InCount);
-			for(int32 Index=0;Index<InCount;Index++)
-			{
-				new(&Elements[Index]) FRigReferenceElement(); 
-			} 
-			Element = Elements;
+			Element = NewElement<FRigReferenceElement>(InCount);
 			break;
 		}
 		default:
diff --git a/Engine/Plugins/Animation/ControlRig/Source/ControlRig/Private/Rigs/RigHierarchyController.cpp b/Engine/Plugins/Animation/ControlRig/Source/ControlRig/Private/Rigs/RigHierarchyController.cpp
--- a/Engine/Plugins/Animation/ControlRig/Source/ControlRig/Private/Rigs/RigHierarchyController.cpp
+++ b/Engine/Plugins/Animation/ControlRig/Source/ControlRig/Private/Rigs/RigHierarchyController.cpp
@@ -180,7 +180,7 @@ FRigElementKey URigHierarchyController::AddBone(FName InName, FRigElementKey InP
 	}
 #endif
 
-	FRigBoneElement* NewElement = new FRigBoneElement();
+	FRigBoneElement* NewElement = URigHierarchy::NewElement<FRigBoneElement>();
 	{
 		TGuardValue<bool> DisableCacheValidityChecks(Hierarchy->bEnableCacheValidityCheck, false);
 		NewElement->Key.Type = ERigElementType::Bone;
@@ -240,7 +240,7 @@ FRigElementKey URigHierarchyController::AddNull(FName InName, FRigElementKey InP
 	}
 #endif
 
-	FRigNullElement* NewElement = new FRigNullElement();
+	FRigNullElement* NewElement = URigHierarchy::NewElement<FRigNullElement>();
 	{
 		TGuardValue<bool> DisableCacheValidityChecks(Hierarchy->bEnableCacheValidityCheck, false);		
 		NewElement->Key.Type = ERigElementType::Null;
@@ -307,7 +307,7 @@ FRigElementKey URigHierarchyController::AddControl(
 	}
 #endif
 
-	FRigControlElement* NewElement = new FRigControlElement();
+	FRigControlElement* NewElement = URigHierarchy::NewElement<FRigControlElement>();
 	{
 		TGuardValue<bool> DisableCacheValidityChecks(Hierarchy->bEnableCacheValidityCheck, false);		
 		NewElement->Key.Type = ERigElementType::Control;
@@ -372,7 +372,7 @@ FRigElementKey URigHierarchyController::AddCurve(FName InName, float InValue, bo
 	}
 #endif
 
-	FRigCurveElement* NewElement = new FRigCurveElement();
+	FRigCurveElement* NewElement = URigHierarchy::NewElement<FRigCurveElement>();
 	{
 		TGuardValue<bool> DisableCacheValidityChecks(Hierarchy->bEnableCacheValidityCheck, false);		
 		NewElement->Key.Type = ERigElementType::Curve;
@@ -421,7 +421,7 @@ FRigElementKey URigHierarchyController::AddRigidBody(FName InName, FRigElementKe
 	}
 #endif
 
-	FRigRigidBodyElement* NewElement = new FRigRigidBodyElement();
+	FRigRigidBodyElement* NewElement = URigHierarchy::NewElement<FRigRigidBodyElement>();
 	{
 		TGuardValue<bool> DisableCacheValidityChecks(Hierarchy->bEnableCacheValidityCheck, false);		
 		NewElement->Key.Type = ERigElementType::RigidBody;
@@ -473,7 +473,7 @@ FRigElementKey URigHierarchyController::AddReference(FName InName, FRigElementKe
 	}
 #endif
 
-	FRigReferenceElement* NewElement = new FRigReferenceElement();
+	FRigReferenceElement* NewElement = URigHierarchy::NewElement<FRigReferenceElement>();
 	{
 		TGuardValue<bool> DisableCacheValidityChecks(Hierarchy->bEnableCacheValidityCheck, false);		
 		NewElement->Key.Type = ERigElementType::Reference;
@@ -1056,38 +1056,38 @@ TArray<FRigElementKey> URigHierarchyController::ImportFromText(FString InContent
 		{
 			case ERigElementType::Bone:
 			{
-				NewElement = new FRigBoneElement();
+				NewElement = URigHierarchy::NewElement<FRigBoneElement>();
 				FRigBoneElement::StaticStruct()->ImportText(*PerElementData.Content, NewElement, nullptr, EPropertyPortFlags::PPF_None, &ErrorPipe, FRigBoneElement::StaticStruct()->GetName(), true);
 				CastChecked<FRigBoneElement>(NewElement)->BoneType = ERigBoneType::User;
 				break;
 			}
 			case ERigElementType::Null:
 			{
-				NewElement = new FRigNullElement();
+				NewElement = URigHierarchy::NewElement<FRigNullElement>();
 				FRigNullElement::StaticStruct()->ImportText(*PerElementData.Content, NewElement, nullptr, EPropertyPortFlags::PPF_None, &ErrorPipe, FRigNullElement::StaticStruct()->GetName(), true);
 				break;
 			}
 			case ERigElementType::Control:
 			{
-				NewElement = new FRigControlElement();
+				NewElement = URigHierarchy::NewElement<FRigControlElement>();
 				FRigControlElement::StaticStruct()->ImportText(*PerElementData.Content, NewElement, nullptr, EPropertyPortFlags::PPF_None, &ErrorPipe, FRigControlElement::StaticStruct()->GetName(), true);
 				break;
 			}
 			case ERigElementType::Curve:
 			{
-				NewElement = new FRigCurveElement();
+				NewElement = URigHierarchy::NewElement<FRigCurveElement>();
 				FRigCurveElement::StaticStruct()->ImportText(*PerElementData.Content, NewElement, nullptr, EPropertyPortFlags::PPF_None, &ErrorPipe, FRigCurveElement::StaticStruct()->GetName(), true);
 				break;
 			}
 			case ERigElementType::RigidBody:
 			{
-				NewElement = new FRigRigidBodyElement();
+				NewElement = URigHierarchy::NewElement<FRigRigidBodyElement>();
 				FRigRigidBodyElement::StaticStruct()->ImportText(*PerElementData.Content, NewElement, nullptr, EPropertyPortFlags::PPF_None, &ErrorPipe, FRigRigidBodyElement::StaticStruct()->GetName(), true);
 				break;
 			}
 			case ERigElementType::Reference:
 			{
-				NewElement = new FRigReferenceElement();
+				NewElement = URigHierarchy::NewElement<FRigReferenceElement>();
 				FRigReferenceElement::StaticStruct()->ImportText(*PerElementData.Content, NewElement, nullptr, EPropertyPortFlags::PPF_None, &ErrorPipe, FRigReferenceElement::StaticStruct()->GetName(), true);
 				break;
 			}
@@ -1139,7 +1139,8 @@ TArray<FRigElementKey> URigHierarchyController::ImportFromText(FString InContent
 				}
 
 				PastedKeys.Add(ExistingElement->GetKey());
-				delete NewElement;
+
+				URigHierarchy::DestroyElement(NewElement);
 				continue;
 			}
 		}
@@ -1744,7 +1745,7 @@ bool URigHierarchyController::RemoveElement(FRigBaseElement* InElement)
 		Notify(ERigHierarchyNotification::HierarchyReset, nullptr);
 	}
 
-	delete InElement;
+	URigHierarchy::DestroyElement(InElement);
 
 	Hierarchy->EnsureCacheValidity();
 
diff --git a/Engine/Plugins/Animation/ControlRig/Source/ControlRig/Public/Rigs/RigHierarchy.h b/Engine/Plugins/Animation/ControlRig/Source/ControlRig/Public/Rigs/RigHierarchy.h
--- a/Engine/Plugins/Animation/ControlRig/Source/ControlRig/Public/Rigs/RigHierarchy.h
+++ b/Engine/Plugins/Animation/ControlRig/Source/ControlRig/Public/Rigs/RigHierarchy.h
@@ -2537,6 +2537,20 @@ private:
 	*/
 	static void DestroyElement(FRigBaseElement*& InElement); 
 
+	/*
+	 * Templated helper function to create an element
+	 */
+	template<typename ElementType = FRigBaseElement>
+	FORCEINLINE static ElementType* NewElement(int32 Num = 1)
+	{
+		ElementType* Elements = (ElementType*)FMemory::Malloc(sizeof(ElementType) * Num);
+		for(int32 Index=0;Index<Num;Index++)
+		{
+			new(&Elements[Index]) ElementType(); 
+		}
+		return Elements;
+	}
+
 	/**
 	 * Marks all affected elements of a given element as dirty
 	 * @param InTransformElement The element that has changed
