Control Rig: Python scripts to add controls for selected bones not printing python interactive commands

diff --git a/Engine/Plugins/Animation/ControlRig/Source/ControlRig/Public/Rigs/RigHierarchy.h b/Engine/Plugins/Animation/ControlRig/Source/ControlRig/Public/Rigs/RigHierarchy.h
--- a/Engine/Plugins/Animation/ControlRig/Source/ControlRig/Public/Rigs/RigHierarchy.h
+++ b/Engine/Plugins/Animation/ControlRig/Source/ControlRig/Public/Rigs/RigHierarchy.h
@@ -1011,9 +1011,9 @@ public:
 	 * @param bSetupUndo If true the transform stack will be setup for undo / redo
 	 */
 	UFUNCTION(BlueprintCallable, Category = URigHierarchy)
-    FORCEINLINE_DEBUGGABLE void SetGlobalTransform(FRigElementKey InKey, FTransform InTransform, bool bInitial = false, bool bAffectChildren = true, bool bSetupUndo = false)
+    FORCEINLINE_DEBUGGABLE void SetGlobalTransform(FRigElementKey InKey, FTransform InTransform, bool bInitial = false, bool bAffectChildren = true, bool bSetupUndo = false, bool bPrintPythonCommand = false)
 	{
-		SetGlobalTransformByIndex(GetIndex(InKey), InTransform, bInitial, bAffectChildren, bSetupUndo);
+		SetGlobalTransformByIndex(GetIndex(InKey), InTransform, bInitial, bAffectChildren, bSetupUndo, bPrintPythonCommand);
 	}
 
 	/**
@@ -1025,13 +1025,13 @@ public:
 	 * @param bSetupUndo If true the transform stack will be setup for undo / redo
 	 */
 	UFUNCTION(BlueprintCallable, Category = URigHierarchy)
-    FORCEINLINE_DEBUGGABLE void SetGlobalTransformByIndex(int32 InElementIndex, FTransform InTransform, bool bInitial = false, bool bAffectChildren = true, bool bSetupUndo = false)
+    FORCEINLINE_DEBUGGABLE void SetGlobalTransformByIndex(int32 InElementIndex, FTransform InTransform, bool bInitial = false, bool bAffectChildren = true, bool bSetupUndo = false, bool bPrintPythonCommand = false)
 	{
 		if(Elements.IsValidIndex(InElementIndex))
 		{
 			if(FRigTransformElement* TransformElement = Cast<FRigTransformElement>(Elements[InElementIndex]))
 			{
-				SetTransform(TransformElement, InTransform, bInitial ? ERigTransformType::InitialGlobal : ERigTransformType::CurrentGlobal, bAffectChildren, bSetupUndo);
+				SetTransform(TransformElement, InTransform, bInitial ? ERigTransformType::InitialGlobal : ERigTransformType::CurrentGlobal, bAffectChildren, bSetupUndo, false, bPrintPythonCommand);
 			}
 		}
 	}
diff --git a/Engine/Plugins/Animation/ControlRig/Source/ControlRig/Public/Rigs/RigHierarchyController.h b/Engine/Plugins/Animation/ControlRig/Source/ControlRig/Public/Rigs/RigHierarchyController.h
--- a/Engine/Plugins/Animation/ControlRig/Source/ControlRig/Public/Rigs/RigHierarchyController.h
+++ b/Engine/Plugins/Animation/ControlRig/Source/ControlRig/Public/Rigs/RigHierarchyController.h
@@ -139,10 +139,11 @@ public:
         FRigElementKey InParent,
         FRigControlSettings InSettings,
         FRigControlValue InValue,
-        bool bSetupUndo = true
+        bool bSetupUndo = true,
+        bool bPrintPythonCommand = false
     )
 	{
-		return AddControl(InName, InParent, InSettings, InValue, FTransform::Identity, FTransform::Identity, bSetupUndo);
+		return AddControl(InName, InParent, InSettings, InValue, FTransform::Identity, FTransform::Identity, bSetupUndo, bPrintPythonCommand);
 	}
 
 	/**
