Control Rig: Support for Stacked Hierarchy for RigHierarchy View

diff --git a/Engine/Plugins/Animation/ControlRig/Source/ControlRig/Private/Settings/ControlRigSettings.cpp b/Engine/Plugins/Animation/ControlRig/Source/ControlRig/Private/Settings/ControlRigSettings.cpp
--- a/Engine/Plugins/Animation/ControlRig/Source/ControlRig/Private/Settings/ControlRigSettings.cpp
+++ b/Engine/Plugins/Animation/ControlRig/Source/ControlRig/Private/Settings/ControlRigSettings.cpp
@@ -31,5 +31,7 @@ UControlRigEditorSettings::UControlRigEditorSettings(const FObjectInitializer& O
 	BackwardsSolveBorderColor = FLinearColor::Yellow;
 	BackwardsAndForwardsBorderColor = FLinearColor::Blue;
 	DynamicHierarchyBorderColor = FLinearColor::Blue;
+	bShowStackedHierarchy = true;
+	MaxStackSize = 16;
 #endif
 }
diff --git a/Engine/Plugins/Animation/ControlRig/Source/ControlRig/Public/Settings/ControlRigSettings.h b/Engine/Plugins/Animation/ControlRig/Source/ControlRig/Public/Settings/ControlRigSettings.h
--- a/Engine/Plugins/Animation/ControlRig/Source/ControlRig/Public/Settings/ControlRigSettings.h
+++ b/Engine/Plugins/Animation/ControlRig/Source/ControlRig/Public/Settings/ControlRigSettings.h
@@ -111,14 +111,28 @@ public:
 	FLinearColor BackwardsAndForwardsBorderColor;
 
 	/**
-	* The border color of the hierarchy tree in dynamic hierarchy mode
-	*/
-	UPROPERTY(EditAnywhere, config, Category = Viewport)
+	 * The border color of the hierarchy tree in dynamic hierarchy mode
+	 */
+	UPROPERTY(EditAnywhere, config, Category = Hierarchy)
 	FLinearColor DynamicHierarchyBorderColor;
 
 	/**
-	* The default node snippet to create when pressing 1 + Left Mouse Button
-	*/
+	 * Option to toggle displaying the stacked hierarchy items.
+	 * Note that changing this option potentially requires to re-open the editors in question. 
+	 */
+	UPROPERTY(EditAnywhere, config, Category = Hierarchy)
+	bool bShowStackedHierarchy;
+
+	/**
+ 	 * The maximum number of stacked items in the view 
+ 	 * Note that changing this option potentially requires to re-open the editors in question. 
+ 	 */
+	UPROPERTY(EditAnywhere, config, Category = Hierarchy, meta = (EditCondition = "bShowStackedHierarchy"))
+	int32 MaxStackSize;
+
+	/**
+	 * The default node snippet to create when pressing 1 + Left Mouse Button
+	 */
 	UPROPERTY(EditAnywhere, config, Category = Snippets, meta = (DisplayName = "1"))
 	FString NodeSnippet_1;
 
diff --git a/Engine/Plugins/Animation/ControlRig/Source/ControlRigEditor/Private/Editor/SRigHierarchyTreeView.cpp b/Engine/Plugins/Animation/ControlRig/Source/ControlRigEditor/Private/Editor/SRigHierarchyTreeView.cpp
--- a/Engine/Plugins/Animation/ControlRig/Source/ControlRigEditor/Private/Editor/SRigHierarchyTreeView.cpp
+++ b/Engine/Plugins/Animation/ControlRig/Source/ControlRigEditor/Private/Editor/SRigHierarchyTreeView.cpp
@@ -15,6 +15,7 @@
 #include "Widgets/Text/SInlineEditableTextBlock.h"
 #include "ControlRig.h"
 #include "ControlRigEditorStyle.h"
+#include "Settings/ControlRigSettings.h"
 
 #define LOCTEXT_NAMESPACE "SRigHierarchyTreeView"
 
@@ -49,9 +50,9 @@ FRigTreeElement::FRigTreeElement(const FRigElementKey& InKey, TWeakPtr<SRigHiera
 }
 
 
-TSharedRef<ITableRow> FRigTreeElement::MakeTreeRowWidget(const TSharedRef<STableViewBase>& InOwnerTable, TSharedRef<FRigTreeElement> InRigTreeElement, TSharedPtr<SRigHierarchyTreeView> InTreeView, const FRigTreeDisplaySettings& InSettings)
+TSharedRef<ITableRow> FRigTreeElement::MakeTreeRowWidget(const TSharedRef<STableViewBase>& InOwnerTable, TSharedRef<FRigTreeElement> InRigTreeElement, TSharedPtr<SRigHierarchyTreeView> InTreeView, const FRigTreeDisplaySettings& InSettings, bool bPinned)
 {
-	return SNew(SRigHierarchyItem, InOwnerTable, InRigTreeElement, InTreeView, InSettings);
+	return SNew(SRigHierarchyItem, InOwnerTable, InRigTreeElement, InTreeView, InSettings, bPinned);
 }
 
 void FRigTreeElement::RequestRename()
@@ -82,7 +83,7 @@ void FRigTreeElement::RefreshDisplaySettings(const URigHierarchy* InHierarchy, c
 //////////////////////////////////////////////////////////////
 /// SRigHierarchyItem
 ///////////////////////////////////////////////////////////
-void SRigHierarchyItem::Construct(const FArguments& InArgs, const TSharedRef<STableViewBase>& OwnerTable, TSharedRef<FRigTreeElement> InRigTreeElement, TSharedPtr<SRigHierarchyTreeView> InTreeView, const FRigTreeDisplaySettings& InSettings)
+void SRigHierarchyItem::Construct(const FArguments& InArgs, const TSharedRef<STableViewBase>& OwnerTable, TSharedRef<FRigTreeElement> InRigTreeElement, TSharedPtr<SRigHierarchyTreeView> InTreeView, const FRigTreeDisplaySettings& InSettings, bool bPinned)
 {
 	WeakRigTreeElement = InRigTreeElement;
 	Delegates = InTreeView->GetRigTreeDelegates();
@@ -186,7 +187,7 @@ void SRigHierarchyTreeView::Construct(const FArguments& InArgs)
 	STreeView<TSharedPtr<FRigTreeElement>>::FArguments SuperArgs;
 	SuperArgs.TreeItemsSource(&RootElements);
 	SuperArgs.SelectionMode(ESelectionMode::Multi);
-	SuperArgs.OnGenerateRow(this, &SRigHierarchyTreeView::MakeTableRowWidget);
+	SuperArgs.OnGenerateRow(this, &SRigHierarchyTreeView::MakeTableRowWidget, false);
 	SuperArgs.OnGetChildren(this, &SRigHierarchyTreeView::HandleGetChildrenForTree);
 	SuperArgs.OnSelectionChanged(FOnRigTreeSelectionChanged::CreateRaw(&Delegates, &FRigTreeDelegates::HandleSelectionChanged));
 	SuperArgs.OnContextMenuOpening(Delegates.OnContextMenuOpening);
@@ -196,6 +197,15 @@ void SRigHierarchyTreeView::Construct(const FArguments& InArgs)
 	SuperArgs.HighlightParentNodesForSelection(true);
 	SuperArgs.ItemHeight(24);
 	SuperArgs.AllowInvisibleItemSelection(true);  //without this we deselect everything when we filter or we collapse
+	
+	SuperArgs.ShouldStackHierarchyHeaders_Lambda([]() -> bool {
+		return UControlRigEditorSettings::Get()->bShowStackedHierarchy;
+	});
+	SuperArgs.OnGeneratePinnedRow(this, &SRigHierarchyTreeView::MakeTableRowWidget, true);
+	SuperArgs.MaxPinnedItems_Lambda([]() -> int32
+	{
+		return FMath::Max<int32>(1, UControlRigEditorSettings::Get()->MaxStackSize);
+	});
 
 	STreeView<TSharedPtr<FRigTreeElement>>::Construct(SuperArgs);
 }
@@ -623,10 +633,10 @@ void SRigHierarchyTreeView::SetExpansionRecursive(TSharedPtr<FRigTreeElement> In
 }
 
 TSharedRef<ITableRow> SRigHierarchyTreeView::MakeTableRowWidget(TSharedPtr<FRigTreeElement> InItem,
-	const TSharedRef<STableViewBase>& OwnerTable)
+	const TSharedRef<STableViewBase>& OwnerTable, bool bPinned)
 {
 	const FRigTreeDisplaySettings& Settings = Delegates.GetDisplaySettings();
-	return InItem->MakeTreeRowWidget(OwnerTable, InItem.ToSharedRef(), SharedThis(this), Settings);
+	return InItem->MakeTreeRowWidget(OwnerTable, InItem.ToSharedRef(), SharedThis(this), Settings, bPinned);
 }
 
 void SRigHierarchyTreeView::HandleGetChildrenForTree(TSharedPtr<FRigTreeElement> InItem,
diff --git a/Engine/Plugins/Animation/ControlRig/Source/ControlRigEditor/Private/Editor/SRigHierarchyTreeView.h b/Engine/Plugins/Animation/ControlRig/Source/ControlRigEditor/Private/Editor/SRigHierarchyTreeView.h
--- a/Engine/Plugins/Animation/ControlRig/Source/ControlRigEditor/Private/Editor/SRigHierarchyTreeView.h
+++ b/Engine/Plugins/Animation/ControlRig/Source/ControlRigEditor/Private/Editor/SRigHierarchyTreeView.h
@@ -174,7 +174,7 @@ public:
 	bool bSupportsRename;
 	TArray<TSharedPtr<FRigTreeElement>> Children;
 
-	TSharedRef<ITableRow> MakeTreeRowWidget(const TSharedRef<STableViewBase>& InOwnerTable, TSharedRef<FRigTreeElement> InRigTreeElement, TSharedPtr<SRigHierarchyTreeView> InTreeView, const FRigTreeDisplaySettings& InSettings);
+	TSharedRef<ITableRow> MakeTreeRowWidget(const TSharedRef<STableViewBase>& InOwnerTable, TSharedRef<FRigTreeElement> InRigTreeElement, TSharedPtr<SRigHierarchyTreeView> InTreeView, const FRigTreeDisplaySettings& InSettings, bool bPinned);
 
 	void RequestRename();
 
@@ -201,7 +201,7 @@ class SRigHierarchyItem : public STableRow<TSharedPtr<FRigTreeElement>>
 {
 public:
 	
-	void Construct(const FArguments& InArgs, const TSharedRef<STableViewBase>& OwnerTable, TSharedRef<FRigTreeElement> InRigTreeElement, TSharedPtr<SRigHierarchyTreeView> InTreeView, const FRigTreeDisplaySettings& InSettings);
+	void Construct(const FArguments& InArgs, const TSharedRef<STableViewBase>& OwnerTable, TSharedRef<FRigTreeElement> InRigTreeElement, TSharedPtr<SRigHierarchyTreeView> InTreeView, const FRigTreeDisplaySettings& InSettings, bool bPinned);
  	void OnNameCommitted(const FText& InText, ETextCommit::Type InCommitType) const;
 	bool OnVerifyNameChanged(const FText& InText, FText& OutErrorMessage);
 	static TPair<const FSlateBrush*, FSlateColor> GetBrushForElementType(const URigHierarchy* InHierarchy, const FRigElementKey& InKey);
@@ -268,7 +268,7 @@ public:
 	bool RemoveElement(FRigElementKey InKey);
 	void RefreshTreeView(bool bRebuildContent = true);
 	void SetExpansionRecursive(TSharedPtr<FRigTreeElement> InElement, bool bTowardsParent, bool bShouldBeExpanded);
-	TSharedRef<ITableRow> MakeTableRowWidget(TSharedPtr<FRigTreeElement> InItem, const TSharedRef<STableViewBase>& OwnerTable);
+	TSharedRef<ITableRow> MakeTableRowWidget(TSharedPtr<FRigTreeElement> InItem, const TSharedRef<STableViewBase>& OwnerTable, bool bPinned);
 	void HandleGetChildrenForTree(TSharedPtr<FRigTreeElement> InItem, TArray<TSharedPtr<FRigTreeElement>>& OutChildren);
 
 	TArray<FRigElementKey> GetSelectedKeys() const;
