Provide ability set ControlRig object binding to a USkeleton
   - Added GetBindableObject helper to determine correct UObject to bind against

diff --git a/Engine/Plugins/Animation/ControlRig/Source/ControlRig/Private/ControlRigObjectBinding.cpp b/Engine/Plugins/Animation/ControlRig/Source/ControlRig/Private/ControlRigObjectBinding.cpp
--- a/Engine/Plugins/Animation/ControlRig/Source/ControlRig/Private/ControlRigObjectBinding.cpp
+++ b/Engine/Plugins/Animation/ControlRig/Source/ControlRig/Private/ControlRigObjectBinding.cpp
@@ -12,72 +12,62 @@ FControlRigObjectBinding::~FControlRigObjectBinding()
 
 void FControlRigObjectBinding::BindToObject(UObject* InObject)
 {
-	// If we are binding to an actor, find the first skeletal mesh component
-	if (AActor* Actor = Cast<AActor>(InObject))
-	{
-		if (UControlRigComponent* ControlRigComponent = Actor->FindComponentByClass<UControlRigComponent>())
-		{
-			SceneComponent = ControlRigComponent;
-		}
-		else if (USkeletalMeshComponent* SkeletalMeshComponent = Actor->FindComponentByClass<USkeletalMeshComponent>())
-		{
-			SceneComponent = SkeletalMeshComponent;
-		}
-	}
-	else if (UControlRigComponent* ControlRigComponent = Cast<UControlRigComponent>(InObject))
-	{
-		SceneComponent = ControlRigComponent;
-	}
-	else if (USkeletalMeshComponent* SkeletalMeshComponent = Cast<USkeletalMeshComponent>(InObject))
-	{
-		SceneComponent = SkeletalMeshComponent;
-	}
-
-	ControlRigBind.Broadcast(SceneComponent.Get());
+	BoundObject = GetBindableObject(InObject);
+	ControlRigBind.Broadcast(BoundObject.Get());
 }
 
 void FControlRigObjectBinding::UnbindFromObject()
 {
-	SceneComponent = nullptr;
+	BoundObject = nullptr;
 
 	ControlRigUnbind.Broadcast();
 }
 
 bool FControlRigObjectBinding::IsBoundToObject(UObject* InObject) const
 {
+	return InObject != nullptr && BoundObject.Get() == GetBindableObject(InObject);
+}
+
+UObject* FControlRigObjectBinding::GetBoundObject() const
+{
+	return BoundObject.Get();
+}
+
+AActor* FControlRigObjectBinding::GetHostingActor() const
+{
+	if (USceneComponent* SceneComponent = Cast<USceneComponent>(BoundObject.Get()))
+	{
+		return SceneComponent->GetOwner();
+	}
+
+	return nullptr;
+}
+
+UObject* FControlRigObjectBinding::GetBindableObject(UObject* InObject) const
+{
+	// If we are binding to an actor, find the first skeletal mesh component
 	if (AActor* Actor = Cast<AActor>(InObject))
 	{
 		if (UControlRigComponent* ControlRigComponent = Actor->FindComponentByClass<UControlRigComponent>())
 		{
-			return SceneComponent.Get() == ControlRigComponent;
+			return ControlRigComponent;
 		}
 		else if (USkeletalMeshComponent* SkeletalMeshComponent = Actor->FindComponentByClass<USkeletalMeshComponent>())
 		{
-			return SceneComponent.Get() == SkeletalMeshComponent;
+			return SkeletalMeshComponent;
 		}
 	}
 	else if (UControlRigComponent* ControlRigComponent = Cast<UControlRigComponent>(InObject))
 	{
-		return SceneComponent.Get() == ControlRigComponent;
+		return ControlRigComponent;
 	}
 	else if (USkeletalMeshComponent* SkeletalMeshComponent = Cast<USkeletalMeshComponent>(InObject))
 	{
-		return SceneComponent.Get() == SkeletalMeshComponent;
+		return SkeletalMeshComponent;
 	}
-
-	return false;
-}
-
-UObject* FControlRigObjectBinding::GetBoundObject() const
-{
-	return SceneComponent.Get();
-}
-
-AActor* FControlRigObjectBinding::GetHostingActor() const
-{
-	if (SceneComponent.Get())
+	else if (USkeleton* Skeleton = Cast<USkeleton>(InObject))
 	{
-		return SceneComponent->GetOwner();
+		return Skeleton;
 	}
 
 	return nullptr;
diff --git a/Engine/Plugins/Animation/ControlRig/Source/ControlRig/Public/ControlRigObjectBinding.h b/Engine/Plugins/Animation/ControlRig/Source/ControlRig/Public/ControlRigObjectBinding.h
--- a/Engine/Plugins/Animation/ControlRig/Source/ControlRig/Public/ControlRigObjectBinding.h
+++ b/Engine/Plugins/Animation/ControlRig/Source/ControlRig/Public/ControlRigObjectBinding.h
@@ -22,10 +22,11 @@ public:
 	virtual bool IsBoundToObject(UObject* InObject) const override;
 	virtual UObject* GetBoundObject() const override;
 	virtual AActor* GetHostingActor() const override;
-
+protected:
+	UObject* GetBindableObject(UObject* InObject) const;
 private:
-	/** The scene component we are bound to */
-	TWeakObjectPtr<USceneComponent> SceneComponent;
+	/** The scene component or USkeleton we are bound to */
+	TWeakObjectPtr<UObject> BoundObject;
 
 	FControlRigBind ControlRigBind;
 	FControlRigUnbind ControlRigUnbind;
