Set _both_ Initial and Current for Transforms when initializing or adding Bone elements

diff --git a/Engine/Plugins/Animation/ControlRig/Source/ControlRig/Private/ControlRig.cpp b/Engine/Plugins/Animation/ControlRig/Source/ControlRig/Private/ControlRig.cpp
--- a/Engine/Plugins/Animation/ControlRig/Source/ControlRig/Private/ControlRig.cpp
+++ b/Engine/Plugins/Animation/ControlRig/Source/ControlRig/Private/ControlRig.cpp
@@ -2693,6 +2693,7 @@ void UControlRig::SetBoneInitialTransformsFromRefSkeleton(const FReferenceSkelet
 			{
 				const FTransform LocalInitialTransform = InReferenceSkeleton.GetRefBonePose()[BoneIndex];
 				DynamicHierarchy->SetTransform(BoneElement, LocalInitialTransform, ERigTransformType::InitialLocal, true, false);
+				DynamicHierarchy->SetTransform(BoneElement, LocalInitialTransform, ERigTransformType::CurrentLocal, true, false);
 			}
 		}
 		return true;
diff --git a/Engine/Plugins/Animation/ControlRig/Source/ControlRig/Private/Rigs/RigHierarchyController.cpp b/Engine/Plugins/Animation/ControlRig/Source/ControlRig/Private/Rigs/RigHierarchyController.cpp
--- a/Engine/Plugins/Animation/ControlRig/Source/ControlRig/Private/Rigs/RigHierarchyController.cpp
+++ b/Engine/Plugins/Animation/ControlRig/Source/ControlRig/Private/Rigs/RigHierarchyController.cpp
@@ -191,10 +191,12 @@ FRigElementKey URigHierarchyController::AddBone(FName InName, FRigElementKey InP
 		if(bTransformInGlobal)
 		{
 			Hierarchy->SetTransform(NewElement, InTransform, ERigTransformType::InitialGlobal, true, false);
+			Hierarchy->SetTransform(NewElement, InTransform, ERigTransformType::CurrentGlobal, true, false);
 		}
 		else
 		{
 			Hierarchy->SetTransform(NewElement, InTransform, ERigTransformType::InitialLocal, true, false);
+			Hierarchy->SetTransform(NewElement, InTransform, ERigTransformType::CurrentLocal, true, false);
 		}
 
 		NewElement->Pose.Current = NewElement->Pose.Initial;
