Add additional FK control rig mode, ::Direct which does not rely on offset transforms but rather sets/gets control transforms directly (matching old behaviour before offset-transform refactor)

diff --git a/Engine/Plugins/Animation/ControlRig/Source/ControlRig/Private/Rigs/FKControlRig.cpp b/Engine/Plugins/Animation/ControlRig/Source/ControlRig/Private/Rigs/FKControlRig.cpp
--- a/Engine/Plugins/Animation/ControlRig/Source/ControlRig/Private/Rigs/FKControlRig.cpp
+++ b/Engine/Plugins/Animation/ControlRig/Source/ControlRig/Private/Rigs/FKControlRig.cpp
@@ -14,6 +14,7 @@
 UFKControlRig::UFKControlRig(const FObjectInitializer& ObjectInitializer)
 	: Super(ObjectInitializer)
 	, ApplyMode(EControlRigFKRigExecuteMode::Replace)
+	, CachedToggleApplyMode(EControlRigFKRigExecuteMode::Replace)
 {
 	bCopyHierarchyBeforeSetup = false;
 	bResetInitialTransformsBeforeSetup = false;
@@ -56,7 +57,7 @@ void UFKControlRig::ExecuteUnits(FRigUnitContext& InOutContext, const FName& InE
 			const FName ControlName = GetControlName(BoneElement->GetName());
 			const FRigElementKey ControlKey(ControlName, ERigElementType::Control);
 			const int32 ControlIndex = GetHierarchy()->GetIndex(ControlKey);
-			if (IsControlActive[ControlIndex])
+			if (ControlIndex != INDEX_NONE && IsControlActive[ControlIndex])
 			{
 				FRigControlElement* Control = GetHierarchy()->Get<FRigControlElement>(ControlIndex);
 				const FTransform LocalTransform = GetHierarchy()->GetLocalTransform(ControlIndex);
@@ -67,7 +68,6 @@ void UFKControlRig::ExecuteUnits(FRigUnitContext& InOutContext, const FName& InE
 						const FTransform OffsetTransform = GetHierarchy()->GetControlOffsetTransform(Control, ERigTransformType::InitialLocal);
 						FTransform Transform = LocalTransform * OffsetTransform;
 						Transform.NormalizeRotation();
-						Transform.SetScale3D(FVector::OneVector);
 						GetHierarchy()->SetTransform(BoneElement, Transform, ERigTransformType::CurrentLocal, true, false);
 						break;
 					}
@@ -76,10 +76,14 @@ void UFKControlRig::ExecuteUnits(FRigUnitContext& InOutContext, const FName& InE
 						const FTransform PreviousTransform = GetHierarchy()->GetTransform(BoneElement, ERigTransformType::CurrentLocal);
 						FTransform Transform = LocalTransform * PreviousTransform;
 						Transform.NormalizeRotation();
-						Transform.SetScale3D(FVector::OneVector);
 						GetHierarchy()->SetTransform(BoneElement, Transform, ERigTransformType::CurrentLocal, true, false);
 						break;
 					}
+					case EControlRigFKRigExecuteMode::Direct:
+					{
+						GetHierarchy()->SetTransform(BoneElement, LocalTransform, ERigTransformType::CurrentLocal, true, false);
+						break;
+					}
 				}
 			}
 			return true;
@@ -91,13 +95,14 @@ void UFKControlRig::ExecuteUnits(FRigUnitContext& InOutContext, const FName& InE
 			const FRigElementKey ControlKey(ControlName, ERigElementType::Control);
 			const int32 ControlIndex = GetHierarchy()->GetIndex(ControlKey);
 			
-			if (IsControlActive[ControlIndex])
+			if (ControlIndex != INDEX_NONE && IsControlActive[ControlIndex])
 			{
 				const float CurveValue = GetHierarchy()->GetControlValue(ControlIndex).Get<float>();
 
 				switch (ApplyMode)
 				{
 					case EControlRigFKRigExecuteMode::Replace:
+					case EControlRigFKRigExecuteMode::Direct:
 					{
 						GetHierarchy()->SetCurveValue(CurveElement, CurveValue, false /*bSetupUndo*/);
 						break;
@@ -134,16 +139,32 @@ void UFKControlRig::ExecuteUnits(FRigUnitContext& InOutContext, const FName& InE
 			const FRigElementKey ControlKey(ControlName, ERigElementType::Control);
 			const int32 ControlIndex = GetHierarchy()->GetIndex(ControlKey);
 			
-			// during inversion we assume Replace Mode
-			FRigControlElement* Control = GetHierarchy()->GetChecked<FRigControlElement>(ControlIndex);
-			const FTransform Offset = GetHierarchy()->GetControlOffsetTransform(Control, ERigTransformType::InitialLocal);
-			const FTransform Current = GetHierarchy()->GetTransform(BoneElement, ERigTransformType::CurrentLocal);
+			switch (ApplyMode)
+			{
+				case EControlRigFKRigExecuteMode::Replace:
+				case EControlRigFKRigExecuteMode::Additive:
+				{
+					// during inversion we assume Replace Mode
+					FRigControlElement* Control = GetHierarchy()->GetChecked<FRigControlElement>(ControlIndex);
+					const FTransform Offset = GetHierarchy()->GetControlOffsetTransform(Control, ERigTransformType::InitialLocal);
+					const FTransform Current = GetHierarchy()->GetTransform(BoneElement, ERigTransformType::CurrentLocal);
 			
-			FTransform Transform = Current.GetRelativeTransform(Offset);
-			Transform.NormalizeRotation();
-			Transform.SetScale3D(FVector::OneVector);
+					FTransform Transform = Current.GetRelativeTransform(Offset);
+					Transform.NormalizeRotation();
+
+					SetControlValue(ControlName, FRigControlValue::Make(FEulerTransform(Transform)), bNotify, Context, bSetupUndo);
+					break;
+				}
 
-			SetControlValue(ControlName, FRigControlValue::Make(FEulerTransform(Transform)), bNotify, Context, bSetupUndo);
+				case EControlRigFKRigExecuteMode::Direct:
+				{
+					FTransform Transform = GetHierarchy()->GetTransform(BoneElement, ERigTransformType::CurrentLocal);
+					Transform.NormalizeRotation();
+					SetControlValue(ControlName, FRigControlValue::Make(FEulerTransform(Transform)), bNotify, Context, bSetupUndo);
+
+					break;
+				}
+			}
 
 		}, true);
 
@@ -285,7 +306,6 @@ void UFKControlRig::CreateRigElements(const FReferenceSkeleton& InReferenceSkele
 			Settings.DisplayName = BoneName;
 
 			OffsetTransform.NormalizeRotation();
-			OffsetTransform.SetScale3D(FVector::OneVector);
 			
 			Controller->AddControl(ControlName, ParentKey, Settings, FRigControlValue::Make(FEulerTransform::Identity), OffsetTransform, FTransform::Identity, false);
 
@@ -368,14 +388,20 @@ void UFKControlRig::CreateRigElements(const USkeletalMesh* InReferenceMesh)
 	}
 }
 
+void UFKControlRig::SetApplyMode(EControlRigFKRigExecuteMode InMode)
+{
+	ApplyMode = InMode;
+}
+
 void UFKControlRig::ToggleApplyMode()
 {
 	if (ApplyMode == EControlRigFKRigExecuteMode::Additive)
 	{
-		ApplyMode = EControlRigFKRigExecuteMode::Replace;
+		ApplyMode = CachedToggleApplyMode;
 	}
 	else
 	{
+		CachedToggleApplyMode = ApplyMode;
 		ApplyMode = EControlRigFKRigExecuteMode::Additive;
 	}
 	if (ApplyMode == EControlRigFKRigExecuteMode::Additive)
diff --git a/Engine/Plugins/Animation/ControlRig/Source/ControlRig/Public/Rigs/FKControlRig.h b/Engine/Plugins/Animation/ControlRig/Source/ControlRig/Public/Rigs/FKControlRig.h
--- a/Engine/Plugins/Animation/ControlRig/Source/ControlRig/Public/Rigs/FKControlRig.h
+++ b/Engine/Plugins/Animation/ControlRig/Source/ControlRig/Public/Rigs/FKControlRig.h
@@ -27,6 +27,9 @@ enum class EControlRigFKRigExecuteMode: uint8
 	/** Applies the authored pose as an additive layer */
 	Additive,
 
+	/** Sets the current pose without the use of offset transforms */
+	Direct,
+
 	/** MAX - invalid */
 	Max UMETA(Hidden),
 };
@@ -53,6 +56,7 @@ public:
 	void SetControlActive(int32 Index, bool bActive);
 	void SetControlActive(const TArray<FFKBoneCheckInfo>& InBoneChecks);
 
+	void SetApplyMode(EControlRigFKRigExecuteMode InMode);
 	void ToggleApplyMode();
 	bool CanToggleApplyMode() const { return true; }
 	bool IsApplyModeAdditive() const { return ApplyMode == EControlRigFKRigExecuteMode::Additive; }
@@ -69,6 +73,7 @@ private:
 
 	UPROPERTY()
 	EControlRigFKRigExecuteMode ApplyMode;
+	EControlRigFKRigExecuteMode CachedToggleApplyMode;
 
 	friend class FControlRigInteractionTest;
 };
\ No newline at end of file
