Control Rig: Improvements on Rig Hierarchy

+ fix validity issues during post load, import from text
+ fix python unit tests and other tests creating incorrect controls
+ enable cache validity checker for large python test (using new global guard)
+ remove parent cache from multi parent element (and simplify transform dependencies)

diff --git a/Engine/Plugins/Animation/ControlRig/Source/ControlRig/Private/ControlRig.cpp b/Engine/Plugins/Animation/ControlRig/Source/ControlRig/Private/ControlRig.cpp
--- a/Engine/Plugins/Animation/ControlRig/Source/ControlRig/Private/ControlRig.cpp
+++ b/Engine/Plugins/Animation/ControlRig/Source/ControlRig/Private/ControlRig.cpp
@@ -775,7 +775,10 @@ void UControlRig::Execute(const EControlRigState InState, const FName& InEventNa
 					
 					if (bResetInitialTransformsBeforeSetup && !bSetupModeEnabled)
 					{
-						GetHierarchy()->CopyPose(CDO->GetHierarchy(), false, true);
+						// when copying the pose from the default we don't copy the weights over,
+						// since we don't want to reset the topology in case it has diverged from the
+						// CDO.
+						GetHierarchy()->CopyPose(CDO->GetHierarchy(), false, true, false);
 					}
 				}
 
diff --git a/Engine/Plugins/Animation/ControlRig/Source/ControlRig/Private/Rigs/RigHierarchy.cpp b/Engine/Plugins/Animation/ControlRig/Source/ControlRig/Private/Rigs/RigHierarchy.cpp
--- a/Engine/Plugins/Animation/ControlRig/Source/ControlRig/Private/Rigs/RigHierarchy.cpp
+++ b/Engine/Plugins/Animation/ControlRig/Source/ControlRig/Private/Rigs/RigHierarchy.cpp
@@ -85,6 +85,12 @@ FAutoConsoleCommandWithWorldAndArgs FCmdControlRigHierarchyTraceFrames
 
 const FRigBaseElementChildrenArray URigHierarchy::EmptyElementArray;
 
+#if URIGHIERARCHY_ENSURE_CACHE_VALIDITY
+bool URigHierarchy::bEnableValidityCheckbyDefault = true;
+#else
+bool URigHierarchy::bEnableValidityCheckbyDefault = false;
+#endif
+
 URigHierarchy::URigHierarchy()
 : TopologyVersion(0)
 , bEnableDirtyPropagation(true)
@@ -103,11 +109,7 @@ URigHierarchy::URigHierarchy()
 , TraceFramesLeft(0)
 , TraceFramesCaptured(0)
 #endif
-#if URIGHIERARCHY_ENSURE_CACHE_VALIDITY
-, bEnableCacheValidityCheck(true)
-#else
-, bEnableCacheValidityCheck(false)
-#endif
+, bEnableCacheValidityCheck(bEnableValidityCheckbyDefault)
 , HierarchyForCacheValidation()
 #if WITH_EDITOR
 , ExecuteContext(nullptr)
@@ -277,90 +279,8 @@ void URigHierarchy::PostLoad()
 	UObject::PostLoad();
 
 	LLM_SCOPE_BYNAME(TEXT("Animation/ControlRig"));
-	
-	struct Local
-	{
-		static bool NeedsCheck(const FRigLocalAndGlobalTransform& InTransform)
-		{
-			return !InTransform.Local.bDirty && !InTransform.Global.bDirty;
-		}
-	};
-
-	// we need to check the elements for integrity (global vs local) to be correct.
-	for(int32 ElementIndex = 0; ElementIndex < Elements.Num(); ElementIndex++)
-	{
-		FRigBaseElement* BaseElement = Elements[ElementIndex];
-
-		if(FRigControlElement* ControlElement = Cast<FRigControlElement>(BaseElement))
-		{
-			if(Local::NeedsCheck(ControlElement->Offset.Initial))
-			{
-				const FTransform ComputedGlobalTransform = SolveParentConstraints(
-					ControlElement->ParentConstraints, ERigTransformType::InitialGlobal,
-					ControlElement->Offset.Get(ERigTransformType::InitialLocal), true,
-					FTransform::Identity, false);
-
-				const FTransform CachedGlobalTransform = ControlElement->Offset.Get(ERigTransformType::InitialGlobal);
-				if(!FRigComputedTransform::Equals(ComputedGlobalTransform, CachedGlobalTransform, 0.01))
-				{
-					ControlElement->Offset.MarkDirty(ERigTransformType::InitialGlobal);
-				}
-			}
-
-			if(Local::NeedsCheck(ControlElement->Pose.Initial))
-			{
-				const FTransform ComputedGlobalTransform = SolveParentConstraints(
-					ControlElement->ParentConstraints, ERigTransformType::InitialGlobal,
-					GetControlOffsetTransform(ControlElement, ERigTransformType::InitialGlobal), true,
-					ControlElement->Pose.Get(ERigTransformType::InitialLocal), true);
-				
-				const FTransform CachedGlobalTransform = ControlElement->Pose.Get(ERigTransformType::InitialGlobal);
-				
-				if(!FRigComputedTransform::Equals(ComputedGlobalTransform, CachedGlobalTransform, 0.01))
-				{
-					// for nulls we perceive the local transform as less relevant
-					ControlElement->Pose.MarkDirty(ERigTransformType::InitialLocal);
-				}
-			}
 
-			// we also need to check the pose here - for controls it is a bit different than for other
-			// types.
-			continue;
-		}
-
-		if(FRigMultiParentElement* MultiParentElement = Cast<FRigMultiParentElement>(BaseElement))
-		{
-			if(Local::NeedsCheck(MultiParentElement->Pose.Initial))
-			{
-				const FTransform ComputedGlobalTransform = SolveParentConstraints(
-					MultiParentElement->ParentConstraints, ERigTransformType::InitialGlobal,
-					FTransform::Identity, false,
-					MultiParentElement->Pose.Get(ERigTransformType::InitialLocal), true);
-				
-				const FTransform CachedGlobalTransform = MultiParentElement->Pose.Get(ERigTransformType::InitialGlobal);
-				
-				if(!FRigComputedTransform::Equals(ComputedGlobalTransform, CachedGlobalTransform, 0.01))
-				{
-					// for nulls we perceive the local transform as less relevant
-					MultiParentElement->Pose.MarkDirty(ERigTransformType::InitialLocal);
-				}
-			}
-		}
-
-		if(FRigTransformElement* TransformElement = Cast<FRigTransformElement>(BaseElement))
-		{
-			if(Local::NeedsCheck(TransformElement->Pose.Initial))
-			{
-				const FTransform ParentTransform = GetParentTransform(TransformElement, ERigTransformType::InitialGlobal);
-				const FTransform ComputedGlobalTransform = TransformElement->Pose.Get(ERigTransformType::InitialLocal) * ParentTransform;
-				const FTransform CachedGlobalTransform = TransformElement->Pose.Get(ERigTransformType::InitialGlobal);
-				if(!FRigComputedTransform::Equals(ComputedGlobalTransform, CachedGlobalTransform, 0.01))
-				{
-					TransformElement->Pose.MarkDirty(ERigTransformType::InitialGlobal);
-				}
-			}
-		}
-	}
+	CleanupInvalidCaches();
 }
 
 void URigHierarchy::Reset()
@@ -390,6 +310,11 @@ void URigHierarchy::Reset()
 	{
 		Notify(ERigHierarchyNotification::HierarchyReset, nullptr);
 	}
+
+	if(HierarchyForCacheValidation)
+	{
+		HierarchyForCacheValidation->Reset();
+	}
 }
 
 void URigHierarchy::CopyHierarchy(URigHierarchy* InHierarchy)
@@ -536,17 +461,72 @@ void URigHierarchy::ClearListeningHierarchy()
 }
 #endif
 
-void URigHierarchy::CopyPose(URigHierarchy* InHierarchy, bool bCurrent, bool bInitial)
+void URigHierarchy::CopyPose(URigHierarchy* InHierarchy, bool bCurrent, bool bInitial, bool bWeights, bool bMatchPoseInGlobalIfNeeded)
 {
 	check(InHierarchy);
 	LLM_SCOPE_BYNAME(TEXT("Animation/ControlRig"));
 
+	// if we need to copy the weights but the hierarchies are topologically
+	// different we need to reset the topology. this is expensive and should
+	// only happen during construction of the hierarchy itself.
+	if(bWeights && (GetTopologyVersion() != InHierarchy->GetTopologyVersion()))
+	{
+		CopyHierarchy(InHierarchy);
+	}
+
+	const bool bPerformTopologyCheck = GetTopologyVersion() != InHierarchy->GetTopologyVersion();
 	for (int32 ElementIndex = 0; ElementIndex < Elements.Num(); ElementIndex++)
 	{
 		FRigBaseElement* Element = Elements[ElementIndex];
 		if(FRigBaseElement* OtherElement = InHierarchy->Find(Element->GetKey()))
 		{
-			Element->CopyPose(OtherElement, bCurrent, bInitial);
+			Element->CopyPose(OtherElement, bCurrent, bInitial, bWeights);
+
+			// if the topologies don't match and we are supposed to match
+			// elements in global space...
+			if(bMatchPoseInGlobalIfNeeded && bPerformTopologyCheck)
+			{
+				FRigMultiParentElement* MultiParentElementA = Cast<FRigMultiParentElement>(Element);
+				FRigMultiParentElement* MultiParentElementB = Cast<FRigMultiParentElement>(OtherElement);
+				if(MultiParentElementA && MultiParentElementB)
+				{
+					if(MultiParentElementA->ParentConstraints.Num() != MultiParentElementB->ParentConstraints.Num())
+					{
+						FRigControlElement* ControlElementA = Cast<FRigControlElement>(Element);
+						FRigControlElement* ControlElementB = Cast<FRigControlElement>(OtherElement);
+						if(ControlElementA && ControlElementB)
+						{
+							if(bCurrent)
+							{
+								ControlElementA->Offset.Set(ERigTransformType::CurrentGlobal, InHierarchy->GetControlOffsetTransform(ControlElementB, ERigTransformType::CurrentGlobal));
+								ControlElementA->Offset.MarkDirty(ERigTransformType::CurrentLocal);
+								ControlElementA->Pose.MarkDirty(ERigTransformType::CurrentGlobal);
+								ControlElementA->Shape.MarkDirty(ERigTransformType::CurrentGlobal);
+							}
+							if(bInitial)
+							{
+								ControlElementA->Offset.Set(ERigTransformType::InitialGlobal, InHierarchy->GetControlOffsetTransform(ControlElementB, ERigTransformType::InitialGlobal));
+								ControlElementA->Offset.MarkDirty(ERigTransformType::InitialLocal);
+								ControlElementA->Pose.MarkDirty(ERigTransformType::InitialGlobal);
+								ControlElementA->Shape.MarkDirty(ERigTransformType::InitialGlobal);
+							}
+						}
+						else
+						{
+							if(bCurrent)
+							{
+								MultiParentElementA->Pose.Set(ERigTransformType::CurrentGlobal, InHierarchy->GetTransform(MultiParentElementB, ERigTransformType::CurrentGlobal));
+								MultiParentElementA->Pose.MarkDirty(ERigTransformType::CurrentLocal);
+							}
+							if(bInitial)
+							{
+								MultiParentElementA->Pose.Set(ERigTransformType::InitialGlobal, InHierarchy->GetTransform(MultiParentElementB, ERigTransformType::InitialGlobal));
+								MultiParentElementA->Pose.MarkDirty(ERigTransformType::InitialLocal);
+							}
+						}
+					}
+				}
+			}
 		}
 	}
 
@@ -636,6 +616,15 @@ void URigHierarchy::ResetPoseToInitial(ERigElementType InTypeFilter)
 			}
 		}
 
+		// reset the weights to the initial values as well
+		if(FRigMultiParentElement* MultiParentElement = Cast<FRigMultiParentElement>(Elements[ElementIndex]))
+		{
+			for(FRigElementParentConstraint& ParentConstraint : MultiParentElement->ParentConstraints)
+			{
+				ParentConstraint.Weight = ParentConstraint.InitialWeight;
+			}
+		}
+
 		if(FRigControlElement* ControlElement = Cast<FRigControlElement>(Elements[ElementIndex]))
 		{
 			ControlElement->Offset.Current = ControlElement->Offset.Initial;
@@ -646,11 +635,6 @@ void URigHierarchy::ResetPoseToInitial(ERigElementType InTypeFilter)
 		{
 			TransformElement->Pose.Current = TransformElement->Pose.Initial;
 		}
-
-		if(FRigMultiParentElement* MultiParentElement = Cast<FRigMultiParentElement>(Elements[ElementIndex]))
-		{
-			MultiParentElement->Parent.Current = MultiParentElement->Parent.Initial;
-		}
 	}
 	
 	EnsureCacheValidity();
@@ -1335,7 +1319,6 @@ bool URigHierarchy::SetParentWeight(FRigBaseElement* InChild, int32 InParentInde
 			}
 
 			TargetWeight = InWeight;
-			MultiParentElement->Parent.MarkDirty(GlobalType);
 
 			if(FRigControlElement* ControlElement = Cast<FRigControlElement>(MultiParentElement))
 			{
@@ -1460,8 +1443,6 @@ bool URigHierarchy::SetParentWeightArray(FRigBaseElement* InChild,  const TArray
 				}
 			}
 
-			MultiParentElement->Parent.MarkDirty(GlobalType);
-
 			if(FRigControlElement* ControlElement = Cast<FRigControlElement>(MultiParentElement))
 			{
 				ControlElement->Offset.MarkDirty(GlobalType);
@@ -2166,6 +2147,16 @@ FTransform URigHierarchy::GetTransform(FRigTransformElement* InTransformElement,
 				const FTransform NewTransform = ComputeLocalControlValue(ControlElement, ControlElement->Pose.Get(OpposedType), GlobalType);
 				InTransformElement->Pose.Set(InTransformType, NewTransform);
 			}
+			else if(FRigMultiParentElement* MultiParentElement = Cast<FRigMultiParentElement>(InTransformElement))
+			{
+				// this is done for nulls and any element that can have more than one parent which 
+				// is not a control
+				const FTransform& GlobalTransform = MultiParentElement->Pose.Get(GlobalType);
+				const FTransform LocalTransform = InverseSolveParentConstraints(
+					GlobalTransform, 
+					MultiParentElement->ParentConstraints, GlobalType, FTransform::Identity);
+				MultiParentElement->Pose.Set(InTransformType, LocalTransform);
+			}
 			else
 			{
 				ParentTransform = GetParentTransform(InTransformElement, GlobalType);
@@ -2179,7 +2170,6 @@ FTransform URigHierarchy::GetTransform(FRigTransformElement* InTransformElement,
 		{
 			if(FRigControlElement* ControlElement = Cast<FRigControlElement>(InTransformElement))
 			{
-				
 				// using GetControlOffsetTransform to check dirty flag before accessing the transform
 				// note: no need to do the same for Pose.Local because there is already an ensure:
 				// "ensure(!InTransformElement->Pose.IsDirty(OpposedType));" above
@@ -2189,6 +2179,16 @@ FTransform URigHierarchy::GetTransform(FRigTransformElement* InTransformElement,
 					ControlElement->Pose.Get(OpposedType), true);
 				ControlElement->Pose.Set(InTransformType, NewTransform);
 			}
+			else if(FRigMultiParentElement* MultiParentElement = Cast<FRigMultiParentElement>(InTransformElement))
+			{
+				// this is done for nulls and any element that can have more than one parent which 
+				// is not a control
+				const FTransform NewTransform = SolveParentConstraints(
+					MultiParentElement->ParentConstraints, InTransformType,
+					FTransform::Identity, false,
+					MultiParentElement->Pose.Get(OpposedType), true);
+				MultiParentElement->Pose.Set(InTransformType, NewTransform);
+			}
 			else
 			{
 				ParentTransform = GetParentTransform(InTransformElement, GlobalType);
@@ -2374,18 +2374,51 @@ FTransform URigHierarchy::GetControlOffsetTransform(FRigControlElement* InContro
 
 		if(IsLocal(InTransformType))
 		{
+			const FTransform& GlobalTransform = InControlElement->Offset.Get(GlobalType);
 			const FTransform LocalTransform = InverseSolveParentConstraints(
-				InControlElement->Offset.Get(GlobalType), 
+				GlobalTransform, 
 				InControlElement->ParentConstraints, GlobalType, FTransform::Identity);
 			InControlElement->Offset.Set(InTransformType, LocalTransform);
+
+			if(bEnableCacheValidityCheck)
+			{
+				const FTransform ComputedTransform = SolveParentConstraints(
+					InControlElement->ParentConstraints, MakeGlobal(InTransformType),
+					LocalTransform, true,
+					FTransform::Identity, false);
+
+				const TArray<FString>& TransformTypeStrings = GetTransformTypeStrings();
+
+				checkf(FRigComputedTransform::Equals(GlobalTransform, ComputedTransform),
+					TEXT("Element '%s' Offset %s Cached vs Computed doesn't match. ('%s' <-> '%s')"),
+					*InControlElement->GetName().ToString(),
+					*TransformTypeStrings[(int32)InTransformType],
+					*GlobalTransform.ToString(), *ComputedTransform.ToString());
+			}
 		}
 		else
 		{
+			const FTransform& LocalTransform = InControlElement->Offset.Get(OpposedType); 
 			const FTransform GlobalTransform = SolveParentConstraints(
 				InControlElement->ParentConstraints, InTransformType,
-				InControlElement->Offset.Get(OpposedType), true,
+				LocalTransform, true,
 				FTransform::Identity, false);
 			InControlElement->Offset.Set(InTransformType, GlobalTransform);
+
+			if(bEnableCacheValidityCheck)
+			{
+				const FTransform ComputedTransform = InverseSolveParentConstraints(
+					GlobalTransform, 
+					InControlElement->ParentConstraints, GlobalType, FTransform::Identity);
+
+				const TArray<FString>& TransformTypeStrings = GetTransformTypeStrings();
+
+				checkf(FRigComputedTransform::Equals(LocalTransform, ComputedTransform),
+					TEXT("Element '%s' Offset %s Cached vs Computed doesn't match. ('%s' <-> '%s')"),
+					*InControlElement->GetName().ToString(),
+					*TransformTypeStrings[(int32)InTransformType],
+					*LocalTransform.ToString(), *ComputedTransform.ToString());
+			}
 		}
 
 		EnsureCacheValidity();
@@ -2523,11 +2556,15 @@ FTransform URigHierarchy::GetControlShapeTransform(FRigControlElement* InControl
 		const FTransform ParentTransform = GetTransform(InControlElement, GlobalType);
 		if(IsLocal(InTransformType))
 		{
-			InControlElement->Shape.Set(InTransformType, InControlElement->Shape.Get(OpposedType).GetRelativeTransform(ParentTransform));
+			FTransform LocalTransform = InControlElement->Shape.Get(OpposedType).GetRelativeTransform(ParentTransform);
+			LocalTransform.NormalizeRotation();
+			InControlElement->Shape.Set(InTransformType, LocalTransform);
 		}
 		else
 		{
-			InControlElement->Shape.Set(InTransformType, InControlElement->Shape.Get(OpposedType) * ParentTransform);
+			FTransform GlobalTransform = InControlElement->Shape.Get(OpposedType) * ParentTransform;
+			GlobalTransform.NormalizeRotation();
+			InControlElement->Shape.Set(InTransformType, GlobalTransform);
 		}
 
 		EnsureCacheValidity();
@@ -2718,23 +2755,16 @@ FTransform URigHierarchy::GetParentTransform(FRigBaseElement* InElement, const E
 	}
 	else if(FRigMultiParentElement* MultiParentElement = Cast<FRigMultiParentElement>(InElement))
 	{
-		FRigComputedTransform& Output = MultiParentElement->Parent[InTransformType];
-
-		if(Output.bDirty)
-		{
-			const FTransform OutputTransform = SolveParentConstraints(
-				MultiParentElement->ParentConstraints,
-				InTransformType,
-				FTransform::Identity,
-				false,
-				FTransform::Identity,
-				false
-			);
-			MultiParentElement->Parent.Set(InTransformType, OutputTransform);
-
-			EnsureCacheValidity();
-		}
-		return Output.Transform;
+		const FTransform OutputTransform = SolveParentConstraints(
+			MultiParentElement->ParentConstraints,
+			InTransformType,
+			FTransform::Identity,
+			false,
+			FTransform::Identity,
+			false
+		);
+		EnsureCacheValidity();
+		return OutputTransform;
 	}
 	return FTransform::Identity;
 }
@@ -3525,16 +3555,9 @@ void URigHierarchy::PropagateDirtyFlags(FRigTransformElement* InTransformElement
 			{
 				if(ERigTransformType::IsGlobal(TypeToDirty))
 				{
-					if(ControlElement->Parent.IsDirty(TypeToDirty) &&
-						ControlElement->Offset.IsDirty(TypeToDirty) &&
-						ControlElement->Pose.IsDirty(TypeToDirty))
-					{
-						continue;
-					}
-				}
-				else
-				{
-					if(ControlElement->Parent.IsDirty(TypeToDirty))
+					if(ControlElement->Offset.IsDirty(TypeToDirty) &&
+						ControlElement->Pose.IsDirty(TypeToDirty) &&
+						ControlElement->Shape.IsDirty(TypeToDirty))
 					{
 						continue;
 					}
@@ -3544,15 +3567,7 @@ void URigHierarchy::PropagateDirtyFlags(FRigTransformElement* InTransformElement
 			{
 				if(ERigTransformType::IsGlobal(TypeToDirty))
 				{
-					if(MultiParentElement->Parent.IsDirty(TypeToDirty) &&
-						MultiParentElement->Pose.IsDirty(TypeToDirty))
-					{
-						continue;
-					}
-				}
-				else
-				{
-					if(MultiParentElement->Parent.IsDirty(TypeToDirty))
+					if(MultiParentElement->Pose.IsDirty(TypeToDirty))
 					{
 						continue;
 					}
@@ -3598,16 +3613,9 @@ void URigHierarchy::PropagateDirtyFlags(FRigTransformElement* InTransformElement
 			{
 				if(ERigTransformType::IsGlobal(TypeToDirty))
 				{
-					if(ControlElement->Parent.IsDirty(TypeToDirty) &&
-						ControlElement->Offset.IsDirty(TypeToDirty) &&
-						ControlElement->Pose.IsDirty(TypeToDirty))
-					{
-						continue;
-					}
-				}
-				else
-				{
-					if(ControlElement->Parent.IsDirty(TypeToDirty))
+					if(ControlElement->Offset.IsDirty(TypeToDirty) &&
+						ControlElement->Pose.IsDirty(TypeToDirty) &&
+					    ControlElement->Shape.IsDirty(TypeToDirty))
 					{
 						continue;
 					}
@@ -3617,15 +3625,7 @@ void URigHierarchy::PropagateDirtyFlags(FRigTransformElement* InTransformElement
 			{
 				if(ERigTransformType::IsGlobal(TypeToDirty))
 				{
-					if(MultiParentElement->Parent.IsDirty(TypeToDirty) &&
-						MultiParentElement->Pose.IsDirty(TypeToDirty))
-					{
-						continue;
-					}
-				}
-				else
-				{
-					if(MultiParentElement->Parent.IsDirty(TypeToDirty))
+					if(MultiParentElement->Pose.IsDirty(TypeToDirty))
 					{
 						continue;
 					}
@@ -3650,10 +3650,6 @@ void URigHierarchy::PropagateDirtyFlags(FRigTransformElement* InTransformElement
 
 			ElementToDirty.Element->Pose.MarkDirty(TypeToDirty);
 		
-			if(FRigMultiParentElement* MultiParentElement = Cast<FRigMultiParentElement>(ElementToDirty.Element))
-			{
-				MultiParentElement->Parent.MarkDirty(GlobalType);
-			}
 			if(FRigControlElement* ControlElement = Cast<FRigControlElement>(ElementToDirty.Element))
 			{
 				ControlElement->Offset.MarkDirty(GlobalType);
@@ -3674,6 +3670,151 @@ void URigHierarchy::PropagateDirtyFlags(FRigTransformElement* InTransformElement
 	}
 }
 
+void URigHierarchy::CleanupInvalidCaches()
+{
+	// create a copy of this hierarchy and pre compute all transforms
+	if(HierarchyForCacheValidation == nullptr)
+	{
+		HierarchyForCacheValidation = NewObject<URigHierarchy>(this, NAME_None, RF_Transient);
+		HierarchyForCacheValidation->bEnableCacheValidityCheck = false;
+	}
+	HierarchyForCacheValidation->CopyHierarchy(this);
+
+	struct Local
+	{
+		static bool NeedsCheck(const FRigLocalAndGlobalTransform& InTransform)
+		{
+			return !InTransform.Local.bDirty && !InTransform.Global.bDirty;
+		}
+	};
+
+	// mark all elements' initial as dirty where needed
+	for(int32 ElementIndex = 0; ElementIndex < HierarchyForCacheValidation->Elements.Num(); ElementIndex++)
+	{
+		FRigBaseElement* BaseElement = HierarchyForCacheValidation->Elements[ElementIndex];
+		if(FRigControlElement* ControlElement = Cast<FRigControlElement>(BaseElement))
+		{
+			if(Local::NeedsCheck(ControlElement->Offset.Initial))
+			{
+				ControlElement->Offset.MarkDirty(ERigTransformType::InitialGlobal);
+			}
+
+			if(Local::NeedsCheck(ControlElement->Pose.Initial))
+			{
+				ControlElement->Pose.MarkDirty(ERigTransformType::InitialGlobal);
+			}
+
+			if(Local::NeedsCheck(ControlElement->Shape.Initial))
+			{
+				ControlElement->Shape.MarkDirty(ERigTransformType::InitialGlobal);
+			}
+			continue;
+		}
+
+		if(FRigMultiParentElement* MultiParentElement = Cast<FRigMultiParentElement>(BaseElement))
+		{
+			if(Local::NeedsCheck(MultiParentElement->Pose.Initial))
+			{
+				MultiParentElement->Pose.MarkDirty(ERigTransformType::InitialLocal);
+			}
+			continue;
+		}
+
+		if(FRigTransformElement* TransformElement = Cast<FRigTransformElement>(BaseElement))
+		{
+			if(Local::NeedsCheck(TransformElement->Pose.Initial))
+			{
+				TransformElement->Pose.MarkDirty(ERigTransformType::InitialGlobal);
+			}
+		}
+	}
+
+	// recompute 
+	HierarchyForCacheValidation->ComputeAllTransforms();
+
+	// we need to check the elements for integrity (global vs local) to be correct.
+	for(int32 ElementIndex = 0; ElementIndex < Elements.Num(); ElementIndex++)
+	{
+		FRigBaseElement* BaseElement = Elements[ElementIndex];
+
+		if(FRigControlElement* ControlElement = Cast<FRigControlElement>(BaseElement))
+		{
+			FRigControlElement* OtherControlElement = HierarchyForCacheValidation->FindChecked<FRigControlElement>(ControlElement->GetKey());
+			
+			if(Local::NeedsCheck(ControlElement->Offset.Initial))
+			{
+				const FTransform CachedGlobalTransform = OtherControlElement->Offset.Get(ERigTransformType::InitialGlobal);
+				const FTransform ComputedGlobalTransform = HierarchyForCacheValidation->GetControlOffsetTransform(OtherControlElement, ERigTransformType::InitialGlobal);
+
+				if(!FRigComputedTransform::Equals(ComputedGlobalTransform, CachedGlobalTransform, 0.01))
+				{
+					ControlElement->Offset.MarkDirty(ERigTransformType::InitialGlobal);
+				}
+			}
+
+			if(Local::NeedsCheck(ControlElement->Pose.Initial))
+			{
+				const FTransform CachedGlobalTransform = ControlElement->Pose.Get(ERigTransformType::InitialGlobal);
+				const FTransform ComputedGlobalTransform = HierarchyForCacheValidation->GetTransform(OtherControlElement, ERigTransformType::InitialGlobal);
+				
+				if(!FRigComputedTransform::Equals(ComputedGlobalTransform, CachedGlobalTransform, 0.01))
+				{
+					ControlElement->Pose.MarkDirty(ERigTransformType::InitialGlobal);
+				}
+			}
+
+			if(Local::NeedsCheck(ControlElement->Shape.Initial))
+			{
+				const FTransform CachedGlobalTransform = ControlElement->Shape.Get(ERigTransformType::InitialGlobal);
+				const FTransform ComputedGlobalTransform = HierarchyForCacheValidation->GetControlShapeTransform(OtherControlElement, ERigTransformType::InitialGlobal);
+				
+				if(!FRigComputedTransform::Equals(ComputedGlobalTransform, CachedGlobalTransform, 0.01))
+				{
+					ControlElement->Shape.MarkDirty(ERigTransformType::InitialGlobal);
+				}
+			}
+			continue;
+		}
+
+		if(FRigMultiParentElement* MultiParentElement = Cast<FRigMultiParentElement>(BaseElement))
+		{
+			FRigMultiParentElement* OtherMultiParentElement = HierarchyForCacheValidation->FindChecked<FRigMultiParentElement>(MultiParentElement->GetKey());
+
+			if(Local::NeedsCheck(MultiParentElement->Pose.Initial))
+			{
+				const FTransform CachedGlobalTransform = MultiParentElement->Pose.Get(ERigTransformType::InitialGlobal);
+				const FTransform ComputedGlobalTransform = HierarchyForCacheValidation->GetTransform(OtherMultiParentElement, ERigTransformType::InitialGlobal);
+				
+				if(!FRigComputedTransform::Equals(ComputedGlobalTransform, CachedGlobalTransform, 0.01))
+				{
+					// for nulls we perceive the local transform as less relevant
+					MultiParentElement->Pose.MarkDirty(ERigTransformType::InitialLocal);
+				}
+			}
+			continue;
+		}
+
+		if(FRigTransformElement* TransformElement = Cast<FRigTransformElement>(BaseElement))
+		{
+			FRigTransformElement* OtherTransformElement = HierarchyForCacheValidation->FindChecked<FRigTransformElement>(TransformElement->GetKey());
+
+			if(Local::NeedsCheck(TransformElement->Pose.Initial))
+			{
+				const FTransform CachedGlobalTransform = TransformElement->Pose.Get(ERigTransformType::InitialGlobal);
+				const FTransform ComputedGlobalTransform = HierarchyForCacheValidation->GetTransform(OtherTransformElement, ERigTransformType::InitialGlobal);
+				
+				if(!FRigComputedTransform::Equals(ComputedGlobalTransform, CachedGlobalTransform, 0.01))
+				{
+					TransformElement->Pose.MarkDirty(ERigTransformType::InitialGlobal);
+				}
+			}
+		}
+	}
+
+	ResetPoseToInitial(ERigElementType::All);
+	EnsureCacheValidity();
+}
+
 void URigHierarchy::EnsureCacheValidityImpl()
 {
 	LLM_SCOPE_BYNAME(TEXT("Animation/ControlRig"));
@@ -3683,14 +3824,7 @@ void URigHierarchy::EnsureCacheValidityImpl()
 	}
 	TGuardValue<bool> Guard(bEnableCacheValidityCheck, false);
 
-	static TArray<FString> TransformTypeStrings;
-	if(TransformTypeStrings.IsEmpty())
-	{
-		for(int32 TransformTypeIndex = 0; TransformTypeIndex < (int32) ERigTransformType::NumTransformTypes; TransformTypeIndex++)
-		{
-			TransformTypeStrings.Add(StaticEnum<ERigTransformType::Type>()->GetDisplayNameTextByValue((int64)TransformTypeIndex).ToString());
-		}
-	}
+	static const TArray<FString>& TransformTypeStrings = GetTransformTypeStrings();
 
 	// make sure that elements which are marked as dirty don't have fully cached children
 	ForEach<FRigTransformElement>([](FRigTransformElement* TransformElement)
@@ -3715,20 +3849,8 @@ void URigHierarchy::EnsureCacheValidityImpl()
 			{
 				if(FRigMultiParentElement* MultiParentElementToDirty = Cast<FRigMultiParentElement>(ElementToDirty.Element))
 				{
-					check(MultiParentElementToDirty->Parent.IsDirty(GlobalType) ||
-                        MultiParentElementToDirty->Parent.IsDirty(LocalType));
-
                     if(FRigControlElement* ControlElementToDirty = Cast<FRigControlElement>(ElementToDirty.Element))
                     {
-                        if(ControlElementToDirty->Parent.IsDirty(GlobalType))
-                        {
-                            checkf(ControlElementToDirty->Offset.IsDirty(GlobalType) ||
-                                    ControlElementToDirty->Offset.IsDirty(LocalType),
-                                    TEXT("Control '%s' %s Parent Cache is dirty, but the Offset is not."),
-                                    *ControlElementToDirty->GetKey().ToString(),
-                                    *TransformTypeString);
-                        }
-
                         if(ControlElementToDirty->Offset.IsDirty(GlobalType))
                         {
                             checkf(ControlElementToDirty->Pose.IsDirty(GlobalType) ||
@@ -3749,14 +3871,11 @@ void URigHierarchy::EnsureCacheValidityImpl()
                     }
                     else
                     {
-                        if(MultiParentElementToDirty->Parent.IsDirty(GlobalType))
-                        {
-                            checkf(MultiParentElementToDirty->Pose.IsDirty(GlobalType) ||
-                                    MultiParentElementToDirty->Pose.IsDirty(LocalType),
-                                    TEXT("MultiParent '%s' %s Parent Cache is dirty, but the Pose is not."),
-									*MultiParentElementToDirty->GetKey().ToString(),
-									*TransformTypeString);
-						}
+                        checkf(MultiParentElementToDirty->Pose.IsDirty(GlobalType) ||
+                                MultiParentElementToDirty->Pose.IsDirty(LocalType),
+                                TEXT("MultiParent '%s' %s Parent Cache is dirty, but the Pose is not."),
+								*MultiParentElementToDirty->GetKey().ToString(),
+								*TransformTypeString);
                     }
 				}
 				else
@@ -3783,7 +3902,7 @@ void URigHierarchy::EnsureCacheValidityImpl()
 	{
 		HierarchyForCacheValidation->CopyHierarchy(this);
 	}
-	HierarchyForCacheValidation->CopyPose(this, true, true);
+	HierarchyForCacheValidation->CopyPose(this, true, true, true);
 
 	// traverse the copied hierarchy and compare cached vs computed values
 	URigHierarchy* HierarchyForLambda = HierarchyForCacheValidation;
@@ -3791,32 +3910,6 @@ void URigHierarchy::EnsureCacheValidityImpl()
 	{
 		bContinue = true;
 
-		if(FRigMultiParentElement* MultiParentElement = Cast<FRigMultiParentElement>(Element))
-		{
-			for(int32 TransformTypeIndex = 0; TransformTypeIndex < (int32) ERigTransformType::NumTransformTypes; TransformTypeIndex++)
-			{
-				const ERigTransformType::Type TransformType = (ERigTransformType::Type)TransformTypeIndex;
-				const ERigTransformType::Type OpposedType = ERigTransformType::SwapLocalAndGlobal(TransformType);
-				const FString& TransformTypeString = TransformTypeStrings[TransformTypeIndex];
-
-				if(ERigTransformType::IsLocal(TransformType))
-				{
-					continue;
-				}
-
-				if(!MultiParentElement->Parent.IsDirty(TransformType))
-				{
-					const FTransform CachedTransform = HierarchyForLambda->GetParentTransform(MultiParentElement, TransformType);
-					MultiParentElement->Parent.MarkDirty(TransformType);
-					const FTransform ComputedTransform = HierarchyForLambda->GetParentTransform(MultiParentElement, TransformType);
-					checkf(FRigComputedTransform::Equals(CachedTransform, ComputedTransform),
-						TEXT("Element '%s' Parent %s Cached vs Computed doesn't match."),
-						*Element->GetName().ToString(),
-						*TransformTypeString);
-				}
-			}
-		}
-
 		if(FRigControlElement* ControlElement = Cast<FRigControlElement>(Element))
 		{
 			for(int32 TransformTypeIndex = 0; TransformTypeIndex < (int32) ERigTransformType::NumTransformTypes; TransformTypeIndex++)
@@ -3831,9 +3924,11 @@ void URigHierarchy::EnsureCacheValidityImpl()
 					ControlElement->Offset.MarkDirty(TransformType);
 					const FTransform ComputedTransform = HierarchyForLambda->GetControlOffsetTransform(ControlElement, TransformType);
 					checkf(FRigComputedTransform::Equals(CachedTransform, ComputedTransform),
-						TEXT("Element '%s' Offset %s Cached vs Computed doesn't match."),
+						TEXT("Element '%s' Offset %s Cached vs Computed doesn't match. ('%s' <-> '%s')"),
 						*Element->GetName().ToString(),
-						*TransformTypeString);
+						*TransformTypeString,
+						*CachedTransform.ToString(),
+						*ComputedTransform.ToString());
 				}
 			}
 		}
@@ -3852,9 +3947,10 @@ void URigHierarchy::EnsureCacheValidityImpl()
 					TransformElement->Pose.MarkDirty(TransformType);
 					const FTransform ComputedTransform = HierarchyForLambda->GetTransform(TransformElement, TransformType);
 					checkf(FRigComputedTransform::Equals(CachedTransform, ComputedTransform),
-						TEXT("Element '%s' Pose %s Cached vs Computed doesn't match."),
+						TEXT("Element '%s' Pose %s Cached vs Computed doesn't match. ('%s' <-> '%s')"),
 						*Element->GetName().ToString(),
-						*TransformTypeString);
+						*TransformTypeString,
+						*CachedTransform.ToString(), *ComputedTransform.ToString());
 				}
 			}
 		}
@@ -3873,9 +3969,10 @@ void URigHierarchy::EnsureCacheValidityImpl()
 					ControlElement->Shape.MarkDirty(TransformType);
 					const FTransform ComputedTransform = HierarchyForLambda->GetControlShapeTransform(ControlElement, TransformType);
 					checkf(FRigComputedTransform::Equals(CachedTransform, ComputedTransform),
-						TEXT("Element '%s' Shape %s Cached vs Computed doesn't match."),
+						TEXT("Element '%s' Shape %s Cached vs Computed doesn't match. ('%s' <-> '%s')"),
 						*Element->GetName().ToString(),
-						*TransformTypeString);
+						*TransformTypeString,
+						*CachedTransform.ToString(), *ComputedTransform.ToString());
 				}
 			}
 		}
@@ -4012,6 +4109,19 @@ URigHierarchy::TElementDependencyMap URigHierarchy::GetDependenciesForVM(URigVM*
 
 #endif
 
+const TArray<FString>& URigHierarchy::GetTransformTypeStrings()
+{
+	static TArray<FString> TransformTypeStrings;
+	if(TransformTypeStrings.IsEmpty())
+	{
+		for(int32 TransformTypeIndex = 0; TransformTypeIndex < (int32) ERigTransformType::NumTransformTypes; TransformTypeIndex++)
+		{
+			TransformTypeStrings.Add(StaticEnum<ERigTransformType::Type>()->GetDisplayNameTextByValue((int64)TransformTypeIndex).ToString());
+		}
+	}
+	return TransformTypeStrings;
+}
+
 void URigHierarchy::PushTransformToStack(const FRigElementKey& InKey, ERigTransformStackEntryType InEntryType,
                                          ERigTransformType::Type InTransformType, const FTransform& InOldTransform, const FTransform& InNewTransform,
                                          bool bAffectChildren, bool bModify)
@@ -4215,13 +4325,16 @@ void URigHierarchy::ComputeAllTransforms()
 		for(int32 TransformTypeIndex = 0; TransformTypeIndex < (int32) ERigTransformType::NumTransformTypes; TransformTypeIndex++)
 		{
 			const ERigTransformType::Type TransformType = (ERigTransformType::Type)TransformTypeIndex; 
+			if(FRigControlElement* ControlElement = Get<FRigControlElement>(ElementIndex))
+			{
+				GetControlOffsetTransform(ControlElement, TransformType);
+			}
 			if(FRigTransformElement* TransformElement = Get<FRigTransformElement>(ElementIndex))
 			{
 				GetTransform(TransformElement, TransformType);
 			}
 			if(FRigControlElement* ControlElement = Get<FRigControlElement>(ElementIndex))
 			{
-				GetControlOffsetTransform(ControlElement, TransformType);
 				GetControlShapeTransform(ControlElement, TransformType);
 			}
 		}
@@ -4271,6 +4384,17 @@ FTransform URigHierarchy::SolveParentConstraints(
 	FRigElementWeight TotalWeight(0.f);
 	ComputeParentConstraintIndices(InConstraints, InTransformType, FirstConstraint, SecondConstraint, NumConstraintsAffecting, TotalWeight);
 
+	// performance improvement for case of a single parent
+	if(NumConstraintsAffecting.Location == 1 &&
+		NumConstraintsAffecting.Rotation == 1 &&
+		NumConstraintsAffecting.Scale == 1 &&
+		FirstConstraint.Location == FirstConstraint.Rotation &&
+		FirstConstraint.Location == FirstConstraint.Scale)
+	{
+		return LazilyComputeParentConstraint(InConstraints, FirstConstraint.Location, InTransformType,
+			InLocalOffsetTransform, bApplyLocalOffsetTransform, InLocalPoseTransform, bApplyLocalPoseTransform);
+	}
+
 	if(NumConstraintsAffecting.Location == 0 ||
 		NumConstraintsAffecting.Rotation == 0 ||
 		NumConstraintsAffecting.Scale == 0)
@@ -4498,7 +4622,6 @@ FTransform URigHierarchy::InverseSolveParentConstraints(
 	const bool bInitial = IsInitial(InTransformType);
 	check(ERigTransformType::IsGlobal(InTransformType));
 
-
 	// collect all of the weights
 	FConstraintIndex FirstConstraint;
 	FConstraintIndex SecondConstraint;
@@ -4506,6 +4629,19 @@ FTransform URigHierarchy::InverseSolveParentConstraints(
 	FRigElementWeight TotalWeight(0.f);
 	ComputeParentConstraintIndices(InConstraints, InTransformType, FirstConstraint, SecondConstraint, NumConstraintsAffecting, TotalWeight);
 
+	// performance improvement for case of a single parent
+	if(NumConstraintsAffecting.Location == 1 &&
+		NumConstraintsAffecting.Rotation == 1 &&
+		NumConstraintsAffecting.Scale == 1 &&
+		FirstConstraint.Location == FirstConstraint.Rotation &&
+		FirstConstraint.Location == FirstConstraint.Scale)
+	{
+		const FTransform Transform = LazilyComputeParentConstraint(InConstraints, FirstConstraint.Location, InTransformType,
+			InLocalOffsetTransform, true, FTransform::Identity, false);
+
+		return InGlobalTransform.GetRelativeTransform(Transform);
+	}
+
 	if(NumConstraintsAffecting.Location == 0 ||
 		NumConstraintsAffecting.Rotation == 0 ||
 		NumConstraintsAffecting.Scale == 0)
@@ -4746,6 +4882,7 @@ FTransform URigHierarchy::LazilyComputeParentConstraint(
 			Transform = InLocalPoseTransform * Transform;
 		}
 
+		Transform.NormalizeRotation();
 		Constraint.Cache.Transform = Transform;
 		Constraint.Cache.bDirty = false;
 	}
@@ -4765,6 +4902,7 @@ void URigHierarchy::ComputeParentConstraintIndices(
 	// find all of the weights affecting this output
 	for(int32 ConstraintIndex = 0; ConstraintIndex < InConstraints.Num(); ConstraintIndex++)
 	{
+		// this is not relying on the cache whatsoever. we might as well remove it.
 		InConstraints[ConstraintIndex].Cache.bDirty = true;
 		
 		const FRigElementWeight& Weight = InConstraints[ConstraintIndex].GetWeight(bInitial);
diff --git a/Engine/Plugins/Animation/ControlRig/Source/ControlRig/Private/Rigs/RigHierarchyController.cpp b/Engine/Plugins/Animation/ControlRig/Source/ControlRig/Private/Rigs/RigHierarchyController.cpp
--- a/Engine/Plugins/Animation/ControlRig/Source/ControlRig/Private/Rigs/RigHierarchyController.cpp
+++ b/Engine/Plugins/Animation/ControlRig/Source/ControlRig/Private/Rigs/RigHierarchyController.cpp
@@ -258,8 +258,6 @@ FRigElementKey URigHierarchyController::AddNull(FName InName, FRigElementKey InP
 			Hierarchy->SetTransform(NewElement, InTransform, ERigTransformType::InitialLocal, true, false);
 		}
 
-		NewElement->Parent.MarkDirty(ERigTransformType::InitialGlobal);
-		NewElement->Parent.Current = NewElement->Parent.Initial;
 		NewElement->Pose.Current = NewElement->Pose.Initial;
 	}
 
@@ -325,11 +323,9 @@ FRigElementKey URigHierarchyController::AddControl(
 		NewElement->Shape.Set(ERigTransformType::InitialLocal, InShapeTransform);  
 		Hierarchy->SetControlValue(NewElement, InValue, ERigControlValueType::Initial, false);
 
-		NewElement->Parent.MarkDirty(ERigTransformType::InitialGlobal);
 		NewElement->Offset.MarkDirty(ERigTransformType::InitialGlobal);
 		NewElement->Pose.MarkDirty(ERigTransformType::InitialGlobal);
 		NewElement->Shape.MarkDirty(ERigTransformType::InitialGlobal);
-		NewElement->Parent.Current = NewElement->Parent.Initial;
 		NewElement->Offset.Current = NewElement->Offset.Initial;
 		NewElement->Pose.Current = NewElement->Pose.Initial;
 		NewElement->Shape.Current = NewElement->Shape.Initial;
@@ -1107,8 +1103,16 @@ TArray<FRigElementKey> URigHierarchyController::ImportFromText(FString InContent
 		{
 			if(FRigBaseElement* ExistingElement = Hierarchy->Find(NewElement->GetKey()))
 			{
-				ExistingElement->CopyPose(NewElement, true, true);
+				ExistingElement->CopyPose(NewElement, true, true, false);
 
+				if(FRigControlElement* ControlElement = Cast<FRigControlElement>(ExistingElement))
+				{
+					Hierarchy->GetControlShapeTransform(ControlElement, ERigTransformType::CurrentLocal);
+					Hierarchy->GetControlShapeTransform(ControlElement, ERigTransformType::InitialLocal);
+					ControlElement->Shape.MarkDirty(ERigTransformType::CurrentGlobal);
+					ControlElement->Shape.MarkDirty(ERigTransformType::InitialGlobal);
+				}
+				
 				TArray<FRigElementKey> CurrentParents = Hierarchy->GetParents(NewElement->GetKey());
 
 				bool bUpdateParents = CurrentParents.Num() != PerElementData.Parents.Num();
@@ -1586,6 +1590,14 @@ int32 URigHierarchyController::AddElement(FRigBaseElement* InElementToAdd, FRigB
 		SetParent(InElementToAdd, InFirstParent, bMaintainGlobalTransform);
 	}
 
+	if(FRigControlElement* ControlElement = Cast<FRigControlElement>(InElementToAdd))
+	{
+		Hierarchy->GetControlShapeTransform(ControlElement, ERigTransformType::CurrentLocal);
+		Hierarchy->GetControlShapeTransform(ControlElement, ERigTransformType::InitialLocal);
+		ControlElement->Shape.MarkDirty(ERigTransformType::CurrentGlobal);
+		ControlElement->Shape.MarkDirty(ERigTransformType::InitialGlobal);
+	}
+
 	// only notify once at the end
 	Notify(ERigHierarchyNotification::ElementAdded, InElementToAdd);
 
@@ -1959,6 +1971,8 @@ bool URigHierarchyController::AddParent(FRigBaseElement* InChild, FRigBaseElemen
 		return false;
 	}
 
+	Hierarchy->EnsureCacheValidity();
+
 	if(bRemoveAllParents)
 	{
 		RemoveAllParents(InChild, bMaintainGlobalTransform);		
@@ -1986,8 +2000,6 @@ bool URigHierarchyController::AddParent(FRigBaseElement* InChild, FRigBaseElemen
 
 		if(FRigControlElement* ControlElement = Cast<FRigControlElement>(InChild))
 		{
-			Hierarchy->GetControlOffsetTransform(ControlElement, ERigTransformType::CurrentLocal);
-			Hierarchy->GetControlOffsetTransform(ControlElement, ERigTransformType::InitialLocal);
 			Hierarchy->GetControlShapeTransform(ControlElement, ERigTransformType::CurrentLocal);
 			Hierarchy->GetControlShapeTransform(ControlElement, ERigTransformType::InitialLocal);
 		}
@@ -2030,9 +2042,6 @@ bool URigHierarchyController::AddParent(FRigBaseElement* InChild, FRigBaseElemen
 
 		if(InWeight > SMALL_NUMBER)
 		{
-			MultiParentElement->Parent.MarkDirty(ERigTransformType::CurrentGlobal);  
-			MultiParentElement->Parent.MarkDirty(ERigTransformType::InitialGlobal);
-
 			if(FRigControlElement* ControlElement = Cast<FRigControlElement>(MultiParentElement))
 			{
 				ControlElement->Offset.MarkDirty(ERigTransformType::CurrentGlobal);  
@@ -2220,9 +2229,6 @@ bool URigHierarchyController::RemoveParent(FRigBaseElement* InChild, FRigBaseEle
 				}
 			}
 
-			MultiParentElement->Parent.MarkDirty(ERigTransformType::CurrentGlobal);  
-			MultiParentElement->Parent.MarkDirty(ERigTransformType::InitialGlobal);  
-
 			if(FRigControlElement* ControlElement = Cast<FRigControlElement>(MultiParentElement))
 			{
 				ControlElement->Offset.MarkDirty(ERigTransformType::CurrentGlobal);  
diff --git a/Engine/Plugins/Animation/ControlRig/Source/ControlRig/Private/Rigs/RigHierarchyElements.cpp b/Engine/Plugins/Animation/ControlRig/Source/ControlRig/Private/Rigs/RigHierarchyElements.cpp
--- a/Engine/Plugins/Animation/ControlRig/Source/ControlRig/Private/Rigs/RigHierarchyElements.cpp
+++ b/Engine/Plugins/Animation/ControlRig/Source/ControlRig/Private/Rigs/RigHierarchyElements.cpp
@@ -159,9 +159,9 @@ void FRigTransformElement::Load(FArchive& Ar, URigHierarchy* Hierarchy, ESeriali
 	}
 }
 
-void FRigTransformElement::CopyPose(FRigBaseElement* InOther, bool bCurrent, bool bInitial)
+void FRigTransformElement::CopyPose(FRigBaseElement* InOther, bool bCurrent, bool bInitial, bool bWeights)
 {
-	Super::CopyPose(InOther, bCurrent, bInitial);
+	Super::CopyPose(InOther, bCurrent, bInitial, bWeights);
 
 	if(FRigTransformElement* Other = Cast<FRigTransformElement>(InOther))
 	{
@@ -259,8 +259,6 @@ void FRigMultiParentElement::Save(FArchive& Ar, URigHierarchy* Hierarchy, ESeria
 
 	if(SerializationPhase == ESerializationPhase::StaticData)
 	{
-		Parent.Save(Ar);
-
 		int32 NumParents = ParentConstraints.Num();
 		Ar << NumParents;
 	}
@@ -287,7 +285,11 @@ void FRigMultiParentElement::Load(FArchive& Ar, URigHierarchy* Hierarchy, ESeria
 
 	if(SerializationPhase == ESerializationPhase::StaticData)
 	{
-		Parent.Load(Ar);
+		if (Ar.CustomVer(FControlRigObjectVersion::GUID) < FControlRigObjectVersion::RemovedMultiParentParentCache)
+		{
+			FRigCurrentAndInitialTransform Parent;
+			Parent.Load(Ar);
+		}
 
 		int32 NumParents = 0;
 		Ar << NumParents;
@@ -332,7 +334,6 @@ void FRigMultiParentElement::CopyFrom(URigHierarchy* InHierarchy, FRigBaseElemen
 	Super::CopyFrom(InHierarchy, InOther, InOtherHierarchy);
 	
 	const FRigMultiParentElement* Source = CastChecked<FRigMultiParentElement>(InOther);
-	Parent = Source->Parent;
 	ParentConstraints.Reset();
 	ParentConstraints.Reserve(Source->ParentConstraints.Num());
 	IndexLookup.Reset();
@@ -349,19 +350,22 @@ void FRigMultiParentElement::CopyFrom(URigHierarchy* InHierarchy, FRigBaseElemen
 	}
 }
 
-void FRigMultiParentElement::CopyPose(FRigBaseElement* InOther, bool bCurrent, bool bInitial)
+void FRigMultiParentElement::CopyPose(FRigBaseElement* InOther, bool bCurrent, bool bInitial, bool bWeights)
 {
-	Super::CopyPose(InOther, bCurrent, bInitial);
+	Super::CopyPose(InOther, bCurrent, bInitial, bWeights);
 
-	if(FRigMultiParentElement* Other = Cast<FRigMultiParentElement>(InOther))
+	if(bWeights)
 	{
-		if(bCurrent)
+		FRigMultiParentElement* Source = Cast<FRigMultiParentElement>(InOther);
+		if(ensure(Source))
 		{
-			Parent.Current = Other->Parent.Current;
-		}
-		if(bInitial)
-		{
-			Parent.Initial = Other->Parent.Initial;
+			if(ensure(ParentConstraints.Num() == Source->ParentConstraints.Num()))
+			{
+				for(int32 ParentIndex = 0; ParentIndex < ParentConstraints.Num(); ParentIndex++)
+				{
+					ParentConstraints[ParentIndex].CopyPose(Source->ParentConstraints[ParentIndex], bCurrent, bInitial);
+				}
+			}
 		}
 	}
 }
@@ -717,9 +721,9 @@ void FRigControlElement::CopyFrom(URigHierarchy* InHierarchy, FRigBaseElement* I
 	Shape = Source->Shape;
 }
 
-void FRigControlElement::CopyPose(FRigBaseElement* InOther, bool bCurrent, bool bInitial)
+void FRigControlElement::CopyPose(FRigBaseElement* InOther, bool bCurrent, bool bInitial, bool bWeights)
 {
-	Super::CopyPose(InOther, bCurrent, bInitial);
+	Super::CopyPose(InOther, bCurrent, bInitial, bWeights);
 	
 	if(FRigControlElement* Other = Cast<FRigControlElement>(InOther))
 	{
@@ -760,9 +764,9 @@ void FRigCurveElement::Load(FArchive& Ar, URigHierarchy* Hierarchy, ESerializati
 	}
 }
 
-void FRigCurveElement::CopyPose(FRigBaseElement* InOther, bool bCurrent, bool bInitial)
+void FRigCurveElement::CopyPose(FRigBaseElement* InOther, bool bCurrent, bool bInitial, bool bWeights)
 {
-	Super::CopyPose(InOther, bCurrent, bInitial);
+	Super::CopyPose(InOther, bCurrent, bInitial, bWeights);
 	
 	if(FRigCurveElement* Other = Cast<FRigCurveElement>(InOther))
 	{
@@ -859,9 +863,9 @@ FTransform FRigReferenceElement::GetReferenceWorldTransform(const FRigUnitContex
 	return FTransform::Identity;
 }
 
-void FRigReferenceElement::CopyPose(FRigBaseElement* InOther, bool bCurrent, bool bInitial)
+void FRigReferenceElement::CopyPose(FRigBaseElement* InOther, bool bCurrent, bool bInitial, bool bWeights)
 {
-	Super::CopyPose(InOther, bCurrent, bInitial);
+	Super::CopyPose(InOther, bCurrent, bInitial, bWeights);
 	
 	if(FRigReferenceElement* Other = Cast<FRigReferenceElement>(InOther))
 	{
diff --git a/Engine/Plugins/Animation/ControlRig/Source/ControlRig/Private/Units/Execution/RigUnit_DynamicHierarchy.cpp b/Engine/Plugins/Animation/ControlRig/Source/ControlRig/Private/Units/Execution/RigUnit_DynamicHierarchy.cpp
--- a/Engine/Plugins/Animation/ControlRig/Source/ControlRig/Private/Units/Execution/RigUnit_DynamicHierarchy.cpp
+++ b/Engine/Plugins/Animation/ControlRig/Source/ControlRig/Private/Units/Execution/RigUnit_DynamicHierarchy.cpp
@@ -3,6 +3,7 @@
 #include "RigUnit_DynamicHierarchy.h"
 #include "Rigs/RigHierarchyController.h"
 #include "Units/RigUnitContext.h"
+#include "ControlRig/Private/Units/Execution/RigUnit_PrepareForExecution.h"
 
 FRigUnit_AddParent_Execute()
 {
@@ -110,6 +111,19 @@ FRigUnit_SwitchParent_Execute()
 				}
 				return;
 			}
+
+			// during setup event also change the initial weights
+			if(ExecuteContext.EventName == FRigUnit_PrepareForExecution::EventName)
+			{
+				if(!ExecuteContext.Hierarchy->SwitchToParent(ChildElement, ParentElement, true, true, DependencyMap, &FailureReason))
+				{
+					if(!FailureReason.IsEmpty())
+					{
+						UE_CONTROLRIG_RIGUNIT_REPORT_WARNING(TEXT("%s"), *FailureReason);
+					}
+					return;
+				}
+			}
 			break;
 		}
 	}
@@ -161,4 +175,10 @@ FRigUnit_HierarchySetParentWeights_Execute()
 	}
 
 	ExecuteContext.Hierarchy->SetParentWeightArray(ChildElement, Weights, false, true);
+
+	// during setup event also change the initial weights
+	if(ExecuteContext.EventName == FRigUnit_PrepareForExecution::EventName)
+	{
+		ExecuteContext.Hierarchy->SetParentWeightArray(ChildElement, Weights, true, true);
+	}
 }
diff --git a/Engine/Plugins/Animation/ControlRig/Source/ControlRig/Public/ControlRigObjectVersion.h b/Engine/Plugins/Animation/ControlRig/Source/ControlRig/Public/ControlRigObjectVersion.h
--- a/Engine/Plugins/Animation/ControlRig/Source/ControlRig/Public/ControlRigObjectVersion.h
+++ b/Engine/Plugins/Animation/ControlRig/Source/ControlRig/Public/ControlRigObjectVersion.h
@@ -51,6 +51,9 @@ struct CONTROLRIG_API FControlRigObjectVersion
 		// Switch limit control over to per channel limits 
 		PerChannelLimits,
 
+		// Removed the parent cache for multi parent elements 
+		RemovedMultiParentParentCache,
+
 		// -----<new versions can be added above this line>-------------------------------------------------
 		VersionPlusOne,
 		LatestVersion = VersionPlusOne - 1
diff --git a/Engine/Plugins/Animation/ControlRig/Source/ControlRig/Public/Rigs/RigHierarchy.h b/Engine/Plugins/Animation/ControlRig/Source/ControlRig/Public/Rigs/RigHierarchy.h
--- a/Engine/Plugins/Animation/ControlRig/Source/ControlRig/Public/Rigs/RigHierarchy.h
+++ b/Engine/Plugins/Animation/ControlRig/Source/ControlRig/Public/Rigs/RigHierarchy.h
@@ -9,6 +9,7 @@
 #include "RigHierarchyPose.h"
 #include "UObject/WeakObjectPtrTemplates.h"
 #include "EdGraph/EdGraphPin.h"
+#include "RigHierarchyDefines.h"
 #if WITH_EDITOR
 #include "RigVMPythonUtils.h"
 #endif
@@ -17,15 +18,6 @@
 class URigHierarchy;
 class URigHierarchyController;
 
-// Define to switch between recursive calls for dirty propagation or flat iteration.
-// Based on this each element will either contains only the next tier of elements to dirty (recursive)
-// or a full flattened list of elements to dirty (no need to recurse)
-#define URIGHIERARCHY_RECURSIVE_DIRTY_PROPAGATION 1
-
-// Debug define which performs a full check on the cache validity for all elements of the hierarchy.
-// This can be useful for debugging cache validity bugs.
-#define URIGHIERARCHY_ENSURE_CACHE_VALIDITY 0
-
 DECLARE_MULTICAST_DELEGATE_ThreeParams(FRigHierarchyModifiedEvent, ERigHierarchyNotification /* type */, URigHierarchy* /* hierarchy */, const FRigBaseElement* /* element */);
 DECLARE_EVENT_FiveParams(URigHierarchy, FRigHierarchyUndoRedoTransformEvent, URigHierarchy*, const FRigElementKey&, ERigTransformType::Type, const FTransform&, bool /* bUndo */);
 
@@ -152,7 +144,7 @@ public:
 	 * Copies the contents of a hierarchy onto this one
 	 */
 	UFUNCTION(BlueprintCallable, Category = URigHierarchy)
-	void CopyPose(URigHierarchy* InHierarchy, bool bCurrent, bool bInitial);
+	void CopyPose(URigHierarchy* InHierarchy, bool bCurrent, bool bInitial, bool bWeights, bool bMatchPoseInGlobalIfNeeded = false);
 
 	/**
 	 * Update all elements that depend on external references
@@ -2563,6 +2555,8 @@ private:
 	void PropagateDirtyFlags(FRigTransformElement* InTransformElement, bool bInitial, bool bAffectChildren) const;
 #endif
 
+public:
+
 	/**
 	* Performs validation of the cache within the hierarchy on any mutation.
 	*/
@@ -2576,6 +2570,13 @@ private:
 		}
 #endif
 	}
+
+	/*
+	 * Cleans up caches after load
+	 */
+	void CleanupInvalidCaches();
+
+private:
 	
 	/**
 	 * The topology version of the hierarchy changes when elements are
@@ -2940,6 +2941,8 @@ protected:
 	
 	bool bEnableCacheValidityCheck;
 
+	static bool bEnableValidityCheckbyDefault;
+
 	UPROPERTY(transient)
 	TObjectPtr<URigHierarchy> HierarchyForCacheValidation;
 	
@@ -2961,10 +2964,13 @@ private:
 	
 #endif
 
+	static const TArray<FString>& GetTransformTypeStrings();
+
 	friend class URigHierarchyController;
 	friend class UControlRig;
 	friend class FControlRigEditor;
 	friend struct FRigHierarchyValidityBracket;
+	friend struct FRigHierarchyGlobalValidityBracket;
 	friend struct FControlRigVisualGraphUtils;
 };
 
@@ -2999,6 +3005,25 @@ struct CONTROLRIG_API FRigHierarchyValidityBracket
 	TWeakObjectPtr<URigHierarchy> HierarchyPtr;
 };
 
+struct CONTROLRIG_API FRigHierarchyGlobalValidityBracket
+{
+public:
+	FRigHierarchyGlobalValidityBracket(bool bEnable = true)
+		: bPreviousValue(URigHierarchy::bEnableValidityCheckbyDefault)
+	{
+		URigHierarchy::bEnableValidityCheckbyDefault = true;
+	}
+
+	~FRigHierarchyGlobalValidityBracket()
+	{
+		URigHierarchy::bEnableValidityCheckbyDefault = bPreviousValue;
+	}
+
+private:
+
+	bool bPreviousValue;
+};
+
 template<>
 FORCEINLINE_DEBUGGABLE FVector2D URigHierarchy::GetControlValue(FRigControlElement* InControlElement, ERigControlValueType InValueType) const
 {
diff --git a/Engine/Plugins/Animation/ControlRig/Source/ControlRig/Public/Rigs/RigHierarchyDefines.h b/Engine/Plugins/Animation/ControlRig/Source/ControlRig/Public/Rigs/RigHierarchyDefines.h
--- a/Engine/Plugins/Animation/ControlRig/Source/ControlRig/Public/Rigs/RigHierarchyDefines.h
+++ b/Engine/Plugins/Animation/ControlRig/Source/ControlRig/Public/Rigs/RigHierarchyDefines.h
@@ -9,6 +9,15 @@
 
 class URigHierarchy;
 
+// Define to switch between recursive calls for dirty propagation or flat iteration.
+// Based on this each element will either contains only the next tier of elements to dirty (recursive)
+// or a full flattened list of elements to dirty (no need to recurse)
+#define URIGHIERARCHY_RECURSIVE_DIRTY_PROPAGATION 1
+
+// Debug define which performs a full check on the cache validity for all elements of the hierarchy.
+// This can be useful for debugging cache validity bugs.
+#define URIGHIERARCHY_ENSURE_CACHE_VALIDITY 0
+
 /* 
  * This is rig element types that we support
  * This can be used as a mask so supported as a bitfield
@@ -391,7 +400,16 @@ public:
 				Get<FVector3f>().X, Get<FVector3f>().Z, Get<FVector3f>().Y); break;
 			case ERigControlType::Scale: ValueStr = FString::Printf(TEXT("unreal.RigHierarchy.make_control_value_from_vector(unreal.Vector(%.6f, %.6f, %.6f))"),
 				Get<FVector3f>().X, Get<FVector3f>().Y, Get<FVector3f>().Z); break;
-			case ERigControlType::Transform:
+			case ERigControlType::Transform: ValueStr = FString::Printf(TEXT("unreal.RigHierarchy.make_control_value_from_euler_transform(unreal.EulerTransform(location=[%.6f,%.6f,%.6f],rotation=[%.6f,%.6f,%.6f],scale=[%.6f,%.6f,%.6f]))"),
+				Get<FTransform_Float>().TranslationX,
+				Get<FTransform_Float>().TranslationY,
+				Get<FTransform_Float>().TranslationZ,
+				Get<FTransform_Float>().GetRotation().Rotator().Pitch,
+				Get<FTransform_Float>().GetRotation().Rotator().Yaw,
+				Get<FTransform_Float>().GetRotation().Rotator().Roll,
+				Get<FTransform_Float>().ScaleX,
+				Get<FTransform_Float>().ScaleY,
+				Get<FTransform_Float>().ScaleZ); break;
 			case ERigControlType::EulerTransform: ValueStr = FString::Printf(TEXT("unreal.RigHierarchy.make_control_value_from_euler_transform(unreal.EulerTransform(location=[%.6f,%.6f,%.6f],rotation=[%.6f,%.6f,%.6f],scale=[%.6f,%.6f,%.6f]))"),
 				Get<FEulerTransform_Float>().TranslationX,
 				Get<FEulerTransform_Float>().TranslationY,
diff --git a/Engine/Plugins/Animation/ControlRig/Source/ControlRig/Public/Rigs/RigHierarchyElements.h b/Engine/Plugins/Animation/ControlRig/Source/ControlRig/Public/Rigs/RigHierarchyElements.h
--- a/Engine/Plugins/Animation/ControlRig/Source/ControlRig/Public/Rigs/RigHierarchyElements.h
+++ b/Engine/Plugins/Animation/ControlRig/Source/ControlRig/Public/Rigs/RigHierarchyElements.h
@@ -515,7 +515,7 @@ public:
 		return Element;
 	}
 
-	virtual void CopyPose(FRigBaseElement* InOther, bool bCurrent, bool bInitial) {}
+	virtual void CopyPose(FRigBaseElement* InOther, bool bCurrent, bool bInitial, bool bWeights) {}
 
 protected:
 
@@ -591,7 +591,7 @@ protected:
 
 public:
 	
-	virtual void CopyPose(FRigBaseElement* InOther, bool bCurrent, bool bInitial) override;
+	virtual void CopyPose(FRigBaseElement* InOther, bool bCurrent, bool bInitial, bool bWeights) override;
 	
 protected:
 
@@ -731,10 +731,26 @@ public:
 	{
 		return bInitial ? InitialWeight : Weight;
 	}
+
+	FORCEINLINE void CopyPose(const FRigElementParentConstraint& InOther, bool bCurrent, bool bInitial)
+	{
+		if(bCurrent)
+		{
+			Weight = InOther.Weight;
+		}
+		if(bInitial)
+		{
+			InitialWeight = InOther.InitialWeight;
+		}
+		Cache.bDirty = true;
+	}
 };
 
-//typedef TArray<FRigElementParentConstraint> FRigElementParentConstraintArray;
+#if URIGHIERARCHY_ENSURE_CACHE_VALIDITY
+typedef TArray<FRigElementParentConstraint, TInlineAllocator<8>> FRigElementParentConstraintArray;
+#else
 typedef TArray<FRigElementParentConstraint, TInlineAllocator<1>> FRigElementParentConstraintArray;
+#endif
 
 USTRUCT(BlueprintType)
 struct CONTROLRIG_API FRigMultiParentElement : public FRigTransformElement
@@ -753,9 +769,6 @@ public:
 	virtual void Save(FArchive& A, URigHierarchy* Hierarchy, ESerializationPhase SerializationPhase) override;
 	virtual void Load(FArchive& Ar, URigHierarchy* Hierarchy, ESerializationPhase SerializationPhase) override;
 	
-	UPROPERTY(BlueprintReadOnly, Category = RigElement)
-	FRigCurrentAndInitialTransform Parent;
-
 	FRigElementParentConstraintArray ParentConstraints;
 	TMap<FRigElementKey, int32> IndexLookup;
 
@@ -771,7 +784,7 @@ protected:
 
 public:
 	
-	virtual void CopyPose(FRigBaseElement* InOther, bool bCurrent, bool bInitial) override;
+	virtual void CopyPose(FRigBaseElement* InOther, bool bCurrent, bool bInitial, bool bWeights) override;
 
 protected:
 	
@@ -1013,7 +1026,7 @@ protected:
 
 public:
 	
-	virtual void CopyPose(FRigBaseElement* InOther, bool bCurrent, bool bInitial) override;
+	virtual void CopyPose(FRigBaseElement* InOther, bool bCurrent, bool bInitial, bool bWeights) override;
 
 protected:
 
@@ -1055,7 +1068,7 @@ public:
 
 public:
 	
-	virtual void CopyPose(FRigBaseElement* InOther, bool bCurrent, bool bInitial) override;
+	virtual void CopyPose(FRigBaseElement* InOther, bool bCurrent, bool bInitial, bool bWeights) override;
 
 protected:
 	
@@ -1136,7 +1149,7 @@ public:
 
 	FTransform GetReferenceWorldTransform(const FRigUnitContext* InContext, bool bInitial) const;
 
-	virtual void CopyPose(FRigBaseElement* InOther, bool bCurrent, bool bInitial) override;
+	virtual void CopyPose(FRigBaseElement* InOther, bool bCurrent, bool bInitial, bool bWeights) override;
 
 protected:
 
diff --git a/Engine/Plugins/Animation/ControlRig/Source/ControlRigDeveloper/Private/ControlRigBlueprint.cpp b/Engine/Plugins/Animation/ControlRig/Source/ControlRigDeveloper/Private/ControlRigBlueprint.cpp
--- a/Engine/Plugins/Animation/ControlRig/Source/ControlRigDeveloper/Private/ControlRigBlueprint.cpp
+++ b/Engine/Plugins/Animation/ControlRig/Source/ControlRigDeveloper/Private/ControlRigBlueprint.cpp
@@ -189,6 +189,7 @@ UClass* UControlRigBlueprint::RegenerateClass(UClass* ClassToRegenerate, UObject
 		TGuardValue<bool> NotificationGuard(bSuspendAllNotifications, true);
 		Result = Super::RegenerateClass(ClassToRegenerate, PreviousCDO);
 	}
+	Hierarchy->CleanupInvalidCaches();
 	PropagateHierarchyFromBPToInstances();
 	return Result;
 }
@@ -3737,7 +3738,7 @@ void UControlRigBlueprint::PropagatePoseFromInstanceToBP(UControlRig* InControlR
 	// current transforms in BP and CDO are meaningless, no need to copy them
 	// we use BP hierarchy to initialize CDO and instances' hierarchy, 
 	// so it should always be in the initial state.
-	Hierarchy->CopyPose(InControlRig->GetHierarchy(), false, true);
+	Hierarchy->CopyPose(InControlRig->GetHierarchy(), false, true, false, true);
 }
 
 void UControlRigBlueprint::PropagatePoseFromBPToInstances()
@@ -3747,7 +3748,7 @@ void UControlRigBlueprint::PropagatePoseFromBPToInstances()
 		if (UControlRig* DefaultObject = Cast<UControlRig>(MyControlRigClass->GetDefaultObject(false)))
 		{
 			DefaultObject->PostInitInstanceIfRequired();
-			DefaultObject->GetHierarchy()->CopyPose(Hierarchy, true, true);
+			DefaultObject->GetHierarchy()->CopyPose(Hierarchy, true, true, true);
 
 			TArray<UObject*> ArchetypeInstances;
 			DefaultObject->GetArchetypeInstances(ArchetypeInstances);
@@ -3756,7 +3757,7 @@ void UControlRigBlueprint::PropagatePoseFromBPToInstances()
 				if (UControlRig* InstanceRig = Cast<UControlRig>(ArchetypeInstance))
 				{
 					InstanceRig->PostInitInstanceIfRequired();
-					InstanceRig->GetHierarchy()->CopyPose(Hierarchy, true, true);
+					InstanceRig->GetHierarchy()->CopyPose(Hierarchy, true, true, true);
 				}
 			}
 		}
diff --git a/Engine/Plugins/Animation/ControlRig/Source/ControlRigEditor/Private/Editor/ControlRigEditor.cpp b/Engine/Plugins/Animation/ControlRig/Source/ControlRigEditor/Private/Editor/ControlRigEditor.cpp
--- a/Engine/Plugins/Animation/ControlRig/Source/ControlRigEditor/Private/Editor/ControlRigEditor.cpp
+++ b/Engine/Plugins/Animation/ControlRig/Source/ControlRigEditor/Private/Editor/ControlRigEditor.cpp
@@ -1301,7 +1301,7 @@ void FControlRigEditor::HandleSetObjectBeingDebugged(UObject* InObject)
 			{
 				// copy the initial transforms back to the blueprint
 				// no need to call modify here since this code only modifies the bp if the preview mesh changed
-				RigBlueprint->Hierarchy->CopyPose(DebuggedControlRig->GetHierarchy(), false, true);
+				RigBlueprint->Hierarchy->CopyPose(DebuggedControlRig->GetHierarchy(), false, true, false);
 			}
 		}
 
@@ -5884,7 +5884,9 @@ void FControlRigEditor::HandleOnControlModified(UControlRig* Subject, FRigContro
 		TargetControlElement->Settings = SourceControlElement->Settings;
 		Blueprint->Hierarchy->OnModified().Broadcast(ERigHierarchyNotification::ControlSettingChanged, Blueprint->Hierarchy, TargetControlElement);
 
-		TargetControlElement->CopyPose(SourceControlElement, true, true);
+		// we copy the pose including the weights since we want the topology to align during setup mode.
+		// i.e. dynamic reparenting should be reset here.
+		TargetControlElement->CopyPose(SourceControlElement, true, true, true);
 	}
 }
 
