Control Rig: Relay decisions around searchterm to schema

diff --git a/Engine/Plugins/Animation/ControlRig/Source/ControlRigDeveloper/Private/Graph/ControlRigGraphSchema.cpp b/Engine/Plugins/Animation/ControlRig/Source/ControlRigDeveloper/Private/Graph/ControlRigGraphSchema.cpp
--- a/Engine/Plugins/Animation/ControlRig/Source/ControlRigDeveloper/Private/Graph/ControlRigGraphSchema.cpp
+++ b/Engine/Plugins/Animation/ControlRig/Source/ControlRigDeveloper/Private/Graph/ControlRigGraphSchema.cpp
@@ -1000,6 +1000,19 @@ bool UControlRigGraphSchema::CanGraphBeDropped(TSharedPtr<FEdGraphSchemaAction>
 	return false;
 }
 
+FString UControlRigGraphSchema::GetFindReferenceSearchTerm(const FEdGraphSchemaAction* InGraphAction) const
+{
+	if(InGraphAction)
+	{
+		if(InGraphAction->GetTypeId() == FEdGraphSchemaAction_K2Var::StaticGetTypeId())
+		{
+			const FEdGraphSchemaAction_K2Var* VarAction = (const FEdGraphSchemaAction_K2Var*)InGraphAction;
+			return VarAction->GetVariableName().ToString();
+		}
+	}
+	return Super::GetFindReferenceSearchTerm(InGraphAction);
+}
+
 FReply UControlRigGraphSchema::BeginGraphDragAction(TSharedPtr<FEdGraphSchemaAction> InAction, const FPointerEvent& MouseEvent) const
 {
 	if (!InAction.IsValid())
diff --git a/Engine/Plugins/Animation/ControlRig/Source/ControlRigDeveloper/Public/Graph/ControlRigGraphSchema.h b/Engine/Plugins/Animation/ControlRig/Source/ControlRigDeveloper/Public/Graph/ControlRigGraphSchema.h
--- a/Engine/Plugins/Animation/ControlRig/Source/ControlRigDeveloper/Public/Graph/ControlRigGraphSchema.h
+++ b/Engine/Plugins/Animation/ControlRig/Source/ControlRigDeveloper/Public/Graph/ControlRigGraphSchema.h
@@ -201,6 +201,7 @@ public:
 	virtual void BreakPinLinks(UEdGraphPin& TargetPin, bool bSendsNodeNotifcation) const override;
 	virtual void BreakSinglePinLink(UEdGraphPin* SourcePin, UEdGraphPin* TargetPin) const override;
 	virtual bool CanGraphBeDropped(TSharedPtr<FEdGraphSchemaAction> InAction) const override;
+	virtual FString GetFindReferenceSearchTerm(const FEdGraphSchemaAction* InGraphAction) const override;
 	virtual FReply BeginGraphDragAction(TSharedPtr<FEdGraphSchemaAction> InAction, const FPointerEvent& MouseEvent = FPointerEvent() ) const override;
 	virtual class FConnectionDrawingPolicy* CreateConnectionDrawingPolicy(int32 InBackLayerID, int32 InFrontLayerID, float InZoomFactor, const FSlateRect& InClippingRect, class FSlateWindowElementList& InDrawElements, class UEdGraph* InGraphObj) const override;
 	virtual bool ShouldHidePinDefaultValue(UEdGraphPin* Pin) const override;
diff --git a/Engine/Source/Editor/Kismet/Private/SMyBlueprint.cpp b/Engine/Source/Editor/Kismet/Private/SMyBlueprint.cpp
--- a/Engine/Source/Editor/Kismet/Private/SMyBlueprint.cpp
+++ b/Engine/Source/Editor/Kismet/Private/SMyBlueprint.cpp
@@ -2705,8 +2705,30 @@ bool SMyBlueprint::IsImplementationDesiredAsFunction(const UFunction* OverrideFu
 
 void SMyBlueprint::OnFindReference()
 {
-	bool bUseQuotes = true;
 	FString SearchTerm;
+	for(const UEdGraph* UberGraph : GetBlueprintObj()->UbergraphPages)
+	{
+		if(const UEdGraphSchema* Schema = UberGraph->GetSchema())
+		{
+			TArray<TSharedPtr<FEdGraphSchemaAction> > SelectedActions;
+			GraphActionMenu->GetSelectedActions(SelectedActions);
+			if(SelectedActions.Num() == 1)
+			{
+				if (const FEdGraphSchemaAction* GraphAction = SelectedActions[0].Get())
+				{
+					SearchTerm = Schema->GetFindReferenceSearchTerm(GraphAction);
+					if(!SearchTerm.IsEmpty())
+					{
+						BlueprintEditorPtr.Pin()->SummonSearchUI(true, SearchTerm);
+						return;
+					}
+				}
+			}
+		}
+	}
+
+	bool bUseQuotes = true;
+
 	if (FEdGraphSchemaAction_K2Graph* GraphAction = SelectionAsGraph())
 	{
 		SearchTerm = GraphAction->FuncName.ToString();
diff --git a/Engine/Source/Runtime/Engine/Classes/EdGraph/EdGraphSchema.h b/Engine/Source/Runtime/Engine/Classes/EdGraph/EdGraphSchema.h
--- a/Engine/Source/Runtime/Engine/Classes/EdGraph/EdGraphSchema.h
+++ b/Engine/Source/Runtime/Engine/Classes/EdGraph/EdGraphSchema.h
@@ -994,6 +994,11 @@ class ENGINE_API UEdGraphSchema : public UObject
 	 */
 	virtual bool CanGraphBeDropped(TSharedPtr<FEdGraphSchemaAction> InAction) const { return false; }
 
+	/*
+	 * Returns a custom reference string for searching within the blueprint based on a given action
+	 */
+	virtual FString GetFindReferenceSearchTerm(const FEdGraphSchemaAction* InGraphAction) const { return FString(); }
+
 	/*
 	 * Begins a drag and drop action to drag a graph action into another graph
 	 */
