Control Rig: Fix nodes with struct pins in AnimBP that are split can result in duplicate pins.

diff --git a/Engine/Plugins/Animation/ControlRig/Source/ControlRigDeveloper/Private/AnimGraphNode_ControlRig.cpp b/Engine/Plugins/Animation/ControlRig/Source/ControlRigDeveloper/Private/AnimGraphNode_ControlRig.cpp
--- a/Engine/Plugins/Animation/ControlRig/Source/ControlRigDeveloper/Private/AnimGraphNode_ControlRig.cpp
+++ b/Engine/Plugins/Animation/ControlRig/Source/ControlRigDeveloper/Private/AnimGraphNode_ControlRig.cpp
@@ -69,6 +69,13 @@ void UAnimGraphNode_ControlRig::CreateCustomPins(TArray<UEdGraphPin*>* OldPins)
 				{
 					for (UEdGraphPin* OldPin : *OldPins)
 					{
+						// do not rebuild sub pins as they will be treated after in UK2Node::RestoreSplitPins
+						const bool bIsSubPin = OldPin->ParentPin && OldPin->ParentPin->bHidden;
+						if (bIsSubPin)
+						{
+							continue;
+						}
+						
 						bool bFound = false;
 						for (UEdGraphPin* CurrentPin : Pins)
 						{
