Control Rig: Ensure to clear template graph after spawning template node

diff --git a/Engine/Plugins/Animation/ControlRig/Source/ControlRigEditor/Private/Graph/NodeSpawners/ControlRigBranchNodeSpawner.cpp b/Engine/Plugins/Animation/ControlRig/Source/ControlRigEditor/Private/Graph/NodeSpawners/ControlRigBranchNodeSpawner.cpp
--- a/Engine/Plugins/Animation/ControlRig/Source/ControlRigEditor/Private/Graph/NodeSpawners/ControlRigBranchNodeSpawner.cpp
+++ b/Engine/Plugins/Animation/ControlRig/Source/ControlRigEditor/Private/Graph/NodeSpawners/ControlRigBranchNodeSpawner.cpp
@@ -52,7 +52,7 @@ UEdGraphNode* UControlRigBranchNodeSpawner::Invoke(UEdGraph* ParentGraph, FBindi
 	UControlRigGraphNode* NewNode = nullptr;
 
 	bool const bIsTemplateNode = FBlueprintNodeTemplateCache::IsTemplateOuter(ParentGraph);
-	bool const bUndo = !bIsTemplateNode;
+	bool const bIsUserFacingNode = !bIsTemplateNode;
 	bool const bPrintCommand = !bIsTemplateNode;
 
 	// First create a backing member for our node
@@ -74,16 +74,16 @@ UEdGraphNode* UControlRigBranchNodeSpawner::Invoke(UEdGraph* ParentGraph, FBindi
 
 	FName Name = *URigVMBranchNode::BranchName;
 
-	if (!bIsTemplateNode)
+	if (bIsUserFacingNode)
 	{
 		Controller->OpenUndoBracket(FString::Printf(TEXT("Add '%s' Node"), *Name.ToString()));
 	}
 
-	if (URigVMBranchNode* ModelNode = Controller->AddBranchNode(Location, Name.ToString(), bUndo, bPrintCommand))
+	if (URigVMBranchNode* ModelNode = Controller->AddBranchNode(Location, Name.ToString(), bIsUserFacingNode, bPrintCommand))
 	{
 		NewNode = Cast<UControlRigGraphNode>(RigGraph->FindNodeForModelNodeName(ModelNode->GetFName()));
 
-		if (NewNode && bUndo)
+		if (NewNode && bIsUserFacingNode)
 		{
 			Controller->ClearNodeSelection(true);
 			Controller->SelectNode(ModelNode, true, true);
@@ -91,14 +91,18 @@ UEdGraphNode* UControlRigBranchNodeSpawner::Invoke(UEdGraph* ParentGraph, FBindi
 			UControlRigUnitNodeSpawner::HookupMutableNode(ModelNode, RigBlueprint);
 		}
 
-		if (bUndo)
+		if (bIsUserFacingNode)
 		{
 			Controller->CloseUndoBracket();
 		}
+		else
+		{
+			Controller->RemoveNode(ModelNode, false);
+		}
 	}
 	else
 	{
-		if (bUndo)
+		if (bIsUserFacingNode)
 		{
 			Controller->CancelUndoBracket();
 		}
diff --git a/Engine/Plugins/Animation/ControlRig/Source/ControlRigEditor/Private/Graph/NodeSpawners/ControlRigEnumNodeSpawner.cpp b/Engine/Plugins/Animation/ControlRig/Source/ControlRigEditor/Private/Graph/NodeSpawners/ControlRigEnumNodeSpawner.cpp
--- a/Engine/Plugins/Animation/ControlRig/Source/ControlRigEditor/Private/Graph/NodeSpawners/ControlRigEnumNodeSpawner.cpp
+++ b/Engine/Plugins/Animation/ControlRig/Source/ControlRigEditor/Private/Graph/NodeSpawners/ControlRigEnumNodeSpawner.cpp
@@ -55,6 +55,7 @@ UEdGraphNode* UControlRigEnumNodeSpawner::Invoke(UEdGraph* ParentGraph, FBinding
 	UControlRigGraphNode* NewNode = nullptr;
 
 	bool const bIsTemplateNode = FBlueprintNodeTemplateCache::IsTemplateOuter(ParentGraph);
+	bool const bIsUserFacingNode = !bIsTemplateNode;
 
 	if (bIsTemplateNode)
 	{
@@ -74,8 +75,6 @@ UEdGraphNode* UControlRigEnumNodeSpawner::Invoke(UEdGraph* ParentGraph, FBinding
 		return NewNode;
 	}
 
-	bool const bUndo = !bIsTemplateNode;
-
 	// First create a backing member for our node
 	UControlRigGraph* RigGraph = Cast<UControlRigGraph>(ParentGraph);
 	if(RigGraph == nullptr) return nullptr;
@@ -98,24 +97,28 @@ UEdGraphNode* UControlRigEnumNodeSpawner::Invoke(UEdGraph* ParentGraph, FBinding
 		Controller->OpenUndoBracket(FString::Printf(TEXT("Add '%s' Node"), *Name.ToString()));
 	}
 
-	if (URigVMEnumNode* ModelNode = Controller->AddEnumNode(*Enum->GetPathName(), Location, Name.ToString(), bUndo, !bIsTemplateNode))
+	if (URigVMEnumNode* ModelNode = Controller->AddEnumNode(*Enum->GetPathName(), Location, Name.ToString(), bIsUserFacingNode, !bIsTemplateNode))
 	{
 		NewNode = Cast<UControlRigGraphNode>(RigGraph->FindNodeForModelNodeName(ModelNode->GetFName()));
 
-		if (NewNode && bUndo)
+		if (NewNode && bIsUserFacingNode)
 		{
 			Controller->ClearNodeSelection(true);
 			Controller->SelectNode(ModelNode, true, true);
 		}
 
-		if (bUndo)
+		if (bIsUserFacingNode)
 		{
 			Controller->CloseUndoBracket();
 		}
+		else
+		{
+			Controller->RemoveNode(ModelNode, false);
+		}
 	}
 	else
 	{
-		if (bUndo)
+		if (bIsUserFacingNode)
 		{
 			Controller->CancelUndoBracket();
 		}
diff --git a/Engine/Plugins/Animation/ControlRig/Source/ControlRigEditor/Private/Graph/NodeSpawners/ControlRigFunctionRefNodeSpawner.cpp b/Engine/Plugins/Animation/ControlRig/Source/ControlRigEditor/Private/Graph/NodeSpawners/ControlRigFunctionRefNodeSpawner.cpp
--- a/Engine/Plugins/Animation/ControlRig/Source/ControlRigEditor/Private/Graph/NodeSpawners/ControlRigFunctionRefNodeSpawner.cpp
+++ b/Engine/Plugins/Animation/ControlRig/Source/ControlRigEditor/Private/Graph/NodeSpawners/ControlRigFunctionRefNodeSpawner.cpp
@@ -280,7 +280,7 @@ UControlRigGraphNode* UControlRigFunctionRefNodeSpawner::SpawnNode(UEdGraph* Par
 	if (RigBlueprint != nullptr && RigGraph != nullptr)
 	{
 		bool const bIsTemplateNode = FBlueprintNodeTemplateCache::IsTemplateOuter(ParentGraph);
-		bool const bUndo = !bIsTemplateNode;
+		bool const bIsUserFacingNode = !bIsTemplateNode;
 
 		FName Name = bIsTemplateNode ? *InFunction->GetName() : FControlRigBlueprintUtils::ValidateName(RigBlueprint, InFunction->GetName());
 		URigVMController* Controller = bIsTemplateNode ? RigGraph->GetTemplateController() : RigBlueprint->GetController(ParentGraph);
@@ -290,12 +290,12 @@ UControlRigGraphNode* UControlRigFunctionRefNodeSpawner::SpawnNode(UEdGraph* Par
 			Controller->OpenUndoBracket(FString::Printf(TEXT("Add '%s' Node"), *Name.ToString()));
 		}
 
-		if (URigVMFunctionReferenceNode* ModelNode = Controller->AddFunctionReferenceNode(InFunction, Location, Name.ToString(), bUndo, !bIsTemplateNode))
+		if (URigVMFunctionReferenceNode* ModelNode = Controller->AddFunctionReferenceNode(InFunction, Location, Name.ToString(), bIsUserFacingNode, !bIsTemplateNode))
 		{
 			NewNode = Cast<UControlRigGraphNode>(RigGraph->FindNodeForModelNodeName(ModelNode->GetFName()));
 			check(NewNode);
 
-			if (NewNode && bUndo)
+			if (NewNode && bIsUserFacingNode)
 			{
 				Controller->ClearNodeSelection(true);
 				Controller->SelectNode(ModelNode, true, true);
@@ -327,7 +327,7 @@ UControlRigGraphNode* UControlRigFunctionRefNodeSpawner::SpawnNode(UEdGraph* Par
 					{
 						if(!MappedVariablePair.Value.IsNone())
 						{
-							Controller->SetRemappedVariable(ModelNode, MappedVariablePair.Key, MappedVariablePair.Value, bUndo);
+							Controller->SetRemappedVariable(ModelNode, MappedVariablePair.Key, MappedVariablePair.Value, bIsUserFacingNode);
 						}
 					}
 				}
@@ -341,14 +341,18 @@ UControlRigGraphNode* UControlRigFunctionRefNodeSpawner::SpawnNode(UEdGraph* Par
 				Controller->SuspendNotifications(false);
 			}
 
-			if (bUndo)
+			if (bIsUserFacingNode)
 			{
 				Controller->CloseUndoBracket();
 			}
+			else
+			{
+				Controller->RemoveNode(ModelNode, false);
+			}
 		}
 		else
 		{
-			if (bUndo)
+			if (bIsUserFacingNode)
 			{
 				Controller->CancelUndoBracket();
 			}
diff --git a/Engine/Plugins/Animation/ControlRig/Source/ControlRigEditor/Private/Graph/NodeSpawners/ControlRigPrototypeNodeSpawner.cpp b/Engine/Plugins/Animation/ControlRig/Source/ControlRigEditor/Private/Graph/NodeSpawners/ControlRigPrototypeNodeSpawner.cpp
--- a/Engine/Plugins/Animation/ControlRig/Source/ControlRigEditor/Private/Graph/NodeSpawners/ControlRigPrototypeNodeSpawner.cpp
+++ b/Engine/Plugins/Animation/ControlRig/Source/ControlRigEditor/Private/Graph/NodeSpawners/ControlRigPrototypeNodeSpawner.cpp
@@ -108,7 +108,7 @@ UControlRigGraphNode* UControlRigPrototypeNodeSpawner::SpawnNode(UEdGraph* Paren
 	if (RigBlueprint != nullptr && RigGraph != nullptr)
 	{
 		bool const bIsTemplateNode = FBlueprintNodeTemplateCache::IsTemplateOuter(ParentGraph);
-		bool const bUndo = !bIsTemplateNode;
+		bool const bIsUserFacingNode = !bIsTemplateNode;
 
 		const FRigVMPrototype* Prototype = FRigVMRegistry::Get().FindPrototype(InNotation);
 		if (Prototype == nullptr)
@@ -124,11 +124,11 @@ UControlRigGraphNode* UControlRigPrototypeNodeSpawner::SpawnNode(UEdGraph* Paren
 			Controller->OpenUndoBracket(FString::Printf(TEXT("Add '%s' Node"), *Name.ToString()));
 		}
 
-		if (URigVMPrototypeNode* ModelNode = Controller->AddPrototypeNode(InNotation, Location, Name.ToString(), bUndo, !bIsTemplateNode))
+		if (URigVMPrototypeNode* ModelNode = Controller->AddPrototypeNode(InNotation, Location, Name.ToString(), bIsUserFacingNode, !bIsTemplateNode))
 		{
 			NewNode = Cast<UControlRigGraphNode>(RigGraph->FindNodeForModelNodeName(ModelNode->GetFName()));
 
-			if (NewNode && bUndo)
+			if (NewNode && bIsUserFacingNode)
 			{
 				Controller->ClearNodeSelection(true);
 				Controller->SelectNode(ModelNode, true, true);
@@ -136,14 +136,19 @@ UControlRigGraphNode* UControlRigPrototypeNodeSpawner::SpawnNode(UEdGraph* Paren
 				UControlRigUnitNodeSpawner::HookupMutableNode(ModelNode, RigBlueprint);
 			}
 
-			if (bUndo)
+			if (bIsUserFacingNode)
 			{
 				Controller->CloseUndoBracket();
 			}
+			else
+			{
+				NewNode->ModelNodePath = Prototype->GetNotation().ToString();
+				Controller->RemoveNode(ModelNode, false);
+			}
 		}
 		else
 		{
-			if (bUndo)
+			if (bIsUserFacingNode)
 			{
 				Controller->CancelUndoBracket();
 			}
diff --git a/Engine/Plugins/Animation/ControlRig/Source/ControlRigEditor/Private/Graph/NodeSpawners/ControlRigUnitNodeSpawner.cpp b/Engine/Plugins/Animation/ControlRig/Source/ControlRigEditor/Private/Graph/NodeSpawners/ControlRigUnitNodeSpawner.cpp
--- a/Engine/Plugins/Animation/ControlRig/Source/ControlRigEditor/Private/Graph/NodeSpawners/ControlRigUnitNodeSpawner.cpp
+++ b/Engine/Plugins/Animation/ControlRig/Source/ControlRigEditor/Private/Graph/NodeSpawners/ControlRigUnitNodeSpawner.cpp
@@ -131,7 +131,7 @@ UControlRigGraphNode* UControlRigUnitNodeSpawner::SpawnNode(UEdGraph* ParentGrap
 	if (RigBlueprint != nullptr && RigGraph != nullptr)
 	{
 		bool const bIsTemplateNode = FBlueprintNodeTemplateCache::IsTemplateOuter(ParentGraph);
-		bool const bUndo = !bIsTemplateNode;
+		bool const bIsUserFacingNode = !bIsTemplateNode;
 
 		FName Name = bIsTemplateNode ? *StructTemplate->GetDisplayNameText().ToString() : FControlRigBlueprintUtils::ValidateName(RigBlueprint, StructTemplate->GetFName().ToString());
 		URigVMController* Controller = bIsTemplateNode ? RigGraph->GetTemplateController() : RigBlueprint->GetController(ParentGraph);
@@ -144,12 +144,12 @@ UControlRigGraphNode* UControlRigUnitNodeSpawner::SpawnNode(UEdGraph* ParentGrap
 		FRigVMUnitNodeCreatedContext& UnitNodeCreatedContext = Controller->GetUnitNodeCreatedContext();
 		FRigVMUnitNodeCreatedContext::FScope ReasonScope(UnitNodeCreatedContext, ERigVMNodeCreatedReason::NodeSpawner);
 
-		if (URigVMUnitNode* ModelNode = Controller->AddUnitNode(StructTemplate, FRigUnit::GetMethodName(), Location, Name.ToString(), bUndo, !bIsTemplateNode))
+		if (URigVMUnitNode* ModelNode = Controller->AddUnitNode(StructTemplate, FRigUnit::GetMethodName(), Location, Name.ToString(), bIsUserFacingNode, !bIsTemplateNode))
 		{
 			NewNode = Cast<UControlRigGraphNode>(RigGraph->FindNodeForModelNodeName(ModelNode->GetFName()));
 			check(NewNode);
 
-			if (NewNode && bUndo)
+			if (NewNode && bIsUserFacingNode)
 			{
 				Controller->ClearNodeSelection(true);
 				Controller->SelectNode(ModelNode, true, true);
@@ -168,7 +168,7 @@ UControlRigGraphNode* UControlRigUnitNodeSpawner::SpawnNode(UEdGraph* ParentGrap
 					for (const TPair<FString, bool>& Pair : ExpansionMap.Values)
 					{
 						FString PinPath = FString::Printf(TEXT("%s.%s"), *ModelNode->GetName(), *Pair.Key);
-						Controller->SetPinExpansion(PinPath, Pair.Value, bUndo);
+						Controller->SetPinExpansion(PinPath, Pair.Value, bIsUserFacingNode);
 					}
 				}
 
@@ -213,14 +213,18 @@ UControlRigGraphNode* UControlRigUnitNodeSpawner::SpawnNode(UEdGraph* ParentGrap
 			}
 #endif
 
-			if (bUndo)
+			if (bIsUserFacingNode)
 			{
 				Controller->CloseUndoBracket();
 			}
+			else
+			{
+				Controller->RemoveNode(ModelNode, false);
+			}
 		}
 		else
 		{
-			if (bUndo)
+			if (bIsUserFacingNode)
 			{
 				Controller->CancelUndoBracket();
 			}
diff --git a/Engine/Plugins/Animation/ControlRig/Source/ControlRigEditor/Private/Graph/NodeSpawners/ControlRigVariableNodeSpawner.cpp b/Engine/Plugins/Animation/ControlRig/Source/ControlRigEditor/Private/Graph/NodeSpawners/ControlRigVariableNodeSpawner.cpp
--- a/Engine/Plugins/Animation/ControlRig/Source/ControlRigEditor/Private/Graph/NodeSpawners/ControlRigVariableNodeSpawner.cpp
+++ b/Engine/Plugins/Animation/ControlRig/Source/ControlRigEditor/Private/Graph/NodeSpawners/ControlRigVariableNodeSpawner.cpp
@@ -110,6 +110,7 @@ UEdGraphNode* UControlRigVariableNodeSpawner::Invoke(UEdGraph* ParentGraph, FBin
 	UControlRigGraphNode* NewNode = nullptr;
 
 	bool const bIsTemplateNode = FBlueprintNodeTemplateCache::IsTemplateOuter(ParentGraph);
+	bool const bIsUserFacingNode = !bIsTemplateNode;
 
 	// First create a backing member for our node
 	UControlRigGraph* RigGraph = Cast<UControlRigGraph>(ParentGraph);
@@ -128,7 +129,7 @@ UEdGraphNode* UControlRigVariableNodeSpawner::Invoke(UEdGraph* ParentGraph, FBin
 
 	URigVMController* Controller = bIsTemplateNode ? RigGraph->GetTemplateController() : RigBlueprint->GetController(ParentGraph);
 
-	if (!bIsTemplateNode)
+	if (bIsUserFacingNode)
 	{
 		Controller->OpenUndoBracket(TEXT("Add Variable"));
 	}
@@ -170,7 +171,7 @@ UEdGraphNode* UControlRigVariableNodeSpawner::Invoke(UEdGraph* ParentGraph, FBin
 			}
 		}
 
-		if (!bIsTemplateNode)
+		if (bIsUserFacingNode)
 		{
 			if (NewNode)
 			{
