Control Rig: fix for build break in CL 19074359

diff --git a/Engine/Plugins/Animation/ControlRig/Source/ControlRigDeveloper/Public/Graph/ControlRigGraphNode.h b/Engine/Plugins/Animation/ControlRig/Source/ControlRigDeveloper/Public/Graph/ControlRigGraphNode.h
--- a/Engine/Plugins/Animation/ControlRig/Source/ControlRigDeveloper/Public/Graph/ControlRigGraphNode.h
+++ b/Engine/Plugins/Animation/ControlRig/Source/ControlRigDeveloper/Public/Graph/ControlRigGraphNode.h
@@ -222,4 +222,5 @@ private:
 	friend class SControlRigGraphNode;
 	friend class FControlRigArgumentLayout;
 	friend class FControlRigGraphDetails;
+	friend class UControlRigPrototypeNodeSpawner;
 };
