Control Rig: Perform Reset() in URigHierarchy::BeginDestroy instead of ~URigHierarachy(). BeginDestroy is by design the place to perform clean up for UObject. By the time it reaches the destructor, it would not be safe to access any UProperty members. It was OK before since Reset only touches native class members. However, with the addition of TObjectPtr<URigHierarchy> HierarchyForCacheValidation, its clean up needs to happen earilier in begin destroy

diff --git a/Engine/Plugins/Animation/ControlRig/Source/ControlRig/Private/Rigs/RigHierarchy.cpp b/Engine/Plugins/Animation/ControlRig/Source/ControlRig/Private/Rigs/RigHierarchy.cpp
--- a/Engine/Plugins/Animation/ControlRig/Source/ControlRig/Private/Rigs/RigHierarchy.cpp
+++ b/Engine/Plugins/Animation/ControlRig/Source/ControlRig/Private/Rigs/RigHierarchy.cpp
@@ -122,9 +122,13 @@ URigHierarchy::URigHierarchy()
 #endif
 }
 
-URigHierarchy::~URigHierarchy()
+void URigHierarchy::BeginDestroy()
 {
+	// reset needs to be in begin destroy since
+	// reset touches a UObject member of this hierarchy,
+	// which will be GCed by the time this hierarchy reaches destructor
 	Reset();
+	Super::BeginDestroy();
 }
 
 void URigHierarchy::Serialize(FArchive& Ar)
diff --git a/Engine/Plugins/Animation/ControlRig/Source/ControlRig/Public/Rigs/RigHierarchy.h b/Engine/Plugins/Animation/ControlRig/Source/ControlRig/Public/Rigs/RigHierarchy.h
--- a/Engine/Plugins/Animation/ControlRig/Source/ControlRig/Public/Rigs/RigHierarchy.h
+++ b/Engine/Plugins/Animation/ControlRig/Source/ControlRig/Public/Rigs/RigHierarchy.h
@@ -95,9 +95,9 @@ public:
 	typedef TPair<int32, TArray<int32>> TElementDependencyMapPair;
 
 	URigHierarchy();
-	virtual ~URigHierarchy();
 
 	// UObject interface
+	virtual void BeginDestroy() override;
 	virtual void Serialize(FArchive& Ar) override;
 	void Save(FArchive& Ar);
 	void Load(FArchive& Ar);
