Control Rig: Fixed wrong Position/Rotator pin default value for anim graph nodes

diff --git a/Engine/Plugins/Animation/ControlRig/Source/ControlRig/Private/Rigs/RigHierarchy.cpp b/Engine/Plugins/Animation/ControlRig/Source/ControlRig/Private/Rigs/RigHierarchy.cpp
--- a/Engine/Plugins/Animation/ControlRig/Source/ControlRig/Private/Rigs/RigHierarchy.cpp
+++ b/Engine/Plugins/Animation/ControlRig/Source/ControlRig/Private/Rigs/RigHierarchy.cpp
@@ -928,7 +928,11 @@ FString URigHierarchy::GetControlPinDefaultValue(FRigControlElement* InControlEl
 		{
 			if(bForEdGraph)
 			{
-				return FVector(Value.Get<FVector3f>()).ToString();
+				// NOTE: We can not use ToString() here since the FDefaultValueHelper::IsStringValidVector used in
+				// EdGraphSchema_K2 expects a string with format '#,#,#', while FVector:ToString is returning the value
+				// with format 'X=# Y=# Z=#'				
+				const FVector Vector(Value.Get<FVector3f>());
+				return FString::Printf(TEXT("%3.3f,%3.3f,%3.3f"), Vector.X, Vector.Y, Vector.Z);
 			}
 			return Value.ToString<FVector>();
 		}
@@ -936,8 +940,9 @@ FString URigHierarchy::GetControlPinDefaultValue(FRigControlElement* InControlEl
 		{
 				if(bForEdGraph)
 				{
-					const FRotator Rotator = FRotator::MakeFromEuler((FVector)Value.GetRef<FVector3f>());
-					return Rotator.ToString();
+					// NOTE: se explanations above for Position/Scale
+					const FRotator Rotator = FRotator::MakeFromEuler(Value.GetRef<FVector3f>());
+					return FString::Printf(TEXT("%3.3f,%3.3f,%3.3f"), Rotator.Pitch, Rotator.Yaw, Rotator.Roll);
 				}
 				return Value.ToString<FRotator>();
 		}
