Control Rig: Fix duplicated blueprint objects compiling before models refeshed.

diff --git a/Engine/Plugins/Animation/ControlRig/Source/ControlRigDeveloper/Private/ControlRigBlueprint.cpp b/Engine/Plugins/Animation/ControlRig/Source/ControlRigDeveloper/Private/ControlRigBlueprint.cpp
--- a/Engine/Plugins/Animation/ControlRig/Source/ControlRigDeveloper/Private/ControlRigBlueprint.cpp
+++ b/Engine/Plugins/Animation/ControlRig/Source/ControlRigDeveloper/Private/ControlRigBlueprint.cpp
@@ -2129,6 +2129,8 @@ void UControlRigBlueprint::ReplaceDeprecatedNodes()
 
 void UControlRigBlueprint::PostDuplicate(bool bDuplicateForPIE)
 {
+	RefreshAllModels();
+	
 	Super::PostDuplicate(bDuplicateForPIE);
 	
 	if (URigHierarchyController* Controller = Hierarchy->GetController(true))
