Control Rig: Deal with inverse frame range in snapper

diff --git a/Engine/Plugins/Animation/ControlRig/Source/ControlRigEditor/Private/Tools/ControlRigSnapper.cpp b/Engine/Plugins/Animation/ControlRig/Source/ControlRigEditor/Private/Tools/ControlRigSnapper.cpp
--- a/Engine/Plugins/Animation/ControlRig/Source/ControlRigEditor/Private/Tools/ControlRigSnapper.cpp
+++ b/Engine/Plugins/Animation/ControlRig/Source/ControlRigEditor/Private/Tools/ControlRigSnapper.cpp
@@ -354,6 +354,11 @@ struct FGuidAndActor
 
 static void CalculateFramesToSnap(ISequencer* InSequencer, UMovieScene* MovieScene, FFrameNumber StartFrame, FFrameNumber EndFrame, TArray<FFrameNumber>& OutFrames)
 {
+	if(StartFrame > EndFrame)
+	{
+		Swap(StartFrame, EndFrame);
+	}
+	
 	FFrameRate TickResolution = MovieScene->GetTickResolution();
 	FFrameRate DisplayResolution = MovieScene->GetDisplayRate();
 
