Control Rig: Ensure backwards compatibility with gizmo libraries

diff --git a/Engine/Plugins/Animation/ControlRig/Source/ControlRigDeveloper/Private/ControlRigBlueprint.cpp b/Engine/Plugins/Animation/ControlRig/Source/ControlRigDeveloper/Private/ControlRigBlueprint.cpp
--- a/Engine/Plugins/Animation/ControlRig/Source/ControlRigDeveloper/Private/ControlRigBlueprint.cpp
+++ b/Engine/Plugins/Animation/ControlRig/Source/ControlRigDeveloper/Private/ControlRigBlueprint.cpp
@@ -539,15 +539,15 @@ void UControlRigBlueprint::PostLoad()
 	}
 
 	// upgrade the gizmo libraries to shape libraries
-	if(GizmoLibrary_DEPRECATED.IsValid() || GetLinkerCustomVersion(FControlRigObjectVersion::GUID) < FControlRigObjectVersion::RenameGizmoToShape)
+	if(!GizmoLibrary_DEPRECATED.IsNull() || GetLinkerCustomVersion(FControlRigObjectVersion::GUID) < FControlRigObjectVersion::RenameGizmoToShape)
 	{
 		// if it's an older file and it doesn't have the GizmoLibrary stored,
 		// refer to the previous default.
 		ShapeLibraries.Reset();
 
-		if(GizmoLibrary_DEPRECATED.IsValid())
+		if(!GizmoLibrary_DEPRECATED.IsNull())
 		{
-			ShapeLibrariesToLoadOnPackageLoaded.Add(GizmoLibrary_DEPRECATED->GetPathName());
+			ShapeLibrariesToLoadOnPackageLoaded.Add(GizmoLibrary_DEPRECATED.ToString());
 		}
 		else
 		{
