Control Rig: For 5.0 Control Rig has been promoted out of beta.

diff --git a/Engine/Plugins/Animation/ControlRig/ControlRig.uplugin b/Engine/Plugins/Animation/ControlRig/ControlRig.uplugin
--- a/Engine/Plugins/Animation/ControlRig/ControlRig.uplugin
+++ b/Engine/Plugins/Animation/ControlRig/ControlRig.uplugin
@@ -12,7 +12,6 @@
 	"SupportURL" : "",
 	"EnabledByDefault" : true,
 	"CanContainContent" : true,
-	"IsBetaVersion" : true,
 	"Installed" : false,
 	"Plugins":
 	[
