Fix up deprecations for CL# 19104650
#rnx

diff --git a/Engine/Plugins/Animation/ControlRig/Source/ControlRigEditor/Private/ControlRigEditorModule.cpp b/Engine/Plugins/Animation/ControlRig/Source/ControlRigEditor/Private/ControlRigEditorModule.cpp
--- a/Engine/Plugins/Animation/ControlRig/Source/ControlRigEditor/Private/ControlRigEditorModule.cpp
+++ b/Engine/Plugins/Animation/ControlRig/Source/ControlRigEditor/Private/ControlRigEditorModule.cpp
@@ -278,7 +278,7 @@ void FControlRigEditorModule::StartupModule()
 					}
 
 					TArray<FRigVMReferenceNodeData> ReferenceNodeDatas;
-					ReferenceNodeDataProperty->ImportText(*ReferenceNodeDataString, &ReferenceNodeDatas, EPropertyPortFlags::PPF_None, nullptr);
+					ReferenceNodeDataProperty->ImportText_Direct(*ReferenceNodeDataString, &ReferenceNodeDatas, nullptr, EPropertyPortFlags::PPF_None);
 
 					for(const FRigVMReferenceNodeData& ReferenceNodeData : ReferenceNodeDatas)
 					{
diff --git a/Engine/Plugins/Animation/IKRig/Source/IKRigDeveloper/Private/AnimGraphNode_IKRig.cpp b/Engine/Plugins/Animation/IKRig/Source/IKRigDeveloper/Private/AnimGraphNode_IKRig.cpp
--- a/Engine/Plugins/Animation/IKRig/Source/IKRigDeveloper/Private/AnimGraphNode_IKRig.cpp
+++ b/Engine/Plugins/Animation/IKRig/Source/IKRigDeveloper/Private/AnimGraphNode_IKRig.cpp
@@ -464,7 +464,7 @@ void UAnimGraphNode_IKRig::SetPinDefaultValue(UEdGraphPin* InPin, const FName& I
 
 		// fallback to default behaviour 
 		FString DefaultValue;
-		InProperty->ExportTextItem(DefaultValue, Memory, nullptr, nullptr, PPF_None);
+		InProperty->ExportTextItem_Direct(DefaultValue, Memory, nullptr, nullptr, PPF_None);
 		return DefaultValue;
 	};
 	
