Control Rig: reverted an old change. GetTransform Node is once again allowed to executes during init. Otherwise, if a transform constraint node takes input from get transform nodes, it would cache the wrong offset since it relies on GetTransform to report valid data during init stage

diff --git a/Engine/Plugins/Animation/ControlRig/Source/ControlRig/Private/Units/Hierarchy/RigUnit_GetTransform.cpp b/Engine/Plugins/Animation/ControlRig/Source/ControlRig/Private/Units/Hierarchy/RigUnit_GetTransform.cpp
--- a/Engine/Plugins/Animation/ControlRig/Source/ControlRig/Private/Units/Hierarchy/RigUnit_GetTransform.cpp
+++ b/Engine/Plugins/Animation/ControlRig/Source/ControlRig/Private/Units/Hierarchy/RigUnit_GetTransform.cpp
@@ -21,7 +21,9 @@ FRigUnit_GetTransform_Execute()
 			case EControlRigState::Init:
 			{
 				CachedIndex.Reset();
-				break;
+				// there is no "break;" here because an old node, Transform Constrtaint,
+				// still caches data during init stage. Thus, if a it takes a GetTransform
+				// as input, that GetTransform needs to execute and output a valid transform					
 			}
 			case EControlRigState::Update:
 			{
