Control Rig: Adapt to default name change in gizmo shape

diff --git a/Engine/Plugins/Animation/ControlRig/Source/ControlRig/Private/Rigs/RigHierarchyController.cpp b/Engine/Plugins/Animation/ControlRig/Source/ControlRig/Private/Rigs/RigHierarchyController.cpp
--- a/Engine/Plugins/Animation/ControlRig/Source/ControlRig/Private/Rigs/RigHierarchyController.cpp
+++ b/Engine/Plugins/Animation/ControlRig/Source/ControlRig/Private/Rigs/RigHierarchyController.cpp
@@ -1262,6 +1262,11 @@ TArray<FRigElementKey> URigHierarchyController::ImportFromHierarchyContainer(con
 		Settings.ShapeColor = Control.GizmoColor;
 		Settings.ControlEnum = Control.ControlEnum;
 
+		if(Settings.ShapeName == FRigControl().GizmoName)
+		{
+			Settings.ShapeName = FControlRigShapeDefinition().ShapeName; 
+		}
+
 		FRigControlValue InitialValue = Control.InitialValue;
 		if(!InitialValue.IsValid())
 		{
diff --git a/Engine/Plugins/Animation/ControlRig/Source/ControlRig/Private/Rigs/RigHierarchyElements.cpp b/Engine/Plugins/Animation/ControlRig/Source/ControlRig/Private/Rigs/RigHierarchyElements.cpp
--- a/Engine/Plugins/Animation/ControlRig/Source/ControlRig/Private/Rigs/RigHierarchyElements.cpp
+++ b/Engine/Plugins/Animation/ControlRig/Source/ControlRig/Private/Rigs/RigHierarchyElements.cpp
@@ -2,6 +2,7 @@
 
 #include "Rigs/RigHierarchyElements.h"
 #include "Rigs/RigHierarchy.h"
+#include "Rigs/RigControlHierarchy.h"
 #include "Units/RigUnitContext.h"
 #include "ControlRigObjectVersion.h"
 #include "ControlRigGizmoLibrary.h"
@@ -513,6 +514,15 @@ void FRigControlSettings::Load(FArchive& Ar)
 	Ar << bShapeEnabled;
 	Ar << bShapeVisible;
 	Ar << ShapeName;
+
+	if(Ar.CustomVer(FControlRigObjectVersion::GUID) < FControlRigObjectVersion::RenameGizmoToShape)
+	{
+		if(ShapeName == FRigControl().GizmoName)
+		{
+			ShapeName = FControlRigShapeDefinition().ShapeName; 
+		}
+	}
+	
 	Ar << ShapeColor;
 	Ar << bIsTransientControl;
 	Ar << ControlEnumPathName;
