Control Rig: Fix for static analysis ControlRigTemplateNodeSpawner.cpp

diff --git a/Engine/Plugins/Animation/ControlRig/Source/ControlRigEditor/Private/Graph/NodeSpawners/ControlRigTemplateNodeSpawner.cpp b/Engine/Plugins/Animation/ControlRig/Source/ControlRigEditor/Private/Graph/NodeSpawners/ControlRigTemplateNodeSpawner.cpp
--- a/Engine/Plugins/Animation/ControlRig/Source/ControlRigEditor/Private/Graph/NodeSpawners/ControlRigTemplateNodeSpawner.cpp
+++ b/Engine/Plugins/Animation/ControlRig/Source/ControlRigEditor/Private/Graph/NodeSpawners/ControlRigTemplateNodeSpawner.cpp
@@ -146,7 +146,10 @@ UControlRigGraphNode* UControlRigTemplateNodeSpawner::SpawnNode(UEdGraph* Parent
 			}
 			else
 			{
-				NewNode->ModelNodePath = Template->GetNotation().ToString();
+				if(NewNode)
+				{
+					NewNode->ModelNodePath = Template->GetNotation().ToString();
+				}
 				Controller->RemoveNode(ModelNode, false);
 			}
 		}
