Control Rig: make sure to use correct skeletal mesh when requesting anim pose

diff --git a/Engine/Plugins/Animation/ControlRig/Source/ControlRig/Private/Sequencer/MovieSceneControlRigParameterSection.cpp b/Engine/Plugins/Animation/ControlRig/Source/ControlRig/Private/Sequencer/MovieSceneControlRigParameterSection.cpp
--- a/Engine/Plugins/Animation/ControlRig/Source/ControlRig/Private/Sequencer/MovieSceneControlRigParameterSection.cpp
+++ b/Engine/Plugins/Animation/ControlRig/Source/ControlRig/Private/Sequencer/MovieSceneControlRigParameterSection.cpp
@@ -2604,8 +2604,10 @@ bool UMovieSceneControlRigParameterSection::LoadAnimSequenceIntoThisSection(UAni
 
 		// retrieve the pose using the services that persona and sequencer rely on
 		// rather than accessing the low level raw tracks.
+		FAnimPoseEvaluationOptions EvaluationOptions;
+		EvaluationOptions.OptionalSkeletalMesh = SkelMeshComp->SkeletalMesh;
 		FAnimPose AnimPose;
-		UAnimPoseExtensions::GetAnimPoseAtTime(AnimSequence, SequenceSecond, FAnimPoseEvaluationOptions(), AnimPose);
+		UAnimPoseExtensions::GetAnimPoseAtTime(AnimSequence, SequenceSecond, EvaluationOptions, AnimPose);
 
 		TArray<FName> BoneNames;
 		UAnimPoseExtensions::GetBoneNames(AnimPose, BoneNames);
