Control Rig: Validate VMRuntime settings

diff --git a/Engine/Plugins/Animation/ControlRig/Source/ControlRigEditor/Private/Editor/ControlRigEditor.cpp b/Engine/Plugins/Animation/ControlRig/Source/ControlRigEditor/Private/Editor/ControlRigEditor.cpp
--- a/Engine/Plugins/Animation/ControlRig/Source/ControlRigEditor/Private/Editor/ControlRigEditor.cpp
+++ b/Engine/Plugins/Animation/ControlRig/Source/ControlRigEditor/Private/Editor/ControlRigEditor.cpp
@@ -4305,6 +4305,7 @@ void FControlRigEditor::OnFinishedChangingProperties(const FPropertyChangedEvent
 
 		if (PropertyChangedEvent.MemberProperty->GetNameCPP() == TEXT("VMRuntimeSettings"))
 		{
+			ControlRigBP->VMRuntimeSettings.Validate();
 			ControlRigBP->PropagateRuntimeSettingsFromBPToInstances();
 			return;
 		}
diff --git a/Engine/Source/Runtime/RigVM/Public/RigVMCore/RigVMExecuteContext.h b/Engine/Source/Runtime/RigVM/Public/RigVMCore/RigVMExecuteContext.h
--- a/Engine/Source/Runtime/RigVM/Public/RigVMCore/RigVMExecuteContext.h
+++ b/Engine/Source/Runtime/RigVM/Public/RigVMCore/RigVMExecuteContext.h
@@ -177,6 +177,14 @@ struct RIGVM_API FRigVMRuntimeSettings
 	 * The function to use for logging anything from the VM to the host
 	 */
 	TFunction<void(EMessageSeverity::Type,const FRigVMExecuteContext*,const FString&)> LogFunction = nullptr;
+
+	/*
+	 * Validate the settings
+	 */
+	void Validate()
+	{
+		MaximumArraySize = FMath::Clamp(MaximumArraySize, 1, INT32_MAX);
+	}
 };
 
 /**
