Sequencer:  Make sure to evaluate Sequencer when snapping from current time. Recent regression.

diff --git a/Engine/Plugins/Animation/ControlRig/Source/ControlRigEditor/Private/Tools/ControlRigSnapper.cpp b/Engine/Plugins/Animation/ControlRig/Source/ControlRigEditor/Private/Tools/ControlRigSnapper.cpp
--- a/Engine/Plugins/Animation/ControlRig/Source/ControlRigEditor/Private/Tools/ControlRigSnapper.cpp
+++ b/Engine/Plugins/Animation/ControlRig/Source/ControlRigEditor/Private/Tools/ControlRigSnapper.cpp
@@ -451,7 +451,11 @@ bool FControlRigSnapper::SnapIt(FFrameNumber StartFrame, FFrameNumber EndFrame,c
 			UE_LOG(LogControlRig, Error, TEXT("Number of Frames %d to Snap different than Parent Frames %d"), Frames.Num(),WorldTransformToSnap.Num());
 			return false;
 		}
-
+		// If we are snapping from the current time there are some assumptions to not re-calcuate values so we need to make sure sequencer is up to date
+		if (StartFrame == Sequencer->GetLocalTime().Time.RoundToFrame())
+		{
+			Sequencer->ForceEvaluate();
+		}
 		TArray<FGuidAndActor > ActorsToSnap;
 		//There may be Actors here not in Sequencer so we add them to sequencer also
 		for (const FActorForWorldTransforms& ActorSelection : ActorToSnap.Actors)
