Control Rig: Ensure to remember the original default parent per element

diff --git a/Engine/Plugins/Animation/ControlRig/Source/ControlRig/Private/Rigs/RigHierarchy.cpp b/Engine/Plugins/Animation/ControlRig/Source/ControlRig/Private/Rigs/RigHierarchy.cpp
--- a/Engine/Plugins/Animation/ControlRig/Source/ControlRig/Private/Rigs/RigHierarchy.cpp
+++ b/Engine/Plugins/Animation/ControlRig/Source/ControlRig/Private/Rigs/RigHierarchy.cpp
@@ -1106,6 +1106,21 @@ FRigBaseElementParentArray URigHierarchy::GetParents(const FRigBaseElement* InEl
 	return Parents;
 }
 
+FRigElementKey URigHierarchy::GetDefaultParent(FRigElementKey InKey) const
+{
+	FRigElementKey DefaultParent;
+	if(const FRigElementKey* DefaultParentPtr = DefaultParentPerElement.Find(InKey))
+	{
+		DefaultParent = *DefaultParentPtr;
+	}
+	else
+	{
+		DefaultParent = GetFirstParent(InKey);
+		DefaultParentPerElement.Add(InKey, DefaultParent);
+	}
+	return DefaultParent;
+}
+
 FRigElementKey URigHierarchy::GetFirstParent(FRigElementKey InKey) const
 {
 	if(FRigBaseElement* FirstParent = GetFirstParent(Find(InKey)))
@@ -1251,7 +1266,10 @@ FRigElementKey URigHierarchy::GetActiveParent(const FRigElementKey& InKey) const
 			{
 				if (!(ParentKeys[ParentIndex] == URigHierarchy::GetDefaultParentKey() || ParentKeys[ParentIndex] == URigHierarchy::GetWorldSpaceReferenceKey()))
 				{
-					return URigHierarchy::GetDefaultParentKey();
+					if(ParentKeys[ParentIndex] == GetDefaultParent(InKey))
+					{
+						return URigHierarchy::GetDefaultParentKey();
+					}
 				}
 			}
 			return ParentKeys[ParentIndex];
@@ -3368,14 +3386,14 @@ FRigElementKey URigHierarchy::PreprocessParentElementKeyForSpaceSwitching(const
 	}
 	else if(InParentKey == GetDefaultParentKey())
 	{
-		const FRigElementKey FirstParent = GetFirstParent(InChildKey);
-		if(FirstParent == GetWorldSpaceReferenceKey())
+		const FRigElementKey DefaultParent = GetDefaultParent(InChildKey);
+		if(DefaultParent == GetWorldSpaceReferenceKey())
 		{
 			return FRigElementKey();
 		}
 		else
 		{
-			return FirstParent;
+			return DefaultParent;
 		}
 	}
 
diff --git a/Engine/Plugins/Animation/ControlRig/Source/ControlRig/Public/Rigs/RigHierarchy.h b/Engine/Plugins/Animation/ControlRig/Source/ControlRig/Public/Rigs/RigHierarchy.h
--- a/Engine/Plugins/Animation/ControlRig/Source/ControlRig/Public/Rigs/RigHierarchy.h
+++ b/Engine/Plugins/Animation/ControlRig/Source/ControlRig/Public/Rigs/RigHierarchy.h
@@ -1586,6 +1586,14 @@ public:
 	 */
 	FRigBaseElementParentArray GetParents(const FRigBaseElement* InElement, bool bRecursive = false) const;
 
+	/**
+	 * Returns the default parent element's key of a given child key
+	 * @param InKey The key of the element to retrieve the parent for
+	 * @return Returns the default parent element key
+	 */
+	UFUNCTION(BlueprintCallable, Category = URigHierarchy)
+	FRigElementKey GetDefaultParent(FRigElementKey InKey) const;
+
 	/**
 	 * Returns the first parent element of a given element key
 	 * @param InKey The key of the element to retrieve the parents for
@@ -2945,6 +2953,8 @@ protected:
 
 	UPROPERTY(transient)
 	TObjectPtr<URigHierarchy> HierarchyForCacheValidation;
+
+	mutable TMap<FRigElementKey, FRigElementKey> DefaultParentPerElement;  
 	
 private:
 	
