Control Rig: Ensure to remember the original default parent per element

diff --git a/Engine/Plugins/Animation/ControlRig/Source/ControlRig/Private/Rigs/RigHierarchy.cpp b/Engine/Plugins/Animation/ControlRig/Source/ControlRig/Private/Rigs/RigHierarchy.cpp
--- a/Engine/Plugins/Animation/ControlRig/Source/ControlRig/Private/Rigs/RigHierarchy.cpp
+++ b/Engine/Plugins/Animation/ControlRig/Source/ControlRig/Private/Rigs/RigHierarchy.cpp
@@ -309,6 +309,7 @@ void URigHierarchy::Reset()
 
 	ResetPoseHash = INDEX_NONE;
 	ResetPoseIsFilteredOut.Reset();
+	DefaultParentPerElement.Reset();
 
 	if(!IsGarbageCollecting())
 	{
