Animation Tools: Snapper Crash from Range being too large, limit it.

diff --git a/Engine/Plugins/Animation/ControlRig/Source/ControlRigEditor/Private/EditMode/SControlRigSnapper.cpp b/Engine/Plugins/Animation/ControlRig/Source/ControlRigEditor/Private/EditMode/SControlRigSnapper.cpp
--- a/Engine/Plugins/Animation/ControlRig/Source/ControlRigEditor/Private/EditMode/SControlRigSnapper.cpp
+++ b/Engine/Plugins/Animation/ControlRig/Source/ControlRigEditor/Private/EditMode/SControlRigSnapper.cpp
@@ -126,6 +126,32 @@ public:
 	}
 };
 
+//function to keep framenumber from beign too large do to user error, it can cause a crash if we snap too many frames.
+static void KeepFrameInRange(FFrameNumber& KeepFrameInRange, UMovieScene* MovieScene)
+{
+	if (MovieScene)
+	{
+		//limit values to 10x times the range from the end points.
+		FFrameNumber StartFrame = MovieScene->GetPlaybackRange().GetLowerBoundValue();
+		FFrameNumber EndFrame = MovieScene->GetPlaybackRange().GetUpperBoundValue();
+		if (EndFrame < StartFrame)
+		{
+			FFrameNumber Temp = StartFrame;
+			StartFrame = EndFrame;
+			EndFrame = Temp;
+		}
+		const FFrameNumber Max = (EndFrame - StartFrame) * 10;
+		if (KeepFrameInRange < (StartFrame - Max))
+		{
+			KeepFrameInRange = (StartFrame - Max);
+		}
+		else if (KeepFrameInRange > (EndFrame + Max))
+		{
+			KeepFrameInRange = (EndFrame + Max);
+		}
+	}
+}
+
 void SControlRigSnapper::Construct(const FArguments& InArgs)
 {
 	ClearActors();
@@ -265,6 +291,7 @@ void SControlRigSnapper::Construct(const FArguments& InArgs)
 											if (NewFrameTime.IsSet())
 											{
 												StartFrame = FFrameNumber((int32)NewFrameTime.GetValue());
+												KeepFrameInRange(StartFrame, Sequencer.Pin()->GetFocusedMovieSceneSequence()->GetMovieScene());
 											}
 										}
 									})
@@ -298,6 +325,7 @@ void SControlRigSnapper::Construct(const FArguments& InArgs)
 											if (NewFrameTime.IsSet())
 											{
 												EndFrame = FFrameNumber((int32)NewFrameTime.GetValue());
+												KeepFrameInRange(EndFrame, Sequencer.Pin()->GetFocusedMovieSceneSequence()->GetMovieScene());
 											}
 										}
 									})
