Sequencer: Animation Mode: Sub Float/Bool/Enum/Int Controls will show up individually in the details panel, other floats that aren't parented will show up collapsed as expected. We do this by adding some new details objects that allow for indivudal controls to not get merged into one details panel.

diff --git a/Engine/Plugins/Animation/ControlRig/Source/ControlRigEditor/Private/EditMode/ControlRigControlsProxy.cpp b/Engine/Plugins/Animation/ControlRig/Source/ControlRigEditor/Private/EditMode/ControlRigControlsProxy.cpp
--- a/Engine/Plugins/Animation/ControlRig/Source/ControlRigEditor/Private/EditMode/ControlRigControlsProxy.cpp
+++ b/Engine/Plugins/Animation/ControlRig/Source/ControlRigEditor/Private/EditMode/ControlRigControlsProxy.cpp
@@ -712,6 +712,11 @@ void UControlRigDetailPanelControlProxies::AddProxy(const FName& Name, UControlR
 			case ERigControlType::Float:
 			{
 				Proxy = NewObject<UControlRigFloatControlProxy>(GetTransientPackage(), NAME_None);
+				FRigControlElement* ParentControlElement = Cast<FRigControlElement>(ControlRig->GetHierarchy()->GetFirstParent(ControlElement));
+				if (ParentControlElement)
+				{
+					Proxy->bIsIndividual = true;
+				}
 				break;
 
 			}
@@ -727,6 +732,11 @@ void UControlRigDetailPanelControlProxies::AddProxy(const FName& Name, UControlR
 					EnumProxy->Enum.EnumType = ControlElement->Settings.ControlEnum;
 					Proxy = EnumProxy;
 				}
+				FRigControlElement* ParentControlElement = Cast<FRigControlElement>(ControlRig->GetHierarchy()->GetFirstParent(ControlElement));
+				if (ParentControlElement)
+				{
+					Proxy->bIsIndividual = true;
+				}
 				break;
 
 			}
@@ -747,6 +757,11 @@ void UControlRigDetailPanelControlProxies::AddProxy(const FName& Name, UControlR
 			case ERigControlType::Bool:
 			{
 				Proxy = NewObject<UControlRigBoolControlProxy>(GetTransientPackage(), NAME_None);
+				FRigControlElement* ParentControlElement = Cast<FRigControlElement>(ControlRig->GetHierarchy()->GetFirstParent(ControlElement));
+				if (ParentControlElement)
+				{
+					Proxy->bIsIndividual = true;
+				}
 				break;
 
 			}
@@ -824,11 +839,7 @@ void UControlRigDetailPanelControlProxies::SelectProxy(const FName& Name, bool b
 		{
 			if (!SelectedProxies.Contains(Proxy))
 			{
-				//don't show more than 5 for performance
-				if (SelectedProxies.Num() < 5)
-				{
-					SelectedProxies.Add(Proxy);
-				}
+				SelectedProxies.Add(Proxy);
 			}
 		}
 		else
diff --git a/Engine/Plugins/Animation/ControlRig/Source/ControlRigEditor/Private/EditMode/ControlRigControlsProxy.h b/Engine/Plugins/Animation/ControlRig/Source/ControlRigEditor/Private/EditMode/ControlRigControlsProxy.h
--- a/Engine/Plugins/Animation/ControlRig/Source/ControlRigEditor/Private/EditMode/ControlRigControlsProxy.h
+++ b/Engine/Plugins/Animation/ControlRig/Source/ControlRigEditor/Private/EditMode/ControlRigControlsProxy.h
@@ -36,6 +36,7 @@ public:
 	virtual void PostEditUndo() override;
 #endif
 	void SetIsMultiple(bool bIsVal); 
+	void SetIsIndividual(bool bIsVal);
 public:
 
 	FRigControlElement* GetControlElement() const;
@@ -45,6 +46,10 @@ public:
 	bool bSelected;
 	UPROPERTY(VisibleAnywhere, Category = "Control")
 	FName ControlName;
+
+	//if individual it will show up independently, this will happen for certain nested controls
+	bool bIsIndividual = false;
+
 protected:
 	bool bIsMultiple = 0;
 	FName Name;
diff --git a/Engine/Plugins/Animation/ControlRig/Source/ControlRigEditor/Private/EditMode/SControlRigDetails.cpp b/Engine/Plugins/Animation/ControlRig/Source/ControlRigEditor/Private/EditMode/SControlRigDetails.cpp
--- a/Engine/Plugins/Animation/ControlRig/Source/ControlRigEditor/Private/EditMode/SControlRigDetails.cpp
+++ b/Engine/Plugins/Animation/ControlRig/Source/ControlRigEditor/Private/EditMode/SControlRigDetails.cpp
@@ -235,6 +235,9 @@ void SControlRigDetails::Construct(const FArguments& InArgs, FControlRigEditMode
 		DetailsViewArgs.bShowScrollBar = false; // Don't need to show this, as we are putting it in a scroll box
 	}
 
+	FDetailsViewArgs IndividualDetailsViewArgs = DetailsViewArgs;
+	IndividualDetailsViewArgs.bAllowMultipleTopLevelObjects = true; //this is the secret sauce to show multiple objects in a details view
+
 	ControlEulerTransformDetailsView = FModuleManager::GetModuleChecked<FPropertyEditorModule>("PropertyEditor").CreateDetailView(DetailsViewArgs);
 	ControlEulerTransformDetailsView->SetKeyframeHandler(SharedThis(this));
 	ControlEulerTransformDetailsView->SetIsPropertyVisibleDelegate(FIsPropertyVisible::CreateSP(this, &SControlRigDetails::ShouldShowPropertyOnDetailCustomization));
@@ -259,24 +262,50 @@ void SControlRigDetails::Construct(const FArguments& InArgs, FControlRigEditMode
 	ControlFloatDetailsView->SetIsPropertyReadOnlyDelegate(FIsPropertyReadOnly::CreateSP(this, &SControlRigDetails::IsReadOnlyPropertyOnDetailCustomization));
 	ControlFloatDetailsView->SetGenericLayoutDetailsDelegate(FOnGetDetailCustomizationInstance::CreateStatic(&FControlRigEditModeGenericDetails::MakeInstance, ModeTools));
 
+	IndividualControlFloatDetailsView = FModuleManager::GetModuleChecked<FPropertyEditorModule>("PropertyEditor").CreateDetailView(IndividualDetailsViewArgs);
+	IndividualControlFloatDetailsView->SetKeyframeHandler(SharedThis(this));
+	IndividualControlFloatDetailsView->SetIsPropertyVisibleDelegate(FIsPropertyVisible::CreateSP(this, &SControlRigDetails::ShouldShowPropertyOnDetailCustomization));
+	IndividualControlFloatDetailsView->SetIsPropertyReadOnlyDelegate(FIsPropertyReadOnly::CreateSP(this, &SControlRigDetails::IsReadOnlyPropertyOnDetailCustomization));
+	IndividualControlFloatDetailsView->SetGenericLayoutDetailsDelegate(FOnGetDetailCustomizationInstance::CreateStatic(&FControlRigEditModeGenericDetails::MakeInstance, ModeTools));
+
 	ControlEnumDetailsView = FModuleManager::GetModuleChecked<FPropertyEditorModule>("PropertyEditor").CreateDetailView(DetailsViewArgs);
 	ControlEnumDetailsView->SetKeyframeHandler(SharedThis(this));
 	ControlEnumDetailsView->SetIsPropertyVisibleDelegate(FIsPropertyVisible::CreateSP(this, &SControlRigDetails::ShouldShowPropertyOnDetailCustomization));
 	ControlEnumDetailsView->SetIsPropertyReadOnlyDelegate(FIsPropertyReadOnly::CreateSP(this, &SControlRigDetails::IsReadOnlyPropertyOnDetailCustomization));
 	ControlEnumDetailsView->SetGenericLayoutDetailsDelegate(FOnGetDetailCustomizationInstance::CreateStatic(&FControlRigEditModeGenericDetails::MakeInstance, ModeTools));
 
+	IndividualControlEnumDetailsView = FModuleManager::GetModuleChecked<FPropertyEditorModule>("PropertyEditor").CreateDetailView(IndividualDetailsViewArgs);
+	IndividualControlEnumDetailsView->SetKeyframeHandler(SharedThis(this));
+	IndividualControlEnumDetailsView->SetIsPropertyVisibleDelegate(FIsPropertyVisible::CreateSP(this, &SControlRigDetails::ShouldShowPropertyOnDetailCustomization));
+	IndividualControlEnumDetailsView->SetIsPropertyReadOnlyDelegate(FIsPropertyReadOnly::CreateSP(this, &SControlRigDetails::IsReadOnlyPropertyOnDetailCustomization));
+	IndividualControlEnumDetailsView->SetGenericLayoutDetailsDelegate(FOnGetDetailCustomizationInstance::CreateStatic(&FControlRigEditModeGenericDetails::MakeInstance, ModeTools));
+
 	ControlIntegerDetailsView = FModuleManager::GetModuleChecked<FPropertyEditorModule>("PropertyEditor").CreateDetailView(DetailsViewArgs);
 	ControlIntegerDetailsView->SetKeyframeHandler(SharedThis(this));
 	ControlIntegerDetailsView->SetIsPropertyVisibleDelegate(FIsPropertyVisible::CreateSP(this, &SControlRigDetails::ShouldShowPropertyOnDetailCustomization));
 	ControlIntegerDetailsView->SetIsPropertyReadOnlyDelegate(FIsPropertyReadOnly::CreateSP(this, &SControlRigDetails::IsReadOnlyPropertyOnDetailCustomization));
 	ControlIntegerDetailsView->SetGenericLayoutDetailsDelegate(FOnGetDetailCustomizationInstance::CreateStatic(&FControlRigEditModeGenericDetails::MakeInstance, ModeTools));
 
+	IndividualControlIntegerDetailsView = FModuleManager::GetModuleChecked<FPropertyEditorModule>("PropertyEditor").CreateDetailView(IndividualDetailsViewArgs);
+	IndividualControlIntegerDetailsView->SetKeyframeHandler(SharedThis(this));
+	IndividualControlIntegerDetailsView->SetIsPropertyVisibleDelegate(FIsPropertyVisible::CreateSP(this, &SControlRigDetails::ShouldShowPropertyOnDetailCustomization));
+	IndividualControlIntegerDetailsView->SetIsPropertyReadOnlyDelegate(FIsPropertyReadOnly::CreateSP(this, &SControlRigDetails::IsReadOnlyPropertyOnDetailCustomization));
+	IndividualControlIntegerDetailsView->SetGenericLayoutDetailsDelegate(FOnGetDetailCustomizationInstance::CreateStatic(&FControlRigEditModeGenericDetails::MakeInstance, ModeTools));
+
+
 	ControlBoolDetailsView = FModuleManager::GetModuleChecked<FPropertyEditorModule>("PropertyEditor").CreateDetailView(DetailsViewArgs);
 	ControlBoolDetailsView->SetKeyframeHandler(SharedThis(this));
 	ControlBoolDetailsView->SetIsPropertyVisibleDelegate(FIsPropertyVisible::CreateSP(this, &SControlRigDetails::ShouldShowPropertyOnDetailCustomization));
 	ControlBoolDetailsView->SetIsPropertyReadOnlyDelegate(FIsPropertyReadOnly::CreateSP(this, &SControlRigDetails::IsReadOnlyPropertyOnDetailCustomization));
 	ControlBoolDetailsView->SetGenericLayoutDetailsDelegate(FOnGetDetailCustomizationInstance::CreateStatic(&FControlRigEditModeGenericDetails::MakeInstance, ModeTools));
 
+	IndividualControlBoolDetailsView = FModuleManager::GetModuleChecked<FPropertyEditorModule>("PropertyEditor").CreateDetailView(IndividualDetailsViewArgs);
+	IndividualControlBoolDetailsView->SetKeyframeHandler(SharedThis(this));
+	IndividualControlBoolDetailsView->SetIsPropertyVisibleDelegate(FIsPropertyVisible::CreateSP(this, &SControlRigDetails::ShouldShowPropertyOnDetailCustomization));
+	IndividualControlBoolDetailsView->SetIsPropertyReadOnlyDelegate(FIsPropertyReadOnly::CreateSP(this, &SControlRigDetails::IsReadOnlyPropertyOnDetailCustomization));
+	IndividualControlBoolDetailsView->SetGenericLayoutDetailsDelegate(FOnGetDetailCustomizationInstance::CreateStatic(&FControlRigEditModeGenericDetails::MakeInstance, ModeTools));
+
+
 	ControlVectorDetailsView = FModuleManager::GetModuleChecked<FPropertyEditorModule>("PropertyEditor").CreateDetailView(DetailsViewArgs);
 	ControlVectorDetailsView->SetKeyframeHandler(SharedThis(this));
 	ControlVectorDetailsView->SetIsPropertyVisibleDelegate(FIsPropertyVisible::CreateSP(this, &SControlRigDetails::ShouldShowPropertyOnDetailCustomization));
@@ -342,6 +371,26 @@ void SControlRigDetails::Construct(const FArguments& InArgs, FControlRigEditMode
 				[
 					ControlFloatDetailsView.ToSharedRef()
 				]
+			+ SVerticalBox::Slot()
+				.AutoHeight()
+				[
+					IndividualControlBoolDetailsView.ToSharedRef()
+				]
+			+ SVerticalBox::Slot()
+				.AutoHeight()
+				[
+					IndividualControlIntegerDetailsView.ToSharedRef()
+				]
+			+ SVerticalBox::Slot()
+				.AutoHeight()
+				[
+					IndividualControlEnumDetailsView.ToSharedRef()
+				]
+			+ SVerticalBox::Slot()
+				.AutoHeight()
+				[
+					IndividualControlFloatDetailsView.ToSharedRef()
+				]
 
 			]
 		];
@@ -381,6 +430,11 @@ void SControlRigDetails::UpdateProxies()
 		TArray<TWeakObjectPtr<>> Bools;
 		TArray<TWeakObjectPtr<>> Integers;
 		TArray<TWeakObjectPtr<>> Enums;
+		TArray<TWeakObjectPtr<>> IndividualFloats;
+		TArray<TWeakObjectPtr<>> IndividualBools;
+		TArray<TWeakObjectPtr<>> IndividualIntegers;
+		TArray<TWeakObjectPtr<>> IndividualEnums;
+
 		if (UControlRig* ControlRig = GetControlRig())
 		{
 			if (FControlRigEditMode* EditMode = static_cast<FControlRigEditMode*>(ModeTools->GetActiveMode(FControlRigEditMode::ModeName)))
@@ -404,7 +458,14 @@ void SControlRigDetails::UpdateProxies()
 						}
 						else if (Proxy->GetClass() == UControlRigFloatControlProxy::StaticClass())
 						{
-							Floats.Add(Proxy);
+							if (Proxy->bIsIndividual)
+							{
+								IndividualFloats.Add(Proxy);
+							}
+							else
+							{
+								Floats.Add(Proxy);
+							}
 						}
 						else if (Proxy->GetClass() == UControlRigVectorControlProxy::StaticClass())
 						{
@@ -416,15 +477,36 @@ void SControlRigDetails::UpdateProxies()
 						}
 						else if (Proxy->GetClass() == UControlRigBoolControlProxy::StaticClass())
 						{
-							Bools.Add(Proxy);
+							if (Proxy->bIsIndividual)
+							{
+								IndividualBools.Add(Proxy);
+							}
+							else
+							{
+								Bools.Add(Proxy);
+							}
 						}
 						else if (Proxy->GetClass() == UControlRigEnumControlProxy::StaticClass())
 						{
-							Enums.Add(Proxy);
+							if (Proxy->bIsIndividual)
+							{
+								IndividualEnums.Add(Proxy);
+							}
+							else
+							{
+								Enums.Add(Proxy);
+							}
 						}
 						else if (Proxy->GetClass() == UControlRigIntegerControlProxy::StaticClass())
 						{
-							Integers.Add(Proxy);
+							if (Proxy->bIsIndividual)
+							{
+								IndividualIntegers.Add(Proxy);
+							}
+							else
+							{
+								Integers.Add(Proxy);
+							}
 						}
 					}
 				}
@@ -436,17 +518,29 @@ void SControlRigDetails::UpdateProxies()
 						Proxy->SetIsMultiple(Transforms.Num() > 1);
 					}
 				}
+				for (TWeakObjectPtr<>& Object : Floats)
+				{
+					UControlRigControlsProxy* Proxy = Cast<UControlRigControlsProxy>(Object.Get());
+					if (Proxy)
+					{
+						Proxy->SetIsMultiple(Transforms.Num() > 1);
+					}
+				}
 			}
 		}
 		SetTransformDetailsObjects(Transforms);
 		SetTransformNoScaleDetailsObjects(TransformNoScales);
 		SetEulerTransformDetailsObjects(Eulers);
-		SetFloatDetailsObjects(Floats);
 		SetVectorDetailsObjects(Vectors);
 		SetVector2DDetailsObjects(Vector2Ds);
-		SetBoolDetailsObjects(Bools);
-		SetIntegerDetailsObjects(Integers);
-		SetEnumDetailsObjects(Enums);
+		SetFloatDetailsObjects(Floats, false);
+		SetBoolDetailsObjects(Bools,false);
+		SetIntegerDetailsObjects(Integers,false);
+		SetEnumDetailsObjects(Enums,false);
+		SetFloatDetailsObjects(IndividualFloats, true);
+		SetBoolDetailsObjects(IndividualBools, true);
+		SetIntegerDetailsObjects(IndividualIntegers, true);
+		SetEnumDetailsObjects(IndividualEnums, true);
 	});
 }
 
@@ -474,42 +568,71 @@ void SControlRigDetails::SetTransformNoScaleDetailsObjects(const TArray<TWeakObj
 	}
 }
 
-void SControlRigDetails::SetFloatDetailsObjects(const TArray<TWeakObjectPtr<>>& InObjects)
+void SControlRigDetails::SetFloatDetailsObjects(const TArray<TWeakObjectPtr<>>& InObjects, bool bIsIndividual)
 {
-	if (ControlFloatDetailsView)
+	if (bIsIndividual)
+	{
+		if (IndividualControlFloatDetailsView)
+		{
+			IndividualControlFloatDetailsView->SetObjects(InObjects);
+		}
+	}
+	else if (ControlFloatDetailsView)
 	{
 		ControlFloatDetailsView->SetObjects(InObjects);
 	}
 }
 
-void SControlRigDetails::SetBoolDetailsObjects(const TArray<TWeakObjectPtr<>>& InObjects)
+void SControlRigDetails::SetBoolDetailsObjects(const TArray<TWeakObjectPtr<>>& InObjects, bool bIsIndividual)
 {
-	if (ControlBoolDetailsView)
+	if (bIsIndividual)
+	{
+		if (IndividualControlBoolDetailsView)
+		{
+			IndividualControlBoolDetailsView->SetObjects(InObjects);
+		}
+	}
+	else if (ControlBoolDetailsView)
 	{
 		ControlBoolDetailsView->SetObjects(InObjects);
 	}
 }
 
-void SControlRigDetails::SetIntegerDetailsObjects(const TArray<TWeakObjectPtr<>>& InObjects)
+void SControlRigDetails::SetIntegerDetailsObjects(const TArray<TWeakObjectPtr<>>& InObjects, bool bIsIndividual)
 {
-	if (ControlIntegerDetailsView)
+
+	if (bIsIndividual)
+	{
+		if (IndividualControlIntegerDetailsView)
+		{
+			IndividualControlIntegerDetailsView->SetObjects(InObjects);
+		}
+	}
+	else if (ControlIntegerDetailsView)
 	{
 		ControlIntegerDetailsView->SetObjects(InObjects);
 	}
 }
-void SControlRigDetails::SetEnumDetailsObjects(const TArray<TWeakObjectPtr<>>& InObjects)
+void SControlRigDetails::SetEnumDetailsObjects(const TArray<TWeakObjectPtr<>>& InObjects, bool bIsIndividual)
 {
-	if (ControlVectorDetailsView)
+	if (bIsIndividual)
 	{
-		ControlVectorDetailsView->SetObjects(InObjects);
+		if (IndividualControlEnumDetailsView)
+		{
+			IndividualControlEnumDetailsView->SetObjects(InObjects);
+		}
+	}
+	else if (ControlEnumDetailsView)
+	{
+		ControlEnumDetailsView->SetObjects(InObjects);
 	}
 }
 
 void SControlRigDetails::SetVectorDetailsObjects(const TArray<TWeakObjectPtr<>>& InObjects)
 {
-	if (ControlEnumDetailsView)
+	if (ControlVectorDetailsView)
 	{
-		ControlEnumDetailsView->SetObjects(InObjects);
+		ControlVectorDetailsView->SetObjects(InObjects);
 	}
 }
 
diff --git a/Engine/Plugins/Animation/ControlRig/Source/ControlRigEditor/Private/EditMode/SControlRigDetails.h b/Engine/Plugins/Animation/ControlRig/Source/ControlRigEditor/Private/EditMode/SControlRigDetails.h
--- a/Engine/Plugins/Animation/ControlRig/Source/ControlRigEditor/Private/EditMode/SControlRigDetails.h
+++ b/Engine/Plugins/Animation/ControlRig/Source/ControlRigEditor/Private/EditMode/SControlRigDetails.h
@@ -50,18 +50,6 @@ class SControlRigDetails: public SCompoundWidget, public FControlRigBaseDockable
 
 	void Construct(const FArguments& InArgs, FControlRigEditMode& InEditMode);
 
-	/** Set the objects to be displayed in the details panel */
-	void SetSettingsDetailsObject(const TWeakObjectPtr<>& InObject);
-	void SetEulerTransformDetailsObjects(const TArray<TWeakObjectPtr<>>& InObjects);
-	void SetTransformDetailsObjects(const TArray<TWeakObjectPtr<>>& InObjects);
-	void SetTransformNoScaleDetailsObjects(const TArray<TWeakObjectPtr<>>& InObjects);
-	void SetFloatDetailsObjects(const TArray<TWeakObjectPtr<>>& InObjects);
-	void SetBoolDetailsObjects(const TArray<TWeakObjectPtr<>>& InObjects);
-	void SetIntegerDetailsObjects(const TArray<TWeakObjectPtr<>>& InObjects);
-	void SetEnumDetailsObjects(const TArray<TWeakObjectPtr<>>& InObjects);
-	void SetVectorDetailsObjects(const TArray<TWeakObjectPtr<>>& InObjects);
-	void SetVector2DDetailsObjects(const TArray<TWeakObjectPtr<>>& InObjects);
-
 	// IDetailKeyframeHandler interface
 	virtual bool IsPropertyKeyable(const UClass* InObjectClass, const class IPropertyHandle& PropertyHandle) const override;
 	virtual bool IsPropertyKeyingEnabled() const override;
@@ -73,10 +61,24 @@ class SControlRigDetails: public SCompoundWidget, public FControlRigBaseDockable
 	bool IsReadOnlyPropertyOnDetailCustomization(const struct FPropertyAndParent& InPropertyAndParent) const;
 
 private:
+
+	/** Set the objects to be displayed in the details panel */
+	void SetSettingsDetailsObject(const TWeakObjectPtr<>& InObject);
+	void SetEulerTransformDetailsObjects(const TArray<TWeakObjectPtr<>>& InObjects);
+	void SetTransformDetailsObjects(const TArray<TWeakObjectPtr<>>& InObjects);
+	void SetTransformNoScaleDetailsObjects(const TArray<TWeakObjectPtr<>>& InObjects);
+	void SetVectorDetailsObjects(const TArray<TWeakObjectPtr<>>& InObjects);
+	void SetVector2DDetailsObjects(const TArray<TWeakObjectPtr<>>& InObjects);
+	void SetFloatDetailsObjects(const TArray<TWeakObjectPtr<>>& InObjects, bool bIsIndividual);
+	void SetBoolDetailsObjects(const TArray<TWeakObjectPtr<>>& InObjects, bool bIsIndividual);
+	void SetIntegerDetailsObjects(const TArray<TWeakObjectPtr<>>& InObjects, bool bIsIndividual);
+	void SetEnumDetailsObjects(const TArray<TWeakObjectPtr<>>& InObjects, bool bIsIndividual);
+
 	void UpdateProxies();
 	virtual void HandleControlSelected(UControlRig* Subject, FRigControlElement* InControl, bool bSelected) override;
 	virtual void HandleControlAdded(UControlRig* ControlRig, bool bIsAdded) override;
 
+	//these views will hold more than one of the same tyhpe
 	TSharedPtr<IDetailsView> ControlEulerTransformDetailsView;
 	TSharedPtr<IDetailsView> ControlTransformDetailsView;
 	TSharedPtr<IDetailsView> ControlTransformNoScaleDetailsView;
@@ -86,6 +88,11 @@ private:
 	TSharedPtr<IDetailsView> ControlEnumDetailsView;
 	TSharedPtr<IDetailsView> ControlVector2DDetailsView;
 	TSharedPtr<IDetailsView> ControlVectorDetailsView;
+	//these will show more than one of the same type, will happen for controls with parents.
+	TSharedPtr<IDetailsView> IndividualControlFloatDetailsView;
+	TSharedPtr<IDetailsView> IndividualControlBoolDetailsView;
+	TSharedPtr<IDetailsView> IndividualControlIntegerDetailsView;
+	TSharedPtr<IDetailsView> IndividualControlEnumDetailsView;
 
 };
 
