Space Switching: keep tangents when switching space to keep arriving / leaving animation as is
Code was by Benoit, Mike reviewed but Mike getting it in due to it being late.

diff --git a/Engine/Plugins/Animation/ControlRig/Source/ControlRigEditor/Private/Sequencer/ControlRigSpaceChannelEditors.cpp b/Engine/Plugins/Animation/ControlRig/Source/ControlRigEditor/Private/Sequencer/ControlRigSpaceChannelEditors.cpp
--- a/Engine/Plugins/Animation/ControlRig/Source/ControlRigEditor/Private/Sequencer/ControlRigSpaceChannelEditors.cpp
+++ b/Engine/Plugins/Animation/ControlRig/Source/ControlRigEditor/Private/Sequencer/ControlRigSpaceChannelEditors.cpp
@@ -374,7 +374,15 @@ FKeyHandle FControlRigSpaceChannelHelpers::SequencerKeyControlRigSpaceChannel(UC
 		FControlRigSnapper Snapper;
 		Snapper.GetControlRigControlTransforms(Sequencer, ControlRig, ControlKey.Name, Frames, ControlRigParentWorldTransforms, ControlWorldTransforms);
 
-		// if current Value not the same as the previous add new space kehy
+		// store tangents to keep the animation when switching 
+		TArray<FMovieSceneTangentData> Tangents;
+		UMovieSceneControlRigParameterSection* ControlRigSection = Cast<UMovieSceneControlRigParameterSection>(SectionToKey);
+		if (ControlRigSection)
+		{
+			EvaluateTangentAtThisTime(ControlRig, ControlRigSection, ControlKey.Name, Time, Tangents);
+		}
+		
+		// if current Value not the same as the previous add new space key
 		if (PreviousValue != Value)
 		{
 			SectionToKey->Modify();
@@ -393,9 +401,9 @@ FKeyHandle FControlRigSpaceChannelHelpers::SequencerKeyControlRigSpaceChannel(UC
 		}
 		else //same as previous
 		{
-			if (UMovieSceneControlRigParameterSection* Section = Cast<UMovieSceneControlRigParameterSection>(SectionToKey))
+			if (ControlRigSection)
 			{
-				Section->Modify();
+				ControlRigSection->Modify();
 				TArray<FFrameNumber> OurKeyTimes;
 				TArray<FKeyHandle> OurKeyHandles;
 				TRange<FFrameNumber> CurrentFrameRange;
@@ -407,10 +415,10 @@ FKeyHandle FControlRigSpaceChannelHelpers::SequencerKeyControlRigSpaceChannel(UC
 					ChannelInterface.DeleteKeys(OurKeyHandles);
 				}
 				//now delete any extra TimeOfDelete -1
-				FControlRigSpaceChannelHelpers::DeleteTransformKeysAtThisTime(ControlRig, Section, ControlKey.Name, Time - 1);
+				// NOTE do we need to update the tangents here?
+				FControlRigSpaceChannelHelpers::DeleteTransformKeysAtThisTime(ControlRig, ControlRigSection, ControlKey.Name, Time - 1);
 				bSetPreviousKey = false; //also don't set previous
 			}
-
 		}
 
 		//do any compensation or previous key adding
@@ -418,28 +426,52 @@ FKeyHandle FControlRigSpaceChannelHelpers::SequencerKeyControlRigSpaceChannel(UC
 		Context.SetKey = EControlRigSetKey::Always;
 		FFrameRate TickResolution = Sequencer->GetFocusedTickResolution();
 
+		// set previous key if we're going to switch space
 		if (bSetPreviousKey)
 		{
 			const FFrameTime GlobalTime(Time -1);
+
 			FMovieSceneContext SceneContext = FMovieSceneContext(FMovieSceneEvaluationRange(GlobalTime, TickResolution), Sequencer->GetPlaybackStatus()).SetHasJumped(true);
 			Sequencer->GetEvaluationTemplate().Evaluate(SceneContext, *Sequencer);
 
 			Context.LocalTime = TickResolution.AsSeconds(FFrameTime(Time - 1));
 			ControlRig->SetControlGlobalTransform(ControlKey.Name, ControlWorldTransforms[0], true, Context);
+
+			// need to update the tangents here to keep the arriving animation
+			if (ControlRigSection)
+			{
+				// ControlRigSection->Modify(); ?
+				SetTangentsAtThisTime(ControlRig, ControlRigSection, ControlKey.Name, GlobalTime.GetFrame(), Tangents);
+			}
 		}
 
+		// effectively switch to new space
 		URigHierarchy::TElementDependencyMap Dependencies = ControlRig->GetHierarchy()->GetDependenciesForVM(ControlRig->GetVM());
 		ControlRig->GetHierarchy()->SwitchToParent(ControlKey, SpaceKey, false, true, Dependencies, nullptr);
+
+		// add new keys in the new space context
 		int32 FramesIndex = 0;
 		for (const FFrameNumber& Frame : Frames)
 		{
 			const FFrameTime GlobalTime(Frame);
+			
 			FMovieSceneContext SceneContext = FMovieSceneContext(FMovieSceneEvaluationRange(GlobalTime, TickResolution), Sequencer->GetPlaybackStatus()).SetHasJumped(true);
 			Sequencer->GetEvaluationTemplate().Evaluate(SceneContext, *Sequencer);
 
 			ControlRig->Evaluate_AnyThread();
 			Context.LocalTime = TickResolution.AsSeconds(FFrameTime(Frame));
 			ControlRig->SetControlGlobalTransform(ControlKey.Name, ControlWorldTransforms[FramesIndex], true, Context);
+
+			// need to update the tangents here to keep the leaving animation
+			if (bSetPreviousKey && FramesIndex == 0)
+			{
+				if (ControlRigSection)
+				{
+					// ControlRigSection->Modify(); ?
+					SetTangentsAtThisTime(ControlRig, ControlRigSection, ControlKey.Name, GlobalTime.GetFrame(), Tangents);
+				}
+			}
+
 			FramesIndex++;
 		}
 		Channel->FindSpaceIntervals();
@@ -1084,7 +1116,7 @@ void FControlRigSpaceChannelHelpers::CompensateIfNeeded(UControlRig* ControlRig,
 	{
 		return;
 	}
-
+	
 	//we need to check all controls for 1) space and 2) previous frame and if so we automatically compensate.
 	TArray<FRigControlElement*> Controls = ControlRig->GetHierarchy()->GetControls();
 	bool bDidIt = false;
@@ -1328,6 +1360,166 @@ TArray<FKeyBarCurveModel::FBarRange> FControlRigSpaceChannelHelpers::FindRanges(
 	return  Range;
 }
 
+// NOTE use this function in HandleSpaceKeyTimeChanged, DeleteTransformKeysAtThisTime, ...
+TPair<FRigControlElement*, FChannelMapInfo*> FControlRigSpaceChannelHelpers::GetControlAndChannelInfo(UControlRig* ControlRig, UMovieSceneControlRigParameterSection* Section, FName ControlName)
+{
+	FRigControlElement* ControlElement = ControlRig ? ControlRig->FindControl(ControlName) : nullptr;
+	FChannelMapInfo* pChannelIndex = Section ? Section->ControlChannelMap.Find(ControlName) : nullptr;
+
+	return TPair<FRigControlElement*, FChannelMapInfo*>(ControlElement, pChannelIndex);
+}
+
+// NOTE use this function in HandleSpaceKeyTimeChanged, DeleteTransformKeysAtThisTime, ...
+int32 FControlRigSpaceChannelHelpers::GetNumFloatChannels(const ERigControlType InControlType)
+{
+	switch (InControlType)
+	{
+	case ERigControlType::Position:
+	case ERigControlType::Scale:
+	case ERigControlType::Rotator:
+		return 3;
+	case ERigControlType::TransformNoScale:
+		return 6;
+	case ERigControlType::Transform:
+	case ERigControlType::EulerTransform:
+		return 9;
+	default:
+		break;
+	}
+	return 0;
+}
+
+void FControlRigSpaceChannelHelpers::EvaluateTangentAtThisTime(	UControlRig* ControlRig,
+																UMovieSceneControlRigParameterSection* Section,
+																FName ControlName,
+																FFrameNumber Time,
+																TArray<FMovieSceneTangentData>& OutTangents)
+{
+	// NOTE this might be moved to FMovieSceneFloatChannel
+	auto EvaluateTangent = [](const FMovieSceneFloatChannel* InChannel, FFrameNumber InTime)
+	{
+		const TMovieSceneChannelData<const FMovieSceneFloatValue> ChannelInterface = InChannel->GetData();
+		const TArrayView<const FFrameNumber> Times = ChannelInterface.GetTimes();
+
+		// Need at least two keys to evaluate a derivative
+		if (Times.Num() < 2)
+		{
+			static const FMovieSceneTangentData DefaultTangent;
+			return DefaultTangent;
+		}
+
+		// look around to get the closest key tangent if fairly close
+		const int32 Tolerance = static_cast<int32>(InChannel->GetTickResolution().AsDecimal() * 0.01);
+		const int32 ExistingIndex = ChannelInterface.FindKey(InTime, Tolerance);
+		if (ExistingIndex != INDEX_NONE)
+		{
+			const FMovieSceneFloatValue& Value = ChannelInterface.GetValues()[ExistingIndex];
+			const FFrameNumber& Time = Times[ExistingIndex];
+			const int32 DiffToKey = InTime.Value - Time.Value;
+
+			// if the closest key is within a threshold, we use it's tangent directly instead of computing one
+			if (FMath::Abs(DiffToKey) <= Tolerance)
+			{
+				if (DiffToKey == 0)
+				{
+					return Value.Tangent;				
+				}
+
+				FMovieSceneTangentData Tangent = Value.Tangent;
+				if (DiffToKey < 0) // pretty close to the next key
+				{
+				 	Tangent.LeaveTangent = 0.f;
+				}
+				else // pretty close to the previous key
+				{
+					Tangent.ArriveTangent = 0.f;
+				}
+				return Tangent;
+			}
+		}
+
+		// compute tangent using central difference
+		// NOTE we may wanna compute a backward / forward difference instead
+		const int32 Delta = static_cast<int32>(InChannel->GetTickResolution().AsDecimal() * 0.1);
+
+		float PrevValue = 0.f; InChannel->Evaluate(InTime-Delta, PrevValue);
+		float NextValue = 0.f; InChannel->Evaluate(InTime+Delta, NextValue);
+		const float TangentValue = (NextValue - PrevValue) / (2.f*Delta);
+		
+		FMovieSceneTangentData TangentData;
+		TangentData.ArriveTangent = TangentValue;
+		TangentData.LeaveTangent = TangentValue;
+	
+		return TangentData;
+	};
+	
+	FChannelMapInfo* pChannelIndex = nullptr;
+	FRigControlElement* ControlElement = nullptr;
+	Tie(ControlElement, pChannelIndex) = GetControlAndChannelInfo(ControlRig, Section, ControlName);
+	
+	if (!pChannelIndex || !ControlElement)
+	{
+		return;
+	}
+
+	// get the number of float channels to treat
+	const int NumChannels = GetNumFloatChannels(ControlElement->Settings.ControlType);
+	if (!NumChannels)
+	{
+		return;
+	}
+
+	// compute and store tangents
+	OutTangents.SetNum(NumChannels);
+
+	const TArrayView<FMovieSceneFloatChannel*> FloatChannels = Section->GetChannelProxy().GetChannels<FMovieSceneFloatChannel>();
+	int32 ChannelIndex = pChannelIndex->ChannelIndex;
+	for (int32 Index = 0; Index < NumChannels; ++Index, ++ChannelIndex)
+	{
+		OutTangents[Index] = EvaluateTangent(FloatChannels[ChannelIndex], Time);
+	}
+}
+
+// NOTE we may pass an enum to tell which tangent we wanna set (arrive, leave, both)
+void FControlRigSpaceChannelHelpers::SetTangentsAtThisTime(	UControlRig* ControlRig,
+															UMovieSceneControlRigParameterSection* Section,
+															FName ControlName,
+															FFrameNumber Time,
+															const TArray<FMovieSceneTangentData>& InTangents)
+{
+	FChannelMapInfo* pChannelIndex = nullptr;
+    FRigControlElement* ControlElement = nullptr;
+    Tie(ControlElement, pChannelIndex) = GetControlAndChannelInfo(ControlRig, Section, ControlName);
+    
+    if (!pChannelIndex || !ControlElement)
+    {
+    	return;
+    }
+
+	// get the number of float channels to treat
+	const int32 NumChannels = GetNumFloatChannels(ControlElement->Settings.ControlType);
+	if (!NumChannels || InTangents.Num() != NumChannels)
+	{
+		return;
+	}
+
+	const TArrayView<FMovieSceneFloatChannel*> FloatChannels = Section->GetChannelProxy().GetChannels<FMovieSceneFloatChannel>();
+	int32 ChannelIndex = pChannelIndex->ChannelIndex;
+
+	for (int32 Index = 0; Index < NumChannels; ++Index, ++ChannelIndex)
+	{
+		TMovieSceneChannelData<FMovieSceneFloatValue> ChannelInterface = FloatChannels[ChannelIndex]->GetData();
+		const int32 KeyIndex = ChannelInterface.FindKey(Time);
+		if (KeyIndex != INDEX_NONE)
+		{
+			FMovieSceneFloatValue& Value = ChannelInterface.GetValues()[KeyIndex];
+			Value.Tangent.ArriveTangent = InTangents[Index].ArriveTangent;
+			Value.Tangent.LeaveTangent = InTangents[Index].LeaveTangent;
+			Value.TangentMode = RCTM_Break;	
+		}
+	}
+}
+
 void DrawExtra(FMovieSceneControlRigSpaceChannel* Channel, const UMovieSceneSection* Owner,const FGeometry& AllottedGeometry, FSequencerSectionPainter& Painter)
 {
 	if (const UMovieSceneControlRigParameterSection* Section = Cast<UMovieSceneControlRigParameterSection>(Owner))
diff --git a/Engine/Plugins/Animation/ControlRig/Source/ControlRigEditor/Public/ControlRigSpaceChannelEditors.h b/Engine/Plugins/Animation/ControlRig/Source/ControlRigEditor/Public/ControlRigSpaceChannelEditors.h
--- a/Engine/Plugins/Animation/ControlRig/Source/ControlRigEditor/Public/ControlRigSpaceChannelEditors.h
+++ b/Engine/Plugins/Animation/ControlRig/Source/ControlRigEditor/Public/ControlRigSpaceChannelEditors.h
@@ -60,6 +60,14 @@ struct FControlRigSpaceChannelHelpers
 	static FReply OpenBakeDialog(ISequencer* Sequencer, FMovieSceneControlRigSpaceChannel* Channel, int32 KeyIndex, UMovieSceneSection* SectionToKey);
 	static TArray<FKeyBarCurveModel::FBarRange> FindRanges(FMovieSceneControlRigSpaceChannel* Channel, const UMovieSceneSection* Section);
 
+	// evaluates the tangent by central difference at a specific time. (return the closest key tangent if within an acceptable tolerance) 
+	static void EvaluateTangentAtThisTime(UControlRig* ControlRig, UMovieSceneControlRigParameterSection* Section, FName ControlName, FFrameNumber Time, TArray<FMovieSceneTangentData>& OutTangents);
+	// set the tangent at a specific time. (if there's a key to set at this time)
+	static void SetTangentsAtThisTime(UControlRig* ControlRig, UMovieSceneControlRigParameterSection* Section, FName ControlName, FFrameNumber Time, const TArray<FMovieSceneTangentData>& InTangents);
+	// retrieve the control and the channel infos for that ControlRig/Section.
+	static TPair<FRigControlElement*, FChannelMapInfo*> GetControlAndChannelInfo(UControlRig* ControlRig, UMovieSceneControlRigParameterSection* Section, FName ControlName);
+	// retrieve the number of float channels based on the control type.
+	static int32 GetNumFloatChannels(const ERigControlType InControlType);
 };
 
 //template specialization
