Control Rig: Implement upgrade path to all nodes

diff --git a/Engine/Plugins/Animation/ControlRig/Source/ControlRig/Private/Units/Collision/RigUnit_WorldCollision.cpp b/Engine/Plugins/Animation/ControlRig/Source/ControlRig/Private/Units/Collision/RigUnit_WorldCollision.cpp
--- a/Engine/Plugins/Animation/ControlRig/Source/ControlRig/Private/Units/Collision/RigUnit_WorldCollision.cpp
+++ b/Engine/Plugins/Animation/ControlRig/Source/ControlRig/Private/Units/Collision/RigUnit_WorldCollision.cpp
@@ -46,6 +46,16 @@ FRigUnit_SphereTraceWorld_Execute()
 	}
 }
 
+FRigVMStructUpgradeInfo FRigUnit_SphereTraceWorld::GetUpgradeInfo() const
+{
+	FRigUnit_SphereTraceByTraceChannel NewNode;
+	NewNode.Start = Start;
+	NewNode.End = End;
+	NewNode.Radius = Radius;
+
+	return FRigVMStructUpgradeInfo(*this, NewNode);
+}
+
 FRigUnit_SphereTraceByTraceChannel_Execute()
 {
 	DECLARE_SCOPE_HIERARCHICAL_COUNTER_RIGUNIT()
diff --git a/Engine/Plugins/Animation/ControlRig/Source/ControlRig/Private/Units/Collision/RigUnit_WorldCollision.h b/Engine/Plugins/Animation/ControlRig/Source/ControlRig/Private/Units/Collision/RigUnit_WorldCollision.h
--- a/Engine/Plugins/Animation/ControlRig/Source/ControlRig/Private/Units/Collision/RigUnit_WorldCollision.h
+++ b/Engine/Plugins/Animation/ControlRig/Source/ControlRig/Private/Units/Collision/RigUnit_WorldCollision.h
@@ -56,6 +56,9 @@ struct CONTROLRIG_API FRigUnit_SphereTraceWorld : public FRigUnit
 	/** Hit normal in rig / global Space */
 	UPROPERTY(meta = (Output))
 	FVector HitNormal;
+
+	RIGVM_METHOD()
+	virtual FRigVMStructUpgradeInfo GetUpgradeInfo() const override;
 };
 
 /**
diff --git a/Engine/Plugins/Animation/ControlRig/Source/ControlRig/Private/Units/Control/RigUnit_Control.cpp b/Engine/Plugins/Animation/ControlRig/Source/ControlRig/Private/Units/Control/RigUnit_Control.cpp
--- a/Engine/Plugins/Animation/ControlRig/Source/ControlRig/Private/Units/Control/RigUnit_Control.cpp
+++ b/Engine/Plugins/Animation/ControlRig/Source/ControlRig/Private/Units/Control/RigUnit_Control.cpp
@@ -70,3 +70,8 @@ FEulerTransform FRigUnit_Control::StaticGetFilteredTransform(const FEulerTransfo
 	InFilter.FilterTransform(FilteredTransform);
 	return FilteredTransform;
 }
+
+FRigVMStructUpgradeInfo FRigUnit_Control::GetUpgradeInfo() const
+{
+	return FRigUnit::GetUpgradeInfo();
+}
diff --git a/Engine/Plugins/Animation/ControlRig/Source/ControlRig/Private/Units/Control/RigUnit_Control_StaticMesh.cpp b/Engine/Plugins/Animation/ControlRig/Source/ControlRig/Private/Units/Control/RigUnit_Control_StaticMesh.cpp
--- a/Engine/Plugins/Animation/ControlRig/Source/ControlRig/Private/Units/Control/RigUnit_Control_StaticMesh.cpp
+++ b/Engine/Plugins/Animation/ControlRig/Source/ControlRig/Private/Units/Control/RigUnit_Control_StaticMesh.cpp
@@ -9,4 +9,10 @@ FRigUnit_Control_StaticMesh::FRigUnit_Control_StaticMesh()
 FRigUnit_Control_StaticMesh_Execute()
 {
 	FRigUnit_Control::StaticExecute(RigVMExecuteContext, Transform, Base, InitTransform, Result, Filter, Context);
+}
+
+FRigVMStructUpgradeInfo FRigUnit_Control_StaticMesh::GetUpgradeInfo() const
+{
+	// this node is no longer supported
+	return FRigVMStructUpgradeInfo();
 }
\ No newline at end of file
diff --git a/Engine/Plugins/Animation/ControlRig/Source/ControlRig/Private/Units/Debug/RigUnit_DebugBezier.cpp b/Engine/Plugins/Animation/ControlRig/Source/ControlRig/Private/Units/Debug/RigUnit_DebugBezier.cpp
--- a/Engine/Plugins/Animation/ControlRig/Source/ControlRig/Private/Units/Debug/RigUnit_DebugBezier.cpp
+++ b/Engine/Plugins/Animation/ControlRig/Source/ControlRig/Private/Units/Debug/RigUnit_DebugBezier.cpp
@@ -20,6 +20,12 @@ FRigUnit_DebugBezier_Execute()
 		Context);
 }
 
+FRigVMStructUpgradeInfo FRigUnit_DebugBezier::GetUpgradeInfo() const
+{
+	// this node is no longer supported
+	return FRigVMStructUpgradeInfo();
+}
+
 FRigUnit_DebugBezierItemSpace_Execute()
 {
     DECLARE_SCOPE_HIERARCHICAL_COUNTER_RIGUNIT()
@@ -40,4 +46,10 @@ FRigUnit_DebugBezierItemSpace_Execute()
 	}
 
 	Context.DrawInterface->DrawBezier(Transform, Bezier, MinimumU, MaximumU, Color, Thickness, Detail);
+}
+
+FRigVMStructUpgradeInfo FRigUnit_DebugBezierItemSpace::GetUpgradeInfo() const
+{
+	// this node is no longer supported
+	return FRigVMStructUpgradeInfo();
 }
\ No newline at end of file
diff --git a/Engine/Plugins/Animation/ControlRig/Source/ControlRig/Private/Units/Debug/RigUnit_DebugBezier.h b/Engine/Plugins/Animation/ControlRig/Source/ControlRig/Private/Units/Debug/RigUnit_DebugBezier.h
--- a/Engine/Plugins/Animation/ControlRig/Source/ControlRig/Private/Units/Debug/RigUnit_DebugBezier.h
+++ b/Engine/Plugins/Animation/ControlRig/Source/ControlRig/Private/Units/Debug/RigUnit_DebugBezier.h
@@ -52,6 +52,9 @@ struct CONTROLRIG_API FRigUnit_DebugBezier : public FRigUnit_DebugBaseMutable
 
 	UPROPERTY(meta = (Input, Constant))
 	bool bEnabled;
+
+	RIGVM_METHOD()
+	virtual FRigVMStructUpgradeInfo GetUpgradeInfo() const override;
 };
 
 USTRUCT(meta=(DisplayName="Draw Bezier", Deprecated = "4.25"))
@@ -101,4 +104,7 @@ struct CONTROLRIG_API FRigUnit_DebugBezierItemSpace : public FRigUnit_DebugBaseM
 
 	UPROPERTY(meta = (Input, Constant))
 	bool bEnabled;
+
+	RIGVM_METHOD()
+	virtual FRigVMStructUpgradeInfo GetUpgradeInfo() const override;
 };
diff --git a/Engine/Plugins/Animation/ControlRig/Source/ControlRig/Private/Units/Debug/RigUnit_DebugLine.cpp b/Engine/Plugins/Animation/ControlRig/Source/ControlRig/Private/Units/Debug/RigUnit_DebugLine.cpp
--- a/Engine/Plugins/Animation/ControlRig/Source/ControlRig/Private/Units/Debug/RigUnit_DebugLine.cpp
+++ b/Engine/Plugins/Animation/ControlRig/Source/ControlRig/Private/Units/Debug/RigUnit_DebugLine.cpp
@@ -18,6 +18,22 @@ FRigUnit_DebugLine_Execute()
 		Context);
 }
 
+FRigVMStructUpgradeInfo FRigUnit_DebugLine::GetUpgradeInfo() const
+{
+	FRigUnit_DebugLineItemSpace NewNode;
+	NewNode.A = A;
+	NewNode.B = B;
+	NewNode.Color = Color;
+	NewNode.Thickness = Thickness;
+	NewNode.Space = FRigElementKey(Space, ERigElementType::Bone);
+	NewNode.WorldOffset = WorldOffset;
+	NewNode.bEnabled = bEnabled;
+
+	FRigVMStructUpgradeInfo Info(*this, NewNode);
+	Info.AddRemappedPin(TEXT("Space"), TEXT("Space.Name"));
+	return Info;
+}
+
 FRigUnit_DebugLineItemSpace_Execute()
 {
     DECLARE_SCOPE_HIERARCHICAL_COUNTER_RIGUNIT()
diff --git a/Engine/Plugins/Animation/ControlRig/Source/ControlRig/Private/Units/Debug/RigUnit_DebugLine.h b/Engine/Plugins/Animation/ControlRig/Source/ControlRig/Private/Units/Debug/RigUnit_DebugLine.h
--- a/Engine/Plugins/Animation/ControlRig/Source/ControlRig/Private/Units/Debug/RigUnit_DebugLine.h
+++ b/Engine/Plugins/Animation/ControlRig/Source/ControlRig/Private/Units/Debug/RigUnit_DebugLine.h
@@ -42,6 +42,9 @@ struct CONTROLRIG_API FRigUnit_DebugLine : public FRigUnit_DebugBaseMutable
 
 	UPROPERTY(meta = (Input, Constant))
 	bool bEnabled;
+
+	RIGVM_METHOD()
+	virtual FRigVMStructUpgradeInfo GetUpgradeInfo() const override;
 };
 
 /**
diff --git a/Engine/Plugins/Animation/ControlRig/Source/ControlRig/Private/Units/Debug/RigUnit_DebugLineStrip.cpp b/Engine/Plugins/Animation/ControlRig/Source/ControlRig/Private/Units/Debug/RigUnit_DebugLineStrip.cpp
--- a/Engine/Plugins/Animation/ControlRig/Source/ControlRig/Private/Units/Debug/RigUnit_DebugLineStrip.cpp
+++ b/Engine/Plugins/Animation/ControlRig/Source/ControlRig/Private/Units/Debug/RigUnit_DebugLineStrip.cpp
@@ -17,6 +17,21 @@ FRigUnit_DebugLineStrip_Execute()
 		Context);
 }
 
+FRigVMStructUpgradeInfo FRigUnit_DebugLineStrip::GetUpgradeInfo() const
+{
+	FRigUnit_DebugLineStripItemSpace NewNode;
+	NewNode.Points = Points;
+	NewNode.Color = Color;
+	NewNode.Thickness = Thickness;
+	NewNode.Space = FRigElementKey(Space, ERigElementType::Bone);
+	NewNode.WorldOffset = WorldOffset;
+	NewNode.bEnabled = bEnabled;
+
+	FRigVMStructUpgradeInfo Info(*this, NewNode);
+	Info.AddRemappedPin(TEXT("Space"), TEXT("Space.Name"));
+	return Info;
+}
+
 FRigUnit_DebugLineStripItemSpace_Execute()
 {
     DECLARE_SCOPE_HIERARCHICAL_COUNTER_RIGUNIT()
diff --git a/Engine/Plugins/Animation/ControlRig/Source/ControlRig/Private/Units/Debug/RigUnit_DebugLineStrip.h b/Engine/Plugins/Animation/ControlRig/Source/ControlRig/Private/Units/Debug/RigUnit_DebugLineStrip.h
--- a/Engine/Plugins/Animation/ControlRig/Source/ControlRig/Private/Units/Debug/RigUnit_DebugLineStrip.h
+++ b/Engine/Plugins/Animation/ControlRig/Source/ControlRig/Private/Units/Debug/RigUnit_DebugLineStrip.h
@@ -38,6 +38,9 @@ struct CONTROLRIG_API FRigUnit_DebugLineStrip : public FRigUnit_DebugBaseMutable
 	
 	UPROPERTY(meta = (Input, Constant))
 	bool bEnabled;
+
+	RIGVM_METHOD()
+	virtual FRigVMStructUpgradeInfo GetUpgradeInfo() const override;
 };
 
 /**
diff --git a/Engine/Plugins/Animation/ControlRig/Source/ControlRig/Private/Units/Debug/RigUnit_DebugPoint.cpp b/Engine/Plugins/Animation/ControlRig/Source/ControlRig/Private/Units/Debug/RigUnit_DebugPoint.cpp
--- a/Engine/Plugins/Animation/ControlRig/Source/ControlRig/Private/Units/Debug/RigUnit_DebugPoint.cpp
+++ b/Engine/Plugins/Animation/ControlRig/Source/ControlRig/Private/Units/Debug/RigUnit_DebugPoint.cpp
@@ -1,6 +1,7 @@
 // Copyright Epic Games, Inc. All Rights Reserved.
 
 #include "Units/Debug/RigUnit_DebugPoint.h"
+#include "Units/Debug/RigUnit_VisualDebug.h"
 #include "Units/RigUnitContext.h"
 
 FRigUnit_DebugPoint_Execute()
@@ -40,6 +41,22 @@ FRigUnit_DebugPoint_Execute()
 	}
 }
 
+FRigVMStructUpgradeInfo FRigUnit_DebugPoint::GetUpgradeInfo() const
+{
+	FRigUnit_VisualDebugVector NewNode;
+	NewNode.Value = Vector;
+	NewNode.Color = Color;
+	NewNode.Thickness = Thickness;
+	NewNode.BoneSpace = Space;
+	NewNode.Scale = Scale;
+	NewNode.bEnabled = bEnabled;
+
+	FRigVMStructUpgradeInfo Info(*this, NewNode);
+	Info.AddRemappedPin(TEXT("Vector"), TEXT("Value"));
+	Info.AddRemappedPin(TEXT("Space"), TEXT("BoneSpace"));
+	return Info;
+}
+
 FRigUnit_DebugPointMutable_Execute()
 {
     DECLARE_SCOPE_HIERARCHICAL_COUNTER_RIGUNIT()
@@ -75,4 +92,10 @@ FRigUnit_DebugPointMutable_Execute()
 			break;
 		}
 	}
-}
\ No newline at end of file
+}
+
+FRigVMStructUpgradeInfo FRigUnit_DebugPointMutable::GetUpgradeInfo() const
+{
+	// this node is no longer supported
+	return FRigVMStructUpgradeInfo();
+}
diff --git a/Engine/Plugins/Animation/ControlRig/Source/ControlRig/Private/Units/Debug/RigUnit_DebugPoint.h b/Engine/Plugins/Animation/ControlRig/Source/ControlRig/Private/Units/Debug/RigUnit_DebugPoint.h
--- a/Engine/Plugins/Animation/ControlRig/Source/ControlRig/Private/Units/Debug/RigUnit_DebugPoint.h
+++ b/Engine/Plugins/Animation/ControlRig/Source/ControlRig/Private/Units/Debug/RigUnit_DebugPoint.h
@@ -60,6 +60,9 @@ struct CONTROLRIG_API FRigUnit_DebugPoint : public FRigUnit_DebugBase
 
 	UPROPERTY(meta = (Input, Constant))
 	bool bEnabled;
+
+	RIGVM_METHOD()
+	virtual FRigVMStructUpgradeInfo GetUpgradeInfo() const override;
 };
 
 USTRUCT(meta=(DisplayName="Draw Point", Keywords="Draw Vector", Deprecated = "4.25.0"))
@@ -106,4 +109,7 @@ struct CONTROLRIG_API FRigUnit_DebugPointMutable : public FRigUnit_DebugBaseMuta
 
 	UPROPERTY(meta = (Input, Constant))
 	bool bEnabled;
+
+	RIGVM_METHOD()
+	virtual FRigVMStructUpgradeInfo GetUpgradeInfo() const override;
 };
\ No newline at end of file
diff --git a/Engine/Plugins/Animation/ControlRig/Source/ControlRig/Private/Units/Debug/RigUnit_DebugPrimitives.cpp b/Engine/Plugins/Animation/ControlRig/Source/ControlRig/Private/Units/Debug/RigUnit_DebugPrimitives.cpp
--- a/Engine/Plugins/Animation/ControlRig/Source/ControlRig/Private/Units/Debug/RigUnit_DebugPrimitives.cpp
+++ b/Engine/Plugins/Animation/ControlRig/Source/ControlRig/Private/Units/Debug/RigUnit_DebugPrimitives.cpp
@@ -18,6 +18,22 @@ FRigUnit_DebugRectangle_Execute()
 		Context);
 }
 
+FRigVMStructUpgradeInfo FRigUnit_DebugRectangle::GetUpgradeInfo() const
+{
+	FRigUnit_DebugRectangleItemSpace NewNode;
+	NewNode.Transform = Transform;
+	NewNode.Color = Color;
+	NewNode.Scale = Scale;
+	NewNode.Thickness = Thickness;
+	NewNode.Space = FRigElementKey(Space, ERigElementType::Bone);
+	NewNode.WorldOffset = WorldOffset;
+	NewNode.bEnabled = bEnabled;
+
+	FRigVMStructUpgradeInfo Info(*this, NewNode);
+	Info.AddRemappedPin(TEXT("Space"), TEXT("Space.Name"));
+	return Info;
+}
+
 FRigUnit_DebugRectangleItemSpace_Execute()
 {
     DECLARE_SCOPE_HIERARCHICAL_COUNTER_RIGUNIT()
@@ -58,6 +74,25 @@ FRigUnit_DebugArc_Execute()
 		Context);
 }
 
+FRigVMStructUpgradeInfo FRigUnit_DebugArc::GetUpgradeInfo() const
+{
+	FRigUnit_DebugArcItemSpace NewNode;
+	NewNode.Transform = Transform;
+	NewNode.Color = Color;
+	NewNode.Radius = Radius;
+	NewNode.MinimumDegrees = MinimumDegrees;
+	NewNode.MaximumDegrees = MaximumDegrees;
+	NewNode.Thickness = Thickness;
+	NewNode.Detail = Detail;
+	NewNode.Space = FRigElementKey(Space, ERigElementType::Bone);
+	NewNode.WorldOffset = WorldOffset;
+	NewNode.bEnabled = bEnabled;
+
+	FRigVMStructUpgradeInfo Info(*this, NewNode);
+	Info.AddRemappedPin(TEXT("Space"), TEXT("Space.Name"));
+	return Info;
+}
+
 FRigUnit_DebugArcItemSpace_Execute()
 {
     DECLARE_SCOPE_HIERARCHICAL_COUNTER_RIGUNIT()
diff --git a/Engine/Plugins/Animation/ControlRig/Source/ControlRig/Private/Units/Debug/RigUnit_DebugPrimitives.h b/Engine/Plugins/Animation/ControlRig/Source/ControlRig/Private/Units/Debug/RigUnit_DebugPrimitives.h
--- a/Engine/Plugins/Animation/ControlRig/Source/ControlRig/Private/Units/Debug/RigUnit_DebugPrimitives.h
+++ b/Engine/Plugins/Animation/ControlRig/Source/ControlRig/Private/Units/Debug/RigUnit_DebugPrimitives.h
@@ -43,6 +43,9 @@ struct CONTROLRIG_API FRigUnit_DebugRectangle : public FRigUnit_DebugBaseMutable
 
 	UPROPERTY(meta = (Input, Constant))
 	bool bEnabled;
+
+	RIGVM_METHOD()
+	virtual FRigVMStructUpgradeInfo GetUpgradeInfo() const override;
 };
 
 /**
@@ -139,6 +142,9 @@ struct CONTROLRIG_API FRigUnit_DebugArc : public FRigUnit_DebugBaseMutable
 
 	UPROPERTY(meta = (Input, Constant))
 	bool bEnabled;
+
+	RIGVM_METHOD()
+	virtual FRigVMStructUpgradeInfo GetUpgradeInfo() const override;
 };
 
 /**
diff --git a/Engine/Plugins/Animation/ControlRig/Source/ControlRig/Private/Units/Debug/RigUnit_DebugTransform.cpp b/Engine/Plugins/Animation/ControlRig/Source/ControlRig/Private/Units/Debug/RigUnit_DebugTransform.cpp
--- a/Engine/Plugins/Animation/ControlRig/Source/ControlRig/Private/Units/Debug/RigUnit_DebugTransform.cpp
+++ b/Engine/Plugins/Animation/ControlRig/Source/ControlRig/Private/Units/Debug/RigUnit_DebugTransform.cpp
@@ -1,6 +1,7 @@
 // Copyright Epic Games, Inc. All Rights Reserved.
 
 #include "Units/Debug/RigUnit_DebugTransform.h"
+#include "Units/Debug/RigUnit_VisualDebug.h"
 #include "Units/RigUnitContext.h"
 
 FRigUnit_DebugTransform_Execute()
@@ -43,6 +44,22 @@ FRigUnit_DebugTransform_Execute()
 	}
 }
 
+FRigVMStructUpgradeInfo FRigUnit_DebugTransform::GetUpgradeInfo() const
+{
+	FRigUnit_VisualDebugTransformItemSpace NewNode;
+	NewNode.Value = Transform;
+	NewNode.Space = FRigElementKey(Space, ERigElementType::Bone);
+	NewNode.Scale = Scale;
+	NewNode.Thickness = Thickness;
+	NewNode.Space = FRigElementKey(Space, ERigElementType::Bone);
+	NewNode.bEnabled = bEnabled;
+
+	FRigVMStructUpgradeInfo Info(*this, NewNode);
+	Info.AddRemappedPin(TEXT("Transform"), TEXT("Value"));
+	Info.AddRemappedPin(TEXT("Space"), TEXT("Space.Name"));
+	return Info;
+}
+
 FRigUnit_DebugTransformMutable_Execute()
 {
 	FRigUnit_DebugTransformMutableItemSpace::StaticExecute(
@@ -59,6 +76,23 @@ FRigUnit_DebugTransformMutable_Execute()
 		Context);
 }
 
+FRigVMStructUpgradeInfo FRigUnit_DebugTransformMutable::GetUpgradeInfo() const
+{
+	FRigUnit_DebugTransformMutableItemSpace NewNode;
+	NewNode.Transform = Transform;
+	NewNode.Mode = Mode;
+	NewNode.Color = Color;
+	NewNode.Thickness = Thickness;
+	NewNode.Scale = Scale;
+	NewNode.Space = FRigElementKey(Space, ERigElementType::Bone);
+	NewNode.WorldOffset = WorldOffset;
+	NewNode.bEnabled = bEnabled;
+
+	FRigVMStructUpgradeInfo Info(*this, NewNode);
+	Info.AddRemappedPin(TEXT("Space"), TEXT("Space.Name"));
+	return Info;
+}
+
 FRigUnit_DebugTransformMutableItemSpace_Execute()
 {
     DECLARE_SCOPE_HIERARCHICAL_COUNTER_RIGUNIT()
@@ -153,6 +187,23 @@ FRigUnit_DebugTransformArrayMutable_Execute()
 	}
 }
 
+FRigVMStructUpgradeInfo FRigUnit_DebugTransformArrayMutable::GetUpgradeInfo() const
+{
+	FRigUnit_DebugTransformArrayMutableItemSpace NewNode;
+	NewNode.Transforms = Transforms;
+	NewNode.Mode = Mode;
+	NewNode.Color = Color;
+	NewNode.Thickness = Thickness;
+	NewNode.Scale = Scale;
+	NewNode.Space = FRigElementKey(Space, ERigElementType::Bone);
+	NewNode.WorldOffset = WorldOffset;
+	NewNode.bEnabled = bEnabled;
+
+	FRigVMStructUpgradeInfo Info(*this, NewNode);
+	Info.AddRemappedPin(TEXT("Space"), TEXT("Space.Name"));
+	return Info;
+}
+
 FRigUnit_DebugTransformArrayMutableItemSpace_Execute()
 {
 	DECLARE_SCOPE_HIERARCHICAL_COUNTER_RIGUNIT()
diff --git a/Engine/Plugins/Animation/ControlRig/Source/ControlRig/Private/Units/Debug/RigUnit_DebugTransform.h b/Engine/Plugins/Animation/ControlRig/Source/ControlRig/Private/Units/Debug/RigUnit_DebugTransform.h
--- a/Engine/Plugins/Animation/ControlRig/Source/ControlRig/Private/Units/Debug/RigUnit_DebugTransform.h
+++ b/Engine/Plugins/Animation/ControlRig/Source/ControlRig/Private/Units/Debug/RigUnit_DebugTransform.h
@@ -62,6 +62,9 @@ struct CONTROLRIG_API FRigUnit_DebugTransform : public FRigUnit_DebugBase
 
 	UPROPERTY(meta = (Input, Constant))
 	bool bEnabled;
+
+	RIGVM_METHOD()
+	virtual FRigVMStructUpgradeInfo GetUpgradeInfo() const override;
 };
 
 USTRUCT(meta=(DisplayName="Draw Transform", Deprecated = "4.25"))
@@ -105,6 +108,9 @@ struct CONTROLRIG_API FRigUnit_DebugTransformMutable : public FRigUnit_DebugBase
 
 	UPROPERTY(meta = (Input, Constant))
 	bool bEnabled;
+
+	RIGVM_METHOD()
+	virtual FRigVMStructUpgradeInfo GetUpgradeInfo() const override;
 };
 
 /**
@@ -207,6 +213,9 @@ struct CONTROLRIG_API FRigUnit_DebugTransformArrayMutable : public FRigUnit_Debu
 
 	UPROPERTY(transient)
 	FRigUnit_DebugTransformArrayMutable_WorkData WorkData;
+
+	RIGVM_METHOD()
+	virtual FRigVMStructUpgradeInfo GetUpgradeInfo() const override;
 };
 
 /**
diff --git a/Engine/Plugins/Animation/ControlRig/Source/ControlRig/Private/Units/Debug/RigUnit_VisualDebug.cpp b/Engine/Plugins/Animation/ControlRig/Source/ControlRig/Private/Units/Debug/RigUnit_VisualDebug.cpp
--- a/Engine/Plugins/Animation/ControlRig/Source/ControlRig/Private/Units/Debug/RigUnit_VisualDebug.cpp
+++ b/Engine/Plugins/Animation/ControlRig/Source/ControlRig/Private/Units/Debug/RigUnit_VisualDebug.cpp
@@ -8,6 +8,21 @@ FRigUnit_VisualDebugVector_Execute()
 	FRigUnit_VisualDebugVectorItemSpace::StaticExecute(RigVMExecuteContext, Value, bEnabled, Mode, Color, Thickness, Scale, FRigElementKey(BoneSpace, ERigElementType::Bone), Context);
 }
 
+FRigVMStructUpgradeInfo FRigUnit_VisualDebugVector::GetUpgradeInfo() const
+{
+	FRigUnit_VisualDebugVectorItemSpace NewNode;
+	NewNode.Value = Value;
+	NewNode.Color = Color;
+	NewNode.Thickness = Thickness;
+	NewNode.Space = FRigElementKey(BoneSpace, ERigElementType::Bone);
+	NewNode.Scale = Scale;
+	NewNode.bEnabled = bEnabled;
+
+	FRigVMStructUpgradeInfo Info(*this, NewNode);
+	Info.AddRemappedPin(TEXT("BoneSpace"), TEXT("Space.Name"));
+	return Info;
+}
+
 FRigUnit_VisualDebugVectorItemSpace_Execute()
 {
     DECLARE_SCOPE_HIERARCHICAL_COUNTER_RIGUNIT()
@@ -51,6 +66,20 @@ FRigUnit_VisualDebugQuat_Execute()
 	FRigUnit_VisualDebugQuatItemSpace::StaticExecute(RigVMExecuteContext, Value, bEnabled, Thickness, Scale, FRigElementKey(BoneSpace, ERigElementType::Bone), Context);
 }
 
+FRigVMStructUpgradeInfo FRigUnit_VisualDebugQuat::GetUpgradeInfo() const
+{
+	FRigUnit_VisualDebugQuatItemSpace NewNode;
+	NewNode.Value = Value;
+	NewNode.Thickness = Thickness;
+	NewNode.Space = FRigElementKey(BoneSpace, ERigElementType::Bone);
+	NewNode.Scale = Scale;
+	NewNode.bEnabled = bEnabled;
+
+	FRigVMStructUpgradeInfo Info(*this, NewNode);
+	Info.AddRemappedPin(TEXT("BoneSpace"), TEXT("Space.Name"));
+	return Info;
+}
+
 FRigUnit_VisualDebugQuatItemSpace_Execute()
 {
     DECLARE_SCOPE_HIERARCHICAL_COUNTER_RIGUNIT()
@@ -66,6 +95,20 @@ FRigUnit_VisualDebugTransform_Execute()
 	FRigUnit_VisualDebugTransformItemSpace::StaticExecute(RigVMExecuteContext, Value, bEnabled, Thickness, Scale, FRigElementKey(BoneSpace, ERigElementType::Bone), Context);
 }
 
+FRigVMStructUpgradeInfo FRigUnit_VisualDebugTransform::GetUpgradeInfo() const
+{
+	FRigUnit_VisualDebugTransformItemSpace NewNode;
+	NewNode.Value = Value;
+	NewNode.Thickness = Thickness;
+	NewNode.Space = FRigElementKey(BoneSpace, ERigElementType::Bone);
+	NewNode.Scale = Scale;
+	NewNode.bEnabled = bEnabled;
+
+	FRigVMStructUpgradeInfo Info(*this, NewNode);
+	Info.AddRemappedPin(TEXT("BoneSpace"), TEXT("Space.Name"));
+	return Info;
+}
+
 FRigUnit_VisualDebugTransformItemSpace_Execute()
 {
     DECLARE_SCOPE_HIERARCHICAL_COUNTER_RIGUNIT()
diff --git a/Engine/Plugins/Animation/ControlRig/Source/ControlRig/Private/Units/Debug/RigUnit_VisualDebug.h b/Engine/Plugins/Animation/ControlRig/Source/ControlRig/Private/Units/Debug/RigUnit_VisualDebug.h
--- a/Engine/Plugins/Animation/ControlRig/Source/ControlRig/Private/Units/Debug/RigUnit_VisualDebug.h
+++ b/Engine/Plugins/Animation/ControlRig/Source/ControlRig/Private/Units/Debug/RigUnit_VisualDebug.h
@@ -57,6 +57,9 @@ struct CONTROLRIG_API FRigUnit_VisualDebugVector : public FRigUnit_DebugBase
 
 	UPROPERTY(meta = (Input, EditCondition = "bEnabled"))
 	FName BoneSpace;
+
+	RIGVM_METHOD()
+	virtual FRigVMStructUpgradeInfo GetUpgradeInfo() const override;
 };
 
 /**
@@ -139,6 +142,9 @@ struct CONTROLRIG_API FRigUnit_VisualDebugQuat : public FRigUnit_DebugBase
 
 	UPROPERTY(meta = (Input, EditCondition = "bEnabled"))
 	FName BoneSpace;
+
+	RIGVM_METHOD()
+	virtual FRigVMStructUpgradeInfo GetUpgradeInfo() const override;
 };
 
 /**
@@ -213,6 +219,9 @@ struct CONTROLRIG_API FRigUnit_VisualDebugTransform : public FRigUnit_DebugBase
 
 	UPROPERTY(meta = (Input, EditCondition = "bEnabled"))
 	FName BoneSpace;
+
+	RIGVM_METHOD()
+	virtual FRigVMStructUpgradeInfo GetUpgradeInfo() const override;
 };
 
 /**
diff --git a/Engine/Plugins/Animation/ControlRig/Source/ControlRig/Private/Units/Deprecated/Math/RigUnit_Converter.cpp b/Engine/Plugins/Animation/ControlRig/Source/ControlRig/Private/Units/Deprecated/Math/RigUnit_Converter.cpp
--- a/Engine/Plugins/Animation/ControlRig/Source/ControlRig/Private/Units/Deprecated/Math/RigUnit_Converter.cpp
+++ b/Engine/Plugins/Animation/ControlRig/Source/ControlRig/Private/Units/Deprecated/Math/RigUnit_Converter.cpp
@@ -1,49 +1,115 @@
 // Copyright Epic Games, Inc. All Rights Reserved.
 
 #include "RigUnit_Converter.h"
+#include "Units/Math/RigUnit_MathTransform.h"
+#include "Units/Math/RigUnit_MathQuaternion.h"
 
 FRigUnit_ConvertTransform_Execute()
 {
 	Result.FromFTransform(Input);
 }
 
+FRigVMStructUpgradeInfo FRigUnit_ConvertTransform::GetUpgradeInfo() const
+{
+	FRigUnit_MathTransformToEulerTransform NewNode;
+	NewNode.Value = Input;
+
+	FRigVMStructUpgradeInfo Info(*this, NewNode);
+	Info.AddRemappedPin(TEXT("Space"), TEXT("Space.Name"));
+	Info.AddRemappedPin(TEXT("Input"), TEXT("Value"));
+	return Info;
+}
+
 FRigUnit_ConvertEulerTransform_Execute()
 {
 	Result = Input.ToFTransform();
 }
 
+FRigVMStructUpgradeInfo FRigUnit_ConvertEulerTransform::GetUpgradeInfo() const
+{
+	FRigUnit_MathTransformFromEulerTransform NewNode;
+	NewNode.EulerTransform = Input;
+
+	FRigVMStructUpgradeInfo Info(*this, NewNode);
+	Info.AddRemappedPin(TEXT("Input"), TEXT("EulerTransform"));
+	return Info;
+}
+
 FRigUnit_ConvertRotation_Execute()
 {
 	Result = Input.Quaternion();
 }
 
+FRigVMStructUpgradeInfo FRigUnit_ConvertRotation::GetUpgradeInfo() const
+{
+	FRigUnit_MathQuaternionFromRotator NewNode;
+	NewNode.Rotator = Input;
+
+	FRigVMStructUpgradeInfo Info(*this, NewNode);
+	Info.AddRemappedPin(TEXT("Input"), TEXT("Rotator"));
+	return Info;
+}
 
 FRigUnit_ConvertQuaternion_Execute()
 {
 	Result = Input.Rotator();
 }
 
+FRigVMStructUpgradeInfo FRigUnit_ConvertQuaternion::GetUpgradeInfo() const
+{
+	FRigUnit_MathQuaternionToRotator NewNode;
+	NewNode.Value = Input;
+
+	FRigVMStructUpgradeInfo Info(*this, NewNode);
+	Info.AddRemappedPin(TEXT("Input"), TEXT("Value"));
+	return Info;
+}
+
 FRigUnit_ConvertVectorToRotation_Execute()
 {
 	Result = Input.Rotation();
 }
 
+FRigVMStructUpgradeInfo FRigUnit_ConvertVectorToRotation::GetUpgradeInfo() const
+{
+	// this node is no longer supported
+	return FRigVMStructUpgradeInfo();
+}
+
 FRigUnit_ConvertVectorToQuaternion_Execute()
 {
 	Result = Input.Rotation().Quaternion();
 	Result.Normalize();
 }
 
+FRigVMStructUpgradeInfo FRigUnit_ConvertVectorToQuaternion::GetUpgradeInfo() const
+{
+	// this node is no longer supported
+	return FRigVMStructUpgradeInfo();
+}
+
 FRigUnit_ConvertRotationToVector_Execute()
 {
 	Result = Input.RotateVector(FVector(1.f, 0.f, 0.f));
 }
 
+FRigVMStructUpgradeInfo FRigUnit_ConvertRotationToVector::GetUpgradeInfo() const
+{
+	// this node is no longer supported
+	return FRigVMStructUpgradeInfo();
+}
+
 FRigUnit_ConvertQuaternionToVector_Execute()
 {
 	Result = Input.RotateVector(FVector(1.f, 0.f, 0.f));
 }
 
+FRigVMStructUpgradeInfo FRigUnit_ConvertQuaternionToVector::GetUpgradeInfo() const
+{
+	// this node is no longer supported
+	return FRigVMStructUpgradeInfo();
+}
+
 FRigUnit_ToSwingAndTwist_Execute()
 {
 	if (!TwistAxis.IsZero())
@@ -52,3 +118,12 @@ FRigUnit_ToSwingAndTwist_Execute()
 		Input.ToSwingTwist(TwistAxis, Swing, Twist);
 	}
 }
+
+FRigVMStructUpgradeInfo FRigUnit_ToSwingAndTwist::GetUpgradeInfo() const
+{
+	FRigUnit_MathQuaternionSwingTwist NewNode;
+	NewNode.Input = Input;
+	NewNode.TwistAxis = TwistAxis;
+
+	return FRigVMStructUpgradeInfo(*this, NewNode);
+}
diff --git a/Engine/Plugins/Animation/ControlRig/Source/ControlRig/Private/Units/Deprecated/Math/RigUnit_Converter.h b/Engine/Plugins/Animation/ControlRig/Source/ControlRig/Private/Units/Deprecated/Math/RigUnit_Converter.h
--- a/Engine/Plugins/Animation/ControlRig/Source/ControlRig/Private/Units/Deprecated/Math/RigUnit_Converter.h
+++ b/Engine/Plugins/Animation/ControlRig/Source/ControlRig/Private/Units/Deprecated/Math/RigUnit_Converter.h
@@ -19,6 +19,9 @@ struct CONTROLRIG_API FRigUnit_ConvertTransform : public FRigUnit
 
 	RIGVM_METHOD()
 	virtual void Execute(const FRigUnitContext& Context) override;
+
+	RIGVM_METHOD()
+	virtual FRigVMStructUpgradeInfo GetUpgradeInfo() const override;
 };
 
 USTRUCT(meta = (DisplayName = "ConvertToTransform", Category = "Math|Convert", NodeColor = "0.1 0.1 0.7", Deprecated="4.23.0"))
@@ -34,6 +37,9 @@ struct CONTROLRIG_API FRigUnit_ConvertEulerTransform : public FRigUnit
 
 	RIGVM_METHOD()
 	virtual void Execute(const FRigUnitContext& Context) override;
+
+	RIGVM_METHOD()
+	virtual FRigVMStructUpgradeInfo GetUpgradeInfo() const override;
 };
 
 USTRUCT(meta = (DisplayName = "ConvertToQuaternion", Category = "Math|Convert", NodeColor = "0.1 0.1 0.7", Deprecated="4.23.0"))
@@ -49,6 +55,9 @@ struct CONTROLRIG_API FRigUnit_ConvertRotation : public FRigUnit
 
 	RIGVM_METHOD()
 	virtual void Execute(const FRigUnitContext& Context) override;
+
+	RIGVM_METHOD()
+	virtual FRigVMStructUpgradeInfo GetUpgradeInfo() const override;
 };
 
 USTRUCT(meta = (DisplayName = "ConvertToQuaternionDeprecated", Category = "Math|Convert", NodeColor = "0.1 0.1 0.7", Deprecated="4.23.0"))
@@ -70,6 +79,9 @@ struct CONTROLRIG_API FRigUnit_ConvertQuaternion: public FRigUnit
 
 	RIGVM_METHOD()
 	virtual void Execute(const FRigUnitContext& Context) override;
+
+	RIGVM_METHOD()
+	virtual FRigVMStructUpgradeInfo GetUpgradeInfo() const override;
 };
 
 USTRUCT(meta = (DisplayName = "ConvertVectorToRotation", Category = "Math|Convert", NodeColor = "0.1 0.1 0.7", Deprecated="4.23.0"))
@@ -85,6 +97,9 @@ struct CONTROLRIG_API FRigUnit_ConvertVectorToRotation: public FRigUnit
 
 	RIGVM_METHOD()
 	virtual void Execute(const FRigUnitContext& Context) override;
+
+	RIGVM_METHOD()
+	virtual FRigVMStructUpgradeInfo GetUpgradeInfo() const override;
 };
 
 USTRUCT(meta = (DisplayName = "ConvertVectorToQuaternion", Category = "Math|Convert", NodeColor = "0.1 0.1 0.7", Deprecated="4.23.0"))
@@ -100,8 +115,10 @@ struct CONTROLRIG_API FRigUnit_ConvertVectorToQuaternion: public FRigUnit
 
 	RIGVM_METHOD()
 	virtual void Execute(const FRigUnitContext& Context) override;
-};
 
+	RIGVM_METHOD()
+	virtual FRigVMStructUpgradeInfo GetUpgradeInfo() const override;
+};
 
 USTRUCT(meta = (DisplayName = "ConvertRotationToVector", Category = "Math|Convert", NodeColor = "0.1 0.1 0.7", Deprecated="4.23.0"))
 struct CONTROLRIG_API FRigUnit_ConvertRotationToVector: public FRigUnit
@@ -116,6 +133,9 @@ struct CONTROLRIG_API FRigUnit_ConvertRotationToVector: public FRigUnit
 
 	RIGVM_METHOD()
 	virtual void Execute(const FRigUnitContext& Context) override;
+
+	RIGVM_METHOD()
+	virtual FRigVMStructUpgradeInfo GetUpgradeInfo() const override;
 };
 
 USTRUCT(meta = (DisplayName = "ConvertQuaternionToVector", Category = "Math|Convert", NodeColor = "0.1 0.1 0.7", Deprecated="4.23.0"))
@@ -131,6 +151,9 @@ struct CONTROLRIG_API FRigUnit_ConvertQuaternionToVector: public FRigUnit
 
 	RIGVM_METHOD()
 	virtual void Execute(const FRigUnitContext& Context) override;
+
+	RIGVM_METHOD()
+	virtual FRigVMStructUpgradeInfo GetUpgradeInfo() const override;
 };
 
 USTRUCT(meta = (DisplayName = "ToSwingAndTwist", Category = "Math|Transform", NodeColor = "0.1 0.1 0.7", Deprecated="4.23.0"))
@@ -160,4 +183,7 @@ struct CONTROLRIG_API FRigUnit_ToSwingAndTwist : public FRigUnit
 
 	RIGVM_METHOD()
 	virtual void Execute(const FRigUnitContext& Context) override;
+
+	RIGVM_METHOD()
+	virtual FRigVMStructUpgradeInfo GetUpgradeInfo() const override;
 };
diff --git a/Engine/Plugins/Animation/ControlRig/Source/ControlRig/Private/Units/Deprecated/Math/RigUnit_Float.cpp b/Engine/Plugins/Animation/ControlRig/Source/ControlRig/Private/Units/Deprecated/Math/RigUnit_Float.cpp
--- a/Engine/Plugins/Animation/ControlRig/Source/ControlRig/Private/Units/Deprecated/Math/RigUnit_Float.cpp
+++ b/Engine/Plugins/Animation/ControlRig/Source/ControlRig/Private/Units/Deprecated/Math/RigUnit_Float.cpp
@@ -1,33 +1,112 @@
 // Copyright Epic Games, Inc. All Rights Reserved.
 
 #include "RigUnit_Float.h"
+#include "Units/Math/RigUnit_MathFloat.h"
 
 FRigUnit_Multiply_FloatFloat_Execute()
 {
 	Result = FRigMathLibrary::Multiply(Argument0, Argument1);
 }
 
+FRigVMStructUpgradeInfo FRigUnit_Multiply_FloatFloat::GetUpgradeInfo() const
+{
+	FRigUnit_MathFloatMul NewNode;
+	NewNode.A = Argument0;
+	NewNode.B = Argument1;
+
+	FRigVMStructUpgradeInfo Info(*this, NewNode);
+	Info.AddRemappedPin(TEXT("Argument0"), TEXT("A"));
+	Info.AddRemappedPin(TEXT("Argument1"), TEXT("B"));
+	return Info;
+}
+
 FRigUnit_Add_FloatFloat_Execute()
 {
 	Result = FRigMathLibrary::Add(Argument0, Argument1);
 }
 
+FRigVMStructUpgradeInfo FRigUnit_Add_FloatFloat::GetUpgradeInfo() const
+{
+	FRigUnit_MathFloatAdd NewNode;
+	NewNode.A = Argument0;
+	NewNode.B = Argument1;
+
+	FRigVMStructUpgradeInfo Info(*this, NewNode);
+	Info.AddRemappedPin(TEXT("Argument0"), TEXT("A"));
+	Info.AddRemappedPin(TEXT("Argument1"), TEXT("B"));
+	return Info;
+}
+
 FRigUnit_Subtract_FloatFloat_Execute()
 {
 	Result = FRigMathLibrary::Subtract(Argument0, Argument1);
 }
 
+FRigVMStructUpgradeInfo FRigUnit_Subtract_FloatFloat::GetUpgradeInfo() const
+{
+	FRigUnit_MathFloatSub NewNode;
+	NewNode.A = Argument0;
+	NewNode.B = Argument1;
+
+	FRigVMStructUpgradeInfo Info(*this, NewNode);
+	Info.AddRemappedPin(TEXT("Argument0"), TEXT("A"));
+	Info.AddRemappedPin(TEXT("Argument1"), TEXT("B"));
+	return Info;
+}
+
 FRigUnit_Divide_FloatFloat_Execute()
 {
 	Result = FRigMathLibrary::Divide(Argument0, Argument1);
 }
 
+FRigVMStructUpgradeInfo FRigUnit_Divide_FloatFloat::GetUpgradeInfo() const
+{
+	FRigUnit_MathFloatDiv NewNode;
+	NewNode.A = Argument0;
+	NewNode.B = Argument1;
+
+	FRigVMStructUpgradeInfo Info(*this, NewNode);
+	Info.AddRemappedPin(TEXT("Argument0"), TEXT("A"));
+	Info.AddRemappedPin(TEXT("Argument1"), TEXT("B"));
+	return Info;
+}
+
 FRigUnit_Clamp_Float_Execute()
 {
 	Result = FMath::Clamp(Value, Min, Max);
 }
 
+FRigVMStructUpgradeInfo FRigUnit_Clamp_Float::GetUpgradeInfo() const
+{
+	FRigUnit_MathFloatClamp NewNode;
+	NewNode.Value = Value;
+	NewNode.Minimum = Min;
+	NewNode.Maximum = Max;
+
+	FRigVMStructUpgradeInfo Info(*this, NewNode);
+	Info.AddRemappedPin(TEXT("Min"), TEXT("Minimum"));
+	Info.AddRemappedPin(TEXT("Max"), TEXT("Maximum"));
+	return Info;
+}
+
 FRigUnit_MapRange_Float_Execute()
 {
 	Result = FMath::GetMappedRangeValueClamped(FVector2f(MinIn, MaxIn), FVector2f(MinOut, MaxOut), Value);
 }
+
+FRigVMStructUpgradeInfo FRigUnit_MapRange_Float::GetUpgradeInfo() const
+{
+	FRigUnit_MathFloatRemap NewNode;
+	NewNode.Value = Value;
+	NewNode.SourceMinimum = MinIn;
+	NewNode.SourceMaximum = MaxIn;
+	NewNode.TargetMinimum = MinOut;
+	NewNode.TargetMaximum = MaxOut;
+
+	FRigVMStructUpgradeInfo Info(*this, NewNode);
+	Info.AddRemappedPin(TEXT("MinIn"), TEXT("SourceMinimum"));
+	Info.AddRemappedPin(TEXT("MaxIn"), TEXT("SourceMaximum"));
+	Info.AddRemappedPin(TEXT("MinOut"), TEXT("TargetMinimum"));
+	Info.AddRemappedPin(TEXT("MaxOut"), TEXT("TargetMaximum"));
+	return Info;
+}
\ No newline at end of file
diff --git a/Engine/Plugins/Animation/ControlRig/Source/ControlRig/Private/Units/Deprecated/Math/RigUnit_Float.h b/Engine/Plugins/Animation/ControlRig/Source/ControlRig/Private/Units/Deprecated/Math/RigUnit_Float.h
--- a/Engine/Plugins/Animation/ControlRig/Source/ControlRig/Private/Units/Deprecated/Math/RigUnit_Float.h
+++ b/Engine/Plugins/Animation/ControlRig/Source/ControlRig/Private/Units/Deprecated/Math/RigUnit_Float.h
@@ -7,7 +7,7 @@
 #include "RigUnit_Float.generated.h"
 
 /** Two args and a result of float type */
-USTRUCT(meta=(Abstract, NodeColor = "0.1 0.7 0.1", Deprecated="4.23.0"))
+USTRUCT(meta=(Abstract, NodeColor = "0.1 0.7 0.1"))
 struct CONTROLRIG_API FRigUnit_BinaryFloatOp : public FRigUnit
 {
 	GENERATED_BODY()
@@ -29,6 +29,9 @@ struct CONTROLRIG_API FRigUnit_Multiply_FloatFloat : public FRigUnit_BinaryFloat
 
 	RIGVM_METHOD()
 	virtual void Execute(const FRigUnitContext& Context) override;
+
+	RIGVM_METHOD()
+	virtual FRigVMStructUpgradeInfo GetUpgradeInfo() const override;
 };
 
 USTRUCT(meta=(DisplayName="Add", Category="Math|Float", Keywords = "+,Sum", Deprecated="4.23.0"))
@@ -38,6 +41,9 @@ struct CONTROLRIG_API FRigUnit_Add_FloatFloat : public FRigUnit_BinaryFloatOp
 
 	RIGVM_METHOD()
 	virtual void Execute(const FRigUnitContext& Context) override;
+
+	RIGVM_METHOD()
+	virtual FRigVMStructUpgradeInfo GetUpgradeInfo() const override;
 };
 
 USTRUCT(meta=(DisplayName="Subtract", Category="Math|Float", Keywords = "-", Deprecated="4.23.0"))
@@ -47,6 +53,9 @@ struct CONTROLRIG_API FRigUnit_Subtract_FloatFloat : public FRigUnit_BinaryFloat
 
 	RIGVM_METHOD()
 	virtual void Execute(const FRigUnitContext& Context) override;
+
+	RIGVM_METHOD()
+	virtual FRigVMStructUpgradeInfo GetUpgradeInfo() const override;
 };
 
 USTRUCT(meta=(DisplayName="Divide", Category="Math|Float", Keywords = "/", Deprecated="4.23.0"))
@@ -56,6 +65,9 @@ struct CONTROLRIG_API FRigUnit_Divide_FloatFloat : public FRigUnit_BinaryFloatOp
 
 	RIGVM_METHOD()
 	virtual void Execute(const FRigUnitContext& Context) override;
+
+	RIGVM_METHOD()
+	virtual FRigVMStructUpgradeInfo GetUpgradeInfo() const override;
 };
 
 /** Two args and a result of float type */
@@ -78,6 +90,9 @@ struct CONTROLRIG_API FRigUnit_Clamp_Float: public FRigUnit
 
 	RIGVM_METHOD()
 	virtual void Execute(const FRigUnitContext& Context) override;
+
+	RIGVM_METHOD()
+	virtual FRigVMStructUpgradeInfo GetUpgradeInfo() const override;
 };
 
 /** Two args and a result of float type */
@@ -106,4 +121,7 @@ struct CONTROLRIG_API FRigUnit_MapRange_Float: public FRigUnit
 
 	RIGVM_METHOD()
 	virtual void Execute(const FRigUnitContext& Context) override;
+
+	RIGVM_METHOD()
+	virtual FRigVMStructUpgradeInfo GetUpgradeInfo() const override;
 };
\ No newline at end of file
diff --git a/Engine/Plugins/Animation/ControlRig/Source/ControlRig/Private/Units/Deprecated/Math/RigUnit_Quaternion.cpp b/Engine/Plugins/Animation/ControlRig/Source/ControlRig/Private/Units/Deprecated/Math/RigUnit_Quaternion.cpp
--- a/Engine/Plugins/Animation/ControlRig/Source/ControlRig/Private/Units/Deprecated/Math/RigUnit_Quaternion.cpp
+++ b/Engine/Plugins/Animation/ControlRig/Source/ControlRig/Private/Units/Deprecated/Math/RigUnit_Quaternion.cpp
@@ -1,6 +1,7 @@
 // Copyright Epic Games, Inc. All Rights Reserved.
 
 #include "RigUnit_Quaternion.h"
+#include "Units/Math/RigUnit_MathQuaternion.h"
 
 FRigUnit_MultiplyQuaternion_Execute()
 {
@@ -8,12 +9,34 @@ FRigUnit_MultiplyQuaternion_Execute()
 	Result.Normalize();
 }
 
+FRigVMStructUpgradeInfo FRigUnit_MultiplyQuaternion::GetUpgradeInfo() const
+{
+	FRigUnit_MathQuaternionMul NewNode;
+	NewNode.A = Argument0;
+	NewNode.B = Argument1;
+
+	FRigVMStructUpgradeInfo Info(*this, NewNode);
+	Info.AddRemappedPin(TEXT("Argument0"), TEXT("A"));
+	Info.AddRemappedPin(TEXT("Argument1"), TEXT("B"));
+	return Info;
+}
+
 FRigUnit_InverseQuaterion_Execute()
 {
 	Result = Argument.Inverse();
 	Result.Normalize();
 }
 
+FRigVMStructUpgradeInfo FRigUnit_InverseQuaterion::GetUpgradeInfo() const
+{
+	FRigUnit_MathQuaternionInverse NewNode;
+	NewNode.Value = Argument;
+
+	FRigVMStructUpgradeInfo Info(*this, NewNode);
+	Info.AddRemappedPin(TEXT("Argument"), TEXT("Value"));
+	return Info;
+}
+
 FRigUnit_QuaternionToAxisAndAngle_Execute()
 {
 	FVector NewAxis = Axis.GetSafeNormal();
@@ -21,12 +44,33 @@ FRigUnit_QuaternionToAxisAndAngle_Execute()
 	Angle = FMath::RadiansToDegrees(Angle);
 }
 
+FRigVMStructUpgradeInfo FRigUnit_QuaternionToAxisAndAngle::GetUpgradeInfo() const
+{
+	FRigUnit_MathQuaternionToAxisAndAngle NewNode;
+	NewNode.Value = Argument;
+
+	FRigVMStructUpgradeInfo Info(*this, NewNode);
+	Info.AddRemappedPin(TEXT("Argument"), TEXT("Value"));
+	return Info;
+}
+
 FRigUnit_QuaternionFromAxisAndAngle_Execute()
 {
 	FVector NewAxis = Axis.GetSafeNormal();
 	Result = FQuat(NewAxis, FMath::DegreesToRadians(Angle));
 }
 
+FRigVMStructUpgradeInfo FRigUnit_QuaternionFromAxisAndAngle::GetUpgradeInfo() const
+{
+	FRigUnit_MathQuaternionToAxisAndAngle NewNode;
+	NewNode.Axis = Axis;
+	NewNode.Angle = Angle;
+
+	FRigVMStructUpgradeInfo Info(*this, NewNode);
+	Info.AddRemappedPin(TEXT("Result"), TEXT("Value"));
+	return Info;
+}
+
 FRigUnit_QuaternionToAngle_Execute()
 {
     DECLARE_SCOPE_HIERARCHICAL_COUNTER_RIGUNIT()
@@ -48,3 +92,8 @@ FRigUnit_QuaternionToAngle_Execute()
 	}
 }
 
+FRigVMStructUpgradeInfo FRigUnit_QuaternionToAngle::GetUpgradeInfo() const
+{
+	// this node is no longer supported
+	return FRigVMStructUpgradeInfo();
+}
diff --git a/Engine/Plugins/Animation/ControlRig/Source/ControlRig/Private/Units/Deprecated/Math/RigUnit_Quaternion.h b/Engine/Plugins/Animation/ControlRig/Source/ControlRig/Private/Units/Deprecated/Math/RigUnit_Quaternion.h
--- a/Engine/Plugins/Animation/ControlRig/Source/ControlRig/Private/Units/Deprecated/Math/RigUnit_Quaternion.h
+++ b/Engine/Plugins/Animation/ControlRig/Source/ControlRig/Private/Units/Deprecated/Math/RigUnit_Quaternion.h
@@ -7,7 +7,7 @@
 #include "RigUnit_Quaternion.generated.h"
 
 /** Two args and a result of Quaternion type */
-USTRUCT(meta=(Abstract, NodeColor = "0.1 0.7 0.1", Deprecated="4.23.0"))
+USTRUCT(meta=(Abstract, NodeColor = "0.1 0.7 0.1"))
 struct CONTROLRIG_API FRigUnit_BinaryQuaternionOp : public FRigUnit
 {
 	GENERATED_BODY()
@@ -35,10 +35,13 @@ struct CONTROLRIG_API FRigUnit_MultiplyQuaternion : public FRigUnit_BinaryQuater
 
 	RIGVM_METHOD()
 	virtual void Execute(const FRigUnitContext& Context) override;
+
+	RIGVM_METHOD()
+	virtual FRigVMStructUpgradeInfo GetUpgradeInfo() const override;
 };
 
 /** Two args and a result of Quaternion type */
-USTRUCT(meta=(Abstract, NodeColor = "0.1 0.7 0.1", Deprecated="4.23.0"))
+USTRUCT(meta=(Abstract, NodeColor = "0.1 0.7 0.1"))
 struct CONTROLRIG_API FRigUnit_UnaryQuaternionOp : public FRigUnit
 {
 	GENERATED_BODY()
@@ -57,6 +60,9 @@ struct CONTROLRIG_API FRigUnit_InverseQuaterion: public FRigUnit_UnaryQuaternion
 
 	RIGVM_METHOD()
 	virtual void Execute(const FRigUnitContext& Context) override;
+
+	RIGVM_METHOD()
+	virtual FRigVMStructUpgradeInfo GetUpgradeInfo() const override;
 };
 
 USTRUCT(meta = (DisplayName = "To Axis And Angle(Quaternion)", Category = "Math|Quaternion", NodeColor = "0.1 0.7 0.1", Deprecated="4.23.0"))
@@ -75,6 +81,9 @@ struct CONTROLRIG_API FRigUnit_QuaternionToAxisAndAngle : public FRigUnit
 
 	RIGVM_METHOD()
 	virtual void Execute(const FRigUnitContext& Context) override;
+
+	RIGVM_METHOD()
+	virtual FRigVMStructUpgradeInfo GetUpgradeInfo() const override;
 };
 
 USTRUCT(meta = (DisplayName = "From Axis And Angle(Quaternion)", Category = "Math|Quaternion", NodeColor = "0.1 0.7 0.1", Deprecated="4.23.0"))
@@ -99,6 +108,9 @@ struct CONTROLRIG_API FRigUnit_QuaternionFromAxisAndAngle : public FRigUnit
 
 	RIGVM_METHOD()
 	virtual void Execute(const FRigUnitContext& Context) override;
+
+	RIGVM_METHOD()
+	virtual FRigVMStructUpgradeInfo GetUpgradeInfo() const override;
 };
 
 USTRUCT(meta = (DisplayName = "Get Angle Around Axis", Category = "Math|Quaternion", NodeColor = "0.1 0.7 0.1", Deprecated="4.23.0"))
@@ -123,5 +135,8 @@ struct CONTROLRIG_API FRigUnit_QuaternionToAngle : public FRigUnit
 
 	RIGVM_METHOD()
 	virtual void Execute(const FRigUnitContext& Context) override;
+	
+	RIGVM_METHOD()
+	virtual FRigVMStructUpgradeInfo GetUpgradeInfo() const override;
 };
 
diff --git a/Engine/Plugins/Animation/ControlRig/Source/ControlRig/Private/Units/Deprecated/Math/RigUnit_Transform.cpp b/Engine/Plugins/Animation/ControlRig/Source/ControlRig/Private/Units/Deprecated/Math/RigUnit_Transform.cpp
--- a/Engine/Plugins/Animation/ControlRig/Source/ControlRig/Private/Units/Deprecated/Math/RigUnit_Transform.cpp
+++ b/Engine/Plugins/Animation/ControlRig/Source/ControlRig/Private/Units/Deprecated/Math/RigUnit_Transform.cpp
@@ -1,14 +1,39 @@
 // Copyright Epic Games, Inc. All Rights Reserved.
 
 #include "RigUnit_Transform.h"
+#include "Units/Math/RigUnit_MathTransform.h"
 
 FRigUnit_MultiplyTransform_Execute()
 {
 	Result = Argument0*Argument1;
 }
 
+FRigVMStructUpgradeInfo FRigUnit_MultiplyTransform::GetUpgradeInfo() const
+{
+	FRigUnit_MathTransformMul NewNode;
+	NewNode.A = Argument0;
+	NewNode.B = Argument1;
+
+	FRigVMStructUpgradeInfo Info(*this, NewNode);
+	Info.AddRemappedPin(TEXT("Argument0"), TEXT("A"));
+	Info.AddRemappedPin(TEXT("Argument1"), TEXT("B"));
+	return Info;
+}
+
 FRigUnit_GetRelativeTransform_Execute()
 {
 	Result = Argument0.GetRelativeTransform(Argument1);
 }
 
+FRigVMStructUpgradeInfo FRigUnit_GetRelativeTransform::GetUpgradeInfo() const
+{
+	FRigUnit_MathTransformMakeRelative NewNode;
+	NewNode.Global = Argument0;
+	NewNode.Parent = Argument1;
+
+	FRigVMStructUpgradeInfo Info(*this, NewNode);
+	Info.AddRemappedPin(TEXT("Argument0"), TEXT("Global"));
+	Info.AddRemappedPin(TEXT("Argument1"), TEXT("Parent"));
+	Info.AddRemappedPin(TEXT("Result"), TEXT("Local"));
+	return Info;
+}
diff --git a/Engine/Plugins/Animation/ControlRig/Source/ControlRig/Private/Units/Deprecated/Math/RigUnit_Transform.h b/Engine/Plugins/Animation/ControlRig/Source/ControlRig/Private/Units/Deprecated/Math/RigUnit_Transform.h
--- a/Engine/Plugins/Animation/ControlRig/Source/ControlRig/Private/Units/Deprecated/Math/RigUnit_Transform.h
+++ b/Engine/Plugins/Animation/ControlRig/Source/ControlRig/Private/Units/Deprecated/Math/RigUnit_Transform.h
@@ -7,7 +7,7 @@
 #include "RigUnit_Transform.generated.h"
 
 /** Two args and a result of Transform type */
-USTRUCT(meta=(Abstract, NodeColor = "0.1 0.7 0.1", Deprecated="4.23.0"))
+USTRUCT(meta=(Abstract, NodeColor = "0.1 0.7 0.1"))
 struct CONTROLRIG_API FRigUnit_BinaryTransformOp : public FRigUnit
 {
 	GENERATED_BODY()
@@ -29,6 +29,9 @@ struct CONTROLRIG_API FRigUnit_MultiplyTransform : public FRigUnit_BinaryTransfo
 
 	RIGVM_METHOD()
 	virtual void Execute(const FRigUnitContext& Context) override;
+
+	RIGVM_METHOD()
+	virtual FRigVMStructUpgradeInfo GetUpgradeInfo() const override;
 };
 
 USTRUCT(meta = (DisplayName = "GetRelativeTransform", Category = "Math|Transform", Deprecated="4.23.0"))
@@ -38,5 +41,8 @@ struct CONTROLRIG_API FRigUnit_GetRelativeTransform : public FRigUnit_BinaryTran
 
 	RIGVM_METHOD()
 	virtual void Execute(const FRigUnitContext& Context) override;
+
+	RIGVM_METHOD()
+	virtual FRigVMStructUpgradeInfo GetUpgradeInfo() const override;
 };
 
diff --git a/Engine/Plugins/Animation/ControlRig/Source/ControlRig/Private/Units/Deprecated/Math/RigUnit_Vector.cpp b/Engine/Plugins/Animation/ControlRig/Source/ControlRig/Private/Units/Deprecated/Math/RigUnit_Vector.cpp
--- a/Engine/Plugins/Animation/ControlRig/Source/ControlRig/Private/Units/Deprecated/Math/RigUnit_Vector.cpp
+++ b/Engine/Plugins/Animation/ControlRig/Source/ControlRig/Private/Units/Deprecated/Math/RigUnit_Vector.cpp
@@ -1,28 +1,89 @@
 // Copyright Epic Games, Inc. All Rights Reserved.
 
 #include "RigUnit_Vector.h"
+#include "Units/Math/RigUnit_MathVector.h"
 
 FRigUnit_Multiply_VectorVector_Execute()
 {
 	Result = FRigMathLibrary::Multiply(Argument0, Argument1);
 }
 
+FRigVMStructUpgradeInfo FRigUnit_Multiply_VectorVector::GetUpgradeInfo() const
+{
+	FRigUnit_MathVectorMul NewNode;
+	NewNode.A = Argument0;
+	NewNode.B = Argument1;
+
+	FRigVMStructUpgradeInfo Info(*this, NewNode);
+	Info.AddRemappedPin(TEXT("Argument0"), TEXT("A"));
+	Info.AddRemappedPin(TEXT("Argument1"), TEXT("B"));
+	return Info;
+}
+
 FRigUnit_Add_VectorVector_Execute()
 {
 	Result = FRigMathLibrary::Add(Argument0, Argument1);
 }
 
+FRigVMStructUpgradeInfo FRigUnit_Add_VectorVector::GetUpgradeInfo() const
+{
+	FRigUnit_MathVectorAdd NewNode;
+	NewNode.A = Argument0;
+	NewNode.B = Argument1;
+
+	FRigVMStructUpgradeInfo Info(*this, NewNode);
+	Info.AddRemappedPin(TEXT("Argument0"), TEXT("A"));
+	Info.AddRemappedPin(TEXT("Argument1"), TEXT("B"));
+	return Info;
+}
+
 FRigUnit_Subtract_VectorVector_Execute()
 {
 	Result = FRigMathLibrary::Subtract(Argument0, Argument1);
 }
 
+FRigVMStructUpgradeInfo FRigUnit_Subtract_VectorVector::GetUpgradeInfo() const
+{
+	FRigUnit_MathVectorSub NewNode;
+	NewNode.A = Argument0;
+	NewNode.B = Argument1;
+
+	FRigVMStructUpgradeInfo Info(*this, NewNode);
+	Info.AddRemappedPin(TEXT("Argument0"), TEXT("A"));
+	Info.AddRemappedPin(TEXT("Argument1"), TEXT("B"));
+	return Info;
+}
+
 FRigUnit_Divide_VectorVector_Execute()
 {
 	Result = FRigMathLibrary::Divide(Argument0, Argument1);
 }
 
+FRigVMStructUpgradeInfo FRigUnit_Divide_VectorVector::GetUpgradeInfo() const
+{
+	FRigUnit_MathVectorDiv NewNode;
+	NewNode.A = Argument0;
+	NewNode.B = Argument1;
+
+	FRigVMStructUpgradeInfo Info(*this, NewNode);
+	Info.AddRemappedPin(TEXT("Argument0"), TEXT("A"));
+	Info.AddRemappedPin(TEXT("Argument1"), TEXT("B"));
+	return Info;
+}
+
 FRigUnit_Distance_VectorVector_Execute()
 {
 	Result = (Argument0 - Argument1).Size();
 }
+
+FRigVMStructUpgradeInfo FRigUnit_Distance_VectorVector::GetUpgradeInfo() const
+{
+	FRigUnit_MathVectorDistance NewNode;
+	NewNode.A = Argument0;
+	NewNode.B = Argument1;
+
+	FRigVMStructUpgradeInfo Info(*this, NewNode);
+	Info.AddRemappedPin(TEXT("Argument0"), TEXT("A"));
+	Info.AddRemappedPin(TEXT("Argument1"), TEXT("B"));
+	return Info;
+}
\ No newline at end of file
diff --git a/Engine/Plugins/Animation/ControlRig/Source/ControlRig/Private/Units/Deprecated/Math/RigUnit_Vector.h b/Engine/Plugins/Animation/ControlRig/Source/ControlRig/Private/Units/Deprecated/Math/RigUnit_Vector.h
--- a/Engine/Plugins/Animation/ControlRig/Source/ControlRig/Private/Units/Deprecated/Math/RigUnit_Vector.h
+++ b/Engine/Plugins/Animation/ControlRig/Source/ControlRig/Private/Units/Deprecated/Math/RigUnit_Vector.h
@@ -7,7 +7,7 @@
 #include "RigUnit_Vector.generated.h"
 
 /** Two args and a result of Vector type */
-USTRUCT(meta=(Abstract, NodeColor = "0.1 0.7 0.1", Deprecated="4.23.0"))
+USTRUCT(meta=(Abstract, NodeColor = "0.1 0.7 0.1"))
 struct CONTROLRIG_API FRigUnit_BinaryVectorOp : public FRigUnit
 {
 	GENERATED_BODY()
@@ -29,6 +29,9 @@ struct CONTROLRIG_API FRigUnit_Multiply_VectorVector : public FRigUnit_BinaryVec
 
 	RIGVM_METHOD()
 	virtual void Execute(const FRigUnitContext& Context) override;
+
+	RIGVM_METHOD()
+	virtual FRigVMStructUpgradeInfo GetUpgradeInfo() const override;
 };
 
 USTRUCT(meta=(DisplayName="Add(Vector)", Category="Math|Vector", Deprecated="4.23.0"))
@@ -38,6 +41,9 @@ struct CONTROLRIG_API FRigUnit_Add_VectorVector : public FRigUnit_BinaryVectorOp
 
 	RIGVM_METHOD()
 	virtual void Execute(const FRigUnitContext& Context) override;
+
+	RIGVM_METHOD()
+	virtual FRigVMStructUpgradeInfo GetUpgradeInfo() const override;
 };
 
 USTRUCT(meta=(DisplayName="Subtract(Vector)", Category="Math|Vector", Deprecated="4.23.0"))
@@ -47,6 +53,9 @@ struct CONTROLRIG_API FRigUnit_Subtract_VectorVector : public FRigUnit_BinaryVec
 
 	RIGVM_METHOD()
 	virtual void Execute(const FRigUnitContext& Context) override;
+
+	RIGVM_METHOD()
+	virtual FRigVMStructUpgradeInfo GetUpgradeInfo() const override;
 };
 
 USTRUCT(meta=(DisplayName="Divide(Vector)", Category="Math|Vector", Deprecated="4.23.0"))
@@ -56,6 +65,9 @@ struct CONTROLRIG_API FRigUnit_Divide_VectorVector : public FRigUnit_BinaryVecto
 
 	RIGVM_METHOD()
 	virtual void Execute(const FRigUnitContext& Context) override;
+
+	RIGVM_METHOD()
+	virtual FRigVMStructUpgradeInfo GetUpgradeInfo() const override;
 };
 
 USTRUCT(meta = (DisplayName = "Distance", Category = "Math|Vector", Deprecated="4.23.0", NodeColor = "0.1 0.7 0.1"))
@@ -74,4 +86,7 @@ struct CONTROLRIG_API FRigUnit_Distance_VectorVector : public FRigUnit
 
 	RIGVM_METHOD()
 	virtual void Execute(const FRigUnitContext& Context) override;
+
+	RIGVM_METHOD()
+	virtual FRigVMStructUpgradeInfo GetUpgradeInfo() const override;
 };
\ No newline at end of file
diff --git a/Engine/Plugins/Animation/ControlRig/Source/ControlRig/Private/Units/Deprecated/RigUnit_AimConstraint.cpp b/Engine/Plugins/Animation/ControlRig/Source/ControlRig/Private/Units/Deprecated/RigUnit_AimConstraint.cpp
--- a/Engine/Plugins/Animation/ControlRig/Source/ControlRig/Private/Units/Deprecated/RigUnit_AimConstraint.cpp
+++ b/Engine/Plugins/Animation/ControlRig/Source/ControlRig/Private/Units/Deprecated/RigUnit_AimConstraint.cpp
@@ -88,6 +88,12 @@ FRigUnit_AimConstraint_Execute()
 	}
 }
 
+FRigVMStructUpgradeInfo FRigUnit_AimConstraint::GetUpgradeInfo() const
+{
+	// this node is no longer supported
+	return FRigVMStructUpgradeInfo();
+}
+
 /*
 void FRigUnit_AimConstraint::AddConstraintData(EAimConstraintType ConstraintType, const int32 TargetIndex, const FTransform& SourceTransform, const FTransform& InBaseTransform)
 {
diff --git a/Engine/Plugins/Animation/ControlRig/Source/ControlRig/Private/Units/Deprecated/RigUnit_AimConstraint.h b/Engine/Plugins/Animation/ControlRig/Source/ControlRig/Private/Units/Deprecated/RigUnit_AimConstraint.h
--- a/Engine/Plugins/Animation/ControlRig/Source/ControlRig/Private/Units/Deprecated/RigUnit_AimConstraint.h
+++ b/Engine/Plugins/Animation/ControlRig/Source/ControlRig/Private/Units/Deprecated/RigUnit_AimConstraint.h
@@ -95,4 +95,7 @@ struct CONTROLRIG_API FRigUnit_AimConstraint : public FRigUnitMutable
 	// note that Targets.Num () != ConstraintData.Num()
 	UPROPERTY()
 	FRigUnit_AimConstraint_WorkData WorkData;
+
+	RIGVM_METHOD()
+	virtual FRigVMStructUpgradeInfo GetUpgradeInfo() const override;
 };
diff --git a/Engine/Plugins/Animation/ControlRig/Source/ControlRig/Private/Units/Deprecated/RigUnit_ApplyFK.cpp b/Engine/Plugins/Animation/ControlRig/Source/ControlRig/Private/Units/Deprecated/RigUnit_ApplyFK.cpp
--- a/Engine/Plugins/Animation/ControlRig/Source/ControlRig/Private/Units/Deprecated/RigUnit_ApplyFK.cpp
+++ b/Engine/Plugins/Animation/ControlRig/Source/ControlRig/Private/Units/Deprecated/RigUnit_ApplyFK.cpp
@@ -54,6 +54,12 @@ FRigUnit_ApplyFK_Execute()
 	}
 }
 
+FRigVMStructUpgradeInfo FRigUnit_ApplyFK::GetUpgradeInfo() const
+{
+	// this node is no longer supported
+	return FRigVMStructUpgradeInfo();
+}
+
 #if WITH_DEV_AUTOMATION_TESTS
 #include "Units/RigUnitTest.h"
 
diff --git a/Engine/Plugins/Animation/ControlRig/Source/ControlRig/Private/Units/Deprecated/RigUnit_ApplyFK.h b/Engine/Plugins/Animation/ControlRig/Source/ControlRig/Private/Units/Deprecated/RigUnit_ApplyFK.h
--- a/Engine/Plugins/Animation/ControlRig/Source/ControlRig/Private/Units/Deprecated/RigUnit_ApplyFK.h
+++ b/Engine/Plugins/Animation/ControlRig/Source/ControlRig/Private/Units/Deprecated/RigUnit_ApplyFK.h
@@ -56,4 +56,7 @@ struct CONTROLRIG_API FRigUnit_ApplyFK : public FRigUnitMutable
 	// Transform op option. Use if ETransformSpace is BaseJoint
 	UPROPERTY(EditAnywhere, Category = "ApplyFK", meta = (Input))
 	FName BaseJoint;
+
+	RIGVM_METHOD()
+	virtual FRigVMStructUpgradeInfo GetUpgradeInfo() const override;
 };
diff --git a/Engine/Plugins/Animation/ControlRig/Source/ControlRig/Private/Units/Deprecated/RigUnit_BlendTransform.cpp b/Engine/Plugins/Animation/ControlRig/Source/ControlRig/Private/Units/Deprecated/RigUnit_BlendTransform.cpp
--- a/Engine/Plugins/Animation/ControlRig/Source/ControlRig/Private/Units/Deprecated/RigUnit_BlendTransform.cpp
+++ b/Engine/Plugins/Animation/ControlRig/Source/ControlRig/Private/Units/Deprecated/RigUnit_BlendTransform.cpp
@@ -49,3 +49,9 @@ FRigUnit_BlendTransform_Execute()
 	// if failed on any of the above, it will just use source as target pose
 	Result = Source;
 }
+
+FRigVMStructUpgradeInfo FRigUnit_BlendTransform::GetUpgradeInfo() const
+{
+	// this node is no longer supported
+	return FRigVMStructUpgradeInfo();
+}
diff --git a/Engine/Plugins/Animation/ControlRig/Source/ControlRig/Private/Units/Deprecated/RigUnit_BlendTransform.h b/Engine/Plugins/Animation/ControlRig/Source/ControlRig/Private/Units/Deprecated/RigUnit_BlendTransform.h
--- a/Engine/Plugins/Animation/ControlRig/Source/ControlRig/Private/Units/Deprecated/RigUnit_BlendTransform.h
+++ b/Engine/Plugins/Animation/ControlRig/Source/ControlRig/Private/Units/Deprecated/RigUnit_BlendTransform.h
@@ -37,4 +37,7 @@ struct CONTROLRIG_API FRigUnit_BlendTransform : public FRigUnit
 
 	RIGVM_METHOD()
 	virtual void Execute(const FRigUnitContext& Context) override;
+
+	RIGVM_METHOD()
+	virtual FRigVMStructUpgradeInfo GetUpgradeInfo() const override;
 };
\ No newline at end of file
diff --git a/Engine/Plugins/Animation/ControlRig/Source/ControlRig/Private/Units/Deprecated/RigUnit_GetJointTransform.cpp b/Engine/Plugins/Animation/ControlRig/Source/ControlRig/Private/Units/Deprecated/RigUnit_GetJointTransform.cpp
--- a/Engine/Plugins/Animation/ControlRig/Source/ControlRig/Private/Units/Deprecated/RigUnit_GetJointTransform.cpp
+++ b/Engine/Plugins/Animation/ControlRig/Source/ControlRig/Private/Units/Deprecated/RigUnit_GetJointTransform.cpp
@@ -3,6 +3,7 @@
 #include "RigUnit_GetJointTransform.h"
 #include "Units/RigUnitContext.h"
 #include "HelperUtil.h"
+#include "Units/Hierarchy/RigUnit_GetTransform.h"
 
 FRigUnit_GetJointTransform_Execute()
 {
@@ -34,3 +35,9 @@ FRigUnit_GetJointTransform_Execute()
 		}
 	}
 }
+
+FRigVMStructUpgradeInfo FRigUnit_GetJointTransform::GetUpgradeInfo() const
+{
+	// this node is no longer supported
+	return FRigVMStructUpgradeInfo();
+}
diff --git a/Engine/Plugins/Animation/ControlRig/Source/ControlRig/Private/Units/Deprecated/RigUnit_GetJointTransform.h b/Engine/Plugins/Animation/ControlRig/Source/ControlRig/Private/Units/Deprecated/RigUnit_GetJointTransform.h
--- a/Engine/Plugins/Animation/ControlRig/Source/ControlRig/Private/Units/Deprecated/RigUnit_GetJointTransform.h
+++ b/Engine/Plugins/Animation/ControlRig/Source/ControlRig/Private/Units/Deprecated/RigUnit_GetJointTransform.h
@@ -41,4 +41,7 @@ struct CONTROLRIG_API FRigUnit_GetJointTransform : public FRigUnitMutable
 	// possibly space, relative transform so on can be input
 	UPROPERTY(meta=(Output))
 	FTransform Output;
+
+	RIGVM_METHOD()
+	virtual FRigVMStructUpgradeInfo GetUpgradeInfo() const override;
 };
diff --git a/Engine/Plugins/Animation/ControlRig/Source/ControlRig/Private/Units/Deprecated/RigUnit_TwoBoneIKFK.cpp b/Engine/Plugins/Animation/ControlRig/Source/ControlRig/Private/Units/Deprecated/RigUnit_TwoBoneIKFK.cpp
--- a/Engine/Plugins/Animation/ControlRig/Source/ControlRig/Private/Units/Deprecated/RigUnit_TwoBoneIKFK.cpp
+++ b/Engine/Plugins/Animation/ControlRig/Source/ControlRig/Private/Units/Deprecated/RigUnit_TwoBoneIKFK.cpp
@@ -117,3 +117,8 @@ FRigUnit_TwoBoneIKFK_Execute()
 	}
 }
 
+FRigVMStructUpgradeInfo FRigUnit_TwoBoneIKFK::GetUpgradeInfo() const
+{
+	// this node is no longer supported
+	return FRigVMStructUpgradeInfo();
+}
diff --git a/Engine/Plugins/Animation/ControlRig/Source/ControlRig/Private/Units/Drawing/RigUnit_DrawContainer.cpp b/Engine/Plugins/Animation/ControlRig/Source/ControlRig/Private/Units/Drawing/RigUnit_DrawContainer.cpp
--- a/Engine/Plugins/Animation/ControlRig/Source/ControlRig/Private/Units/Drawing/RigUnit_DrawContainer.cpp
+++ b/Engine/Plugins/Animation/ControlRig/Source/ControlRig/Private/Units/Drawing/RigUnit_DrawContainer.cpp
@@ -33,6 +33,12 @@ FRigUnit_DrawContainerGetInstruction_Execute()
 	}
 }
 
+FRigVMStructUpgradeInfo FRigUnit_DrawContainerGetInstruction::GetUpgradeInfo() const
+{
+	// this node is no longer supported
+	return FRigVMStructUpgradeInfo();
+}
+
 FRigUnit_DrawContainerSetThickness_Execute()
 {
 	DECLARE_SCOPE_HIERARCHICAL_COUNTER_RIGUNIT()
@@ -57,6 +63,12 @@ FRigUnit_DrawContainerSetThickness_Execute()
 	}
 }
 
+FRigVMStructUpgradeInfo FRigUnit_DrawContainerSetThickness::GetUpgradeInfo() const
+{
+	// this node is no longer supported
+	return FRigVMStructUpgradeInfo();
+}
+
 FRigUnit_DrawContainerSetColor_Execute()
 {
 	DECLARE_SCOPE_HIERARCHICAL_COUNTER_RIGUNIT()
@@ -81,6 +93,12 @@ FRigUnit_DrawContainerSetColor_Execute()
 	}
 }
 
+FRigVMStructUpgradeInfo FRigUnit_DrawContainerSetColor::GetUpgradeInfo() const
+{
+	// this node is no longer supported
+	return FRigVMStructUpgradeInfo();
+}
+
 FRigUnit_DrawContainerSetTransform_Execute()
 {
 	DECLARE_SCOPE_HIERARCHICAL_COUNTER_RIGUNIT()
@@ -103,4 +121,10 @@ FRigUnit_DrawContainerSetTransform_Execute()
 			break;
 		}
 	}
-}
\ No newline at end of file
+}
+
+FRigVMStructUpgradeInfo FRigUnit_DrawContainerSetTransform::GetUpgradeInfo() const
+{
+	// this node is no longer supported
+	return FRigVMStructUpgradeInfo();
+}
diff --git a/Engine/Plugins/Animation/ControlRig/Source/ControlRig/Private/Units/Drawing/RigUnit_DrawContainer.h b/Engine/Plugins/Animation/ControlRig/Source/ControlRig/Private/Units/Drawing/RigUnit_DrawContainer.h
--- a/Engine/Plugins/Animation/ControlRig/Source/ControlRig/Private/Units/Drawing/RigUnit_DrawContainer.h
+++ b/Engine/Plugins/Animation/ControlRig/Source/ControlRig/Private/Units/Drawing/RigUnit_DrawContainer.h
@@ -31,6 +31,9 @@ struct CONTROLRIG_API FRigUnit_DrawContainerGetInstruction : public FRigUnit
 
 	UPROPERTY(meta = (Output))
 	FTransform Transform;
+
+	RIGVM_METHOD()
+	virtual FRigVMStructUpgradeInfo GetUpgradeInfo() const override;
 };
 
 /**
@@ -55,6 +58,9 @@ struct CONTROLRIG_API FRigUnit_DrawContainerSetColor: public FRigUnitMutable
 
 	UPROPERTY(meta = (Input))
 	FLinearColor Color;
+
+	RIGVM_METHOD()
+	virtual FRigVMStructUpgradeInfo GetUpgradeInfo() const override;
 };
 
 /**
@@ -79,6 +85,9 @@ struct CONTROLRIG_API FRigUnit_DrawContainerSetThickness : public FRigUnitMutabl
 
 	UPROPERTY(meta = (Input))
 	float Thickness;
+
+	RIGVM_METHOD()
+	virtual FRigVMStructUpgradeInfo GetUpgradeInfo() const override;
 };
 
 /**
@@ -103,4 +112,7 @@ struct CONTROLRIG_API FRigUnit_DrawContainerSetTransform : public FRigUnitMutabl
 
 	UPROPERTY(meta = (Input))
 	FTransform Transform;
+
+	RIGVM_METHOD()
+	virtual FRigVMStructUpgradeInfo GetUpgradeInfo() const override;
 };
\ No newline at end of file
diff --git a/Engine/Plugins/Animation/ControlRig/Source/ControlRig/Private/Units/Execution/RigUnit_Collection.cpp b/Engine/Plugins/Animation/ControlRig/Source/ControlRig/Private/Units/Execution/RigUnit_Collection.cpp
--- a/Engine/Plugins/Animation/ControlRig/Source/ControlRig/Private/Units/Execution/RigUnit_Collection.cpp
+++ b/Engine/Plugins/Animation/ControlRig/Source/ControlRig/Private/Units/Execution/RigUnit_Collection.cpp
@@ -25,6 +25,16 @@ FRigUnit_CollectionChain_Execute()
 	FRigUnit_CollectionChainArray::StaticExecute(RigVMExecuteContext, FirstItem, LastItem, Reverse, Collection.Keys, Context);
 }
 
+FRigVMStructUpgradeInfo FRigUnit_CollectionChain::GetUpgradeInfo() const
+{
+	FRigUnit_CollectionChainArray NewNode;
+	NewNode.FirstItem = FirstItem;
+	NewNode.LastItem = LastItem;
+	NewNode.Reverse = Reverse;
+
+	return FRigVMStructUpgradeInfo(*this, NewNode);
+}
+
 FRigUnit_CollectionChainArray_Execute()
 {
 	uint32 Hash = FRigUnit_CollectionChain_Hash + Context.Hierarchy->GetTopologyVersion() * 17;
@@ -68,6 +78,15 @@ FRigUnit_CollectionNameSearch_Execute()
 	FRigUnit_CollectionNameSearchArray::StaticExecute(RigVMExecuteContext, PartialName, TypeToSearch, Collection.Keys, Context);
 }
 
+FRigVMStructUpgradeInfo FRigUnit_CollectionNameSearch::GetUpgradeInfo() const
+{
+	FRigUnit_CollectionNameSearchArray NewNode;
+	NewNode.PartialName = PartialName;
+	NewNode.TypeToSearch = TypeToSearch;
+
+	return FRigVMStructUpgradeInfo(*this, NewNode);
+}
+
 FRigUnit_CollectionNameSearchArray_Execute()
 {
 	DECLARE_SCOPE_HIERARCHICAL_COUNTER_RIGUNIT()
@@ -95,6 +114,17 @@ FRigUnit_CollectionChildren_Execute()
 	FRigUnit_CollectionChildrenArray::StaticExecute(RigVMExecuteContext, Parent, bIncludeParent, bRecursive, TypeToSearch, Collection.Keys, Context);
 }
 
+FRigVMStructUpgradeInfo FRigUnit_CollectionChildren::GetUpgradeInfo() const
+{
+	FRigUnit_CollectionChildrenArray NewNode;
+	NewNode.Parent = Parent;
+	NewNode.bIncludeParent = bIncludeParent;
+	NewNode.bRecursive = bRecursive;
+	NewNode.TypeToSearch = TypeToSearch;
+
+	return FRigVMStructUpgradeInfo(*this, NewNode);
+}
+
 FRigUnit_CollectionChildrenArray_Execute()
 {
 	DECLARE_SCOPE_HIERARCHICAL_COUNTER_RIGUNIT()
@@ -170,6 +200,18 @@ FRigUnit_CollectionReplaceItems_Execute()
 	FRigUnit_CollectionReplaceItemsArray::StaticExecute(RigVMExecuteContext, Items.Keys, Old, New, RemoveInvalidItems, bAllowDuplicates, Collection.Keys, Context);
 }
 
+FRigVMStructUpgradeInfo FRigUnit_CollectionReplaceItems::GetUpgradeInfo() const
+{
+	FRigUnit_CollectionReplaceItemsArray NewNode;
+	NewNode.Items = Items.Keys;
+	NewNode.Old = Old;
+	NewNode.New = New;
+	NewNode.RemoveInvalidItems = RemoveInvalidItems;
+	NewNode.bAllowDuplicates = bAllowDuplicates;
+
+	return FRigVMStructUpgradeInfo(*this, NewNode);
+}
+
 FRigUnit_CollectionReplaceItemsArray_Execute()
 {
 	DECLARE_SCOPE_HIERARCHICAL_COUNTER_RIGUNIT()
@@ -247,6 +289,16 @@ FRigUnit_CollectionGetParentIndices_Execute()
 	FRigUnit_CollectionGetParentIndicesItemArray::StaticExecute(RigVMExecuteContext, Collection.Keys, ParentIndices, Context);
 }
 
+FRigVMStructUpgradeInfo FRigUnit_CollectionGetParentIndices::GetUpgradeInfo() const
+{
+	FRigUnit_CollectionGetParentIndicesItemArray NewNode;
+	NewNode.Items = Collection.Keys;
+
+	FRigVMStructUpgradeInfo Info(*this, NewNode);
+	Info.AddRemappedPin(TEXT("Collection"), TEXT("Items"));
+	return Info;
+}
+
 FRigUnit_CollectionGetParentIndicesItemArray_Execute()
 {
 	DECLARE_SCOPE_HIERARCHICAL_COUNTER_RIGUNIT()
@@ -319,6 +371,12 @@ FRigUnit_CollectionUnion_Execute()
 	Collection = FRigElementKeyCollection::MakeUnion(A, B, bAllowDuplicates);
 }
 
+FRigVMStructUpgradeInfo FRigUnit_CollectionUnion::GetUpgradeInfo() const
+{
+	// this node is no longer supported. you can rely on generic array nodes for this now
+	return FRigVMStructUpgradeInfo();
+}
+
 FRigUnit_CollectionIntersection_Execute()
 {
 	DECLARE_SCOPE_HIERARCHICAL_COUNTER_RIGUNIT()
@@ -326,6 +384,12 @@ FRigUnit_CollectionIntersection_Execute()
 	Collection = FRigElementKeyCollection::MakeIntersection(A, B);
 }
 
+FRigVMStructUpgradeInfo FRigUnit_CollectionIntersection::GetUpgradeInfo() const
+{
+	// this node is no longer supported. you can rely on generic array nodes for this now
+	return FRigVMStructUpgradeInfo();
+}
+
 FRigUnit_CollectionDifference_Execute()
 {
 	DECLARE_SCOPE_HIERARCHICAL_COUNTER_RIGUNIT()
@@ -333,6 +397,12 @@ FRigUnit_CollectionDifference_Execute()
 	Collection = FRigElementKeyCollection::MakeDifference(A, B);
 }
 
+FRigVMStructUpgradeInfo FRigUnit_CollectionDifference::GetUpgradeInfo() const
+{
+	// this node is no longer supported. you can rely on generic array nodes for this now
+	return FRigVMStructUpgradeInfo();
+}
+
 FRigUnit_CollectionReverse_Execute()
 {
 	DECLARE_SCOPE_HIERARCHICAL_COUNTER_RIGUNIT()
@@ -340,6 +410,12 @@ FRigUnit_CollectionReverse_Execute()
 	Reversed = FRigElementKeyCollection::MakeReversed(Collection);
 }
 
+FRigVMStructUpgradeInfo FRigUnit_CollectionReverse::GetUpgradeInfo() const
+{
+	// this node is no longer supported. you can rely on generic array nodes for this now
+	return FRigVMStructUpgradeInfo();
+}
+
 FRigUnit_CollectionCount_Execute()
 {
 	DECLARE_SCOPE_HIERARCHICAL_COUNTER_RIGUNIT()
@@ -347,6 +423,12 @@ FRigUnit_CollectionCount_Execute()
 	Count = Collection.Num();
 }
 
+FRigVMStructUpgradeInfo FRigUnit_CollectionCount::GetUpgradeInfo() const
+{
+	// this node is no longer supported. you can rely on generic array nodes for this now
+	return FRigVMStructUpgradeInfo();
+}
+
 FRigUnit_CollectionItemAtIndex_Execute()
 {
 	DECLARE_SCOPE_HIERARCHICAL_COUNTER_RIGUNIT()
@@ -361,6 +443,12 @@ FRigUnit_CollectionItemAtIndex_Execute()
 	}
 }
 
+FRigVMStructUpgradeInfo FRigUnit_CollectionItemAtIndex::GetUpgradeInfo() const
+{
+	// this node is no longer supported. you can rely on generic array nodes for this now
+	return FRigVMStructUpgradeInfo();
+}
+
 FRigUnit_CollectionLoop_Execute()
 {
     DECLARE_SCOPE_HIERARCHICAL_COUNTER_RIGUNIT()
@@ -378,10 +466,22 @@ FRigUnit_CollectionLoop_Execute()
 	}
 }
 
+FRigVMStructUpgradeInfo FRigUnit_CollectionLoop::GetUpgradeInfo() const
+{
+	// this node is no longer supported. you can rely on generic array nodes for this now
+	return FRigVMStructUpgradeInfo();
+}
+
 FRigUnit_CollectionAddItem_Execute()
 {
 	DECLARE_SCOPE_HIERARCHICAL_COUNTER_RIGUNIT();
 
 	Result = Collection;
 	Result.AddUnique(Item);
+}
+
+FRigVMStructUpgradeInfo FRigUnit_CollectionAddItem::GetUpgradeInfo() const
+{
+	// this node is no longer supported. you can rely on generic array nodes for this now
+	return FRigVMStructUpgradeInfo();
 }
\ No newline at end of file
diff --git a/Engine/Plugins/Animation/ControlRig/Source/ControlRig/Private/Units/Execution/RigUnit_Collection.h b/Engine/Plugins/Animation/ControlRig/Source/ControlRig/Private/Units/Execution/RigUnit_Collection.h
--- a/Engine/Plugins/Animation/ControlRig/Source/ControlRig/Private/Units/Execution/RigUnit_Collection.h
+++ b/Engine/Plugins/Animation/ControlRig/Source/ControlRig/Private/Units/Execution/RigUnit_Collection.h
@@ -47,6 +47,9 @@ struct CONTROLRIG_API FRigUnit_CollectionChain : public FRigUnit_CollectionBase
 
 	UPROPERTY(meta = (Output))
 	FRigElementKeyCollection Collection;
+
+	RIGVM_METHOD()
+	virtual FRigVMStructUpgradeInfo GetUpgradeInfo() const override;
 };
 
 /**
@@ -106,6 +109,9 @@ struct CONTROLRIG_API FRigUnit_CollectionNameSearch : public FRigUnit_Collection
 
 	UPROPERTY(meta = (Output))
 	FRigElementKeyCollection Collection;
+
+	RIGVM_METHOD()
+	virtual FRigVMStructUpgradeInfo GetUpgradeInfo() const override;
 };
 
 /**
@@ -170,6 +176,9 @@ struct CONTROLRIG_API FRigUnit_CollectionChildren : public FRigUnit_CollectionBa
 
 	UPROPERTY(meta = (Output))
 	FRigElementKeyCollection Collection;
+
+	RIGVM_METHOD()
+	virtual FRigVMStructUpgradeInfo GetUpgradeInfo() const override;
 };
 
 /**
@@ -243,6 +252,9 @@ struct CONTROLRIG_API FRigUnit_CollectionReplaceItems : public FRigUnit_Collecti
 
 	UPROPERTY(meta = (Output))
 	FRigElementKeyCollection Collection;
+
+	RIGVM_METHOD()
+	virtual FRigVMStructUpgradeInfo GetUpgradeInfo() const override;
 };
 
 /**
@@ -354,6 +366,9 @@ struct CONTROLRIG_API FRigUnit_CollectionGetParentIndices : public FRigUnit_Coll
 
 	UPROPERTY(meta = (Output))
 	TArray<int32> ParentIndices;
+
+	RIGVM_METHOD()
+	virtual FRigVMStructUpgradeInfo GetUpgradeInfo() const override;
 };
 
 /**
@@ -409,6 +424,9 @@ struct CONTROLRIG_API FRigUnit_CollectionUnion : public FRigUnit_CollectionBase
 
 	UPROPERTY(meta = (Output))
 	FRigElementKeyCollection Collection;
+
+	RIGVM_METHOD()
+	virtual FRigVMStructUpgradeInfo GetUpgradeInfo() const override;
 };
 
 /**
@@ -435,6 +453,9 @@ struct CONTROLRIG_API FRigUnit_CollectionIntersection : public FRigUnit_Collecti
 
 	UPROPERTY(meta = (Output))
 	FRigElementKeyCollection Collection;
+
+	RIGVM_METHOD()
+	virtual FRigVMStructUpgradeInfo GetUpgradeInfo() const override;
 };
 
 /**
@@ -461,6 +482,9 @@ struct CONTROLRIG_API FRigUnit_CollectionDifference : public FRigUnit_Collection
 
 	UPROPERTY(meta = (Output))
 	FRigElementKeyCollection Collection;
+
+	RIGVM_METHOD()
+	virtual FRigVMStructUpgradeInfo GetUpgradeInfo() const override;
 };
 
 /**
@@ -483,6 +507,9 @@ struct CONTROLRIG_API FRigUnit_CollectionReverse : public FRigUnit_CollectionBas
 
 	UPROPERTY(meta = (Output))
 	FRigElementKeyCollection Reversed;
+
+	RIGVM_METHOD()
+	virtual FRigVMStructUpgradeInfo GetUpgradeInfo() const override;
 };
 
 /**
@@ -507,6 +534,9 @@ struct CONTROLRIG_API FRigUnit_CollectionCount : public FRigUnit_CollectionBase
 
 	UPROPERTY(meta = (Output))
 	int32 Count;
+
+	RIGVM_METHOD()
+	virtual FRigVMStructUpgradeInfo GetUpgradeInfo() const override;
 };
 
 /**
@@ -535,6 +565,9 @@ struct CONTROLRIG_API FRigUnit_CollectionItemAtIndex : public FRigUnit_Collectio
 
 	UPROPERTY(meta = (Output))
 	FRigElementKey Item;
+
+	RIGVM_METHOD()
+	virtual FRigVMStructUpgradeInfo GetUpgradeInfo() const override;
 };
 
 /**
@@ -585,6 +618,9 @@ struct CONTROLRIG_API FRigUnit_CollectionLoop : public FRigUnit_CollectionBaseMu
 
 	UPROPERTY(meta = (Output))
 	FControlRigExecuteContext Completed;
+
+	RIGVM_METHOD()
+	virtual FRigVMStructUpgradeInfo GetUpgradeInfo() const override;
 };
 
 /**
@@ -612,4 +648,7 @@ struct CONTROLRIG_API FRigUnit_CollectionAddItem : public FRigUnit_CollectionBas
 
 	UPROPERTY(meta = (Output))
 	FRigElementKeyCollection Result;
+
+	RIGVM_METHOD()
+	virtual FRigVMStructUpgradeInfo GetUpgradeInfo() const override;
 };
\ No newline at end of file
diff --git a/Engine/Plugins/Animation/ControlRig/Source/ControlRig/Private/Units/Execution/RigUnit_DynamicHierarchy.cpp b/Engine/Plugins/Animation/ControlRig/Source/ControlRig/Private/Units/Execution/RigUnit_DynamicHierarchy.cpp
--- a/Engine/Plugins/Animation/ControlRig/Source/ControlRig/Private/Units/Execution/RigUnit_DynamicHierarchy.cpp
+++ b/Engine/Plugins/Animation/ControlRig/Source/ControlRig/Private/Units/Execution/RigUnit_DynamicHierarchy.cpp
@@ -136,6 +136,16 @@ FRigUnit_HierarchyGetParentWeights_Execute()
 	FRigUnit_HierarchyGetParentWeightsArray::StaticExecute(RigVMExecuteContext, Child, Weights, Parents.Keys, Context);
 }
 
+FRigVMStructUpgradeInfo FRigUnit_HierarchyGetParentWeights::GetUpgradeInfo() const
+{
+	FRigUnit_HierarchyGetParentWeightsArray NewNode;
+	NewNode.Child = Child;
+	NewNode.Weights = Weights;
+	NewNode.Parents = Parents.Keys;
+
+	return FRigVMStructUpgradeInfo();
+}
+
 FRigUnit_HierarchyGetParentWeightsArray_Execute()
 {
 	if((Context.State != EControlRigState::Update) || (Context.Hierarchy == nullptr))
diff --git a/Engine/Plugins/Animation/ControlRig/Source/ControlRig/Private/Units/Execution/RigUnit_DynamicHierarchy.h b/Engine/Plugins/Animation/ControlRig/Source/ControlRig/Private/Units/Execution/RigUnit_DynamicHierarchy.h
--- a/Engine/Plugins/Animation/ControlRig/Source/ControlRig/Private/Units/Execution/RigUnit_DynamicHierarchy.h
+++ b/Engine/Plugins/Animation/ControlRig/Source/ControlRig/Private/Units/Execution/RigUnit_DynamicHierarchy.h
@@ -133,6 +133,9 @@ struct CONTROLRIG_API FRigUnit_HierarchyGetParentWeights : public FRigUnit_Dynam
 	 */
 	UPROPERTY(meta = (Output))
 	FRigElementKeyCollection Parents;
+
+	RIGVM_METHOD()
+	virtual FRigVMStructUpgradeInfo GetUpgradeInfo() const override;
 };
 
 /**
diff --git a/Engine/Plugins/Animation/ControlRig/Source/ControlRig/Private/Units/Execution/RigUnit_Hierarchy.cpp b/Engine/Plugins/Animation/ControlRig/Source/ControlRig/Private/Units/Execution/RigUnit_Hierarchy.cpp
--- a/Engine/Plugins/Animation/ControlRig/Source/ControlRig/Private/Units/Execution/RigUnit_Hierarchy.cpp
+++ b/Engine/Plugins/Animation/ControlRig/Source/ControlRig/Private/Units/Execution/RigUnit_Hierarchy.cpp
@@ -3,6 +3,7 @@
 #include "RigUnit_Hierarchy.h"
 #include "Units/Execution/RigUnit_Item.h"
 #include "Units/RigUnitContext.h"
+#include "RigUnit_Collection.h"
 
 FRigUnit_HierarchyGetParent_Execute()
 {
@@ -39,6 +40,16 @@ FRigUnit_HierarchyGetParents_Execute()
 	FRigUnit_HierarchyGetParentsItemArray::StaticExecute(RigVMExecuteContext, Child, bIncludeChild, bReverse, Parents.Keys, CachedChild, CachedParents, Context);
 }
 
+FRigVMStructUpgradeInfo FRigUnit_HierarchyGetParents::GetUpgradeInfo() const
+{
+	FRigUnit_HierarchyGetParentsItemArray NewNode;
+	NewNode.Child = Child;
+	NewNode.bIncludeChild = bIncludeChild;
+	NewNode.bReverse = bReverse;
+	NewNode.Parents = Parents.Keys;
+	return FRigVMStructUpgradeInfo(*this, NewNode);
+}
+
 FRigUnit_HierarchyGetParentsItemArray_Execute()
 {
     DECLARE_SCOPE_HIERARCHICAL_COUNTER_RIGUNIT()
@@ -111,11 +122,29 @@ FRigUnit_HierarchyGetChildren_Execute()
 	Children = CachedChildren;
 }
 
+FRigVMStructUpgradeInfo FRigUnit_HierarchyGetChildren::GetUpgradeInfo() const
+{
+	FRigUnit_CollectionChildrenArray NewNode;
+	NewNode.Parent = Parent;
+	NewNode.bRecursive = bRecursive;
+	NewNode.bIncludeParent = bIncludeParent;
+	NewNode.TypeToSearch = ERigElementType::All;
+	return FRigVMStructUpgradeInfo(*this, NewNode);
+}
+
 FRigUnit_HierarchyGetSiblings_Execute()
 {
 	FRigUnit_HierarchyGetSiblingsItemArray::StaticExecute(RigVMExecuteContext, Item, bIncludeItem, Siblings.Keys, CachedItem, CachedSiblings, Context);
 }
 
+FRigVMStructUpgradeInfo FRigUnit_HierarchyGetSiblings::GetUpgradeInfo() const
+{
+	FRigUnit_HierarchyGetSiblingsItemArray NewNode;
+	NewNode.Item = Item;
+	NewNode.bIncludeItem = bIncludeItem;
+	return FRigVMStructUpgradeInfo(*this, NewNode);
+}
+
 FRigUnit_HierarchyGetSiblingsItemArray_Execute()
 {
     DECLARE_SCOPE_HIERARCHICAL_COUNTER_RIGUNIT()
@@ -164,6 +193,16 @@ FRigUnit_HierarchyGetPose_Execute()
 	FRigUnit_HierarchyGetPoseItemArray::StaticExecute(RigVMExecuteContext, Initial, ElementType, ItemsToGet.Keys, Pose, Context);
 }
 
+FRigVMStructUpgradeInfo FRigUnit_HierarchyGetPose::GetUpgradeInfo() const
+{
+	FRigUnit_HierarchyGetPoseItemArray NewNode;
+	NewNode.Initial = Initial;
+	NewNode.ElementType = ElementType;
+	NewNode.ItemsToGet = ItemsToGet.Keys;
+
+	return FRigVMStructUpgradeInfo(*this, NewNode);
+}
+
 FRigUnit_HierarchyGetPoseItemArray_Execute()
 {
 	Pose = Context.Hierarchy->GetPose(Initial, ElementType, ItemsToGet);
@@ -174,6 +213,17 @@ FRigUnit_HierarchySetPose_Execute()
 	FRigUnit_HierarchySetPoseItemArray::StaticExecute(RigVMExecuteContext, Pose, ElementType, Space, ItemsToSet.Keys, Weight, ExecuteContext, Context);
 }
 
+FRigVMStructUpgradeInfo FRigUnit_HierarchySetPose::GetUpgradeInfo() const
+{
+	FRigUnit_HierarchySetPoseItemArray NewNode;
+	NewNode.Pose = Pose;
+	NewNode.ElementType = ElementType;
+	NewNode.Space = Space;
+	NewNode.ItemsToSet = ItemsToSet.Keys;
+
+	return FRigVMStructUpgradeInfo(*this, NewNode);
+}
+
 FRigUnit_HierarchySetPoseItemArray_Execute()
 {
 	ExecuteContext.Hierarchy->SetPose(
@@ -195,6 +245,15 @@ FRigUnit_PoseGetItems_Execute()
 	FRigUnit_PoseGetItemsItemArray::StaticExecute(RigVMExecuteContext, Pose, ElementType, Items.Keys, Context);
 }
 
+FRigVMStructUpgradeInfo FRigUnit_PoseGetItems::GetUpgradeInfo() const
+{
+	FRigUnit_PoseGetItemsItemArray NewNode;
+	NewNode.Pose = Pose;
+	NewNode.ElementType = ElementType;
+
+	return FRigVMStructUpgradeInfo(*this, NewNode);
+}
+
 FRigUnit_PoseGetItemsItemArray_Execute()
 {
 	Items.Reset();
@@ -356,7 +415,7 @@ FRigUnit_PoseGetDelta_Execute()
 }
 
 bool FRigUnit_PoseGetDelta::ArePoseElementsEqual(const FRigPoseElement& A, const FRigPoseElement& B,
-	EBoneGetterSetterMode Space, float PositionU, float RotationU, float ScaleU, float CurveU)
+                                                 EBoneGetterSetterMode Space, float PositionU, float RotationU, float ScaleU, float CurveU)
 {
 	const FRigElementKey& KeyA = A.Index.GetKey();
 	const FRigElementKey& KeyB = B.Index.GetKey();
diff --git a/Engine/Plugins/Animation/ControlRig/Source/ControlRig/Private/Units/Execution/RigUnit_Hierarchy.h b/Engine/Plugins/Animation/ControlRig/Source/ControlRig/Private/Units/Execution/RigUnit_Hierarchy.h
--- a/Engine/Plugins/Animation/ControlRig/Source/ControlRig/Private/Units/Execution/RigUnit_Hierarchy.h
+++ b/Engine/Plugins/Animation/ControlRig/Source/ControlRig/Private/Units/Execution/RigUnit_Hierarchy.h
@@ -89,6 +89,9 @@ struct CONTROLRIG_API FRigUnit_HierarchyGetParents : public FRigUnit_HierarchyBa
 	// Used to cache the internally
 	UPROPERTY()
 	FRigElementKeyCollection CachedParents;
+
+	RIGVM_METHOD()
+	virtual FRigVMStructUpgradeInfo GetUpgradeInfo() const override;
 };
 
 /**
@@ -172,6 +175,9 @@ struct CONTROLRIG_API FRigUnit_HierarchyGetChildren : public FRigUnit_HierarchyB
 	// Used to cache the internally
 	UPROPERTY()
 	FRigElementKeyCollection CachedChildren;
+
+	RIGVM_METHOD()
+	virtual FRigVMStructUpgradeInfo GetUpgradeInfo() const override;
 };
 
 /**
@@ -209,6 +215,9 @@ struct CONTROLRIG_API FRigUnit_HierarchyGetSiblings : public FRigUnit_HierarchyB
 	// Used to cache the internally
 	UPROPERTY()
 	FRigElementKeyCollection CachedSiblings;
+
+	RIGVM_METHOD()
+	virtual FRigVMStructUpgradeInfo GetUpgradeInfo() const override;
 };
 
 /**
@@ -279,6 +288,9 @@ struct CONTROLRIG_API FRigUnit_HierarchyGetPose : public FRigUnit_HierarchyBase
 
 	UPROPERTY(meta = (Output))
 	FRigPose Pose;
+
+	RIGVM_METHOD()
+	virtual FRigVMStructUpgradeInfo GetUpgradeInfo() const override;
 };
 
 /**
@@ -349,6 +361,9 @@ struct CONTROLRIG_API FRigUnit_HierarchySetPose : public FRigUnit_HierarchyBaseM
 
 	UPROPERTY(meta = (Input))
 	float Weight;
+
+	RIGVM_METHOD()
+	virtual FRigVMStructUpgradeInfo GetUpgradeInfo() const override;
 };
 
 /**
@@ -437,6 +452,9 @@ struct CONTROLRIG_API FRigUnit_PoseGetItems : public FRigUnit_HierarchyBase
 
 	UPROPERTY(meta = (Output))
 	FRigElementKeyCollection Items;
+
+	RIGVM_METHOD()
+	virtual FRigVMStructUpgradeInfo GetUpgradeInfo() const override;
 };
 
 /**
diff --git a/Engine/Plugins/Animation/ControlRig/Source/ControlRig/Private/Units/Hierarchy/RigUnit_AddBoneTransform.cpp b/Engine/Plugins/Animation/ControlRig/Source/ControlRig/Private/Units/Hierarchy/RigUnit_AddBoneTransform.cpp
--- a/Engine/Plugins/Animation/ControlRig/Source/ControlRig/Private/Units/Hierarchy/RigUnit_AddBoneTransform.cpp
+++ b/Engine/Plugins/Animation/ControlRig/Source/ControlRig/Private/Units/Hierarchy/RigUnit_AddBoneTransform.cpp
@@ -1,6 +1,7 @@
 // Copyright Epic Games, Inc. All Rights Reserved.
 
 #include "Units/Hierarchy/RigUnit_AddBoneTransform.h"
+#include "Units/Hierarchy/RigUnit_OffsetTransform.h"
 #include "Units/RigUnitContext.h"
 
 FRigUnit_AddBoneTransform_Execute()
@@ -53,3 +54,14 @@ FRigUnit_AddBoneTransform_Execute()
 	}
 }
 
+FRigVMStructUpgradeInfo FRigUnit_AddBoneTransform::GetUpgradeInfo() const
+{
+	FRigUnit_OffsetTransformForItem NewNode;
+	NewNode.Item = FRigElementKey(Bone, ERigElementType::Bone);
+	NewNode.Weight = Weight;
+	NewNode.bPropagateToChildren = bPropagateToChildren;
+
+	FRigVMStructUpgradeInfo Info(*this, NewNode);
+	Info.AddRemappedPin(TEXT("Bone"), TEXT("Item.Name"));
+	return Info;
+}
\ No newline at end of file
diff --git a/Engine/Plugins/Animation/ControlRig/Source/ControlRig/Private/Units/Hierarchy/RigUnit_BoneName.cpp b/Engine/Plugins/Animation/ControlRig/Source/ControlRig/Private/Units/Hierarchy/RigUnit_BoneName.cpp
--- a/Engine/Plugins/Animation/ControlRig/Source/ControlRig/Private/Units/Hierarchy/RigUnit_BoneName.cpp
+++ b/Engine/Plugins/Animation/ControlRig/Source/ControlRig/Private/Units/Hierarchy/RigUnit_BoneName.cpp
@@ -11,10 +11,40 @@ FRigUnit_BoneName_Execute()
 {
 }
 
+FRigVMStructUpgradeInfo FRigUnit_BoneName::GetUpgradeInfo() const
+{
+	FRigUnit_Item NewNode;
+	NewNode.Item = FRigElementKey(Bone, ERigElementType::Bone);
+
+	FRigVMStructUpgradeInfo Info(*this, NewNode);
+	Info.AddRemappedPin(TEXT("Bone"), TEXT("Item.Name"));
+	return Info;
+}
+
 FRigUnit_SpaceName_Execute()
 {
 }
 
+FRigVMStructUpgradeInfo FRigUnit_SpaceName::GetUpgradeInfo() const
+{
+	FRigUnit_Item NewNode;
+	NewNode.Item = FRigElementKey(Space, ERigElementType::Null);
+
+	FRigVMStructUpgradeInfo Info(*this, NewNode);
+	Info.AddRemappedPin(TEXT("Space"), TEXT("Item.Name"));
+	return Info;
+}
+
 FRigUnit_ControlName_Execute()
 {
-}
\ No newline at end of file
+}
+
+FRigVMStructUpgradeInfo FRigUnit_ControlName::GetUpgradeInfo() const
+{
+	FRigUnit_Item NewNode;
+	NewNode.Item = FRigElementKey(Control, ERigElementType::Control);
+
+	FRigVMStructUpgradeInfo Info(*this, NewNode);
+	Info.AddRemappedPin(TEXT("Control"), TEXT("Item.Name"));
+	return Info;
+}
diff --git a/Engine/Plugins/Animation/ControlRig/Source/ControlRig/Private/Units/Hierarchy/RigUnit_BoneName.h b/Engine/Plugins/Animation/ControlRig/Source/ControlRig/Private/Units/Hierarchy/RigUnit_BoneName.h
--- a/Engine/Plugins/Animation/ControlRig/Source/ControlRig/Private/Units/Hierarchy/RigUnit_BoneName.h
+++ b/Engine/Plugins/Animation/ControlRig/Source/ControlRig/Private/Units/Hierarchy/RigUnit_BoneName.h
@@ -45,6 +45,9 @@ struct CONTROLRIG_API FRigUnit_BoneName : public FRigUnit
 	 */
 	UPROPERTY(meta = (Input, Output))
 	FName Bone;
+
+	RIGVM_METHOD()
+	virtual FRigVMStructUpgradeInfo GetUpgradeInfo() const override;
 };
 
 /**
@@ -66,6 +69,9 @@ struct CONTROLRIG_API FRigUnit_SpaceName : public FRigUnit
 	 */
 	UPROPERTY(meta = (Input, Output))
 	FName Space;
+
+	RIGVM_METHOD()
+	virtual FRigVMStructUpgradeInfo GetUpgradeInfo() const override;
 };
 
 /**
@@ -87,4 +93,7 @@ struct CONTROLRIG_API FRigUnit_ControlName : public FRigUnit
 	 */
 	UPROPERTY(meta = (Input, Output))
 	FName Control;
+
+	RIGVM_METHOD()
+	virtual FRigVMStructUpgradeInfo GetUpgradeInfo() const override;
 };
diff --git a/Engine/Plugins/Animation/ControlRig/Source/ControlRig/Private/Units/Hierarchy/RigUnit_GetBoneTransform.cpp b/Engine/Plugins/Animation/ControlRig/Source/ControlRig/Private/Units/Hierarchy/RigUnit_GetBoneTransform.cpp
--- a/Engine/Plugins/Animation/ControlRig/Source/ControlRig/Private/Units/Hierarchy/RigUnit_GetBoneTransform.cpp
+++ b/Engine/Plugins/Animation/ControlRig/Source/ControlRig/Private/Units/Hierarchy/RigUnit_GetBoneTransform.cpp
@@ -3,6 +3,7 @@
 #include "RigUnit_GetBoneTransform.h"
 #include "RigUnit_GetInitialBoneTransform.h"
 #include "Units/RigUnitContext.h"
+#include "Units/Hierarchy/RigUnit_GetTransform.h"
 
 FRigUnit_GetBoneTransform_Execute()
 {
@@ -52,6 +53,17 @@ FRigUnit_GetBoneTransform_Execute()
 	}
 }
 
+FRigVMStructUpgradeInfo FRigUnit_GetBoneTransform::GetUpgradeInfo() const
+{
+	FRigUnit_GetTransform NewNode;
+	NewNode.Item = FRigElementKey(Bone, ERigElementType::Control);
+	NewNode.Space = Space;
+
+	FRigVMStructUpgradeInfo Info(*this, NewNode);
+	Info.AddRemappedPin(TEXT("Bone"), TEXT("Item.Name"));
+	return Info;
+}
+
 #if WITH_DEV_AUTOMATION_TESTS
 #include "Units/RigUnitTest.h"
 
diff --git a/Engine/Plugins/Animation/ControlRig/Source/ControlRig/Private/Units/Hierarchy/RigUnit_GetBoneTransform.h b/Engine/Plugins/Animation/ControlRig/Source/ControlRig/Private/Units/Hierarchy/RigUnit_GetBoneTransform.h
--- a/Engine/Plugins/Animation/ControlRig/Source/ControlRig/Private/Units/Hierarchy/RigUnit_GetBoneTransform.h
+++ b/Engine/Plugins/Animation/ControlRig/Source/ControlRig/Private/Units/Hierarchy/RigUnit_GetBoneTransform.h
@@ -53,4 +53,7 @@ struct CONTROLRIG_API FRigUnit_GetBoneTransform : public FRigUnit
 	// Used to cache the internally used bone index
 	UPROPERTY(transient)
 	FCachedRigElement CachedBone;
+
+	RIGVM_METHOD()
+	virtual FRigVMStructUpgradeInfo GetUpgradeInfo() const override;
 };
diff --git a/Engine/Plugins/Animation/ControlRig/Source/ControlRig/Private/Units/Hierarchy/RigUnit_GetControlInitialTransform.cpp b/Engine/Plugins/Animation/ControlRig/Source/ControlRig/Private/Units/Hierarchy/RigUnit_GetControlInitialTransform.cpp
--- a/Engine/Plugins/Animation/ControlRig/Source/ControlRig/Private/Units/Hierarchy/RigUnit_GetControlInitialTransform.cpp
+++ b/Engine/Plugins/Animation/ControlRig/Source/ControlRig/Private/Units/Hierarchy/RigUnit_GetControlInitialTransform.cpp
@@ -2,6 +2,7 @@
 
 #include "RigUnit_GetControlInitialTransform.h"
 #include "Units/RigUnitContext.h"
+#include "Units/Hierarchy/RigUnit_GetTransform.h"
 
 FRigUnit_GetControlInitialTransform_Execute()
 {
@@ -50,3 +51,15 @@ FRigUnit_GetControlInitialTransform_Execute()
 		}
 	}
 }
+
+FRigVMStructUpgradeInfo FRigUnit_GetControlInitialTransform::GetUpgradeInfo() const
+{
+	FRigUnit_GetTransform NewNode;
+	NewNode.Item = FRigElementKey(Control, ERigElementType::Control);
+	NewNode.Space = Space;
+	NewNode.bInitial = true;
+
+	FRigVMStructUpgradeInfo Info(*this, NewNode);
+	Info.AddRemappedPin(TEXT("Control"), TEXT("Item.Name"));
+	return Info;
+}
diff --git a/Engine/Plugins/Animation/ControlRig/Source/ControlRig/Private/Units/Hierarchy/RigUnit_GetControlInitialTransform.h b/Engine/Plugins/Animation/ControlRig/Source/ControlRig/Private/Units/Hierarchy/RigUnit_GetControlInitialTransform.h
--- a/Engine/Plugins/Animation/ControlRig/Source/ControlRig/Private/Units/Hierarchy/RigUnit_GetControlInitialTransform.h
+++ b/Engine/Plugins/Animation/ControlRig/Source/ControlRig/Private/Units/Hierarchy/RigUnit_GetControlInitialTransform.h
@@ -41,4 +41,7 @@ struct CONTROLRIG_API FRigUnit_GetControlInitialTransform : public FRigUnit
 	// Used to cache the internally used bone index
 	UPROPERTY()
 	FCachedRigElement CachedControlIndex;
+
+	RIGVM_METHOD()
+	virtual FRigVMStructUpgradeInfo GetUpgradeInfo() const override;
 };
diff --git a/Engine/Plugins/Animation/ControlRig/Source/ControlRig/Private/Units/Hierarchy/RigUnit_GetInitialBoneTransform.cpp b/Engine/Plugins/Animation/ControlRig/Source/ControlRig/Private/Units/Hierarchy/RigUnit_GetInitialBoneTransform.cpp
--- a/Engine/Plugins/Animation/ControlRig/Source/ControlRig/Private/Units/Hierarchy/RigUnit_GetInitialBoneTransform.cpp
+++ b/Engine/Plugins/Animation/ControlRig/Source/ControlRig/Private/Units/Hierarchy/RigUnit_GetInitialBoneTransform.cpp
@@ -2,6 +2,7 @@
 
 #include "RigUnit_GetInitialBoneTransform.h"
 #include "Units/RigUnitContext.h"
+#include "Units/Hierarchy/RigUnit_GetTransform.h"
 
 FRigUnit_GetInitialBoneTransform_Execute()
 {
@@ -50,6 +51,18 @@ FRigUnit_GetInitialBoneTransform_Execute()
 	}
 }
 
+FRigVMStructUpgradeInfo FRigUnit_GetInitialBoneTransform::GetUpgradeInfo() const
+{
+	FRigUnit_GetTransform NewNode;
+	NewNode.Item = FRigElementKey(Bone, ERigElementType::Control);
+	NewNode.Space = Space;
+	NewNode.bInitial = true;
+
+	FRigVMStructUpgradeInfo Info(*this, NewNode);
+	Info.AddRemappedPin(TEXT("Bone"), TEXT("Item.Name"));
+	return Info;
+}
+
 #if WITH_DEV_AUTOMATION_TESTS
 #include "Units/RigUnitTest.h"
 
diff --git a/Engine/Plugins/Animation/ControlRig/Source/ControlRig/Private/Units/Hierarchy/RigUnit_GetInitialBoneTransform.h b/Engine/Plugins/Animation/ControlRig/Source/ControlRig/Private/Units/Hierarchy/RigUnit_GetInitialBoneTransform.h
--- a/Engine/Plugins/Animation/ControlRig/Source/ControlRig/Private/Units/Hierarchy/RigUnit_GetInitialBoneTransform.h
+++ b/Engine/Plugins/Animation/ControlRig/Source/ControlRig/Private/Units/Hierarchy/RigUnit_GetInitialBoneTransform.h
@@ -53,4 +53,7 @@ struct CONTROLRIG_API FRigUnit_GetInitialBoneTransform : public FRigUnit
 	// Used to cache the internally used bone index
 	UPROPERTY(transient)
 	FCachedRigElement CachedBone;
+
+	RIGVM_METHOD()
+	virtual FRigVMStructUpgradeInfo GetUpgradeInfo() const override;
 };
diff --git a/Engine/Plugins/Animation/ControlRig/Source/ControlRig/Private/Units/Hierarchy/RigUnit_GetRelativeBoneTransform.cpp b/Engine/Plugins/Animation/ControlRig/Source/ControlRig/Private/Units/Hierarchy/RigUnit_GetRelativeBoneTransform.cpp
--- a/Engine/Plugins/Animation/ControlRig/Source/ControlRig/Private/Units/Hierarchy/RigUnit_GetRelativeBoneTransform.cpp
+++ b/Engine/Plugins/Animation/ControlRig/Source/ControlRig/Private/Units/Hierarchy/RigUnit_GetRelativeBoneTransform.cpp
@@ -2,6 +2,7 @@
 
 #include "RigUnit_GetRelativeBoneTransform.h"
 #include "Units/RigUnitContext.h"
+#include "Units/Hierarchy/RigUnit_GetRelativeTransform.h"
 
 FRigUnit_GetRelativeBoneTransform_Execute()
 {
@@ -41,6 +42,19 @@ FRigUnit_GetRelativeBoneTransform_Execute()
 	}
 }
 
+FRigVMStructUpgradeInfo FRigUnit_GetRelativeBoneTransform::GetUpgradeInfo() const
+{
+	FRigUnit_GetRelativeTransformForItem NewNode;
+	NewNode.Parent = FRigElementKey(Space, ERigElementType::Bone);
+	NewNode.Child = FRigElementKey(Bone, ERigElementType::Bone);
+
+	FRigVMStructUpgradeInfo Info(*this, NewNode);
+	Info.AddRemappedPin(TEXT("Space"), TEXT("Parent.Name"));
+	Info.AddRemappedPin(TEXT("Bone"), TEXT("Child.Name"));
+	Info.AddRemappedPin(TEXT("Transform"), TEXT("RelativeTransform"));
+	return Info;
+}
+
 #if WITH_DEV_AUTOMATION_TESTS
 #include "Units/RigUnitTest.h"
 
diff --git a/Engine/Plugins/Animation/ControlRig/Source/ControlRig/Private/Units/Hierarchy/RigUnit_GetRelativeBoneTransform.h b/Engine/Plugins/Animation/ControlRig/Source/ControlRig/Private/Units/Hierarchy/RigUnit_GetRelativeBoneTransform.h
--- a/Engine/Plugins/Animation/ControlRig/Source/ControlRig/Private/Units/Hierarchy/RigUnit_GetRelativeBoneTransform.h
+++ b/Engine/Plugins/Animation/ControlRig/Source/ControlRig/Private/Units/Hierarchy/RigUnit_GetRelativeBoneTransform.h
@@ -53,4 +53,7 @@ struct CONTROLRIG_API FRigUnit_GetRelativeBoneTransform : public FRigUnit
 	// Used to cache the internally used space index
 	UPROPERTY(transient)
 	FCachedRigElement CachedSpace;
+
+	RIGVM_METHOD()
+	virtual FRigVMStructUpgradeInfo GetUpgradeInfo() const override;
 };
diff --git a/Engine/Plugins/Animation/ControlRig/Source/ControlRig/Private/Units/Hierarchy/RigUnit_GetSpaceTransform.cpp b/Engine/Plugins/Animation/ControlRig/Source/ControlRig/Private/Units/Hierarchy/RigUnit_GetSpaceTransform.cpp
--- a/Engine/Plugins/Animation/ControlRig/Source/ControlRig/Private/Units/Hierarchy/RigUnit_GetSpaceTransform.cpp
+++ b/Engine/Plugins/Animation/ControlRig/Source/ControlRig/Private/Units/Hierarchy/RigUnit_GetSpaceTransform.cpp
@@ -2,6 +2,7 @@
 
 #include "RigUnit_GetSpaceTransform.h"
 #include "Units/RigUnitContext.h"
+#include "Units/Hierarchy/RigUnit_GetTransform.h"
 
 FRigUnit_GetSpaceTransform_Execute()
 {
@@ -45,3 +46,15 @@ FRigUnit_GetSpaceTransform_Execute()
 		}
 	}
 }
+
+FRigVMStructUpgradeInfo FRigUnit_GetSpaceTransform::GetUpgradeInfo() const
+{
+	FRigUnit_GetTransform NewNode;
+	NewNode.Item = FRigElementKey(Space, ERigElementType::Control);
+	NewNode.Space = SpaceType;
+
+	FRigVMStructUpgradeInfo Info(*this, NewNode);
+	Info.AddRemappedPin(TEXT("Space"), TEXT("Item.Name"));
+	Info.AddRemappedPin(TEXT("SpaceType"), TEXT("Space"));
+	return Info;
+}
diff --git a/Engine/Plugins/Animation/ControlRig/Source/ControlRig/Private/Units/Hierarchy/RigUnit_GetSpaceTransform.h b/Engine/Plugins/Animation/ControlRig/Source/ControlRig/Private/Units/Hierarchy/RigUnit_GetSpaceTransform.h
--- a/Engine/Plugins/Animation/ControlRig/Source/ControlRig/Private/Units/Hierarchy/RigUnit_GetSpaceTransform.h
+++ b/Engine/Plugins/Animation/ControlRig/Source/ControlRig/Private/Units/Hierarchy/RigUnit_GetSpaceTransform.h
@@ -41,4 +41,7 @@ struct CONTROLRIG_API FRigUnit_GetSpaceTransform : public FRigUnit
 	// Used to cache the internally used bone index
 	UPROPERTY()
 	FCachedRigElement CachedSpaceIndex;
+
+	RIGVM_METHOD()
+	virtual FRigVMStructUpgradeInfo GetUpgradeInfo() const override;
 };
diff --git a/Engine/Plugins/Animation/ControlRig/Source/ControlRig/Private/Units/Hierarchy/RigUnit_GetTransform.cpp b/Engine/Plugins/Animation/ControlRig/Source/ControlRig/Private/Units/Hierarchy/RigUnit_GetTransform.cpp
--- a/Engine/Plugins/Animation/ControlRig/Source/ControlRig/Private/Units/Hierarchy/RigUnit_GetTransform.cpp
+++ b/Engine/Plugins/Animation/ControlRig/Source/ControlRig/Private/Units/Hierarchy/RigUnit_GetTransform.cpp
@@ -88,6 +88,15 @@ FRigUnit_GetTransformArray_Execute()
 	FRigUnit_GetTransformItemArray::StaticExecute(RigVMExecuteContext, Items.Keys, Space, bInitial, Transforms, CachedIndex, Context);
 }
 
+FRigVMStructUpgradeInfo FRigUnit_GetTransformArray::GetUpgradeInfo() const
+{
+	FRigUnit_GetTransformItemArray NewNode;
+	NewNode.Items = Items.GetKeys();
+	NewNode.Space = Space;
+	NewNode.bInitial = bInitial;
+	return FRigVMStructUpgradeInfo(*this, NewNode);
+}
+
 FRigUnit_GetTransformItemArray_Execute()
 {
     DECLARE_SCOPE_HIERARCHICAL_COUNTER_RIGUNIT()
diff --git a/Engine/Plugins/Animation/ControlRig/Source/ControlRig/Private/Units/Hierarchy/RigUnit_GetTransform.h b/Engine/Plugins/Animation/ControlRig/Source/ControlRig/Private/Units/Hierarchy/RigUnit_GetTransform.h
--- a/Engine/Plugins/Animation/ControlRig/Source/ControlRig/Private/Units/Hierarchy/RigUnit_GetTransform.h
+++ b/Engine/Plugins/Animation/ControlRig/Source/ControlRig/Private/Units/Hierarchy/RigUnit_GetTransform.h
@@ -111,6 +111,9 @@ struct CONTROLRIG_API FRigUnit_GetTransformArray : public FRigUnit
 	// Used to cache the internally
 	UPROPERTY()
 	TArray<FCachedRigElement> CachedIndex;
+
+	RIGVM_METHOD()
+	virtual FRigVMStructUpgradeInfo GetUpgradeInfo() const override;
 };
 
 /**
diff --git a/Engine/Plugins/Animation/ControlRig/Source/ControlRig/Private/Units/Hierarchy/RigUnit_ParentSwitchConstraint.cpp b/Engine/Plugins/Animation/ControlRig/Source/ControlRig/Private/Units/Hierarchy/RigUnit_ParentSwitchConstraint.cpp
--- a/Engine/Plugins/Animation/ControlRig/Source/ControlRig/Private/Units/Hierarchy/RigUnit_ParentSwitchConstraint.cpp
+++ b/Engine/Plugins/Animation/ControlRig/Source/ControlRig/Private/Units/Hierarchy/RigUnit_ParentSwitchConstraint.cpp
@@ -8,6 +8,18 @@ FRigUnit_ParentSwitchConstraint_Execute()
 	FRigUnit_ParentSwitchConstraintArray::StaticExecute(RigVMExecuteContext, Subject, ParentIndex, Parents.Keys, InitialGlobalTransform, Weight, Transform, Switched, CachedSubject, CachedParent, RelativeOffset, ExecuteContext, Context);	
 }
 
+FRigVMStructUpgradeInfo FRigUnit_ParentSwitchConstraint::GetUpgradeInfo() const
+{
+	FRigUnit_ParentSwitchConstraintArray NewNode;
+	NewNode.Subject = Subject;
+	NewNode.ParentIndex = ParentIndex;
+	NewNode.Parents = Parents.Keys;
+	NewNode.InitialGlobalTransform = InitialGlobalTransform;
+	NewNode.Weight = Weight;
+
+	return FRigVMStructUpgradeInfo(*this, NewNode);
+}
+
 FRigUnit_ParentSwitchConstraintArray_Execute()
 {
     DECLARE_SCOPE_HIERARCHICAL_COUNTER_RIGUNIT()
@@ -87,3 +99,9 @@ FRigUnit_ParentSwitchConstraintArray_Execute()
 		}
 	}
 }
+
+FRigVMStructUpgradeInfo FRigUnit_ParentSwitchConstraintArray::GetUpgradeInfo() const
+{
+	// this node is no longer supported.
+	return FRigVMStructUpgradeInfo();
+}
diff --git a/Engine/Plugins/Animation/ControlRig/Source/ControlRig/Private/Units/Hierarchy/RigUnit_ParentSwitchConstraint.h b/Engine/Plugins/Animation/ControlRig/Source/ControlRig/Private/Units/Hierarchy/RigUnit_ParentSwitchConstraint.h
--- a/Engine/Plugins/Animation/ControlRig/Source/ControlRig/Private/Units/Hierarchy/RigUnit_ParentSwitchConstraint.h
+++ b/Engine/Plugins/Animation/ControlRig/Source/ControlRig/Private/Units/Hierarchy/RigUnit_ParentSwitchConstraint.h
@@ -80,6 +80,9 @@ struct CONTROLRIG_API FRigUnit_ParentSwitchConstraint : public FRigUnitMutable
 	// The cached relative offset between subject and parent
 	UPROPERTY()
 	FTransform RelativeOffset;
+
+	RIGVM_METHOD()
+	virtual FRigVMStructUpgradeInfo GetUpgradeInfo() const override;
 };
 
 /**
@@ -156,4 +159,7 @@ struct CONTROLRIG_API FRigUnit_ParentSwitchConstraintArray : public FRigUnitMuta
 	// The cached relative offset between subject and parent
 	UPROPERTY()
 	FTransform RelativeOffset;
+
+	RIGVM_METHOD()
+	virtual FRigVMStructUpgradeInfo GetUpgradeInfo() const override;
 };
diff --git a/Engine/Plugins/Animation/ControlRig/Source/ControlRig/Private/Units/Hierarchy/RigUnit_PropagateTransform.cpp b/Engine/Plugins/Animation/ControlRig/Source/ControlRig/Private/Units/Hierarchy/RigUnit_PropagateTransform.cpp
--- a/Engine/Plugins/Animation/ControlRig/Source/ControlRig/Private/Units/Hierarchy/RigUnit_PropagateTransform.cpp
+++ b/Engine/Plugins/Animation/ControlRig/Source/ControlRig/Private/Units/Hierarchy/RigUnit_PropagateTransform.cpp
@@ -64,4 +64,10 @@ FRigUnit_PropagateTransform_Execute()
 		}
 		*/
 	}
-}
\ No newline at end of file
+}
+
+FRigVMStructUpgradeInfo FRigUnit_PropagateTransform::GetUpgradeInfo() const
+{
+	// this node is no longer supported
+	return FRigVMStructUpgradeInfo();
+}
diff --git a/Engine/Plugins/Animation/ControlRig/Source/ControlRig/Private/Units/Hierarchy/RigUnit_PropagateTransform.h b/Engine/Plugins/Animation/ControlRig/Source/ControlRig/Private/Units/Hierarchy/RigUnit_PropagateTransform.h
--- a/Engine/Plugins/Animation/ControlRig/Source/ControlRig/Private/Units/Hierarchy/RigUnit_PropagateTransform.h
+++ b/Engine/Plugins/Animation/ControlRig/Source/ControlRig/Private/Units/Hierarchy/RigUnit_PropagateTransform.h
@@ -54,4 +54,7 @@ struct CONTROLRIG_API FRigUnit_PropagateTransform : public FRigUnitMutable
 	// Used to cache the item internally
 	UPROPERTY(transient)
 	FCachedRigElement CachedIndex;
+
+	RIGVM_METHOD()
+	virtual FRigVMStructUpgradeInfo GetUpgradeInfo() const override;
 };
diff --git a/Engine/Plugins/Animation/ControlRig/Source/ControlRig/Private/Units/Hierarchy/RigUnit_SetBoneInitialTransform.cpp b/Engine/Plugins/Animation/ControlRig/Source/ControlRig/Private/Units/Hierarchy/RigUnit_SetBoneInitialTransform.cpp
--- a/Engine/Plugins/Animation/ControlRig/Source/ControlRig/Private/Units/Hierarchy/RigUnit_SetBoneInitialTransform.cpp
+++ b/Engine/Plugins/Animation/ControlRig/Source/ControlRig/Private/Units/Hierarchy/RigUnit_SetBoneInitialTransform.cpp
@@ -3,6 +3,7 @@
 #include "RigUnit_SetBoneInitialTransform.h"
 #include "Units/RigUnitContext.h"
 #include "Units/Execution/RigUnit_PrepareForExecution.h"
+#include "Units/Hierarchy/RigUnit_SetTransform.h"
 
 FRigUnit_SetBoneInitialTransform_Execute()
 {
@@ -43,3 +44,15 @@ FRigUnit_SetBoneInitialTransform_Execute()
 		}
 	}
 }
+
+FRigVMStructUpgradeInfo FRigUnit_SetBoneInitialTransform::GetUpgradeInfo() const
+{
+	FRigUnit_SetTransform NewNode;
+	NewNode.Item = FRigElementKey(Bone, ERigElementType::Bone);
+	NewNode.Space = Space;
+	NewNode.Transform = Transform;
+
+	FRigVMStructUpgradeInfo Info(*this, NewNode);
+	Info.AddRemappedPin(TEXT("Bone"), TEXT("Item.Name"));
+	return Info;
+}
diff --git a/Engine/Plugins/Animation/ControlRig/Source/ControlRig/Private/Units/Hierarchy/RigUnit_SetBoneInitialTransform.h b/Engine/Plugins/Animation/ControlRig/Source/ControlRig/Private/Units/Hierarchy/RigUnit_SetBoneInitialTransform.h
--- a/Engine/Plugins/Animation/ControlRig/Source/ControlRig/Private/Units/Hierarchy/RigUnit_SetBoneInitialTransform.h
+++ b/Engine/Plugins/Animation/ControlRig/Source/ControlRig/Private/Units/Hierarchy/RigUnit_SetBoneInitialTransform.h
@@ -59,4 +59,7 @@ struct CONTROLRIG_API FRigUnit_SetBoneInitialTransform : public FRigUnitMutable
 	// Used to cache the internally used bone index
 	UPROPERTY(transient)
 	FCachedRigElement CachedBone;
+
+	RIGVM_METHOD()
+	virtual FRigVMStructUpgradeInfo GetUpgradeInfo() const override;
 };
diff --git a/Engine/Plugins/Animation/ControlRig/Source/ControlRig/Private/Units/Hierarchy/RigUnit_SetBoneRotation.cpp b/Engine/Plugins/Animation/ControlRig/Source/ControlRig/Private/Units/Hierarchy/RigUnit_SetBoneRotation.cpp
--- a/Engine/Plugins/Animation/ControlRig/Source/ControlRig/Private/Units/Hierarchy/RigUnit_SetBoneRotation.cpp
+++ b/Engine/Plugins/Animation/ControlRig/Source/ControlRig/Private/Units/Hierarchy/RigUnit_SetBoneRotation.cpp
@@ -2,6 +2,7 @@
 
 #include "RigUnit_SetBoneRotation.h"
 #include "Units/RigUnitContext.h"
+#include "Units/Hierarchy/RigUnit_SetTransform.h"
 
 FRigUnit_SetBoneRotation_Execute()
 {
@@ -76,6 +77,18 @@ FRigUnit_SetBoneRotation_Execute()
 	}
 }
 
+FRigVMStructUpgradeInfo FRigUnit_SetBoneRotation::GetUpgradeInfo() const
+{
+	FRigUnit_SetRotation NewNode;
+	NewNode.Item = FRigElementKey(Bone, ERigElementType::Bone);
+	NewNode.Space = Space;
+	NewNode.Rotation = Rotation;
+
+	FRigVMStructUpgradeInfo Info(*this, NewNode);
+	Info.AddRemappedPin(TEXT("Bone"), TEXT("Item.Name"));
+	return Info;
+}
+
 #if WITH_DEV_AUTOMATION_TESTS
 #include "Units/RigUnitTest.h"
 
diff --git a/Engine/Plugins/Animation/ControlRig/Source/ControlRig/Private/Units/Hierarchy/RigUnit_SetBoneRotation.h b/Engine/Plugins/Animation/ControlRig/Source/ControlRig/Private/Units/Hierarchy/RigUnit_SetBoneRotation.h
--- a/Engine/Plugins/Animation/ControlRig/Source/ControlRig/Private/Units/Hierarchy/RigUnit_SetBoneRotation.h
+++ b/Engine/Plugins/Animation/ControlRig/Source/ControlRig/Private/Units/Hierarchy/RigUnit_SetBoneRotation.h
@@ -73,4 +73,7 @@ struct CONTROLRIG_API FRigUnit_SetBoneRotation : public FRigUnitMutable
 	// Used to cache the internally used bone index
 	UPROPERTY(transient)
 	FCachedRigElement CachedBone;
+
+	RIGVM_METHOD()
+	virtual FRigVMStructUpgradeInfo GetUpgradeInfo() const override;
 };
diff --git a/Engine/Plugins/Animation/ControlRig/Source/ControlRig/Private/Units/Hierarchy/RigUnit_SetBoneTransform.cpp b/Engine/Plugins/Animation/ControlRig/Source/ControlRig/Private/Units/Hierarchy/RigUnit_SetBoneTransform.cpp
--- a/Engine/Plugins/Animation/ControlRig/Source/ControlRig/Private/Units/Hierarchy/RigUnit_SetBoneTransform.cpp
+++ b/Engine/Plugins/Animation/ControlRig/Source/ControlRig/Private/Units/Hierarchy/RigUnit_SetBoneTransform.cpp
@@ -2,6 +2,7 @@
 
 #include "RigUnit_SetBoneTransform.h"
 #include "Units/RigUnitContext.h"
+#include "Units/Hierarchy/RigUnit_SetTransform.h"
 
 FRigUnit_SetBoneTransform_Execute()
 {
@@ -73,6 +74,18 @@ FRigUnit_SetBoneTransform_Execute()
 	}
 }
 
+FRigVMStructUpgradeInfo FRigUnit_SetBoneTransform::GetUpgradeInfo() const
+{
+	FRigUnit_SetTransform NewNode;
+	NewNode.Item = FRigElementKey(Bone, ERigElementType::Bone);
+	NewNode.Space = Space;
+	NewNode.Transform = Transform;
+
+	FRigVMStructUpgradeInfo Info(*this, NewNode);
+	Info.AddRemappedPin(TEXT("Bone"), TEXT("Item.Name"));
+	return Info;
+}
+
 #if WITH_DEV_AUTOMATION_TESTS
 #include "Units/RigUnitTest.h"
 
diff --git a/Engine/Plugins/Animation/ControlRig/Source/ControlRig/Private/Units/Hierarchy/RigUnit_SetBoneTransform.h b/Engine/Plugins/Animation/ControlRig/Source/ControlRig/Private/Units/Hierarchy/RigUnit_SetBoneTransform.h
--- a/Engine/Plugins/Animation/ControlRig/Source/ControlRig/Private/Units/Hierarchy/RigUnit_SetBoneTransform.h
+++ b/Engine/Plugins/Animation/ControlRig/Source/ControlRig/Private/Units/Hierarchy/RigUnit_SetBoneTransform.h
@@ -78,4 +78,7 @@ struct CONTROLRIG_API FRigUnit_SetBoneTransform : public FRigUnitMutable
 	// Used to cache the internally used bone index
 	UPROPERTY(transient)
 	FCachedRigElement CachedBone;
+
+	RIGVM_METHOD()
+	virtual FRigVMStructUpgradeInfo GetUpgradeInfo() const override;
 };
diff --git a/Engine/Plugins/Animation/ControlRig/Source/ControlRig/Private/Units/Hierarchy/RigUnit_SetBoneTranslation.cpp b/Engine/Plugins/Animation/ControlRig/Source/ControlRig/Private/Units/Hierarchy/RigUnit_SetBoneTranslation.cpp
--- a/Engine/Plugins/Animation/ControlRig/Source/ControlRig/Private/Units/Hierarchy/RigUnit_SetBoneTranslation.cpp
+++ b/Engine/Plugins/Animation/ControlRig/Source/ControlRig/Private/Units/Hierarchy/RigUnit_SetBoneTranslation.cpp
@@ -2,6 +2,7 @@
 
 #include "RigUnit_SetBoneTranslation.h"
 #include "Units/RigUnitContext.h"
+#include "Units/Hierarchy/RigUnit_SetTransform.h"
 
 FRigUnit_SetBoneTranslation_Execute()
 {
@@ -75,3 +76,15 @@ FRigUnit_SetBoneTranslation_Execute()
 		}
 	}
 }
+
+FRigVMStructUpgradeInfo FRigUnit_SetBoneTranslation::GetUpgradeInfo() const
+{
+	FRigUnit_SetTranslation NewNode;
+	NewNode.Item = FRigElementKey(Bone, ERigElementType::Bone);
+	NewNode.Space = Space;
+	NewNode.Translation = Translation;
+
+	FRigVMStructUpgradeInfo Info(*this, NewNode);
+	Info.AddRemappedPin(TEXT("Bone"), TEXT("Item.Name"));
+	return Info;
+}
diff --git a/Engine/Plugins/Animation/ControlRig/Source/ControlRig/Private/Units/Hierarchy/RigUnit_SetBoneTranslation.h b/Engine/Plugins/Animation/ControlRig/Source/ControlRig/Private/Units/Hierarchy/RigUnit_SetBoneTranslation.h
--- a/Engine/Plugins/Animation/ControlRig/Source/ControlRig/Private/Units/Hierarchy/RigUnit_SetBoneTranslation.h
+++ b/Engine/Plugins/Animation/ControlRig/Source/ControlRig/Private/Units/Hierarchy/RigUnit_SetBoneTranslation.h
@@ -73,4 +73,7 @@ struct CONTROLRIG_API FRigUnit_SetBoneTranslation : public FRigUnitMutable
 	// Used to cache the internally used bone index
 	UPROPERTY(transient)
 	FCachedRigElement CachedBone;
+
+	RIGVM_METHOD()
+	virtual FRigVMStructUpgradeInfo GetUpgradeInfo() const override;
 };
diff --git a/Engine/Plugins/Animation/ControlRig/Source/ControlRig/Private/Units/Hierarchy/RigUnit_SetControlTransform.cpp b/Engine/Plugins/Animation/ControlRig/Source/ControlRig/Private/Units/Hierarchy/RigUnit_SetControlTransform.cpp
--- a/Engine/Plugins/Animation/ControlRig/Source/ControlRig/Private/Units/Hierarchy/RigUnit_SetControlTransform.cpp
+++ b/Engine/Plugins/Animation/ControlRig/Source/ControlRig/Private/Units/Hierarchy/RigUnit_SetControlTransform.cpp
@@ -3,6 +3,7 @@
 #include "RigUnit_SetControlTransform.h"
 #include "Units/RigUnitContext.h"
 #include "Math/ControlRigMathLibrary.h"
+#include "Units/Hierarchy/RigUnit_SetTransform.h"
 
 FRigUnit_SetControlBool_Execute()
 {
@@ -491,6 +492,18 @@ FRigUnit_SetControlTransform_Execute()
 	}
 }
 
+FRigVMStructUpgradeInfo FRigUnit_SetControlTransform::GetUpgradeInfo() const
+{
+	FRigUnit_SetTransform NewNode;
+	NewNode.Item = FRigElementKey(Control, ERigElementType::Control);
+	NewNode.Space = Space;
+	NewNode.Transform = Transform;
+
+	FRigVMStructUpgradeInfo Info(*this, NewNode);
+	Info.AddRemappedPin(TEXT("Control"), TEXT("Item.Name"));
+	return Info;
+}
+
 #if WITH_DEV_AUTOMATION_TESTS
 #include "Units/RigUnitTest.h"
 #include "Units/Math/RigUnit_MathTransform.h"
diff --git a/Engine/Plugins/Animation/ControlRig/Source/ControlRig/Private/Units/Hierarchy/RigUnit_SetControlTransform.h b/Engine/Plugins/Animation/ControlRig/Source/ControlRig/Private/Units/Hierarchy/RigUnit_SetControlTransform.h
--- a/Engine/Plugins/Animation/ControlRig/Source/ControlRig/Private/Units/Hierarchy/RigUnit_SetControlTransform.h
+++ b/Engine/Plugins/Animation/ControlRig/Source/ControlRig/Private/Units/Hierarchy/RigUnit_SetControlTransform.h
@@ -579,4 +579,7 @@ struct CONTROLRIG_API FRigUnit_SetControlTransform : public FRigUnitMutable
 	// Used to cache the internally used bone index
 	UPROPERTY()
 	FCachedRigElement CachedControlIndex;
+
+	RIGVM_METHOD()
+	virtual FRigVMStructUpgradeInfo GetUpgradeInfo() const override;
 };
diff --git a/Engine/Plugins/Animation/ControlRig/Source/ControlRig/Private/Units/Hierarchy/RigUnit_SetRelativeBoneTransform.cpp b/Engine/Plugins/Animation/ControlRig/Source/ControlRig/Private/Units/Hierarchy/RigUnit_SetRelativeBoneTransform.cpp
--- a/Engine/Plugins/Animation/ControlRig/Source/ControlRig/Private/Units/Hierarchy/RigUnit_SetRelativeBoneTransform.cpp
+++ b/Engine/Plugins/Animation/ControlRig/Source/ControlRig/Private/Units/Hierarchy/RigUnit_SetRelativeBoneTransform.cpp
@@ -2,6 +2,7 @@
 
 #include "RigUnit_SetRelativeBoneTransform.h"
 #include "Units/RigUnitContext.h"
+#include "Units/Hierarchy/RigUnit_SetRelativeTransform.h"
 
 FRigUnit_SetRelativeBoneTransform_Execute()
 {
@@ -52,6 +53,20 @@ FRigUnit_SetRelativeBoneTransform_Execute()
 	}
 }
 
+FRigVMStructUpgradeInfo FRigUnit_SetRelativeBoneTransform::GetUpgradeInfo() const
+{
+	FRigUnit_SetRelativeTransformForItem NewNode;
+	NewNode.Parent = FRigElementKey(Space, ERigElementType::Bone);
+	NewNode.Child = FRigElementKey(Bone, ERigElementType::Bone);
+	NewNode.RelativeTransform = Transform;
+
+	FRigVMStructUpgradeInfo Info(*this, NewNode);
+	Info.AddRemappedPin(TEXT("Space"), TEXT("Parent.Name"));
+	Info.AddRemappedPin(TEXT("Bone"), TEXT("Child.Name"));
+	Info.AddRemappedPin(TEXT("RelativeTransform"), TEXT("Transform"));
+	return Info;
+}
+
 #if WITH_DEV_AUTOMATION_TESTS
 #include "Units/RigUnitTest.h"
 
diff --git a/Engine/Plugins/Animation/ControlRig/Source/ControlRig/Private/Units/Hierarchy/RigUnit_SetRelativeBoneTransform.h b/Engine/Plugins/Animation/ControlRig/Source/ControlRig/Private/Units/Hierarchy/RigUnit_SetRelativeBoneTransform.h
--- a/Engine/Plugins/Animation/ControlRig/Source/ControlRig/Private/Units/Hierarchy/RigUnit_SetRelativeBoneTransform.h
+++ b/Engine/Plugins/Animation/ControlRig/Source/ControlRig/Private/Units/Hierarchy/RigUnit_SetRelativeBoneTransform.h
@@ -72,4 +72,7 @@ struct CONTROLRIG_API FRigUnit_SetRelativeBoneTransform : public FRigUnitMutable
 	// Used to cache the internally used space index
 	UPROPERTY(transient)
 	FCachedRigElement CachedSpaceIndex;
+
+	RIGVM_METHOD()
+	virtual FRigVMStructUpgradeInfo GetUpgradeInfo() const override;
 };
diff --git a/Engine/Plugins/Animation/ControlRig/Source/ControlRig/Private/Units/Hierarchy/RigUnit_SetSpaceInitialTransform.cpp b/Engine/Plugins/Animation/ControlRig/Source/ControlRig/Private/Units/Hierarchy/RigUnit_SetSpaceInitialTransform.cpp
--- a/Engine/Plugins/Animation/ControlRig/Source/ControlRig/Private/Units/Hierarchy/RigUnit_SetSpaceInitialTransform.cpp
+++ b/Engine/Plugins/Animation/ControlRig/Source/ControlRig/Private/Units/Hierarchy/RigUnit_SetSpaceInitialTransform.cpp
@@ -3,6 +3,7 @@
 #include "RigUnit_SetSpaceInitialTransform.h"
 #include "Units/RigUnitContext.h"
 #include "Units/Execution/RigUnit_PrepareForExecution.h"
+#include "Units/Hierarchy/RigUnit_SetTransform.h"
 
 FRigUnit_SetSpaceInitialTransform_Execute()
 {
@@ -43,3 +44,16 @@ FRigUnit_SetSpaceInitialTransform_Execute()
 		}
 	}
 }
+
+FRigVMStructUpgradeInfo FRigUnit_SetSpaceInitialTransform::GetUpgradeInfo() const
+{
+	FRigUnit_SetTransform NewNode;
+	NewNode.Item = FRigElementKey(SpaceName, ERigElementType::Null);
+	NewNode.Space = Space;
+	NewNode.Transform = Transform;
+	NewNode.bInitial = true;
+
+	FRigVMStructUpgradeInfo Info(*this, NewNode);
+	Info.AddRemappedPin(TEXT("SpaceName"), TEXT("Item.Name"));
+	return Info;
+}
diff --git a/Engine/Plugins/Animation/ControlRig/Source/ControlRig/Private/Units/Hierarchy/RigUnit_SetSpaceInitialTransform.h b/Engine/Plugins/Animation/ControlRig/Source/ControlRig/Private/Units/Hierarchy/RigUnit_SetSpaceInitialTransform.h
--- a/Engine/Plugins/Animation/ControlRig/Source/ControlRig/Private/Units/Hierarchy/RigUnit_SetSpaceInitialTransform.h
+++ b/Engine/Plugins/Animation/ControlRig/Source/ControlRig/Private/Units/Hierarchy/RigUnit_SetSpaceInitialTransform.h
@@ -49,4 +49,7 @@ struct CONTROLRIG_API FRigUnit_SetSpaceInitialTransform : public FRigUnitMutable
 	// user to internally cache the index of the bone
 	UPROPERTY()
 	FCachedRigElement CachedSpaceIndex;
+
+	RIGVM_METHOD()
+	virtual FRigVMStructUpgradeInfo GetUpgradeInfo() const override;
 };
diff --git a/Engine/Plugins/Animation/ControlRig/Source/ControlRig/Private/Units/Hierarchy/RigUnit_SetSpaceTransform.cpp b/Engine/Plugins/Animation/ControlRig/Source/ControlRig/Private/Units/Hierarchy/RigUnit_SetSpaceTransform.cpp
--- a/Engine/Plugins/Animation/ControlRig/Source/ControlRig/Private/Units/Hierarchy/RigUnit_SetSpaceTransform.cpp
+++ b/Engine/Plugins/Animation/ControlRig/Source/ControlRig/Private/Units/Hierarchy/RigUnit_SetSpaceTransform.cpp
@@ -3,6 +3,7 @@
 #include "RigUnit_SetSpaceTransform.h"
 #include "Units/RigUnitContext.h"
 #include "Math/ControlRigMathLibrary.h"
+#include "Units/Hierarchy/RigUnit_SetTransform.h"
 
 FRigUnit_SetSpaceTransform_Execute()
 {
@@ -65,3 +66,16 @@ FRigUnit_SetSpaceTransform_Execute()
 		}
 	}
 }
+
+FRigVMStructUpgradeInfo FRigUnit_SetSpaceTransform::GetUpgradeInfo() const
+{
+	FRigUnit_SetTransform NewNode;
+	NewNode.Item = FRigElementKey(Space, ERigElementType::Null);
+	NewNode.Space = SpaceType;
+	NewNode.Transform = Transform;
+
+	FRigVMStructUpgradeInfo Info(*this, NewNode);
+	Info.AddRemappedPin(TEXT("Space"), TEXT("Item.Name"));
+	Info.AddRemappedPin(TEXT("SpaceType"), TEXT("Space"));
+	return Info;
+}
diff --git a/Engine/Plugins/Animation/ControlRig/Source/ControlRig/Private/Units/Hierarchy/RigUnit_SetSpaceTransform.h b/Engine/Plugins/Animation/ControlRig/Source/ControlRig/Private/Units/Hierarchy/RigUnit_SetSpaceTransform.h
--- a/Engine/Plugins/Animation/ControlRig/Source/ControlRig/Private/Units/Hierarchy/RigUnit_SetSpaceTransform.h
+++ b/Engine/Plugins/Animation/ControlRig/Source/ControlRig/Private/Units/Hierarchy/RigUnit_SetSpaceTransform.h
@@ -51,4 +51,7 @@ struct CONTROLRIG_API FRigUnit_SetSpaceTransform : public FRigUnitMutable
 	// Used to cache the internally used bone index
 	UPROPERTY()
 	FCachedRigElement CachedSpaceIndex;
+
+	RIGVM_METHOD()
+	virtual FRigVMStructUpgradeInfo GetUpgradeInfo() const override;
 };
diff --git a/Engine/Plugins/Animation/ControlRig/Source/ControlRig/Private/Units/Hierarchy/RigUnit_SetTransform.cpp b/Engine/Plugins/Animation/ControlRig/Source/ControlRig/Private/Units/Hierarchy/RigUnit_SetTransform.cpp
--- a/Engine/Plugins/Animation/ControlRig/Source/ControlRig/Private/Units/Hierarchy/RigUnit_SetTransform.cpp
+++ b/Engine/Plugins/Animation/ControlRig/Source/ControlRig/Private/Units/Hierarchy/RigUnit_SetTransform.cpp
@@ -164,6 +164,19 @@ FRigUnit_SetTransformArray_Execute()
 	FRigUnit_SetTransformItemArray::StaticExecute(RigVMExecuteContext, Items.Keys, Space, bInitial, Transforms, Weight, bPropagateToChildren, CachedIndex, ExecuteContext, Context);
 }
 
+FRigVMStructUpgradeInfo FRigUnit_SetTransformArray::GetUpgradeInfo() const
+{
+	FRigUnit_SetTransformItemArray NewNode;
+	NewNode.Items = Items.GetKeys();
+	NewNode.Space = Space;
+	NewNode.bInitial = bInitial;
+	NewNode.Transforms = Transforms;
+	NewNode.Weight = Weight;
+	NewNode.bPropagateToChildren = bPropagateToChildren;
+
+	return FRigVMStructUpgradeInfo(*this, NewNode);
+}
+
 FRigUnit_SetTransformItemArray_Execute()
 {
 	DECLARE_SCOPE_HIERARCHICAL_COUNTER_RIGUNIT()
diff --git a/Engine/Plugins/Animation/ControlRig/Source/ControlRig/Private/Units/Highlevel/Harmonics/RigUnit_BoneHarmonics.cpp b/Engine/Plugins/Animation/ControlRig/Source/ControlRig/Private/Units/Highlevel/Harmonics/RigUnit_BoneHarmonics.cpp
--- a/Engine/Plugins/Animation/ControlRig/Source/ControlRig/Private/Units/Highlevel/Harmonics/RigUnit_BoneHarmonics.cpp
+++ b/Engine/Plugins/Animation/ControlRig/Source/ControlRig/Private/Units/Highlevel/Harmonics/RigUnit_BoneHarmonics.cpp
@@ -33,6 +33,38 @@ FRigUnit_BoneHarmonics_Execute()
 		Context);
 }
 
+FRigVMStructUpgradeInfo FRigUnit_BoneHarmonics::GetUpgradeInfo() const
+{
+	FRigUnit_ItemHarmonics NewNode;
+	
+	for(int32 BoneIndex = 0;BoneIndex<Bones.Num();BoneIndex++)
+	{
+		FRigUnit_Harmonics_TargetItem Target;
+		Target.Item = FRigElementKey(Bones[BoneIndex].Bone, ERigElementType::Bone);
+		Target.Ratio = Bones[BoneIndex].Ratio;
+		NewNode.Targets.Add(Target);
+	}
+
+	NewNode.WaveSpeed = WaveSpeed;
+	NewNode.WaveFrequency = WaveFrequency;
+	NewNode.WaveAmplitude = WaveAmplitude;
+	NewNode.WaveOffset = WaveOffset;
+	NewNode.WaveNoise = WaveNoise;
+	NewNode.WaveEase = WaveEase;
+	NewNode.WaveMinimum = WaveMinimum;
+	NewNode.WaveMaximum = WaveMaximum;
+	NewNode.RotationOrder = RotationOrder;
+
+	FRigVMStructUpgradeInfo Info(*this, NewNode);
+	for(int32 BoneIndex = 0;BoneIndex<Bones.Num();BoneIndex++)
+	{
+		Info.AddRemappedPin(
+			FString::Printf(TEXT("Bones.%d.Bone"), BoneIndex),
+			FString::Printf(TEXT("Targets.%d.Item.Name"), BoneIndex));
+	}
+	return Info;
+}
+
 FRigUnit_ItemHarmonics_Execute()
 {
     DECLARE_SCOPE_HIERARCHICAL_COUNTER_RIGUNIT()
diff --git a/Engine/Plugins/Animation/ControlRig/Source/ControlRig/Private/Units/Highlevel/Harmonics/RigUnit_BoneHarmonics.h b/Engine/Plugins/Animation/ControlRig/Source/ControlRig/Private/Units/Highlevel/Harmonics/RigUnit_BoneHarmonics.h
--- a/Engine/Plugins/Animation/ControlRig/Source/ControlRig/Private/Units/Highlevel/Harmonics/RigUnit_BoneHarmonics.h
+++ b/Engine/Plugins/Animation/ControlRig/Source/ControlRig/Private/Units/Highlevel/Harmonics/RigUnit_BoneHarmonics.h
@@ -140,6 +140,9 @@ struct CONTROLRIG_API FRigUnit_BoneHarmonics : public FRigUnit_HighlevelBaseMuta
 
 	UPROPERTY(transient)
 	FRigUnit_BoneHarmonics_WorkData WorkData;
+
+	RIGVM_METHOD()
+	virtual FRigVMStructUpgradeInfo GetUpgradeInfo() const override;
 };
 
 /**
diff --git a/Engine/Plugins/Animation/ControlRig/Source/ControlRig/Private/Units/Highlevel/Harmonics/RigUnit_ChainHarmonics.cpp b/Engine/Plugins/Animation/ControlRig/Source/ControlRig/Private/Units/Highlevel/Harmonics/RigUnit_ChainHarmonics.cpp
--- a/Engine/Plugins/Animation/ControlRig/Source/ControlRig/Private/Units/Highlevel/Harmonics/RigUnit_ChainHarmonics.cpp
+++ b/Engine/Plugins/Animation/ControlRig/Source/ControlRig/Private/Units/Highlevel/Harmonics/RigUnit_ChainHarmonics.cpp
@@ -20,6 +20,23 @@ FRigUnit_ChainHarmonics_Execute()
 		Context);
 }
 
+FRigVMStructUpgradeInfo FRigUnit_ChainHarmonics::GetUpgradeInfo() const
+{
+	FRigUnit_ChainHarmonicsPerItem NewNode;
+	NewNode.ChainRoot = FRigElementKey(ChainRoot, ERigElementType::Bone);
+	NewNode.Speed = Speed;
+	NewNode.Reach = Reach;
+	NewNode.Wave = Wave;
+	NewNode.WaveCurve = WaveCurve;
+	NewNode.Pendulum = Pendulum;
+	NewNode.bDrawDebug = bDrawDebug;
+	NewNode.DrawWorldOffset = DrawWorldOffset;
+
+	FRigVMStructUpgradeInfo Info(*this, NewNode);
+	Info.AddRemappedPin(TEXT("ChainRoot"), TEXT("ChainRoot.Name"));
+	return Info;
+}
+
 FRigUnit_ChainHarmonicsPerItem_Execute()
 {
     DECLARE_SCOPE_HIERARCHICAL_COUNTER_RIGUNIT()
diff --git a/Engine/Plugins/Animation/ControlRig/Source/ControlRig/Private/Units/Highlevel/Harmonics/RigUnit_ChainHarmonics.h b/Engine/Plugins/Animation/ControlRig/Source/ControlRig/Private/Units/Highlevel/Harmonics/RigUnit_ChainHarmonics.h
--- a/Engine/Plugins/Animation/ControlRig/Source/ControlRig/Private/Units/Highlevel/Harmonics/RigUnit_ChainHarmonics.h
+++ b/Engine/Plugins/Animation/ControlRig/Source/ControlRig/Private/Units/Highlevel/Harmonics/RigUnit_ChainHarmonics.h
@@ -226,6 +226,9 @@ struct CONTROLRIG_API FRigUnit_ChainHarmonics : public FRigUnit_HighlevelBaseMut
 
 	UPROPERTY(transient)
 	FRigUnit_ChainHarmonics_WorkData WorkData;
+
+	RIGVM_METHOD()
+	virtual FRigVMStructUpgradeInfo GetUpgradeInfo() const override;
 };
 
 /**
diff --git a/Engine/Plugins/Animation/ControlRig/Source/ControlRig/Private/Units/Highlevel/Hierarchy/RigUnit_AimBone.cpp b/Engine/Plugins/Animation/ControlRig/Source/ControlRig/Private/Units/Highlevel/Hierarchy/RigUnit_AimBone.cpp
--- a/Engine/Plugins/Animation/ControlRig/Source/ControlRig/Private/Units/Highlevel/Hierarchy/RigUnit_AimBone.cpp
+++ b/Engine/Plugins/Animation/ControlRig/Source/ControlRig/Private/Units/Highlevel/Hierarchy/RigUnit_AimBone.cpp
@@ -173,6 +173,30 @@ FRigUnit_AimBone_Execute()
 		Context);
 }
 
+FRigVMStructUpgradeInfo FRigUnit_AimBone::GetUpgradeInfo() const
+{
+	FRigUnit_AimItem NewNode;
+	NewNode.Item = FRigElementKey(Bone, ERigElementType::Bone);
+	NewNode.Primary.Weight = Primary.Weight;
+	NewNode.Primary.Axis = Primary.Axis;
+	NewNode.Primary.Target = Primary.Target;
+	NewNode.Primary.Kind = Primary.Kind;
+	NewNode.Primary.Space = FRigElementKey(Primary.Space, ERigElementType::Bone);
+	NewNode.Secondary.Weight = Secondary.Weight;
+	NewNode.Secondary.Axis = Secondary.Axis;
+	NewNode.Secondary.Target = Secondary.Target;
+	NewNode.Secondary.Kind = Secondary.Kind;
+	NewNode.Secondary.Space = FRigElementKey(Secondary.Space, ERigElementType::Bone);
+	NewNode.Weight = Weight;
+	NewNode.DebugSettings = DebugSettings;
+
+	FRigVMStructUpgradeInfo Info(*this, NewNode);
+	Info.AddRemappedPin(TEXT("Bone"), TEXT("Item.Name"));
+	Info.AddRemappedPin(TEXT("Primary.Space"), TEXT("Primary.Space.Name"));
+	Info.AddRemappedPin(TEXT("Secondary.Space"), TEXT("Secondary.Space.Name"));
+	return Info;
+}
+
 FRigUnit_AimItem_Execute()
 {
     DECLARE_SCOPE_HIERARCHICAL_COUNTER_RIGUNIT()
@@ -216,4 +240,5 @@ FRigUnit_AimItem_Execute()
 		Context);
 
 	Hierarchy->SetGlobalTransform(CachedItem, Transform);
-}
\ No newline at end of file
+}
+
diff --git a/Engine/Plugins/Animation/ControlRig/Source/ControlRig/Private/Units/Highlevel/Hierarchy/RigUnit_AimBone.h b/Engine/Plugins/Animation/ControlRig/Source/ControlRig/Private/Units/Highlevel/Hierarchy/RigUnit_AimBone.h
--- a/Engine/Plugins/Animation/ControlRig/Source/ControlRig/Private/Units/Highlevel/Hierarchy/RigUnit_AimBone.h
+++ b/Engine/Plugins/Animation/ControlRig/Source/ControlRig/Private/Units/Highlevel/Hierarchy/RigUnit_AimBone.h
@@ -276,6 +276,9 @@ struct CONTROLRIG_API FRigUnit_AimBone : public FRigUnit_HighlevelBaseMutable
 
 	UPROPERTY()
 	FCachedRigElement SecondaryCachedSpace;
+
+	RIGVM_METHOD()
+	virtual FRigVMStructUpgradeInfo GetUpgradeInfo() const override;
 };
 
 /**
diff --git a/Engine/Plugins/Animation/ControlRig/Source/ControlRig/Private/Units/Highlevel/Hierarchy/RigUnit_CCDIK.cpp b/Engine/Plugins/Animation/ControlRig/Source/ControlRig/Private/Units/Highlevel/Hierarchy/RigUnit_CCDIK.cpp
--- a/Engine/Plugins/Animation/ControlRig/Source/ControlRig/Private/Units/Highlevel/Hierarchy/RigUnit_CCDIK.cpp
+++ b/Engine/Plugins/Animation/ControlRig/Source/ControlRig/Private/Units/Highlevel/Hierarchy/RigUnit_CCDIK.cpp
@@ -48,11 +48,33 @@ FRigUnit_CCDIK_Execute()
 		Context);
 }
 
+FRigVMStructUpgradeInfo FRigUnit_CCDIK::GetUpgradeInfo() const
+{
+	// this node is no longer supported and the upgrade path is too complex.
+	return FRigVMStructUpgradeInfo();
+}
+
 FRigUnit_CCDIKPerItem_Execute()
 {
 	FRigUnit_CCDIKItemArray::StaticExecute(RigVMExecuteContext, Items.Keys, EffectorTransform, Precision, Weight, MaxIterations, bStartFromTail, BaseRotationLimit, RotationLimits, bPropagateToChildren, WorkData, ExecuteContext, Context);
 }
 
+FRigVMStructUpgradeInfo FRigUnit_CCDIKPerItem::GetUpgradeInfo() const
+{
+	FRigUnit_CCDIKItemArray NewNode;
+	NewNode.Items = Items.Keys;
+	NewNode.EffectorTransform = EffectorTransform;
+	NewNode.Precision = Precision;
+	NewNode.Weight = Weight;
+	NewNode.MaxIterations = MaxIterations;
+	NewNode.bStartFromTail = bStartFromTail;
+	NewNode.BaseRotationLimit = BaseRotationLimit;
+	NewNode.RotationLimits = RotationLimits;
+	NewNode.bPropagateToChildren = bPropagateToChildren;
+
+	return FRigVMStructUpgradeInfo(*this, NewNode);
+}
+
 FRigUnit_CCDIKItemArray_Execute()
 {
     DECLARE_SCOPE_HIERARCHICAL_COUNTER_RIGUNIT()
diff --git a/Engine/Plugins/Animation/ControlRig/Source/ControlRig/Private/Units/Highlevel/Hierarchy/RigUnit_CCDIK.h b/Engine/Plugins/Animation/ControlRig/Source/ControlRig/Private/Units/Highlevel/Hierarchy/RigUnit_CCDIK.h
--- a/Engine/Plugins/Animation/ControlRig/Source/ControlRig/Private/Units/Highlevel/Hierarchy/RigUnit_CCDIK.h
+++ b/Engine/Plugins/Animation/ControlRig/Source/ControlRig/Private/Units/Highlevel/Hierarchy/RigUnit_CCDIK.h
@@ -168,6 +168,9 @@ struct CONTROLRIG_API FRigUnit_CCDIK : public FRigUnit_HighlevelBaseMutable
 
 	UPROPERTY()
 	FRigUnit_CCDIK_WorkData WorkData;
+
+	RIGVM_METHOD()
+	virtual FRigVMStructUpgradeInfo GetUpgradeInfo() const override;
 };
 
 /**
@@ -252,6 +255,9 @@ struct CONTROLRIG_API FRigUnit_CCDIKPerItem : public FRigUnit_HighlevelBaseMutab
 
 	UPROPERTY()
 	FRigUnit_CCDIK_WorkData WorkData;
+
+	RIGVM_METHOD()
+	virtual FRigVMStructUpgradeInfo GetUpgradeInfo() const override;
 };
 
 /**
diff --git a/Engine/Plugins/Animation/ControlRig/Source/ControlRig/Private/Units/Highlevel/Hierarchy/RigUnit_DistributeRotation.cpp b/Engine/Plugins/Animation/ControlRig/Source/ControlRig/Private/Units/Highlevel/Hierarchy/RigUnit_DistributeRotation.cpp
--- a/Engine/Plugins/Animation/ControlRig/Source/ControlRig/Private/Units/Highlevel/Hierarchy/RigUnit_DistributeRotation.cpp
+++ b/Engine/Plugins/Animation/ControlRig/Source/ControlRig/Private/Units/Highlevel/Hierarchy/RigUnit_DistributeRotation.cpp
@@ -57,11 +57,28 @@ FRigUnit_DistributeRotation_Execute()
 		Context);
 }
 
+FRigVMStructUpgradeInfo FRigUnit_DistributeRotation::GetUpgradeInfo() const
+{
+	// this node is no longer supported and the upgrade path is too complex.
+	return FRigVMStructUpgradeInfo();
+}
+
 FRigUnit_DistributeRotationForCollection_Execute()
 {
 	FRigUnit_DistributeRotationForItemArray::StaticExecute(RigVMExecuteContext, Items.Keys, Rotations, RotationEaseType, Weight, WorkData, ExecuteContext, Context);
 }
 
+FRigVMStructUpgradeInfo FRigUnit_DistributeRotationForCollection::GetUpgradeInfo() const
+{
+	FRigUnit_DistributeRotationForItemArray NewNode;
+	NewNode.Items = Items.Keys;
+	NewNode.Rotations = Rotations;
+	NewNode.RotationEaseType = RotationEaseType;
+	NewNode.Weight = Weight;
+
+	return FRigVMStructUpgradeInfo(*this, NewNode);
+}
+
 FRigUnit_DistributeRotationForItemArray_Execute()
 {
     DECLARE_SCOPE_HIERARCHICAL_COUNTER_RIGUNIT()
diff --git a/Engine/Plugins/Animation/ControlRig/Source/ControlRig/Private/Units/Highlevel/Hierarchy/RigUnit_DistributeRotation.h b/Engine/Plugins/Animation/ControlRig/Source/ControlRig/Private/Units/Highlevel/Hierarchy/RigUnit_DistributeRotation.h
--- a/Engine/Plugins/Animation/ControlRig/Source/ControlRig/Private/Units/Highlevel/Hierarchy/RigUnit_DistributeRotation.h
+++ b/Engine/Plugins/Animation/ControlRig/Source/ControlRig/Private/Units/Highlevel/Hierarchy/RigUnit_DistributeRotation.h
@@ -121,6 +121,9 @@ struct CONTROLRIG_API FRigUnit_DistributeRotation : public FRigUnit_HighlevelBas
 
 	UPROPERTY(transient)
 	FRigUnit_DistributeRotation_WorkData WorkData;
+
+	RIGVM_METHOD()
+	virtual FRigVMStructUpgradeInfo GetUpgradeInfo() const override;
 };
 
 /**
@@ -168,6 +171,9 @@ struct CONTROLRIG_API FRigUnit_DistributeRotationForCollection : public FRigUnit
 
 	UPROPERTY(transient)
 	FRigUnit_DistributeRotation_WorkData WorkData;
+
+	RIGVM_METHOD()
+	virtual FRigVMStructUpgradeInfo GetUpgradeInfo() const override;
 };
 
 /**
diff --git a/Engine/Plugins/Animation/ControlRig/Source/ControlRig/Private/Units/Highlevel/Hierarchy/RigUnit_FABRIK.cpp b/Engine/Plugins/Animation/ControlRig/Source/ControlRig/Private/Units/Highlevel/Hierarchy/RigUnit_FABRIK.cpp
--- a/Engine/Plugins/Animation/ControlRig/Source/ControlRig/Private/Units/Highlevel/Hierarchy/RigUnit_FABRIK.cpp
+++ b/Engine/Plugins/Animation/ControlRig/Source/ControlRig/Private/Units/Highlevel/Hierarchy/RigUnit_FABRIK.cpp
@@ -36,11 +36,30 @@ FRigUnit_FABRIK_Execute()
 		Context);
 }
 
+FRigVMStructUpgradeInfo FRigUnit_FABRIK::GetUpgradeInfo() const
+{
+	// this node is no longer supported and the upgrade path is too complex.
+	return FRigVMStructUpgradeInfo();
+}
+
 FRigUnit_FABRIKPerItem_Execute()
 {
 	FRigUnit_FABRIKItemArray::StaticExecute(RigVMExecuteContext, Items.Keys, EffectorTransform, Precision, Weight, bPropagateToChildren, MaxIterations, WorkData, bSetEffectorTransform,ExecuteContext, Context);
 }
 
+FRigVMStructUpgradeInfo FRigUnit_FABRIKPerItem::GetUpgradeInfo() const
+{
+	FRigUnit_FABRIKItemArray NewNode;
+	NewNode.Items = Items.Keys;
+	NewNode.EffectorTransform = EffectorTransform;
+	NewNode.Precision = Precision;
+	NewNode.Weight = Weight;
+	NewNode.bPropagateToChildren = bPropagateToChildren;
+	NewNode.MaxIterations = MaxIterations;
+
+	return FRigVMStructUpgradeInfo(*this, NewNode);
+}
+
 FRigUnit_FABRIKItemArray_Execute()
 {
     DECLARE_SCOPE_HIERARCHICAL_COUNTER_RIGUNIT()
diff --git a/Engine/Plugins/Animation/ControlRig/Source/ControlRig/Private/Units/Highlevel/Hierarchy/RigUnit_FABRIK.h b/Engine/Plugins/Animation/ControlRig/Source/ControlRig/Private/Units/Highlevel/Hierarchy/RigUnit_FABRIK.h
--- a/Engine/Plugins/Animation/ControlRig/Source/ControlRig/Private/Units/Highlevel/Hierarchy/RigUnit_FABRIK.h
+++ b/Engine/Plugins/Animation/ControlRig/Source/ControlRig/Private/Units/Highlevel/Hierarchy/RigUnit_FABRIK.h
@@ -101,6 +101,9 @@ struct CONTROLRIG_API FRigUnit_FABRIK : public FRigUnit_HighlevelBaseMutable
 	*/
 	UPROPERTY(meta = (Input))
 	bool bSetEffectorTransform;
+
+	RIGVM_METHOD()
+	virtual FRigVMStructUpgradeInfo GetUpgradeInfo() const override;
 };
 
 /**
@@ -172,6 +175,9 @@ struct CONTROLRIG_API FRigUnit_FABRIKPerItem : public FRigUnit_HighlevelBaseMuta
 	*/
 	UPROPERTY(meta = (Input))
 	bool bSetEffectorTransform;
+
+	RIGVM_METHOD()
+	virtual FRigVMStructUpgradeInfo GetUpgradeInfo() const override;
 };
 
 /**
diff --git a/Engine/Plugins/Animation/ControlRig/Source/ControlRig/Private/Units/Highlevel/Hierarchy/RigUnit_FitChainToCurve.cpp b/Engine/Plugins/Animation/ControlRig/Source/ControlRig/Private/Units/Highlevel/Hierarchy/RigUnit_FitChainToCurve.cpp
--- a/Engine/Plugins/Animation/ControlRig/Source/ControlRig/Private/Units/Highlevel/Hierarchy/RigUnit_FitChainToCurve.cpp
+++ b/Engine/Plugins/Animation/ControlRig/Source/ControlRig/Private/Units/Highlevel/Hierarchy/RigUnit_FitChainToCurve.cpp
@@ -43,6 +43,12 @@ FRigUnit_FitChainToCurve_Execute()
 		Context);
 }
 
+FRigVMStructUpgradeInfo FRigUnit_FitChainToCurve::GetUpgradeInfo() const
+{
+	// this node is no longer supported and the upgrade path is too complex.
+	return FRigVMStructUpgradeInfo();
+}
+
 FRigUnit_FitChainToCurvePerItem_Execute()
 {
 	FRigUnit_FitChainToCurveItemArray::StaticExecute(
@@ -66,6 +72,12 @@ FRigUnit_FitChainToCurvePerItem_Execute()
 		Context);
 }
 
+FRigVMStructUpgradeInfo FRigUnit_FitChainToCurvePerItem::GetUpgradeInfo() const
+{
+	// this node is no longer supported
+	return FRigVMStructUpgradeInfo();
+}
+
 FRigUnit_FitChainToCurveItemArray_Execute()
 {
     DECLARE_SCOPE_HIERARCHICAL_COUNTER_RIGUNIT()
@@ -455,3 +467,9 @@ FRigUnit_FitChainToCurveItemArray_Execute()
 		// Context.DrawInterface->DrawPoints(DebugSettings.WorldOffset, ItemPositions, DebugSettings.Scale * 6.f, DebugSettings.SegmentsColor);
 	}
 }
+
+FRigVMStructUpgradeInfo FRigUnit_FitChainToCurveItemArray::GetUpgradeInfo() const
+{
+	// this node is no longer supported
+	return FRigVMStructUpgradeInfo();
+}
diff --git a/Engine/Plugins/Animation/ControlRig/Source/ControlRig/Private/Units/Highlevel/Hierarchy/RigUnit_FitChainToCurve.h b/Engine/Plugins/Animation/ControlRig/Source/ControlRig/Private/Units/Highlevel/Hierarchy/RigUnit_FitChainToCurve.h
--- a/Engine/Plugins/Animation/ControlRig/Source/ControlRig/Private/Units/Highlevel/Hierarchy/RigUnit_FitChainToCurve.h
+++ b/Engine/Plugins/Animation/ControlRig/Source/ControlRig/Private/Units/Highlevel/Hierarchy/RigUnit_FitChainToCurve.h
@@ -245,6 +245,9 @@ struct CONTROLRIG_API FRigUnit_FitChainToCurve : public FRigUnit_HighlevelBaseMu
 
 	UPROPERTY(transient)
 	FRigUnit_FitChainToCurve_WorkData WorkData;
+
+	RIGVM_METHOD()
+	virtual FRigVMStructUpgradeInfo GetUpgradeInfo() const override;
 };
 
 /**
@@ -362,6 +365,9 @@ struct CONTROLRIG_API FRigUnit_FitChainToCurvePerItem : public FRigUnit_Highleve
 
 	UPROPERTY(transient)
 	FRigUnit_FitChainToCurve_WorkData WorkData;
+
+	RIGVM_METHOD()
+	virtual FRigVMStructUpgradeInfo GetUpgradeInfo() const override;
 };
 
 /**
@@ -479,4 +485,7 @@ struct CONTROLRIG_API FRigUnit_FitChainToCurveItemArray : public FRigUnit_Highle
 
 	UPROPERTY(transient)
 	FRigUnit_FitChainToCurve_WorkData WorkData;
+
+	RIGVM_METHOD()
+	virtual FRigVMStructUpgradeInfo GetUpgradeInfo() const override;
 };
diff --git a/Engine/Plugins/Animation/ControlRig/Source/ControlRig/Private/Units/Highlevel/Hierarchy/RigUnit_ModifyBoneTransforms.cpp b/Engine/Plugins/Animation/ControlRig/Source/ControlRig/Private/Units/Highlevel/Hierarchy/RigUnit_ModifyBoneTransforms.cpp
--- a/Engine/Plugins/Animation/ControlRig/Source/ControlRig/Private/Units/Highlevel/Hierarchy/RigUnit_ModifyBoneTransforms.cpp
+++ b/Engine/Plugins/Animation/ControlRig/Source/ControlRig/Private/Units/Highlevel/Hierarchy/RigUnit_ModifyBoneTransforms.cpp
@@ -26,6 +26,12 @@ FRigUnit_ModifyBoneTransforms_Execute()
 		Context);
 }
 
+FRigVMStructUpgradeInfo FRigUnit_ModifyBoneTransforms::GetUpgradeInfo() const
+{
+	// this node is no longer supported
+	return FRigVMStructUpgradeInfo();
+}
+
 #if WITH_DEV_AUTOMATION_TESTS
 #include "Units/RigUnitTest.h"
 
diff --git a/Engine/Plugins/Animation/ControlRig/Source/ControlRig/Private/Units/Highlevel/Hierarchy/RigUnit_ModifyBoneTransforms.h b/Engine/Plugins/Animation/ControlRig/Source/ControlRig/Private/Units/Highlevel/Hierarchy/RigUnit_ModifyBoneTransforms.h
--- a/Engine/Plugins/Animation/ControlRig/Source/ControlRig/Private/Units/Highlevel/Hierarchy/RigUnit_ModifyBoneTransforms.h
+++ b/Engine/Plugins/Animation/ControlRig/Source/ControlRig/Private/Units/Highlevel/Hierarchy/RigUnit_ModifyBoneTransforms.h
@@ -123,4 +123,7 @@ struct CONTROLRIG_API FRigUnit_ModifyBoneTransforms : public FRigUnit_HighlevelB
 	// Used to cache the internally used bone index
 	UPROPERTY(transient)
 	FRigUnit_ModifyBoneTransforms_WorkData WorkData;
+
+	RIGVM_METHOD()
+	virtual FRigVMStructUpgradeInfo GetUpgradeInfo() const override;
 };
diff --git a/Engine/Plugins/Animation/ControlRig/Source/ControlRig/Private/Units/Highlevel/Hierarchy/RigUnit_SlideChain.cpp b/Engine/Plugins/Animation/ControlRig/Source/ControlRig/Private/Units/Highlevel/Hierarchy/RigUnit_SlideChain.cpp
--- a/Engine/Plugins/Animation/ControlRig/Source/ControlRig/Private/Units/Highlevel/Hierarchy/RigUnit_SlideChain.cpp
+++ b/Engine/Plugins/Animation/ControlRig/Source/ControlRig/Private/Units/Highlevel/Hierarchy/RigUnit_SlideChain.cpp
@@ -33,11 +33,27 @@ FRigUnit_SlideChain_Execute()
 		Context);
 }
 
+FRigVMStructUpgradeInfo FRigUnit_SlideChain::GetUpgradeInfo() const
+{
+	// this node is no longer supported and the upgrade path is too complex.
+	return FRigVMStructUpgradeInfo();
+}
+
 FRigUnit_SlideChainPerItem_Execute()
 {
 	FRigUnit_SlideChainItemArray::StaticExecute(RigVMExecuteContext, Items.Keys, SlideAmount, bPropagateToChildren, WorkData, ExecuteContext, Context);
 }
 
+FRigVMStructUpgradeInfo FRigUnit_SlideChainPerItem::GetUpgradeInfo() const
+{
+	FRigUnit_SlideChainItemArray NewNode;
+	NewNode.Items = Items.Keys;
+	NewNode.SlideAmount = SlideAmount;
+	NewNode.bPropagateToChildren = bPropagateToChildren;
+
+	return FRigVMStructUpgradeInfo(*this, NewNode);
+}
+
 FRigUnit_SlideChainItemArray_Execute()
 {
 	URigHierarchy* Hierarchy = ExecuteContext.Hierarchy;
diff --git a/Engine/Plugins/Animation/ControlRig/Source/ControlRig/Private/Units/Highlevel/Hierarchy/RigUnit_SlideChain.h b/Engine/Plugins/Animation/ControlRig/Source/ControlRig/Private/Units/Highlevel/Hierarchy/RigUnit_SlideChain.h
--- a/Engine/Plugins/Animation/ControlRig/Source/ControlRig/Private/Units/Highlevel/Hierarchy/RigUnit_SlideChain.h
+++ b/Engine/Plugins/Animation/ControlRig/Source/ControlRig/Private/Units/Highlevel/Hierarchy/RigUnit_SlideChain.h
@@ -78,6 +78,9 @@ struct CONTROLRIG_API FRigUnit_SlideChain: public FRigUnit_HighlevelBaseMutable
 
 	UPROPERTY(transient)
 	FRigUnit_SlideChain_WorkData WorkData;
+
+	RIGVM_METHOD()
+	virtual FRigVMStructUpgradeInfo GetUpgradeInfo() const override;
 };
 
 /**
@@ -119,6 +122,9 @@ struct CONTROLRIG_API FRigUnit_SlideChainPerItem: public FRigUnit_HighlevelBaseM
 
 	UPROPERTY(transient)
 	FRigUnit_SlideChain_WorkData WorkData;
+
+	RIGVM_METHOD()
+	virtual FRigVMStructUpgradeInfo GetUpgradeInfo() const override;
 };
 
 /**
diff --git a/Engine/Plugins/Animation/ControlRig/Source/ControlRig/Private/Units/Highlevel/Hierarchy/RigUnit_TransformConstraint.cpp b/Engine/Plugins/Animation/ControlRig/Source/ControlRig/Private/Units/Highlevel/Hierarchy/RigUnit_TransformConstraint.cpp
--- a/Engine/Plugins/Animation/ControlRig/Source/ControlRig/Private/Units/Highlevel/Hierarchy/RigUnit_TransformConstraint.cpp
+++ b/Engine/Plugins/Animation/ControlRig/Source/ControlRig/Private/Units/Highlevel/Hierarchy/RigUnit_TransformConstraint.cpp
@@ -20,6 +20,22 @@ FRigUnit_TransformConstraint_Execute()
 		Context);
 }
 
+FRigVMStructUpgradeInfo FRigUnit_TransformConstraint::GetUpgradeInfo() const
+{
+	FRigUnit_TransformConstraintPerItem NewNode;
+	NewNode.Item = FRigElementKey(Bone, ERigElementType::Bone);
+	NewNode.BaseTransformSpace = BaseTransformSpace;
+	NewNode.BaseTransform = BaseTransform;
+	NewNode.BaseItem = FRigElementKey(BaseBone, ERigElementType::Bone);
+	NewNode.Targets = Targets;
+	NewNode.bUseInitialTransforms = bUseInitialTransforms;
+
+	FRigVMStructUpgradeInfo Info(*this, NewNode);
+	Info.AddRemappedPin(TEXT("Bone"), TEXT("Item.Name"));
+	Info.AddRemappedPin(TEXT("BaseBone"), TEXT("BaseItem.Name"));
+	return Info;
+}
+
 FRigUnit_TransformConstraintPerItem_Execute()
 {
     DECLARE_SCOPE_HIERARCHICAL_COUNTER_RIGUNIT()
@@ -166,6 +182,12 @@ void FRigUnit_TransformConstraintPerItem::AddConstraintData(const TArrayView<con
 	OutConstraintDataToTargets.Add(NewIndex, TargetIndex);
 }
 
+FRigVMStructUpgradeInfo FRigUnit_TransformConstraintPerItem::GetUpgradeInfo() const
+{
+	// this node is no longer supported
+	return FRigVMStructUpgradeInfo();
+}
+
 FRigUnit_ParentConstraint_Execute()
  {
 	if (Weight < KINDA_SMALL_NUMBER)
@@ -537,6 +559,18 @@ FRigUnit_PositionConstraint_Execute()
 	}
 }
 
+FRigVMStructUpgradeInfo FRigUnit_PositionConstraint::GetUpgradeInfo() const
+{
+	FRigUnit_PositionConstraintLocalSpaceOffset NewNode;
+	NewNode.Child = Child;
+	NewNode.bMaintainOffset = bMaintainOffset;
+	NewNode.Filter = Filter;
+	NewNode.Parents = Parents;
+	NewNode.Weight = Weight;
+
+	return FRigVMStructUpgradeInfo(*this, NewNode);
+}
+
 #if WITH_DEV_AUTOMATION_TESTS
 #include "Units/RigUnitTest.h"
 
@@ -964,6 +998,19 @@ FRigUnit_RotationConstraint_Execute()
 	} 	
 }
 
+FRigVMStructUpgradeInfo FRigUnit_RotationConstraint::GetUpgradeInfo() const
+{
+	FRigUnit_RotationConstraintLocalSpaceOffset NewNode;
+	NewNode.Child = Child;
+	NewNode.bMaintainOffset = bMaintainOffset;
+	NewNode.Filter = Filter;
+	NewNode.Parents = Parents;
+	NewNode.AdvancedSettings = AdvancedSettings;
+	NewNode.Weight = Weight;
+
+	return FRigVMStructUpgradeInfo(*this, NewNode);
+}
+
 #if WITH_DEV_AUTOMATION_TESTS
 #include "Units/RigUnitTest.h"
 
@@ -1559,6 +1606,18 @@ FRigUnit_ScaleConstraint_Execute()
 	}
 }
 
+FRigVMStructUpgradeInfo FRigUnit_ScaleConstraint::GetUpgradeInfo() const
+{
+	FRigUnit_ScaleConstraintLocalSpaceOffset NewNode;
+	NewNode.Child = Child;
+	NewNode.bMaintainOffset = bMaintainOffset;
+	NewNode.Filter = Filter;
+	NewNode.Parents = Parents;
+	NewNode.Weight = Weight;
+
+	return FRigVMStructUpgradeInfo(*this, NewNode);
+}
+
 #if WITH_DEV_AUTOMATION_TESTS
 #include "Units/RigUnitTest.h"
 
diff --git a/Engine/Plugins/Animation/ControlRig/Source/ControlRig/Private/Units/Highlevel/Hierarchy/RigUnit_TransformConstraint.h b/Engine/Plugins/Animation/ControlRig/Source/ControlRig/Private/Units/Highlevel/Hierarchy/RigUnit_TransformConstraint.h
--- a/Engine/Plugins/Animation/ControlRig/Source/ControlRig/Private/Units/Highlevel/Hierarchy/RigUnit_TransformConstraint.h
+++ b/Engine/Plugins/Animation/ControlRig/Source/ControlRig/Private/Units/Highlevel/Hierarchy/RigUnit_TransformConstraint.h
@@ -109,6 +109,11 @@ private:
 
 	UPROPERTY(transient)
 	FRigUnit_TransformConstraint_WorkData WorkData;
+
+public:
+	
+	RIGVM_METHOD()
+	virtual FRigVMStructUpgradeInfo GetUpgradeInfo() const override;
 };
 
 /**
@@ -175,6 +180,11 @@ private:
 
 	UPROPERTY(transient)
 	FRigUnit_TransformConstraint_WorkData WorkData;
+
+public:
+
+	RIGVM_METHOD()
+	virtual FRigVMStructUpgradeInfo GetUpgradeInfo() const override;
 };
 
 USTRUCT()
@@ -307,6 +317,9 @@ struct FRigUnit_PositionConstraint : public FRigUnit_HighlevelBaseMutable
 
 	UPROPERTY(EditAnywhere, Category = "Constraint", meta = (Input))
 	float Weight;
+
+	RIGVM_METHOD()
+	virtual FRigVMStructUpgradeInfo GetUpgradeInfo() const override;
 };
 
 /**
@@ -404,6 +417,9 @@ struct FRigUnit_RotationConstraint : public FRigUnit_HighlevelBaseMutable
 
 	UPROPERTY(EditAnywhere, Category = "Constraint", meta = (Input))
 	float Weight;
+
+	RIGVM_METHOD()
+	virtual FRigVMStructUpgradeInfo GetUpgradeInfo() const override;
 };
 
 /**
@@ -477,6 +493,9 @@ struct FRigUnit_ScaleConstraint : public FRigUnit_HighlevelBaseMutable
 
 	UPROPERTY(EditAnywhere, Category = "Constraint", meta = (Input))
 	float Weight;
+
+	RIGVM_METHOD()
+	virtual FRigVMStructUpgradeInfo GetUpgradeInfo() const override;
 };
 
 /**
diff --git a/Engine/Plugins/Animation/ControlRig/Source/ControlRig/Private/Units/Highlevel/Hierarchy/RigUnit_TwistBones.cpp b/Engine/Plugins/Animation/ControlRig/Source/ControlRig/Private/Units/Highlevel/Hierarchy/RigUnit_TwistBones.cpp
--- a/Engine/Plugins/Animation/ControlRig/Source/ControlRig/Private/Units/Highlevel/Hierarchy/RigUnit_TwistBones.cpp
+++ b/Engine/Plugins/Animation/ControlRig/Source/ControlRig/Private/Units/Highlevel/Hierarchy/RigUnit_TwistBones.cpp
@@ -35,6 +35,12 @@ FRigUnit_TwistBones_Execute()
 		Context);
 }
 
+FRigVMStructUpgradeInfo FRigUnit_TwistBones::GetUpgradeInfo() const
+{
+	// this node is no longer supported
+	return FRigVMStructUpgradeInfo();
+}
+
 FRigUnit_TwistBonesPerItem_Execute()
 {
     DECLARE_SCOPE_HIERARCHICAL_COUNTER_RIGUNIT()
@@ -137,3 +143,9 @@ FRigUnit_TwistBonesPerItem_Execute()
 		Hierarchy->SetGlobalTransform(CachedItems[Index], ItemTransforms[Index], bPropagateToChildren);
 	}
 }
+
+FRigVMStructUpgradeInfo FRigUnit_TwistBonesPerItem::GetUpgradeInfo() const
+{
+	// this node is no longer supported
+	return FRigVMStructUpgradeInfo();
+}
diff --git a/Engine/Plugins/Animation/ControlRig/Source/ControlRig/Private/Units/Highlevel/Hierarchy/RigUnit_TwistBones.h b/Engine/Plugins/Animation/ControlRig/Source/ControlRig/Private/Units/Highlevel/Hierarchy/RigUnit_TwistBones.h
--- a/Engine/Plugins/Animation/ControlRig/Source/ControlRig/Private/Units/Highlevel/Hierarchy/RigUnit_TwistBones.h
+++ b/Engine/Plugins/Animation/ControlRig/Source/ControlRig/Private/Units/Highlevel/Hierarchy/RigUnit_TwistBones.h
@@ -94,6 +94,9 @@ struct CONTROLRIG_API FRigUnit_TwistBones : public FRigUnit_HighlevelBaseMutable
 
 	UPROPERTY(transient)
 	FRigUnit_TwistBones_WorkData WorkData;
+
+	RIGVM_METHOD()
+	virtual FRigVMStructUpgradeInfo GetUpgradeInfo() const override;
 };
 
 /**
@@ -156,4 +159,7 @@ struct CONTROLRIG_API FRigUnit_TwistBonesPerItem : public FRigUnit_HighlevelBase
 
 	UPROPERTY(transient)
 	FRigUnit_TwistBones_WorkData WorkData;
+
+	RIGVM_METHOD()
+	virtual FRigVMStructUpgradeInfo GetUpgradeInfo() const override;
 };
diff --git a/Engine/Plugins/Animation/ControlRig/Source/ControlRig/Private/Units/Highlevel/Hierarchy/RigUnit_TwoBoneIKSimple.cpp b/Engine/Plugins/Animation/ControlRig/Source/ControlRig/Private/Units/Highlevel/Hierarchy/RigUnit_TwoBoneIKSimple.cpp
--- a/Engine/Plugins/Animation/ControlRig/Source/ControlRig/Private/Units/Highlevel/Hierarchy/RigUnit_TwoBoneIKSimple.cpp
+++ b/Engine/Plugins/Animation/ControlRig/Source/ControlRig/Private/Units/Highlevel/Hierarchy/RigUnit_TwoBoneIKSimple.cpp
@@ -35,6 +35,37 @@ FRigUnit_TwoBoneIKSimple_Execute()
 		Context);
 }
 
+FRigVMStructUpgradeInfo FRigUnit_TwoBoneIKSimple::GetUpgradeInfo() const
+{
+	FRigUnit_TwoBoneIKSimplePerItem NewNode;
+	NewNode.ItemA = FRigElementKey(BoneA, ERigElementType::Bone);
+	NewNode.ItemA = FRigElementKey(BoneB, ERigElementType::Bone);
+	NewNode.EffectorItem = FRigElementKey(EffectorBone, ERigElementType::Bone);
+	NewNode.Effector = Effector;
+	NewNode.PrimaryAxis = PrimaryAxis;
+	NewNode.SecondaryAxis = SecondaryAxis;
+	NewNode.SecondaryAxisWeight = SecondaryAxisWeight;
+	NewNode.PoleVector = PoleVector;
+	NewNode.PoleVectorKind = PoleVectorKind;
+	NewNode.PoleVectorSpace = FRigElementKey(PoleVectorSpace, ERigElementType::Bone);
+	NewNode.bEnableStretch = bEnableStretch;
+	NewNode.StretchStartRatio = StretchStartRatio;
+	NewNode.StretchMaximumRatio = StretchMaximumRatio;
+	NewNode.Weight = Weight;
+	NewNode.ItemALength = BoneALength;
+	NewNode.ItemBLength = BoneBLength;
+	NewNode.bPropagateToChildren = bPropagateToChildren;
+	NewNode.DebugSettings = DebugSettings;
+	
+	FRigVMStructUpgradeInfo Info(*this, NewNode);
+	Info.AddRemappedPin(TEXT("BoneA"), TEXT("ItemA.Name"));
+	Info.AddRemappedPin(TEXT("BoneB"), TEXT("ItemB.Name"));
+	Info.AddRemappedPin(TEXT("EffectorBone"), TEXT("EffectorItem.Name"));
+	Info.AddRemappedPin(TEXT("BoneALength"), TEXT("ItemALength"));
+	Info.AddRemappedPin(TEXT("BoneBLength"), TEXT("ItemBLength"));
+	return Info;
+}
+
 FRigUnit_TwoBoneIKSimplePerItem_Execute()
 {
     DECLARE_SCOPE_HIERARCHICAL_COUNTER_RIGUNIT()
diff --git a/Engine/Plugins/Animation/ControlRig/Source/ControlRig/Private/Units/Highlevel/Hierarchy/RigUnit_TwoBoneIKSimple.h b/Engine/Plugins/Animation/ControlRig/Source/ControlRig/Private/Units/Highlevel/Hierarchy/RigUnit_TwoBoneIKSimple.h
--- a/Engine/Plugins/Animation/ControlRig/Source/ControlRig/Private/Units/Highlevel/Hierarchy/RigUnit_TwoBoneIKSimple.h
+++ b/Engine/Plugins/Animation/ControlRig/Source/ControlRig/Private/Units/Highlevel/Hierarchy/RigUnit_TwoBoneIKSimple.h
@@ -201,6 +201,9 @@ struct CONTROLRIG_API FRigUnit_TwoBoneIKSimple : public FRigUnit_HighlevelBaseMu
 	FCachedRigElement CachedEffectorBoneIndex;
 	UPROPERTY()
 	FCachedRigElement CachedPoleVectorSpaceIndex;
+
+	RIGVM_METHOD()
+	virtual FRigVMStructUpgradeInfo GetUpgradeInfo() const override;
 };
 
 /**
diff --git a/Engine/Plugins/Animation/ControlRig/Source/ControlRig/Private/Units/Math/RigUnit_MathFloat.cpp b/Engine/Plugins/Animation/ControlRig/Source/ControlRig/Private/Units/Math/RigUnit_MathFloat.cpp
--- a/Engine/Plugins/Animation/ControlRig/Source/ControlRig/Private/Units/Math/RigUnit_MathFloat.cpp
+++ b/Engine/Plugins/Animation/ControlRig/Source/ControlRig/Private/Units/Math/RigUnit_MathFloat.cpp
@@ -233,6 +233,12 @@ FRigUnit_MathFloatSelectBool_Execute()
 	Result = Condition ? IfTrue : IfFalse;
 }
 
+FRigVMStructUpgradeInfo FRigUnit_MathFloatSelectBool::GetUpgradeInfo() const
+{
+	// this node is no longer supported
+	return FRigVMStructUpgradeInfo();
+}
+
 FRigUnit_MathFloatDeg_Execute()
 {
     DECLARE_SCOPE_HIERARCHICAL_COUNTER_RIGUNIT()
diff --git a/Engine/Plugins/Animation/ControlRig/Source/ControlRig/Private/Units/Math/RigUnit_MathFloat.h b/Engine/Plugins/Animation/ControlRig/Source/ControlRig/Private/Units/Math/RigUnit_MathFloat.h
--- a/Engine/Plugins/Animation/ControlRig/Source/ControlRig/Private/Units/Math/RigUnit_MathFloat.h
+++ b/Engine/Plugins/Animation/ControlRig/Source/ControlRig/Private/Units/Math/RigUnit_MathFloat.h
@@ -65,7 +65,7 @@ struct CONTROLRIG_API FRigUnit_MathFloatBinaryOp : public FRigUnit_MathFloatBase
 /**
  * Returns PI
  */
-USTRUCT(meta=(DisplayName="Pi"))
+USTRUCT(meta=(DisplayName="Pi", PrototypeName="Pi"))
 struct CONTROLRIG_API FRigUnit_MathFloatConstPi : public FRigUnit_MathFloatConstant
 {
 	GENERATED_BODY()
@@ -77,7 +77,7 @@ struct CONTROLRIG_API FRigUnit_MathFloatConstPi : public FRigUnit_MathFloatConst
 /**
  * Returns PI * 0.5
  */
-USTRUCT(meta=(DisplayName="Half Pi"))
+USTRUCT(meta=(DisplayName="Half Pi", PrototypeName="HalfPi"))
 struct CONTROLRIG_API FRigUnit_MathFloatConstHalfPi : public FRigUnit_MathFloatConstant
 {
 	GENERATED_BODY()
@@ -89,7 +89,7 @@ struct CONTROLRIG_API FRigUnit_MathFloatConstHalfPi : public FRigUnit_MathFloatC
 /**
  * Returns PI * 2.0
  */
-USTRUCT(meta=(DisplayName="Two Pi", Keywords="Tau"))
+USTRUCT(meta=(DisplayName="Two Pi", PrototypeName="TwoPi", Keywords="Tau"))
 struct CONTROLRIG_API FRigUnit_MathFloatConstTwoPi : public FRigUnit_MathFloatConstant
 {
 	GENERATED_BODY()
@@ -101,7 +101,7 @@ struct CONTROLRIG_API FRigUnit_MathFloatConstTwoPi : public FRigUnit_MathFloatCo
 /**
  * Returns E
  */
-USTRUCT(meta=(DisplayName="E", Keywords="Euler"))
+USTRUCT(meta=(DisplayName="E", PrototypeName="E", Keywords="Euler"))
 struct CONTROLRIG_API FRigUnit_MathFloatConstE : public FRigUnit_MathFloatConstant
 {
 	GENERATED_BODY()
@@ -727,6 +727,9 @@ struct CONTROLRIG_API FRigUnit_MathFloatSelectBool : public FRigUnit_MathFloatBa
 
 	UPROPERTY(meta=(Output))
 	float Result;
+
+	RIGVM_METHOD()
+	virtual FRigVMStructUpgradeInfo GetUpgradeInfo() const override;
 };
 
 /**
@@ -828,7 +831,7 @@ struct CONTROLRIG_API FRigUnit_MathFloatAtan : public FRigUnit_MathFloatUnaryOp
 /**
  * Computes the angles alpha, beta and gamma (in radians) between the three sides A, B and C
  */
-USTRUCT(meta = (DisplayName = "Law Of Cosine"))
+USTRUCT(meta = (DisplayName = "Law Of Cosine", PrototypeName="LawOfCosine"))
 struct CONTROLRIG_API FRigUnit_MathFloatLawOfCosine : public FRigUnit_MathFloatBase
 {
 	GENERATED_BODY()
@@ -867,7 +870,7 @@ struct CONTROLRIG_API FRigUnit_MathFloatLawOfCosine : public FRigUnit_MathFloatB
 /**
  * Computes the base-e exponential of the given value 
  */
-USTRUCT(meta = (DisplayName = "Exponential"))
+USTRUCT(meta = (DisplayName = "Exponential", PrototypeName="Exponential"))
 struct CONTROLRIG_API FRigUnit_MathFloatExponential : public FRigUnit_MathFloatUnaryOp
 {
 	GENERATED_BODY()
diff --git a/Engine/Plugins/Animation/ControlRig/Source/ControlRig/Private/Units/Math/RigUnit_MathQuaternion.cpp b/Engine/Plugins/Animation/ControlRig/Source/ControlRig/Private/Units/Math/RigUnit_MathQuaternion.cpp
--- a/Engine/Plugins/Animation/ControlRig/Source/ControlRig/Private/Units/Math/RigUnit_MathQuaternion.cpp
+++ b/Engine/Plugins/Animation/ControlRig/Source/ControlRig/Private/Units/Math/RigUnit_MathQuaternion.cpp
@@ -123,6 +123,12 @@ FRigUnit_MathQuaternionSelectBool_Execute()
 	Result = Condition ? IfTrue : IfFalse;
 }
 
+FRigVMStructUpgradeInfo FRigUnit_MathQuaternionSelectBool::GetUpgradeInfo() const
+{
+	// this node is no longer supported
+	return FRigVMStructUpgradeInfo();
+}
+
 FRigUnit_MathQuaternionDot_Execute()
 {
     DECLARE_SCOPE_HIERARCHICAL_COUNTER_RIGUNIT()
diff --git a/Engine/Plugins/Animation/ControlRig/Source/ControlRig/Private/Units/Math/RigUnit_MathQuaternion.h b/Engine/Plugins/Animation/ControlRig/Source/ControlRig/Private/Units/Math/RigUnit_MathQuaternion.h
--- a/Engine/Plugins/Animation/ControlRig/Source/ControlRig/Private/Units/Math/RigUnit_MathQuaternion.h
+++ b/Engine/Plugins/Animation/ControlRig/Source/ControlRig/Private/Units/Math/RigUnit_MathQuaternion.h
@@ -394,6 +394,9 @@ struct CONTROLRIG_API FRigUnit_MathQuaternionSelectBool : public FRigUnit_MathQu
 
 	UPROPERTY(meta=(Output))
 	FQuat Result;
+
+	RIGVM_METHOD()
+	virtual FRigVMStructUpgradeInfo GetUpgradeInfo() const override;
 };
 
 /**
diff --git a/Engine/Plugins/Animation/ControlRig/Source/ControlRig/Private/Units/Math/RigUnit_MathTransform.cpp b/Engine/Plugins/Animation/ControlRig/Source/ControlRig/Private/Units/Math/RigUnit_MathTransform.cpp
--- a/Engine/Plugins/Animation/ControlRig/Source/ControlRig/Private/Units/Math/RigUnit_MathTransform.cpp
+++ b/Engine/Plugins/Animation/ControlRig/Source/ControlRig/Private/Units/Math/RigUnit_MathTransform.cpp
@@ -119,6 +119,12 @@ FRigUnit_MathTransformSelectBool_Execute()
 	Result = Condition ? IfTrue : IfFalse;
 }
 
+FRigVMStructUpgradeInfo FRigUnit_MathTransformSelectBool::GetUpgradeInfo() const
+{
+	// this node is no longer supported
+	return FRigVMStructUpgradeInfo();
+}
+
 FRigUnit_MathTransformRotateVector_Execute()
 {
     DECLARE_SCOPE_HIERARCHICAL_COUNTER_RIGUNIT()
diff --git a/Engine/Plugins/Animation/ControlRig/Source/ControlRig/Private/Units/Math/RigUnit_MathTransform.h b/Engine/Plugins/Animation/ControlRig/Source/ControlRig/Private/Units/Math/RigUnit_MathTransform.h
--- a/Engine/Plugins/Animation/ControlRig/Source/ControlRig/Private/Units/Math/RigUnit_MathTransform.h
+++ b/Engine/Plugins/Animation/ControlRig/Source/ControlRig/Private/Units/Math/RigUnit_MathTransform.h
@@ -280,6 +280,9 @@ struct CONTROLRIG_API FRigUnit_MathTransformSelectBool : public FRigUnit_MathTra
 
 	UPROPERTY(meta=(Output))
 	FTransform Result;
+
+	RIGVM_METHOD()
+	virtual FRigVMStructUpgradeInfo GetUpgradeInfo() const override;
 };
 
 /**
diff --git a/Engine/Plugins/Animation/ControlRig/Source/ControlRig/Private/Units/Math/RigUnit_MathVector.cpp b/Engine/Plugins/Animation/ControlRig/Source/ControlRig/Private/Units/Math/RigUnit_MathVector.cpp
--- a/Engine/Plugins/Animation/ControlRig/Source/ControlRig/Private/Units/Math/RigUnit_MathVector.cpp
+++ b/Engine/Plugins/Animation/ControlRig/Source/ControlRig/Private/Units/Math/RigUnit_MathVector.cpp
@@ -227,6 +227,12 @@ FRigUnit_MathVectorSelectBool_Execute()
 	Result = Condition ? IfTrue : IfFalse;
 }
 
+FRigVMStructUpgradeInfo FRigUnit_MathVectorSelectBool::GetUpgradeInfo() const
+{
+	// this node is no longer supported
+	return FRigVMStructUpgradeInfo();
+}
+
 FRigUnit_MathVectorDeg_Execute()
 {
     DECLARE_SCOPE_HIERARCHICAL_COUNTER_RIGUNIT()
@@ -381,10 +387,22 @@ FRigUnit_MathVectorBezierFourPoint_Execute()
 	FControlRigMathLibrary::FourPointBezier(Bezier, T, Result, Tangent);
 }
 
+FRigVMStructUpgradeInfo FRigUnit_MathVectorBezierFourPoint::GetUpgradeInfo() const
+{
+	// this node is no longer supported
+	return FRigVMStructUpgradeInfo();
+}
+
 FRigUnit_MathVectorMakeBezierFourPoint_Execute()
 {
 }
 
+FRigVMStructUpgradeInfo FRigUnit_MathVectorMakeBezierFourPoint::GetUpgradeInfo() const
+{
+	// this node is no longer supported
+	return FRigVMStructUpgradeInfo();
+}
+
 FRigUnit_MathVectorClampSpatially_Execute()
 {
 	DECLARE_SCOPE_HIERARCHICAL_COUNTER_RIGUNIT()
diff --git a/Engine/Plugins/Animation/ControlRig/Source/ControlRig/Private/Units/Math/RigUnit_MathVector.h b/Engine/Plugins/Animation/ControlRig/Source/ControlRig/Private/Units/Math/RigUnit_MathVector.h
--- a/Engine/Plugins/Animation/ControlRig/Source/ControlRig/Private/Units/Math/RigUnit_MathVector.h
+++ b/Engine/Plugins/Animation/ControlRig/Source/ControlRig/Private/Units/Math/RigUnit_MathVector.h
@@ -514,6 +514,9 @@ struct CONTROLRIG_API FRigUnit_MathVectorSelectBool : public FRigUnit_MathVector
 
 	UPROPERTY(meta=(Output))
 	FVector Result;
+
+	RIGVM_METHOD()
+	virtual FRigVMStructUpgradeInfo GetUpgradeInfo() const override;
 };
 
 /**
@@ -862,6 +865,9 @@ struct CONTROLRIG_API FRigUnit_MathVectorBezierFourPoint : public FRigUnit_MathV
 
 	UPROPERTY(meta=(Output))
 	FVector Tangent;
+
+	RIGVM_METHOD()
+	virtual FRigVMStructUpgradeInfo GetUpgradeInfo() const override;
 };
 
 /**
@@ -882,6 +888,9 @@ struct CONTROLRIG_API FRigUnit_MathVectorMakeBezierFourPoint : public FRigUnit_M
 
 	UPROPERTY(meta = (Input, Output))
 	FCRFourPointBezier Bezier;
+
+	RIGVM_METHOD()
+	virtual FRigVMStructUpgradeInfo GetUpgradeInfo() const override;
 };
 
 /**
diff --git a/Engine/Plugins/Animation/ControlRig/Source/ControlRig/Private/Units/Math/RigUnit_Noise.cpp b/Engine/Plugins/Animation/ControlRig/Source/ControlRig/Private/Units/Math/RigUnit_Noise.cpp
--- a/Engine/Plugins/Animation/ControlRig/Source/ControlRig/Private/Units/Math/RigUnit_Noise.cpp
+++ b/Engine/Plugins/Animation/ControlRig/Source/ControlRig/Private/Units/Math/RigUnit_Noise.cpp
@@ -12,25 +12,58 @@ FRigUnit_NoiseFloat_Execute()
 		return;
 	}
 
-	float Noise = FMath::PerlinNoise1D(Value * Frequency + Time) + 0.5f;
+	const float Noise = FMath::PerlinNoise1D(Value * Frequency + Time) + 0.5f;
 	Result = FMath::Lerp<float>(Minimum, Maximum, Noise);
 	Time = Time + Speed * Context.DeltaTime;
 }
 
+FRigUnit_NoiseDouble_Execute()
+{
+	DECLARE_SCOPE_HIERARCHICAL_COUNTER_RIGUNIT()
+	if (Context.State == EControlRigState::Init)
+	{
+		Time = 0.f;
+		return;
+	}
+
+	const double Noise = double(FMath::PerlinNoise1D(Value * Frequency + Time)) + 0.5;
+	Result = FMath::Lerp<double>(Minimum, Maximum, Noise);
+	Time = Time + Speed * Context.DeltaTime;
+}
+
 FRigUnit_NoiseVector_Execute()
 {
-    DECLARE_SCOPE_HIERARCHICAL_COUNTER_RIGUNIT()
+	FRigUnit_NoiseVector2::StaticExecute(RigVMExecuteContext, Position, Speed, Frequency, Minimum, Maximum, Result, Time, Context);
+}
+
+FRigVMStructUpgradeInfo FRigUnit_NoiseVector::GetUpgradeInfo() const
+{
+	FRigUnit_NoiseVector2 NewNode;
+	NewNode.Value = Position;
+	NewNode.Speed = Speed;
+	NewNode.Frequency = Frequency;
+	NewNode.Minimum = Minimum;
+	NewNode.Maximum = Maximum;
+
+	FRigVMStructUpgradeInfo Info(*this, NewNode);
+	Info.AddRemappedPin(TEXT("Position"), TEXT("Value"));
+	return Info;
+}
+
+FRigUnit_NoiseVector2_Execute()
+{
+	DECLARE_SCOPE_HIERARCHICAL_COUNTER_RIGUNIT()
 	if (Context.State == EControlRigState::Init)
 	{
 		Time = FVector::ZeroVector;
 		return;
 	}
 
-	float NoiseX = FMath::PerlinNoise1D(Position.X * Frequency.X + Time.X) + 0.5f;
-	float NoiseY = FMath::PerlinNoise1D(Position.Y * Frequency.Y + Time.Y) + 0.5f;
-	float NoiseZ = FMath::PerlinNoise1D(Position.Z * Frequency.Z + Time.Z) + 0.5f;
-	Result.X = FMath::Lerp<float>(Minimum, Maximum, NoiseX);
-	Result.Y = FMath::Lerp<float>(Minimum, Maximum, NoiseY);
-	Result.Z = FMath::Lerp<float>(Minimum, Maximum, NoiseZ);
+	const double NoiseX = double(FMath::PerlinNoise1D(Value.X * Frequency.X + Time.X)) + 0.5;
+	const double NoiseY = double(FMath::PerlinNoise1D(Value.Y * Frequency.Y + Time.Y)) + 0.5;
+	const double NoiseZ = double(FMath::PerlinNoise1D(Value.Z * Frequency.Z + Time.Z)) + 0.5;
+	Result.X = FMath::Lerp<double>(Minimum, Maximum, NoiseX);
+	Result.Y = FMath::Lerp<double>(Minimum, Maximum, NoiseY);
+	Result.Z = FMath::Lerp<double>(Minimum, Maximum, NoiseZ);
 	Time = Time + Speed * Context.DeltaTime;
 }
diff --git a/Engine/Plugins/Animation/ControlRig/Source/ControlRig/Private/Units/Math/RigUnit_Noise.h b/Engine/Plugins/Animation/ControlRig/Source/ControlRig/Private/Units/Math/RigUnit_Noise.h
--- a/Engine/Plugins/Animation/ControlRig/Source/ControlRig/Private/Units/Math/RigUnit_Noise.h
+++ b/Engine/Plugins/Animation/ControlRig/Source/ControlRig/Private/Units/Math/RigUnit_Noise.h
@@ -45,10 +45,50 @@ struct CONTROLRIG_API FRigUnit_NoiseFloat : public FRigUnit_MathBase
 	float Time;
 };
 
+/**
+ * Generates a float through a noise fluctuation process between a min and a max through speed
+ */
+USTRUCT(meta=(DisplayName="Noise (Double)", Category="Math|Noise", PrototypeName="Noise"))
+struct CONTROLRIG_API FRigUnit_NoiseDouble : public FRigUnit_MathBase
+{
+	GENERATED_BODY()
+
+	FRigUnit_NoiseDouble()
+	{
+		Value = Minimum = Result = Time = 0.0;
+		Speed = 0.1;
+		Frequency = Maximum = 1.0;
+	}
+
+	RIGVM_METHOD()
+	virtual void Execute(const FRigUnitContext& Context) override;
+
+	UPROPERTY(meta = (Input))
+	double Value;
+
+	UPROPERTY(meta = (Input))
+	double Speed;
+
+	UPROPERTY(meta = (Input))
+	double Frequency;
+
+	UPROPERTY(meta = (Input))
+	double Minimum;
+
+	UPROPERTY(meta = (Input))
+	double Maximum;
+
+	UPROPERTY(meta = (Output))
+	double Result;
+
+	UPROPERTY()
+	double Time;
+};
+
 /**
  * Generates a vector through a noise fluctuation process between a min and a max through speed
  */
-USTRUCT(meta = (DisplayName = "Noise (Vector)", Category = "Math|Noise", PrototypeName = "Noise"))
+USTRUCT(meta = (DisplayName = "Noise (Vector)", Category = "Math|Noise", PrototypeName = "Noise", Deprecated = "5.0.0"))
 struct CONTROLRIG_API FRigUnit_NoiseVector : public FRigUnit_MathBase
 {
 	GENERATED_BODY()
@@ -85,4 +125,49 @@ struct CONTROLRIG_API FRigUnit_NoiseVector : public FRigUnit_MathBase
 
 	UPROPERTY()
 	FVector Time;
+
+	RIGVM_METHOD()
+	virtual FRigVMStructUpgradeInfo GetUpgradeInfo() const override;
+};
+
+/**
+ * Generates a vector through a noise fluctuation process between a min and a max through speed
+ */
+USTRUCT(meta = (DisplayName = "Noise (Vector)", Category = "Math|Noise", PrototypeName = "Noise"))
+struct CONTROLRIG_API FRigUnit_NoiseVector2 : public FRigUnit_MathBase
+{
+	GENERATED_BODY()
+
+	FRigUnit_NoiseVector2()
+	{
+		Value = Result = Time = FVector::ZeroVector;
+		Frequency = FVector::OneVector;
+		Speed = FVector(0.1f, 0.1f, 0.1f);
+		Minimum = 0.0;
+		Maximum = 1.0;
+	}
+
+	RIGVM_METHOD()
+	virtual void Execute(const FRigUnitContext& Context) override;
+
+	UPROPERTY(meta = (Input))
+	FVector Value;
+
+	UPROPERTY(meta = (Input))
+	FVector Speed;
+
+	UPROPERTY(meta = (Input))
+	FVector Frequency;
+
+	UPROPERTY(meta = (Input))
+	double Minimum;
+
+	UPROPERTY(meta = (Input))
+	double Maximum;
+
+	UPROPERTY(meta = (Output))
+	FVector Result;
+
+	UPROPERTY()
+	FVector Time;
 };
\ No newline at end of file
diff --git a/Engine/Plugins/Animation/ControlRig/Source/ControlRig/Private/Units/Simulation/RigUnit_PointSimulation.cpp b/Engine/Plugins/Animation/ControlRig/Source/ControlRig/Private/Units/Simulation/RigUnit_PointSimulation.cpp
--- a/Engine/Plugins/Animation/ControlRig/Source/ControlRig/Private/Units/Simulation/RigUnit_PointSimulation.cpp
+++ b/Engine/Plugins/Animation/ControlRig/Source/ControlRig/Private/Units/Simulation/RigUnit_PointSimulation.cpp
@@ -254,3 +254,9 @@ FRigUnit_PointSimulation_Execute()
 		Context.DrawInterface->DrawPointSimulation(DebugSettings.WorldOffset, Simulation, DebugSettings.Color, DebugSettings.Scale, DebugSettings.CollisionScale, DebugSettings.bDrawPointsAsSpheres);
 	}
 }
+
+FRigVMStructUpgradeInfo FRigUnit_PointSimulation::GetUpgradeInfo() const
+{
+	// this node is no longer supported
+	return FRigVMStructUpgradeInfo();
+}
diff --git a/Engine/Plugins/Animation/ControlRig/Source/ControlRig/Private/Units/Simulation/RigUnit_PointSimulation.h b/Engine/Plugins/Animation/ControlRig/Source/ControlRig/Private/Units/Simulation/RigUnit_PointSimulation.h
--- a/Engine/Plugins/Animation/ControlRig/Source/ControlRig/Private/Units/Simulation/RigUnit_PointSimulation.h
+++ b/Engine/Plugins/Animation/ControlRig/Source/ControlRig/Private/Units/Simulation/RigUnit_PointSimulation.h
@@ -203,5 +203,8 @@ struct CONTROLRIG_API FRigUnit_PointSimulation : public FRigUnit_SimBaseMutable
 
 	UPROPERTY(transient)
 	FRigUnit_PointSimulation_WorkData WorkData;
+
+	RIGVM_METHOD()
+	virtual FRigVMStructUpgradeInfo GetUpgradeInfo() const override;
 };
 
diff --git a/Engine/Plugins/Animation/ControlRig/Source/ControlRig/Private/Units/Simulation/RigUnit_SpringInterp.cpp b/Engine/Plugins/Animation/ControlRig/Source/ControlRig/Private/Units/Simulation/RigUnit_SpringInterp.cpp
--- a/Engine/Plugins/Animation/ControlRig/Source/ControlRig/Private/Units/Simulation/RigUnit_SpringInterp.cpp
+++ b/Engine/Plugins/Animation/ControlRig/Source/ControlRig/Private/Units/Simulation/RigUnit_SpringInterp.cpp
@@ -33,6 +33,12 @@ FRigUnit_SpringInterp_Execute()
 		Result = UKismetMathLibrary::FloatSpringInterp(Result, Target, SpringState, Stiffness, CriticalDamping, RemainingTime, Mass);
 	}
 }
+
+FRigVMStructUpgradeInfo FRigUnit_SpringInterp::GetUpgradeInfo() const
+{
+	// this node is no longer supported and the new implementation is vastly different
+	return FRigVMStructUpgradeInfo();
+}
  
 FRigUnit_SpringInterpVector_Execute()
 {
@@ -58,6 +64,12 @@ FRigUnit_SpringInterpVector_Execute()
 	}
 }
 
+FRigVMStructUpgradeInfo FRigUnit_SpringInterpVector::GetUpgradeInfo() const
+{
+	// this node is no longer supported and the new implementation is vastly different
+	return FRigVMStructUpgradeInfo();
+}
+
 FRigUnit_SpringInterpV2_Execute()
 {
 	DECLARE_SCOPE_HIERARCHICAL_COUNTER_RIGUNIT()
diff --git a/Engine/Plugins/Animation/ControlRig/Source/ControlRig/Private/Units/Simulation/RigUnit_SpringInterp.h b/Engine/Plugins/Animation/ControlRig/Source/ControlRig/Private/Units/Simulation/RigUnit_SpringInterp.h
--- a/Engine/Plugins/Animation/ControlRig/Source/ControlRig/Private/Units/Simulation/RigUnit_SpringInterp.h
+++ b/Engine/Plugins/Animation/ControlRig/Source/ControlRig/Private/Units/Simulation/RigUnit_SpringInterp.h
@@ -46,6 +46,9 @@ struct CONTROLRIG_API FRigUnit_SpringInterp : public FRigUnit_SimBase
  
 	UPROPERTY()
 	FFloatSpringState SpringState;
+
+	RIGVM_METHOD()
+	virtual FRigVMStructUpgradeInfo GetUpgradeInfo() const override;
 };
  
 /**
@@ -88,6 +91,9 @@ struct CONTROLRIG_API FRigUnit_SpringInterpVector : public FRigUnit_SimBase
  
 	UPROPERTY()
 	FVectorSpringState SpringState;
+
+	RIGVM_METHOD()
+	virtual FRigVMStructUpgradeInfo GetUpgradeInfo() const override;
 };
 
 /**
diff --git a/Engine/Plugins/Animation/ControlRig/Source/ControlRig/Public/Units/Control/RigUnit_Control.h b/Engine/Plugins/Animation/ControlRig/Source/ControlRig/Public/Units/Control/RigUnit_Control.h
--- a/Engine/Plugins/Animation/ControlRig/Source/ControlRig/Public/Units/Control/RigUnit_Control.h
+++ b/Engine/Plugins/Animation/ControlRig/Source/ControlRig/Public/Units/Control/RigUnit_Control.h
@@ -62,4 +62,7 @@ struct CONTROLRIG_API FRigUnit_Control : public FRigUnit
 	/** The filter determines what axes can be manipulated by the in-viewport widgets */
 	UPROPERTY(meta=(Input))
  	FTransformFilter Filter;
+
+	RIGVM_METHOD()
+	virtual FRigVMStructUpgradeInfo GetUpgradeInfo() const override;
 };
\ No newline at end of file
diff --git a/Engine/Plugins/Animation/ControlRig/Source/ControlRig/Public/Units/Control/RigUnit_Control_StaticMesh.h b/Engine/Plugins/Animation/ControlRig/Source/ControlRig/Public/Units/Control/RigUnit_Control_StaticMesh.h
--- a/Engine/Plugins/Animation/ControlRig/Source/ControlRig/Public/Units/Control/RigUnit_Control_StaticMesh.h
+++ b/Engine/Plugins/Animation/ControlRig/Source/ControlRig/Public/Units/Control/RigUnit_Control_StaticMesh.h
@@ -22,4 +22,7 @@ struct CONTROLRIG_API FRigUnit_Control_StaticMesh : public FRigUnit_Control
 	/** The the transform the mesh will be rendered with (applied on top of the control's transform in the viewport) */
 	UPROPERTY(meta=(Input))
 	FTransform MeshTransform;
+
+	RIGVM_METHOD()
+	virtual FRigVMStructUpgradeInfo GetUpgradeInfo() const override;
 };
\ No newline at end of file
diff --git a/Engine/Plugins/Animation/ControlRig/Source/ControlRig/Public/Units/Deprecated/RigUnit_TwoBoneIKFK.h b/Engine/Plugins/Animation/ControlRig/Source/ControlRig/Public/Units/Deprecated/RigUnit_TwoBoneIKFK.h
--- a/Engine/Plugins/Animation/ControlRig/Source/ControlRig/Public/Units/Deprecated/RigUnit_TwoBoneIKFK.h
+++ b/Engine/Plugins/Animation/ControlRig/Source/ControlRig/Public/Units/Deprecated/RigUnit_TwoBoneIKFK.h
@@ -93,4 +93,7 @@ private:
 	
 	UPROPERTY(transient)
 	float LowerLimbLength;
+
+	RIGVM_METHOD()
+	virtual FRigVMStructUpgradeInfo GetUpgradeInfo() const override;
 };
diff --git a/Engine/Plugins/Animation/ControlRig/Source/ControlRig/Public/Units/Hierarchy/RigUnit_AddBoneTransform.h b/Engine/Plugins/Animation/ControlRig/Source/ControlRig/Public/Units/Hierarchy/RigUnit_AddBoneTransform.h
--- a/Engine/Plugins/Animation/ControlRig/Source/ControlRig/Public/Units/Hierarchy/RigUnit_AddBoneTransform.h
+++ b/Engine/Plugins/Animation/ControlRig/Source/ControlRig/Public/Units/Hierarchy/RigUnit_AddBoneTransform.h
@@ -70,4 +70,7 @@ struct FRigUnit_AddBoneTransform : public FRigUnitMutable
 	// Used to cache the internally used bone index
 	UPROPERTY(transient)
 	FCachedRigElement CachedBone;
+
+	RIGVM_METHOD()
+	virtual FRigVMStructUpgradeInfo GetUpgradeInfo() const override;
 };
diff --git a/Engine/Plugins/Animation/ControlRig/Source/ControlRig/Public/Units/Hierarchy/RigUnit_SetTransform.h b/Engine/Plugins/Animation/ControlRig/Source/ControlRig/Public/Units/Hierarchy/RigUnit_SetTransform.h
--- a/Engine/Plugins/Animation/ControlRig/Source/ControlRig/Public/Units/Hierarchy/RigUnit_SetTransform.h
+++ b/Engine/Plugins/Animation/ControlRig/Source/ControlRig/Public/Units/Hierarchy/RigUnit_SetTransform.h
@@ -322,6 +322,9 @@ struct CONTROLRIG_API FRigUnit_SetTransformArray : public FRigUnitMutable
 	// Used to cache the internally
 	UPROPERTY()
 	TArray<FCachedRigElement> CachedIndex;
+	
+	RIGVM_METHOD()
+	virtual FRigVMStructUpgradeInfo GetUpgradeInfo() const override;
 };
 
 /**
diff --git a/Engine/Plugins/Animation/ControlRigSpline/Source/ControlRigSpline/Private/ControlRigSplineUnits.cpp b/Engine/Plugins/Animation/ControlRigSpline/Source/ControlRigSpline/Private/ControlRigSplineUnits.cpp
--- a/Engine/Plugins/Animation/ControlRigSpline/Source/ControlRigSpline/Private/ControlRigSplineUnits.cpp
+++ b/Engine/Plugins/Animation/ControlRigSpline/Source/ControlRigSpline/Private/ControlRigSplineUnits.cpp
@@ -252,6 +252,28 @@ FRigUnit_FitChainToSplineCurve_Execute()
 		Context);
 }
 
+FRigVMStructUpgradeInfo FRigUnit_FitChainToSplineCurve::GetUpgradeInfo() const
+{
+	FRigUnit_FitChainToSplineCurveItemArray NewNode;
+	NewNode.Items = Items.GetKeys();
+	NewNode.Spline = Spline;
+	NewNode.Alignment = Alignment;
+	NewNode.Minimum = Minimum;
+	NewNode.Maximum = Maximum;
+	NewNode.SamplingPrecision = SamplingPrecision;
+	NewNode.PrimaryAxis = PrimaryAxis;
+	NewNode.SecondaryAxis = SecondaryAxis;
+	NewNode.PoleVectorPosition = PoleVectorPosition;
+	NewNode.Rotations = Rotations;
+	NewNode.RotationEaseType = RotationEaseType;
+	NewNode.Weight = Weight;
+	NewNode.bPropagateToChildren = bPropagateToChildren;
+	NewNode.DebugSettings = DebugSettings;
+	NewNode.WorkData = WorkData;
+
+	return FRigVMStructUpgradeInfo(*this, NewNode);
+}
+
 FRigUnit_FitChainToSplineCurveItemArray_Execute()
 {
     DECLARE_SCOPE_HIERARCHICAL_COUNTER_RIGUNIT()
@@ -664,6 +686,15 @@ FRigUnit_FitSplineCurveToChain_Execute()
 	FRigUnit_FitSplineCurveToChainItemArray::StaticExecute(RigVMExecuteContext, Items.Keys, Spline, ExecuteContext, Context);
 }
 
+FRigVMStructUpgradeInfo FRigUnit_FitSplineCurveToChain::GetUpgradeInfo() const
+{
+	FRigUnit_FitSplineCurveToChainItemArray NewNode;
+	NewNode.Items = Items.GetKeys();
+	NewNode.Spline = Spline;
+
+	return FRigVMStructUpgradeInfo(*this, NewNode);
+}
+
 FRigUnit_FitSplineCurveToChainItemArray_Execute()
 {
     DECLARE_SCOPE_HIERARCHICAL_COUNTER_RIGUNIT()
diff --git a/Engine/Plugins/Animation/ControlRigSpline/Source/ControlRigSpline/Public/ControlRigSplineUnits.h b/Engine/Plugins/Animation/ControlRigSpline/Source/ControlRigSpline/Public/ControlRigSplineUnits.h
--- a/Engine/Plugins/Animation/ControlRigSpline/Source/ControlRigSpline/Public/ControlRigSplineUnits.h
+++ b/Engine/Plugins/Animation/ControlRigSpline/Source/ControlRigSpline/Public/ControlRigSplineUnits.h
@@ -339,6 +339,9 @@ struct CONTROLRIGSPLINE_API FRigUnit_FitChainToSplineCurve : public FRigUnit_Hig
 
 	UPROPERTY(transient)
 	FRigUnit_FitChainToCurve_WorkData WorkData;
+
+	RIGVM_METHOD()
+	virtual FRigVMStructUpgradeInfo GetUpgradeInfo() const override;
 };
 
 /**
@@ -484,6 +487,9 @@ struct CONTROLRIGSPLINE_API FRigUnit_FitSplineCurveToChain : public FRigUnit_Hig
 	 */
 	UPROPERTY(meta = (Input, Output))
 	FControlRigSpline Spline;	
+
+	RIGVM_METHOD()
+	virtual FRigVMStructUpgradeInfo GetUpgradeInfo() const override;
 };
 
 /**
diff --git a/Engine/Plugins/Experimental/FullBodyIK/Source/FullBodyIK/Private/RigUnit_FullbodyIK.cpp b/Engine/Plugins/Experimental/FullBodyIK/Source/FullBodyIK/Private/RigUnit_FullbodyIK.cpp
--- a/Engine/Plugins/Experimental/FullBodyIK/Source/FullBodyIK/Private/RigUnit_FullbodyIK.cpp
+++ b/Engine/Plugins/Experimental/FullBodyIK/Source/FullBodyIK/Private/RigUnit_FullbodyIK.cpp
@@ -596,3 +596,7 @@ IMPLEMENT_RIGUNIT_AUTOMATION_TEST(FRigUnit_FullbodyIK)
 	return true;
 }
 #endif
+FRigVMStructUpgradeInfo FRigUnit_FullbodyIK::GetUpgradeInfo() const
+{
+	return FRigUnit_HighlevelBaseMutable::GetUpgradeInfo();
+}
diff --git a/Engine/Plugins/Experimental/FullBodyIK/Source/FullBodyIK/Private/RigUnit_FullbodyIK.h b/Engine/Plugins/Experimental/FullBodyIK/Source/FullBodyIK/Private/RigUnit_FullbodyIK.h
--- a/Engine/Plugins/Experimental/FullBodyIK/Source/FullBodyIK/Private/RigUnit_FullbodyIK.h
+++ b/Engine/Plugins/Experimental/FullBodyIK/Source/FullBodyIK/Private/RigUnit_FullbodyIK.h
@@ -144,4 +144,7 @@ struct FRigUnit_FullbodyIK : public FRigUnit_HighlevelBaseMutable
 
 	UPROPERTY(transient)
 	FRigUnit_FullbodyIK_WorkData WorkData;
+
+	RIGVM_METHOD()
+	virtual FRigVMStructUpgradeInfo GetUpgradeInfo() const override;
 };
diff --git a/Engine/Source/Programs/UnrealHeaderTool/Private/HeaderParser.cpp b/Engine/Source/Programs/UnrealHeaderTool/Private/HeaderParser.cpp
--- a/Engine/Source/Programs/UnrealHeaderTool/Private/HeaderParser.cpp
+++ b/Engine/Source/Programs/UnrealHeaderTool/Private/HeaderParser.cpp
@@ -2079,6 +2079,21 @@ FUnrealScriptStructDefinitionInfo& FHeaderParser::CompileStructDeclaration()
 		Throwf(TEXT("Expected a GENERATED_BODY() at the start of struct"));
 	}
 
+	// check if the struct is marked as deprecated and does not implement the upgrade path
+	if(StructDef.HasMetaData(TEXT("Deprecated")))
+	{
+		const FRigVMStructInfo& StructRigVMInfo = StructDef.GetRigVMInfo();
+		if(!StructRigVMInfo.bHasGetUpgradeInfoMethod && !StructRigVMInfo.Methods.IsEmpty())
+		{
+			LogError(TEXT(
+				"RigVMStruct '%s' is marked as deprecated but is missing GetUpgradeInfo method.\n"
+				"Please implement a method like below:\n\n"
+				"RIGVM_METHOD()\n"
+				"virtual FRigVMStructUpgradeInfo GetUpgradeInfo() const override;"),
+				*StructDef.GetNameCPP());
+		}
+	}
+
 	// Validate sparse class data
 	CheckSparseClassData(StructDef);
 	return StructDef;
@@ -6210,7 +6225,17 @@ void FHeaderParser::CompileRigVMMethodDeclaration(FUnrealStructDefinitionInfo& S
 	FRigVMMethodInfo MethodInfo;
 	MethodInfo.ReturnType = FString(ReturnTypeToken.Value);
 	MethodInfo.Name = FString(NameToken.Value);
-	
+
+	// look out for the upgrade info method
+	static const FString GetUpgradeInfoString = TEXT("GetUpgradeInfo");
+	static const FString RigVMStructUpgradeInfoString = TEXT("FRigVMStructUpgradeInfo");
+	if(MethodInfo.ReturnType == RigVMStructUpgradeInfoString && MethodInfo.Name == GetUpgradeInfoString)
+	{
+		FRigVMStructInfo& StructRigVMInfo = StructDef.GetRigVMInfo();
+		StructRigVMInfo.bHasGetUpgradeInfoMethod = true;
+		return;
+	}
+
 	FString ParamString = FString::Join(ParamsContent, TEXT(" "));
 	if (!ParamString.IsEmpty())
 	{
diff --git a/Engine/Source/Programs/UnrealHeaderTool/Private/ParserHelper.h b/Engine/Source/Programs/UnrealHeaderTool/Private/ParserHelper.h
--- a/Engine/Source/Programs/UnrealHeaderTool/Private/ParserHelper.h
+++ b/Engine/Source/Programs/UnrealHeaderTool/Private/ParserHelper.h
@@ -711,6 +711,7 @@ struct FRigVMMethodInfo
 struct FRigVMStructInfo
 {
 	bool bHasRigVM = false;
+	bool bHasGetUpgradeInfoMethod = false;
 	FString Name;
 	FRigVMParameterArray Members;
 	TArray<FRigVMMethodInfo> Methods;
diff --git a/Engine/Source/Programs/UnrealHeaderTool/Resources/UHTDebugging/RigVMStructHeader.h b/Engine/Source/Programs/UnrealHeaderTool/Resources/UHTDebugging/RigVMStructHeader.h
--- a/Engine/Source/Programs/UnrealHeaderTool/Resources/UHTDebugging/RigVMStructHeader.h
+++ b/Engine/Source/Programs/UnrealHeaderTool/Resources/UHTDebugging/RigVMStructHeader.h
@@ -33,9 +33,11 @@ struct FRigVMStructBase
 
 	UPROPERTY(meta = (Output))
 	float InheritedOutput;
+
+	virtual FRigVMStructUpgradeInfo GetUpgradeInfo() const {};
 };
 
-USTRUCT()
+USTRUCT(meta = (Deprecated = "5.0.0"))
 struct FRigVMMethodStruct : public FRigVMStructBase
 {
 	GENERATED_BODY()
@@ -93,4 +95,7 @@ struct FRigVMMethodStruct : public FRigVMStructBase
 
 	UPROPERTY()
 	TEnumAsByte<ERigVMTestNameSpaceEnum::Type> HiddenNameSpaceEnum;
+
+	RIGVM_METHOD()
+	virtual FRigVMStructUpgradeInfo GetUpgradeInfo() const override;
 };
diff --git a/Engine/Source/Runtime/RigVM/Private/RigVMCore/RigVMStruct.cpp b/Engine/Source/Runtime/RigVM/Private/RigVMCore/RigVMStruct.cpp
--- a/Engine/Source/Runtime/RigVM/Private/RigVMCore/RigVMStruct.cpp
+++ b/Engine/Source/Runtime/RigVM/Private/RigVMCore/RigVMStruct.cpp
@@ -97,7 +97,28 @@ bool FRigVMStructUpgradeInfo::IsValid() const
 	return NodePath.IsEmpty() && (OldStruct != nullptr) && (NewStruct != nullptr);
 }
 
-#if WITH_EDITOR
+const FString& FRigVMStructUpgradeInfo::GetDefaultValueForPin(const FName& InPinName) const
+{
+	static const FString EmptyString = FString();
+	if(const FString* DefaultValue = DefaultValues.Find(InPinName))
+	{
+		return *DefaultValue;
+	}
+	return EmptyString;
+}
+
+void FRigVMStructUpgradeInfo::AddRemappedPin(const FString& InOldPinPath, const FString& InNewPinPath, bool bAsInput,
+	bool bAsOutput)
+{
+	if(bAsInput)
+	{
+		InputLinkMap.Add(InOldPinPath, InNewPinPath);
+	}
+	if(bAsOutput)
+	{
+		OutputLinkMap.Add(InOldPinPath, InNewPinPath);
+	}
+}
 
 void FRigVMStructUpgradeInfo::SetDefaultValues(const FRigVMStruct* InNewStructMemory)
 {
@@ -106,8 +127,6 @@ void FRigVMStructUpgradeInfo::SetDefaultValues(const FRigVMStruct* InNewStructMe
 	DefaultValues = InNewStructMemory->GetDefaultValues(NewStruct);
 }
 
-#endif
-
 ////////////////////////////////////////////////////////////////////////////////////////////////////////////////////////
 
 const FName FRigVMStruct::DeprecatedMetaName("Deprecated");
@@ -399,6 +418,8 @@ ERigVMPinDirection FRigVMStruct::GetPinDirectionFromProperty(FProperty* InProper
 	return ERigVMPinDirection::Hidden;
 }
 
+#endif
+
 FString FRigVMStruct::ExportToFullyQualifiedText(const FProperty* InMemberProperty, const uint8* InMemberMemoryPtr, bool bUseQuotes)
 {
 	check(InMemberProperty);
@@ -503,9 +524,6 @@ TMap<FName, FString> FRigVMStruct::GetDefaultValues(UScriptStruct* InScriptStruc
 {
 	check(InScriptStruct);
 	
-	FRigVMStructUpgradeInfo Info;
-	Info.OldStruct = InScriptStruct;
-
 	TMap<FName, FString> DefaultValues;
 	for (TFieldIterator<FProperty> It(InScriptStruct); It; ++It)
 	{
@@ -513,6 +531,15 @@ TMap<FName, FString> FRigVMStruct::GetDefaultValues(UScriptStruct* InScriptStruc
 		{
 			continue;
 		}
+
+#if WITH_EDITOR
+		if(!It->HasMetaData(FRigVMStruct::InputMetaName) &&
+			!It->HasMetaData(FRigVMStruct::OutputMetaName) &&
+			!It->HasMetaData(FRigVMStruct::VisibleMetaName))
+		{
+			continue;
+		}
+#endif
 		
 		const FName PropertyName = It->GetFName();
 		const uint8* StructMemberMemoryPtr = It->ContainerPtrToValuePtr<uint8>(this);
@@ -567,12 +594,12 @@ bool FRigVMStruct::ApplyUpgradeInfo(const FRigVMStructUpgradeInfo& InUpgradeInfo
 {
 	check(InUpgradeInfo.IsValid());
 
-	for(const TPair<FName, FString>& Pair : InUpgradeInfo.DefaultValues)
+	for(const TPair<FName, FString>& Pair : InUpgradeInfo.GetDefaultValues())
 	{
 		const FName& PropertyName = Pair.Key;
 		const FString& DefaultValue = Pair.Value;
 		
-		const FProperty* Property = InUpgradeInfo.NewStruct->FindPropertyByName(PropertyName);
+		const FProperty* Property = InUpgradeInfo.GetNewStruct()->FindPropertyByName(PropertyName);
 		if(Property == nullptr)
 		{
 			return false;
@@ -580,7 +607,7 @@ bool FRigVMStruct::ApplyUpgradeInfo(const FRigVMStructUpgradeInfo& InUpgradeInfo
 
 		uint8* MemberMemory = Property->ContainerPtrToValuePtr<uint8>(this);
 
-		FRigVMStructApplyUpgradeInfoErrorContext ErrorPipe(InUpgradeInfo.OldStruct, InUpgradeInfo.NewStruct, PropertyName);
+		FRigVMStructApplyUpgradeInfoErrorContext ErrorPipe(InUpgradeInfo.GetOldStruct(), InUpgradeInfo.GetNewStruct(), PropertyName);
         Property->ImportText_Direct(*DefaultValue, MemberMemory, nullptr, PPF_None, &ErrorPipe);
 
 		if(ErrorPipe.NumErrors > 0)
@@ -591,5 +618,3 @@ bool FRigVMStruct::ApplyUpgradeInfo(const FRigVMStructUpgradeInfo& InUpgradeInfo
 
 	return true;
 }
-
-#endif
\ No newline at end of file
diff --git a/Engine/Source/Runtime/RigVM/Public/RigVMCore/RigVMStruct.h b/Engine/Source/Runtime/RigVM/Public/RigVMCore/RigVMStruct.h
--- a/Engine/Source/Runtime/RigVM/Public/RigVMCore/RigVMStruct.h
+++ b/Engine/Source/Runtime/RigVM/Public/RigVMCore/RigVMStruct.h
@@ -138,18 +138,38 @@ private:
 /** Structure used to upgrade to a new implementation of a node */
 struct RIGVM_API FRigVMStructUpgradeInfo
 {
+public:
+	
 	FRigVMStructUpgradeInfo();
 
+	template<typename Old, typename New>
+	FRigVMStructUpgradeInfo(const Old& InOld, const New& InNew)
+	{
+		OldStruct = Old::StaticStruct();
+		NewStruct = New::StaticStruct();
+		SetDefaultValues(&InNew);
+	}
+
 	// returns true if this upgrade info can be applied
 	bool IsValid() const;
 
-#if WITH_EDITOR
-	
-	// sets the default values from the new struct.
-	void SetDefaultValues(const FRigVMStruct* InNewStructMemory);
+	// returns the old struct trying to be upgraded
+	UScriptStruct* GetOldStruct() const { return OldStruct; }
 
-#endif
+	// returns the new struct to upgrade to
+	UScriptStruct* GetNewStruct() const { return NewStruct; }
+
+	// returns the map for all default values
+	const TMap<FName, FString>& GetDefaultValues() const { return DefaultValues; }
+
+	// returns the default value for a given pin
+	const FString& GetDefaultValueForPin(const FName& InPinName) const;
+
+	// adds a pin to be remapped
+	void AddRemappedPin(const FString& InOldPinPath, const FString& InNewPinPath, bool bAsInput = true, bool bAsOutput = true);
 
+private:
+	
 	// The complete node path including models / collapse node.
 	// The path may look like "RigGraph|CollapseNode1|Add"
 	FString NodePath;
@@ -160,22 +180,21 @@ struct RIGVM_API FRigVMStructUpgradeInfo
 	// The new struct this upgrade info is targeting
 	UScriptStruct* NewStruct;
 
-	// Remapping info for re-mapping pins
-	// Entries can be root pins or sub pins
-	TMap<FString, FString> PinNameMap;
-
 	// Remapping info for re-linking inputs 
-	// (takes precedence over pin name map)
 	// Entries can be root pins or sub pins
 	TMap<FString, FString> InputLinkMap;
 
 	// Remapping info for re-linking outputs 
-	// (takes precedence over pin name map)
 	// Entries can be root pins or sub pins
 	TMap<FString, FString> OutputLinkMap;
 
 	// New sets of default values
 	TMap<FName, FString> DefaultValues;
+
+	// sets the default values from the new struct.
+	void SetDefaultValues(const FRigVMStruct* InNewStructMemory);
+
+	friend class URigVMController;
 };
 
 /**
@@ -207,11 +226,11 @@ public:
 	static bool CheckPinExists(UScriptStruct* InStruct, const FName& PinName, const FString& ExpectedType = FString(), FString* OutErrorMessage = nullptr);
 	static bool CheckMetadata(UScriptStruct* InStruct, const FName& PinName, const FName& InMetadataKey, FString* OutErrorMessage = nullptr);
 	static bool CheckFunctionExists(UScriptStruct* InStruct, const FName& FunctionName, FString* OutErrorMessage = nullptr);
+#endif
 	static FString ExportToFullyQualifiedText(const FProperty* InMemberProperty, const uint8* InMemberMemoryPtr, bool bUseQuotes = true);
 	static FString ExportToFullyQualifiedText(UScriptStruct* InStruct, const uint8* InStructMemoryPtr);
 	FString ExportToFullyQualifiedText(UScriptStruct* InScriptStruct, const FName& InPropertyName, const uint8* InStructMemoryPointer = nullptr) const;
 	virtual FRigVMStructUpgradeInfo GetUpgradeInfo() const { return FRigVMStructUpgradeInfo(); }
-#endif
 
 	static const FName DeprecatedMetaName;
 	static const FName InputMetaName;
@@ -250,10 +269,8 @@ public:
 protected:
 
 	static float GetRatioFromIndex(int32 InIndex, int32 InCount);
-#if WITH_EDITOR
 	TMap<FName, FString> GetDefaultValues(UScriptStruct* InScriptStruct) const;
 	bool ApplyUpgradeInfo(const FRigVMStructUpgradeInfo& InUpgradeInfo);
-#endif
 
 	friend struct FRigVMStructUpgradeInfo;
 	friend class FRigVMGraphStructUpgradeInfoTest;
