Control Rig: Rename PrototypeName to TemplateName

diff --git a/Engine/Plugins/Animation/ControlRig/Source/ControlRig/Private/Units/Debug/RigUnit_VisualDebug.h b/Engine/Plugins/Animation/ControlRig/Source/ControlRig/Private/Units/Debug/RigUnit_VisualDebug.h
--- a/Engine/Plugins/Animation/ControlRig/Source/ControlRig/Private/Units/Debug/RigUnit_VisualDebug.h
+++ b/Engine/Plugins/Animation/ControlRig/Source/ControlRig/Private/Units/Debug/RigUnit_VisualDebug.h
@@ -18,7 +18,7 @@ enum class ERigUnitVisualDebugPointMode : uint8
 	Max UMETA(Hidden),
 };
 
-USTRUCT(meta=(DisplayName = "Visual Debug Vector", PrototypeName = "VisualDebug", Keywords = "Draw,Point", Deprecated = "4.25", Varying))
+USTRUCT(meta=(DisplayName = "Visual Debug Vector", TemplateName="VisualDebug", Keywords = "Draw,Point", Deprecated = "4.25", Varying))
 struct CONTROLRIG_API FRigUnit_VisualDebugVector : public FRigUnit_DebugBase
 {
 	GENERATED_BODY()
@@ -65,7 +65,7 @@ struct CONTROLRIG_API FRigUnit_VisualDebugVector : public FRigUnit_DebugBase
 /**
  * Debug draw parameters for a Point or Vector given a vector
  */
-USTRUCT(meta=(DisplayName = "Visual Debug Vector", PrototypeName = "VisualDebug", Keywords = "Draw,Point", Varying))
+USTRUCT(meta=(DisplayName = "Visual Debug Vector", TemplateName="VisualDebug", Keywords = "Draw,Point", Varying))
 struct CONTROLRIG_API FRigUnit_VisualDebugVectorItemSpace : public FRigUnit_DebugBase
 {
 	GENERATED_BODY()
@@ -111,7 +111,7 @@ struct CONTROLRIG_API FRigUnit_VisualDebugVectorItemSpace : public FRigUnit_Debu
 	FRigElementKey Space;
 };
 
-USTRUCT(meta = (DisplayName = "Visual Debug Quat", PrototypeName = "VisualDebug", Keywords = "Draw,Rotation", Deprecated = "4.25", Varying))
+USTRUCT(meta = (DisplayName = "Visual Debug Quat", TemplateName="VisualDebug", Keywords = "Draw,Rotation", Deprecated = "4.25", Varying))
 struct CONTROLRIG_API FRigUnit_VisualDebugQuat : public FRigUnit_DebugBase
 {
 	GENERATED_BODY()
@@ -150,7 +150,7 @@ struct CONTROLRIG_API FRigUnit_VisualDebugQuat : public FRigUnit_DebugBase
 /**
  * Debug draw parameters for an Axis given a quaternion
  */
-USTRUCT(meta = (DisplayName = "Visual Debug Quat", PrototypeName = "VisualDebug", Keywords = "Draw,Rotation", Varying))
+USTRUCT(meta = (DisplayName = "Visual Debug Quat", TemplateName="VisualDebug", Keywords = "Draw,Rotation", Varying))
 struct CONTROLRIG_API FRigUnit_VisualDebugQuatItemSpace : public FRigUnit_DebugBase
 {
 	GENERATED_BODY()
@@ -188,7 +188,7 @@ struct CONTROLRIG_API FRigUnit_VisualDebugQuatItemSpace : public FRigUnit_DebugB
 	FRigElementKey Space;
 };
 
-USTRUCT(meta=(DisplayName="Visual Debug Transform", PrototypeName = "VisualDebug", Keywords = "Draw,Axes", Deprecated = "4.25", Varying))
+USTRUCT(meta=(DisplayName="Visual Debug Transform", TemplateName="VisualDebug", Keywords = "Draw,Axes", Deprecated = "4.25", Varying))
 struct CONTROLRIG_API FRigUnit_VisualDebugTransform : public FRigUnit_DebugBase
 {
 	GENERATED_BODY()
@@ -227,7 +227,7 @@ struct CONTROLRIG_API FRigUnit_VisualDebugTransform : public FRigUnit_DebugBase
 /**
  * Debug draw parameters for an Axis given a transform
  */
-USTRUCT(meta=(DisplayName="Visual Debug Transform", PrototypeName = "VisualDebug", Keywords = "Draw,Axes", Varying))
+USTRUCT(meta=(DisplayName="Visual Debug Transform", TemplateName="VisualDebug", Keywords = "Draw,Axes", Varying))
 struct CONTROLRIG_API FRigUnit_VisualDebugTransformItemSpace : public FRigUnit_DebugBase
 {
 	GENERATED_BODY()
diff --git a/Engine/Plugins/Animation/ControlRig/Source/ControlRig/Private/Units/Hierarchy/RigUnit_GetControlTransform.h b/Engine/Plugins/Animation/ControlRig/Source/ControlRig/Private/Units/Hierarchy/RigUnit_GetControlTransform.h
--- a/Engine/Plugins/Animation/ControlRig/Source/ControlRig/Private/Units/Hierarchy/RigUnit_GetControlTransform.h
+++ b/Engine/Plugins/Animation/ControlRig/Source/ControlRig/Private/Units/Hierarchy/RigUnit_GetControlTransform.h
@@ -8,7 +8,7 @@
 /**
  * GetControlBool is used to retrieve a single Bool from a hierarchy.
  */
-USTRUCT(meta=(DisplayName="Get Control Bool", Category="Controls", DocumentationPolicy = "Strict", Keywords="GetControlBool", PrototypeName = "GetControlValue", NodeColor="0.462745, 1,0, 0.329412",Varying))
+USTRUCT(meta=(DisplayName="Get Control Bool", Category="Controls", DocumentationPolicy = "Strict", Keywords="GetControlBool", TemplateName="GetControlValue", NodeColor="0.462745, 1,0, 0.329412",Varying))
 struct CONTROLRIG_API FRigUnit_GetControlBool : public FRigUnit
 {
 	GENERATED_BODY()
@@ -39,7 +39,7 @@ struct CONTROLRIG_API FRigUnit_GetControlBool : public FRigUnit
 /**
  * GetControlFloat is used to retrieve a single Float from a hierarchy.
  */
-USTRUCT(meta=(DisplayName="Get Control Float", Category="Controls", DocumentationPolicy = "Strict", Keywords="GetControlFloat", PrototypeName = "GetControlValue", NodeColor="0.462745, 1,0, 0.329412", Varying))
+USTRUCT(meta=(DisplayName="Get Control Float", Category="Controls", DocumentationPolicy = "Strict", Keywords="GetControlFloat", TemplateName="GetControlValue", NodeColor="0.462745, 1,0, 0.329412", Varying))
 struct CONTROLRIG_API FRigUnit_GetControlFloat : public FRigUnit
 {
 	GENERATED_BODY()
@@ -80,7 +80,7 @@ struct CONTROLRIG_API FRigUnit_GetControlFloat : public FRigUnit
 /**
  * GetControlFloat is used to retrieve a single Integer from a hierarchy.
  */
-USTRUCT(meta=(DisplayName="Get Control Integer", Category="Controls", DocumentationPolicy = "Strict", Keywords="GetControlInteger,GetControlEnum", PrototypeName = "GetControlValue", NodeColor="0.462745, 1,0, 0.329412", Varying))
+USTRUCT(meta=(DisplayName="Get Control Integer", Category="Controls", DocumentationPolicy = "Strict", Keywords="GetControlInteger,GetControlEnum", TemplateName="GetControlValue", NodeColor="0.462745, 1,0, 0.329412", Varying))
 struct CONTROLRIG_API FRigUnit_GetControlInteger : public FRigUnit
 {
 	GENERATED_BODY()
@@ -121,7 +121,7 @@ struct CONTROLRIG_API FRigUnit_GetControlInteger : public FRigUnit
 /**
  * GetControlVector2D is used to retrieve a single Vector2D from a hierarchy.
  */
-USTRUCT(meta=(DisplayName="Get Control Vector2D", Category="Controls", DocumentationPolicy = "Strict", Keywords="GetControlVector2D", PrototypeName = "GetControlValue", NodeColor="0.462745, 1,0, 0.329412", Varying))
+USTRUCT(meta=(DisplayName="Get Control Vector2D", Category="Controls", DocumentationPolicy = "Strict", Keywords="GetControlVector2D", TemplateName="GetControlValue", NodeColor="0.462745, 1,0, 0.329412", Varying))
 struct CONTROLRIG_API FRigUnit_GetControlVector2D : public FRigUnit
 {
 	GENERATED_BODY()
@@ -162,7 +162,7 @@ struct CONTROLRIG_API FRigUnit_GetControlVector2D : public FRigUnit
 /**
  * GetControlVector is used to retrieve a single Vector from a hierarchy, can be used for Controls of type "Position" or "Scale".
  */
-USTRUCT(meta=(DisplayName="Get Control Vector", Category="Controls", DocumentationPolicy = "Strict", Keywords="GetControlVector, GetControlPosition, GetControlScale", PrototypeName = "GetControlValue", NodeColor="0.462745, 1,0, 0.329412", Varying))
+USTRUCT(meta=(DisplayName="Get Control Vector", Category="Controls", DocumentationPolicy = "Strict", Keywords="GetControlVector, GetControlPosition, GetControlScale", TemplateName="GetControlValue", NodeColor="0.462745, 1,0, 0.329412", Varying))
 struct CONTROLRIG_API FRigUnit_GetControlVector : public FRigUnit
 {
 	GENERATED_BODY()
@@ -211,7 +211,7 @@ struct CONTROLRIG_API FRigUnit_GetControlVector : public FRigUnit
 /**
  * GetControlRotator is used to retrieve a single Rotator from a hierarchy.
  */
-USTRUCT(meta=(DisplayName="Get Control Rotator", Category="Controls", DocumentationPolicy = "Strict", Keywords="GetControlRotator", PrototypeName = "GetControlValue", NodeColor="0.462745, 1,0, 0.329412", Varying))
+USTRUCT(meta=(DisplayName="Get Control Rotator", Category="Controls", DocumentationPolicy = "Strict", Keywords="GetControlRotator", TemplateName="GetControlValue", NodeColor="0.462745, 1,0, 0.329412", Varying))
 struct CONTROLRIG_API FRigUnit_GetControlRotator : public FRigUnit
 {
 	GENERATED_BODY()
@@ -261,7 +261,7 @@ struct CONTROLRIG_API FRigUnit_GetControlRotator : public FRigUnit
 /**
  * GetControlTransform is used to retrieve a single transform from a hierarchy.
  */
-USTRUCT(meta=(DisplayName="Get Control Transform", Category="Controls", DocumentationPolicy = "Strict", Keywords="GetControlTransform", PrototypeName = "GetControlValue", NodeColor="0.462745, 1,0, 0.329412", Varying))
+USTRUCT(meta=(DisplayName="Get Control Transform", Category="Controls", DocumentationPolicy = "Strict", Keywords="GetControlTransform", TemplateName="GetControlValue", NodeColor="0.462745, 1,0, 0.329412", Varying))
 struct CONTROLRIG_API FRigUnit_GetControlTransform : public FRigUnit
 {
 	GENERATED_BODY()
diff --git a/Engine/Plugins/Animation/ControlRig/Source/ControlRig/Private/Units/Hierarchy/RigUnit_SendEvent.h b/Engine/Plugins/Animation/ControlRig/Source/ControlRig/Private/Units/Hierarchy/RigUnit_SendEvent.h
--- a/Engine/Plugins/Animation/ControlRig/Source/ControlRig/Private/Units/Hierarchy/RigUnit_SendEvent.h
+++ b/Engine/Plugins/Animation/ControlRig/Source/ControlRig/Private/Units/Hierarchy/RigUnit_SendEvent.h
@@ -8,7 +8,7 @@
 /**
  * SendEvent is used to notify the engine / editor of a change that happend within the Control Rig.
  */
-USTRUCT(meta=(DisplayName="Send Event", Category="Hierarchy", DocumentationPolicy="Strict", Keywords = "SendEvent", PrototypeName = "Event,Notify,Notification", NodeColor="1, 0, 0"))
+USTRUCT(meta=(DisplayName="Send Event", Category="Hierarchy", DocumentationPolicy="Strict", Keywords = "SendEvent", TemplateName="Event,Notify,Notification", NodeColor="1, 0, 0"))
 struct CONTROLRIG_API FRigUnit_SendEvent : public FRigUnitMutable
 {
 	GENERATED_BODY()
diff --git a/Engine/Plugins/Animation/ControlRig/Source/ControlRig/Private/Units/Hierarchy/RigUnit_SetControlColor.h b/Engine/Plugins/Animation/ControlRig/Source/ControlRig/Private/Units/Hierarchy/RigUnit_SetControlColor.h
--- a/Engine/Plugins/Animation/ControlRig/Source/ControlRig/Private/Units/Hierarchy/RigUnit_SetControlColor.h
+++ b/Engine/Plugins/Animation/ControlRig/Source/ControlRig/Private/Units/Hierarchy/RigUnit_SetControlColor.h
@@ -8,7 +8,7 @@
 /**
  * SetControlColor is used to change the gizmo's color on a control at runtime
  */
-USTRUCT(meta=(DisplayName="Set Control Color", Category="Controls", DocumentationPolicy="Strict", Keywords = "SetControlColor,SetGizmoColor", PrototypeName = "SetControlColor", NodeColor = "0.0 0.36470600962638855 1.0"))
+USTRUCT(meta=(DisplayName="Set Control Color", Category="Controls", DocumentationPolicy="Strict", Keywords = "SetControlColor,SetGizmoColor", TemplateName="SetControlColor", NodeColor = "0.0 0.36470600962638855 1.0"))
 struct CONTROLRIG_API FRigUnit_SetControlColor : public FRigUnitMutable
 {
 	GENERATED_BODY()
diff --git a/Engine/Plugins/Animation/ControlRig/Source/ControlRig/Private/Units/Hierarchy/RigUnit_SetControlOffset.h b/Engine/Plugins/Animation/ControlRig/Source/ControlRig/Private/Units/Hierarchy/RigUnit_SetControlOffset.h
--- a/Engine/Plugins/Animation/ControlRig/Source/ControlRig/Private/Units/Hierarchy/RigUnit_SetControlOffset.h
+++ b/Engine/Plugins/Animation/ControlRig/Source/ControlRig/Private/Units/Hierarchy/RigUnit_SetControlOffset.h
@@ -9,7 +9,7 @@
  * SetControlOffset is used to perform a change in the hierarchy by setting a single control's transform.
  * This is typically only used during the Setup Event.
  */
-USTRUCT(meta=(DisplayName="Set Control Offset", Category="Controls", DocumentationPolicy="Strict", Keywords = "SetControlOffset,Initial,InitialTransform,SetInitialTransform,SetInitialControlTransform", PrototypeName = "SetControlValue", NodeColor="0, 0.364706, 1.0"))
+USTRUCT(meta=(DisplayName="Set Control Offset", Category="Controls", DocumentationPolicy="Strict", Keywords = "SetControlOffset,Initial,InitialTransform,SetInitialTransform,SetInitialControlTransform", TemplateName="SetControlValue", NodeColor="0, 0.364706, 1.0"))
 struct CONTROLRIG_API FRigUnit_SetControlOffset : public FRigUnitMutable
 {
 	GENERATED_BODY()
diff --git a/Engine/Plugins/Animation/ControlRig/Source/ControlRig/Private/Units/Hierarchy/RigUnit_SetControlTransform.h b/Engine/Plugins/Animation/ControlRig/Source/ControlRig/Private/Units/Hierarchy/RigUnit_SetControlTransform.h
--- a/Engine/Plugins/Animation/ControlRig/Source/ControlRig/Private/Units/Hierarchy/RigUnit_SetControlTransform.h
+++ b/Engine/Plugins/Animation/ControlRig/Source/ControlRig/Private/Units/Hierarchy/RigUnit_SetControlTransform.h
@@ -8,7 +8,7 @@
 /**
  * SetControlBool is used to perform a change in the hierarchy by setting a single control's bool value.
  */
-USTRUCT(meta=(DisplayName="Set Control Bool", Category="Controls", DocumentationPolicy="Strict", Keywords = "SetControlBool,SetGizmoBool", PrototypeName = "SetControlValue", NodeColor = "0.0 0.36470600962638855 1.0"))
+USTRUCT(meta=(DisplayName="Set Control Bool", Category="Controls", DocumentationPolicy="Strict", Keywords = "SetControlBool,SetGizmoBool", TemplateName="SetControlValue", NodeColor = "0.0 0.36470600962638855 1.0"))
 struct CONTROLRIG_API FRigUnit_SetControlBool : public FRigUnitMutable
 {
 	GENERATED_BODY()
@@ -61,7 +61,7 @@ struct CONTROLRIG_API FRigUnit_SetMultiControlBool_Entry
 /**
  * SetMultiControlBool is used to perform a change in the hierarchy by setting multiple controls' bool value.
  */
-USTRUCT(meta = (DisplayName = "Set Multiple Controls Bool", Category = "Controls", DocumentationPolicy = "Strict", Keywords = "SetMultipleControlsBool,SetControlBool,SetGizmoBool", PrototypeName = "SetMultiControlValue", NodeColor = "0.0 0.36470600962638855 1.0"))
+USTRUCT(meta = (DisplayName = "Set Multiple Controls Bool", Category = "Controls", DocumentationPolicy = "Strict", Keywords = "SetMultipleControlsBool,SetControlBool,SetGizmoBool", TemplateName="SetMultiControlValue", NodeColor = "0.0 0.36470600962638855 1.0"))
 struct CONTROLRIG_API FRigUnit_SetMultiControlBool : public FRigUnitMutable
 {
 	GENERATED_BODY()
@@ -88,7 +88,7 @@ struct CONTROLRIG_API FRigUnit_SetMultiControlBool : public FRigUnitMutable
 /**
  * SetControlFloat is used to perform a change in the hierarchy by setting a single control's float value.
  */
-USTRUCT(meta=(DisplayName="Set Control Float", Category="Controls", DocumentationPolicy="Strict", Keywords = "SetControlFloat,SetGizmoFloat", PrototypeName = "SetControlValue", NodeColor = "0.0 0.36470600962638855 1.0"))
+USTRUCT(meta=(DisplayName="Set Control Float", Category="Controls", DocumentationPolicy="Strict", Keywords = "SetControlFloat,SetGizmoFloat", TemplateName="SetControlValue", NodeColor = "0.0 0.36470600962638855 1.0"))
 struct CONTROLRIG_API FRigUnit_SetControlFloat : public FRigUnitMutable
 {
 	GENERATED_BODY()
@@ -149,7 +149,7 @@ struct CONTROLRIG_API FRigUnit_SetMultiControlFloat_Entry
 /**
  * SetMultiControlFloat is used to perform a change in the hierarchy by setting multiple controls' float value.
  */
-USTRUCT(meta = (DisplayName = "Set Multiple Controls Float", Category = "Controls", DocumentationPolicy = "Strict", Keywords = "SetMultipleControlsFloat,SetControlFloat,SetGizmoFloat", PrototypeName = "SetMultiControlValue", NodeColor = "0.0 0.36470600962638855 1.0"))
+USTRUCT(meta = (DisplayName = "Set Multiple Controls Float", Category = "Controls", DocumentationPolicy = "Strict", Keywords = "SetMultipleControlsFloat,SetControlFloat,SetGizmoFloat", TemplateName="SetMultiControlValue", NodeColor = "0.0 0.36470600962638855 1.0"))
 struct CONTROLRIG_API FRigUnit_SetMultiControlFloat : public FRigUnitMutable
 {
 	GENERATED_BODY()
@@ -185,7 +185,7 @@ struct CONTROLRIG_API FRigUnit_SetMultiControlFloat : public FRigUnitMutable
 /**
  * SetControlInteger is used to perform a change in the hierarchy by setting a single control's int32 value.
  */
-USTRUCT(meta=(DisplayName="Set Control Integer", Category="Controls", DocumentationPolicy="Strict", Keywords = "SetControlInteger,SetGizmoInteger", PrototypeName = "SetControlValue", NodeColor = "0.0 0.36470600962638855 1.0"))
+USTRUCT(meta=(DisplayName="Set Control Integer", Category="Controls", DocumentationPolicy="Strict", Keywords = "SetControlInteger,SetGizmoInteger", TemplateName="SetControlValue", NodeColor = "0.0 0.36470600962638855 1.0"))
 struct CONTROLRIG_API FRigUnit_SetControlInteger : public FRigUnitMutable
 {
 	GENERATED_BODY()
@@ -246,7 +246,7 @@ struct CONTROLRIG_API FRigUnit_SetMultiControlInteger_Entry
 /**
  * SetMultiControlInteger is used to perform a change in the hierarchy by setting multiple controls' integer value.
  */
-USTRUCT(meta = (DisplayName = "Set Multiple Controls Integer", Category = "Controls", DocumentationPolicy = "Strict", Keywords = "SetMultipleControlsInteger,SetControlInteger,SetGizmoInteger", PrototypeName = "SetMultiControlValue", NodeColor = "0.0 0.36470600962638855 1.0"))
+USTRUCT(meta = (DisplayName = "Set Multiple Controls Integer", Category = "Controls", DocumentationPolicy = "Strict", Keywords = "SetMultipleControlsInteger,SetControlInteger,SetGizmoInteger", TemplateName="SetMultiControlValue", NodeColor = "0.0 0.36470600962638855 1.0"))
 struct CONTROLRIG_API FRigUnit_SetMultiControlInteger : public FRigUnitMutable
 {
 	GENERATED_BODY()
@@ -280,7 +280,7 @@ struct CONTROLRIG_API FRigUnit_SetMultiControlInteger : public FRigUnitMutable
 /**
  * SetControlVector2D is used to perform a change in the hierarchy by setting a single control's Vector2D value.
  */
-USTRUCT(meta=(DisplayName="Set Control Vector2D", Category="Controls", DocumentationPolicy="Strict", Keywords = "SetControlVector2D,SetGizmoVector2D", PrototypeName = "SetControlValue", NodeColor = "0.0 0.36470600962638855 1.0"))
+USTRUCT(meta=(DisplayName="Set Control Vector2D", Category="Controls", DocumentationPolicy="Strict", Keywords = "SetControlVector2D,SetGizmoVector2D", TemplateName="SetControlValue", NodeColor = "0.0 0.36470600962638855 1.0"))
 struct CONTROLRIG_API FRigUnit_SetControlVector2D : public FRigUnitMutable
 {
 	GENERATED_BODY()
@@ -342,7 +342,7 @@ struct CONTROLRIG_API FRigUnit_SetMultiControlVector2D_Entry
 /**
  * SetMultiControlVector2D is used to perform a change in the hierarchy by setting multiple controls' vector2D value.
  */
-USTRUCT(meta = (DisplayName = "Set Multiple Controls Vector2D", Category = "Controls", DocumentationPolicy = "Strict", Keywords = "SetMultipleControlsVector2D,SetControlVector2D,SetGizmoVector2D", PrototypeName = "SetMultiControlValue", NodeColor = "0.0 0.36470600962638855 1.0"))
+USTRUCT(meta = (DisplayName = "Set Multiple Controls Vector2D", Category = "Controls", DocumentationPolicy = "Strict", Keywords = "SetMultipleControlsVector2D,SetControlVector2D,SetGizmoVector2D", TemplateName="SetMultiControlValue", NodeColor = "0.0 0.36470600962638855 1.0"))
 struct CONTROLRIG_API FRigUnit_SetMultiControlVector2D : public FRigUnitMutable
 {
 	GENERATED_BODY()
@@ -376,7 +376,7 @@ struct CONTROLRIG_API FRigUnit_SetMultiControlVector2D : public FRigUnitMutable
 /**
  * SetControlVector is used to perform a change in the hierarchy by setting a single control's Vector value.
  */
-USTRUCT(meta=(DisplayName="Set Control Vector", Category="Controls", DocumentationPolicy="Strict", Keywords = "SetControlVector,SetGizmoVector", PrototypeName = "SetControlValue", NodeColor = "0.0 0.36470600962638855 1.0"))
+USTRUCT(meta=(DisplayName="Set Control Vector", Category="Controls", DocumentationPolicy="Strict", Keywords = "SetControlVector,SetGizmoVector", TemplateName="SetControlValue", NodeColor = "0.0 0.36470600962638855 1.0"))
 struct CONTROLRIG_API FRigUnit_SetControlVector : public FRigUnitMutable
 {
 	GENERATED_BODY()
@@ -424,7 +424,7 @@ struct CONTROLRIG_API FRigUnit_SetControlVector : public FRigUnitMutable
 /**
  * SetControlRotator is used to perform a change in the hierarchy by setting a single control's Rotator value.
  */
-USTRUCT(meta=(DisplayName="Set Control Rotator", Category="Controls", DocumentationPolicy="Strict", Keywords = "SetControlRotator,SetGizmoRotator", PrototypeName = "SetControlValue", NodeColor = "0.0 0.36470600962638855 1.0"))
+USTRUCT(meta=(DisplayName="Set Control Rotator", Category="Controls", DocumentationPolicy="Strict", Keywords = "SetControlRotator,SetGizmoRotator", TemplateName="SetControlValue", NodeColor = "0.0 0.36470600962638855 1.0"))
 struct CONTROLRIG_API FRigUnit_SetControlRotator : public FRigUnitMutable
 {
 	GENERATED_BODY()
@@ -503,7 +503,7 @@ struct CONTROLRIG_API FRigUnit_SetMultiControlRotator_Entry
 /**
  * SetMultiControlRotator is used to perform a change in the hierarchy by setting multiple controls' rotator value.
  */
-USTRUCT(meta = (DisplayName = "Set Multiple Controls Rotator", Category = "Controls", DocumentationPolicy = "Strict", Keywords = "SetMultipleControlsRotator,SetControlRotator,SetGizmoRotator", PrototypeName = "SetMultiControlValue", NodeColor = "0.0 0.36470600962638855 1.0"))
+USTRUCT(meta = (DisplayName = "Set Multiple Controls Rotator", Category = "Controls", DocumentationPolicy = "Strict", Keywords = "SetMultipleControlsRotator,SetControlRotator,SetGizmoRotator", TemplateName="SetMultiControlValue", NodeColor = "0.0 0.36470600962638855 1.0"))
 struct CONTROLRIG_API FRigUnit_SetMultiControlRotator : public FRigUnitMutable
 {
 	GENERATED_BODY()
@@ -537,7 +537,7 @@ struct CONTROLRIG_API FRigUnit_SetMultiControlRotator : public FRigUnitMutable
 /**
  * SetControlTransform is used to perform a change in the hierarchy by setting a single control's transform.
  */
-USTRUCT(meta=(DisplayName="Set Control Transform", Category="Controls", DocumentationPolicy="Strict", Keywords = "SetControlTransform,SetGizmoTransform", PrototypeName = "SetControlValue", Deprecated = "4.25"))
+USTRUCT(meta=(DisplayName="Set Control Transform", Category="Controls", DocumentationPolicy="Strict", Keywords = "SetControlTransform,SetGizmoTransform", TemplateName="SetControlValue", Deprecated = "4.25"))
 struct CONTROLRIG_API FRigUnit_SetControlTransform : public FRigUnitMutable
 {
 	GENERATED_BODY()
diff --git a/Engine/Plugins/Animation/ControlRig/Source/ControlRig/Private/Units/Hierarchy/RigUnit_SetControlVisibility.h b/Engine/Plugins/Animation/ControlRig/Source/ControlRig/Private/Units/Hierarchy/RigUnit_SetControlVisibility.h
--- a/Engine/Plugins/Animation/ControlRig/Source/ControlRig/Private/Units/Hierarchy/RigUnit_SetControlVisibility.h
+++ b/Engine/Plugins/Animation/ControlRig/Source/ControlRig/Private/Units/Hierarchy/RigUnit_SetControlVisibility.h
@@ -8,7 +8,7 @@
 /**
  * SetControlVisibility is used to change the gizmo visibility on a control at runtime
  */
-USTRUCT(meta=(DisplayName="Set Control Visibility", Category="Controls", DocumentationPolicy="Strict", Keywords = "SetControlVisibility,Visibility,Hide,Show,Hidden,Visible,SetGizmoVisibility", PrototypeName = "SetControlVisibility", NodeColor="0, 0.364706, 1.0"))
+USTRUCT(meta=(DisplayName="Set Control Visibility", Category="Controls", DocumentationPolicy="Strict", Keywords = "SetControlVisibility,Visibility,Hide,Show,Hidden,Visible,SetGizmoVisibility", TemplateName="SetControlVisibility", NodeColor="0, 0.364706, 1.0"))
 struct CONTROLRIG_API FRigUnit_SetControlVisibility : public FRigUnitMutable
 {
 	GENERATED_BODY()
diff --git a/Engine/Plugins/Animation/ControlRig/Source/ControlRig/Private/Units/Math/RigUnit_MathBool.h b/Engine/Plugins/Animation/ControlRig/Source/ControlRig/Private/Units/Math/RigUnit_MathBool.h
--- a/Engine/Plugins/Animation/ControlRig/Source/ControlRig/Private/Units/Math/RigUnit_MathBool.h
+++ b/Engine/Plugins/Animation/ControlRig/Source/ControlRig/Private/Units/Math/RigUnit_MathBool.h
@@ -99,7 +99,7 @@ struct CONTROLRIG_API FRigUnit_MathBoolConstFalse : public FRigUnit_MathBoolCons
 /**
  * Returns true if the condition is false
  */
-USTRUCT(meta=(DisplayName="Not", PrototypeName="Not", Keywords="!"))
+USTRUCT(meta=(DisplayName="Not", TemplateName="Not", Keywords="!"))
 struct CONTROLRIG_API FRigUnit_MathBoolNot : public FRigUnit_MathBoolUnaryOp
 {
 	GENERATED_BODY()
@@ -111,7 +111,7 @@ struct CONTROLRIG_API FRigUnit_MathBoolNot : public FRigUnit_MathBoolUnaryOp
 /**
  * Returns true if both conditions are true
  */
-USTRUCT(meta=(DisplayName="And", PrototypeName="And", Keywords="&&"))
+USTRUCT(meta=(DisplayName="And", TemplateName="And", Keywords="&&"))
 struct CONTROLRIG_API FRigUnit_MathBoolAnd : public FRigUnit_MathBoolBinaryOp
 {
 	GENERATED_BODY()
@@ -123,7 +123,7 @@ struct CONTROLRIG_API FRigUnit_MathBoolAnd : public FRigUnit_MathBoolBinaryOp
 /**
  * Returns true if both conditions are false
  */
-USTRUCT(meta=(DisplayName="Nand", PrototypeName="Nand"))
+USTRUCT(meta=(DisplayName="Nand", TemplateName="Nand"))
 struct CONTROLRIG_API FRigUnit_MathBoolNand : public FRigUnit_MathBoolBinaryOp
 {
 	GENERATED_BODY()
@@ -135,7 +135,7 @@ struct CONTROLRIG_API FRigUnit_MathBoolNand : public FRigUnit_MathBoolBinaryOp
 /**
  * Returns true if one of the conditions is true
  */
-USTRUCT(meta=(DisplayName="Or", PrototypeName="Or", Keywords="||"))
+USTRUCT(meta=(DisplayName="Or", TemplateName="Or", Keywords="||"))
 struct CONTROLRIG_API FRigUnit_MathBoolOr : public FRigUnit_MathBoolBinaryOp
 {
 	GENERATED_BODY()
@@ -147,7 +147,7 @@ struct CONTROLRIG_API FRigUnit_MathBoolOr : public FRigUnit_MathBoolBinaryOp
 /**
  * Returns true if the value A equals B
  */
-USTRUCT(meta=(DisplayName="Equals", PrototypeName="Equals", Keywords="Same,=="))
+USTRUCT(meta=(DisplayName="Equals", TemplateName="Equals", Keywords="Same,=="))
 struct CONTROLRIG_API FRigUnit_MathBoolEquals : public FRigUnit_MathBoolBase
 {
 	GENERATED_BODY()
@@ -173,7 +173,7 @@ struct CONTROLRIG_API FRigUnit_MathBoolEquals : public FRigUnit_MathBoolBase
 /**
  * Returns true if the value A does not equal B
  */
-USTRUCT(meta=(DisplayName="Not Equals", PrototypeName="NotEquals", Keywords="Different,!=,Xor"))
+USTRUCT(meta=(DisplayName="Not Equals", TemplateName="NotEquals", Keywords="Different,!=,Xor"))
 struct CONTROLRIG_API FRigUnit_MathBoolNotEquals : public FRigUnit_MathBoolBase
 {
 	GENERATED_BODY()
@@ -199,7 +199,7 @@ struct CONTROLRIG_API FRigUnit_MathBoolNotEquals : public FRigUnit_MathBoolBase
 /**
 * Returns true if the value has changed from the last run
 */
-USTRUCT(meta=(DisplayName="Toggled", PrototypeName="Toggled", Keywords="Changed,Different"))
+USTRUCT(meta=(DisplayName="Toggled", TemplateName="Toggled", Keywords="Changed,Different"))
 struct CONTROLRIG_API FRigUnit_MathBoolToggled : public FRigUnit_MathBoolBase
 {
 	GENERATED_BODY()
diff --git a/Engine/Plugins/Animation/ControlRig/Source/ControlRig/Private/Units/Math/RigUnit_MathColor.h b/Engine/Plugins/Animation/ControlRig/Source/ControlRig/Private/Units/Math/RigUnit_MathColor.h
--- a/Engine/Plugins/Animation/ControlRig/Source/ControlRig/Private/Units/Math/RigUnit_MathColor.h
+++ b/Engine/Plugins/Animation/ControlRig/Source/ControlRig/Private/Units/Math/RigUnit_MathColor.h
@@ -35,7 +35,7 @@ struct CONTROLRIG_API FRigUnit_MathColorBinaryOp : public FRigUnit_MathColorBase
 /**
  * Makes a vector from a single float
  */
-USTRUCT(meta=(DisplayName="From Float", PrototypeName="FromFloat", Keywords="Make,Construct"))
+USTRUCT(meta=(DisplayName="From Float", TemplateName="FromFloat", Keywords="Make,Construct"))
 struct CONTROLRIG_API FRigUnit_MathColorFromFloat : public FRigUnit_MathColorBase
 {
 	GENERATED_BODY()
@@ -58,7 +58,7 @@ struct CONTROLRIG_API FRigUnit_MathColorFromFloat : public FRigUnit_MathColorBas
 /**
  * Returns the sum of the two values
  */
-USTRUCT(meta=(DisplayName="Add", PrototypeName="Add", Keywords="Sum,+"))
+USTRUCT(meta=(DisplayName="Add", TemplateName="Add", Keywords="Sum,+"))
 struct CONTROLRIG_API FRigUnit_MathColorAdd : public FRigUnit_MathColorBinaryOp
 {
 	GENERATED_BODY()
@@ -70,7 +70,7 @@ struct CONTROLRIG_API FRigUnit_MathColorAdd : public FRigUnit_MathColorBinaryOp
 /**
  * Returns the difference of the two values
  */
-USTRUCT(meta=(DisplayName="Subtract", PrototypeName="Subtract", Keywords="-"))
+USTRUCT(meta=(DisplayName="Subtract", TemplateName="Subtract", Keywords="-"))
 struct CONTROLRIG_API FRigUnit_MathColorSub : public FRigUnit_MathColorBinaryOp
 {
 	GENERATED_BODY()
@@ -82,7 +82,7 @@ struct CONTROLRIG_API FRigUnit_MathColorSub : public FRigUnit_MathColorBinaryOp
 /**
  * Returns the product of the two values
  */
-USTRUCT(meta=(DisplayName="Multiply", PrototypeName="Multiply", Keywords="Product,*"))
+USTRUCT(meta=(DisplayName="Multiply", TemplateName="Multiply", Keywords="Product,*"))
 struct CONTROLRIG_API FRigUnit_MathColorMul : public FRigUnit_MathColorBinaryOp
 {
 	GENERATED_BODY()
@@ -99,7 +99,7 @@ struct CONTROLRIG_API FRigUnit_MathColorMul : public FRigUnit_MathColorBinaryOp
 /**
  * Linearly interpolates between A and B using the ratio T
  */
-USTRUCT(meta=(DisplayName="Interpolate", PrototypeName="Interpolate", Keywords="Lerp,Mix,Blend"))
+USTRUCT(meta=(DisplayName="Interpolate", TemplateName="Interpolate", Keywords="Lerp,Mix,Blend"))
 struct CONTROLRIG_API FRigUnit_MathColorLerp : public FRigUnit_MathColorBase
 {
 	GENERATED_BODY()
diff --git a/Engine/Plugins/Animation/ControlRig/Source/ControlRig/Private/Units/Math/RigUnit_MathFloat.h b/Engine/Plugins/Animation/ControlRig/Source/ControlRig/Private/Units/Math/RigUnit_MathFloat.h
--- a/Engine/Plugins/Animation/ControlRig/Source/ControlRig/Private/Units/Math/RigUnit_MathFloat.h
+++ b/Engine/Plugins/Animation/ControlRig/Source/ControlRig/Private/Units/Math/RigUnit_MathFloat.h
@@ -65,7 +65,7 @@ struct CONTROLRIG_API FRigUnit_MathFloatBinaryOp : public FRigUnit_MathFloatBase
 /**
  * Returns PI
  */
-USTRUCT(meta=(DisplayName="Pi", PrototypeName="Pi"))
+USTRUCT(meta=(DisplayName="Pi", TemplateName="Pi"))
 struct CONTROLRIG_API FRigUnit_MathFloatConstPi : public FRigUnit_MathFloatConstant
 {
 	GENERATED_BODY()
@@ -77,7 +77,7 @@ struct CONTROLRIG_API FRigUnit_MathFloatConstPi : public FRigUnit_MathFloatConst
 /**
  * Returns PI * 0.5
  */
-USTRUCT(meta=(DisplayName="Half Pi", PrototypeName="HalfPi"))
+USTRUCT(meta=(DisplayName="Half Pi", TemplateName="HalfPi"))
 struct CONTROLRIG_API FRigUnit_MathFloatConstHalfPi : public FRigUnit_MathFloatConstant
 {
 	GENERATED_BODY()
@@ -89,7 +89,7 @@ struct CONTROLRIG_API FRigUnit_MathFloatConstHalfPi : public FRigUnit_MathFloatC
 /**
  * Returns PI * 2.0
  */
-USTRUCT(meta=(DisplayName="Two Pi", PrototypeName="TwoPi", Keywords="Tau"))
+USTRUCT(meta=(DisplayName="Two Pi", TemplateName="TwoPi", Keywords="Tau"))
 struct CONTROLRIG_API FRigUnit_MathFloatConstTwoPi : public FRigUnit_MathFloatConstant
 {
 	GENERATED_BODY()
@@ -101,7 +101,7 @@ struct CONTROLRIG_API FRigUnit_MathFloatConstTwoPi : public FRigUnit_MathFloatCo
 /**
  * Returns E
  */
-USTRUCT(meta=(DisplayName="E", PrototypeName="E", Keywords="Euler"))
+USTRUCT(meta=(DisplayName="E", TemplateName="E", Keywords="Euler"))
 struct CONTROLRIG_API FRigUnit_MathFloatConstE : public FRigUnit_MathFloatConstant
 {
 	GENERATED_BODY()
@@ -113,7 +113,7 @@ struct CONTROLRIG_API FRigUnit_MathFloatConstE : public FRigUnit_MathFloatConsta
 /**
  * Returns the sum of the two values
  */
-USTRUCT(meta=(DisplayName="Add", PrototypeName="Add", Keywords="Sum,+"))
+USTRUCT(meta=(DisplayName="Add", TemplateName="Add", Keywords="Sum,+"))
 struct CONTROLRIG_API FRigUnit_MathFloatAdd : public FRigUnit_MathFloatBinaryOp
 {
 	GENERATED_BODY()
@@ -125,7 +125,7 @@ struct CONTROLRIG_API FRigUnit_MathFloatAdd : public FRigUnit_MathFloatBinaryOp
 /**
  * Returns the difference of the two values
  */
-USTRUCT(meta=(DisplayName="Subtract", PrototypeName="Subtract", Keywords="-"))
+USTRUCT(meta=(DisplayName="Subtract", TemplateName="Subtract", Keywords="-"))
 struct CONTROLRIG_API FRigUnit_MathFloatSub : public FRigUnit_MathFloatBinaryOp
 {
 	GENERATED_BODY()
@@ -137,7 +137,7 @@ struct CONTROLRIG_API FRigUnit_MathFloatSub : public FRigUnit_MathFloatBinaryOp
 /**
  * Returns the product of the two values
  */
-USTRUCT(meta=(DisplayName="Multiply", PrototypeName="Multiply", Keywords="Product,*"))
+USTRUCT(meta=(DisplayName="Multiply", TemplateName="Multiply", Keywords="Product,*"))
 struct CONTROLRIG_API FRigUnit_MathFloatMul : public FRigUnit_MathFloatBinaryOp
 {
 	GENERATED_BODY()
@@ -155,7 +155,7 @@ struct CONTROLRIG_API FRigUnit_MathFloatMul : public FRigUnit_MathFloatBinaryOp
 /**
  * Returns the division of the two values
  */
-USTRUCT(meta=(DisplayName="Divide", PrototypeName="Divide", Keywords="Division,Divisor,/"))
+USTRUCT(meta=(DisplayName="Divide", TemplateName="Divide", Keywords="Division,Divisor,/"))
 struct CONTROLRIG_API FRigUnit_MathFloatDiv : public FRigUnit_MathFloatBinaryOp
 {
 	GENERATED_BODY()
@@ -172,7 +172,7 @@ struct CONTROLRIG_API FRigUnit_MathFloatDiv : public FRigUnit_MathFloatBinaryOp
 /**
  * Returns the modulo of the two values
  */
-USTRUCT(meta=(DisplayName="Modulo", PrototypeName="Modulo", Keywords="%,fmod"))
+USTRUCT(meta=(DisplayName="Modulo", TemplateName="Modulo", Keywords="%,fmod"))
 struct CONTROLRIG_API FRigUnit_MathFloatMod : public FRigUnit_MathFloatBinaryOp
 {
 	GENERATED_BODY()
@@ -190,7 +190,7 @@ struct CONTROLRIG_API FRigUnit_MathFloatMod : public FRigUnit_MathFloatBinaryOp
 /**
  * Returns the smaller of the two values
  */
-USTRUCT(meta=(DisplayName="Minimum", PrototypeName="Minimum"))
+USTRUCT(meta=(DisplayName="Minimum", TemplateName="Minimum"))
 struct CONTROLRIG_API FRigUnit_MathFloatMin : public FRigUnit_MathFloatBinaryOp
 {
 	GENERATED_BODY()
@@ -202,7 +202,7 @@ struct CONTROLRIG_API FRigUnit_MathFloatMin : public FRigUnit_MathFloatBinaryOp
 /**
  * Returns the larger of the two values
  */
-USTRUCT(meta=(DisplayName="Maximum", PrototypeName="Maximum"))
+USTRUCT(meta=(DisplayName="Maximum", TemplateName="Maximum"))
 struct CONTROLRIG_API FRigUnit_MathFloatMax : public FRigUnit_MathFloatBinaryOp
 {
 	GENERATED_BODY()
@@ -214,7 +214,7 @@ struct CONTROLRIG_API FRigUnit_MathFloatMax : public FRigUnit_MathFloatBinaryOp
 /**
  * Returns the value of A raised to the power of B.
  */
-USTRUCT(meta=(DisplayName="Power", PrototypeName="Power"))
+USTRUCT(meta=(DisplayName="Power", TemplateName="Power"))
 struct CONTROLRIG_API FRigUnit_MathFloatPow : public FRigUnit_MathFloatBinaryOp
 {
 	GENERATED_BODY()
@@ -232,7 +232,7 @@ struct CONTROLRIG_API FRigUnit_MathFloatPow : public FRigUnit_MathFloatBinaryOp
 /**
  * Returns the square root of the given value
  */
-USTRUCT(meta=(DisplayName="Sqrt", PrototypeName="Sqrt", Keywords="Root,Square"))
+USTRUCT(meta=(DisplayName="Sqrt", TemplateName="Sqrt", Keywords="Root,Square"))
 struct CONTROLRIG_API FRigUnit_MathFloatSqrt : public FRigUnit_MathFloatUnaryOp
 {
 	GENERATED_BODY()
@@ -244,7 +244,7 @@ struct CONTROLRIG_API FRigUnit_MathFloatSqrt : public FRigUnit_MathFloatUnaryOp
 /**
  * Returns the negative value
  */
-USTRUCT(meta=(DisplayName="Negate", PrototypeName="Negate", Keywords="-,Abs"))
+USTRUCT(meta=(DisplayName="Negate", TemplateName="Negate", Keywords="-,Abs"))
 struct CONTROLRIG_API FRigUnit_MathFloatNegate : public FRigUnit_MathFloatUnaryOp
 {
 	GENERATED_BODY()
@@ -256,7 +256,7 @@ struct CONTROLRIG_API FRigUnit_MathFloatNegate : public FRigUnit_MathFloatUnaryO
 /**
  * Returns the absolute (positive) value
  */
-USTRUCT(meta=(DisplayName="Absolute", PrototypeName="Absolute", Keywords="Abs,Neg"))
+USTRUCT(meta=(DisplayName="Absolute", TemplateName="Absolute", Keywords="Abs,Neg"))
 struct CONTROLRIG_API FRigUnit_MathFloatAbs : public FRigUnit_MathFloatUnaryOp
 {
 	GENERATED_BODY()
@@ -268,7 +268,7 @@ struct CONTROLRIG_API FRigUnit_MathFloatAbs : public FRigUnit_MathFloatUnaryOp
 /**
  * Returns the closest lower full number (integer) of the value
  */
-USTRUCT(meta=(DisplayName="Floor", PrototypeName="Floor", Keywords="Round"))
+USTRUCT(meta=(DisplayName="Floor", TemplateName="Floor", Keywords="Round"))
 struct CONTROLRIG_API FRigUnit_MathFloatFloor : public FRigUnit_MathFloatBase
 {
 	GENERATED_BODY()
@@ -295,7 +295,7 @@ struct CONTROLRIG_API FRigUnit_MathFloatFloor : public FRigUnit_MathFloatBase
 /**
  * Returns the closest higher full number (integer) of the value
  */
-USTRUCT(meta=(DisplayName="Ceiling", PrototypeName="Ceiling", Keywords="Round"))
+USTRUCT(meta=(DisplayName="Ceiling", TemplateName="Ceiling", Keywords="Round"))
 struct CONTROLRIG_API FRigUnit_MathFloatCeil : public FRigUnit_MathFloatBase
 {
 	GENERATED_BODY()
@@ -322,7 +322,7 @@ struct CONTROLRIG_API FRigUnit_MathFloatCeil : public FRigUnit_MathFloatBase
 /**
  * Returns the closest higher full number (integer) of the value
  */
-USTRUCT(meta=(DisplayName="Round", PrototypeName="Round"))
+USTRUCT(meta=(DisplayName="Round", TemplateName="Round"))
 struct CONTROLRIG_API FRigUnit_MathFloatRound : public FRigUnit_MathFloatBase
 {
 	GENERATED_BODY()
@@ -349,7 +349,7 @@ struct CONTROLRIG_API FRigUnit_MathFloatRound : public FRigUnit_MathFloatBase
 /**
  * Returns the float cast to an int (this uses floor)
  */
-USTRUCT(meta=(DisplayName="To Int", PrototypeName="Convert"))
+USTRUCT(meta=(DisplayName="To Int", TemplateName="Convert"))
 struct CONTROLRIG_API FRigUnit_MathFloatToInt : public FRigUnit_MathFloatBase
 {
 	GENERATED_BODY()
@@ -373,7 +373,7 @@ struct CONTROLRIG_API FRigUnit_MathFloatToInt : public FRigUnit_MathFloatBase
 /**
  * Returns the sign of the value (+1 for >= 0.f, -1 for < 0.f)
  */
-USTRUCT(meta=(DisplayName="Sign", PrototypeName="Sign"))
+USTRUCT(meta=(DisplayName="Sign", TemplateName="Sign"))
 struct CONTROLRIG_API FRigUnit_MathFloatSign : public FRigUnit_MathFloatUnaryOp
 {
 	GENERATED_BODY()
@@ -385,7 +385,7 @@ struct CONTROLRIG_API FRigUnit_MathFloatSign : public FRigUnit_MathFloatUnaryOp
 /**
  * Clamps the given value within the range provided by minimum and maximum
  */
-USTRUCT(meta=(DisplayName="Clamp", PrototypeName="Clamp", Keywords="Range,Remap"))
+USTRUCT(meta=(DisplayName="Clamp", TemplateName="Clamp", Keywords="Range,Remap"))
 struct CONTROLRIG_API FRigUnit_MathFloatClamp : public FRigUnit_MathFloatBase
 {
 	GENERATED_BODY()
@@ -414,7 +414,7 @@ struct CONTROLRIG_API FRigUnit_MathFloatClamp : public FRigUnit_MathFloatBase
 /**
  * Linearly interpolates between A and B using the ratio T
  */
-USTRUCT(meta=(DisplayName="Interpolate", PrototypeName="Interpolate", Keywords="Lerp,Mix,Blend"))
+USTRUCT(meta=(DisplayName="Interpolate", TemplateName="Interpolate", Keywords="Lerp,Mix,Blend"))
 struct CONTROLRIG_API FRigUnit_MathFloatLerp : public FRigUnit_MathFloatBase
 {
 	GENERATED_BODY()
@@ -444,7 +444,7 @@ struct CONTROLRIG_API FRigUnit_MathFloatLerp : public FRigUnit_MathFloatBase
 /**
  * Remaps the given value from a source range to a target range.
  */
-USTRUCT(meta=(DisplayName="Remap", PrototypeName="Remap", Keywords="Rescale,Scale"))
+USTRUCT(meta=(DisplayName="Remap", TemplateName="Remap", Keywords="Rescale,Scale"))
 struct CONTROLRIG_API FRigUnit_MathFloatRemap : public FRigUnit_MathFloatBase
 {
 	GENERATED_BODY()
@@ -485,7 +485,7 @@ struct CONTROLRIG_API FRigUnit_MathFloatRemap : public FRigUnit_MathFloatBase
 /**
  * Returns true if the value A equals B
  */
-USTRUCT(meta=(DisplayName="Equals", PrototypeName="Equals", Keywords="Same,=="))
+USTRUCT(meta=(DisplayName="Equals", TemplateName="Equals", Keywords="Same,=="))
 struct CONTROLRIG_API FRigUnit_MathFloatEquals : public FRigUnit_MathFloatBase
 {
 	GENERATED_BODY()
@@ -511,7 +511,7 @@ struct CONTROLRIG_API FRigUnit_MathFloatEquals : public FRigUnit_MathFloatBase
 /**
  * Returns true if the value A does not equal B
  */
-USTRUCT(meta=(DisplayName="Not Equals", PrototypeName="NotEquals", Keywords="Different,!="))
+USTRUCT(meta=(DisplayName="Not Equals", TemplateName="NotEquals", Keywords="Different,!="))
 struct CONTROLRIG_API FRigUnit_MathFloatNotEquals : public FRigUnit_MathFloatBase
 {
 	GENERATED_BODY()
@@ -537,7 +537,7 @@ struct CONTROLRIG_API FRigUnit_MathFloatNotEquals : public FRigUnit_MathFloatBas
 /**
  * Returns true if the value A is greater than B
  */
-USTRUCT(meta=(DisplayName="Greater", PrototypeName="Greater", Keywords="Larger,Bigger,>"))
+USTRUCT(meta=(DisplayName="Greater", TemplateName="Greater", Keywords="Larger,Bigger,>"))
 struct CONTROLRIG_API FRigUnit_MathFloatGreater : public FRigUnit_MathFloatBase
 {
 	GENERATED_BODY()
@@ -564,7 +564,7 @@ struct CONTROLRIG_API FRigUnit_MathFloatGreater : public FRigUnit_MathFloatBase
 /**
  * Returns true if the value A is less than B
  */
-USTRUCT(meta=(DisplayName="Less", PrototypeName="Less", Keywords="Smaller,<"))
+USTRUCT(meta=(DisplayName="Less", TemplateName="Less", Keywords="Smaller,<"))
 struct CONTROLRIG_API FRigUnit_MathFloatLess : public FRigUnit_MathFloatBase
 {
 	GENERATED_BODY()
@@ -591,7 +591,7 @@ struct CONTROLRIG_API FRigUnit_MathFloatLess : public FRigUnit_MathFloatBase
 /**
  * Returns true if the value A is greater than or equal to B
  */
-USTRUCT(meta=(DisplayName="Greater Equal", PrototypeName="GreaterEqual", Keywords="Larger,Bigger,>="))
+USTRUCT(meta=(DisplayName="Greater Equal", TemplateName="GreaterEqual", Keywords="Larger,Bigger,>="))
 struct CONTROLRIG_API FRigUnit_MathFloatGreaterEqual : public FRigUnit_MathFloatBase
 {
 	GENERATED_BODY()
@@ -618,7 +618,7 @@ struct CONTROLRIG_API FRigUnit_MathFloatGreaterEqual : public FRigUnit_MathFloat
 /**
  * Returns true if the value A is less than or equal to B
  */
-USTRUCT(meta=(DisplayName="Less Equal", PrototypeName="LessEqual", Keywords="Smaller,<="))
+USTRUCT(meta=(DisplayName="Less Equal", TemplateName="LessEqual", Keywords="Smaller,<="))
 struct CONTROLRIG_API FRigUnit_MathFloatLessEqual : public FRigUnit_MathFloatBase
 {
 	GENERATED_BODY()
@@ -645,7 +645,7 @@ struct CONTROLRIG_API FRigUnit_MathFloatLessEqual : public FRigUnit_MathFloatBas
 /**
  * Returns true if the value is nearly zero
  */
-USTRUCT(meta=(DisplayName="Is Nearly Zero", PrototypeName="IsNearlyZero", Keywords="AlmostZero,0"))
+USTRUCT(meta=(DisplayName="Is Nearly Zero", TemplateName="IsNearlyZero", Keywords="AlmostZero,0"))
 struct CONTROLRIG_API FRigUnit_MathFloatIsNearlyZero : public FRigUnit_MathFloatBase
 {
 	GENERATED_BODY()
@@ -672,7 +672,7 @@ struct CONTROLRIG_API FRigUnit_MathFloatIsNearlyZero : public FRigUnit_MathFloat
 /**
  * Returns true if the value A is almost equal to B
  */
-USTRUCT(meta=(DisplayName="Is Nearly Equal", PrototypeName="IsNearlyEqual", Keywords="AlmostEqual,=="))
+USTRUCT(meta=(DisplayName="Is Nearly Equal", TemplateName="IsNearlyEqual", Keywords="AlmostEqual,=="))
 struct CONTROLRIG_API FRigUnit_MathFloatIsNearlyEqual : public FRigUnit_MathFloatBase
 {
 	GENERATED_BODY()
@@ -702,7 +702,7 @@ struct CONTROLRIG_API FRigUnit_MathFloatIsNearlyEqual : public FRigUnit_MathFloa
 /**
  * Return one of the two values based on the condition
  */
-USTRUCT(meta=(DisplayName="Select", PrototypeName="Select", Keywords="Pick,If", Deprecated = "4.26.0"))
+USTRUCT(meta=(DisplayName="Select", TemplateName="Select", Keywords="Pick,If", Deprecated = "4.26.0"))
 struct CONTROLRIG_API FRigUnit_MathFloatSelectBool : public FRigUnit_MathFloatBase
 {
 	GENERATED_BODY()
@@ -735,7 +735,7 @@ struct CONTROLRIG_API FRigUnit_MathFloatSelectBool : public FRigUnit_MathFloatBa
 /**
  * Returns the degrees of a given value in radians
  */
-USTRUCT(meta=(DisplayName="Degrees", PrototypeName="Degrees"))
+USTRUCT(meta=(DisplayName="Degrees", TemplateName="Degrees"))
 struct CONTROLRIG_API FRigUnit_MathFloatDeg : public FRigUnit_MathFloatUnaryOp
 {
 	GENERATED_BODY()
@@ -747,7 +747,7 @@ struct CONTROLRIG_API FRigUnit_MathFloatDeg : public FRigUnit_MathFloatUnaryOp
 /**
  * Returns the radians of a given value in degrees
  */
-USTRUCT(meta=(DisplayName="Radians", PrototypeName="Radians"))
+USTRUCT(meta=(DisplayName="Radians", TemplateName="Radians"))
 struct CONTROLRIG_API FRigUnit_MathFloatRad : public FRigUnit_MathFloatUnaryOp
 {
 	GENERATED_BODY()
@@ -759,7 +759,7 @@ struct CONTROLRIG_API FRigUnit_MathFloatRad : public FRigUnit_MathFloatUnaryOp
 /**
  * Returns the sinus value of the given value (in radians)
  */
-USTRUCT(meta=(DisplayName="Sin", PrototypeName="Sin"))
+USTRUCT(meta=(DisplayName="Sin", TemplateName="Sin"))
 struct CONTROLRIG_API FRigUnit_MathFloatSin : public FRigUnit_MathFloatUnaryOp
 {
 	GENERATED_BODY()
@@ -771,7 +771,7 @@ struct CONTROLRIG_API FRigUnit_MathFloatSin : public FRigUnit_MathFloatUnaryOp
 /**
  * Returns the cosinus value of the given value (in radians)
  */
-USTRUCT(meta=(DisplayName="Cos", PrototypeName="Cos"))
+USTRUCT(meta=(DisplayName="Cos", TemplateName="Cos"))
 struct CONTROLRIG_API FRigUnit_MathFloatCos : public FRigUnit_MathFloatUnaryOp
 {
 	GENERATED_BODY()
@@ -783,7 +783,7 @@ struct CONTROLRIG_API FRigUnit_MathFloatCos : public FRigUnit_MathFloatUnaryOp
 /**
  * Returns the tangens value of the given value (in radians)
  */
-USTRUCT(meta=(DisplayName="Tan", PrototypeName="Tan"))
+USTRUCT(meta=(DisplayName="Tan", TemplateName="Tan"))
 struct CONTROLRIG_API FRigUnit_MathFloatTan : public FRigUnit_MathFloatUnaryOp
 {
 	GENERATED_BODY()
@@ -795,7 +795,7 @@ struct CONTROLRIG_API FRigUnit_MathFloatTan : public FRigUnit_MathFloatUnaryOp
 /**
  * Returns the inverse sinus value (in radians) of the given value
  */
-USTRUCT(meta=(DisplayName="Asin", PrototypeName="Asin", Keywords="Arcsin"))
+USTRUCT(meta=(DisplayName="Asin", TemplateName="Asin", Keywords="Arcsin"))
 struct CONTROLRIG_API FRigUnit_MathFloatAsin : public FRigUnit_MathFloatUnaryOp
 {
 	GENERATED_BODY()
@@ -807,7 +807,7 @@ struct CONTROLRIG_API FRigUnit_MathFloatAsin : public FRigUnit_MathFloatUnaryOp
 /**
  * Returns the inverse cosinus value (in radians) of the given value
  */
-USTRUCT(meta=(DisplayName="Acos", PrototypeName="Acos", Keywords="Arccos"))
+USTRUCT(meta=(DisplayName="Acos", TemplateName="Acos", Keywords="Arccos"))
 struct CONTROLRIG_API FRigUnit_MathFloatAcos : public FRigUnit_MathFloatUnaryOp
 {
 	GENERATED_BODY()
@@ -819,7 +819,7 @@ struct CONTROLRIG_API FRigUnit_MathFloatAcos : public FRigUnit_MathFloatUnaryOp
 /**
  * Returns the inverse tangens value (in radians) of the given value
  */
-USTRUCT(meta=(DisplayName="Atan", PrototypeName="Atan", Keywords="Arctan"))
+USTRUCT(meta=(DisplayName="Atan", TemplateName="Atan", Keywords="Arctan"))
 struct CONTROLRIG_API FRigUnit_MathFloatAtan : public FRigUnit_MathFloatUnaryOp
 {
 	GENERATED_BODY()
@@ -831,7 +831,7 @@ struct CONTROLRIG_API FRigUnit_MathFloatAtan : public FRigUnit_MathFloatUnaryOp
 /**
  * Computes the angles alpha, beta and gamma (in radians) between the three sides A, B and C
  */
-USTRUCT(meta = (DisplayName = "Law Of Cosine", PrototypeName="LawOfCosine"))
+USTRUCT(meta = (DisplayName = "Law Of Cosine", TemplateName="LawOfCosine"))
 struct CONTROLRIG_API FRigUnit_MathFloatLawOfCosine : public FRigUnit_MathFloatBase
 {
 	GENERATED_BODY()
@@ -870,7 +870,7 @@ struct CONTROLRIG_API FRigUnit_MathFloatLawOfCosine : public FRigUnit_MathFloatB
 /**
  * Computes the base-e exponential of the given value 
  */
-USTRUCT(meta = (DisplayName = "Exponential", PrototypeName="Exponential"))
+USTRUCT(meta = (DisplayName = "Exponential", TemplateName="Exponential"))
 struct CONTROLRIG_API FRigUnit_MathFloatExponential : public FRigUnit_MathFloatUnaryOp
 {
 	GENERATED_BODY()
diff --git a/Engine/Plugins/Animation/ControlRig/Source/ControlRig/Private/Units/Math/RigUnit_MathInt.h b/Engine/Plugins/Animation/ControlRig/Source/ControlRig/Private/Units/Math/RigUnit_MathInt.h
--- a/Engine/Plugins/Animation/ControlRig/Source/ControlRig/Private/Units/Math/RigUnit_MathInt.h
+++ b/Engine/Plugins/Animation/ControlRig/Source/ControlRig/Private/Units/Math/RigUnit_MathInt.h
@@ -51,7 +51,7 @@ struct CONTROLRIG_API FRigUnit_MathIntBinaryOp : public FRigUnit_MathIntBase
 /**
  * Returns the sum of the two values
  */
-USTRUCT(meta=(DisplayName="Add", PrototypeName="Add", Keywords="Sum,+"))
+USTRUCT(meta=(DisplayName="Add", TemplateName="Add", Keywords="Sum,+"))
 struct CONTROLRIG_API FRigUnit_MathIntAdd : public FRigUnit_MathIntBinaryOp
 {
 	GENERATED_BODY()
@@ -63,7 +63,7 @@ struct CONTROLRIG_API FRigUnit_MathIntAdd : public FRigUnit_MathIntBinaryOp
 /**
  * Returns the difference of the two values
  */
-USTRUCT(meta=(DisplayName="Subtract", PrototypeName="Subtract", Keywords="-"))
+USTRUCT(meta=(DisplayName="Subtract", TemplateName="Subtract", Keywords="-"))
 struct CONTROLRIG_API FRigUnit_MathIntSub : public FRigUnit_MathIntBinaryOp
 {
 	GENERATED_BODY()
@@ -75,7 +75,7 @@ struct CONTROLRIG_API FRigUnit_MathIntSub : public FRigUnit_MathIntBinaryOp
 /**
  * Returns the product of the two values
  */
-USTRUCT(meta=(DisplayName="Multiply", PrototypeName="Multiply", Keywords="Product,*"))
+USTRUCT(meta=(DisplayName="Multiply", TemplateName="Multiply", Keywords="Product,*"))
 struct CONTROLRIG_API FRigUnit_MathIntMul : public FRigUnit_MathIntBinaryOp
 {
 	GENERATED_BODY()
@@ -93,7 +93,7 @@ struct CONTROLRIG_API FRigUnit_MathIntMul : public FRigUnit_MathIntBinaryOp
 /**
  * Returns the division of the two values
  */
-USTRUCT(meta=(DisplayName="Divide", PrototypeName="Divide", Keywords="Division,Divisor,/"))
+USTRUCT(meta=(DisplayName="Divide", TemplateName="Divide", Keywords="Division,Divisor,/"))
 struct CONTROLRIG_API FRigUnit_MathIntDiv : public FRigUnit_MathIntBinaryOp
 {
 	GENERATED_BODY()
@@ -110,7 +110,7 @@ struct CONTROLRIG_API FRigUnit_MathIntDiv : public FRigUnit_MathIntBinaryOp
 /**
  * Returns the modulo of the two values
  */
-USTRUCT(meta=(DisplayName="Modulo", PrototypeName="Modulo", Keywords="%,fmod"))
+USTRUCT(meta=(DisplayName="Modulo", TemplateName="Modulo", Keywords="%,fmod"))
 struct CONTROLRIG_API FRigUnit_MathIntMod : public FRigUnit_MathIntBinaryOp
 {
 	GENERATED_BODY()
@@ -128,7 +128,7 @@ struct CONTROLRIG_API FRigUnit_MathIntMod : public FRigUnit_MathIntBinaryOp
 /**
  * Returns the smaller of the two values
  */
-USTRUCT(meta=(DisplayName="Minimum", PrototypeName="Minimum"))
+USTRUCT(meta=(DisplayName="Minimum", TemplateName="Minimum"))
 struct CONTROLRIG_API FRigUnit_MathIntMin : public FRigUnit_MathIntBinaryOp
 {
 	GENERATED_BODY()
@@ -140,7 +140,7 @@ struct CONTROLRIG_API FRigUnit_MathIntMin : public FRigUnit_MathIntBinaryOp
 /**
  * Returns the larger of the two values
  */
-USTRUCT(meta=(DisplayName="Maximum", PrototypeName="Maximum"))
+USTRUCT(meta=(DisplayName="Maximum", TemplateName="Maximum"))
 struct CONTROLRIG_API FRigUnit_MathIntMax : public FRigUnit_MathIntBinaryOp
 {
 	GENERATED_BODY()
@@ -152,7 +152,7 @@ struct CONTROLRIG_API FRigUnit_MathIntMax : public FRigUnit_MathIntBinaryOp
 /**
  * Returns the value of A raised to the power of B.
  */
-USTRUCT(meta=(DisplayName="Power", PrototypeName="Power"))
+USTRUCT(meta=(DisplayName="Power", TemplateName="Power"))
 struct CONTROLRIG_API FRigUnit_MathIntPow : public FRigUnit_MathIntBinaryOp
 {
 	GENERATED_BODY()
@@ -170,7 +170,7 @@ struct CONTROLRIG_API FRigUnit_MathIntPow : public FRigUnit_MathIntBinaryOp
 /**
  * Returns the negative value
  */
-USTRUCT(meta=(DisplayName="Negate", PrototypeName="Negate", Keywords="-,Abs"))
+USTRUCT(meta=(DisplayName="Negate", TemplateName="Negate", Keywords="-,Abs"))
 struct CONTROLRIG_API FRigUnit_MathIntNegate : public FRigUnit_MathIntUnaryOp
 {
 	GENERATED_BODY()
@@ -182,7 +182,7 @@ struct CONTROLRIG_API FRigUnit_MathIntNegate : public FRigUnit_MathIntUnaryOp
 /**
  * Returns the absolute (positive) value
  */
-USTRUCT(meta=(DisplayName="Absolute", PrototypeName="Absolute", Keywords="Abs,Neg"))
+USTRUCT(meta=(DisplayName="Absolute", TemplateName="Absolute", Keywords="Abs,Neg"))
 struct CONTROLRIG_API FRigUnit_MathIntAbs : public FRigUnit_MathIntUnaryOp
 {
 	GENERATED_BODY()
@@ -194,7 +194,7 @@ struct CONTROLRIG_API FRigUnit_MathIntAbs : public FRigUnit_MathIntUnaryOp
 /**
  * Returns the int cast to a float
  */
-USTRUCT(meta=(DisplayName="To Float", PrototypeName="Convert"))
+USTRUCT(meta=(DisplayName="To Float", TemplateName="Convert"))
 struct CONTROLRIG_API FRigUnit_MathIntToFloat : public FRigUnit_MathIntBase
 {
 	GENERATED_BODY()
@@ -218,7 +218,7 @@ struct CONTROLRIG_API FRigUnit_MathIntToFloat : public FRigUnit_MathIntBase
 /**
  * Returns the sign of the value (+1 for >= 0, -1 for < 0)
  */
-USTRUCT(meta=(DisplayName="Sign", PrototypeName="Sign"))
+USTRUCT(meta=(DisplayName="Sign", TemplateName="Sign"))
 struct CONTROLRIG_API FRigUnit_MathIntSign : public FRigUnit_MathIntUnaryOp
 {
 	GENERATED_BODY()
@@ -230,7 +230,7 @@ struct CONTROLRIG_API FRigUnit_MathIntSign : public FRigUnit_MathIntUnaryOp
 /**
  * Clamps the given value within the range provided by minimum and maximum
  */
-USTRUCT(meta=(DisplayName="Clamp", PrototypeName="Clamp", Keywords="Range,Remap"))
+USTRUCT(meta=(DisplayName="Clamp", TemplateName="Clamp", Keywords="Range,Remap"))
 struct CONTROLRIG_API FRigUnit_MathIntClamp : public FRigUnit_MathIntBase
 {
 	GENERATED_BODY()
@@ -259,7 +259,7 @@ struct CONTROLRIG_API FRigUnit_MathIntClamp : public FRigUnit_MathIntBase
 /**
  * Returns true if the value A equals B
  */
-USTRUCT(meta=(DisplayName="Equals", PrototypeName="Equals", Keywords="Same,=="))
+USTRUCT(meta=(DisplayName="Equals", TemplateName="Equals", Keywords="Same,=="))
 struct CONTROLRIG_API FRigUnit_MathIntEquals : public FRigUnit_MathIntBase
 {
 	GENERATED_BODY()
@@ -285,7 +285,7 @@ struct CONTROLRIG_API FRigUnit_MathIntEquals : public FRigUnit_MathIntBase
 /**
  * Returns true if the value A does not equal B
  */
-USTRUCT(meta=(DisplayName="Not Equals", PrototypeName="NotEquals", Keywords="Different,!="))
+USTRUCT(meta=(DisplayName="Not Equals", TemplateName="NotEquals", Keywords="Different,!="))
 struct CONTROLRIG_API FRigUnit_MathIntNotEquals : public FRigUnit_MathIntBase
 {
 	GENERATED_BODY()
@@ -311,7 +311,7 @@ struct CONTROLRIG_API FRigUnit_MathIntNotEquals : public FRigUnit_MathIntBase
 /**
  * Returns true if the value A is greater than B
  */
-USTRUCT(meta=(DisplayName="Greater", PrototypeName="Greater", Keywords="Larger,Bigger,>"))
+USTRUCT(meta=(DisplayName="Greater", TemplateName="Greater", Keywords="Larger,Bigger,>"))
 struct CONTROLRIG_API FRigUnit_MathIntGreater : public FRigUnit_MathIntBase
 {
 	GENERATED_BODY()
@@ -338,7 +338,7 @@ struct CONTROLRIG_API FRigUnit_MathIntGreater : public FRigUnit_MathIntBase
 /**
  * Returns true if the value A is less than B
  */
-USTRUCT(meta=(DisplayName="Less", PrototypeName="Less", Keywords="Smaller,<"))
+USTRUCT(meta=(DisplayName="Less", TemplateName="Less", Keywords="Smaller,<"))
 struct CONTROLRIG_API FRigUnit_MathIntLess : public FRigUnit_MathIntBase
 {
 	GENERATED_BODY()
@@ -365,7 +365,7 @@ struct CONTROLRIG_API FRigUnit_MathIntLess : public FRigUnit_MathIntBase
 /**
  * Returns true if the value A is greater than or equal to B
  */
-USTRUCT(meta=(DisplayName="Greater Equal", PrototypeName="GreaterEqual", Keywords="Larger,Bigger,>="))
+USTRUCT(meta=(DisplayName="Greater Equal", TemplateName="GreaterEqual", Keywords="Larger,Bigger,>="))
 struct CONTROLRIG_API FRigUnit_MathIntGreaterEqual : public FRigUnit_MathIntBase
 {
 	GENERATED_BODY()
@@ -392,7 +392,7 @@ struct CONTROLRIG_API FRigUnit_MathIntGreaterEqual : public FRigUnit_MathIntBase
 /**
  * Returns true if the value A is less than or equal to B
  */
-USTRUCT(meta=(DisplayName="Less Equal", PrototypeName="LessEqual", Keywords="Smaller,<="))
+USTRUCT(meta=(DisplayName="Less Equal", TemplateName="LessEqual", Keywords="Smaller,<="))
 struct CONTROLRIG_API FRigUnit_MathIntLessEqual : public FRigUnit_MathIntBase
 {
 	GENERATED_BODY()
diff --git a/Engine/Plugins/Animation/ControlRig/Source/ControlRig/Private/Units/Math/RigUnit_MathMatrix.h b/Engine/Plugins/Animation/ControlRig/Source/ControlRig/Private/Units/Math/RigUnit_MathMatrix.h
--- a/Engine/Plugins/Animation/ControlRig/Source/ControlRig/Private/Units/Math/RigUnit_MathMatrix.h
+++ b/Engine/Plugins/Animation/ControlRig/Source/ControlRig/Private/Units/Math/RigUnit_MathMatrix.h
@@ -173,7 +173,7 @@ struct CONTROLRIG_API FRigUnit_MathMatrixFromVectors : public FRigUnit_MathMatri
 /**
  * Returns the product of the two values
  */
-USTRUCT(meta=(DisplayName="Multiply", PrototypeName="Multiply", Keywords="Product,*,Global"))
+USTRUCT(meta=(DisplayName="Multiply", TemplateName="Multiply", Keywords="Product,*,Global"))
 struct CONTROLRIG_API FRigUnit_MathMatrixMul : public FRigUnit_MathMatrixBinaryOp
 {
 	GENERATED_BODY()
@@ -184,7 +184,7 @@ struct CONTROLRIG_API FRigUnit_MathMatrixMul : public FRigUnit_MathMatrixBinaryO
 /**
  * Returns the inverse value
  */
-USTRUCT(meta=(DisplayName="Inverse", PrototypeName="Inverse"))
+USTRUCT(meta=(DisplayName="Inverse", TemplateName="Inverse"))
 struct CONTROLRIG_API FRigUnit_MathMatrixInverse : public FRigUnit_MathMatrixUnaryOp
 {
 	GENERATED_BODY()
diff --git a/Engine/Plugins/Animation/ControlRig/Source/ControlRig/Private/Units/Math/RigUnit_MathQuaternion.h b/Engine/Plugins/Animation/ControlRig/Source/ControlRig/Private/Units/Math/RigUnit_MathQuaternion.h
--- a/Engine/Plugins/Animation/ControlRig/Source/ControlRig/Private/Units/Math/RigUnit_MathQuaternion.h
+++ b/Engine/Plugins/Animation/ControlRig/Source/ControlRig/Private/Units/Math/RigUnit_MathQuaternion.h
@@ -52,7 +52,7 @@ struct CONTROLRIG_API FRigUnit_MathQuaternionBinaryOp : public FRigUnit_MathQuat
 /**
  * Makes a quaternion from an axis and an angle in radians
  */
-USTRUCT(meta=(DisplayName="From Axis And Angle", PrototypeName="FromAxisAndAngle", Keywords="Make,Construct"))
+USTRUCT(meta=(DisplayName="From Axis And Angle", TemplateName="FromAxisAndAngle", Keywords="Make,Construct"))
 struct CONTROLRIG_API FRigUnit_MathQuaternionFromAxisAndAngle : public FRigUnit_MathQuaternionBase
 {
 	GENERATED_BODY()
@@ -80,7 +80,7 @@ struct CONTROLRIG_API FRigUnit_MathQuaternionFromAxisAndAngle : public FRigUnit_
 /**
  * Makes a quaternion from euler values in degrees
  */
-USTRUCT(meta=(DisplayName="From Euler", PrototypeName="FromEuler", Keywords="Make,Construct"))
+USTRUCT(meta=(DisplayName="From Euler", TemplateName="FromEuler", Keywords="Make,Construct"))
 struct CONTROLRIG_API FRigUnit_MathQuaternionFromEuler : public FRigUnit_MathQuaternionBase
 {
 	GENERATED_BODY()
@@ -108,7 +108,7 @@ struct CONTROLRIG_API FRigUnit_MathQuaternionFromEuler : public FRigUnit_MathQua
 /**
  * Makes a quaternion from a rotator
  */
-USTRUCT(meta=(DisplayName="From Rotator", PrototypeName="FromRotator", Keywords="Make,Construct"))
+USTRUCT(meta=(DisplayName="From Rotator", TemplateName="FromRotator", Keywords="Make,Construct"))
 struct CONTROLRIG_API FRigUnit_MathQuaternionFromRotator : public FRigUnit_MathQuaternionBase
 {
 	GENERATED_BODY()
@@ -132,7 +132,7 @@ struct CONTROLRIG_API FRigUnit_MathQuaternionFromRotator : public FRigUnit_MathQ
 /**
  * Makes a quaternion from two vectors, representing the shortest rotation between the two vectors.
  */
-USTRUCT(meta=(DisplayName="From Two Vectors", PrototypeName="FromTwoVectors", Keywords="Make,Construct"))
+USTRUCT(meta=(DisplayName="From Two Vectors", TemplateName="FromTwoVectors", Keywords="Make,Construct"))
 struct CONTROLRIG_API FRigUnit_MathQuaternionFromTwoVectors : public FRigUnit_MathQuaternionBase
 {
 	GENERATED_BODY()
@@ -159,7 +159,7 @@ struct CONTROLRIG_API FRigUnit_MathQuaternionFromTwoVectors : public FRigUnit_Ma
 /**
  * Retrieves the axis and angle of a quaternion in radians
  */
-USTRUCT(meta=(DisplayName="To Axis And Angle", PrototypeName="ToAxisAndAngle", Keywords="Make,Construct,GetAxis,GetAngle"))
+USTRUCT(meta=(DisplayName="To Axis And Angle", TemplateName="ToAxisAndAngle", Keywords="Make,Construct,GetAxis,GetAngle"))
 struct CONTROLRIG_API FRigUnit_MathQuaternionToAxisAndAngle : public FRigUnit_MathQuaternionBase
 {
 	GENERATED_BODY()
@@ -187,7 +187,7 @@ struct CONTROLRIG_API FRigUnit_MathQuaternionToAxisAndAngle : public FRigUnit_Ma
 /**
  * Scales a quaternion's angle
  */
-USTRUCT(meta=(DisplayName="Scale", PrototypeName="Scale", Keywords="Multiply,Angle,Scale", Constant))
+USTRUCT(meta=(DisplayName="Scale", TemplateName="Scale", Keywords="Multiply,Angle,Scale", Constant))
 struct CONTROLRIG_API FRigUnit_MathQuaternionScale : public FRigUnit_MathQuaternionBase
 {
 	GENERATED_BODY()
@@ -211,7 +211,7 @@ struct CONTROLRIG_API FRigUnit_MathQuaternionScale : public FRigUnit_MathQuatern
 /**
  * Retrieves the euler angles in degrees
  */
-USTRUCT(meta=(DisplayName="To Euler", PrototypeName="ToEuler", Keywords="Make,Construct"))
+USTRUCT(meta=(DisplayName="To Euler", TemplateName="ToEuler", Keywords="Make,Construct"))
 struct CONTROLRIG_API FRigUnit_MathQuaternionToEuler : public FRigUnit_MathQuaternionBase
 {
 	GENERATED_BODY()
@@ -239,7 +239,7 @@ struct CONTROLRIG_API FRigUnit_MathQuaternionToEuler : public FRigUnit_MathQuate
 /**
  * Retrieves the rotator
  */
-USTRUCT(meta=(DisplayName="To Rotator", PrototypeName="ToRotator", Keywords="Make,Construct"))
+USTRUCT(meta=(DisplayName="To Rotator", TemplateName="ToRotator", Keywords="Make,Construct"))
 struct CONTROLRIG_API FRigUnit_MathQuaternionToRotator : public FRigUnit_MathQuaternionBase
 {
 	GENERATED_BODY()
@@ -263,7 +263,7 @@ struct CONTROLRIG_API FRigUnit_MathQuaternionToRotator : public FRigUnit_MathQua
 /**
  * Returns the product of the two values
  */
-USTRUCT(meta=(DisplayName="Multiply", PrototypeName="Multiply", Keywords="Product,*"))
+USTRUCT(meta=(DisplayName="Multiply", TemplateName="Multiply", Keywords="Product,*"))
 struct CONTROLRIG_API FRigUnit_MathQuaternionMul : public FRigUnit_MathQuaternionBinaryOp
 {
 	GENERATED_BODY()
@@ -274,7 +274,7 @@ struct CONTROLRIG_API FRigUnit_MathQuaternionMul : public FRigUnit_MathQuaternio
 /**
  * Returns the negative value
  */
-USTRUCT(meta=(DisplayName="Inverse", PrototypeName="Inverse"))
+USTRUCT(meta=(DisplayName="Inverse", TemplateName="Inverse"))
 struct CONTROLRIG_API FRigUnit_MathQuaternionInverse : public FRigUnit_MathQuaternionUnaryOp
 {
 	GENERATED_BODY()
@@ -285,7 +285,7 @@ struct CONTROLRIG_API FRigUnit_MathQuaternionInverse : public FRigUnit_MathQuate
 /**
  * Linearly interpolates between A and B using the ratio T
  */
-USTRUCT(meta=(DisplayName="Interpolate", PrototypeName="Interpolate", Keywords="Lerp,Mix,Blend,Slerp,SphericalInterpolate"))
+USTRUCT(meta=(DisplayName="Interpolate", TemplateName="Interpolate", Keywords="Lerp,Mix,Blend,Slerp,SphericalInterpolate"))
 struct CONTROLRIG_API FRigUnit_MathQuaternionSlerp : public FRigUnit_MathQuaternionBase
 {
 	GENERATED_BODY()
@@ -315,7 +315,7 @@ struct CONTROLRIG_API FRigUnit_MathQuaternionSlerp : public FRigUnit_MathQuatern
 /**
  * Returns true if the value A equals B
  */
-USTRUCT(meta=(DisplayName="Equals", PrototypeName="Equals", Keywords="Same,=="))
+USTRUCT(meta=(DisplayName="Equals", TemplateName="Equals", Keywords="Same,=="))
 struct CONTROLRIG_API FRigUnit_MathQuaternionEquals : public FRigUnit_MathQuaternionBase
 {
 	GENERATED_BODY()
@@ -342,7 +342,7 @@ struct CONTROLRIG_API FRigUnit_MathQuaternionEquals : public FRigUnit_MathQuater
 /**
  * Returns true if the value A does not equal B
  */
-USTRUCT(meta=(DisplayName="Not Equals", PrototypeName="NotEquals", Keywords="Different,!="))
+USTRUCT(meta=(DisplayName="Not Equals", TemplateName="NotEquals", Keywords="Different,!="))
 struct CONTROLRIG_API FRigUnit_MathQuaternionNotEquals : public FRigUnit_MathQuaternionBase
 {
 	GENERATED_BODY()
@@ -369,7 +369,7 @@ struct CONTROLRIG_API FRigUnit_MathQuaternionNotEquals : public FRigUnit_MathQua
 /**
  * Return one of the two values based on the condition
  */
-USTRUCT(meta=(DisplayName="Select", PrototypeName="Select", Keywords="Pick,If", Deprecated = "4.26.0"))
+USTRUCT(meta=(DisplayName="Select", TemplateName="Select", Keywords="Pick,If", Deprecated = "4.26.0"))
 struct CONTROLRIG_API FRigUnit_MathQuaternionSelectBool : public FRigUnit_MathQuaternionBase
 {
 	GENERATED_BODY()
@@ -402,7 +402,7 @@ struct CONTROLRIG_API FRigUnit_MathQuaternionSelectBool : public FRigUnit_MathQu
 /**
  * Returns the dot product between two quaternions
  */
-USTRUCT(meta=(DisplayName="Dot", PrototypeName="Dot,|"))
+USTRUCT(meta=(DisplayName="Dot", TemplateName="Dot,|"))
 struct CONTROLRIG_API FRigUnit_MathQuaternionDot : public FRigUnit_MathQuaternionBase
 {
 	GENERATED_BODY()
@@ -429,7 +429,7 @@ struct CONTROLRIG_API FRigUnit_MathQuaternionDot : public FRigUnit_MathQuaternio
 /**
  * Returns the normalized quaternion
  */
-USTRUCT(meta=(DisplayName="Unit", PrototypeName="Unit", Keywords="Normalize"))
+USTRUCT(meta=(DisplayName="Unit", TemplateName="Unit", Keywords="Normalize"))
 struct CONTROLRIG_API FRigUnit_MathQuaternionUnit : public FRigUnit_MathQuaternionUnaryOp
 {
 	GENERATED_BODY()
@@ -440,7 +440,7 @@ struct CONTROLRIG_API FRigUnit_MathQuaternionUnit : public FRigUnit_MathQuaterni
 /**
  * Rotates a given vector by the quaternion
  */
-USTRUCT(meta=(DisplayName="Rotate", PrototypeName="Multiply", Keywords="Transform"))
+USTRUCT(meta=(DisplayName="Rotate", TemplateName="Multiply", Keywords="Transform"))
 struct CONTROLRIG_API FRigUnit_MathQuaternionRotateVector : public FRigUnit_MathQuaternionBase
 {
 	GENERATED_BODY()
@@ -467,7 +467,7 @@ struct CONTROLRIG_API FRigUnit_MathQuaternionRotateVector : public FRigUnit_Math
 /**
  * Rotates a given vector by the quaternion
  */
-USTRUCT(meta = (DisplayName = "Axis", PrototypeName = "Axis", Keywords = "GetAxis,xAxis,yAxis,zAxis"))
+USTRUCT(meta = (DisplayName = "Axis", TemplateName="Axis", Keywords = "GetAxis,xAxis,yAxis,zAxis"))
 struct CONTROLRIG_API FRigUnit_MathQuaternionGetAxis: public FRigUnit_MathQuaternionBase
 {
 	GENERATED_BODY()
diff --git a/Engine/Plugins/Animation/ControlRig/Source/ControlRig/Private/Units/Math/RigUnit_MathRBFInterpolate.h b/Engine/Plugins/Animation/ControlRig/Source/ControlRig/Private/Units/Math/RigUnit_MathRBFInterpolate.h
--- a/Engine/Plugins/Animation/ControlRig/Source/ControlRig/Private/Units/Math/RigUnit_MathRBFInterpolate.h
+++ b/Engine/Plugins/Animation/ControlRig/Source/ControlRig/Private/Units/Math/RigUnit_MathRBFInterpolate.h
@@ -182,7 +182,7 @@ struct FMathRBFInterpolateQuatFloat_Target
 	float Value = 0.0f;
 };
 
-USTRUCT(meta = (Abstract, DisplayName = "RBF Quaternion to Float", PrototypeName="RBFInterpolateQuatToFloat"))
+USTRUCT(meta = (Abstract, DisplayName = "RBF Quaternion to Float", TemplateName="RBFInterpolateQuatToFloat"))
 struct CONTROLRIG_API FRigUnit_MathRBFInterpolateQuatFloat : 
 	public FRigUnit_MathRBFInterpolateQuatBase
 {
@@ -211,7 +211,7 @@ struct FMathRBFInterpolateQuatVector_Target
 	FVector Value = FVector::ZeroVector;
 };
 
-USTRUCT(meta = (Abstract, DisplayName = "RBF Quaternion to Vector", PrototypeName = "RBFInterpolateQuatToVector"))
+USTRUCT(meta = (Abstract, DisplayName = "RBF Quaternion to Vector", TemplateName="RBFInterpolateQuatToVector"))
 struct CONTROLRIG_API FRigUnit_MathRBFInterpolateQuatVector :
 	public FRigUnit_MathRBFInterpolateQuatBase
 {
@@ -240,7 +240,7 @@ struct FMathRBFInterpolateQuatColor_Target
 	FLinearColor Value = FLinearColor::Transparent;
 };
 
-USTRUCT(meta = (Abstract, DisplayName = "RBF Quaternion to Color", PrototypeName = "RBFInterpolateQuatToColor"))
+USTRUCT(meta = (Abstract, DisplayName = "RBF Quaternion to Color", TemplateName="RBFInterpolateQuatToColor"))
 struct CONTROLRIG_API FRigUnit_MathRBFInterpolateQuatColor :
 	public FRigUnit_MathRBFInterpolateQuatBase
 {
@@ -269,7 +269,7 @@ struct FMathRBFInterpolateQuatQuat_Target
 	FQuat Value = FQuat::Identity;
 };
 
-USTRUCT(meta = (Abstract, DisplayName = "RBF Quaternion to Quaternion", PrototypeName = "RBFInterpolateQuatToQuat"))
+USTRUCT(meta = (Abstract, DisplayName = "RBF Quaternion to Quaternion", TemplateName="RBFInterpolateQuatToQuat"))
 struct CONTROLRIG_API FRigUnit_MathRBFInterpolateQuatQuat :
 	public FRigUnit_MathRBFInterpolateQuatBase
 {
@@ -298,7 +298,7 @@ struct FMathRBFInterpolateQuatXform_Target
 	FTransform Value = FTransform::Identity;
 };
 
-USTRUCT(meta = (Abstract, DisplayName = "RBF Quaternion to Transform", PrototypeName = "RBFInterpolateQuatToXform"))
+USTRUCT(meta = (Abstract, DisplayName = "RBF Quaternion to Transform", TemplateName="RBFInterpolateQuatToXform"))
 struct CONTROLRIG_API FRigUnit_MathRBFInterpolateQuatXform :
 	public FRigUnit_MathRBFInterpolateQuatBase
 {
@@ -329,7 +329,7 @@ struct FMathRBFInterpolateVectorFloat_Target
 	float Value = 0.0f;
 };
 
-USTRUCT(meta = (Abstract, DisplayName = "RBF Vector to Float", PrototypeName = "RBFInterpolateVectorToFloat", Keywords = "RBF,Interpolate,Vector"))
+USTRUCT(meta = (Abstract, DisplayName = "RBF Vector to Float", TemplateName="RBFInterpolateVectorToFloat", Keywords = "RBF,Interpolate,Vector"))
 struct CONTROLRIG_API FRigUnit_MathRBFInterpolateVectorFloat :
 	public FRigUnit_MathRBFInterpolateVectorBase
 {
@@ -358,7 +358,7 @@ struct FMathRBFInterpolateVectorVector_Target
 	FVector Value = FVector::ZeroVector;
 };
 
-USTRUCT(meta = (Abstract, DisplayName = "RBF Vector to Vector", PrototypeName = "RBFInterpolateVectorToVector", Keywords = "RBF,Interpolate,Vector"))
+USTRUCT(meta = (Abstract, DisplayName = "RBF Vector to Vector", TemplateName="RBFInterpolateVectorToVector", Keywords = "RBF,Interpolate,Vector"))
 struct CONTROLRIG_API FRigUnit_MathRBFInterpolateVectorVector :
 	public FRigUnit_MathRBFInterpolateVectorBase
 {
@@ -387,7 +387,7 @@ struct FMathRBFInterpolateVectorColor_Target
 	FLinearColor Value = FLinearColor::Transparent;
 };
 
-USTRUCT(meta = (Abstract, DisplayName = "RBF Vector to Color", PrototypeName = "RBFInterpolateVectorToColor", Keywords = "RBF,Interpolate,Vector"))
+USTRUCT(meta = (Abstract, DisplayName = "RBF Vector to Color", TemplateName="RBFInterpolateVectorToColor", Keywords = "RBF,Interpolate,Vector"))
 struct CONTROLRIG_API FRigUnit_MathRBFInterpolateVectorColor :
 	public FRigUnit_MathRBFInterpolateVectorBase
 {
@@ -417,7 +417,7 @@ struct FMathRBFInterpolateVectorQuat_Target
 	FQuat Value = FQuat::Identity;
 };
 
-USTRUCT(meta = (Abstract, DisplayName = "RBF Vector to Quat", PrototypeName = "RBFInterpolateVectorToQuat", Keywords = "RBF,Interpolate,Vector"))
+USTRUCT(meta = (Abstract, DisplayName = "RBF Vector to Quat", TemplateName="RBFInterpolateVectorToQuat", Keywords = "RBF,Interpolate,Vector"))
 struct CONTROLRIG_API FRigUnit_MathRBFInterpolateVectorQuat :
 	public FRigUnit_MathRBFInterpolateVectorBase
 {
@@ -445,7 +445,7 @@ struct FMathRBFInterpolateVectorXform_Target
 	FTransform Value = FTransform::Identity;
 };
 
-USTRUCT(meta = (Abstract, DisplayName = "RBF Vector to Transform", PrototypeName = "RBFInterpolateVectorToXform", Keywords = "RBF,Interpolate,Vector"))
+USTRUCT(meta = (Abstract, DisplayName = "RBF Vector to Transform", TemplateName="RBFInterpolateVectorToXform", Keywords = "RBF,Interpolate,Vector"))
 struct CONTROLRIG_API FRigUnit_MathRBFInterpolateVectorXform :
 	public FRigUnit_MathRBFInterpolateVectorBase
 {
diff --git a/Engine/Plugins/Animation/ControlRig/Source/ControlRig/Private/Units/Math/RigUnit_MathTransform.h b/Engine/Plugins/Animation/ControlRig/Source/ControlRig/Private/Units/Math/RigUnit_MathTransform.h
--- a/Engine/Plugins/Animation/ControlRig/Source/ControlRig/Private/Units/Math/RigUnit_MathTransform.h
+++ b/Engine/Plugins/Animation/ControlRig/Source/ControlRig/Private/Units/Math/RigUnit_MathTransform.h
@@ -59,7 +59,7 @@ struct CONTROLRIG_API FRigUnit_MathTransformBinaryOp : public FRigUnit_MathTrans
 /**
  * Makes a quaternion based transform from a euler based transform
  */
-USTRUCT(meta=(DisplayName="From Euler Transform", PrototypeName="FromEulerTransform", Keywords="Make,Construct"))
+USTRUCT(meta=(DisplayName="From Euler Transform", TemplateName="FromEulerTransform", Keywords="Make,Construct"))
 struct CONTROLRIG_API FRigUnit_MathTransformFromEulerTransform : public FRigUnit_MathTransformBase
 {
 	GENERATED_BODY()
@@ -83,7 +83,7 @@ struct CONTROLRIG_API FRigUnit_MathTransformFromEulerTransform : public FRigUnit
 /**
  * Retrieves a euler based transform from a quaternion based transform
  */
-USTRUCT(meta=(DisplayName="To Euler Transform", PrototypeName="ToEulerTransform", Keywords="Make,Construct"))
+USTRUCT(meta=(DisplayName="To Euler Transform", TemplateName="ToEulerTransform", Keywords="Make,Construct"))
 struct CONTROLRIG_API FRigUnit_MathTransformToEulerTransform : public FRigUnit_MathTransformBase
 {
 	GENERATED_BODY()
@@ -107,7 +107,7 @@ struct CONTROLRIG_API FRigUnit_MathTransformToEulerTransform : public FRigUnit_M
 /**
  * Returns the product of the two values
  */
-USTRUCT(meta=(DisplayName="Multiply", PrototypeName="Multiply", Keywords="Product,*,Global"))
+USTRUCT(meta=(DisplayName="Multiply", TemplateName="Multiply", Keywords="Product,*,Global"))
 struct CONTROLRIG_API FRigUnit_MathTransformMul : public FRigUnit_MathTransformBinaryOp
 {
 	GENERATED_BODY()
@@ -118,7 +118,7 @@ struct CONTROLRIG_API FRigUnit_MathTransformMul : public FRigUnit_MathTransformB
 /**
  * Returns the relative local transform within a parent's transform
  */
-USTRUCT(meta=(DisplayName="Make Relative", PrototypeName="MakeRelative", Keywords="Local,Global,Absolute"))
+USTRUCT(meta=(DisplayName="Make Relative", TemplateName="MakeRelative", Keywords="Local,Global,Absolute"))
 struct CONTROLRIG_API FRigUnit_MathTransformMakeRelative : public FRigUnit_MathTransformBase
 {
 	GENERATED_BODY()
@@ -144,7 +144,7 @@ struct CONTROLRIG_API FRigUnit_MathTransformMakeRelative : public FRigUnit_MathT
 /**
  * Returns the absolute global transform within a parent's transform
  */
-USTRUCT(meta = (DisplayName = "Make Absolute", PrototypeName = "MakeAbsolute", Keywords = "Local,Global,Relative"))
+USTRUCT(meta = (DisplayName = "Make Absolute", TemplateName="MakeAbsolute", Keywords = "Local,Global,Relative"))
 struct CONTROLRIG_API FRigUnit_MathTransformMakeAbsolute : public FRigUnit_MathTransformBase
 {
 	GENERATED_BODY()
@@ -214,7 +214,7 @@ struct CONTROLRIG_API FRigUnit_MathTransformAccumulateArray : public FRigUnit_Ma
 /**
  * Returns the negative value
  */
-USTRUCT(meta=(DisplayName="Inverse", PrototypeName="Inverse"))
+USTRUCT(meta=(DisplayName="Inverse", TemplateName="Inverse"))
 struct CONTROLRIG_API FRigUnit_MathTransformInverse : public FRigUnit_MathTransformUnaryOp
 {
 	GENERATED_BODY()
@@ -225,7 +225,7 @@ struct CONTROLRIG_API FRigUnit_MathTransformInverse : public FRigUnit_MathTransf
 /**
  * Linearly interpolates between A and B using the ratio T
  */
-USTRUCT(meta=(DisplayName="Interpolate", PrototypeName="Interpolate", Keywords="Lerp,Mix,Blend"))
+USTRUCT(meta=(DisplayName="Interpolate", TemplateName="Interpolate", Keywords="Lerp,Mix,Blend"))
 struct CONTROLRIG_API FRigUnit_MathTransformLerp : public FRigUnit_MathTransformBase
 {
 	GENERATED_BODY()
@@ -255,7 +255,7 @@ struct CONTROLRIG_API FRigUnit_MathTransformLerp : public FRigUnit_MathTransform
 /**
  * Return one of the two values based on the condition
  */
-USTRUCT(meta=(DisplayName="Select", PrototypeName="Select", Keywords="Pick,If", Deprecated = "4.26.0"))
+USTRUCT(meta=(DisplayName="Select", TemplateName="Select", Keywords="Pick,If", Deprecated = "4.26.0"))
 struct CONTROLRIG_API FRigUnit_MathTransformSelectBool : public FRigUnit_MathTransformBase
 {
 	GENERATED_BODY()
@@ -288,7 +288,7 @@ struct CONTROLRIG_API FRigUnit_MathTransformSelectBool : public FRigUnit_MathTra
 /**
  * Rotates a given vector (direction) by the transform
  */
-USTRUCT(meta=(DisplayName="Transform Direction", PrototypeName="Rotate", Keywords="Transform,Direction"))
+USTRUCT(meta=(DisplayName="Transform Direction", TemplateName="Rotate", Keywords="Transform,Direction"))
 struct CONTROLRIG_API FRigUnit_MathTransformRotateVector : public FRigUnit_MathTransformBase
 {
 	GENERATED_BODY()
@@ -315,7 +315,7 @@ struct CONTROLRIG_API FRigUnit_MathTransformRotateVector : public FRigUnit_MathT
 /**
  * Rotates a given vector (location) by the transform
  */
-USTRUCT(meta=(DisplayName="Transform Location", PrototypeName="Multiply"))
+USTRUCT(meta=(DisplayName="Transform Location", TemplateName="Multiply"))
 struct CONTROLRIG_API FRigUnit_MathTransformTransformVector : public FRigUnit_MathTransformBase
 {
 	GENERATED_BODY()
@@ -382,7 +382,7 @@ struct CONTROLRIG_API FRigUnit_MathTransformFromSRT : public FRigUnit_MathTransf
 /**
  * Clamps a position using a plane collision, cylindric collision or spherical collision.
  */
-USTRUCT(meta = (DisplayName = "Clamp Spatially", PrototypeName = "ClampSpatially", Keywords = "Collide,Collision"))
+USTRUCT(meta = (DisplayName = "Clamp Spatially", TemplateName="ClampSpatially", Keywords = "Collide,Collision"))
 struct CONTROLRIG_API FRigUnit_MathTransformClampSpatially : public FRigUnit_MathTransformBase
 {
 	GENERATED_BODY()
diff --git a/Engine/Plugins/Animation/ControlRig/Source/ControlRig/Private/Units/Math/RigUnit_MathVector.h b/Engine/Plugins/Animation/ControlRig/Source/ControlRig/Private/Units/Math/RigUnit_MathVector.h
--- a/Engine/Plugins/Animation/ControlRig/Source/ControlRig/Private/Units/Math/RigUnit_MathVector.h
+++ b/Engine/Plugins/Animation/ControlRig/Source/ControlRig/Private/Units/Math/RigUnit_MathVector.h
@@ -52,7 +52,7 @@ struct CONTROLRIG_API FRigUnit_MathVectorBinaryOp : public FRigUnit_MathVectorBa
 /**
  * Makes a vector from a single float
  */
-USTRUCT(meta=(DisplayName="From Float", PrototypeName="FromFloat", Keywords="Make,Construct"))
+USTRUCT(meta=(DisplayName="From Float", TemplateName="FromFloat", Keywords="Make,Construct"))
 struct CONTROLRIG_API FRigUnit_MathVectorFromFloat : public FRigUnit_MathVectorBase
 {
 	GENERATED_BODY()
@@ -75,7 +75,7 @@ struct CONTROLRIG_API FRigUnit_MathVectorFromFloat : public FRigUnit_MathVectorB
 /**
  * Returns the sum of the two values
  */
-USTRUCT(meta=(DisplayName="Add", PrototypeName="Add", Keywords="Sum,+"))
+USTRUCT(meta=(DisplayName="Add", TemplateName="Add", Keywords="Sum,+"))
 struct CONTROLRIG_API FRigUnit_MathVectorAdd : public FRigUnit_MathVectorBinaryOp
 {
 	GENERATED_BODY()
@@ -87,7 +87,7 @@ struct CONTROLRIG_API FRigUnit_MathVectorAdd : public FRigUnit_MathVectorBinaryO
 /**
  * Returns the difference of the two values
  */
-USTRUCT(meta=(DisplayName="Subtract", PrototypeName="Subtract", Keywords="-"))
+USTRUCT(meta=(DisplayName="Subtract", TemplateName="Subtract", Keywords="-"))
 struct CONTROLRIG_API FRigUnit_MathVectorSub : public FRigUnit_MathVectorBinaryOp
 {
 	GENERATED_BODY()
@@ -99,7 +99,7 @@ struct CONTROLRIG_API FRigUnit_MathVectorSub : public FRigUnit_MathVectorBinaryO
 /**
  * Returns the product of the two values
  */
-USTRUCT(meta=(DisplayName="Multiply", PrototypeName="Multiply", Keywords="Product,*"))
+USTRUCT(meta=(DisplayName="Multiply", TemplateName="Multiply", Keywords="Product,*"))
 struct CONTROLRIG_API FRigUnit_MathVectorMul : public FRigUnit_MathVectorBinaryOp
 {
 	GENERATED_BODY()
@@ -116,7 +116,7 @@ struct CONTROLRIG_API FRigUnit_MathVectorMul : public FRigUnit_MathVectorBinaryO
 /**
  * Returns the product of the the vector and the float value
  */
-USTRUCT(meta = (DisplayName = "Scale", PrototypeName = "Scale", Keywords = "Multiply,Product,*,ByScalar,ByFloat"))
+USTRUCT(meta = (DisplayName = "Scale", TemplateName="Scale", Keywords = "Multiply,Product,*,ByScalar,ByFloat"))
 struct CONTROLRIG_API FRigUnit_MathVectorScale : public FRigUnit_MathVectorBase
 {
 	GENERATED_BODY()
@@ -143,7 +143,7 @@ struct CONTROLRIG_API FRigUnit_MathVectorScale : public FRigUnit_MathVectorBase
 /**
  * Returns the division of the two values
  */
-USTRUCT(meta=(DisplayName="Divide", PrototypeName="Divide", Keywords="Division,Divisor,/"))
+USTRUCT(meta=(DisplayName="Divide", TemplateName="Divide", Keywords="Division,Divisor,/"))
 struct CONTROLRIG_API FRigUnit_MathVectorDiv : public FRigUnit_MathVectorBinaryOp
 {
 	GENERATED_BODY()
@@ -160,7 +160,7 @@ struct CONTROLRIG_API FRigUnit_MathVectorDiv : public FRigUnit_MathVectorBinaryO
 /**
  * Returns the modulo of the two values
  */
-USTRUCT(meta=(DisplayName="Modulo", PrototypeName="Modulo", Keywords="%,fmod"))
+USTRUCT(meta=(DisplayName="Modulo", TemplateName="Modulo", Keywords="%,fmod"))
 struct CONTROLRIG_API FRigUnit_MathVectorMod : public FRigUnit_MathVectorBinaryOp
 {
 	GENERATED_BODY()
@@ -178,7 +178,7 @@ struct CONTROLRIG_API FRigUnit_MathVectorMod : public FRigUnit_MathVectorBinaryO
 /**
  * Returns the smaller of the two values for each component
  */
-USTRUCT(meta=(DisplayName="Minimum", PrototypeName="Minimum"))
+USTRUCT(meta=(DisplayName="Minimum", TemplateName="Minimum"))
 struct CONTROLRIG_API FRigUnit_MathVectorMin : public FRigUnit_MathVectorBinaryOp
 {
 	GENERATED_BODY()
@@ -190,7 +190,7 @@ struct CONTROLRIG_API FRigUnit_MathVectorMin : public FRigUnit_MathVectorBinaryO
 /**
  * Returns the larger of the two values each component
  */
-USTRUCT(meta=(DisplayName="Maximum", PrototypeName="Maximum"))
+USTRUCT(meta=(DisplayName="Maximum", TemplateName="Maximum"))
 struct CONTROLRIG_API FRigUnit_MathVectorMax : public FRigUnit_MathVectorBinaryOp
 {
 	GENERATED_BODY()
@@ -202,7 +202,7 @@ struct CONTROLRIG_API FRigUnit_MathVectorMax : public FRigUnit_MathVectorBinaryO
 /**
  * Returns the negative value
  */
-USTRUCT(meta=(DisplayName="Negate", PrototypeName="Negate", Keywords="-,Abs"))
+USTRUCT(meta=(DisplayName="Negate", TemplateName="Negate", Keywords="-,Abs"))
 struct CONTROLRIG_API FRigUnit_MathVectorNegate : public FRigUnit_MathVectorUnaryOp
 {
 	GENERATED_BODY()
@@ -214,7 +214,7 @@ struct CONTROLRIG_API FRigUnit_MathVectorNegate : public FRigUnit_MathVectorUnar
 /**
  * Returns the absolute (positive) value
  */
-USTRUCT(meta=(DisplayName="Absolute", PrototypeName="Absolute", Keywords="Abs,Neg"))
+USTRUCT(meta=(DisplayName="Absolute", TemplateName="Absolute", Keywords="Abs,Neg"))
 struct CONTROLRIG_API FRigUnit_MathVectorAbs : public FRigUnit_MathVectorUnaryOp
 {
 	GENERATED_BODY()
@@ -226,7 +226,7 @@ struct CONTROLRIG_API FRigUnit_MathVectorAbs : public FRigUnit_MathVectorUnaryOp
 /**
  * Returns the closest lower full number (integer) of the value for each component
  */
-USTRUCT(meta=(DisplayName="Floor", PrototypeName="Floor", Keywords="Round"))
+USTRUCT(meta=(DisplayName="Floor", TemplateName="Floor", Keywords="Round"))
 struct CONTROLRIG_API FRigUnit_MathVectorFloor : public FRigUnit_MathVectorUnaryOp
 {
 	GENERATED_BODY()
@@ -238,7 +238,7 @@ struct CONTROLRIG_API FRigUnit_MathVectorFloor : public FRigUnit_MathVectorUnary
 /**
  * Returns the closest higher full number (integer) of the value for each component
  */
-USTRUCT(meta=(DisplayName="Ceiling", PrototypeName="Ceiling", Keywords="Round"))
+USTRUCT(meta=(DisplayName="Ceiling", TemplateName="Ceiling", Keywords="Round"))
 struct CONTROLRIG_API FRigUnit_MathVectorCeil : public FRigUnit_MathVectorUnaryOp
 {
 	GENERATED_BODY()
@@ -250,7 +250,7 @@ struct CONTROLRIG_API FRigUnit_MathVectorCeil : public FRigUnit_MathVectorUnaryO
 /**
  * Returns the closest higher full number (integer) of the value for each component
  */
-USTRUCT(meta=(DisplayName="Round", PrototypeName="Round"))
+USTRUCT(meta=(DisplayName="Round", TemplateName="Round"))
 struct CONTROLRIG_API FRigUnit_MathVectorRound : public FRigUnit_MathVectorUnaryOp
 {
 	GENERATED_BODY()
@@ -262,7 +262,7 @@ struct CONTROLRIG_API FRigUnit_MathVectorRound : public FRigUnit_MathVectorUnary
 /**
  * Returns the sign of the value (+1 for >= FVector(0.f, 0.f, 0.f), -1 for < 0.f) for each component
  */
-USTRUCT(meta=(DisplayName="Sign", PrototypeName="Sign"))
+USTRUCT(meta=(DisplayName="Sign", TemplateName="Sign"))
 struct CONTROLRIG_API FRigUnit_MathVectorSign : public FRigUnit_MathVectorUnaryOp
 {
 	GENERATED_BODY()
@@ -274,7 +274,7 @@ struct CONTROLRIG_API FRigUnit_MathVectorSign : public FRigUnit_MathVectorUnaryO
 /**
  * Clamps the given value within the range provided by minimum and maximum for each component
  */
-USTRUCT(meta=(DisplayName="Clamp", PrototypeName="Clamp", Keywords="Range,Remap"))
+USTRUCT(meta=(DisplayName="Clamp", TemplateName="Clamp", Keywords="Range,Remap"))
 struct CONTROLRIG_API FRigUnit_MathVectorClamp : public FRigUnit_MathVectorBase
 {
 	GENERATED_BODY()
@@ -304,7 +304,7 @@ struct CONTROLRIG_API FRigUnit_MathVectorClamp : public FRigUnit_MathVectorBase
 /**
  * Linearly interpolates between A and B using the ratio T
  */
-USTRUCT(meta=(DisplayName="Interpolate", PrototypeName="Interpolate", Keywords="Lerp,Mix,Blend"))
+USTRUCT(meta=(DisplayName="Interpolate", TemplateName="Interpolate", Keywords="Lerp,Mix,Blend"))
 struct CONTROLRIG_API FRigUnit_MathVectorLerp : public FRigUnit_MathVectorBase
 {
 	GENERATED_BODY()
@@ -335,7 +335,7 @@ struct CONTROLRIG_API FRigUnit_MathVectorLerp : public FRigUnit_MathVectorBase
 /**
  * Remaps the given value from a source range to a target range for each component
  */
-USTRUCT(meta=(DisplayName="Remap", PrototypeName="Remap", Keywords="Rescale,Scale"))
+USTRUCT(meta=(DisplayName="Remap", TemplateName="Remap", Keywords="Rescale,Scale"))
 struct CONTROLRIG_API FRigUnit_MathVectorRemap : public FRigUnit_MathVectorBase
 {
 	GENERATED_BODY()
@@ -376,7 +376,7 @@ struct CONTROLRIG_API FRigUnit_MathVectorRemap : public FRigUnit_MathVectorBase
 /**
  * Returns true if the value A equals B
  */
-USTRUCT(meta=(DisplayName="Equals", PrototypeName="Equals", Keywords="Same,=="))
+USTRUCT(meta=(DisplayName="Equals", TemplateName="Equals", Keywords="Same,=="))
 struct CONTROLRIG_API FRigUnit_MathVectorEquals : public FRigUnit_MathVectorBase
 {
 	GENERATED_BODY()
@@ -403,7 +403,7 @@ struct CONTROLRIG_API FRigUnit_MathVectorEquals : public FRigUnit_MathVectorBase
 /**
  * Returns true if the value A does not equal B
  */
-USTRUCT(meta=(DisplayName="Not Equals", PrototypeName="NotEquals", Keywords="Different,!="))
+USTRUCT(meta=(DisplayName="Not Equals", TemplateName="NotEquals", Keywords="Different,!="))
 struct CONTROLRIG_API FRigUnit_MathVectorNotEquals : public FRigUnit_MathVectorBase
 {
 	GENERATED_BODY()
@@ -430,7 +430,7 @@ struct CONTROLRIG_API FRigUnit_MathVectorNotEquals : public FRigUnit_MathVectorB
 /**
  * Returns true if the value is nearly zero
  */
-USTRUCT(meta=(DisplayName="Is Nearly Zero", PrototypeName="IsNearlyZero", Keywords="AlmostZero,0"))
+USTRUCT(meta=(DisplayName="Is Nearly Zero", TemplateName="IsNearlyZero", Keywords="AlmostZero,0"))
 struct CONTROLRIG_API FRigUnit_MathVectorIsNearlyZero : public FRigUnit_MathVectorBase
 {
 	GENERATED_BODY()
@@ -458,7 +458,7 @@ struct CONTROLRIG_API FRigUnit_MathVectorIsNearlyZero : public FRigUnit_MathVect
 /**
  * Returns true if the value A is almost equal to B
  */
-USTRUCT(meta=(DisplayName="Is Nearly Equal", PrototypeName="IsNearlyEqual", Keywords="AlmostEqual,=="))
+USTRUCT(meta=(DisplayName="Is Nearly Equal", TemplateName="IsNearlyEqual", Keywords="AlmostEqual,=="))
 struct CONTROLRIG_API FRigUnit_MathVectorIsNearlyEqual : public FRigUnit_MathVectorBase
 {
 	GENERATED_BODY()
@@ -489,7 +489,7 @@ struct CONTROLRIG_API FRigUnit_MathVectorIsNearlyEqual : public FRigUnit_MathVec
 /**
  * Return one of the two values based on the condition
  */
-USTRUCT(meta=(DisplayName="Select", PrototypeName="Select", Keywords="Pick,If", Deprecated = "4.26.0"))
+USTRUCT(meta=(DisplayName="Select", TemplateName="Select", Keywords="Pick,If", Deprecated = "4.26.0"))
 struct CONTROLRIG_API FRigUnit_MathVectorSelectBool : public FRigUnit_MathVectorBase
 {
 	GENERATED_BODY()
@@ -522,7 +522,7 @@ struct CONTROLRIG_API FRigUnit_MathVectorSelectBool : public FRigUnit_MathVector
 /**
  * Returns the degrees of a given value in radians
  */
-USTRUCT(meta=(DisplayName="Degrees", PrototypeName="Degrees"))
+USTRUCT(meta=(DisplayName="Degrees", TemplateName="Degrees"))
 struct CONTROLRIG_API FRigUnit_MathVectorDeg : public FRigUnit_MathVectorUnaryOp
 {
 	GENERATED_BODY()
@@ -534,7 +534,7 @@ struct CONTROLRIG_API FRigUnit_MathVectorDeg : public FRigUnit_MathVectorUnaryOp
 /**
  * Returns the radians of a given value in degrees
  */
-USTRUCT(meta=(DisplayName="Radians", PrototypeName="Radians"))
+USTRUCT(meta=(DisplayName="Radians", TemplateName="Radians"))
 struct CONTROLRIG_API FRigUnit_MathVectorRad : public FRigUnit_MathVectorUnaryOp
 {
 	GENERATED_BODY()
@@ -546,7 +546,7 @@ struct CONTROLRIG_API FRigUnit_MathVectorRad : public FRigUnit_MathVectorUnaryOp
 /**
  * Returns the squared length of the vector
  */
-USTRUCT(meta=(DisplayName="Length Squared", PrototypeName="LengthSquared", Keywords="Length,Size,Magnitude"))
+USTRUCT(meta=(DisplayName="Length Squared", TemplateName="LengthSquared", Keywords="Length,Size,Magnitude"))
 struct CONTROLRIG_API FRigUnit_MathVectorLengthSquared : public FRigUnit_MathVectorBase
 {
 	GENERATED_BODY()
@@ -570,7 +570,7 @@ struct CONTROLRIG_API FRigUnit_MathVectorLengthSquared : public FRigUnit_MathVec
 /**
  * Returns the length of the vector
  */
-USTRUCT(meta=(DisplayName="Length", PrototypeName="Length", Keywords="Size,Magnitude"))
+USTRUCT(meta=(DisplayName="Length", TemplateName="Length", Keywords="Size,Magnitude"))
 struct CONTROLRIG_API FRigUnit_MathVectorLength : public FRigUnit_MathVectorBase
 {
 	GENERATED_BODY()
@@ -594,7 +594,7 @@ struct CONTROLRIG_API FRigUnit_MathVectorLength : public FRigUnit_MathVectorBase
 /**
  * Returns the distance from A to B
  */
-USTRUCT(meta=(DisplayName="Distance Between", PrototypeName="Distance"))
+USTRUCT(meta=(DisplayName="Distance Between", TemplateName="Distance"))
 struct CONTROLRIG_API FRigUnit_MathVectorDistance : public FRigUnit_MathVectorBase
 {
 	GENERATED_BODY()
@@ -621,7 +621,7 @@ struct CONTROLRIG_API FRigUnit_MathVectorDistance : public FRigUnit_MathVectorBa
 /**
  * Returns the cross product between two vectors
  */
-USTRUCT(meta=(DisplayName="Cross", PrototypeName="Cross", Keywords="^"))
+USTRUCT(meta=(DisplayName="Cross", TemplateName="Cross", Keywords="^"))
 struct CONTROLRIG_API FRigUnit_MathVectorCross : public FRigUnit_MathVectorBinaryOp
 {
 	GENERATED_BODY()
@@ -632,7 +632,7 @@ struct CONTROLRIG_API FRigUnit_MathVectorCross : public FRigUnit_MathVectorBinar
 /**
  * Returns the dot product between two vectors
  */
-USTRUCT(meta=(DisplayName="Dot", PrototypeName="Dot,|"))
+USTRUCT(meta=(DisplayName="Dot", TemplateName="Dot,|"))
 struct CONTROLRIG_API FRigUnit_MathVectorDot : public FRigUnit_MathVectorBase
 {
 	GENERATED_BODY()
@@ -659,7 +659,7 @@ struct CONTROLRIG_API FRigUnit_MathVectorDot : public FRigUnit_MathVectorBase
 /**
  * Returns the normalized value
  */
-USTRUCT(meta=(DisplayName="Unit", PrototypeName="Unit", Keywords="Normalize"))
+USTRUCT(meta=(DisplayName="Unit", TemplateName="Unit", Keywords="Normalize"))
 struct CONTROLRIG_API FRigUnit_MathVectorUnit : public FRigUnit_MathVectorUnaryOp
 {
 	GENERATED_BODY()
@@ -671,7 +671,7 @@ struct CONTROLRIG_API FRigUnit_MathVectorUnit : public FRigUnit_MathVectorUnaryO
 /**
  * Sets the length of a given vector
  */
-USTRUCT(meta = (DisplayName = "SetLength", PrototypeName = "SetLength", Keywords = "Unit,Normalize,Scale"))
+USTRUCT(meta = (DisplayName = "SetLength", TemplateName="SetLength", Keywords = "Unit,Normalize,Scale"))
 struct CONTROLRIG_API FRigUnit_MathVectorSetLength: public FRigUnit_MathVectorBase
 {
 	GENERATED_BODY()
@@ -698,7 +698,7 @@ struct CONTROLRIG_API FRigUnit_MathVectorSetLength: public FRigUnit_MathVectorBa
 /**
  * Clamps the length of a given vector between a minimum and maximum
  */
-USTRUCT(meta = (DisplayName = "ClampLength", PrototypeName = "ClampLength", Keywords = "Unit,Normalize,Scale"))
+USTRUCT(meta = (DisplayName = "ClampLength", TemplateName="ClampLength", Keywords = "Unit,Normalize,Scale"))
 struct CONTROLRIG_API FRigUnit_MathVectorClampLength: public FRigUnit_MathVectorBase
 {
 	GENERATED_BODY()
@@ -730,7 +730,7 @@ struct CONTROLRIG_API FRigUnit_MathVectorClampLength: public FRigUnit_MathVector
 /**
  * Mirror a vector about a normal vector.
  */
-USTRUCT(meta=(DisplayName="Mirror", PrototypeName="Mirror"))
+USTRUCT(meta=(DisplayName="Mirror", TemplateName="Mirror"))
 struct CONTROLRIG_API FRigUnit_MathVectorMirror : public FRigUnit_MathVectorBase
 {
 	GENERATED_BODY()
@@ -757,7 +757,7 @@ struct CONTROLRIG_API FRigUnit_MathVectorMirror : public FRigUnit_MathVectorBase
 /**
  * Returns the angle between two vectors in radians
  */
-USTRUCT(meta=(DisplayName="Angle Between", PrototypeName="AngleBetween"))
+USTRUCT(meta=(DisplayName="Angle Between", TemplateName="AngleBetween"))
 struct CONTROLRIG_API FRigUnit_MathVectorAngle : public FRigUnit_MathVectorBase
 {
 	GENERATED_BODY()
@@ -784,7 +784,7 @@ struct CONTROLRIG_API FRigUnit_MathVectorAngle : public FRigUnit_MathVectorBase
 /**
  * Returns true if the two vectors are parallel
  */
-USTRUCT(meta=(DisplayName="Parallel", PrototypeName="Parallel"))
+USTRUCT(meta=(DisplayName="Parallel", TemplateName="Parallel"))
 struct CONTROLRIG_API FRigUnit_MathVectorParallel : public FRigUnit_MathVectorBase
 {
 	GENERATED_BODY()
@@ -811,7 +811,7 @@ struct CONTROLRIG_API FRigUnit_MathVectorParallel : public FRigUnit_MathVectorBa
 /**
  * Returns true if the two vectors are orthogonal
  */
-USTRUCT(meta=(DisplayName="Orthogonal", PrototypeName="Orthogonal"))
+USTRUCT(meta=(DisplayName="Orthogonal", TemplateName="Orthogonal"))
 struct CONTROLRIG_API FRigUnit_MathVectorOrthogonal : public FRigUnit_MathVectorBase
 {
 	GENERATED_BODY()
@@ -896,7 +896,7 @@ struct CONTROLRIG_API FRigUnit_MathVectorMakeBezierFourPoint : public FRigUnit_M
 /**
  * Clamps a position using a plane collision, cylindric collision or spherical collision.
  */
-USTRUCT(meta = (DisplayName = "Clamp Spatially", PrototypeName = "ClampSpatially", Keywords="Collide,Collision"))
+USTRUCT(meta = (DisplayName = "Clamp Spatially", TemplateName="ClampSpatially", Keywords="Collide,Collision"))
 struct CONTROLRIG_API FRigUnit_MathVectorClampSpatially: public FRigUnit_MathVectorBase
 {
 	GENERATED_BODY()
diff --git a/Engine/Plugins/Animation/ControlRig/Source/ControlRig/Private/Units/Math/RigUnit_Noise.h b/Engine/Plugins/Animation/ControlRig/Source/ControlRig/Private/Units/Math/RigUnit_Noise.h
--- a/Engine/Plugins/Animation/ControlRig/Source/ControlRig/Private/Units/Math/RigUnit_Noise.h
+++ b/Engine/Plugins/Animation/ControlRig/Source/ControlRig/Private/Units/Math/RigUnit_Noise.h
@@ -8,7 +8,7 @@
 /**
  * Generates a float through a noise fluctuation process between a min and a max through speed
  */
-USTRUCT(meta=(DisplayName="Noise (Float)", Category="Math|Noise", PrototypeName="Noise"))
+USTRUCT(meta=(DisplayName="Noise (Float)", Category="Math|Noise", TemplateName="Noise"))
 struct CONTROLRIG_API FRigUnit_NoiseFloat : public FRigUnit_MathBase
 {
 	GENERATED_BODY()
@@ -48,7 +48,7 @@ struct CONTROLRIG_API FRigUnit_NoiseFloat : public FRigUnit_MathBase
 /**
  * Generates a float through a noise fluctuation process between a min and a max through speed
  */
-USTRUCT(meta=(DisplayName="Noise (Double)", Category="Math|Noise", PrototypeName="Noise"))
+USTRUCT(meta=(DisplayName="Noise (Double)", Category="Math|Noise", TemplateName="Noise"))
 struct CONTROLRIG_API FRigUnit_NoiseDouble : public FRigUnit_MathBase
 {
 	GENERATED_BODY()
@@ -88,7 +88,7 @@ struct CONTROLRIG_API FRigUnit_NoiseDouble : public FRigUnit_MathBase
 /**
  * Generates a vector through a noise fluctuation process between a min and a max through speed
  */
-USTRUCT(meta = (DisplayName = "Noise (Vector)", Category = "Math|Noise", PrototypeName = "Noise", Deprecated = "5.0.0"))
+USTRUCT(meta = (DisplayName = "Noise (Vector)", Category = "Math|Noise", TemplateName="Noise", Deprecated = "5.0.0"))
 struct CONTROLRIG_API FRigUnit_NoiseVector : public FRigUnit_MathBase
 {
 	GENERATED_BODY()
@@ -133,7 +133,7 @@ struct CONTROLRIG_API FRigUnit_NoiseVector : public FRigUnit_MathBase
 /**
  * Generates a vector through a noise fluctuation process between a min and a max through speed
  */
-USTRUCT(meta = (DisplayName = "Noise (Vector)", Category = "Math|Noise", PrototypeName = "Noise"))
+USTRUCT(meta = (DisplayName = "Noise (Vector)", Category = "Math|Noise", TemplateName="Noise"))
 struct CONTROLRIG_API FRigUnit_NoiseVector2 : public FRigUnit_MathBase
 {
 	GENERATED_BODY()
diff --git a/Engine/Plugins/Animation/ControlRig/Source/ControlRig/Private/Units/Math/RigUnit_Random.h b/Engine/Plugins/Animation/ControlRig/Source/ControlRig/Private/Units/Math/RigUnit_Random.h
--- a/Engine/Plugins/Animation/ControlRig/Source/ControlRig/Private/Units/Math/RigUnit_Random.h
+++ b/Engine/Plugins/Animation/ControlRig/Source/ControlRig/Private/Units/Math/RigUnit_Random.h
@@ -8,7 +8,7 @@
 /**
  * Generates a random float between a min and a max
  */
-USTRUCT(meta=(DisplayName="Random (Float)", Category="Math|Random", PrototypeName="Random"))
+USTRUCT(meta=(DisplayName="Random (Float)", Category="Math|Random", TemplateName="Random"))
 struct CONTROLRIG_API FRigUnit_RandomFloat : public FRigUnit_MathBase
 {
 	GENERATED_BODY()
@@ -61,7 +61,7 @@ struct CONTROLRIG_API FRigUnit_RandomFloat : public FRigUnit_MathBase
 /**
  * Generates a random vector between a min and a max
  */
-USTRUCT(meta = (DisplayName = "Random (Vector)", Category = "Math|Random", PrototypeName = "Random"))
+USTRUCT(meta = (DisplayName = "Random (Vector)", Category = "Math|Random", TemplateName="Random"))
 struct CONTROLRIG_API FRigUnit_RandomVector: public FRigUnit_MathBase
 {
 	GENERATED_BODY()
diff --git a/Engine/Plugins/Animation/ControlRig/Source/ControlRig/Private/Units/Simulation/RigUnit_Accumulate.h b/Engine/Plugins/Animation/ControlRig/Source/ControlRig/Private/Units/Simulation/RigUnit_Accumulate.h
--- a/Engine/Plugins/Animation/ControlRig/Source/ControlRig/Private/Units/Simulation/RigUnit_Accumulate.h
+++ b/Engine/Plugins/Animation/ControlRig/Source/ControlRig/Private/Units/Simulation/RigUnit_Accumulate.h
@@ -14,7 +14,7 @@ struct CONTROLRIG_API FRigUnit_AccumulateBase : public FRigUnit_SimBase
 /**
  * Adds a value over time over and over again
  */
-USTRUCT(meta=(DisplayName="Accumulate Add (Float)", PrototypeName="AccumulateAdd", Keywords="Simulate,++"))
+USTRUCT(meta=(DisplayName="Accumulate Add (Float)", TemplateName="AccumulateAdd", Keywords="Simulate,++"))
 struct CONTROLRIG_API FRigUnit_AccumulateFloatAdd : public FRigUnit_AccumulateBase
 {
 	GENERATED_BODY()
@@ -47,7 +47,7 @@ struct CONTROLRIG_API FRigUnit_AccumulateFloatAdd : public FRigUnit_AccumulateBa
 /**
  * Adds a vector over time over and over again
  */
-USTRUCT(meta = (DisplayName="Accumulate Add (Vector)", PrototypeName="AccumulateAdd", Keywords="Simulate,++"))
+USTRUCT(meta = (DisplayName="Accumulate Add (Vector)", TemplateName="AccumulateAdd", Keywords="Simulate,++"))
 struct CONTROLRIG_API FRigUnit_AccumulateVectorAdd : public FRigUnit_AccumulateBase
 {
 	GENERATED_BODY()
@@ -80,7 +80,7 @@ struct CONTROLRIG_API FRigUnit_AccumulateVectorAdd : public FRigUnit_AccumulateB
 /**
  * Multiplies a value over time over and over again
  */
-USTRUCT(meta=(DisplayName="Accumulate Mul (Float)", PrototypeName="AccumulateMul", Keywords="Simulate,**"))
+USTRUCT(meta=(DisplayName="Accumulate Mul (Float)", TemplateName="AccumulateMul", Keywords="Simulate,**"))
 struct CONTROLRIG_API FRigUnit_AccumulateFloatMul : public FRigUnit_AccumulateBase
 {
 	GENERATED_BODY()
@@ -113,7 +113,7 @@ struct CONTROLRIG_API FRigUnit_AccumulateFloatMul : public FRigUnit_AccumulateBa
 /**
  * Multiplies a vector over time over and over again
  */
-USTRUCT(meta = (DisplayName="Accumulate Mul (Vector)", PrototypeName="AccumulateMul", Keywords="Simulate,**"))
+USTRUCT(meta = (DisplayName="Accumulate Mul (Vector)", TemplateName="AccumulateMul", Keywords="Simulate,**"))
 struct CONTROLRIG_API FRigUnit_AccumulateVectorMul : public FRigUnit_AccumulateBase
 {
 	GENERATED_BODY()
@@ -146,7 +146,7 @@ struct CONTROLRIG_API FRigUnit_AccumulateVectorMul : public FRigUnit_AccumulateB
 /**
  * Multiplies a quaternion over time over and over again
  */
-USTRUCT(meta=(DisplayName="Accumulate Mul (Quaternion)", PrototypeName="AccumulateMul", Keywords="Simulate,**"))
+USTRUCT(meta=(DisplayName="Accumulate Mul (Quaternion)", TemplateName="AccumulateMul", Keywords="Simulate,**"))
 struct CONTROLRIG_API FRigUnit_AccumulateQuatMul : public FRigUnit_AccumulateBase
 {
 	GENERATED_BODY()
@@ -182,7 +182,7 @@ struct CONTROLRIG_API FRigUnit_AccumulateQuatMul : public FRigUnit_AccumulateBas
 /**
  * Multiplies a transform over time over and over again
  */
-USTRUCT(meta = (DisplayName="Accumulate Mul (Transform)", PrototypeName="AccumulateMul", Keywords="Simulate,**"))
+USTRUCT(meta = (DisplayName="Accumulate Mul (Transform)", TemplateName="AccumulateMul", Keywords="Simulate,**"))
 struct CONTROLRIG_API FRigUnit_AccumulateTransformMul : public FRigUnit_AccumulateBase
 {
 	GENERATED_BODY()
@@ -218,7 +218,7 @@ struct CONTROLRIG_API FRigUnit_AccumulateTransformMul : public FRigUnit_Accumula
 /**
  * Interpolates two values over time over and over again
  */
-USTRUCT(meta=(DisplayName="Accumulate Lerp (Float)", PrototypeName="AccumulateLerp", Keywords="Simulate,Ramp"))
+USTRUCT(meta=(DisplayName="Accumulate Lerp (Float)", TemplateName="AccumulateLerp", Keywords="Simulate,Ramp"))
 struct CONTROLRIG_API FRigUnit_AccumulateFloatLerp : public FRigUnit_AccumulateBase
 {
 	GENERATED_BODY()
@@ -254,7 +254,7 @@ struct CONTROLRIG_API FRigUnit_AccumulateFloatLerp : public FRigUnit_AccumulateB
 /**
  * Interpolates two vectors over time over and over again
  */
-USTRUCT(meta = (DisplayName="Accumulate Lerp (Vector)", PrototypeName="AccumulateLerp", Keywords="Simulate,Ramp"))
+USTRUCT(meta = (DisplayName="Accumulate Lerp (Vector)", TemplateName="AccumulateLerp", Keywords="Simulate,Ramp"))
 struct CONTROLRIG_API FRigUnit_AccumulateVectorLerp : public FRigUnit_AccumulateBase
 {
 	GENERATED_BODY()
@@ -291,7 +291,7 @@ struct CONTROLRIG_API FRigUnit_AccumulateVectorLerp : public FRigUnit_Accumulate
 /**
  * Interpolates two quaternions over time over and over again
  */
-USTRUCT(meta=(DisplayName="Accumulate Lerp (Quaternion)", PrototypeName="AccumulateLerp", Keywords="Simulate,Ramp"))
+USTRUCT(meta=(DisplayName="Accumulate Lerp (Quaternion)", TemplateName="AccumulateLerp", Keywords="Simulate,Ramp"))
 struct CONTROLRIG_API FRigUnit_AccumulateQuatLerp : public FRigUnit_AccumulateBase
 {
 	GENERATED_BODY()
@@ -328,7 +328,7 @@ struct CONTROLRIG_API FRigUnit_AccumulateQuatLerp : public FRigUnit_AccumulateBa
 /**
  * Interpolates two transforms over time over and over again
  */
-USTRUCT(meta = (DisplayName="Accumulate Lerp (Transform)", PrototypeName="AccumulateLerp", Keywords="Simulate,Ramp"))
+USTRUCT(meta = (DisplayName="Accumulate Lerp (Transform)", TemplateName="AccumulateLerp", Keywords="Simulate,Ramp"))
 struct CONTROLRIG_API FRigUnit_AccumulateTransformLerp : public FRigUnit_AccumulateBase
 {
 	GENERATED_BODY()
@@ -365,7 +365,7 @@ struct CONTROLRIG_API FRigUnit_AccumulateTransformLerp : public FRigUnit_Accumul
 /**
  * Accumulates the min and max values over time
  */
-USTRUCT(meta=(DisplayName="Accumulate Range (Float)", PrototypeName="AccumulateRange", Keywords="Range"))
+USTRUCT(meta=(DisplayName="Accumulate Range (Float)", TemplateName="AccumulateRange", Keywords="Range"))
 struct CONTROLRIG_API FRigUnit_AccumulateFloatRange : public FRigUnit_AccumulateBase
 {
 	GENERATED_BODY()
@@ -397,7 +397,7 @@ struct CONTROLRIG_API FRigUnit_AccumulateFloatRange : public FRigUnit_Accumulate
 /**
  * Accumulates the min and max values over time
  */
-USTRUCT(meta = (DisplayName="Accumulate Range (Vector)", PrototypeName="AccumulateRange", Keywords="Range"))
+USTRUCT(meta = (DisplayName="Accumulate Range (Vector)", TemplateName="AccumulateRange", Keywords="Range"))
 struct CONTROLRIG_API FRigUnit_AccumulateVectorRange : public FRigUnit_AccumulateBase
 {
 	GENERATED_BODY()
diff --git a/Engine/Plugins/Animation/ControlRig/Source/ControlRig/Private/Units/Simulation/RigUnit_AlphaInterp.h b/Engine/Plugins/Animation/ControlRig/Source/ControlRig/Private/Units/Simulation/RigUnit_AlphaInterp.h
--- a/Engine/Plugins/Animation/ControlRig/Source/ControlRig/Private/Units/Simulation/RigUnit_AlphaInterp.h
+++ b/Engine/Plugins/Animation/ControlRig/Source/ControlRig/Private/Units/Simulation/RigUnit_AlphaInterp.h
@@ -9,7 +9,7 @@
 /**
  * Takes in a float value and outputs an accumulated value with a customized scale and clamp
  */
-USTRUCT(meta=(DisplayName="Alpha Interpolate", Keywords="Alpha,Lerp,LinearInterpolate", Category = "Simulation|Time", PrototypeName = "AlphaInterp", MenuDescSuffix = "(Float)"))
+USTRUCT(meta=(DisplayName="Alpha Interpolate", Keywords="Alpha,Lerp,LinearInterpolate", Category = "Simulation|Time", TemplateName="AlphaInterp", MenuDescSuffix = "(Float)"))
 struct CONTROLRIG_API FRigUnit_AlphaInterp : public FRigUnit_SimBase
 {
 	GENERATED_BODY()
@@ -82,7 +82,7 @@ struct CONTROLRIG_API FRigUnit_AlphaInterp : public FRigUnit_SimBase
 /**
  * Takes in a vector value and outputs an accumulated value with a customized scale and clamp
  */
-USTRUCT(meta=(DisplayName="Alpha Interpolate", Keywords="Alpha,Lerp,LinearInterpolate", Category = "Simulation|Time", PrototypeName = "AlphaInterp", MenuDescSuffix = "(Vector)"))
+USTRUCT(meta=(DisplayName="Alpha Interpolate", Keywords="Alpha,Lerp,LinearInterpolate", Category = "Simulation|Time", TemplateName="AlphaInterp", MenuDescSuffix = "(Vector)"))
 struct CONTROLRIG_API FRigUnit_AlphaInterpVector : public FRigUnit_SimBase
 {
 	GENERATED_BODY()
diff --git a/Engine/Plugins/Animation/ControlRig/Source/ControlRig/Private/Units/Simulation/RigUnit_DeltaFromPrevious.h b/Engine/Plugins/Animation/ControlRig/Source/ControlRig/Private/Units/Simulation/RigUnit_DeltaFromPrevious.h
--- a/Engine/Plugins/Animation/ControlRig/Source/ControlRig/Private/Units/Simulation/RigUnit_DeltaFromPrevious.h
+++ b/Engine/Plugins/Animation/ControlRig/Source/ControlRig/Private/Units/Simulation/RigUnit_DeltaFromPrevious.h
@@ -8,7 +8,7 @@
 /**
  * Computes the difference from the previous value going through the node
  */
-USTRUCT(meta=(DisplayName="DeltaFromPrevious (Float)", Category = "Simulation|Time", PrototypeName="DeltaFromPrevious", Keywords="Difference,Velocity,Acceleration"))
+USTRUCT(meta=(DisplayName="DeltaFromPrevious (Float)", Category = "Simulation|Time", TemplateName="DeltaFromPrevious", Keywords="Difference,Velocity,Acceleration"))
 struct CONTROLRIG_API FRigUnit_DeltaFromPreviousFloat : public FRigUnit_SimBase
 {
 	GENERATED_BODY()
@@ -37,7 +37,7 @@ struct CONTROLRIG_API FRigUnit_DeltaFromPreviousFloat : public FRigUnit_SimBase
 /**
  * Computes the difference from the previous value going through the node
  */
-USTRUCT(meta=(DisplayName="DeltaFromPrevious (Vector)", Category = "Simulation|Time", PrototypeName="DeltaFromPrevious", Keywords="Difference,Velocity,Acceleration"))
+USTRUCT(meta=(DisplayName="DeltaFromPrevious (Vector)", Category = "Simulation|Time", TemplateName="DeltaFromPrevious", Keywords="Difference,Velocity,Acceleration"))
 struct CONTROLRIG_API FRigUnit_DeltaFromPreviousVector : public FRigUnit_SimBase
 {
 	GENERATED_BODY()
@@ -66,7 +66,7 @@ struct CONTROLRIG_API FRigUnit_DeltaFromPreviousVector : public FRigUnit_SimBase
 /**
  * Computes the difference from the previous value going through the node
  */
-USTRUCT(meta=(DisplayName="DeltaFromPrevious (Quaternion)", Category = "Simulation|Time", PrototypeName="DeltaFromPrevious", Keywords="Difference,Velocity,Acceleration"))
+USTRUCT(meta=(DisplayName="DeltaFromPrevious (Quaternion)", Category = "Simulation|Time", TemplateName="DeltaFromPrevious", Keywords="Difference,Velocity,Acceleration"))
 struct CONTROLRIG_API FRigUnit_DeltaFromPreviousQuat : public FRigUnit_SimBase
 {
 	GENERATED_BODY()
@@ -95,7 +95,7 @@ struct CONTROLRIG_API FRigUnit_DeltaFromPreviousQuat : public FRigUnit_SimBase
 /**
  * Computes the difference from the previous value going through the node
  */
-USTRUCT(meta=(DisplayName="DeltaFromPrevious (Transform)", Category = "Simulation|Time", PrototypeName="DeltaFromPrevious", Keywords="Difference,Velocity,Acceleration"))
+USTRUCT(meta=(DisplayName="DeltaFromPrevious (Transform)", Category = "Simulation|Time", TemplateName="DeltaFromPrevious", Keywords="Difference,Velocity,Acceleration"))
 struct CONTROLRIG_API FRigUnit_DeltaFromPreviousTransform : public FRigUnit_SimBase
 {
 	GENERATED_BODY()
diff --git a/Engine/Plugins/Animation/ControlRig/Source/ControlRig/Private/Units/Simulation/RigUnit_Kalman.h b/Engine/Plugins/Animation/ControlRig/Source/ControlRig/Private/Units/Simulation/RigUnit_Kalman.h
--- a/Engine/Plugins/Animation/ControlRig/Source/ControlRig/Private/Units/Simulation/RigUnit_Kalman.h
+++ b/Engine/Plugins/Animation/ControlRig/Source/ControlRig/Private/Units/Simulation/RigUnit_Kalman.h
@@ -9,7 +9,7 @@
  * Averages a value over time.
  * This uses a Kalman Filter internally.
  */
-USTRUCT(meta=(DisplayName="Average Over Time (Float)", Category = "Simulation|Time", PrototypeName="KalmanFilter", Keywords="Average,Smooth,KalmanFilter"))
+USTRUCT(meta=(DisplayName="Average Over Time (Float)", Category = "Simulation|Time", TemplateName="KalmanFilter", Keywords="Average,Smooth,KalmanFilter"))
 struct CONTROLRIG_API FRigUnit_KalmanFloat : public FRigUnit_SimBase
 {
 	GENERATED_BODY()
@@ -44,7 +44,7 @@ struct CONTROLRIG_API FRigUnit_KalmanFloat : public FRigUnit_SimBase
  * Averages a value over time.
  * This uses a Kalman Filter internally.
  */
-USTRUCT(meta=(DisplayName="Average Over Time (Vector)", Category = "Simulation|Time", PrototypeName="KalmanFilter", Keywords="Average,Smooth,KalmanFilter"))
+USTRUCT(meta=(DisplayName="Average Over Time (Vector)", Category = "Simulation|Time", TemplateName="KalmanFilter", Keywords="Average,Smooth,KalmanFilter"))
 struct CONTROLRIG_API FRigUnit_KalmanVector : public FRigUnit_SimBase
 {
 	GENERATED_BODY()
@@ -79,7 +79,7 @@ struct CONTROLRIG_API FRigUnit_KalmanVector : public FRigUnit_SimBase
  * Averages a transform over time.
  * This uses a Kalman Filter internally.
  */
-USTRUCT(meta=(DisplayName="Average Over Time (Transform)", Category = "Simulation|Time", PrototypeName="KalmanFilter", Keywords="Average,Smooth,KalmanFilter"))
+USTRUCT(meta=(DisplayName="Average Over Time (Transform)", Category = "Simulation|Time", TemplateName="KalmanFilter", Keywords="Average,Smooth,KalmanFilter"))
 struct CONTROLRIG_API FRigUnit_KalmanTransform : public FRigUnit_SimBase
 {
 	GENERATED_BODY()
diff --git a/Engine/Plugins/Animation/ControlRig/Source/ControlRig/Private/Units/Simulation/RigUnit_SpringInterp.h b/Engine/Plugins/Animation/ControlRig/Source/ControlRig/Private/Units/Simulation/RigUnit_SpringInterp.h
--- a/Engine/Plugins/Animation/ControlRig/Source/ControlRig/Private/Units/Simulation/RigUnit_SpringInterp.h
+++ b/Engine/Plugins/Animation/ControlRig/Source/ControlRig/Private/Units/Simulation/RigUnit_SpringInterp.h
@@ -9,7 +9,7 @@
 /**
  * Uses a simple spring model to interpolate a float from Current to Target.
  */
-USTRUCT(meta=(DisplayName="Spring Interpolate", Keywords="Alpha,SpringInterpolate", PrototypeName = "SpringInterp", Deprecated="5.0"))
+USTRUCT(meta=(DisplayName="Spring Interpolate", Keywords="Alpha,SpringInterpolate", TemplateName="SpringInterp", Deprecated="5.0"))
 struct CONTROLRIG_API FRigUnit_SpringInterp : public FRigUnit_SimBase
 {
 	GENERATED_BODY()
@@ -54,7 +54,7 @@ struct CONTROLRIG_API FRigUnit_SpringInterp : public FRigUnit_SimBase
 /**
  * Uses a simple spring model to interpolate a vector from Current to Target.
  */
-USTRUCT(meta=(DisplayName="Spring Interpolate", Keywords="Alpha,SpringInterpolate", PrototypeName = "SpringInterp", MenuDescSuffix = "(Vector)", Deprecated="5.0"))
+USTRUCT(meta=(DisplayName="Spring Interpolate", Keywords="Alpha,SpringInterpolate", TemplateName="SpringInterp", MenuDescSuffix = "(Vector)", Deprecated="5.0"))
 struct CONTROLRIG_API FRigUnit_SpringInterpVector : public FRigUnit_SimBase
 {
 	GENERATED_BODY()
@@ -99,7 +99,7 @@ struct CONTROLRIG_API FRigUnit_SpringInterpVector : public FRigUnit_SimBase
 /**
  * Uses a simple spring model to interpolate a float from Current to Target.
  */
-USTRUCT(meta=(DisplayName="Spring Interpolate", Keywords="Alpha,SpringInterpolate,Verlet", Category = "Simulation|Springs", PrototypeName = "SpringInterp"))
+USTRUCT(meta=(DisplayName="Spring Interpolate", Keywords="Alpha,SpringInterpolate,Verlet", Category = "Simulation|Springs", TemplateName="SpringInterp"))
 struct CONTROLRIG_API FRigUnit_SpringInterpV2 : public FRigUnit_SimBase
 {
 	GENERATED_BODY()
@@ -189,7 +189,7 @@ struct CONTROLRIG_API FRigUnit_SpringInterpV2 : public FRigUnit_SimBase
 /**
  * Uses a simple spring model to interpolate a vector from Current to Target.
  */
-USTRUCT(meta=(DisplayName="Spring Interpolate", Keywords="Alpha,SpringInterpolate,Verlet", Category = "Simulation|Springs", PrototypeName = "SpringInterp", MenuDescSuffix = "(Vector)"))
+USTRUCT(meta=(DisplayName="Spring Interpolate", Keywords="Alpha,SpringInterpolate,Verlet", Category = "Simulation|Springs", TemplateName="SpringInterp", MenuDescSuffix = "(Vector)"))
 struct CONTROLRIG_API FRigUnit_SpringInterpVectorV2 : public FRigUnit_SimBase
 {
 	GENERATED_BODY()
@@ -279,7 +279,7 @@ struct CONTROLRIG_API FRigUnit_SpringInterpVectorV2 : public FRigUnit_SimBase
 /**
  * Uses a simple spring model to interpolate a quaternion from Current to Target.
  */
-USTRUCT(meta=(DisplayName="Spring Interpolate", Keywords="Alpha,SpringInterpolate,Verlet", Category = "Simulation|Springs", PrototypeName = "SpringInterp", MenuDescSuffix = "(Quaternion)"))
+USTRUCT(meta=(DisplayName="Spring Interpolate", Keywords="Alpha,SpringInterpolate,Verlet", Category = "Simulation|Springs", TemplateName="SpringInterp", MenuDescSuffix = "(Quaternion)"))
 struct CONTROLRIG_API FRigUnit_SpringInterpQuaternionV2 : public FRigUnit_SimBase
 {
 	GENERATED_BODY()
diff --git a/Engine/Plugins/Animation/ControlRig/Source/ControlRig/Private/Units/Simulation/RigUnit_TimeOffset.h b/Engine/Plugins/Animation/ControlRig/Source/ControlRig/Private/Units/Simulation/RigUnit_TimeOffset.h
--- a/Engine/Plugins/Animation/ControlRig/Source/ControlRig/Private/Units/Simulation/RigUnit_TimeOffset.h
+++ b/Engine/Plugins/Animation/ControlRig/Source/ControlRig/Private/Units/Simulation/RigUnit_TimeOffset.h
@@ -8,7 +8,7 @@
 /**
  * Records a value over time and can access the value from the past
  */
-USTRUCT(meta=(DisplayName="Value Over Time (Float)", Category = "Simulation|Time", PrototypeName="TimeOffset", Keywords="Buffer,Delta,History,Previous,TimeOffset,Delay"))
+USTRUCT(meta=(DisplayName="Value Over Time (Float)", Category = "Simulation|Time", TemplateName="TimeOffset", Keywords="Buffer,Delta,History,Previous,TimeOffset,Delay"))
 struct CONTROLRIG_API FRigUnit_TimeOffsetFloat : public FRigUnit_SimBase
 {
 	GENERATED_BODY()
@@ -61,7 +61,7 @@ struct CONTROLRIG_API FRigUnit_TimeOffsetFloat : public FRigUnit_SimBase
 /**
  * Records a value over time and can access the value from the past
  */
-USTRUCT(meta=(DisplayName="Value Over Time (Vector)", Category = "Simulation|Time", PrototypeName="TimeOffset", Keywords="Buffer,Delta,History,Previous,TimeOffset,Delay"))
+USTRUCT(meta=(DisplayName="Value Over Time (Vector)", Category = "Simulation|Time", TemplateName="TimeOffset", Keywords="Buffer,Delta,History,Previous,TimeOffset,Delay"))
 struct CONTROLRIG_API FRigUnit_TimeOffsetVector : public FRigUnit_SimBase
 {
 	GENERATED_BODY()
@@ -114,7 +114,7 @@ struct CONTROLRIG_API FRigUnit_TimeOffsetVector : public FRigUnit_SimBase
 /**
  * Records a value over time and can access the value from the past
  */
-USTRUCT(meta=(DisplayName="Value Over Time (Transform)", Category = "Simulation|Time", PrototypeName="TimeOffset", Keywords="Buffer,Delta,History,Previous,TimeOffset,Delay"))
+USTRUCT(meta=(DisplayName="Value Over Time (Transform)", Category = "Simulation|Time", TemplateName="TimeOffset", Keywords="Buffer,Delta,History,Previous,TimeOffset,Delay"))
 struct CONTROLRIG_API FRigUnit_TimeOffsetTransform : public FRigUnit_SimBase
 {
 	GENERATED_BODY()
diff --git a/Engine/Plugins/Animation/ControlRig/Source/ControlRig/Private/Units/Simulation/RigUnit_Verlet.h b/Engine/Plugins/Animation/ControlRig/Source/ControlRig/Private/Units/Simulation/RigUnit_Verlet.h
--- a/Engine/Plugins/Animation/ControlRig/Source/ControlRig/Private/Units/Simulation/RigUnit_Verlet.h
+++ b/Engine/Plugins/Animation/ControlRig/Source/ControlRig/Private/Units/Simulation/RigUnit_Verlet.h
@@ -10,7 +10,7 @@
  * Simulates a single position over time using Verlet integration. It is recommended to use SpringInterp instead as it
  * is more accurate and stable, and has more meaningful parameters.
  */
-USTRUCT(meta=(DisplayName="Verlet (Vector)", Category = "Simulation|Springs", PrototypeName="Verlet", Keywords="Simulate,Integrate"))
+USTRUCT(meta=(DisplayName="Verlet (Vector)", Category = "Simulation|Springs", TemplateName="Verlet", Keywords="Simulate,Integrate"))
 struct CONTROLRIG_API FRigUnit_VerletIntegrateVector : public FRigUnit_SimBase
 {
 	GENERATED_BODY()
diff --git a/Engine/Source/Runtime/RigVM/Private/RigVMCore/RigVMRegistry.cpp b/Engine/Source/Runtime/RigVM/Private/RigVMCore/RigVMRegistry.cpp
--- a/Engine/Source/Runtime/RigVM/Private/RigVMCore/RigVMRegistry.cpp
+++ b/Engine/Source/Runtime/RigVM/Private/RigVMCore/RigVMRegistry.cpp
@@ -5,7 +5,7 @@
 #include "UObject/UObjectIterator.h"
 
 FRigVMRegistry FRigVMRegistry::s_RigVMRegistry;
-const FName FRigVMRegistry::TemplateNameMetaName = TEXT("PrototypeName");
+const FName FRigVMRegistry::TemplateNameMetaName = TEXT("TemplateName");
 
 FRigVMRegistry& FRigVMRegistry::Get()
 {
