Fixed crash when switching space: FindKey may also return Num() and not just INDEX_NONE as the documentation said

diff --git a/Engine/Plugins/Animation/ControlRig/Source/ControlRigEditor/Private/Sequencer/ControlRigSpaceChannelEditors.cpp b/Engine/Plugins/Animation/ControlRig/Source/ControlRigEditor/Private/Sequencer/ControlRigSpaceChannelEditors.cpp
--- a/Engine/Plugins/Animation/ControlRig/Source/ControlRigEditor/Private/Sequencer/ControlRigSpaceChannelEditors.cpp
+++ b/Engine/Plugins/Animation/ControlRig/Source/ControlRigEditor/Private/Sequencer/ControlRigSpaceChannelEditors.cpp
@@ -1408,18 +1408,21 @@ void FControlRigSpaceChannelHelpers::EvaluateTangentAtThisTime(	UControlRig* Con
 			return DefaultTangent;
 		}
 
+		const TArrayView<const FMovieSceneFloatValue> Values = ChannelInterface.GetValues();
+		
 		// look around to get the closest key tangent if fairly close
 		const int32 Tolerance = static_cast<int32>(InChannel->GetTickResolution().AsDecimal() * 0.01);
+		// NOTE FindKey might return Times.Num() (see AlgoImpl::LowerBoundInternal)
 		const int32 ExistingIndex = ChannelInterface.FindKey(InTime, Tolerance);
-		if (ExistingIndex != INDEX_NONE)
+		if (Times.IsValidIndex(ExistingIndex) && Values.IsValidIndex(ExistingIndex))
 		{
-			const FMovieSceneFloatValue& Value = ChannelInterface.GetValues()[ExistingIndex];
 			const FFrameNumber& Time = Times[ExistingIndex];
 			const int32 DiffToKey = InTime.Value - Time.Value;
 
 			// if the closest key is within a threshold, we use it's tangent directly instead of computing one
 			if (FMath::Abs(DiffToKey) <= Tolerance)
 			{
+				const FMovieSceneFloatValue& Value = Values[ExistingIndex];
 				if (DiffToKey == 0)
 				{
 					return Value.Tangent;				
diff --git a/Engine/Source/Runtime/MovieScene/Public/Channels/MovieSceneChannelData.h b/Engine/Source/Runtime/MovieScene/Public/Channels/MovieSceneChannelData.h
--- a/Engine/Source/Runtime/MovieScene/Public/Channels/MovieSceneChannelData.h
+++ b/Engine/Source/Runtime/MovieScene/Public/Channels/MovieSceneChannelData.h
@@ -555,7 +555,7 @@ struct TMovieSceneChannelData<const ValueType>
 	 *
 	 * @param InTime         The time at which to search
 	 * @param InTolerance    A tolerance of frame numbers to allow either side of the specified time
-	 * @return The index of the key closest to InTime and within InTolerance, or INDEX_NONE
+	 * @return The index of the key closest to InTime and within InTolerance, INDEX_NONE or Times.Num()
 	 */
 	int32 FindKey(FFrameNumber InTime, FFrameNumber InTolerance = 0) const
 	{
