Control Rig: User interface access to upgrade node functionality

diff --git a/Engine/Plugins/Animation/ControlRig/Source/ControlRig/Private/Units/Hierarchy/RigUnit_GetBoneTransform.cpp b/Engine/Plugins/Animation/ControlRig/Source/ControlRig/Private/Units/Hierarchy/RigUnit_GetBoneTransform.cpp
--- a/Engine/Plugins/Animation/ControlRig/Source/ControlRig/Private/Units/Hierarchy/RigUnit_GetBoneTransform.cpp
+++ b/Engine/Plugins/Animation/ControlRig/Source/ControlRig/Private/Units/Hierarchy/RigUnit_GetBoneTransform.cpp
@@ -56,7 +56,7 @@ FRigUnit_GetBoneTransform_Execute()
 FRigVMStructUpgradeInfo FRigUnit_GetBoneTransform::GetUpgradeInfo() const
 {
 	FRigUnit_GetTransform NewNode;
-	NewNode.Item = FRigElementKey(Bone, ERigElementType::Control);
+	NewNode.Item = FRigElementKey(Bone, ERigElementType::Bone);
 	NewNode.Space = Space;
 
 	FRigVMStructUpgradeInfo Info(*this, NewNode);
diff --git a/Engine/Plugins/Animation/ControlRig/Source/ControlRig/Private/Units/Hierarchy/RigUnit_GetInitialBoneTransform.cpp b/Engine/Plugins/Animation/ControlRig/Source/ControlRig/Private/Units/Hierarchy/RigUnit_GetInitialBoneTransform.cpp
--- a/Engine/Plugins/Animation/ControlRig/Source/ControlRig/Private/Units/Hierarchy/RigUnit_GetInitialBoneTransform.cpp
+++ b/Engine/Plugins/Animation/ControlRig/Source/ControlRig/Private/Units/Hierarchy/RigUnit_GetInitialBoneTransform.cpp
@@ -54,7 +54,7 @@ FRigUnit_GetInitialBoneTransform_Execute()
 FRigVMStructUpgradeInfo FRigUnit_GetInitialBoneTransform::GetUpgradeInfo() const
 {
 	FRigUnit_GetTransform NewNode;
-	NewNode.Item = FRigElementKey(Bone, ERigElementType::Control);
+	NewNode.Item = FRigElementKey(Bone, ERigElementType::Bone);
 	NewNode.Space = Space;
 	NewNode.bInitial = true;
 
diff --git a/Engine/Plugins/Animation/ControlRig/Source/ControlRig/Private/Units/Hierarchy/RigUnit_GetSpaceTransform.cpp b/Engine/Plugins/Animation/ControlRig/Source/ControlRig/Private/Units/Hierarchy/RigUnit_GetSpaceTransform.cpp
--- a/Engine/Plugins/Animation/ControlRig/Source/ControlRig/Private/Units/Hierarchy/RigUnit_GetSpaceTransform.cpp
+++ b/Engine/Plugins/Animation/ControlRig/Source/ControlRig/Private/Units/Hierarchy/RigUnit_GetSpaceTransform.cpp
@@ -50,7 +50,7 @@ FRigUnit_GetSpaceTransform_Execute()
 FRigVMStructUpgradeInfo FRigUnit_GetSpaceTransform::GetUpgradeInfo() const
 {
 	FRigUnit_GetTransform NewNode;
-	NewNode.Item = FRigElementKey(Space, ERigElementType::Control);
+	NewNode.Item = FRigElementKey(Space, ERigElementType::Null);
 	NewNode.Space = SpaceType;
 
 	FRigVMStructUpgradeInfo Info(*this, NewNode);
diff --git a/Engine/Plugins/Animation/ControlRig/Source/ControlRig/Private/Units/Hierarchy/RigUnit_SetBoneInitialTransform.cpp b/Engine/Plugins/Animation/ControlRig/Source/ControlRig/Private/Units/Hierarchy/RigUnit_SetBoneInitialTransform.cpp
--- a/Engine/Plugins/Animation/ControlRig/Source/ControlRig/Private/Units/Hierarchy/RigUnit_SetBoneInitialTransform.cpp
+++ b/Engine/Plugins/Animation/ControlRig/Source/ControlRig/Private/Units/Hierarchy/RigUnit_SetBoneInitialTransform.cpp
@@ -51,6 +51,7 @@ FRigVMStructUpgradeInfo FRigUnit_SetBoneInitialTransform::GetUpgradeInfo() const
 	NewNode.Item = FRigElementKey(Bone, ERigElementType::Bone);
 	NewNode.Space = Space;
 	NewNode.Transform = Transform;
+	NewNode.bInitial = true;
 
 	FRigVMStructUpgradeInfo Info(*this, NewNode);
 	Info.AddRemappedPin(TEXT("Bone"), TEXT("Item.Name"));
diff --git a/Engine/Plugins/Animation/ControlRig/Source/ControlRig/Private/Units/Highlevel/Hierarchy/RigUnit_TwoBoneIKSimple.cpp b/Engine/Plugins/Animation/ControlRig/Source/ControlRig/Private/Units/Highlevel/Hierarchy/RigUnit_TwoBoneIKSimple.cpp
--- a/Engine/Plugins/Animation/ControlRig/Source/ControlRig/Private/Units/Highlevel/Hierarchy/RigUnit_TwoBoneIKSimple.cpp
+++ b/Engine/Plugins/Animation/ControlRig/Source/ControlRig/Private/Units/Highlevel/Hierarchy/RigUnit_TwoBoneIKSimple.cpp
@@ -39,7 +39,7 @@ FRigVMStructUpgradeInfo FRigUnit_TwoBoneIKSimple::GetUpgradeInfo() const
 {
 	FRigUnit_TwoBoneIKSimplePerItem NewNode;
 	NewNode.ItemA = FRigElementKey(BoneA, ERigElementType::Bone);
-	NewNode.ItemA = FRigElementKey(BoneB, ERigElementType::Bone);
+	NewNode.ItemB = FRigElementKey(BoneB, ERigElementType::Bone);
 	NewNode.EffectorItem = FRigElementKey(EffectorBone, ERigElementType::Bone);
 	NewNode.Effector = Effector;
 	NewNode.PrimaryAxis = PrimaryAxis;
diff --git a/Engine/Plugins/Animation/ControlRig/Source/ControlRigEditor/Private/ControlRigEditorModule.cpp b/Engine/Plugins/Animation/ControlRig/Source/ControlRigEditor/Private/ControlRigEditorModule.cpp
--- a/Engine/Plugins/Animation/ControlRig/Source/ControlRigEditor/Private/ControlRigEditorModule.cpp
+++ b/Engine/Plugins/Animation/ControlRig/Source/ControlRigEditor/Private/ControlRigEditorModule.cpp
@@ -1685,6 +1685,8 @@ void FControlRigEditorModule::GetContextMenuActions(const UControlRigGraphSchema
 				TArray<FName> SelectedNodeNames = Model->GetSelectNodes();
 				SelectedNodeNames.AddUnique(Context->Node->GetFName());
 
+				bool bCanNodeBeUpgraded = false; 
+
 				for(const FName& SelectedNodeName : SelectedNodeNames)
 				{
 					if (URigVMNode* ModelNode = Model->FindNodeByName(SelectedNodeName))
@@ -1854,6 +1856,8 @@ void FControlRigEditorModule::GetContextMenuActions(const UControlRigGraphSchema
 								}
 							}
 						}
+						
+						bCanNodeBeUpgraded = bCanNodeBeUpgraded || ModelNode->CanBeUpgraded();
 					}
 				}
 
@@ -2216,6 +2220,21 @@ void FControlRigEditorModule::GetContextMenuActions(const UControlRigGraphSchema
 							})
 						));
 					}
+
+					if(bCanNodeBeUpgraded)
+					{
+						FToolMenuSection& VersioningSection = Menu->AddSection("EdGraphSchemaVersioning", LOCTEXT("VersioningHeader", "Versioning"));
+						VersioningSection.AddMenuEntry(
+							"Upgrade Nodes",
+							LOCTEXT("UpgradeNodes", "Upgrade Nodes"),
+							LOCTEXT("UpgradeNodes_Tooltip", "Upgrades deprecated nodes to their current implementation"),
+							FSlateIcon(),
+							FUIAction(FExecuteAction::CreateLambda([Model, Controller]() {
+								TArray<FName> Nodes = Model->GetSelectNodes();
+								Controller->UpgradeNodes(Nodes, true, true);
+							})
+						));
+					}
 					
 					OrganizationSection.AddSubMenu("Alignment", LOCTEXT("AlignmentHeader", "Alignment"), FText(), FNewToolMenuDelegate::CreateLambda([](UToolMenu* AlignmentMenu)
 					{
diff --git a/Engine/Source/Developer/RigVMDeveloper/Private/RigVMModel/RigVMController.cpp b/Engine/Source/Developer/RigVMDeveloper/Private/RigVMModel/RigVMController.cpp
--- a/Engine/Source/Developer/RigVMDeveloper/Private/RigVMModel/RigVMController.cpp
+++ b/Engine/Source/Developer/RigVMDeveloper/Private/RigVMModel/RigVMController.cpp
@@ -1667,6 +1667,18 @@ TArray<URigVMNode*> URigVMController::UpgradeNodes(const TArray<FName>& InNodeNa
 				*NodeNamesJoined));
 	}
 
+	// log a warning for all nodes which are still marked deprecated
+	for(URigVMNode* Node : Nodes)
+	{
+		if(URigVMUnitNode* UnitNode = Cast<URigVMUnitNode>(Node))
+		{
+			if(UnitNode->IsDeprecated())
+			{
+				ReportWarningf(TEXT("Node %s cannot be upgraded. There is no automatic upgrade path available."), *UnitNode->GetNodePath());
+			}
+		}
+	}
+
 	return Nodes;
 }
 
