Control Rig: Fix crash when dragging asset to level after reload

diff --git a/Engine/Plugins/Animation/ControlRig/Source/ControlRigEditor/Private/ControlRigThumbnailRenderer.cpp b/Engine/Plugins/Animation/ControlRig/Source/ControlRigEditor/Private/ControlRigThumbnailRenderer.cpp
--- a/Engine/Plugins/Animation/ControlRig/Source/ControlRigEditor/Private/ControlRigThumbnailRenderer.cpp
+++ b/Engine/Plugins/Animation/ControlRig/Source/ControlRigEditor/Private/ControlRigThumbnailRenderer.cpp
@@ -103,65 +103,68 @@ void UControlRigThumbnailRenderer::AddAdditionalPreviewSceneContent(UObject* Obj
 
 		FTransform ComponentToWorld = ThumbnailScene->GetPreviewActor()->GetSkeletalMeshComponent()->GetComponentToWorld();
 
-		ControlRig->GetHierarchy()->ForEach<FRigControlElement>([&](FRigControlElement* ControlElement) -> bool
+		if (URigHierarchy* Hierarchy = ControlRig->GetHierarchy())
 		{
-			switch (ControlElement->Settings.ControlType)
+			Hierarchy->ForEach<FRigControlElement>([&](FRigControlElement* ControlElement) -> bool
 			{
-				case ERigControlType::Float:
-				case ERigControlType::Integer:
-				case ERigControlType::Vector2D:
-				case ERigControlType::Position:
-				case ERigControlType::Scale:
-				case ERigControlType::Rotator:
-				case ERigControlType::Transform:
-				case ERigControlType::TransformNoScale:
-				case ERigControlType::EulerTransform:
+				switch (ControlElement->Settings.ControlType)
 				{
-					if (const FControlRigShapeDefinition* ShapeDef = RigBlueprint->GetControlShapeByName(ControlElement->Settings.ShapeName))
+					case ERigControlType::Float:
+					case ERigControlType::Integer:
+					case ERigControlType::Vector2D:
+					case ERigControlType::Position:
+					case ERigControlType::Scale:
+					case ERigControlType::Rotator:
+					case ERigControlType::Transform:
+					case ERigControlType::TransformNoScale:
+					case ERigControlType::EulerTransform:
 					{
-						UStaticMesh* StaticMesh = ShapeDef->StaticMesh.Get();
-						if (StaticMesh == nullptr) // not yet loaded
-						{
-							return true;
-						}
-
-						const FTransform ShapeGlobalTransform = ControlRig->GetHierarchy()->GetGlobalControlShapeTransform(ControlElement->GetKey());
-
-						FActorSpawnParameters SpawnInfo;
-						SpawnInfo.SpawnCollisionHandlingOverride = ESpawnActorCollisionHandlingMethod::AlwaysSpawn;
-						SpawnInfo.bNoFail = true;
-						SpawnInfo.ObjectFlags = RF_Transient;
-						AStaticMeshActor* ShapeActor = PreviewWorld->SpawnActor<AStaticMeshActor>(SpawnInfo);
-						ShapeActor->SetActorEnableCollision(false);
-
-						if(!ShapeDef->Library.IsValid())
-						{
-							return true;
-						}
-						
-						UMaterial* DefaultMaterial = ShapeDef->Library.Get()->DefaultMaterial.Get();
-						if (DefaultMaterial == nullptr) // not yet loaded
+						if (const FControlRigShapeDefinition* ShapeDef = RigBlueprint->GetControlShapeByName(ControlElement->Settings.ShapeName))
 						{
-							return true;
+							UStaticMesh* StaticMesh = ShapeDef->StaticMesh.Get();
+							if (StaticMesh == nullptr) // not yet loaded
+							{
+								return true;
+							}
+
+							const FTransform ShapeGlobalTransform = ControlRig->GetHierarchy()->GetGlobalControlShapeTransform(ControlElement->GetKey());
+
+							FActorSpawnParameters SpawnInfo;
+							SpawnInfo.SpawnCollisionHandlingOverride = ESpawnActorCollisionHandlingMethod::AlwaysSpawn;
+							SpawnInfo.bNoFail = true;
+							SpawnInfo.ObjectFlags = RF_Transient;
+							AStaticMeshActor* ShapeActor = PreviewWorld->SpawnActor<AStaticMeshActor>(SpawnInfo);
+							ShapeActor->SetActorEnableCollision(false);
+
+							if (!ShapeDef->Library.IsValid())
+							{
+								return true;
+							}
+
+							UMaterial* DefaultMaterial = ShapeDef->Library.Get()->DefaultMaterial.Get();
+							if (DefaultMaterial == nullptr) // not yet loaded
+							{
+								return true;
+							}
+
+							UMaterialInstanceDynamic* MaterialInstance = UMaterialInstanceDynamic::Create(DefaultMaterial, ShapeActor);
+							MaterialInstance->SetVectorParameterValue(ShapeDef->Library.Get()->MaterialColorParameter,FVector(ControlElement->Settings.ShapeColor));
+							ShapeActor->GetStaticMeshComponent()->SetMaterial(0, MaterialInstance);
+
+							ShapeActors.Add(ControlElement->GetName(), ShapeActor);
+
+							ShapeActor->GetStaticMeshComponent()->SetStaticMesh(StaticMesh);
+							ShapeActor->SetActorTransform(ShapeDef->Transform * ShapeGlobalTransform);
 						}
-
-						UMaterialInstanceDynamic* MaterialInstance = UMaterialInstanceDynamic::Create(DefaultMaterial, ShapeActor);
-						MaterialInstance->SetVectorParameterValue(ShapeDef->Library.Get()->MaterialColorParameter, FVector(ControlElement->Settings.ShapeColor));
-						ShapeActor->GetStaticMeshComponent()->SetMaterial(0, MaterialInstance);
-
-						ShapeActors.Add(ControlElement->GetName(), ShapeActor);
-
-						ShapeActor->GetStaticMeshComponent()->SetStaticMesh(StaticMesh);
-						ShapeActor->SetActorTransform(ShapeDef->Transform * ShapeGlobalTransform);
+						break;
+					}
+					default:
+					{
+						break;
 					}
-					break;
-				}
-				default:
-				{
-					break;
 				}
-			}
-			return true;
-        });
+				return true;
+			});
+		}
 	}
 }
\ No newline at end of file
