Control Rig: Fix static analysis warnings

diff --git a/Engine/Plugins/Animation/ControlRig/Source/ControlRigEditor/Private/EditMode/ControlRigEditMode.cpp b/Engine/Plugins/Animation/ControlRig/Source/ControlRigEditor/Private/EditMode/ControlRigEditMode.cpp
--- a/Engine/Plugins/Animation/ControlRig/Source/ControlRigEditor/Private/EditMode/ControlRigEditMode.cpp
+++ b/Engine/Plugins/Animation/ControlRig/Source/ControlRigEditor/Private/EditMode/ControlRigEditMode.cpp
@@ -745,14 +745,17 @@ void FControlRigEditMode::Render(const FSceneView* View, FViewport* Viewport, FP
 						}
 						const FTransform ParentTransform = Hierarchy->GetTransform(ParentTransformElement, ERigTransformType::CurrentGlobal);
 						const bool bHitTesting = bBoolSetHitProxies && (ParentTransformElement->GetType() == ERigElementType::Bone);
-						if (bHitTesting)
+						if (PDI)
 						{
-							PDI->SetHitProxy(new HFKRigBoneProxy(ParentTransformElement->GetName(), ControlRig));
-						}
-                        PDI->DrawLine(ComponentTransform.TransformPosition(Transform.GetLocation()),ComponentTransform.TransformPosition(ParentTransform.GetLocation()), Color, SDPG_Foreground);
-						if (bHitTesting)
-						{
-							PDI->SetHitProxy(nullptr);
+							if (bHitTesting)
+							{
+								PDI->SetHitProxy(new HFKRigBoneProxy(ParentTransformElement->GetName(), ControlRig));
+							}
+							PDI->DrawLine(ComponentTransform.TransformPosition(Transform.GetLocation()),ComponentTransform.TransformPosition(ParentTransform.GetLocation()), Color, SDPG_Foreground);
+							if (bHitTesting)
+							{
+								PDI->SetHitProxy(nullptr);
+							}
 						}
 					}
                 }
diff --git a/Engine/Plugins/Animation/ControlRig/Source/ControlRigEditor/Private/Sequencer/ControlRigSpaceChannelEditors.cpp b/Engine/Plugins/Animation/ControlRig/Source/ControlRigEditor/Private/Sequencer/ControlRigSpaceChannelEditors.cpp
--- a/Engine/Plugins/Animation/ControlRig/Source/ControlRigEditor/Private/Sequencer/ControlRigSpaceChannelEditors.cpp
+++ b/Engine/Plugins/Animation/ControlRig/Source/ControlRigEditor/Private/Sequencer/ControlRigSpaceChannelEditors.cpp
@@ -1422,7 +1422,7 @@ void FControlRigSpaceChannelHelpers::EvaluateTangentAtThisTime(	UControlRig* Con
 			// if the closest key is within a threshold, we use it's tangent directly instead of computing one
 			if (FMath::Abs(DiffToKey) <= Tolerance)
 			{
-				const FMovieSceneFloatValue& Value = Values[ExistingIndex];
+				const FMovieSceneFloatValue Value = Values[ExistingIndex];
 				if (DiffToKey == 0)
 				{
 					return Value.Tangent;				
@@ -1515,7 +1515,7 @@ void FControlRigSpaceChannelHelpers::SetTangentsAtThisTime(	UControlRig* Control
 		const int32 KeyIndex = ChannelInterface.FindKey(Time);
 		if (KeyIndex != INDEX_NONE)
 		{
-			FMovieSceneFloatValue& Value = ChannelInterface.GetValues()[KeyIndex];
+			FMovieSceneFloatValue Value = ChannelInterface.GetValues()[KeyIndex];
 			Value.Tangent.ArriveTangent = InTangents[Index].ArriveTangent;
 			Value.Tangent.LeaveTangent = InTangents[Index].LeaveTangent;
 			Value.TangentMode = RCTM_Break;	
