Control Rig: Transform Location - added Keyword and fixed Description #Jira UE-132737 #preflight

diff --git a/Engine/Plugins/Animation/ControlRig/Source/ControlRig/Private/Units/Math/RigUnit_MathTransform.h b/Engine/Plugins/Animation/ControlRig/Source/ControlRig/Private/Units/Math/RigUnit_MathTransform.h
--- a/Engine/Plugins/Animation/ControlRig/Source/ControlRig/Private/Units/Math/RigUnit_MathTransform.h
+++ b/Engine/Plugins/Animation/ControlRig/Source/ControlRig/Private/Units/Math/RigUnit_MathTransform.h
@@ -313,9 +313,9 @@ struct CONTROLRIG_API FRigUnit_MathTransformRotateVector : public FRigUnit_MathT
 };
 
 /**
- * Rotates a given vector (location) by the transform
+ * Multiplies a given vector (location) by the transform
  */
-USTRUCT(meta=(DisplayName="Transform Location", TemplateName="Multiply"))
+USTRUCT(meta=(DisplayName="Transform Location", TemplateName="Multiply", Keywords="Multiply"))
 struct CONTROLRIG_API FRigUnit_MathTransformTransformVector : public FRigUnit_MathTransformBase
 {
 	GENERATED_BODY()
