attempt to fix function reference registry while GC ing: the initialization is defered at first ControlRig editor initialization instead of when loading CR's module

diff --git a/Engine/Plugins/Animation/ControlRig/Source/ControlRigEditor/Private/ControlRigEditorModule.cpp b/Engine/Plugins/Animation/ControlRig/Source/ControlRigEditor/Private/ControlRigEditorModule.cpp
--- a/Engine/Plugins/Animation/ControlRig/Source/ControlRigEditor/Private/ControlRigEditorModule.cpp
+++ b/Engine/Plugins/Animation/ControlRig/Source/ControlRigEditor/Private/ControlRigEditorModule.cpp
@@ -249,44 +249,6 @@ void FControlRigEditorModule::StartupModule()
 	//UThumbnailManager::Get().RegisterCustomRenderer(UControlRigPoseAsset::StaticClass(), UControlRigPoseThumbnailRenderer::StaticClass());
 
 	bFilterAssetBySkeleton = true;
-
-	FFunctionGraphTask::CreateAndDispatchWhenReady([]()
-	{
-		if(URigVMBuildData* BuildData = URigVMController::GetBuildData())
-		{
-			const FArrayProperty* ReferenceNodeDataProperty =
-				CastField<FArrayProperty>(UControlRigBlueprint::StaticClass()->FindPropertyByName(TEXT("FunctionReferenceNodeData")));
-			if(ReferenceNodeDataProperty)
-			{
-				const FAssetRegistryModule& AssetRegistryModule = FModuleManager::LoadModuleChecked<FAssetRegistryModule>("AssetRegistry");
-
-				// find all control rigs in the project
-				TArray<FAssetData> ControlRigAssetDatas;
-				FARFilter ControlRigAssetFilter;
-				ControlRigAssetFilter.ClassNames.Add(UControlRigBlueprint::StaticClass()->GetFName());
-				AssetRegistryModule.Get().GetAssets(ControlRigAssetFilter, ControlRigAssetDatas);
-
-				// loop over all control rigs in the project
-				for(const FAssetData& ControlRigAssetData : ControlRigAssetDatas)
-				{
-					const FString ReferenceNodeDataString =
-						ControlRigAssetData.GetTagValueRef<FString>(ReferenceNodeDataProperty->GetFName());
-					if(ReferenceNodeDataString.IsEmpty())
-					{
-						continue;
-					}
-
-					TArray<FRigVMReferenceNodeData> ReferenceNodeDatas;
-					ReferenceNodeDataProperty->ImportText_Direct(*ReferenceNodeDataString, &ReferenceNodeDatas, nullptr, EPropertyPortFlags::PPF_None);
-
-					for(const FRigVMReferenceNodeData& ReferenceNodeData : ReferenceNodeDatas)
-					{
-						BuildData->RegisterFunctionReference(ReferenceNodeData);
-					}
-				}
-			}
-		}
-	}, TStatId(), NULL, ENamedThreads::GameThread);
 }
 
 void FControlRigEditorModule::ShutdownModule()
diff --git a/Engine/Plugins/Animation/ControlRig/Source/ControlRigEditor/Private/Editor/ControlRigEditor.cpp b/Engine/Plugins/Animation/ControlRig/Source/ControlRigEditor/Private/Editor/ControlRigEditor.cpp
--- a/Engine/Plugins/Animation/ControlRig/Source/ControlRigEditor/Private/Editor/ControlRigEditor.cpp
+++ b/Engine/Plugins/Animation/ControlRig/Source/ControlRigEditor/Private/Editor/ControlRigEditor.cpp
@@ -96,6 +96,7 @@
 #include "Kismet2/WatchedPin.h"
 #include "ToolMenus.h"
 #include "Editor/EditorStyle/Private/SlateEditorStyle.h"
+#include "AssetRegistryModule.h"
 
 #define LOCTEXT_NAMESPACE "ControlRigEditor"
 
@@ -103,6 +104,8 @@ const FName ControlRigEditorAppName(TEXT("ControlRigEditorApp"));
 
 const FName FControlRigEditorModes::ControlRigEditorMode("Rigging");
 
+bool FControlRigEditor::bAreFunctionReferencesInitialized = false;
+
 namespace ControlRigEditorTabs
 {
 	const FName DetailsTab(TEXT("DetailsTab"));
@@ -337,6 +340,9 @@ void FControlRigEditor::InitControlRigEditor(const EToolkitMode::Type Mode, cons
 
 	CommonInitialization(ControlRigBlueprints, false);
 
+	// update function references once
+	InitFunctionReferences();
+	
 	// user-defined-struct can change even after load
 	// refresh the models such that pins are updated to match
 	// the latest struct member layout
@@ -568,6 +574,54 @@ void FControlRigEditor::InitControlRigEditor(const EToolkitMode::Type Mode, cons
 	}
 }
 
+void FControlRigEditor::InitFunctionReferences()
+{
+	if (bAreFunctionReferencesInitialized)
+	{
+		return;
+	}
+	
+	FFunctionGraphTask::CreateAndDispatchWhenReady([]()
+	{
+		if(URigVMBuildData* BuildData = URigVMController::GetBuildData())
+		{
+			const FArrayProperty* ReferenceNodeDataProperty =
+				CastField<FArrayProperty>(UControlRigBlueprint::StaticClass()->FindPropertyByName(TEXT("FunctionReferenceNodeData")));
+			if(ReferenceNodeDataProperty)
+			{
+				const FAssetRegistryModule& AssetRegistryModule = FModuleManager::LoadModuleChecked<FAssetRegistryModule>("AssetRegistry");
+
+				// find all control rigs in the project
+				TArray<FAssetData> ControlRigAssetDatas;
+				FARFilter ControlRigAssetFilter;
+				ControlRigAssetFilter.ClassNames.Add(UControlRigBlueprint::StaticClass()->GetFName());
+				AssetRegistryModule.Get().GetAssets(ControlRigAssetFilter, ControlRigAssetDatas);
+
+				// loop over all control rigs in the project
+				for(const FAssetData& ControlRigAssetData : ControlRigAssetDatas)
+				{
+					const FString ReferenceNodeDataString =
+						ControlRigAssetData.GetTagValueRef<FString>(ReferenceNodeDataProperty->GetFName());
+					if(ReferenceNodeDataString.IsEmpty())
+					{
+						continue;
+					}
+
+					TArray<FRigVMReferenceNodeData> ReferenceNodeDatas;
+					ReferenceNodeDataProperty->ImportText_Direct(*ReferenceNodeDataString, &ReferenceNodeDatas, nullptr, EPropertyPortFlags::PPF_None);
+
+					for(const FRigVMReferenceNodeData& ReferenceNodeData : ReferenceNodeDatas)
+					{
+						BuildData->RegisterFunctionReference(ReferenceNodeData);
+					}
+				}
+			}
+		}
+	}, TStatId(), NULL, ENamedThreads::GameThread);
+	
+	bAreFunctionReferencesInitialized = true;
+}
+
 void FControlRigEditor::BindCommands()
 {
 	DECLARE_SCOPE_HIERARCHICAL_COUNTER_FUNC()
diff --git a/Engine/Plugins/Animation/ControlRig/Source/ControlRigEditor/Private/Editor/ControlRigEditor.h b/Engine/Plugins/Animation/ControlRig/Source/ControlRigEditor/Private/Editor/ControlRigEditor.h
--- a/Engine/Plugins/Animation/ControlRig/Source/ControlRigEditor/Private/Editor/ControlRigEditor.h
+++ b/Engine/Plugins/Animation/ControlRig/Source/ControlRigEditor/Private/Editor/ControlRigEditor.h
@@ -101,6 +101,9 @@ private:
 
 	FControlRigEditorClosed ControlRigEditorClosedDelegate;
 
+	static bool bAreFunctionReferencesInitialized;
+	static void InitFunctionReferences();
+
 public:
 	// IToolkit Interface
 	virtual FName GetToolkitFName() const override;
