ensure that the skeletal mesh is still bound to an anim instance after compiling.
when using a post process ABP, its compilation results in clearing the anim script instance leaving the skeletal mesh unbound

diff --git a/Engine/Plugins/Animation/ControlRig/Source/ControlRigEditor/Private/Editor/ControlRigEditor.cpp b/Engine/Plugins/Animation/ControlRig/Source/ControlRigEditor/Private/Editor/ControlRigEditor.cpp
--- a/Engine/Plugins/Animation/ControlRig/Source/ControlRigEditor/Private/Editor/ControlRigEditor.cpp
+++ b/Engine/Plugins/Animation/ControlRig/Source/ControlRigEditor/Private/Editor/ControlRigEditor.cpp
@@ -1933,6 +1933,18 @@ void FControlRigEditor::Compile()
 			FBlueprintEditor::Compile();
 		}
 
+		// ensure the skeletal mesh is still bound
+		UControlRigSkeletalMeshComponent* SkelMeshComponent = Cast<UControlRigSkeletalMeshComponent>(GetPersonaToolkit()->GetPreviewScene()->GetPreviewMeshComponent());
+		if (SkelMeshComponent)
+		{
+			bool bWasCreated = false;
+			FAnimCustomInstanceHelper::BindToSkeletalMeshComponent<UControlRigLayerInstance>(SkelMeshComponent, bWasCreated);
+			if (bWasCreated)
+			{
+				OnAnimInitialized();
+			}
+		}
+		
 		if (ControlRig)
 		{
 			ControlRig->ControlRigLog = &ControlRigLog;
