Curve Editor: Improve multi paste operation by adding an IntentName to the channel editor data.

Previously, the intent for control rigs included the control node name, ie. Foot_L.Location.X. This means it could not match Foot_R.Location.X. With this change, the intent for control rigs are filled in with just Location.X.
The intent for non control rig transforms remains the same - it includes the group and the display name. For static mesh transforms, this is Location.X.

With this change, you can now copy one node's channels to multiple objects. For example, copy Node1's Location channels and pasting them onto one or more nodes results in all of those getting Node1's Location channels.

You can also now copy multiple nodes to multiple objects in order. For example, copying Node1's Location and Node2's Location channels onto Node3 and Node4 results in Node3 getting Node1's Location channel and Node4 getting Node2's Location channel.

diff --git a/Engine/Plugins/Animation/ControlRig/Source/ControlRig/Private/Sequencer/MovieSceneControlRigParameterSection.cpp b/Engine/Plugins/Animation/ControlRig/Source/ControlRig/Private/Sequencer/MovieSceneControlRigParameterSection.cpp
--- a/Engine/Plugins/Animation/ControlRig/Source/ControlRig/Private/Sequencer/MovieSceneControlRigParameterSection.cpp
+++ b/Engine/Plugins/Animation/ControlRig/Source/ControlRig/Private/Sequencer/MovieSceneControlRigParameterSection.cpp
@@ -110,6 +110,7 @@ struct FParameterVectorChannelEditorData
 			TotalName += ".X";
 			MetaData[0].SetIdentifiers(FName(*TotalName), FCommonChannelData::ChannelX);
 			TotalName = NameAsString;
+			MetaData[0].IntentName = FCommonChannelData::ChannelX;
 			MetaData[0].Group = GroupName;
 			MetaData[0].bEnabled = bEnabledOverride;
 			MetaData[0].SortOrder = SortStartIndex++;
@@ -119,6 +120,7 @@ struct FParameterVectorChannelEditorData
 			TotalName += ".Y";
 			MetaData[1].SetIdentifiers(FName(*TotalName), FCommonChannelData::ChannelY);
 			TotalName = NameAsString;
+			MetaData[1].IntentName = FCommonChannelData::ChannelY;
 			MetaData[1].Group = GroupName;
 			MetaData[1].bEnabled = bEnabledOverride;
 			MetaData[1].SortOrder = SortStartIndex++;
@@ -128,6 +130,7 @@ struct FParameterVectorChannelEditorData
 			TotalName += ".Z";
 			MetaData[2].SetIdentifiers(FName(*TotalName), FCommonChannelData::ChannelZ);
 			TotalName = NameAsString;
+			MetaData[2].IntentName = FCommonChannelData::ChannelZ;
 			MetaData[2].Group = GroupName;
 			MetaData[2].bEnabled = bEnabledOverride;
 			MetaData[2].SortOrder = SortStartIndex++;
@@ -137,6 +140,7 @@ struct FParameterVectorChannelEditorData
 			TotalName += ".W";
 			MetaData[3].SetIdentifiers(FName(*TotalName), FCommonChannelData::ChannelW);
 			TotalName = NameAsString;
+			MetaData[3].IntentName = FCommonChannelData::ChannelW;
 			MetaData[3].Group = GroupName;
 			MetaData[3].bEnabled = bEnabledOverride;
 			MetaData[3].SortOrder = SortStartIndex++;
@@ -340,6 +344,7 @@ struct FParameterTransformChannelEditorData
 			//MetaData[0].SetIdentifiers("Location.X", FCommonChannelData::ChannelX, LocationGroup);
 			TotalName += ".Location.X";
 			MetaData[0].SetIdentifiers(FName(*TotalName), NSLOCTEXT("MovieSceneControlParameterRigSection", "Location.X", "Location.X"), TransformGroup);
+			MetaData[0].IntentName = NSLOCTEXT("MovieSceneControlParameterRigSection", "Location.X", "Location.X");
 			TotalName = NameAsString;
 
 			MetaData[0].bEnabled = bEnabledOverride && EnumHasAllFlags(Mask, EMovieSceneTransformChannel::TranslationX);
@@ -350,6 +355,7 @@ struct FParameterTransformChannelEditorData
 			//MetaData[1].SetIdentifiers("Location.Y", FCommonChannelData::ChannelY, LocationGroup);
 			TotalName += ".Location.Y";
 			MetaData[1].SetIdentifiers(FName(*TotalName), NSLOCTEXT("MovieSceneControlParameterRigSection", "Location.Y", "Location.Y"), TransformGroup);
+			MetaData[1].IntentName = NSLOCTEXT("MovieSceneControlParameterRigSection", "Location.Y", "Location.Y");
 			TotalName = NameAsString;
 
 			MetaData[1].bEnabled = bEnabledOverride && EnumHasAllFlags(Mask, EMovieSceneTransformChannel::TranslationY);
@@ -360,6 +366,7 @@ struct FParameterTransformChannelEditorData
 			//MetaData[2].SetIdentifiers("Location.Z", FCommonChannelData::ChannelZ, LocationGroup);
 			TotalName += ".Location.Z";
 			MetaData[2].SetIdentifiers(FName(*TotalName), NSLOCTEXT("MovieSceneControlParameterRigSection", "Location.Z", "Location.Z"), TransformGroup);
+			MetaData[2].IntentName = NSLOCTEXT("MovieSceneControlParameterRigSection", "Location.Z", "Location.Z");
 			TotalName = NameAsString;
 
 			MetaData[2].bEnabled = bEnabledOverride && EnumHasAllFlags(Mask, EMovieSceneTransformChannel::TranslationZ);
@@ -371,6 +378,7 @@ struct FParameterTransformChannelEditorData
 			//MetaData[3].SetIdentifiers("Rotation.X", NSLOCTEXT("MovieSceneTransformSection", "RotationX", "Roll"), RotationGroup);
 			TotalName += ".Rotation.X";
 			MetaData[3].SetIdentifiers(FName(*TotalName), NSLOCTEXT("MovieSceneControlParameterRigSection", "Rotation.X", "Rotation.Roll"), TransformGroup);
+			MetaData[3].IntentName = NSLOCTEXT("MovieSceneControlParameterRigSection", "Rotation.X", "Rotation.Roll");
 			TotalName = NameAsString;
 
 			MetaData[3].bEnabled = bEnabledOverride && EnumHasAllFlags(Mask, EMovieSceneTransformChannel::RotationX);
@@ -381,6 +389,7 @@ struct FParameterTransformChannelEditorData
 			//MetaData[4].SetIdentifiers("Rotation.Y", NSLOCTEXT("MovieSceneTransformSection", "RotationY", "Pitch"), RotationGroup);
 			TotalName += ".Rotation.Y";
 			MetaData[4].SetIdentifiers(FName(*TotalName), NSLOCTEXT("MovieSceneControlParameterRigSection", "Rotation.Y", "Rotation.Pitch"), TransformGroup);
+			MetaData[4].IntentName = NSLOCTEXT("MovieSceneControlParameterRigSection", "Rotation.Y", "Rotation.Pitch");
 			TotalName = NameAsString;
 
 			MetaData[4].bEnabled = bEnabledOverride && EnumHasAllFlags(Mask, EMovieSceneTransformChannel::RotationY);
@@ -391,6 +400,7 @@ struct FParameterTransformChannelEditorData
 			//MetaData[5].SetIdentifiers("Rotation.Z", NSLOCTEXT("MovieSceneTransformSection", "RotationZ", "Yaw"), RotationGroup);
 			TotalName += ".Rotation.Z";
 			MetaData[5].SetIdentifiers(FName(*TotalName), NSLOCTEXT("MovieSceneControlParameterRigSection", "Rotation.Z", "Rotation.Yaw"), TransformGroup);
+			MetaData[5].IntentName = NSLOCTEXT("MovieSceneControlParameterRigSection", "Rotation.Z", "Rotation.Yaw");
 			TotalName = NameAsString;
 
 			MetaData[5].bEnabled = bEnabledOverride && EnumHasAllFlags(Mask, EMovieSceneTransformChannel::RotationZ);
@@ -402,6 +412,7 @@ struct FParameterTransformChannelEditorData
 			//MetaData[6].SetIdentifiers("Scale.X", FCommonChannelData::ChannelX, ScaleGroup);
 			TotalName += ".Scale.X";
 			MetaData[6].SetIdentifiers(FName(*TotalName), NSLOCTEXT("MovieSceneControlParameterRigSection", "Scale.X", "Scale.X"), TransformGroup);
+			MetaData[6].IntentName = NSLOCTEXT("MovieSceneControlParameterRigSection", "Scale.X", "Scale.X");
 			TotalName = NameAsString;
 
 			MetaData[6].bEnabled = bEnabledOverride && EnumHasAllFlags(Mask, EMovieSceneTransformChannel::ScaleX);
@@ -412,6 +423,7 @@ struct FParameterTransformChannelEditorData
 			//MetaData[7].SetIdentifiers("Scale.Y", FCommonChannelData::ChannelY, ScaleGroup);
 			TotalName += ".Scale.Y";
 			MetaData[7].SetIdentifiers(FName(*TotalName), NSLOCTEXT("MovieSceneControlParameterRigSection", "Scale.Y", "Scale.Y"), TransformGroup);
+			MetaData[7].IntentName = NSLOCTEXT("MovieSceneControlParameterRigSection", "Scale.Y", "Scale.Y");
 			TotalName = NameAsString;
 
 			MetaData[7].bEnabled = bEnabledOverride && EnumHasAllFlags(Mask, EMovieSceneTransformChannel::ScaleY);
@@ -422,6 +434,7 @@ struct FParameterTransformChannelEditorData
 			//MetaData[8].SetIdentifiers("Scale.Z", FCommonChannelData::ChannelZ, ScaleGroup);
 			TotalName += ".Scale.Z";
 			MetaData[8].SetIdentifiers(FName(*TotalName), NSLOCTEXT("MovieSceneControlParameterRigSection", "Scale.Z", "Scale.Z"), TransformGroup);
+			MetaData[8].IntentName = NSLOCTEXT("MovieSceneControlParameterRigSection", "Scale.Z", "Scale.Z");
 			TotalName = NameAsString;
 
 			MetaData[8].bEnabled = bEnabledOverride && EnumHasAllFlags(Mask, EMovieSceneTransformChannel::ScaleZ);
diff --git a/Engine/Source/Editor/CurveEditor/Private/CurveEditor.cpp b/Engine/Source/Editor/CurveEditor/Private/CurveEditor.cpp
--- a/Engine/Source/Editor/CurveEditor/Private/CurveEditor.cpp
+++ b/Engine/Source/Editor/CurveEditor/Private/CurveEditor.cpp
@@ -1355,29 +1355,6 @@ void FCurveEditor::PasteKeys(TSet<FCurveModelID> CurveModelIDs)
 	}
 
 	bool bSelectionNeedsLongNames = false;
-	FCurveEditorTreeItemID LastRootItem;
-
-	for (const TTuple<FCurveEditorTreeItemID, ECurveEditorTreeSelectionState>& Pair : GetTreeSelection())
-	{
-		FCurveEditorTreeItem* TreeItem = &GetTreeItem(Pair.Key);
-
-		FCurveEditorTreeItemID ParentId = TreeItem->GetParentID();
-		while (ParentId.IsValid())
-		{
-			TreeItem = &GetTreeItem(ParentId);
-			ParentId = TreeItem->GetParentID();
-		}
-
-		if (!LastRootItem.IsValid())
-		{
-			LastRootItem = TreeItem->GetID();
-		}
-		else if (TreeItem->GetID() != LastRootItem)
-		{
-			bSelectionNeedsLongNames = true;
-			break;
-		}
-	}
 
 	if (CurveModelIDs.Num() == 0)
 	{
@@ -1454,24 +1431,6 @@ void FCurveEditor::PasteKeys(TSet<FCurveModelID> CurveModelIDs)
 	for (UCurveEditorCopyBuffer* CopyBuffer : ImportedCopyBuffers)
 	{
 		bool bUseLongDisplayName = bSelectionNeedsLongNames;
-		if (!bUseLongDisplayName)
-		{
-			TOptional<FString> LastRootName;
-			for (UCurveEditorCopyableCurveKeys* CopyableCurveKeys : CopyBuffer->Curves)
-			{
-				FString RootName;
-				CopyableCurveKeys->LongDisplayName.Split(".", &RootName, nullptr);
-				if (!LastRootName.IsSet())
-				{
-					LastRootName = RootName;
-				}
-				else if (!LastRootName.GetValue().Equals(RootName))
-				{
-					bUseLongDisplayName = true;
-					break;
-				}
-			}
-		}
 
 		double TimeOffset = 0.0f;
 		bool bApplyOffset = !CopyBuffer->bAbsolutePosition;
@@ -1491,6 +1450,8 @@ void FCurveEditor::PasteKeys(TSet<FCurveModelID> CurveModelIDs)
 			}
 		}
 
+		TArray<UCurveEditorCopyableCurveKeys*> UsedCurves;
+
 		for (FCurveModelID CurveID : CurveModelIDs)
 		{
 			FCurveModel* Curve = FindCurve(CurveID);
@@ -1498,13 +1459,26 @@ void FCurveEditor::PasteKeys(TSet<FCurveModelID> CurveModelIDs)
 			{
 				const FString CurveLongDisplayName = Curve->GetLongDisplayName().ToString();
 				const FString CurveIntentionName = Curve->GetIntentionName();
-
+				
+				bool bFoundMatch = false;
 				for (UCurveEditorCopyableCurveKeys* CopyableCurveKeys : CopyBuffer->Curves)
 				{
+					// Use up all curves in the copied buffer before reusing. For example, if the following 6 curves have been copied,
+					//   Cube1.Location.X, Cube1.Location.Y, Cube1.Location.Z
+					//   Cube2.Location.X, Cube2.Location.Y, Cube2.Location.Z
+					// and the user attempts to paste onto multiple objects, Cube1's curves will match before Cube2's curves if this restriction is not in place.
+					//
+					if (UsedCurves.Contains(CopyableCurveKeys))
+					{
+						continue;
+					}
+
 					if (bAllCopiedCurvesLongNameEqual ||
 						(!bUseLongDisplayName && CurveIntentionName.Equals(CopyableCurveKeys->IntentionName))
 						|| (bUseLongDisplayName && CurveLongDisplayName.Equals(CopyableCurveKeys->LongDisplayName)))
 					{
+						bFoundMatch = true;
+
 						for (int32 Index = 0; Index < CopyableCurveKeys->KeyPositions.Num(); ++Index)
 						{
 							FKeyPosition KeyPosition = CopyableCurveKeys->KeyPositions[Index];
@@ -1519,8 +1493,22 @@ void FCurveEditor::PasteKeys(TSet<FCurveModelID> CurveModelIDs)
 								Selection.Add(FCurvePointHandle(CurveID, ECurvePointType::Key, KeyHandle.GetValue()));
 							}
 						}
+
+						UsedCurves.Add(CopyableCurveKeys);
+
+						if (UsedCurves.Num() == CopyBuffer->Curves.Num())
+						{
+							UsedCurves.Empty();
+						}
+
+						break; // Just paste one of the matching curves
 					}
 				}
+
+				if (!bFoundMatch)
+				{
+					UE_LOG(LogCurveEditor, Warning, TEXT("Failed to find matching curve to copy onto: %s"), *CurveLongDisplayName);		
+				}			
 			}
 		}
 	}
diff --git a/Engine/Source/Editor/Sequencer/Private/IKeyArea.cpp b/Engine/Source/Editor/Sequencer/Private/IKeyArea.cpp
--- a/Engine/Source/Editor/Sequencer/Private/IKeyArea.cpp
+++ b/Engine/Source/Editor/Sequencer/Private/IKeyArea.cpp
@@ -299,7 +299,11 @@ TUniquePtr<FCurveModel> IKeyArea::CreateCurveEditorModel(TSharedRef<ISequencer>
 
 			FText LongDisplayName = FText::Format(LongDisplayNameFormatString, FFormatOrderedArguments(ValidNames));
 			const FText ShortDisplayName = DisplayText;
-			const FString IntentName = ChannelHandle.GetMetaData()->Group.IsEmptyOrWhitespace() ? DisplayText.ToString() : FString::Format(TEXT("{0}.{1}"), { ChannelHandle.GetMetaData()->Group.ToString(), DisplayText.ToString() });
+			FString IntentName = ChannelHandle.GetMetaData()->IntentName.ToString();
+			if (IntentName.IsEmpty())
+			{
+				IntentName = ChannelHandle.GetMetaData()->Group.IsEmptyOrWhitespace() ? DisplayText.ToString() : FString::Format(TEXT("{0}.{1}"), { ChannelHandle.GetMetaData()->Group.ToString(), DisplayText.ToString() });
+			}
 
 			CurveModel->SetShortDisplayName(DisplayText);
 			CurveModel->SetLongDisplayName(LongDisplayName);
diff --git a/Engine/Source/Runtime/MovieScene/Public/Channels/MovieSceneChannelEditorData.h b/Engine/Source/Runtime/MovieScene/Public/Channels/MovieSceneChannelEditorData.h
--- a/Engine/Source/Runtime/MovieScene/Public/Channels/MovieSceneChannelEditorData.h
+++ b/Engine/Source/Runtime/MovieScene/Public/Channels/MovieSceneChannelEditorData.h
@@ -53,6 +53,8 @@ struct FMovieSceneChannelMetaData
 	FText DisplayText;
 	/** Name to group this channel with others of the same group name */
 	FText Group;
+	/** Intent name */
+	FText IntentName;
 	/** Optional color to draw underneath the keys on this channel */
 	TOptional<FLinearColor> Color;
 };
