Control Rig: Remove ability to watch pins inside functions

diff --git a/Engine/Plugins/Animation/ControlRig/Source/ControlRigEditor/Private/Editor/ControlRigEditor.cpp b/Engine/Plugins/Animation/ControlRig/Source/ControlRigEditor/Private/Editor/ControlRigEditor.cpp
--- a/Engine/Plugins/Animation/ControlRig/Source/ControlRigEditor/Private/Editor/ControlRigEditor.cpp
+++ b/Engine/Plugins/Animation/ControlRig/Source/ControlRigEditor/Private/Editor/ControlRigEditor.cpp
@@ -2548,7 +2548,9 @@ bool FControlRigEditor::CanStartWatchingPin() const
 	{
 		if (URigVMPin* ModelPin = GetFocusedModel()->FindPin(Pin->GetName()))
 		{
-			return ModelPin->GetParentPin() == nullptr && !ModelPin->RequiresWatch();
+			return ModelPin->GetParentPin() == nullptr &&
+					!ModelPin->RequiresWatch() &&
+					!ModelPin->GetGraph()->GetRootGraph()->IsA<URigVMFunctionLibrary>();
 		}
 	}
 	return false;
diff --git a/Engine/Source/Developer/RigVMDeveloper/Private/RigVMModel/RigVMController.cpp b/Engine/Source/Developer/RigVMDeveloper/Private/RigVMModel/RigVMController.cpp
--- a/Engine/Source/Developer/RigVMDeveloper/Private/RigVMModel/RigVMController.cpp
+++ b/Engine/Source/Developer/RigVMDeveloper/Private/RigVMModel/RigVMController.cpp
@@ -6272,7 +6272,7 @@ bool URigVMController::SetPinIsWatched(URigVMPin* InPin, bool bIsWatched, bool b
 	URigVMGraph* Graph = GetGraph();
 	check(Graph);
 
-	if (Graph->IsA<URigVMFunctionLibrary>())
+	if (Graph->GetRootGraph()->IsA<URigVMFunctionLibrary>())
 	{
 		ReportError(TEXT("Cannot watch pins in function library graphs."));
 		return false;
