update bPropagateToChildren default value for some nodes saved in previous versions

diff --git a/Engine/Plugins/Animation/ControlRig/Source/ControlRigDeveloper/Private/ControlRigBlueprint.cpp b/Engine/Plugins/Animation/ControlRig/Source/ControlRigDeveloper/Private/ControlRigBlueprint.cpp
--- a/Engine/Plugins/Animation/ControlRig/Source/ControlRigDeveloper/Private/ControlRigBlueprint.cpp
+++ b/Engine/Plugins/Animation/ControlRig/Source/ControlRigDeveloper/Private/ControlRigBlueprint.cpp
@@ -576,6 +576,8 @@ void UControlRigBlueprint::PostLoad()
 		}
 	}
 
+	PatchPropagateToChildren();
+	
 #if WITH_EDITOR
 	// delay compilation until the package has been loaded
 	FCoreUObjectDelegates::OnEndLoadPackage.AddUObject(this, &UControlRigBlueprint::HandlePackageDone);
@@ -3738,6 +3740,87 @@ void UControlRigBlueprint::PatchVariableNodesWithIncorrectType()
 	}
 }
 
+// change the default value form False to True for transform nodes
+void UControlRigBlueprint::PatchPropagateToChildren()
+{
+	// no need to update default value past this version
+	if (GetLinkerCustomVersion(FControlRigObjectVersion::GUID) >= FControlRigObjectVersion::RenameGizmoToShape)
+	{
+		return;
+	}
+	
+	auto IsNullOrControl = [](const URigVMPin* InPin)
+	{
+		const bool bHasItem = InPin->GetCPPTypeObject() == FRigElementKey::StaticStruct() && InPin->GetName() == "Item";
+		if (!bHasItem)
+		{
+			return false;
+		}
+
+		if (const URigVMPin* TypePin = InPin->FindSubPin(TEXT("Type")))
+		{
+			const FString& TypeValue = TypePin->GetDefaultValue();
+			return TypeValue == TEXT("Null") || TypeValue == TEXT("Control");
+		}
+		
+		return false;
+	};
+
+	auto IsPropagateChildren = [](const URigVMPin* InPin)
+	{
+		return InPin->GetCPPType() == TEXT("bool") && InPin->GetName() == TEXT("bPropagateToChildren");
+	};
+
+	auto FindPropagatePin = [IsNullOrControl, IsPropagateChildren](const URigVMNode* InNode)-> URigVMPin*
+	{
+		URigVMPin* PropagatePin = nullptr;
+		URigVMPin* ItemPin = nullptr;  
+		for (URigVMPin* Pin: InNode->GetPins())
+		{
+			// look for Item pin
+			if (!ItemPin && IsNullOrControl(Pin))
+			{
+				ItemPin = Pin;
+			}
+
+			// look for bPropagateToChildren pin
+			if (!PropagatePin && IsPropagateChildren(Pin))
+			{
+				PropagatePin = Pin;
+			}
+
+			// return propagation pin if both found
+			if (ItemPin && PropagatePin)
+			{
+				return PropagatePin;
+			}
+		}
+		return nullptr;
+	};
+
+	for (URigVMGraph* Graph : GetAllModels())
+	{
+		TArray< const URigVMPin* > PinsToUpdate;
+		for (const URigVMNode* Node : Graph->GetNodes())
+		{
+			if (const URigVMPin* PropagatePin = FindPropagatePin(Node))
+			{
+				PinsToUpdate.Add(PropagatePin);
+			}
+		}
+		
+		if (URigVMController* Controller = GetOrCreateController(Graph))
+		{
+			Controller->SuspendNotifications(true);
+			for (const URigVMPin* Pin: PinsToUpdate)
+			{
+				Controller->SetPinDefaultValue(Pin->GetPinPath(), TEXT("True"), false, false, false);
+			}
+			Controller->SuspendNotifications(false);
+		}
+	}
+}
+
 void UControlRigBlueprint::ConvertUnitNodesToTemplateNodes()
 {
 	TGuardValue<bool> GuardNotifsSelf(bSuspendModelNotificationsForSelf, true);
diff --git a/Engine/Plugins/Animation/ControlRig/Source/ControlRigDeveloper/Public/ControlRigBlueprint.h b/Engine/Plugins/Animation/ControlRig/Source/ControlRigDeveloper/Public/ControlRigBlueprint.h
--- a/Engine/Plugins/Animation/ControlRig/Source/ControlRigDeveloper/Public/ControlRigBlueprint.h
+++ b/Engine/Plugins/Animation/ControlRig/Source/ControlRigDeveloper/Public/ControlRigBlueprint.h
@@ -551,6 +551,7 @@ private:
 	void PatchRigElementKeyCacheOnLoad();
 	void PatchBoundVariables();
 	void PatchVariableNodesWithIncorrectType();
+	void PatchPropagateToChildren();
 	void ConvertUnitNodesToTemplateNodes();
 	void PatchParameterNodesOnLoad();
 
