Control Rig: Node audit for templates

diff --git a/Engine/Plugins/Animation/ControlRig/Config/BaseControlRig.ini b/Engine/Plugins/Animation/ControlRig/Config/BaseControlRig.ini
--- a/Engine/Plugins/Animation/ControlRig/Config/BaseControlRig.ini
+++ b/Engine/Plugins/Animation/ControlRig/Config/BaseControlRig.ini
@@ -8,6 +8,19 @@
 +StructRedirects=(OldName="RigJoint",NewName="/Script/ControlRig.RigBone")
 +PropertyRedirects=(OldName="RigBoneHierarchy.Joints",NewName="RigBoneHierarchy.Bones")
 +PropertyRedirects=(OldName="RigHierarchyContainer.BaseHierarchy",NewName="RigHierarchyContainer.BoneHierarchy")
++PropertyRedirects=(OldName="RigUnit_SetTransform.Transform", NewName="RigUnit_SetTransform.Value")
++PropertyRedirects=(OldName="RigUnit_SetTranslation.Translation", NewName="RigUnit_SetTranslation.Value")
++PropertyRedirects=(OldName="RigUnit_SetRotation.Rotation", NewName="RigUnit_SetRotation.Value")
++PropertyRedirects=(OldName="RigUnit_ToWorldSpace_Transform.Transform", NewName="RigUnit_ToWorldSpace_Transform.Value")
++PropertyRedirects=(OldName="RigUnit_ToWorldSpace_Transform.Transform", NewName="RigUnit_ToWorldSpace_Transform.Value")
++PropertyRedirects=(OldName="RigUnit_ToWorldSpace_Location.Location", NewName="RigUnit_ToWorldSpace_Location.Value")
++PropertyRedirects=(OldName="RigUnit_ToWorldSpace_Rotation.Rotation", NewName="RigUnit_ToWorldSpace_Rotation.Value")
++PropertyRedirects=(OldName="RigUnit_ToRigSpace_Transform.Transform", NewName="RigUnit_ToRigSpace_Transform.Value")
++PropertyRedirects=(OldName="RigUnit_ToRigSpace_Location.Location", NewName="RigUnit_ToRigSpace_Location.Value")
++PropertyRedirects=(OldName="RigUnit_ToRigSpace_Rotation.Rotation", NewName="RigUnit_ToRigSpace_Rotation.Value")
++PropertyRedirects=(OldName="RigUnit_MathQuaternionRotateVector.Quaternion", NewName="RigUnit_MathQuaternionRotateVector.Transform")
++PropertyRedirects=(OldName="RigUnit_MathTransformRotateVector.Direction", NewName="RigUnit_MathTransformRotateVector.Vector")
++PropertyRedirects=(OldName="RigUnit_SetRelativeTransformForItem.RelativeTransform", NewName="RigUnit_SetRelativeTransformForItem.Value")
 
 +ClassRedirects=(OldName="RigVMGraph",NewName="RigVMGraph",ValueChanges=(("Model","RigVMModel")))
 
diff --git a/Engine/Plugins/Animation/ControlRig/Source/ControlRig/Private/Units/Hierarchy/RigUnit_SetBoneInitialTransform.cpp b/Engine/Plugins/Animation/ControlRig/Source/ControlRig/Private/Units/Hierarchy/RigUnit_SetBoneInitialTransform.cpp
--- a/Engine/Plugins/Animation/ControlRig/Source/ControlRig/Private/Units/Hierarchy/RigUnit_SetBoneInitialTransform.cpp
+++ b/Engine/Plugins/Animation/ControlRig/Source/ControlRig/Private/Units/Hierarchy/RigUnit_SetBoneInitialTransform.cpp
@@ -50,10 +50,11 @@ FRigVMStructUpgradeInfo FRigUnit_SetBoneInitialTransform::GetUpgradeInfo() const
 	FRigUnit_SetTransform NewNode;
 	NewNode.Item = FRigElementKey(Bone, ERigElementType::Bone);
 	NewNode.Space = Space;
-	NewNode.Transform = Transform;
+	NewNode.Value = Transform;
 	NewNode.bInitial = true;
 
 	FRigVMStructUpgradeInfo Info(*this, NewNode);
 	Info.AddRemappedPin(TEXT("Bone"), TEXT("Item.Name"));
+	Info.AddRemappedPin(TEXT("Transform"), TEXT("Value"));
 	return Info;
 }
diff --git a/Engine/Plugins/Animation/ControlRig/Source/ControlRig/Private/Units/Hierarchy/RigUnit_SetBoneRotation.cpp b/Engine/Plugins/Animation/ControlRig/Source/ControlRig/Private/Units/Hierarchy/RigUnit_SetBoneRotation.cpp
--- a/Engine/Plugins/Animation/ControlRig/Source/ControlRig/Private/Units/Hierarchy/RigUnit_SetBoneRotation.cpp
+++ b/Engine/Plugins/Animation/ControlRig/Source/ControlRig/Private/Units/Hierarchy/RigUnit_SetBoneRotation.cpp
@@ -82,10 +82,11 @@ FRigVMStructUpgradeInfo FRigUnit_SetBoneRotation::GetUpgradeInfo() const
 	FRigUnit_SetRotation NewNode;
 	NewNode.Item = FRigElementKey(Bone, ERigElementType::Bone);
 	NewNode.Space = Space;
-	NewNode.Rotation = Rotation;
+	NewNode.Value = Rotation;
 
 	FRigVMStructUpgradeInfo Info(*this, NewNode);
 	Info.AddRemappedPin(TEXT("Bone"), TEXT("Item.Name"));
+	Info.AddRemappedPin(TEXT("Rotation"), TEXT("Value"));
 	return Info;
 }
 
diff --git a/Engine/Plugins/Animation/ControlRig/Source/ControlRig/Private/Units/Hierarchy/RigUnit_SetBoneTransform.cpp b/Engine/Plugins/Animation/ControlRig/Source/ControlRig/Private/Units/Hierarchy/RigUnit_SetBoneTransform.cpp
--- a/Engine/Plugins/Animation/ControlRig/Source/ControlRig/Private/Units/Hierarchy/RigUnit_SetBoneTransform.cpp
+++ b/Engine/Plugins/Animation/ControlRig/Source/ControlRig/Private/Units/Hierarchy/RigUnit_SetBoneTransform.cpp
@@ -79,10 +79,11 @@ FRigVMStructUpgradeInfo FRigUnit_SetBoneTransform::GetUpgradeInfo() const
 	FRigUnit_SetTransform NewNode;
 	NewNode.Item = FRigElementKey(Bone, ERigElementType::Bone);
 	NewNode.Space = Space;
-	NewNode.Transform = Transform;
+	NewNode.Value = Transform;
 
 	FRigVMStructUpgradeInfo Info(*this, NewNode);
 	Info.AddRemappedPin(TEXT("Bone"), TEXT("Item.Name"));
+	Info.AddRemappedPin(TEXT("Transform"), TEXT("Value"));
 	return Info;
 }
 
diff --git a/Engine/Plugins/Animation/ControlRig/Source/ControlRig/Private/Units/Hierarchy/RigUnit_SetBoneTranslation.cpp b/Engine/Plugins/Animation/ControlRig/Source/ControlRig/Private/Units/Hierarchy/RigUnit_SetBoneTranslation.cpp
--- a/Engine/Plugins/Animation/ControlRig/Source/ControlRig/Private/Units/Hierarchy/RigUnit_SetBoneTranslation.cpp
+++ b/Engine/Plugins/Animation/ControlRig/Source/ControlRig/Private/Units/Hierarchy/RigUnit_SetBoneTranslation.cpp
@@ -82,9 +82,10 @@ FRigVMStructUpgradeInfo FRigUnit_SetBoneTranslation::GetUpgradeInfo() const
 	FRigUnit_SetTranslation NewNode;
 	NewNode.Item = FRigElementKey(Bone, ERigElementType::Bone);
 	NewNode.Space = Space;
-	NewNode.Translation = Translation;
+	NewNode.Value = Translation;
 
 	FRigVMStructUpgradeInfo Info(*this, NewNode);
 	Info.AddRemappedPin(TEXT("Bone"), TEXT("Item.Name"));
+	Info.AddRemappedPin(TEXT("Translation"), TEXT("Value"));
 	return Info;
 }
diff --git a/Engine/Plugins/Animation/ControlRig/Source/ControlRig/Private/Units/Hierarchy/RigUnit_SetControlTransform.cpp b/Engine/Plugins/Animation/ControlRig/Source/ControlRig/Private/Units/Hierarchy/RigUnit_SetControlTransform.cpp
--- a/Engine/Plugins/Animation/ControlRig/Source/ControlRig/Private/Units/Hierarchy/RigUnit_SetControlTransform.cpp
+++ b/Engine/Plugins/Animation/ControlRig/Source/ControlRig/Private/Units/Hierarchy/RigUnit_SetControlTransform.cpp
@@ -497,10 +497,11 @@ FRigVMStructUpgradeInfo FRigUnit_SetControlTransform::GetUpgradeInfo() const
 	FRigUnit_SetTransform NewNode;
 	NewNode.Item = FRigElementKey(Control, ERigElementType::Control);
 	NewNode.Space = Space;
-	NewNode.Transform = Transform;
+	NewNode.Value = Transform;
 
 	FRigVMStructUpgradeInfo Info(*this, NewNode);
 	Info.AddRemappedPin(TEXT("Control"), TEXT("Item.Name"));
+	Info.AddRemappedPin(TEXT("Transform"), TEXT("Value"));
 	return Info;
 }
 
diff --git a/Engine/Plugins/Animation/ControlRig/Source/ControlRig/Private/Units/Hierarchy/RigUnit_SetRelativeBoneTransform.cpp b/Engine/Plugins/Animation/ControlRig/Source/ControlRig/Private/Units/Hierarchy/RigUnit_SetRelativeBoneTransform.cpp
--- a/Engine/Plugins/Animation/ControlRig/Source/ControlRig/Private/Units/Hierarchy/RigUnit_SetRelativeBoneTransform.cpp
+++ b/Engine/Plugins/Animation/ControlRig/Source/ControlRig/Private/Units/Hierarchy/RigUnit_SetRelativeBoneTransform.cpp
@@ -58,7 +58,7 @@ FRigVMStructUpgradeInfo FRigUnit_SetRelativeBoneTransform::GetUpgradeInfo() cons
 	FRigUnit_SetRelativeTransformForItem NewNode;
 	NewNode.Parent = FRigElementKey(Space, ERigElementType::Bone);
 	NewNode.Child = FRigElementKey(Bone, ERigElementType::Bone);
-	NewNode.RelativeTransform = Transform;
+	NewNode.Value = Transform;
 
 	FRigVMStructUpgradeInfo Info(*this, NewNode);
 	Info.AddRemappedPin(TEXT("Space"), TEXT("Parent.Name"));
diff --git a/Engine/Plugins/Animation/ControlRig/Source/ControlRig/Private/Units/Hierarchy/RigUnit_SetRelativeTransform.cpp b/Engine/Plugins/Animation/ControlRig/Source/ControlRig/Private/Units/Hierarchy/RigUnit_SetRelativeTransform.cpp
--- a/Engine/Plugins/Animation/ControlRig/Source/ControlRig/Private/Units/Hierarchy/RigUnit_SetRelativeTransform.cpp
+++ b/Engine/Plugins/Animation/ControlRig/Source/ControlRig/Private/Units/Hierarchy/RigUnit_SetRelativeTransform.cpp
@@ -4,6 +4,7 @@
 #include "Units/Math/RigUnit_MathTransform.h"
 #include "Units/Hierarchy/RigUnit_GetTransform.h"
 #include "Units/Hierarchy/RigUnit_SetTransform.h"
+#include "Units/Hierarchy/RigUnit_GetRelativeTransform.h"
 #include "Units/RigUnitContext.h"
 
 FRigUnit_SetRelativeTransformForItem_Execute()
@@ -19,6 +20,47 @@ FRigUnit_SetRelativeTransformForItem_Execute()
 	FTransform GlobalTransform = FTransform::Identity;
 
 	FRigUnit_GetTransform::StaticExecute(RigVMExecuteContext, Parent, EBoneGetterSetterMode::GlobalSpace, bParentInitial, ParentTransform, CachedParent, Context);
-	FRigUnit_MathTransformMakeAbsolute::StaticExecute(RigVMExecuteContext, RelativeTransform, ParentTransform, GlobalTransform, Context);
+	FRigUnit_MathTransformMakeAbsolute::StaticExecute(RigVMExecuteContext, Value, ParentTransform, GlobalTransform, Context);
 	FRigUnit_SetTransform::StaticExecute(RigVMExecuteContext, Child, EBoneGetterSetterMode::GlobalSpace, false, GlobalTransform, Weight, bPropagateToChildren, CachedChild, ExecuteContext, Context);
-}
\ No newline at end of file
+}
+
+FRigUnit_SetRelativeTranslationForItem_Execute()
+{
+	DECLARE_SCOPE_HIERARCHICAL_COUNTER_RIGUNIT()
+
+	if (Weight < SMALL_NUMBER)
+	{
+		return;
+	}
+
+	FTransform ParentTransform = FTransform::Identity;
+	FTransform LocalTransform = FTransform::Identity;
+	FTransform GlobalTransform = FTransform::Identity;
+
+	FRigUnit_GetTransform::StaticExecute(RigVMExecuteContext, Parent, EBoneGetterSetterMode::GlobalSpace, bParentInitial, ParentTransform, CachedParent, Context);
+	FRigUnit_GetRelativeTransformForItem::StaticExecute(RigVMExecuteContext, Child, false, Parent, bParentInitial, LocalTransform, CachedChild, CachedParent, Context);
+	LocalTransform.SetTranslation(Value);
+	FRigUnit_MathTransformMakeAbsolute::StaticExecute(RigVMExecuteContext, LocalTransform, ParentTransform, GlobalTransform, Context);
+	FRigUnit_SetTransform::StaticExecute(RigVMExecuteContext, Child, EBoneGetterSetterMode::GlobalSpace, false, GlobalTransform, Weight, bPropagateToChildren, CachedChild, ExecuteContext, Context);
+}
+
+FRigUnit_SetRelativeRotationForItem_Execute()
+{
+	DECLARE_SCOPE_HIERARCHICAL_COUNTER_RIGUNIT()
+
+	if (Weight < SMALL_NUMBER)
+	{
+		return;
+	}
+
+	FTransform ParentTransform = FTransform::Identity;
+	FTransform LocalTransform = FTransform::Identity;
+	FTransform GlobalTransform = FTransform::Identity;
+
+	FRigUnit_GetTransform::StaticExecute(RigVMExecuteContext, Parent, EBoneGetterSetterMode::GlobalSpace, bParentInitial, ParentTransform, CachedParent, Context);
+	FRigUnit_GetRelativeTransformForItem::StaticExecute(RigVMExecuteContext, Child, false, Parent, bParentInitial, LocalTransform, CachedChild, CachedParent, Context);
+	LocalTransform.SetRotation(Value);
+	LocalTransform.NormalizeRotation();
+	FRigUnit_MathTransformMakeAbsolute::StaticExecute(RigVMExecuteContext, LocalTransform, ParentTransform, GlobalTransform, Context);
+	FRigUnit_SetTransform::StaticExecute(RigVMExecuteContext, Child, EBoneGetterSetterMode::GlobalSpace, false, GlobalTransform, Weight, bPropagateToChildren, CachedChild, ExecuteContext, Context);
+}
diff --git a/Engine/Plugins/Animation/ControlRig/Source/ControlRig/Private/Units/Hierarchy/RigUnit_SetRelativeTransform.h b/Engine/Plugins/Animation/ControlRig/Source/ControlRig/Private/Units/Hierarchy/RigUnit_SetRelativeTransform.h
--- a/Engine/Plugins/Animation/ControlRig/Source/ControlRig/Private/Units/Hierarchy/RigUnit_SetRelativeTransform.h
+++ b/Engine/Plugins/Animation/ControlRig/Source/ControlRig/Private/Units/Hierarchy/RigUnit_SetRelativeTransform.h
@@ -8,7 +8,7 @@
 /**
  * SetRelativeTransform is used to set a single transform from a hierarchy in the space of another item
  */
-USTRUCT(meta=(DisplayName="Set Relative Transform", Category="Transforms", DocumentationPolicy = "Strict", Keywords = "Offset,Local", NodeColor="0, 0.364706, 1.0", Varying))
+USTRUCT(meta=(DisplayName="Set Relative Transform", Category="Transforms", TemplateName="Set Relative Transform", DocumentationPolicy = "Strict", Keywords = "Offset,Local", NodeColor="0, 0.364706, 1.0", Varying))
 struct CONTROLRIG_API FRigUnit_SetRelativeTransformForItem : public FRigUnitMutable
 {
 	GENERATED_BODY()
@@ -17,7 +17,7 @@ struct CONTROLRIG_API FRigUnit_SetRelativeTransformForItem : public FRigUnitMuta
 		: Child(NAME_None, ERigElementType::Bone)
 		, Parent(NAME_None, ERigElementType::Bone)
 		, bParentInitial(false)
-		, RelativeTransform(FTransform::Identity)
+		, Value(FTransform::Identity)
 		, Weight(1.f)
 		, bPropagateToChildren(true)
 		, CachedChild()
@@ -54,7 +54,145 @@ struct CONTROLRIG_API FRigUnit_SetRelativeTransformForItem : public FRigUnitMuta
 
 	// The transform of the child item relative to the provided parent
 	UPROPERTY(meta=(Input))
-	FTransform RelativeTransform;
+	FTransform Value;
+
+	// Defines how much the change will be applied
+	UPROPERTY(meta = (Input, UIMin = "0.0", UIMax = "1.0"))
+	float Weight;
+
+	// If set to true children of affected items in the hierarchy
+	// will follow the transform change - otherwise only the parent will move.
+	UPROPERTY(meta=(Input))
+	bool bPropagateToChildren;
+
+	// Used to cache the child internally
+	UPROPERTY()
+	FCachedRigElement CachedChild;
+
+	// Used to cache the parent internally
+	UPROPERTY()
+	FCachedRigElement CachedParent;
+};
+
+/**
+ * SetRelativeTranslation is used to set a single translation from a hierarchy in the space of another item
+ */
+USTRUCT(meta=(DisplayName="Set Relative Translation", Category="Transforms", TemplateName="Set Relative Transform", DocumentationPolicy = "Strict", Keywords = "Offset,Local", NodeColor="0, 0.364706, 1.0", Varying))
+struct CONTROLRIG_API FRigUnit_SetRelativeTranslationForItem : public FRigUnitMutable
+{
+	GENERATED_BODY()
+
+	FRigUnit_SetRelativeTranslationForItem()
+		: Child(NAME_None, ERigElementType::Bone)
+		, Parent(NAME_None, ERigElementType::Bone)
+		, bParentInitial(false)
+		, Value(FVector::ZeroVector)
+		, Weight(1.f)
+		, bPropagateToChildren(true)
+		, CachedChild()
+		, CachedParent()
+	{}
+
+	virtual FRigElementKey DetermineSpaceForPin(const FString& InPinPath, void* InUserContext) const override
+	{
+		return Parent;
+	}
+
+	RIGVM_METHOD()
+	virtual void Execute(const FRigUnitContext& Context) override;
+
+	/**
+	 * The child item to set the transform for
+	 */
+	UPROPERTY(meta = (Input, ExpandByDefault))
+	FRigElementKey Child;
+
+	/**
+	 * The parent item to use.
+	 * The child transform will be set in the space of the parent.
+	 */
+	UPROPERTY(meta = (Input, ExpandByDefault))
+	FRigElementKey Parent;
+
+	/**
+	 * Defines if the parent's transform should be determined as current (false) or initial (true).
+	 * Initial transforms for bones and other elements in the hierarchy represent the reference pose's value.
+	 */ 
+	UPROPERTY(meta = (Input))
+	bool bParentInitial;
+
+	// The transform of the child item relative to the provided parent
+	UPROPERTY(meta=(Input))
+	FVector Value;
+
+	// Defines how much the change will be applied
+	UPROPERTY(meta = (Input, UIMin = "0.0", UIMax = "1.0"))
+	float Weight;
+
+	// If set to true children of affected items in the hierarchy
+	// will follow the transform change - otherwise only the parent will move.
+	UPROPERTY(meta=(Input))
+	bool bPropagateToChildren;
+
+	// Used to cache the child internally
+	UPROPERTY()
+	FCachedRigElement CachedChild;
+
+	// Used to cache the parent internally
+	UPROPERTY()
+	FCachedRigElement CachedParent;
+};
+
+/**
+ * SetRelativeRotation is used to set a single rotation from a hierarchy in the space of another item
+ */
+USTRUCT(meta=(DisplayName="Set Relative Rotation", Category="Transforms", TemplateName="Set Relative Transform", DocumentationPolicy = "Strict", Keywords = "Offset,Local", NodeColor="0, 0.364706, 1.0", Varying))
+struct CONTROLRIG_API FRigUnit_SetRelativeRotationForItem : public FRigUnitMutable
+{
+	GENERATED_BODY()
+
+	FRigUnit_SetRelativeRotationForItem()
+		: Child(NAME_None, ERigElementType::Bone)
+		, Parent(NAME_None, ERigElementType::Bone)
+		, bParentInitial(false)
+		, Value(FQuat::Identity)
+		, Weight(1.f)
+		, bPropagateToChildren(true)
+		, CachedChild()
+		, CachedParent()
+	{}
+
+	virtual FRigElementKey DetermineSpaceForPin(const FString& InPinPath, void* InUserContext) const override
+	{
+		return Parent;
+	}
+
+	RIGVM_METHOD()
+	virtual void Execute(const FRigUnitContext& Context) override;
+
+	/**
+	 * The child item to set the transform for
+	 */
+	UPROPERTY(meta = (Input, ExpandByDefault))
+	FRigElementKey Child;
+
+	/**
+	 * The parent item to use.
+	 * The child transform will be set in the space of the parent.
+	 */
+	UPROPERTY(meta = (Input, ExpandByDefault))
+	FRigElementKey Parent;
+
+	/**
+	 * Defines if the parent's transform should be determined as current (false) or initial (true).
+	 * Initial transforms for bones and other elements in the hierarchy represent the reference pose's value.
+	 */ 
+	UPROPERTY(meta = (Input))
+	bool bParentInitial;
+
+	// The transform of the child item relative to the provided parent
+	UPROPERTY(meta=(Input))
+	FQuat Value;
 
 	// Defines how much the change will be applied
 	UPROPERTY(meta = (Input, UIMin = "0.0", UIMax = "1.0"))
diff --git a/Engine/Plugins/Animation/ControlRig/Source/ControlRig/Private/Units/Hierarchy/RigUnit_SetSpaceInitialTransform.cpp b/Engine/Plugins/Animation/ControlRig/Source/ControlRig/Private/Units/Hierarchy/RigUnit_SetSpaceInitialTransform.cpp
--- a/Engine/Plugins/Animation/ControlRig/Source/ControlRig/Private/Units/Hierarchy/RigUnit_SetSpaceInitialTransform.cpp
+++ b/Engine/Plugins/Animation/ControlRig/Source/ControlRig/Private/Units/Hierarchy/RigUnit_SetSpaceInitialTransform.cpp
@@ -50,10 +50,11 @@ FRigVMStructUpgradeInfo FRigUnit_SetSpaceInitialTransform::GetUpgradeInfo() cons
 	FRigUnit_SetTransform NewNode;
 	NewNode.Item = FRigElementKey(SpaceName, ERigElementType::Null);
 	NewNode.Space = Space;
-	NewNode.Transform = Transform;
+	NewNode.Value = Transform;
 	NewNode.bInitial = true;
 
 	FRigVMStructUpgradeInfo Info(*this, NewNode);
 	Info.AddRemappedPin(TEXT("SpaceName"), TEXT("Item.Name"));
+	Info.AddRemappedPin(TEXT("Transform"), TEXT("Value"));
 	return Info;
 }
diff --git a/Engine/Plugins/Animation/ControlRig/Source/ControlRig/Private/Units/Hierarchy/RigUnit_SetSpaceTransform.cpp b/Engine/Plugins/Animation/ControlRig/Source/ControlRig/Private/Units/Hierarchy/RigUnit_SetSpaceTransform.cpp
--- a/Engine/Plugins/Animation/ControlRig/Source/ControlRig/Private/Units/Hierarchy/RigUnit_SetSpaceTransform.cpp
+++ b/Engine/Plugins/Animation/ControlRig/Source/ControlRig/Private/Units/Hierarchy/RigUnit_SetSpaceTransform.cpp
@@ -72,10 +72,11 @@ FRigVMStructUpgradeInfo FRigUnit_SetSpaceTransform::GetUpgradeInfo() const
 	FRigUnit_SetTransform NewNode;
 	NewNode.Item = FRigElementKey(Space, ERigElementType::Null);
 	NewNode.Space = SpaceType;
-	NewNode.Transform = Transform;
+	NewNode.Value = Transform;
 
 	FRigVMStructUpgradeInfo Info(*this, NewNode);
 	Info.AddRemappedPin(TEXT("Space"), TEXT("Item.Name"));
 	Info.AddRemappedPin(TEXT("SpaceType"), TEXT("Space"));
+	Info.AddRemappedPin(TEXT("Transform"), TEXT("Value"));
 	return Info;
 }
diff --git a/Engine/Plugins/Animation/ControlRig/Source/ControlRig/Private/Units/Hierarchy/RigUnit_SetTransform.cpp b/Engine/Plugins/Animation/ControlRig/Source/ControlRig/Private/Units/Hierarchy/RigUnit_SetTransform.cpp
--- a/Engine/Plugins/Animation/ControlRig/Source/ControlRig/Private/Units/Hierarchy/RigUnit_SetTransform.cpp
+++ b/Engine/Plugins/Animation/ControlRig/Source/ControlRig/Private/Units/Hierarchy/RigUnit_SetTransform.cpp
@@ -44,7 +44,7 @@ FRigUnit_SetTransform_Execute()
 					// for controls - set the control offset transform instead
 					if(bInitial && (CachedIndex.GetKey().Type == ERigElementType::Control))
 					{
-						FTransform TransformMutable = Transform;
+						FTransform TransformMutable = Value;
 						FRigUnit_SetControlOffset::StaticExecute(RigVMExecuteContext, CachedIndex.GetKey().Name, TransformMutable, Space, CachedIndex, ExecuteContext, Context);
 						
 						if (ExecuteContext.EventName == FRigUnit_PrepareForExecution::EventName)
@@ -55,7 +55,7 @@ FRigUnit_SetTransform_Execute()
 						return;
 					}
 					
-					FTransform WeightedTransform = Transform;
+					FTransform WeightedTransform = Value;
 					if (Weight < 1.f - SMALL_NUMBER)
 					{
 						FTransform PreviousTransform = WeightedTransform;
@@ -125,9 +125,9 @@ FString FRigUnit_SetTranslation::GetUnitLabel() const
 FRigUnit_SetTranslation_Execute()
 {
 	FTransform Transform = FTransform::Identity;
-	FRigUnit_GetTransform::StaticExecute(RigVMExecuteContext, Item, Space, false, Transform, CachedIndex, Context);
-	Transform.SetLocation(Translation);
-	FRigUnit_SetTransform::StaticExecute(RigVMExecuteContext, Item, Space, false, Transform, Weight, bPropagateToChildren, CachedIndex, ExecuteContext, Context);
+	FRigUnit_GetTransform::StaticExecute(RigVMExecuteContext, Item, Space, bInitial, Transform, CachedIndex, Context);
+	Transform.SetLocation(Value);
+	FRigUnit_SetTransform::StaticExecute(RigVMExecuteContext, Item, Space, bInitial, Transform, Weight, bPropagateToChildren, CachedIndex, ExecuteContext, Context);
 }
 
 FString FRigUnit_SetRotation::GetUnitLabel() const
@@ -139,9 +139,9 @@ FString FRigUnit_SetRotation::GetUnitLabel() const
 FRigUnit_SetRotation_Execute()
 {
 	FTransform Transform = FTransform::Identity;
-	FRigUnit_GetTransform::StaticExecute(RigVMExecuteContext, Item, Space, false, Transform, CachedIndex, Context);
-	Transform.SetRotation(Rotation);
-	FRigUnit_SetTransform::StaticExecute(RigVMExecuteContext, Item, Space, false, Transform, Weight, bPropagateToChildren, CachedIndex, ExecuteContext, Context);
+	FRigUnit_GetTransform::StaticExecute(RigVMExecuteContext, Item, Space, bInitial, Transform, CachedIndex, Context);
+	Transform.SetRotation(Value);
+	FRigUnit_SetTransform::StaticExecute(RigVMExecuteContext, Item, Space, bInitial, Transform, Weight, bPropagateToChildren, CachedIndex, ExecuteContext, Context);
 }
 
 FString FRigUnit_SetScale::GetUnitLabel() const
@@ -153,9 +153,9 @@ FString FRigUnit_SetScale::GetUnitLabel() const
 FRigUnit_SetScale_Execute()
 {
 	FTransform Transform = FTransform::Identity;
-	FRigUnit_GetTransform::StaticExecute(RigVMExecuteContext, Item, Space, false, Transform, CachedIndex, Context);
+	FRigUnit_GetTransform::StaticExecute(RigVMExecuteContext, Item, Space, bInitial, Transform, CachedIndex, Context);
 	Transform.SetScale3D(Scale);
-	FRigUnit_SetTransform::StaticExecute(RigVMExecuteContext, Item, Space, false, Transform, Weight, bPropagateToChildren, CachedIndex, ExecuteContext, Context);
+	FRigUnit_SetTransform::StaticExecute(RigVMExecuteContext, Item, Space, bInitial, Transform, Weight, bPropagateToChildren, CachedIndex, ExecuteContext, Context);
 }
 
 
diff --git a/Engine/Plugins/Animation/ControlRig/Source/ControlRig/Private/Units/Math/RigUnit_MathColor.h b/Engine/Plugins/Animation/ControlRig/Source/ControlRig/Private/Units/Math/RigUnit_MathColor.h
--- a/Engine/Plugins/Animation/ControlRig/Source/ControlRig/Private/Units/Math/RigUnit_MathColor.h
+++ b/Engine/Plugins/Animation/ControlRig/Source/ControlRig/Private/Units/Math/RigUnit_MathColor.h
@@ -35,7 +35,7 @@ struct CONTROLRIG_API FRigUnit_MathColorBinaryOp : public FRigUnit_MathColorBase
 /**
  * Makes a vector from a single float
  */
-USTRUCT(meta=(DisplayName="From Float", TemplateName="FromFloat", Keywords="Make,Construct"))
+USTRUCT(meta=(DisplayName="From Float", TemplateName="Cast", Keywords="Make,Construct"))
 struct CONTROLRIG_API FRigUnit_MathColorFromFloat : public FRigUnit_MathColorBase
 {
 	GENERATED_BODY()
diff --git a/Engine/Plugins/Animation/ControlRig/Source/ControlRig/Private/Units/Math/RigUnit_MathDouble.cpp b/Engine/Plugins/Animation/ControlRig/Source/ControlRig/Private/Units/Math/RigUnit_MathDouble.cpp
new file mode 100644
--- /dev/null
+++ b/Engine/Plugins/Animation/ControlRig/Source/ControlRig/Private/Units/Math/RigUnit_MathDouble.cpp
@@ -0,0 +1,310 @@
+// Copyright Epic Games, Inc. All Rights Reserved.
+
+#include "Units/Math/RigUnit_MathDouble.h"
+#include "Units/RigUnitContext.h"
+
+FRigUnit_MathDoubleConstPi_Execute()
+{
+	Value = PI;
+}
+
+FRigUnit_MathDoubleConstHalfPi_Execute()
+{
+	Value = HALF_PI;
+}
+
+FRigUnit_MathDoubleConstTwoPi_Execute()
+{
+	Value = PI * 2.0;
+}
+
+FRigUnit_MathDoubleConstE_Execute()
+{
+	Value = EULERS_NUMBER;
+}
+
+FRigUnit_MathDoubleAdd_Execute()
+{
+    DECLARE_SCOPE_HIERARCHICAL_COUNTER_RIGUNIT()
+	Result = A + B;
+}
+
+FRigUnit_MathDoubleSub_Execute()
+{
+    DECLARE_SCOPE_HIERARCHICAL_COUNTER_RIGUNIT()
+	Result = A - B;
+}
+
+FRigUnit_MathDoubleMul_Execute()
+{
+    DECLARE_SCOPE_HIERARCHICAL_COUNTER_RIGUNIT()
+	Result = A * B;
+}
+
+FRigUnit_MathDoubleDiv_Execute()
+{
+    DECLARE_SCOPE_HIERARCHICAL_COUNTER_RIGUNIT()
+	if(FMath::IsNearlyZero(B))
+	{
+		UE_CONTROLRIG_RIGUNIT_REPORT_WARNING(TEXT("B is nearly 0.0"));
+		Result = 0.0;
+		return;
+	}
+	Result = A / B;
+}
+
+FRigUnit_MathDoubleMod_Execute()
+{
+    DECLARE_SCOPE_HIERARCHICAL_COUNTER_RIGUNIT()
+	if(FMath::IsNearlyZero(B) || B < 0.0)
+	{
+		UE_CONTROLRIG_RIGUNIT_REPORT_WARNING(TEXT("B needs to be greater than 0"));
+		Result = 0.0;
+		return;
+	}
+	Result = FMath::Fmod(A, B);
+}
+
+FRigUnit_MathDoubleMin_Execute()
+{
+    DECLARE_SCOPE_HIERARCHICAL_COUNTER_RIGUNIT()
+	Result = FMath::Min<double>(A, B);
+}
+
+FRigUnit_MathDoubleMax_Execute()
+{
+    DECLARE_SCOPE_HIERARCHICAL_COUNTER_RIGUNIT()
+	Result = FMath::Max<double>(A, B);
+}
+
+FRigUnit_MathDoublePow_Execute()
+{
+    DECLARE_SCOPE_HIERARCHICAL_COUNTER_RIGUNIT()
+	Result = FMath::Pow(A, B);
+}
+
+FRigUnit_MathDoubleSqrt_Execute()
+{
+    DECLARE_SCOPE_HIERARCHICAL_COUNTER_RIGUNIT()
+	if(Value < 0.0)
+	{
+		UE_CONTROLRIG_RIGUNIT_REPORT_WARNING(TEXT("Value is below zero"));
+		Result = 0.0;
+		return;
+	}
+	Result = FMath::Sqrt(Value);
+}
+
+FRigUnit_MathDoubleNegate_Execute()
+{
+    DECLARE_SCOPE_HIERARCHICAL_COUNTER_RIGUNIT()
+	Result = -Value;
+}
+
+FRigUnit_MathDoubleAbs_Execute()
+{
+    DECLARE_SCOPE_HIERARCHICAL_COUNTER_RIGUNIT()
+	Result = FMath::Abs(Value);
+}
+
+FRigUnit_MathDoubleFloor_Execute()
+{
+    DECLARE_SCOPE_HIERARCHICAL_COUNTER_RIGUNIT()
+	Result = FMath::FloorToDouble(Value);
+	Int = (int32)Result;
+}
+
+FRigUnit_MathDoubleCeil_Execute()
+{
+    DECLARE_SCOPE_HIERARCHICAL_COUNTER_RIGUNIT()
+	Result = FMath::CeilToDouble(Value);
+	Int = (int32)Result;
+}
+
+FRigUnit_MathDoubleRound_Execute()
+{
+    DECLARE_SCOPE_HIERARCHICAL_COUNTER_RIGUNIT()
+	Result = FMath::RoundToDouble(Value);
+	Int = (int32)Result;
+}
+
+FRigUnit_MathDoubleToInt_Execute()
+{
+    DECLARE_SCOPE_HIERARCHICAL_COUNTER_RIGUNIT()
+	Result = FMath::FloorToInt(Value);
+}
+
+FRigUnit_MathDoubleSign_Execute()
+{
+    DECLARE_SCOPE_HIERARCHICAL_COUNTER_RIGUNIT()
+	Result = Value < 0.0 ? -1.0 : 1.0;
+}
+
+FRigUnit_MathDoubleClamp_Execute()
+{
+    DECLARE_SCOPE_HIERARCHICAL_COUNTER_RIGUNIT()
+	Result = FMath::Clamp<double>(Value, Minimum, Maximum);
+}
+
+FRigUnit_MathDoubleLerp_Execute()
+{
+    DECLARE_SCOPE_HIERARCHICAL_COUNTER_RIGUNIT()
+	Result = FMath::Lerp<double>(A, B, T);
+}
+
+FRigUnit_MathDoubleRemap_Execute()
+{
+    DECLARE_SCOPE_HIERARCHICAL_COUNTER_RIGUNIT()
+	double Ratio = 0.0;
+	if (FMath::IsNearlyEqual(SourceMinimum, SourceMaximum))
+	{
+		UE_CONTROLRIG_RIGUNIT_REPORT_WARNING(TEXT("The source minimum and maximum are the same."));
+	}
+	else
+	{
+		Ratio = (Value - SourceMinimum) / (SourceMaximum - SourceMinimum);
+	}
+	if (bClamp)
+	{
+		Ratio = FMath::Clamp<double>(Ratio, 0.0, 1.0);
+	}
+	Result = FMath::Lerp<double>(TargetMinimum, TargetMaximum, Ratio);
+}
+
+FRigUnit_MathDoubleEquals_Execute()
+{
+    DECLARE_SCOPE_HIERARCHICAL_COUNTER_RIGUNIT()
+	Result = A == B;
+}
+
+FRigUnit_MathDoubleNotEquals_Execute()
+{
+    DECLARE_SCOPE_HIERARCHICAL_COUNTER_RIGUNIT()
+	Result = A != B;
+}
+
+FRigUnit_MathDoubleGreater_Execute()
+{
+    DECLARE_SCOPE_HIERARCHICAL_COUNTER_RIGUNIT()
+	Result = A > B;
+}
+
+FRigUnit_MathDoubleLess_Execute()
+{
+    DECLARE_SCOPE_HIERARCHICAL_COUNTER_RIGUNIT()
+	Result = A < B;
+}
+
+FRigUnit_MathDoubleGreaterEqual_Execute()
+{
+    DECLARE_SCOPE_HIERARCHICAL_COUNTER_RIGUNIT()
+	Result = A >= B;
+}
+
+FRigUnit_MathDoubleLessEqual_Execute()
+{
+    DECLARE_SCOPE_HIERARCHICAL_COUNTER_RIGUNIT()
+	Result = A <= B;
+}
+
+FRigUnit_MathDoubleIsNearlyZero_Execute()
+{
+    DECLARE_SCOPE_HIERARCHICAL_COUNTER_RIGUNIT()
+	if(Tolerance < 0.0)
+	{
+		UE_CONTROLRIG_RIGUNIT_REPORT_WARNING(TEXT("Tolerance is below zero"));
+	}
+	Result = FMath::IsNearlyZero(Value, FMath::Max<double>(Tolerance, SMALL_NUMBER));
+}
+
+FRigUnit_MathDoubleIsNearlyEqual_Execute()
+{
+    DECLARE_SCOPE_HIERARCHICAL_COUNTER_RIGUNIT()
+	if(Tolerance < 0.0)
+	{
+		UE_CONTROLRIG_RIGUNIT_REPORT_WARNING(TEXT("Tolerance is below zero"));
+	}
+	Result = FMath::IsNearlyEqual(A, B, FMath::Max<double>(Tolerance, SMALL_NUMBER));
+}
+
+FRigUnit_MathDoubleDeg_Execute()
+{
+    DECLARE_SCOPE_HIERARCHICAL_COUNTER_RIGUNIT()
+	Result = FMath::RadiansToDegrees(Value);
+}
+
+FRigUnit_MathDoubleRad_Execute()
+{
+    DECLARE_SCOPE_HIERARCHICAL_COUNTER_RIGUNIT()
+	Result = FMath::DegreesToRadians(Value);
+}
+
+FRigUnit_MathDoubleSin_Execute()
+{
+    DECLARE_SCOPE_HIERARCHICAL_COUNTER_RIGUNIT()
+	Result = FMath::Sin(Value);
+}
+
+FRigUnit_MathDoubleCos_Execute()
+{
+    DECLARE_SCOPE_HIERARCHICAL_COUNTER_RIGUNIT()
+	Result = FMath::Cos(Value);
+}
+
+FRigUnit_MathDoubleTan_Execute()
+{
+    DECLARE_SCOPE_HIERARCHICAL_COUNTER_RIGUNIT()
+	Result = FMath::Tan(Value);
+}
+
+FRigUnit_MathDoubleAsin_Execute()
+{
+    DECLARE_SCOPE_HIERARCHICAL_COUNTER_RIGUNIT()
+	if (!FMath::IsWithinInclusive<double>(Value, -1.0, 1.0))
+	{
+		UE_CONTROLRIG_RIGUNIT_REPORT_WARNING(TEXT("Value is outside of valid range (-1 to 1)"));
+		Result = 0.0;
+		return;
+	}
+	Result = FMath::Asin(Value);
+}
+
+FRigUnit_MathDoubleAcos_Execute()
+{
+    DECLARE_SCOPE_HIERARCHICAL_COUNTER_RIGUNIT()
+	if (!FMath::IsWithinInclusive<double>(Value, -1.0, 1.0))
+	{
+		UE_CONTROLRIG_RIGUNIT_REPORT_WARNING(TEXT("Value is outside of valid range (-1 to 1)"));
+		Result = 0.0;
+		return;
+	}
+	Result = FMath::Acos(Value);
+}
+
+FRigUnit_MathDoubleAtan_Execute()
+{
+    DECLARE_SCOPE_HIERARCHICAL_COUNTER_RIGUNIT()
+	Result = FMath::Atan(Value);
+}
+
+FRigUnit_MathDoubleLawOfCosine_Execute()
+{
+    DECLARE_SCOPE_HIERARCHICAL_COUNTER_RIGUNIT()
+	if ((A <= 0.0) || (B <= 0.0) || (C <= 0.0) || (A + B < C) || (A + C < B) || (B + C < A))
+	{
+		AlphaAngle = BetaAngle = GammaAngle = 0.0;
+		bValid = false;
+		return;
+	}
+
+	GammaAngle = FMath::Acos((A * A + B * B - C * C) / (2.0 * A * B));
+	BetaAngle = FMath::Acos((A * A + C * C - B * B) / (2.0 * A * C));
+	AlphaAngle = PI - GammaAngle - BetaAngle;
+	bValid = true;
+}
+
+FRigUnit_MathDoubleExponential_Execute()
+{
+	DECLARE_SCOPE_HIERARCHICAL_COUNTER_RIGUNIT()
+	Result = FMath::Exp(Value);
+}
diff --git a/Engine/Plugins/Animation/ControlRig/Source/ControlRig/Private/Units/Math/RigUnit_MathDouble.h b/Engine/Plugins/Animation/ControlRig/Source/ControlRig/Private/Units/Math/RigUnit_MathDouble.h
new file mode 100644
--- /dev/null
+++ b/Engine/Plugins/Animation/ControlRig/Source/ControlRig/Private/Units/Math/RigUnit_MathDouble.h
@@ -0,0 +1,847 @@
+// Copyright Epic Games, Inc. All Rights Reserved.
+
+#pragma once
+
+#include "RigUnit_MathBase.h"
+#include "RigUnit_MathDouble.generated.h"
+
+USTRUCT(meta=(Abstract, Category="Math|Double", MenuDescSuffix="(Double)"))
+struct CONTROLRIG_API FRigUnit_MathDoubleBase : public FRigUnit_MathBase
+{
+	GENERATED_BODY()
+};
+
+USTRUCT(meta=(Abstract))
+struct CONTROLRIG_API FRigUnit_MathDoubleConstant : public FRigUnit_MathDoubleBase
+{
+	GENERATED_BODY()
+
+	FRigUnit_MathDoubleConstant()
+	{
+		Value = 0.0;
+	}
+
+	UPROPERTY(meta=(Output))
+	double Value;
+};
+
+USTRUCT(meta=(Abstract))
+struct CONTROLRIG_API FRigUnit_MathDoubleUnaryOp : public FRigUnit_MathDoubleBase
+{
+	GENERATED_BODY()
+
+	FRigUnit_MathDoubleUnaryOp()
+	{
+		Value = Result = 0.0;
+	}
+
+	UPROPERTY(meta=(Input))
+	double Value;
+
+	UPROPERTY(meta=(Output))
+	double Result;
+};
+
+USTRUCT(meta=(Abstract))
+struct CONTROLRIG_API FRigUnit_MathDoubleBinaryOp : public FRigUnit_MathDoubleBase
+{
+	GENERATED_BODY()
+
+	FRigUnit_MathDoubleBinaryOp()
+	{
+		A = B = Result = 0.0;
+	}
+
+	UPROPERTY(meta=(Input))
+	double A;
+
+	UPROPERTY(meta=(Input))
+	double B;
+
+	UPROPERTY(meta=(Output))
+	double Result;
+};
+
+/**
+ * Returns PI
+ */
+USTRUCT(meta=(DisplayName="Pi", TemplateName="Pi"))
+struct CONTROLRIG_API FRigUnit_MathDoubleConstPi : public FRigUnit_MathDoubleConstant
+{
+	GENERATED_BODY()
+	
+	RIGVM_METHOD()
+	virtual void Execute(const FRigUnitContext& Context) override;
+};
+
+/**
+ * Returns PI * 0.5
+ */
+USTRUCT(meta=(DisplayName="Half Pi", TemplateName="HalfPi"))
+struct CONTROLRIG_API FRigUnit_MathDoubleConstHalfPi : public FRigUnit_MathDoubleConstant
+{
+	GENERATED_BODY()
+
+	RIGVM_METHOD()
+	virtual void Execute(const FRigUnitContext& Context) override;
+};
+
+/**
+ * Returns PI * 2.0
+ */
+USTRUCT(meta=(DisplayName="Two Pi", TemplateName="TwoPi", Keywords="Tau"))
+struct CONTROLRIG_API FRigUnit_MathDoubleConstTwoPi : public FRigUnit_MathDoubleConstant
+{
+	GENERATED_BODY()
+
+	RIGVM_METHOD()
+	virtual void Execute(const FRigUnitContext& Context) override;
+};
+
+/**
+ * Returns E
+ */
+USTRUCT(meta=(DisplayName="E", TemplateName="E", Keywords="Euler"))
+struct CONTROLRIG_API FRigUnit_MathDoubleConstE : public FRigUnit_MathDoubleConstant
+{
+	GENERATED_BODY()
+	
+	RIGVM_METHOD()
+	virtual void Execute(const FRigUnitContext& Context) override;
+};
+
+/**
+ * Returns the sum of the two values
+ */
+USTRUCT(meta=(DisplayName="Add", TemplateName="Add", Keywords="Sum,+"))
+struct CONTROLRIG_API FRigUnit_MathDoubleAdd : public FRigUnit_MathDoubleBinaryOp
+{
+	GENERATED_BODY()
+
+	RIGVM_METHOD()
+	virtual void Execute(const FRigUnitContext& Context) override;
+};
+
+/**
+ * Returns the difference of the two values
+ */
+USTRUCT(meta=(DisplayName="Subtract", TemplateName="Subtract", Keywords="-"))
+struct CONTROLRIG_API FRigUnit_MathDoubleSub : public FRigUnit_MathDoubleBinaryOp
+{
+	GENERATED_BODY()
+
+	RIGVM_METHOD()
+	virtual void Execute(const FRigUnitContext& Context) override;
+};
+
+/**
+ * Returns the product of the two values
+ */
+USTRUCT(meta=(DisplayName="Multiply", TemplateName="Multiply", Keywords="Product,*"))
+struct CONTROLRIG_API FRigUnit_MathDoubleMul : public FRigUnit_MathDoubleBinaryOp
+{
+	GENERATED_BODY()
+
+	FRigUnit_MathDoubleMul()
+	{
+		A = 1.0;
+		B = 1.0;
+	}
+
+	RIGVM_METHOD()
+	virtual void Execute(const FRigUnitContext& Context) override;
+};
+
+/**
+ * Returns the division of the two values
+ */
+USTRUCT(meta=(DisplayName="Divide", TemplateName="Divide", Keywords="Division,Divisor,/"))
+struct CONTROLRIG_API FRigUnit_MathDoubleDiv : public FRigUnit_MathDoubleBinaryOp
+{
+	GENERATED_BODY()
+
+	FRigUnit_MathDoubleDiv()
+	{
+		B = 1.0;
+	}
+
+	RIGVM_METHOD()
+	virtual void Execute(const FRigUnitContext& Context) override;
+};
+
+/**
+ * Returns the modulo of the two values
+ */
+USTRUCT(meta=(DisplayName="Modulo", TemplateName="Modulo", Keywords="%,fmod"))
+struct CONTROLRIG_API FRigUnit_MathDoubleMod : public FRigUnit_MathDoubleBinaryOp
+{
+	GENERATED_BODY()
+
+	FRigUnit_MathDoubleMod()
+	{
+		A = 0.0;
+		B = 1.0;
+	}
+
+	RIGVM_METHOD()
+	virtual void Execute(const FRigUnitContext& Context) override;
+};
+
+/**
+ * Returns the smaller of the two values
+ */
+USTRUCT(meta=(DisplayName="Minimum", TemplateName="Minimum"))
+struct CONTROLRIG_API FRigUnit_MathDoubleMin : public FRigUnit_MathDoubleBinaryOp
+{
+	GENERATED_BODY()
+
+	RIGVM_METHOD()
+	virtual void Execute(const FRigUnitContext& Context) override;
+};
+
+/**
+ * Returns the larger of the two values
+ */
+USTRUCT(meta=(DisplayName="Maximum", TemplateName="Maximum"))
+struct CONTROLRIG_API FRigUnit_MathDoubleMax : public FRigUnit_MathDoubleBinaryOp
+{
+	GENERATED_BODY()
+
+	RIGVM_METHOD()
+	virtual void Execute(const FRigUnitContext& Context) override;
+};
+
+/**
+ * Returns the value of A raised to the power of B.
+ */
+USTRUCT(meta=(DisplayName="Power", TemplateName="Power"))
+struct CONTROLRIG_API FRigUnit_MathDoublePow : public FRigUnit_MathDoubleBinaryOp
+{
+	GENERATED_BODY()
+
+	FRigUnit_MathDoublePow()
+	{
+		A = 1.0;
+		B = 1.0;
+	}
+
+	RIGVM_METHOD()
+	virtual void Execute(const FRigUnitContext& Context) override;
+};
+
+/**
+ * Returns the square root of the given value
+ */
+USTRUCT(meta=(DisplayName="Sqrt", TemplateName="Sqrt", Keywords="Root,Square"))
+struct CONTROLRIG_API FRigUnit_MathDoubleSqrt : public FRigUnit_MathDoubleUnaryOp
+{
+	GENERATED_BODY()
+
+	RIGVM_METHOD()
+	virtual void Execute(const FRigUnitContext& Context) override;
+};
+
+/**
+ * Returns the negative value
+ */
+USTRUCT(meta=(DisplayName="Negate", TemplateName="Negate", Keywords="-,Abs"))
+struct CONTROLRIG_API FRigUnit_MathDoubleNegate : public FRigUnit_MathDoubleUnaryOp
+{
+	GENERATED_BODY()
+
+	RIGVM_METHOD()
+	virtual void Execute(const FRigUnitContext& Context) override;
+};
+
+/**
+ * Returns the absolute (positive) value
+ */
+USTRUCT(meta=(DisplayName="Absolute", TemplateName="Absolute", Keywords="Abs,Neg"))
+struct CONTROLRIG_API FRigUnit_MathDoubleAbs : public FRigUnit_MathDoubleUnaryOp
+{
+	GENERATED_BODY()
+
+	RIGVM_METHOD()
+	virtual void Execute(const FRigUnitContext& Context) override;
+};
+
+/**
+ * Returns the closest lower full number (integer) of the value
+ */
+USTRUCT(meta=(DisplayName="Floor", TemplateName="Floor", Keywords="Round"))
+struct CONTROLRIG_API FRigUnit_MathDoubleFloor : public FRigUnit_MathDoubleBase
+{
+	GENERATED_BODY()
+
+	FRigUnit_MathDoubleFloor()
+	{
+		Value = Result = 0.0;
+		Int = 0;
+	}
+
+	RIGVM_METHOD()
+	virtual void Execute(const FRigUnitContext& Context) override;
+
+	UPROPERTY(meta=(Input))
+	double Value;
+
+	UPROPERTY(meta=(Output))
+	double Result;
+
+	UPROPERTY(meta=(Output))
+	int32 Int;
+};
+
+/**
+ * Returns the closest higher full number (integer) of the value
+ */
+USTRUCT(meta=(DisplayName="Ceiling", TemplateName="Ceiling", Keywords="Round"))
+struct CONTROLRIG_API FRigUnit_MathDoubleCeil : public FRigUnit_MathDoubleBase
+{
+	GENERATED_BODY()
+
+	FRigUnit_MathDoubleCeil()
+	{
+		Value = Result = 0.0;
+		Int = 0;
+	}
+
+	RIGVM_METHOD()
+	virtual void Execute(const FRigUnitContext& Context) override;
+
+	UPROPERTY(meta=(Input))
+	double Value;
+
+	UPROPERTY(meta=(Output))
+	double Result;
+
+	UPROPERTY(meta=(Output))
+	int32 Int;
+};
+
+/**
+ * Returns the closest higher full number (integer) of the value
+ */
+USTRUCT(meta=(DisplayName="Round", TemplateName="Round"))
+struct CONTROLRIG_API FRigUnit_MathDoubleRound : public FRigUnit_MathDoubleBase
+{
+	GENERATED_BODY()
+
+	FRigUnit_MathDoubleRound()
+	{
+		Value = Result = 0.0;
+		Int = 0;
+	}
+
+	RIGVM_METHOD()
+	virtual void Execute(const FRigUnitContext& Context) override;
+
+	UPROPERTY(meta=(Input))
+	double Value;
+
+	UPROPERTY(meta=(Output))
+	double Result;
+
+	UPROPERTY(meta=(Output))
+	int32 Int;
+};
+
+/**
+ * Returns the double cast to an int (this uses floor)
+ */
+USTRUCT(meta=(DisplayName="To Int", TemplateName="Cast"))
+struct CONTROLRIG_API FRigUnit_MathDoubleToInt : public FRigUnit_MathDoubleBase
+{
+	GENERATED_BODY()
+
+	FRigUnit_MathDoubleToInt()
+	{
+		Value = 0.0;
+		Result = 0;
+	}
+
+	RIGVM_METHOD()
+	virtual void Execute(const FRigUnitContext& Context) override;
+
+	UPROPERTY(meta=(Input))
+	double Value;
+
+	UPROPERTY(meta=(Output))
+	int32 Result;
+};
+
+/**
+ * Returns the sign of the value (+1 for >= 0.0, -1 for < 0.0)
+ */
+USTRUCT(meta=(DisplayName="Sign", TemplateName="Sign"))
+struct CONTROLRIG_API FRigUnit_MathDoubleSign : public FRigUnit_MathDoubleUnaryOp
+{
+	GENERATED_BODY()
+
+	RIGVM_METHOD()
+	virtual void Execute(const FRigUnitContext& Context) override;
+};
+
+/**
+ * Clamps the given value within the range provided by minimum and maximum
+ */
+USTRUCT(meta=(DisplayName="Clamp", TemplateName="Clamp", Keywords="Range,Remap"))
+struct CONTROLRIG_API FRigUnit_MathDoubleClamp : public FRigUnit_MathDoubleBase
+{
+	GENERATED_BODY()
+
+	FRigUnit_MathDoubleClamp()
+	{
+		Value = Minimum = Maximum = Result = 0.0;
+	}
+
+	RIGVM_METHOD()
+	virtual void Execute(const FRigUnitContext& Context) override;
+
+	UPROPERTY(meta=(Input))
+	double Value;
+
+	UPROPERTY(meta=(Input))
+	double Minimum;
+
+	UPROPERTY(meta=(Input))
+	double Maximum;
+
+	UPROPERTY(meta=(Output))
+	double Result;
+};
+
+/**
+ * Linearly interpolates between A and B using the ratio T
+ */
+USTRUCT(meta=(DisplayName="Interpolate", TemplateName="Interpolate", Keywords="Lerp,Mix,Blend"))
+struct CONTROLRIG_API FRigUnit_MathDoubleLerp : public FRigUnit_MathDoubleBase
+{
+	GENERATED_BODY()
+
+	FRigUnit_MathDoubleLerp()
+	{
+		A = B = T = Result = 0.0;
+		B = 1.0;
+	}
+
+	RIGVM_METHOD()
+	virtual void Execute(const FRigUnitContext& Context) override;
+
+	UPROPERTY(meta=(Input))
+	double A;
+
+	UPROPERTY(meta=(Input))
+	double B;
+
+	UPROPERTY(meta=(Input))
+	double T;
+
+	UPROPERTY(meta=(Output))
+	double Result;
+};
+
+/**
+ * Remaps the given value from a source range to a target range.
+ */
+USTRUCT(meta=(DisplayName="Remap", TemplateName="Remap", Keywords="Rescale,Scale"))
+struct CONTROLRIG_API FRigUnit_MathDoubleRemap : public FRigUnit_MathDoubleBase
+{
+	GENERATED_BODY()
+
+	FRigUnit_MathDoubleRemap()
+	{
+		Value = SourceMinimum = TargetMinimum = Result = 0.0;
+		SourceMaximum = TargetMaximum = 1.0;
+		bClamp = false;
+	}
+
+	RIGVM_METHOD()
+	virtual void Execute(const FRigUnitContext& Context) override;
+
+	UPROPERTY(meta=(Input))
+	double Value;
+
+	UPROPERTY(meta=(Input))
+	double SourceMinimum;
+
+	UPROPERTY(meta=(Input))
+	double SourceMaximum;
+
+	UPROPERTY(meta=(Input))
+	double TargetMinimum;
+
+	UPROPERTY(meta=(Input))
+	double TargetMaximum;
+
+	/** If set to true the result is clamped to the target range */
+	UPROPERTY(meta=(Input))
+	bool bClamp;
+
+	UPROPERTY(meta=(Output))
+	double Result;
+};
+
+/**
+ * Returns true if the value A equals B
+ */
+USTRUCT(meta=(DisplayName="Equals", TemplateName="Equals", Keywords="Same,=="))
+struct CONTROLRIG_API FRigUnit_MathDoubleEquals : public FRigUnit_MathDoubleBase
+{
+	GENERATED_BODY()
+	RIGVM_METHOD()
+	virtual void Execute(const FRigUnitContext& Context) override;
+
+	FRigUnit_MathDoubleEquals()
+	{
+		A = B = 0.0;
+		Result = true;
+	}
+
+	UPROPERTY(meta=(Input))
+	double A;
+
+	UPROPERTY(meta=(Input))
+	double B;
+
+	UPROPERTY(meta=(Output))
+	bool Result;
+};
+
+/**
+ * Returns true if the value A does not equal B
+ */
+USTRUCT(meta=(DisplayName="Not Equals", TemplateName="NotEquals", Keywords="Different,!="))
+struct CONTROLRIG_API FRigUnit_MathDoubleNotEquals : public FRigUnit_MathDoubleBase
+{
+	GENERATED_BODY()
+	RIGVM_METHOD()
+	virtual void Execute(const FRigUnitContext& Context) override;
+
+	FRigUnit_MathDoubleNotEquals()
+	{
+		A = B = 0.0;
+		Result = false;
+	}
+
+	UPROPERTY(meta=(Input))
+	double A;
+
+	UPROPERTY(meta=(Input))
+	double B;
+
+	UPROPERTY(meta=(Output))
+	bool Result;
+};
+
+/**
+ * Returns true if the value A is greater than B
+ */
+USTRUCT(meta=(DisplayName="Greater", TemplateName="Greater", Keywords="Larger,Bigger,>"))
+struct CONTROLRIG_API FRigUnit_MathDoubleGreater : public FRigUnit_MathDoubleBase
+{
+	GENERATED_BODY()
+
+	FRigUnit_MathDoubleGreater()
+	{
+		A = B = 0.0;
+		Result = false;
+	}
+
+	RIGVM_METHOD()
+	virtual void Execute(const FRigUnitContext& Context) override;
+
+	UPROPERTY(meta=(Input))
+	double A;
+
+	UPROPERTY(meta=(Input))
+	double B;
+
+	UPROPERTY(meta=(Output))
+	bool Result;
+};
+
+/**
+ * Returns true if the value A is less than B
+ */
+USTRUCT(meta=(DisplayName="Less", TemplateName="Less", Keywords="Smaller,<"))
+struct CONTROLRIG_API FRigUnit_MathDoubleLess : public FRigUnit_MathDoubleBase
+{
+	GENERATED_BODY()
+	
+	FRigUnit_MathDoubleLess()
+	{
+		A = B = 0.0;
+		Result = false;
+	}
+
+	RIGVM_METHOD()
+	virtual void Execute(const FRigUnitContext& Context) override;
+
+	UPROPERTY(meta=(Input))
+	double A;
+
+	UPROPERTY(meta=(Input))
+	double B;
+
+	UPROPERTY(meta=(Output))
+	bool Result;
+};
+
+/**
+ * Returns true if the value A is greater than or equal to B
+ */
+USTRUCT(meta=(DisplayName="Greater Equal", TemplateName="GreaterEqual", Keywords="Larger,Bigger,>="))
+struct CONTROLRIG_API FRigUnit_MathDoubleGreaterEqual : public FRigUnit_MathDoubleBase
+{
+	GENERATED_BODY()
+
+	FRigUnit_MathDoubleGreaterEqual()
+	{
+		A = B = 0.0;
+		Result = true;
+	}
+
+	RIGVM_METHOD()
+	virtual void Execute(const FRigUnitContext& Context) override;
+
+	UPROPERTY(meta=(Input))
+	double A;
+
+	UPROPERTY(meta=(Input))
+	double B;
+
+	UPROPERTY(meta=(Output))
+	bool Result;
+};
+
+/**
+ * Returns true if the value A is less than or equal to B
+ */
+USTRUCT(meta=(DisplayName="Less Equal", TemplateName="LessEqual", Keywords="Smaller,<="))
+struct CONTROLRIG_API FRigUnit_MathDoubleLessEqual : public FRigUnit_MathDoubleBase
+{
+	GENERATED_BODY()
+
+	FRigUnit_MathDoubleLessEqual()
+	{
+		A = B = 0.0;
+		Result = true;
+	}
+
+	RIGVM_METHOD()
+	virtual void Execute(const FRigUnitContext& Context) override;
+
+	UPROPERTY(meta=(Input))
+	double A;
+
+	UPROPERTY(meta=(Input))
+	double B;
+
+	UPROPERTY(meta=(Output))
+	bool Result;
+};
+
+/**
+ * Returns true if the value is nearly zero
+ */
+USTRUCT(meta=(DisplayName="Is Nearly Zero", TemplateName="IsNearlyZero", Keywords="AlmostZero,0"))
+struct CONTROLRIG_API FRigUnit_MathDoubleIsNearlyZero : public FRigUnit_MathDoubleBase
+{
+	GENERATED_BODY()
+	
+	FRigUnit_MathDoubleIsNearlyZero()
+	{
+		Value = Tolerance = 0.0;
+		Result = true;
+	}
+	
+	RIGVM_METHOD()
+	virtual void Execute(const FRigUnitContext& Context) override;
+
+	UPROPERTY(meta=(Input))
+	double Value;
+
+	UPROPERTY(meta=(Input))
+	double Tolerance;
+
+	UPROPERTY(meta=(Output))
+	bool Result;
+};
+
+/**
+ * Returns true if the value A is almost equal to B
+ */
+USTRUCT(meta=(DisplayName="Is Nearly Equal", TemplateName="IsNearlyEqual", Keywords="AlmostEqual,=="))
+struct CONTROLRIG_API FRigUnit_MathDoubleIsNearlyEqual : public FRigUnit_MathDoubleBase
+{
+	GENERATED_BODY()
+
+	FRigUnit_MathDoubleIsNearlyEqual()
+	{
+		A = B = Tolerance = 0.0;
+		Result = true;
+	}
+
+	RIGVM_METHOD()
+	virtual void Execute(const FRigUnitContext& Context) override;
+
+	UPROPERTY(meta=(Input))
+	double A;
+
+	UPROPERTY(meta=(Input))
+	double B;
+
+	UPROPERTY(meta=(Input))
+	double Tolerance;
+
+	UPROPERTY(meta=(Output))
+	bool Result;
+};
+
+/**
+ * Returns the degrees of a given value in radians
+ */
+USTRUCT(meta=(DisplayName="Degrees", TemplateName="Degrees"))
+struct CONTROLRIG_API FRigUnit_MathDoubleDeg : public FRigUnit_MathDoubleUnaryOp
+{
+	GENERATED_BODY()
+
+	RIGVM_METHOD()
+	virtual void Execute(const FRigUnitContext& Context) override;
+};
+
+/**
+ * Returns the radians of a given value in degrees
+ */
+USTRUCT(meta=(DisplayName="Radians", TemplateName="Radians"))
+struct CONTROLRIG_API FRigUnit_MathDoubleRad : public FRigUnit_MathDoubleUnaryOp
+{
+	GENERATED_BODY()
+
+	RIGVM_METHOD()
+	virtual void Execute(const FRigUnitContext& Context) override;
+};
+
+/**
+ * Returns the sinus value of the given value (in radians)
+ */
+USTRUCT(meta=(DisplayName="Sin", TemplateName="Sin"))
+struct CONTROLRIG_API FRigUnit_MathDoubleSin : public FRigUnit_MathDoubleUnaryOp
+{
+	GENERATED_BODY()
+
+	RIGVM_METHOD()
+	virtual void Execute(const FRigUnitContext& Context) override;
+};
+
+/**
+ * Returns the cosinus value of the given value (in radians)
+ */
+USTRUCT(meta=(DisplayName="Cos", TemplateName="Cos"))
+struct CONTROLRIG_API FRigUnit_MathDoubleCos : public FRigUnit_MathDoubleUnaryOp
+{
+	GENERATED_BODY()
+
+	RIGVM_METHOD()
+	virtual void Execute(const FRigUnitContext& Context) override;
+};
+
+/**
+ * Returns the tangens value of the given value (in radians)
+ */
+USTRUCT(meta=(DisplayName="Tan", TemplateName="Tan"))
+struct CONTROLRIG_API FRigUnit_MathDoubleTan : public FRigUnit_MathDoubleUnaryOp
+{
+	GENERATED_BODY()
+
+	RIGVM_METHOD()
+	virtual void Execute(const FRigUnitContext& Context) override;
+};
+
+/**
+ * Returns the inverse sinus value (in radians) of the given value
+ */
+USTRUCT(meta=(DisplayName="Asin", TemplateName="Asin", Keywords="Arcsin"))
+struct CONTROLRIG_API FRigUnit_MathDoubleAsin : public FRigUnit_MathDoubleUnaryOp
+{
+	GENERATED_BODY()
+
+	RIGVM_METHOD()
+	virtual void Execute(const FRigUnitContext& Context) override;
+};
+
+/**
+ * Returns the inverse cosinus value (in radians) of the given value
+ */
+USTRUCT(meta=(DisplayName="Acos", TemplateName="Acos", Keywords="Arccos"))
+struct CONTROLRIG_API FRigUnit_MathDoubleAcos : public FRigUnit_MathDoubleUnaryOp
+{
+	GENERATED_BODY()
+
+	RIGVM_METHOD()
+	virtual void Execute(const FRigUnitContext& Context) override;
+};
+
+/**
+ * Returns the inverse tangens value (in radians) of the given value
+ */
+USTRUCT(meta=(DisplayName="Atan", TemplateName="Atan", Keywords="Arctan"))
+struct CONTROLRIG_API FRigUnit_MathDoubleAtan : public FRigUnit_MathDoubleUnaryOp
+{
+	GENERATED_BODY()
+
+	RIGVM_METHOD()
+	virtual void Execute(const FRigUnitContext& Context) override;
+};
+
+/**
+ * Computes the angles alpha, beta and gamma (in radians) between the three sides A, B and C
+ */
+USTRUCT(meta = (DisplayName = "Law Of Cosine", TemplateName="LawOfCosine"))
+struct CONTROLRIG_API FRigUnit_MathDoubleLawOfCosine : public FRigUnit_MathDoubleBase
+{
+	GENERATED_BODY()
+
+	FRigUnit_MathDoubleLawOfCosine()
+	{
+		A = B = C = AlphaAngle = BetaAngle = GammaAngle = 0.0;
+		bValid = false;
+	}
+
+	RIGVM_METHOD()
+	virtual void Execute(const FRigUnitContext& Context) override;
+
+	UPROPERTY(meta = (Input))
+	double A;
+
+	UPROPERTY(meta = (Input))
+	double B;
+
+	UPROPERTY(meta = (Input))
+	double C;
+
+	UPROPERTY(meta = (Output))
+	double AlphaAngle;
+
+	UPROPERTY(meta = (Output))
+	double BetaAngle;
+
+	UPROPERTY(meta = (Output))
+	double GammaAngle;
+
+	UPROPERTY(meta = (Output))
+	bool bValid;
+};
+
+/**
+ * Computes the base-e exponential of the given value 
+ */
+USTRUCT(meta = (DisplayName = "Exponential", TemplateName="Exponential"))
+struct CONTROLRIG_API FRigUnit_MathDoubleExponential : public FRigUnit_MathDoubleUnaryOp
+{
+	GENERATED_BODY()
+
+	RIGVM_METHOD()
+	virtual void Execute(const FRigUnitContext& Context) override;
+};
\ No newline at end of file
diff --git a/Engine/Plugins/Animation/ControlRig/Source/ControlRig/Private/Units/Math/RigUnit_MathFloat.h b/Engine/Plugins/Animation/ControlRig/Source/ControlRig/Private/Units/Math/RigUnit_MathFloat.h
--- a/Engine/Plugins/Animation/ControlRig/Source/ControlRig/Private/Units/Math/RigUnit_MathFloat.h
+++ b/Engine/Plugins/Animation/ControlRig/Source/ControlRig/Private/Units/Math/RigUnit_MathFloat.h
@@ -349,7 +349,7 @@ struct CONTROLRIG_API FRigUnit_MathFloatRound : public FRigUnit_MathFloatBase
 /**
  * Returns the float cast to an int (this uses floor)
  */
-USTRUCT(meta=(DisplayName="To Int", TemplateName="Convert"))
+USTRUCT(meta=(DisplayName="To Int", TemplateName="Cast"))
 struct CONTROLRIG_API FRigUnit_MathFloatToInt : public FRigUnit_MathFloatBase
 {
 	GENERATED_BODY()
diff --git a/Engine/Plugins/Animation/ControlRig/Source/ControlRig/Private/Units/Math/RigUnit_MathInt.h b/Engine/Plugins/Animation/ControlRig/Source/ControlRig/Private/Units/Math/RigUnit_MathInt.h
--- a/Engine/Plugins/Animation/ControlRig/Source/ControlRig/Private/Units/Math/RigUnit_MathInt.h
+++ b/Engine/Plugins/Animation/ControlRig/Source/ControlRig/Private/Units/Math/RigUnit_MathInt.h
@@ -194,7 +194,7 @@ struct CONTROLRIG_API FRigUnit_MathIntAbs : public FRigUnit_MathIntUnaryOp
 /**
  * Returns the int cast to a float
  */
-USTRUCT(meta=(DisplayName="To Float", TemplateName="Convert"))
+USTRUCT(meta=(DisplayName="To Float", TemplateName="Cast"))
 struct CONTROLRIG_API FRigUnit_MathIntToFloat : public FRigUnit_MathIntBase
 {
 	GENERATED_BODY()
diff --git a/Engine/Plugins/Animation/ControlRig/Source/ControlRig/Private/Units/Math/RigUnit_MathMatrix.cpp b/Engine/Plugins/Animation/ControlRig/Source/ControlRig/Private/Units/Math/RigUnit_MathMatrix.cpp
--- a/Engine/Plugins/Animation/ControlRig/Source/ControlRig/Private/Units/Math/RigUnit_MathMatrix.cpp
+++ b/Engine/Plugins/Animation/ControlRig/Source/ControlRig/Private/Units/Math/RigUnit_MathMatrix.cpp
@@ -10,6 +10,22 @@ FRigUnit_MathMatrixFromTransform_Execute()
 	Result = Transform.ToMatrixWithScale();
 }
 
+FRigVMStructUpgradeInfo FRigUnit_MathMatrixFromTransform::GetUpgradeInfo() const
+{
+	FRigUnit_MathMatrixFromTransformV2 NewNode;
+	NewNode.Value = Transform;
+
+	FRigVMStructUpgradeInfo Info(*this, NewNode);
+	Info.AddRemappedPin(TEXT("Transform"), TEXT("Value"));
+	return Info;
+}
+
+FRigUnit_MathMatrixFromTransformV2_Execute()
+{
+	DECLARE_SCOPE_HIERARCHICAL_COUNTER_RIGUNIT()
+	Result = Value.ToMatrixWithScale();
+}
+
 FRigUnit_MathMatrixToTransform_Execute()
 {
     DECLARE_SCOPE_HIERARCHICAL_COUNTER_RIGUNIT()
diff --git a/Engine/Plugins/Animation/ControlRig/Source/ControlRig/Private/Units/Math/RigUnit_MathMatrix.h b/Engine/Plugins/Animation/ControlRig/Source/ControlRig/Private/Units/Math/RigUnit_MathMatrix.h
--- a/Engine/Plugins/Animation/ControlRig/Source/ControlRig/Private/Units/Math/RigUnit_MathMatrix.h
+++ b/Engine/Plugins/Animation/ControlRig/Source/ControlRig/Private/Units/Math/RigUnit_MathMatrix.h
@@ -53,7 +53,7 @@ struct CONTROLRIG_API FRigUnit_MathMatrixBinaryOp : public FRigUnit_MathMatrixBa
 /**
 * Makes a transform from a matrix
 */
-USTRUCT(meta=(DisplayName="To Transform", Keywords="Make,Construct"))
+USTRUCT(meta=(DisplayName="To Transform", TemplateName="Cast", Keywords="Make,Construct"))
 struct CONTROLRIG_API FRigUnit_MathMatrixToTransform : public FRigUnit_MathMatrixBase
 {
 	GENERATED_BODY()
@@ -77,7 +77,7 @@ struct CONTROLRIG_API FRigUnit_MathMatrixToTransform : public FRigUnit_MathMatri
 /**
  * Makes a matrix from a transform
  */
-USTRUCT(meta=(DisplayName="From Transform", Keywords="Make,Construct"))
+USTRUCT(meta=(DisplayName="From Transform", Keywords="Make,Construct", Deprecated="5.0.1"))
 struct CONTROLRIG_API FRigUnit_MathMatrixFromTransform : public FRigUnit_MathMatrixBase
 {
 	GENERATED_BODY()
@@ -96,6 +96,33 @@ struct CONTROLRIG_API FRigUnit_MathMatrixFromTransform : public FRigUnit_MathMat
 
 	UPROPERTY(meta=(Output))
 	FMatrix Result;
+	
+	RIGVM_METHOD()
+	virtual FRigVMStructUpgradeInfo GetUpgradeInfo() const override;
+};
+
+/**
+ * Makes a matrix from a transform
+ */
+USTRUCT(meta=(DisplayName="From Transform", TemplateName="Cast", Keywords="Make,Construct"))
+struct CONTROLRIG_API FRigUnit_MathMatrixFromTransformV2 : public FRigUnit_MathMatrixBase
+{
+	GENERATED_BODY()
+
+	FRigUnit_MathMatrixFromTransformV2()
+	{
+		Value = FTransform::Identity;
+		Result = FMatrix::Identity;
+	}
+
+	RIGVM_METHOD()
+	virtual void Execute(const FRigUnitContext& Context) override;
+
+	UPROPERTY(meta=(Input))
+	FTransform Value;
+
+	UPROPERTY(meta=(Output))
+	FMatrix Result;
 };
 
 /**
diff --git a/Engine/Plugins/Animation/ControlRig/Source/ControlRig/Private/Units/Math/RigUnit_MathQuaternion.cpp b/Engine/Plugins/Animation/ControlRig/Source/ControlRig/Private/Units/Math/RigUnit_MathQuaternion.cpp
--- a/Engine/Plugins/Animation/ControlRig/Source/ControlRig/Private/Units/Math/RigUnit_MathQuaternion.cpp
+++ b/Engine/Plugins/Animation/ControlRig/Source/ControlRig/Private/Units/Math/RigUnit_MathQuaternion.cpp
@@ -1,6 +1,7 @@
 // Copyright Epic Games, Inc. All Rights Reserved.
 
 #include "Units/Math/RigUnit_MathQuaternion.h"
+#include "Units/Math/RigUnit_MathTransform.h"
 #include "Units/RigUnitContext.h"
 #include "AnimationCoreLibrary.h"
 
@@ -28,6 +29,22 @@ FRigUnit_MathQuaternionFromRotator_Execute()
 	Result = FQuat(Rotator);
 }
 
+FRigVMStructUpgradeInfo FRigUnit_MathQuaternionFromRotator::GetUpgradeInfo() const
+{
+	FRigUnit_MathQuaternionFromRotatorV2 NewNode;
+	NewNode.Value = Rotator;
+
+	FRigVMStructUpgradeInfo Info(*this, NewNode);
+	Info.AddRemappedPin(TEXT("Rotator"), TEXT("Value"));
+	return Info;
+}
+
+FRigUnit_MathQuaternionFromRotatorV2_Execute()
+{
+	DECLARE_SCOPE_HIERARCHICAL_COUNTER_RIGUNIT()
+	Result = FQuat(Value);
+}
+
 FRigUnit_MathQuaternionFromTwoVectors_Execute()
 {
     DECLARE_SCOPE_HIERARCHICAL_COUNTER_RIGUNIT()
@@ -75,6 +92,27 @@ FRigUnit_MathQuaternionScale_Execute()
 	Value = FQuat(Axis, Angle * Scale);
 }
 
+FRigVMStructUpgradeInfo FRigUnit_MathQuaternionScale::GetUpgradeInfo() const
+{
+	FRigUnit_MathQuaternionScaleV2 NewNode;
+	NewNode.Value = Value;
+	NewNode.Factor = Scale;
+
+	FRigVMStructUpgradeInfo Info(*this, NewNode);
+	Info.AddRemappedPin(TEXT("Scale"), TEXT("Factor"));
+	Info.AddRemappedPin(TEXT("Value"), TEXT("Result"), false, true);
+	return Info;
+}
+
+FRigUnit_MathQuaternionScaleV2_Execute()
+{
+	DECLARE_SCOPE_HIERARCHICAL_COUNTER_RIGUNIT()
+	FVector Axis = FVector::ZeroVector;
+	float Angle = 0.f;
+	Value.ToAxisAndAngle(Axis, Angle);
+	Result = FQuat(Axis, Angle * Factor);
+}
+
 FRigUnit_MathQuaternionToEuler_Execute()
 {
     DECLARE_SCOPE_HIERARCHICAL_COUNTER_RIGUNIT()
@@ -144,7 +182,7 @@ FRigUnit_MathQuaternionUnit_Execute()
 FRigUnit_MathQuaternionRotateVector_Execute()
 {
     DECLARE_SCOPE_HIERARCHICAL_COUNTER_RIGUNIT()
-	Result = Quaternion.RotateVector(Vector);
+	Result = Transform.RotateVector(Vector);
 }
 
 FRigUnit_MathQuaternionGetAxis_Execute()
@@ -187,4 +225,26 @@ FRigUnit_MathQuaternionSwingTwist_Execute()
 
 FRigUnit_MathQuaternionRotationOrder_Execute()
 {
-}
\ No newline at end of file
+}
+
+FRigUnit_MathQuaternionMakeRelative_Execute()
+{
+	DECLARE_SCOPE_HIERARCHICAL_COUNTER_RIGUNIT()
+	Local = Parent.Inverse() * Global;
+	Local.Normalize();
+}
+
+FRigUnit_MathQuaternionMakeAbsolute_Execute()
+{
+	DECLARE_SCOPE_HIERARCHICAL_COUNTER_RIGUNIT()
+	Global = Parent * Local;
+	Global.Normalize();
+}
+
+FRigUnit_MathQuaternionMirrorTransform_Execute()
+{
+	FTransform Transform = FTransform::Identity;
+	Transform.SetRotation(Value);
+	FRigUnit_MathTransformMirrorTransform::StaticExecute(RigVMExecuteContext, Transform, MirrorAxis, AxisToFlip, CentralTransform, Transform, Context);
+	Result = Transform.GetRotation();
+}
diff --git a/Engine/Plugins/Animation/ControlRig/Source/ControlRig/Private/Units/Math/RigUnit_MathQuaternion.h b/Engine/Plugins/Animation/ControlRig/Source/ControlRig/Private/Units/Math/RigUnit_MathQuaternion.h
--- a/Engine/Plugins/Animation/ControlRig/Source/ControlRig/Private/Units/Math/RigUnit_MathQuaternion.h
+++ b/Engine/Plugins/Animation/ControlRig/Source/ControlRig/Private/Units/Math/RigUnit_MathQuaternion.h
@@ -108,7 +108,7 @@ struct CONTROLRIG_API FRigUnit_MathQuaternionFromEuler : public FRigUnit_MathQua
 /**
  * Makes a quaternion from a rotator
  */
-USTRUCT(meta=(DisplayName="From Rotator", TemplateName="FromRotator", Keywords="Make,Construct"))
+USTRUCT(meta=(DisplayName="From Rotator", Keywords="Make,Construct", Deprecated="5.0.1"))
 struct CONTROLRIG_API FRigUnit_MathQuaternionFromRotator : public FRigUnit_MathQuaternionBase
 {
 	GENERATED_BODY()
@@ -127,6 +127,33 @@ struct CONTROLRIG_API FRigUnit_MathQuaternionFromRotator : public FRigUnit_MathQ
 
 	UPROPERTY(meta=(Output))
 	FQuat Result;
+	
+	RIGVM_METHOD()
+	virtual FRigVMStructUpgradeInfo GetUpgradeInfo() const override;
+};
+
+/**
+ * Makes a quaternion from a rotator
+ */
+USTRUCT(meta=(DisplayName="From Rotator", TemplateName="Cast", Keywords="Make,Construct"))
+struct CONTROLRIG_API FRigUnit_MathQuaternionFromRotatorV2 : public FRigUnit_MathQuaternionBase
+{
+	GENERATED_BODY()
+	
+	FRigUnit_MathQuaternionFromRotatorV2()
+	{
+		Value = FRotator::ZeroRotator;
+		Result = FQuat::Identity;
+	}
+
+	RIGVM_METHOD()
+	virtual void Execute(const FRigUnitContext& Context) override;
+
+	UPROPERTY(meta=(Input))
+	FRotator Value;
+
+	UPROPERTY(meta=(Output))
+	FQuat Result;
 };
 
 /**
@@ -187,7 +214,7 @@ struct CONTROLRIG_API FRigUnit_MathQuaternionToAxisAndAngle : public FRigUnit_Ma
 /**
  * Scales a quaternion's angle
  */
-USTRUCT(meta=(DisplayName="Scale", TemplateName="Scale", Keywords="Multiply,Angle,Scale", Constant))
+USTRUCT(meta=(DisplayName="Scale", Keywords="Multiply,Angle,Scale", Constant, Deprecated = "5.0.1"))
 struct CONTROLRIG_API FRigUnit_MathQuaternionScale : public FRigUnit_MathQuaternionBase
 {
 	GENERATED_BODY()
@@ -206,6 +233,36 @@ struct CONTROLRIG_API FRigUnit_MathQuaternionScale : public FRigUnit_MathQuatern
 
 	UPROPERTY(meta=(Input))
 	float Scale;
+
+	RIGVM_METHOD()
+	virtual FRigVMStructUpgradeInfo GetUpgradeInfo() const override;
+};
+
+/**
+ * Scales a quaternion's angle
+ */
+USTRUCT(meta=(DisplayName="Scale", TemplateName="Scale", Keywords="Multiply,Angle,Scale", Constant))
+struct CONTROLRIG_API FRigUnit_MathQuaternionScaleV2 : public FRigUnit_MathQuaternionBase
+{
+	GENERATED_BODY()
+
+	FRigUnit_MathQuaternionScaleV2()
+	{
+		Value = Result = FQuat::Identity;
+		Factor = 1.f;
+	}
+
+	RIGVM_METHOD()
+	virtual void Execute(const FRigUnitContext& Context) override;
+
+	UPROPERTY(meta=(Input))
+	FQuat Value;
+
+	UPROPERTY(meta=(Input))
+	float Factor;
+
+	UPROPERTY(meta=(Output))
+	FQuat Result;
 };
 
 /**
@@ -239,7 +296,7 @@ struct CONTROLRIG_API FRigUnit_MathQuaternionToEuler : public FRigUnit_MathQuate
 /**
  * Retrieves the rotator
  */
-USTRUCT(meta=(DisplayName="To Rotator", TemplateName="ToRotator", Keywords="Make,Construct"))
+USTRUCT(meta=(DisplayName="To Rotator", TemplateName="Cast", Keywords="Make,Construct"))
 struct CONTROLRIG_API FRigUnit_MathQuaternionToRotator : public FRigUnit_MathQuaternionBase
 {
 	GENERATED_BODY()
@@ -440,14 +497,14 @@ struct CONTROLRIG_API FRigUnit_MathQuaternionUnit : public FRigUnit_MathQuaterni
 /**
  * Rotates a given vector by the quaternion
  */
-USTRUCT(meta=(DisplayName="Rotate", TemplateName="Multiply", Keywords="Transform"))
+USTRUCT(meta=(DisplayName="Rotate Vector", TemplateName="Rotate Vector", Keywords="Transform,Multiply"))
 struct CONTROLRIG_API FRigUnit_MathQuaternionRotateVector : public FRigUnit_MathQuaternionBase
 {
 	GENERATED_BODY()
 
 	FRigUnit_MathQuaternionRotateVector()
 	{
-		Quaternion = FQuat::Identity;
+		Transform = FQuat::Identity;
 		Vector = Result = FVector::ZeroVector;
 	}
 
@@ -455,7 +512,7 @@ struct CONTROLRIG_API FRigUnit_MathQuaternionRotateVector : public FRigUnit_Math
 	virtual void Execute(const FRigUnitContext& Context) override;
 
 	UPROPERTY(meta=(Input))
-	FQuat Quaternion;
+	FQuat Transform;
 
 	UPROPERTY(meta=(Input))
 	FVector Vector;
@@ -542,4 +599,94 @@ struct CONTROLRIG_API FRigUnit_MathQuaternionRotationOrder : public FRigUnit_Mat
 
 	RIGVM_METHOD()
 	virtual void Execute(const FRigUnitContext& Context) override;
-};
\ No newline at end of file
+};
+
+/**
+ * Returns the relative local transform within a parent's transform
+ */
+USTRUCT(meta=(DisplayName="Make Relative", TemplateName="Make Relative", Keywords="Local,Global,Absolute"))
+struct CONTROLRIG_API FRigUnit_MathQuaternionMakeRelative : public FRigUnit_MathQuaternionBase
+{
+	GENERATED_BODY()
+
+	FRigUnit_MathQuaternionMakeRelative()
+	{
+		Global = Parent = Local = FQuat::Identity;
+	}
+
+	RIGVM_METHOD()
+	virtual void Execute(const FRigUnitContext& Context) override;
+
+	UPROPERTY(meta=(Input))
+	FQuat Global;
+
+	UPROPERTY(meta=(Input))
+	FQuat Parent;
+
+	UPROPERTY(meta=(Output))
+	FQuat Local;
+};
+
+/**
+ * Returns the absolute global transform within a parent's transform
+ */
+USTRUCT(meta = (DisplayName = "Make Absolute", TemplateName="Make Absolute", Keywords = "Local,Global,Relative"))
+struct CONTROLRIG_API FRigUnit_MathQuaternionMakeAbsolute : public FRigUnit_MathQuaternionBase
+{
+	GENERATED_BODY()
+
+	FRigUnit_MathQuaternionMakeAbsolute()
+	{
+		Global = Parent = Local = FQuat::Identity;
+	}
+
+	RIGVM_METHOD()
+	virtual void Execute(const FRigUnitContext& Context) override;
+
+	UPROPERTY(meta = (Input))
+	FQuat Local;
+
+	UPROPERTY(meta = (Input))
+	FQuat Parent;
+
+	UPROPERTY(meta = (Output))
+	FQuat Global;
+};
+
+/**
+ * Mirror a rotation about a central transform.
+ */
+USTRUCT(meta=(DisplayName="Mirror", TemplateName="Mirror"))
+struct CONTROLRIG_API FRigUnit_MathQuaternionMirrorTransform : public FRigUnit_MathQuaternionBase
+{
+	GENERATED_BODY()
+
+	FRigUnit_MathQuaternionMirrorTransform()
+	{
+		Value = Result = FQuat::Identity;
+		MirrorAxis = EAxis::X;
+		AxisToFlip = EAxis::Z;
+		CentralTransform = FTransform::Identity;
+	}
+
+	RIGVM_METHOD()
+	virtual void Execute(const FRigUnitContext& Context) override;
+
+	UPROPERTY(meta=(Input))
+	FQuat Value;
+
+	// the axis to mirror against
+	UPROPERTY(meta=(Input))
+	TEnumAsByte<EAxis::Type> MirrorAxis;
+
+	// the axis to flip for rotations
+	UPROPERTY(meta=(Input))
+	TEnumAsByte<EAxis::Type> AxisToFlip;
+
+	// The transform about which to mirror
+	UPROPERTY(meta=(Input))
+	FTransform CentralTransform;
+
+	UPROPERTY(meta=(Output))
+	FQuat Result;
+};
diff --git a/Engine/Plugins/Animation/ControlRig/Source/ControlRig/Private/Units/Math/RigUnit_MathQuaternionTests.cpp b/Engine/Plugins/Animation/ControlRig/Source/ControlRig/Private/Units/Math/RigUnit_MathQuaternionTests.cpp
--- a/Engine/Plugins/Animation/ControlRig/Source/ControlRig/Private/Units/Math/RigUnit_MathQuaternionTests.cpp
+++ b/Engine/Plugins/Animation/ControlRig/Source/ControlRig/Private/Units/Math/RigUnit_MathQuaternionTests.cpp
@@ -2,6 +2,7 @@
 
 #if WITH_EDITOR
 #include "Units/Math/RigUnit_MathQuaternion.h"
+#include "Units/Math/RigUnit_MathTransform.h"
 #include "Units/RigUnitTest.h"
 #include "AnimationCoreLibrary.h"
 
@@ -192,7 +193,7 @@ IMPLEMENT_RIGUNIT_AUTOMATION_TEST(FRigUnit_MathQuaternionUnit)
 
 IMPLEMENT_RIGUNIT_AUTOMATION_TEST(FRigUnit_MathQuaternionRotateVector)
 {
-	Unit.Quaternion = FQuat(FVector(1.f, 0.f, 0.f), -HALF_PI);
+	Unit.Transform = FQuat(FVector(1.f, 0.f, 0.f), -HALF_PI);
 	Unit.Vector = FVector(0.f, 0.f, 1.f);
 	InitAndExecute();
 	AddErrorIfFalse(FRigUnit_MathQuatTest_Utils::IsNearlyEqual(Unit.Result, FVector(0.f, 1.f, 0.f)), TEXT("unexpected result"));
@@ -214,4 +215,46 @@ IMPLEMENT_RIGUNIT_AUTOMATION_TEST(FRigUnit_MathQuaternionGetAxis)
 	return true;
 }
 
+IMPLEMENT_RIGUNIT_AUTOMATION_TEST(FRigUnit_MathQuaternionMakeRelative)
+{
+	FRigUnit_MathTransformMakeRelative TransformUnit;
+	Unit.Global = FQuat(FVector(1.f, 0.f, 0.f), -HALF_PI);
+	Unit.Parent = FQuat(FVector(0.2f, 0.8f, 0.f), -HALF_PI).GetNormalized();
+	
+	TransformUnit.Global.SetRotation(Unit.Global);
+	TransformUnit.Parent.SetRotation(Unit.Parent);
+
+	Context.State = EControlRigState::Init;
+	Unit.Execute(Context);
+	TransformUnit.Execute(Context);
+
+	Context.State = EControlRigState::Update;
+	Unit.Execute(Context);
+	TransformUnit.Execute(Context);
+
+	AddErrorIfFalse(FRigUnit_MathQuatTest_Utils::IsNearlyEqual(Unit.Local, TransformUnit.Local.GetRotation()), TEXT("unexpected result"));
+	return true;
+}
+
+IMPLEMENT_RIGUNIT_AUTOMATION_TEST(FRigUnit_MathQuaternionMakeAbsolute)
+{
+	FRigUnit_MathTransformMakeAbsolute TransformUnit;
+	Unit.Local = FQuat(FVector(1.f, 0.f, 0.f), -HALF_PI);
+	Unit.Parent = FQuat(FVector(0.2f, 0.8f, 0.f), -HALF_PI).GetNormalized();
+	
+	TransformUnit.Local.SetRotation(Unit.Local);
+	TransformUnit.Parent.SetRotation(Unit.Parent);
+
+	Context.State = EControlRigState::Init;
+	Unit.Execute(Context);
+	TransformUnit.Execute(Context);
+
+	Context.State = EControlRigState::Update;
+	Unit.Execute(Context);
+	TransformUnit.Execute(Context);
+
+	AddErrorIfFalse(FRigUnit_MathQuatTest_Utils::IsNearlyEqual(Unit.Global, TransformUnit.Global.GetRotation()), TEXT("unexpected result"));
+	return true;
+}
+
 #endif
diff --git a/Engine/Plugins/Animation/ControlRig/Source/ControlRig/Private/Units/Math/RigUnit_MathRBFInterpolate.h b/Engine/Plugins/Animation/ControlRig/Source/ControlRig/Private/Units/Math/RigUnit_MathRBFInterpolate.h
--- a/Engine/Plugins/Animation/ControlRig/Source/ControlRig/Private/Units/Math/RigUnit_MathRBFInterpolate.h
+++ b/Engine/Plugins/Animation/ControlRig/Source/ControlRig/Private/Units/Math/RigUnit_MathRBFInterpolate.h
@@ -79,7 +79,7 @@ struct CONTROLRIG_API FRigUnit_MathRBFInterpolateBase :
 	GENERATED_BODY()
 };
 
-USTRUCT(meta = (Abstract, Keywords = "RBF,Interpolate,Quaternion"))
+USTRUCT(meta = (Abstract, TemplateName="RBF Quaternion", Keywords = "RBF,Interpolate,Quaternion"))
 struct CONTROLRIG_API FRigUnit_MathRBFInterpolateQuatBase :
 	public FRigUnit_MathRBFInterpolateBase
 {
@@ -125,7 +125,7 @@ protected:
 	static uint64 HashTargets(const TArrayView<const FQuat>& Targets);
 };
 
-USTRUCT(meta = (Abstract, Keywords = "RBF,Interpolate,Vector"))
+USTRUCT(meta = (Abstract, TemplateName="RBF Vector", Keywords = "RBF,Interpolate,Vector"))
 struct CONTROLRIG_API FRigUnit_MathRBFInterpolateVectorBase :
 	public FRigUnit_MathRBFInterpolateBase
 {
@@ -182,7 +182,7 @@ struct FMathRBFInterpolateQuatFloat_Target
 	float Value = 0.0f;
 };
 
-USTRUCT(meta = (Abstract, DisplayName = "RBF Quaternion to Float", TemplateName="RBFInterpolateQuatToFloat"))
+USTRUCT(meta = (Abstract, DisplayName = "RBF Quaternion to Float"))
 struct CONTROLRIG_API FRigUnit_MathRBFInterpolateQuatFloat : 
 	public FRigUnit_MathRBFInterpolateQuatBase
 {
@@ -211,7 +211,7 @@ struct FMathRBFInterpolateQuatVector_Target
 	FVector Value = FVector::ZeroVector;
 };
 
-USTRUCT(meta = (Abstract, DisplayName = "RBF Quaternion to Vector", TemplateName="RBFInterpolateQuatToVector"))
+USTRUCT(meta = (Abstract, DisplayName = "RBF Quaternion to Vector"))
 struct CONTROLRIG_API FRigUnit_MathRBFInterpolateQuatVector :
 	public FRigUnit_MathRBFInterpolateQuatBase
 {
@@ -240,7 +240,7 @@ struct FMathRBFInterpolateQuatColor_Target
 	FLinearColor Value = FLinearColor::Transparent;
 };
 
-USTRUCT(meta = (Abstract, DisplayName = "RBF Quaternion to Color", TemplateName="RBFInterpolateQuatToColor"))
+USTRUCT(meta = (Abstract, DisplayName = "RBF Quaternion to Color"))
 struct CONTROLRIG_API FRigUnit_MathRBFInterpolateQuatColor :
 	public FRigUnit_MathRBFInterpolateQuatBase
 {
@@ -269,7 +269,7 @@ struct FMathRBFInterpolateQuatQuat_Target
 	FQuat Value = FQuat::Identity;
 };
 
-USTRUCT(meta = (Abstract, DisplayName = "RBF Quaternion to Quaternion", TemplateName="RBFInterpolateQuatToQuat"))
+USTRUCT(meta = (Abstract, DisplayName = "RBF Quaternion to Quaternion"))
 struct CONTROLRIG_API FRigUnit_MathRBFInterpolateQuatQuat :
 	public FRigUnit_MathRBFInterpolateQuatBase
 {
@@ -298,7 +298,7 @@ struct FMathRBFInterpolateQuatXform_Target
 	FTransform Value = FTransform::Identity;
 };
 
-USTRUCT(meta = (Abstract, DisplayName = "RBF Quaternion to Transform", TemplateName="RBFInterpolateQuatToXform"))
+USTRUCT(meta = (Abstract, DisplayName = "RBF Quaternion to Transform"))
 struct CONTROLRIG_API FRigUnit_MathRBFInterpolateQuatXform :
 	public FRigUnit_MathRBFInterpolateQuatBase
 {
@@ -329,7 +329,7 @@ struct FMathRBFInterpolateVectorFloat_Target
 	float Value = 0.0f;
 };
 
-USTRUCT(meta = (Abstract, DisplayName = "RBF Vector to Float", TemplateName="RBFInterpolateVectorToFloat", Keywords = "RBF,Interpolate,Vector"))
+USTRUCT(meta = (Abstract, DisplayName = "RBF Vector to Float", Keywords = "RBF,Interpolate,Vector"))
 struct CONTROLRIG_API FRigUnit_MathRBFInterpolateVectorFloat :
 	public FRigUnit_MathRBFInterpolateVectorBase
 {
@@ -358,7 +358,7 @@ struct FMathRBFInterpolateVectorVector_Target
 	FVector Value = FVector::ZeroVector;
 };
 
-USTRUCT(meta = (Abstract, DisplayName = "RBF Vector to Vector", TemplateName="RBFInterpolateVectorToVector", Keywords = "RBF,Interpolate,Vector"))
+USTRUCT(meta = (Abstract, DisplayName = "RBF Vector to Vector", Keywords = "RBF,Interpolate,Vector"))
 struct CONTROLRIG_API FRigUnit_MathRBFInterpolateVectorVector :
 	public FRigUnit_MathRBFInterpolateVectorBase
 {
@@ -387,7 +387,7 @@ struct FMathRBFInterpolateVectorColor_Target
 	FLinearColor Value = FLinearColor::Transparent;
 };
 
-USTRUCT(meta = (Abstract, DisplayName = "RBF Vector to Color", TemplateName="RBFInterpolateVectorToColor", Keywords = "RBF,Interpolate,Vector"))
+USTRUCT(meta = (Abstract, DisplayName = "RBF Vector to Color", Keywords = "RBF,Interpolate,Vector"))
 struct CONTROLRIG_API FRigUnit_MathRBFInterpolateVectorColor :
 	public FRigUnit_MathRBFInterpolateVectorBase
 {
@@ -417,7 +417,7 @@ struct FMathRBFInterpolateVectorQuat_Target
 	FQuat Value = FQuat::Identity;
 };
 
-USTRUCT(meta = (Abstract, DisplayName = "RBF Vector to Quat", TemplateName="RBFInterpolateVectorToQuat", Keywords = "RBF,Interpolate,Vector"))
+USTRUCT(meta = (Abstract, DisplayName = "RBF Vector to Quat", Keywords = "RBF,Interpolate,Vector"))
 struct CONTROLRIG_API FRigUnit_MathRBFInterpolateVectorQuat :
 	public FRigUnit_MathRBFInterpolateVectorBase
 {
@@ -445,7 +445,7 @@ struct FMathRBFInterpolateVectorXform_Target
 	FTransform Value = FTransform::Identity;
 };
 
-USTRUCT(meta = (Abstract, DisplayName = "RBF Vector to Transform", TemplateName="RBFInterpolateVectorToXform", Keywords = "RBF,Interpolate,Vector"))
+USTRUCT(meta = (Abstract, DisplayName = "RBF Vector to Transform", Keywords = "RBF,Interpolate,Vector"))
 struct CONTROLRIG_API FRigUnit_MathRBFInterpolateVectorXform :
 	public FRigUnit_MathRBFInterpolateVectorBase
 {
diff --git a/Engine/Plugins/Animation/ControlRig/Source/ControlRig/Private/Units/Math/RigUnit_MathTransform.cpp b/Engine/Plugins/Animation/ControlRig/Source/ControlRig/Private/Units/Math/RigUnit_MathTransform.cpp
--- a/Engine/Plugins/Animation/ControlRig/Source/ControlRig/Private/Units/Math/RigUnit_MathTransform.cpp
+++ b/Engine/Plugins/Animation/ControlRig/Source/ControlRig/Private/Units/Math/RigUnit_MathTransform.cpp
@@ -5,6 +5,7 @@
 #include "Math/ControlRigMathLibrary.h"
 #include "Units/RigUnitContext.h"
 #include "AnimationCoreLibrary.h"
+#include "Rigs/RigHierarchyDefines.h"
 
 FRigUnit_MathTransformFromEulerTransform_Execute()
 {
@@ -12,6 +13,22 @@ FRigUnit_MathTransformFromEulerTransform_Execute()
 	Result = EulerTransform.ToFTransform();
 }
 
+FRigVMStructUpgradeInfo FRigUnit_MathTransformFromEulerTransform::GetUpgradeInfo() const
+{
+	FRigUnit_MathTransformFromEulerTransformV2 NewNode;
+	NewNode.Value = EulerTransform;
+
+	FRigVMStructUpgradeInfo Info(*this, NewNode);
+	Info.AddRemappedPin(TEXT("EulerTransform"), TEXT("Value"));
+	return Info;
+}
+
+FRigUnit_MathTransformFromEulerTransformV2_Execute()
+{
+	DECLARE_SCOPE_HIERARCHICAL_COUNTER_RIGUNIT()
+	Result = Value.ToFTransform();
+}
+
 FRigUnit_MathTransformToEulerTransform_Execute()
 {
     DECLARE_SCOPE_HIERARCHICAL_COUNTER_RIGUNIT()
@@ -128,7 +145,7 @@ FRigVMStructUpgradeInfo FRigUnit_MathTransformSelectBool::GetUpgradeInfo() const
 FRigUnit_MathTransformRotateVector_Execute()
 {
     DECLARE_SCOPE_HIERARCHICAL_COUNTER_RIGUNIT()
-	Result = Transform.TransformVector(Direction);
+	Result = Transform.TransformVector(Vector);
 }
 
 FRigUnit_MathTransformTransformVector_Execute()
@@ -153,4 +170,16 @@ FRigUnit_MathTransformClampSpatially_Execute()
 	FRigUnit_MathVectorClampSpatially::StaticExecute(RigVMExecuteContext, Value.GetTranslation(), Axis, Type, Minimum, Maximum, Space, bDrawDebug, DebugColor, DebugThickness, Position, Context);
 	Result = Value;
 	Result.SetTranslation(Position);
-}
\ No newline at end of file
+}
+
+FRigUnit_MathTransformMirrorTransform_Execute()
+{
+	FRigMirrorSettings MirrorSettings;
+	MirrorSettings.MirrorAxis = MirrorAxis;
+	MirrorSettings.AxisToFlip = AxisToFlip;
+
+	FTransform Local = FTransform::Identity;
+	FRigUnit_MathTransformMakeRelative::StaticExecute(RigVMExecuteContext, Value, CentralTransform, Local, Context);
+	Local = MirrorSettings.MirrorTransform(Local);
+	FRigUnit_MathTransformMakeAbsolute::StaticExecute(RigVMExecuteContext, Local, CentralTransform, Result, Context);
+}
diff --git a/Engine/Plugins/Animation/ControlRig/Source/ControlRig/Private/Units/Math/RigUnit_MathTransform.h b/Engine/Plugins/Animation/ControlRig/Source/ControlRig/Private/Units/Math/RigUnit_MathTransform.h
--- a/Engine/Plugins/Animation/ControlRig/Source/ControlRig/Private/Units/Math/RigUnit_MathTransform.h
+++ b/Engine/Plugins/Animation/ControlRig/Source/ControlRig/Private/Units/Math/RigUnit_MathTransform.h
@@ -59,7 +59,7 @@ struct CONTROLRIG_API FRigUnit_MathTransformBinaryOp : public FRigUnit_MathTrans
 /**
  * Makes a quaternion based transform from a euler based transform
  */
-USTRUCT(meta=(DisplayName="From Euler Transform", TemplateName="FromEulerTransform", Keywords="Make,Construct"))
+USTRUCT(meta=(DisplayName="From Euler Transform", TemplateName="FromEulerTransform", Keywords="Make,Construct", Deprecated="5.0.1"))
 struct CONTROLRIG_API FRigUnit_MathTransformFromEulerTransform : public FRigUnit_MathTransformBase
 {
 	GENERATED_BODY()
@@ -78,12 +78,39 @@ struct CONTROLRIG_API FRigUnit_MathTransformFromEulerTransform : public FRigUnit
 
 	UPROPERTY(meta=(Output))
 	FTransform Result;
+	
+	RIGVM_METHOD()
+	virtual FRigVMStructUpgradeInfo GetUpgradeInfo() const override;
+};
+
+/**
+ * Makes a quaternion based transform from a euler based transform
+ */
+USTRUCT(meta=(DisplayName="To Transform", TemplateName="Cast", Keywords="Make,Construct"))
+struct CONTROLRIG_API FRigUnit_MathTransformFromEulerTransformV2 : public FRigUnit_MathTransformBase
+{
+	GENERATED_BODY()
+
+	FRigUnit_MathTransformFromEulerTransformV2()
+	{
+		Value = FEulerTransform::Identity;
+		Result = FTransform::Identity;
+	}
+
+	RIGVM_METHOD()
+	virtual void Execute(const FRigUnitContext& Context) override;
+
+	UPROPERTY(meta=(Input))
+	FEulerTransform Value;
+
+	UPROPERTY(meta=(Output))
+	FTransform Result;
 };
 
 /**
  * Retrieves a euler based transform from a quaternion based transform
  */
-USTRUCT(meta=(DisplayName="To Euler Transform", TemplateName="ToEulerTransform", Keywords="Make,Construct"))
+USTRUCT(meta=(DisplayName="To Euler Transform", TemplateName="Cast", Keywords="Make,Construct"))
 struct CONTROLRIG_API FRigUnit_MathTransformToEulerTransform : public FRigUnit_MathTransformBase
 {
 	GENERATED_BODY()
@@ -118,7 +145,7 @@ struct CONTROLRIG_API FRigUnit_MathTransformMul : public FRigUnit_MathTransformB
 /**
  * Returns the relative local transform within a parent's transform
  */
-USTRUCT(meta=(DisplayName="Make Relative", TemplateName="MakeRelative", Keywords="Local,Global,Absolute"))
+USTRUCT(meta=(DisplayName="Make Relative", TemplateName="Make Relative", Keywords="Local,Global,Absolute"))
 struct CONTROLRIG_API FRigUnit_MathTransformMakeRelative : public FRigUnit_MathTransformBase
 {
 	GENERATED_BODY()
@@ -144,7 +171,7 @@ struct CONTROLRIG_API FRigUnit_MathTransformMakeRelative : public FRigUnit_MathT
 /**
  * Returns the absolute global transform within a parent's transform
  */
-USTRUCT(meta = (DisplayName = "Make Absolute", TemplateName="MakeAbsolute", Keywords = "Local,Global,Relative"))
+USTRUCT(meta = (DisplayName = "Make Absolute", TemplateName="Make Absolute", Keywords = "Local,Global,Relative"))
 struct CONTROLRIG_API FRigUnit_MathTransformMakeAbsolute : public FRigUnit_MathTransformBase
 {
 	GENERATED_BODY()
@@ -288,7 +315,7 @@ struct CONTROLRIG_API FRigUnit_MathTransformSelectBool : public FRigUnit_MathTra
 /**
  * Rotates a given vector (direction) by the transform
  */
-USTRUCT(meta=(DisplayName="Transform Direction", TemplateName="Rotate", Keywords="Transform,Direction"))
+USTRUCT(meta=(DisplayName="Rotate Vector", TemplateName="Rotate Vector", Keywords="Transform,Direction,TransformDirection"))
 struct CONTROLRIG_API FRigUnit_MathTransformRotateVector : public FRigUnit_MathTransformBase
 {
 	GENERATED_BODY()
@@ -296,7 +323,7 @@ struct CONTROLRIG_API FRigUnit_MathTransformRotateVector : public FRigUnit_MathT
 	FRigUnit_MathTransformRotateVector()
 	{
 		Transform = FTransform::Identity;
-		Direction = Result = FVector::ZeroVector;
+		Vector = Result = FVector::ZeroVector;
 	}
 
 	RIGVM_METHOD()
@@ -306,7 +333,7 @@ struct CONTROLRIG_API FRigUnit_MathTransformRotateVector : public FRigUnit_MathT
 	FTransform Transform;
 
 	UPROPERTY(meta=(Input))
-	FVector Direction;
+	FVector Vector;
 
 	UPROPERTY(meta=(Output))
 	FVector Result;
@@ -434,4 +461,42 @@ struct CONTROLRIG_API FRigUnit_MathTransformClampSpatially : public FRigUnit_Mat
 
 	UPROPERTY(meta = (Output))
 	FTransform Result;
-};
\ No newline at end of file
+};
+
+/**
+ * Mirror a transform about a central transform.
+ */
+USTRUCT(meta=(DisplayName="Mirror", TemplateName="Mirror"))
+struct CONTROLRIG_API FRigUnit_MathTransformMirrorTransform : public FRigUnit_MathTransformBase
+{
+	GENERATED_BODY()
+
+	FRigUnit_MathTransformMirrorTransform()
+	{
+		Value = Result = FTransform::Identity;
+		MirrorAxis = EAxis::X;
+		AxisToFlip = EAxis::Z;
+		CentralTransform = FTransform::Identity;
+	}
+
+	RIGVM_METHOD()
+	virtual void Execute(const FRigUnitContext& Context) override;
+
+	UPROPERTY(meta=(Input))
+	FTransform Value;
+
+	// the axis to mirror against
+	UPROPERTY(meta=(Input))
+	TEnumAsByte<EAxis::Type> MirrorAxis;
+
+	// the axis to flip for rotations
+	UPROPERTY(meta=(Input))
+	TEnumAsByte<EAxis::Type> AxisToFlip;
+
+	// The transform about which to mirror
+	UPROPERTY(meta=(Input))
+	FTransform CentralTransform;
+
+	UPROPERTY(meta=(Output))
+	FTransform Result;
+};
diff --git a/Engine/Plugins/Animation/ControlRig/Source/ControlRig/Private/Units/Math/RigUnit_MathTransformTests.cpp b/Engine/Plugins/Animation/ControlRig/Source/ControlRig/Private/Units/Math/RigUnit_MathTransformTests.cpp
--- a/Engine/Plugins/Animation/ControlRig/Source/ControlRig/Private/Units/Math/RigUnit_MathTransformTests.cpp
+++ b/Engine/Plugins/Animation/ControlRig/Source/ControlRig/Private/Units/Math/RigUnit_MathTransformTests.cpp
@@ -112,7 +112,7 @@ IMPLEMENT_RIGUNIT_AUTOMATION_TEST(FRigUnit_MathTransformSelectBool)
 IMPLEMENT_RIGUNIT_AUTOMATION_TEST(FRigUnit_MathTransformRotateVector)
 {
 	Unit.Transform = FTransform(FQuat(FVector(1.f, 0.f, 0.f), -HALF_PI), FVector(1.f, 2.f, 3.f));
-	Unit.Direction = FVector(0.f, 0.f, 4.f);
+	Unit.Vector = FVector(0.f, 0.f, 4.f);
 	InitAndExecute();
 	AddErrorIfFalse(FRigUnit_MathTransformTest_Utils::IsNearlyEqual(Unit.Result, FVector(0.f, 4.f, 0.f)), TEXT("unexpected result"));
 	return true;
diff --git a/Engine/Plugins/Animation/ControlRig/Source/ControlRig/Private/Units/Math/RigUnit_MathVector.cpp b/Engine/Plugins/Animation/ControlRig/Source/ControlRig/Private/Units/Math/RigUnit_MathVector.cpp
--- a/Engine/Plugins/Animation/ControlRig/Source/ControlRig/Private/Units/Math/RigUnit_MathVector.cpp
+++ b/Engine/Plugins/Animation/ControlRig/Source/ControlRig/Private/Units/Math/RigUnit_MathVector.cpp
@@ -1,6 +1,7 @@
 // Copyright Epic Games, Inc. All Rights Reserved.
 
 #include "Units/Math/RigUnit_MathVector.h"
+#include "Units/Math/RigUnit_MathTransform.h"
 #include "Units/RigUnitContext.h"
 #include "Math/ControlRigMathLibrary.h"
 
@@ -529,3 +530,22 @@ FRigUnit_MathDistanceToPlane_Execute()
 	}
 }
 
+FRigUnit_MathVectorMakeRelative_Execute()
+{
+	DECLARE_SCOPE_HIERARCHICAL_COUNTER_RIGUNIT()
+	Local = Global - Parent;
+}
+
+FRigUnit_MathVectorMakeAbsolute_Execute()
+{
+	DECLARE_SCOPE_HIERARCHICAL_COUNTER_RIGUNIT()
+	Global = Parent + Local;
+}
+
+FRigUnit_MathVectorMirrorTransform_Execute()
+{
+	FTransform Transform = FTransform::Identity;
+	Transform.SetTranslation(Value);
+	FRigUnit_MathTransformMirrorTransform::StaticExecute(RigVMExecuteContext, Transform, MirrorAxis, AxisToFlip, CentralTransform, Transform, Context);
+	Result = Transform.GetTranslation();
+}
\ No newline at end of file
diff --git a/Engine/Plugins/Animation/ControlRig/Source/ControlRig/Private/Units/Math/RigUnit_MathVector.h b/Engine/Plugins/Animation/ControlRig/Source/ControlRig/Private/Units/Math/RigUnit_MathVector.h
--- a/Engine/Plugins/Animation/ControlRig/Source/ControlRig/Private/Units/Math/RigUnit_MathVector.h
+++ b/Engine/Plugins/Animation/ControlRig/Source/ControlRig/Private/Units/Math/RigUnit_MathVector.h
@@ -52,7 +52,7 @@ struct CONTROLRIG_API FRigUnit_MathVectorBinaryOp : public FRigUnit_MathVectorBa
 /**
  * Makes a vector from a single float
  */
-USTRUCT(meta=(DisplayName="From Float", TemplateName="FromFloat", Keywords="Make,Construct"))
+USTRUCT(meta=(DisplayName="From Float", TemplateName="Cast", Keywords="Make,Construct"))
 struct CONTROLRIG_API FRigUnit_MathVectorFromFloat : public FRigUnit_MathVectorBase
 {
 	GENERATED_BODY()
@@ -226,7 +226,7 @@ struct CONTROLRIG_API FRigUnit_MathVectorAbs : public FRigUnit_MathVectorUnaryOp
 /**
  * Returns the closest lower full number (integer) of the value for each component
  */
-USTRUCT(meta=(DisplayName="Floor", TemplateName="Floor", Keywords="Round"))
+USTRUCT(meta=(DisplayName="Floor", Keywords="Round"))
 struct CONTROLRIG_API FRigUnit_MathVectorFloor : public FRigUnit_MathVectorUnaryOp
 {
 	GENERATED_BODY()
@@ -238,7 +238,7 @@ struct CONTROLRIG_API FRigUnit_MathVectorFloor : public FRigUnit_MathVectorUnary
 /**
  * Returns the closest higher full number (integer) of the value for each component
  */
-USTRUCT(meta=(DisplayName="Ceiling", TemplateName="Ceiling", Keywords="Round"))
+USTRUCT(meta=(DisplayName="Ceiling", Keywords="Round"))
 struct CONTROLRIG_API FRigUnit_MathVectorCeil : public FRigUnit_MathVectorUnaryOp
 {
 	GENERATED_BODY()
@@ -250,7 +250,7 @@ struct CONTROLRIG_API FRigUnit_MathVectorCeil : public FRigUnit_MathVectorUnaryO
 /**
  * Returns the closest higher full number (integer) of the value for each component
  */
-USTRUCT(meta=(DisplayName="Round", TemplateName="Round"))
+USTRUCT(meta=(DisplayName="Round"))
 struct CONTROLRIG_API FRigUnit_MathVectorRound : public FRigUnit_MathVectorUnaryOp
 {
 	GENERATED_BODY()
@@ -730,7 +730,7 @@ struct CONTROLRIG_API FRigUnit_MathVectorClampLength: public FRigUnit_MathVector
 /**
  * Mirror a vector about a normal vector.
  */
-USTRUCT(meta=(DisplayName="Mirror", TemplateName="Mirror"))
+USTRUCT(meta=(DisplayName="Mirror on Normal"))
 struct CONTROLRIG_API FRigUnit_MathVectorMirror : public FRigUnit_MathVectorBase
 {
 	GENERATED_BODY()
@@ -1019,4 +1019,94 @@ struct CONTROLRIG_API FRigUnit_MathDistanceToPlane : public FRigUnit_MathVectorB
 
 	UPROPERTY(meta = (Output))
 	float SignedDistance;
-};
\ No newline at end of file
+};
+
+/**
+ * Returns the relative local vector within a parent's vector
+ */
+USTRUCT(meta=(DisplayName="Make Relative", TemplateName="Make Relative", Keywords="Local,Global,Absolute"))
+struct CONTROLRIG_API FRigUnit_MathVectorMakeRelative : public FRigUnit_MathVectorBase
+{
+	GENERATED_BODY()
+
+	FRigUnit_MathVectorMakeRelative()
+	{
+		Global = Parent = Local = FVector::ZeroVector;
+	}
+
+	RIGVM_METHOD()
+	virtual void Execute(const FRigUnitContext& Context) override;
+
+	UPROPERTY(meta=(Input))
+	FVector Global;
+
+	UPROPERTY(meta=(Input))
+	FVector Parent;
+
+	UPROPERTY(meta=(Output))
+	FVector Local;
+};
+
+/**
+ * Returns the absolute global vector within a parent's vector
+ */
+USTRUCT(meta = (DisplayName = "Make Absolute", TemplateName="Make Absolute", Keywords = "Local,Global,Relative"))
+struct CONTROLRIG_API FRigUnit_MathVectorMakeAbsolute : public FRigUnit_MathVectorBase
+{
+	GENERATED_BODY()
+
+	FRigUnit_MathVectorMakeAbsolute()
+	{
+		Global = Parent = Local = FVector::ZeroVector;
+	}
+
+	RIGVM_METHOD()
+	virtual void Execute(const FRigUnitContext& Context) override;
+
+	UPROPERTY(meta = (Input))
+	FVector Local;
+
+	UPROPERTY(meta = (Input))
+	FVector Parent;
+
+	UPROPERTY(meta = (Output))
+	FVector Global;
+};
+
+/**
+ * Mirror a vector about a central transform.
+ */
+USTRUCT(meta=(DisplayName="Mirror", TemplateName="Mirror"))
+struct CONTROLRIG_API FRigUnit_MathVectorMirrorTransform : public FRigUnit_MathVectorBase
+{
+	GENERATED_BODY()
+
+	FRigUnit_MathVectorMirrorTransform()
+	{
+		Value = Result = FVector::ZeroVector;
+		MirrorAxis = EAxis::X;
+		AxisToFlip = EAxis::Z;
+		CentralTransform = FTransform::Identity;
+	}
+
+	RIGVM_METHOD()
+	virtual void Execute(const FRigUnitContext& Context) override;
+
+	UPROPERTY(meta=(Input))
+	FVector Value;
+
+	// the axis to mirror against
+	UPROPERTY(meta=(Input))
+	TEnumAsByte<EAxis::Type> MirrorAxis;
+
+	// the axis to flip for rotations
+	UPROPERTY(meta=(Input))
+	TEnumAsByte<EAxis::Type> AxisToFlip;
+
+	// The transform about which to mirror
+	UPROPERTY(meta=(Input))
+	FTransform CentralTransform;
+
+	UPROPERTY(meta=(Output))
+	FVector Result;
+};
diff --git a/Engine/Plugins/Animation/ControlRig/Source/ControlRig/Private/Units/Simulation/RigUnit_AlphaInterp.cpp b/Engine/Plugins/Animation/ControlRig/Source/ControlRig/Private/Units/Simulation/RigUnit_AlphaInterp.cpp
--- a/Engine/Plugins/Animation/ControlRig/Source/ControlRig/Private/Units/Simulation/RigUnit_AlphaInterp.cpp
+++ b/Engine/Plugins/Animation/ControlRig/Source/ControlRig/Private/Units/Simulation/RigUnit_AlphaInterp.cpp
@@ -98,6 +98,57 @@ FString FRigUnit_AlphaInterpVector::ProcessPinLabelForInjection(const FString& I
 		Formula += FString::Printf(TEXT(" Clamp(%.02f, %.02f)"), ClampMin, ClampMax);
 	}
 
+	if (Formula.IsEmpty())
+	{
+		return InLabel;
+	}
+	return FString::Printf(TEXT("%s: %s"), *InLabel, *Formula);
+}
+
+FRigUnit_AlphaInterpQuat_Execute()
+{
+    DECLARE_SCOPE_HIERARCHICAL_COUNTER_RIGUNIT()
+
+	ScaleBiasClamp.bMapRange = bMapRange;
+	ScaleBiasClamp.bClampResult = bClampResult;
+	ScaleBiasClamp.bInterpResult = bInterpResult;
+
+	if (Context.State == EControlRigState::Init)
+	{
+		ScaleBiasClamp.Reinitialize();
+	}
+	else
+	{
+		ScaleBiasClamp.InRange = InRange;
+		ScaleBiasClamp.OutRange = OutRange;
+		ScaleBiasClamp.ClampMin = ClampMin;
+		ScaleBiasClamp.ClampMax = ClampMax;
+		ScaleBiasClamp.Scale = Scale;
+		ScaleBiasClamp.Bias = Bias;
+		ScaleBiasClamp.InterpSpeedIncreasing = InterpSpeedIncreasing;
+		ScaleBiasClamp.InterpSpeedDecreasing = InterpSpeedDecreasing;
+
+		const float T = ScaleBiasClamp.ApplyTo(1.f, Context.DeltaTime);
+		Result = FQuat::Slerp(FQuat::Identity, Value, T);
+	}
+}
+
+FString FRigUnit_AlphaInterpQuat::ProcessPinLabelForInjection(const FString& InLabel) const
+{
+	FString Formula;
+	if (bMapRange)
+	{
+		Formula += FString::Printf(TEXT(" Map(%.02f, %.02f, %.02f, %.02f)"), InRange.Min, InRange.Max, OutRange.Min, OutRange.Max);
+	}
+	if (bInterpResult)
+	{
+		Formula += FString::Printf(TEXT(" Interp(%.02f, %.02f)"), InterpSpeedIncreasing, InterpSpeedDecreasing);
+	}
+	if (bClampResult)
+	{
+		Formula += FString::Printf(TEXT(" Clamp(%.02f, %.02f)"), ClampMin, ClampMax);
+	}
+
 	if (Formula.IsEmpty())
 	{
 		return InLabel;
diff --git a/Engine/Plugins/Animation/ControlRig/Source/ControlRig/Private/Units/Simulation/RigUnit_AlphaInterp.h b/Engine/Plugins/Animation/ControlRig/Source/ControlRig/Private/Units/Simulation/RigUnit_AlphaInterp.h
--- a/Engine/Plugins/Animation/ControlRig/Source/ControlRig/Private/Units/Simulation/RigUnit_AlphaInterp.h
+++ b/Engine/Plugins/Animation/ControlRig/Source/ControlRig/Private/Units/Simulation/RigUnit_AlphaInterp.h
@@ -151,3 +151,76 @@ struct CONTROLRIG_API FRigUnit_AlphaInterpVector : public FRigUnit_SimBase
 	UPROPERTY()
 	FInputScaleBiasClamp ScaleBiasClamp;
 };
+
+/**
+ * Takes in a vector value and outputs an accumulated value with a customized scale and clamp
+ */
+USTRUCT(meta=(DisplayName="Alpha Interpolate", Keywords="Alpha,Lerp,LinearInterpolate", Category = "Simulation|Time", TemplateName="AlphaInterp", MenuDescSuffix = "(Quat)"))
+struct CONTROLRIG_API FRigUnit_AlphaInterpQuat : public FRigUnit_SimBase
+{
+	GENERATED_BODY()
+	
+	FRigUnit_AlphaInterpQuat()
+	{
+		Value = Result = FQuat::Identity;
+		ScaleBiasClamp = FInputScaleBiasClamp();
+		bMapRange = ScaleBiasClamp.bMapRange;
+		bClampResult = ScaleBiasClamp.bClampResult;
+		bInterpResult = ScaleBiasClamp.bInterpResult;
+		InRange = ScaleBiasClamp.InRange;
+		OutRange = ScaleBiasClamp.OutRange;
+		Scale = ScaleBiasClamp.Scale;
+		Bias = ScaleBiasClamp.Bias;
+		ClampMin = ScaleBiasClamp.ClampMin;
+		ClampMax = ScaleBiasClamp.ClampMax;
+		InterpSpeedIncreasing = ScaleBiasClamp.InterpSpeedIncreasing;
+		InterpSpeedDecreasing = ScaleBiasClamp.InterpSpeedDecreasing;
+	}
+
+	RIGVM_METHOD()
+	virtual void Execute(const FRigUnitContext& Context) override;
+
+	virtual FString ProcessPinLabelForInjection(const FString& InLabel) const override;
+
+	UPROPERTY(meta=(Input))
+	FQuat Value;
+
+	UPROPERTY(meta = (Input))
+	float Scale;
+
+	UPROPERTY(meta = (Input))
+	float Bias;
+
+	UPROPERTY(meta = (Input, Constant))
+	bool bMapRange;
+
+	UPROPERTY(meta = (Input, EditCondition = "bMapRange"))
+	FInputRange InRange;
+
+	UPROPERTY(meta = (Input, EditCondition = "bMapRange"))
+	FInputRange OutRange;
+
+	UPROPERTY(meta = (Input, Constant))
+	bool bClampResult;
+
+	UPROPERTY(meta = (Input, EditCondition = "bClampResult"))
+	float ClampMin;
+
+	UPROPERTY(meta = (Input, EditCondition = "bClampResult"))
+	float ClampMax;
+
+	UPROPERTY(meta = (Input, Constant))
+	bool bInterpResult;
+
+	UPROPERTY(meta = (Input, EditCondition = "bInterpResult"))
+	float InterpSpeedIncreasing;
+
+	UPROPERTY(meta = (Input, EditCondition = "bInterpResult"))
+	float InterpSpeedDecreasing;
+
+	UPROPERTY(meta=(Output))
+	FQuat Result;
+
+	UPROPERTY()
+	FInputScaleBiasClamp ScaleBiasClamp;
+};
diff --git a/Engine/Plugins/Animation/ControlRig/Source/ControlRig/Public/Units/Hierarchy/RigUnit_SetTransform.h b/Engine/Plugins/Animation/ControlRig/Source/ControlRig/Public/Units/Hierarchy/RigUnit_SetTransform.h
--- a/Engine/Plugins/Animation/ControlRig/Source/ControlRig/Public/Units/Hierarchy/RigUnit_SetTransform.h
+++ b/Engine/Plugins/Animation/ControlRig/Source/ControlRig/Public/Units/Hierarchy/RigUnit_SetTransform.h
@@ -12,7 +12,7 @@
  * actually sets the Control's offset transform and resets the local
  * values to (0, 0, 0).
  */
-USTRUCT(meta=(DisplayName="Set Transform", Category="Transforms", DocumentationPolicy = "Strict", Keywords="SetBoneTransform,SetControlTransform,SetInitialTransform,SetSpaceTransform", NodeColor="0, 0.364706, 1.0", Varying))
+USTRUCT(meta=(DisplayName="Set Transform", Category="Transforms", TemplateName = "Set Transform", DocumentationPolicy = "Strict", Keywords="SetBoneTransform,SetControlTransform,SetInitialTransform,SetSpaceTransform", NodeColor="0, 0.364706, 1.0", Varying))
 struct CONTROLRIG_API FRigUnit_SetTransform : public FRigUnitMutable
 {
 	GENERATED_BODY()
@@ -21,7 +21,7 @@ struct CONTROLRIG_API FRigUnit_SetTransform : public FRigUnitMutable
 		: Item(NAME_None, ERigElementType::Bone)
 		, Space(EBoneGetterSetterMode::GlobalSpace)
 		, bInitial(false)
-		, Transform(FTransform::Identity)
+		, Value(FTransform::Identity)
 		, Weight(1.f)
 		, bPropagateToChildren(true)
 		, CachedIndex()
@@ -66,7 +66,7 @@ struct CONTROLRIG_API FRigUnit_SetTransform : public FRigUnitMutable
 
 	// The new transform of the given item
 	UPROPERTY(meta=(Input))
-	FTransform Transform;
+	FTransform Value;
 
 	// Defines how much the change will be applied
 	UPROPERTY(meta = (Input, UIMin = "0.0", UIMax = "1.0"))
@@ -85,7 +85,7 @@ struct CONTROLRIG_API FRigUnit_SetTransform : public FRigUnitMutable
 /**
  * SetTranslation is used to set a single translation on hierarchy.
  */
-USTRUCT(meta=(DisplayName="Set Translation", Category="Transforms", DocumentationPolicy = "Strict", Keywords="SetBoneTranslation,SetControlTranslation,SetInitialTranslation,SetSpaceTranslation,SetBoneLocation,SetControlLocation,SetInitialLocation,SetSpaceLocation,SetBonePosition,SetControlPosition,SetInitialPosition,SetSpacePosition,SetTranslation,SetLocation,SetPosition", NodeColor="0, 0.364706, 1.0", Varying))
+USTRUCT(meta=(DisplayName="Set Translation", Category="Transforms", TemplateName = "Set Transform", DocumentationPolicy = "Strict", Keywords="SetBoneTranslation,SetControlTranslation,SetInitialTranslation,SetSpaceTranslation,SetBoneLocation,SetControlLocation,SetInitialLocation,SetSpaceLocation,SetBonePosition,SetControlPosition,SetInitialPosition,SetSpacePosition,SetTranslation,SetLocation,SetPosition", NodeColor="0, 0.364706, 1.0", Varying))
 struct CONTROLRIG_API FRigUnit_SetTranslation : public FRigUnitMutable
 {
 	GENERATED_BODY()
@@ -93,7 +93,8 @@ struct CONTROLRIG_API FRigUnit_SetTranslation : public FRigUnitMutable
 	FRigUnit_SetTranslation()
 		: Item(NAME_None, ERigElementType::Bone)
 		, Space(EBoneGetterSetterMode::GlobalSpace)
-		, Translation(FVector::ZeroVector)
+		, bInitial(false)
+		, Value(FVector::ZeroVector)
 		, Weight(1.f)
 		, bPropagateToChildren(true)
 		, CachedIndex()
@@ -128,9 +129,16 @@ struct CONTROLRIG_API FRigUnit_SetTranslation : public FRigUnitMutable
 	UPROPERTY(meta = (Input))
 	EBoneGetterSetterMode Space;
 
+	/**
+	 * Defines if the transform should be set as current (false) or initial (true).
+	 * Initial transforms for bones and other elements in the hierarchy represent the reference pose's value.
+	 */ 
+	UPROPERTY(meta = (Input))
+	bool bInitial;
+
 	// The new translation of the given item
 	UPROPERTY(meta=(Input))
-	FVector Translation;
+	FVector Value;
 
 	// Defines how much the change will be applied
 	UPROPERTY(meta = (Input, UIMin = "0.0", UIMax = "1.0"))
@@ -149,7 +157,7 @@ struct CONTROLRIG_API FRigUnit_SetTranslation : public FRigUnitMutable
 /**
  * SetRotation is used to set a single rotation on hierarchy.
  */
-USTRUCT(meta=(DisplayName="Set Rotation", Category="Transforms", DocumentationPolicy = "Strict", Keywords="SetBoneRotation,SetControlRotation,SetInitialRotation,SetSpaceRotation,SetBoneOrientation,SetControlOrientation,SetInitialOrientation,SetSpaceOrientation,SetRotation,SetOrientation", NodeColor="0, 0.364706, 1.0", Varying))
+USTRUCT(meta=(DisplayName="Set Rotation", Category="Transforms", TemplateName = "Set Transform", DocumentationPolicy = "Strict", Keywords="SetBoneRotation,SetControlRotation,SetInitialRotation,SetSpaceRotation,SetBoneOrientation,SetControlOrientation,SetInitialOrientation,SetSpaceOrientation,SetRotation,SetOrientation", NodeColor="0, 0.364706, 1.0", Varying))
 struct CONTROLRIG_API FRigUnit_SetRotation : public FRigUnitMutable
 {
 	GENERATED_BODY()
@@ -157,7 +165,8 @@ struct CONTROLRIG_API FRigUnit_SetRotation : public FRigUnitMutable
 	FRigUnit_SetRotation()
 		: Item(NAME_None, ERigElementType::Bone)
 		, Space(EBoneGetterSetterMode::GlobalSpace)
-		, Rotation(FQuat::Identity)
+		, bInitial(false)
+		, Value(FQuat::Identity)
 		, Weight(1.f)
 		, bPropagateToChildren(true)
 		, CachedIndex()
@@ -192,9 +201,16 @@ struct CONTROLRIG_API FRigUnit_SetRotation : public FRigUnitMutable
 	UPROPERTY(meta = (Input))
 	EBoneGetterSetterMode Space;
 
+	/**
+	 * Defines if the transform should be set as current (false) or initial (true).
+	 * Initial transforms for bones and other elements in the hierarchy represent the reference pose's value.
+	 */ 
+	UPROPERTY(meta = (Input))
+	bool bInitial;
+
 	// The new rotation of the given item
 	UPROPERTY(meta=(Input))
-	FQuat Rotation;
+	FQuat Value;
 
 	// Defines how much the change will be applied
 	UPROPERTY(meta = (Input, UIMin = "0.0", UIMax = "1.0"))
@@ -221,6 +237,7 @@ struct FRigUnit_SetScale : public FRigUnitMutable
 	FRigUnit_SetScale()
 		: Item(NAME_None, ERigElementType::Bone)
 		, Space(EBoneGetterSetterMode::GlobalSpace)
+		, bInitial(false)
 		, Scale(FVector::OneVector)
 		, Weight(1.f)
 		, bPropagateToChildren(true)
@@ -244,6 +261,13 @@ struct FRigUnit_SetScale : public FRigUnitMutable
 	UPROPERTY(meta = (Input))
 	EBoneGetterSetterMode Space;
 
+	/**
+	 * Defines if the transform should be set as current (false) or initial (true).
+	 * Initial transforms for bones and other elements in the hierarchy represent the reference pose's value.
+	 */ 
+	UPROPERTY(meta = (Input))
+	bool bInitial;
+
 	// The new scale of the given item
 	UPROPERTY(meta=(Input))
 	FVector Scale;
