ControlRig: Disable access to controller during the run of an event

diff --git a/Engine/Plugins/Animation/ControlRig/Source/ControlRig/Private/ControlRig.cpp b/Engine/Plugins/Animation/ControlRig/Source/ControlRig/Private/ControlRig.cpp
--- a/Engine/Plugins/Animation/ControlRig/Source/ControlRig/Private/ControlRig.cpp
+++ b/Engine/Plugins/Animation/ControlRig/Source/ControlRig/Private/ControlRig.cpp
@@ -663,7 +663,12 @@ void UControlRig::Execute(const EControlRigState InState, const FName& InEventNa
 	Context.FramesPerSecond = GetCurrentFramesPerSecond();
 	Context.bDuringInteraction = IsInteracting();
 	Context.State = InState;
+	
 	Context.Hierarchy = GetHierarchy();
+	check(Context.Hierarchy);
+
+	// for now we'll disable any controller access during the run of a control rig
+	FRigHierarchyEnableControllerBracket DisableHierarchyController(Context.Hierarchy, false);
 
 	Context.ToWorldSpaceTransform = FTransform::Identity;
 	Context.OwningComponent = nullptr;
diff --git a/Engine/Plugins/Animation/ControlRig/Source/ControlRig/Private/Rigs/RigHierarchy.cpp b/Engine/Plugins/Animation/ControlRig/Source/ControlRig/Private/Rigs/RigHierarchy.cpp
--- a/Engine/Plugins/Animation/ControlRig/Source/ControlRig/Private/Rigs/RigHierarchy.cpp
+++ b/Engine/Plugins/Animation/ControlRig/Source/ControlRig/Private/Rigs/RigHierarchy.cpp
@@ -102,6 +102,7 @@ URigHierarchy::URigHierarchy()
 , LastInteractedKey()
 , bSuspendNotifications(false)
 , HierarchyController(nullptr)
+, bIsControllerAvailable(true)
 , ResetPoseHash(INDEX_NONE)
 #if WITH_EDITOR
 , bPropagatingChange(false)
@@ -1586,6 +1587,8 @@ bool URigHierarchy::SwitchToParent(FRigElementKey InChild, FRigElementKey InPare
 bool URigHierarchy::SwitchToParent(FRigBaseElement* InChild, FRigBaseElement* InParent, bool bInitial,
 	bool bAffectChildren, const TElementDependencyMap& InDependencyMap, FString* OutFailureReason)
 {
+	FRigHierarchyEnableControllerBracket EnableController(this, true);
+	
 	if(InChild && InParent)
 	{
 		if(!CanSwitchToParent(InChild->GetKey(), InParent->GetKey(), InDependencyMap, OutFailureReason))
@@ -1655,6 +1658,8 @@ bool URigHierarchy::SwitchToWorldSpace(FRigBaseElement* InChild, bool bInitial,
 
 FRigElementKey URigHierarchy::GetOrAddWorldSpaceReference()
 {
+	FRigHierarchyEnableControllerBracket EnableController(this, true);
+
 	const FRigElementKey WorldSpaceReferenceKey = GetWorldSpaceReferenceKey();
 
 	FRigBaseElement* Parent = Find(WorldSpaceReferenceKey);
@@ -1944,8 +1949,17 @@ void URigHierarchy::SendAutoKeyEvent(FRigElementKey InElement, float InOffsetInS
 	SendEvent(Context, bAsynchronous);
 }
 
+bool URigHierarchy::IsControllerAvailable() const
+{
+	return bIsControllerAvailable;
+}
+
 URigHierarchyController* URigHierarchy::GetController(bool bCreateIfNeeded)
 {
+	if(!IsControllerAvailable())
+	{
+		return nullptr;
+	}
 	if(HierarchyController)
 	{
 		return HierarchyController;
diff --git a/Engine/Plugins/Animation/ControlRig/Source/ControlRig/Private/Units/Execution/RigUnit_DynamicHierarchy.cpp b/Engine/Plugins/Animation/ControlRig/Source/ControlRig/Private/Units/Execution/RigUnit_DynamicHierarchy.cpp
--- a/Engine/Plugins/Animation/ControlRig/Source/ControlRig/Private/Units/Execution/RigUnit_DynamicHierarchy.cpp
+++ b/Engine/Plugins/Animation/ControlRig/Source/ControlRig/Private/Units/Execution/RigUnit_DynamicHierarchy.cpp
@@ -26,6 +26,7 @@ FRigUnit_AddParent_Execute()
 		return;
 	}
 
+	FRigHierarchyEnableControllerBracket EnableController(ExecuteContext.Hierarchy, true);
 	if(URigHierarchyController* Controller = ExecuteContext.Hierarchy->GetController(true))
 	{
 		Controller->AddParent(ChildElement, ParentElement, 0.f, true, false);
diff --git a/Engine/Plugins/Animation/ControlRig/Source/ControlRig/Public/Rigs/RigHierarchy.h b/Engine/Plugins/Animation/ControlRig/Source/ControlRig/Public/Rigs/RigHierarchy.h
--- a/Engine/Plugins/Animation/ControlRig/Source/ControlRig/Public/Rigs/RigHierarchy.h
+++ b/Engine/Plugins/Animation/ControlRig/Source/ControlRig/Public/Rigs/RigHierarchy.h
@@ -1986,9 +1986,20 @@ public:
 	 * @return The delegate to listen to for events coming from this hierarchy
 	 */
 	FORCEINLINE FRigEventDelegate& OnEventReceived() { return EventDelegate; }
+	
+	/**
+	 * Returns true if the hierarchy controller is currently available
+	 * The controller may not be available during certain events.
+	 * If the controller is not available then GetController() will return nullptr.
+	 */ 
+	UFUNCTION(BlueprintPure, Category = URigHierarchy)
+	bool IsControllerAvailable() const;
 
 	/**
-	 * Returns a controller for this hierarchy
+	 * Returns a controller for this hierarchy.
+	 * Note: If the controller is not available this will return nullptr 
+	 * even if the bCreateIfNeeded flag is set to true. You can check the 
+	 * controller's availability with IsControllerAvailable().
 	 * @param bCreateIfNeeded Creates a controller if needed
 	 * @return The Controller for this hierarchy
 	 */
@@ -2698,6 +2709,7 @@ private:
 
 	UPROPERTY(Transient)
 	TObjectPtr<URigHierarchyController> HierarchyController;
+	bool bIsControllerAvailable;
 
 	TMap<FRigElementKey, FRigElementKey> PreviousParentMap;
 
@@ -2982,6 +2994,22 @@ private:
 	friend struct FRigHierarchyValidityBracket;
 	friend struct FRigHierarchyGlobalValidityBracket;
 	friend struct FControlRigVisualGraphUtils;
+	friend struct FRigHierarchyEnableControllerBracket;
+};
+
+struct CONTROLRIG_API FRigHierarchyEnableControllerBracket : public TGuardValue<bool>
+{
+private:
+	FRigHierarchyEnableControllerBracket(URigHierarchy* InHierarchy, bool bEnable)
+		: TGuardValue<bool>(InHierarchy->bIsControllerAvailable, bEnable)
+	{
+	}
+
+	friend class URigHierarchy;
+	friend class UControlRig;
+
+	// certain units are allowed to use this
+	friend struct FRigUnit_AddParent;
 };
 
 struct CONTROLRIG_API FRigHierarchyValidityBracket
