Control Rig: Remove dependency on VM within dynamic hierarchy nodes

diff --git a/Engine/Plugins/Animation/ControlRig/Source/ControlRig/Private/Rigs/RigHierarchy.cpp b/Engine/Plugins/Animation/ControlRig/Source/ControlRig/Private/Rigs/RigHierarchy.cpp
--- a/Engine/Plugins/Animation/ControlRig/Source/ControlRig/Private/Rigs/RigHierarchy.cpp
+++ b/Engine/Plugins/Animation/ControlRig/Source/ControlRig/Private/Rigs/RigHierarchy.cpp
@@ -1588,10 +1588,25 @@ bool URigHierarchy::SwitchToParent(FRigBaseElement* InChild, FRigBaseElement* In
 	bool bAffectChildren, const TElementDependencyMap& InDependencyMap, FString* OutFailureReason)
 {
 	FRigHierarchyEnableControllerBracket EnableController(this, true);
+
+	// rely on the VM's dependency map if there's currently an available context.
+	const TElementDependencyMap* DependencyMapPtr = &InDependencyMap;
+
+#if WITH_EDITOR
+	TElementDependencyMap DependencyMapFromVM;
+	if(ExecuteContext != nullptr && DependencyMapPtr->IsEmpty())
+	{
+		if(ExecuteContext->VM)
+		{
+			DependencyMapFromVM = GetDependenciesForVM(ExecuteContext->VM);
+			DependencyMapPtr = &DependencyMapFromVM;
+		}
+	}
+#endif
 	
 	if(InChild && InParent)
 	{
-		if(!CanSwitchToParent(InChild->GetKey(), InParent->GetKey(), InDependencyMap, OutFailureReason))
+		if(!CanSwitchToParent(InChild->GetKey(), InParent->GetKey(), *DependencyMapPtr, OutFailureReason))
 		{
 			return false;
 		}
diff --git a/Engine/Plugins/Animation/ControlRig/Source/ControlRig/Private/Units/Execution/RigUnit_DynamicHierarchy.cpp b/Engine/Plugins/Animation/ControlRig/Source/ControlRig/Private/Units/Execution/RigUnit_DynamicHierarchy.cpp
--- a/Engine/Plugins/Animation/ControlRig/Source/ControlRig/Private/Units/Execution/RigUnit_DynamicHierarchy.cpp
+++ b/Engine/Plugins/Animation/ControlRig/Source/ControlRig/Private/Units/Execution/RigUnit_DynamicHierarchy.cpp
@@ -97,14 +97,8 @@ FRigUnit_SwitchParent_Execute()
 		default:
 		{
 			FString FailureReason;
-			URigHierarchy::TElementDependencyMap DependencyMap;
-#if WITH_EDITOR
-			if(ExecuteContext.VM)
-			{
-				DependencyMap = ExecuteContext.Hierarchy->GetDependenciesForVM(ExecuteContext.VM);
-			}
-#endif
-			if(!ExecuteContext.Hierarchy->SwitchToParent(ChildElement, ParentElement, false, true, DependencyMap, &FailureReason))
+			static const URigHierarchy::TElementDependencyMap EmptyDependencyMap;
+			if(!ExecuteContext.Hierarchy->SwitchToParent(ChildElement, ParentElement, false, true, EmptyDependencyMap, &FailureReason))
 			{
 				if(!FailureReason.IsEmpty())
 				{
