Fix build.

#rnx

diff --git a/Engine/Plugins/Animation/ControlRig/Source/ControlRig/Private/Units/Execution/RigUnit_DynamicHierarchy.cpp b/Engine/Plugins/Animation/ControlRig/Source/ControlRig/Private/Units/Execution/RigUnit_DynamicHierarchy.cpp
--- a/Engine/Plugins/Animation/ControlRig/Source/ControlRig/Private/Units/Execution/RigUnit_DynamicHierarchy.cpp
+++ b/Engine/Plugins/Animation/ControlRig/Source/ControlRig/Private/Units/Execution/RigUnit_DynamicHierarchy.cpp
@@ -108,9 +108,9 @@ FRigUnit_SwitchParent_Execute()
 			}
 
 			// during setup event also change the initial weights
-			if(ExecuteContext.EventName == FRigUnit_PrepareForExecution::EventName)
+			if(ExecuteContext.GetEventName() == FRigUnit_PrepareForExecution::EventName)
 			{
-				if(!ExecuteContext.Hierarchy->SwitchToParent(ChildElement, ParentElement, true, true, DependencyMap, &FailureReason))
+				if(!ExecuteContext.Hierarchy->SwitchToParent(ChildElement, ParentElement, true, true, EmptyDependencyMap, &FailureReason))
 				{
 					if(!FailureReason.IsEmpty())
 					{
@@ -182,7 +182,7 @@ FRigUnit_HierarchySetParentWeights_Execute()
 	ExecuteContext.Hierarchy->SetParentWeightArray(ChildElement, Weights, false, true);
 
 	// during setup event also change the initial weights
-	if(ExecuteContext.EventName == FRigUnit_PrepareForExecution::EventName)
+	if(ExecuteContext.GetEventName() == FRigUnit_PrepareForExecution::EventName)
 	{
 		ExecuteContext.Hierarchy->SetParentWeightArray(ChildElement, Weights, true, true);
 	}
