Control Rig: Improved version of FQuat::FindBetween

diff --git a/Engine/Plugins/Animation/ControlRig/Source/ControlRig/Private/Math/ControlRigMathLibrary.cpp b/Engine/Plugins/Animation/ControlRig/Source/ControlRig/Private/Math/ControlRigMathLibrary.cpp
--- a/Engine/Plugins/Animation/ControlRig/Source/ControlRig/Private/Math/ControlRigMathLibrary.cpp
+++ b/Engine/Plugins/Animation/ControlRig/Source/ControlRig/Private/Math/ControlRigMathLibrary.cpp
@@ -374,3 +374,44 @@ FVector FControlRigMathLibrary::ClampSpatially(const FVector& Value, EAxis::Type
 
 	return Space.TransformPosition(Local);
 }
+
+FQuat FControlRigMathLibrary::FindQuatBetweenVectors(const FVector& A, const FVector& B)
+{
+	return FindQuatBetweenNormals(A.GetSafeNormal(), B.GetSafeNormal());
+}
+
+FQuat FControlRigMathLibrary::FindQuatBetweenNormals(const FVector& A, const FVector& B)
+{
+	const FQuat::FReal Dot = FVector::DotProduct(A, B);
+	FQuat::FReal W = 1 + Dot;
+	FQuat Result;
+
+	if (W < SMALL_NUMBER)
+	{
+		// A and B point in opposite directions
+		W = 2 - W;
+		Result = FQuat(
+			-A.Y * B.Z + A.Z * B.Y,
+			-A.Z * B.X + A.X * B.Z,
+			-A.X * B.Y + A.Y * B.X,
+			W).GetNormalized();
+
+		const FVector Normal = FMath::Abs(A.X) > FMath::Abs(A.Y) ?
+			FVector(0.f, 1.f, 0.f) : FVector(1.f, 0.f, 0.f);
+		const FVector BiNormal = FVector::CrossProduct(A, Normal);
+		const FVector TauNormal = FVector::CrossProduct(A, BiNormal);
+		Result = Result * FQuat(TauNormal, PI);
+	}
+	else
+	{
+		//Axis = FVector::CrossProduct(A, B);
+		Result = FQuat(
+			A.Y * B.Z - A.Z * B.Y,
+			A.Z * B.X - A.X * B.Z,
+			A.X * B.Y - A.Y * B.X,
+			W);
+	}
+
+	Result.Normalize();
+	return Result;
+}
diff --git a/Engine/Plugins/Animation/ControlRig/Source/ControlRig/Private/Units/Highlevel/Harmonics/RigUnit_ChainHarmonics.cpp b/Engine/Plugins/Animation/ControlRig/Source/ControlRig/Private/Units/Highlevel/Harmonics/RigUnit_ChainHarmonics.cpp
--- a/Engine/Plugins/Animation/ControlRig/Source/ControlRig/Private/Units/Highlevel/Harmonics/RigUnit_ChainHarmonics.cpp
+++ b/Engine/Plugins/Animation/ControlRig/Source/ControlRig/Private/Units/Highlevel/Harmonics/RigUnit_ChainHarmonics.cpp
@@ -2,6 +2,7 @@
 
 #include "Units/Highlevel/Harmonics/RigUnit_ChainHarmonics.h"
 #include "Units/RigUnitContext.h"
+#include "Math/ControlRigMathLibrary.h"
 
 FRigUnit_ChainHarmonics_Execute()
 {
@@ -139,7 +140,7 @@ FRigUnit_ChainHarmonicsPerItem_Execute()
 			FVector ReachDirection = (Reach.ReachTarget - GlobalTransform.GetLocation()).GetSafeNormal();
 			ReachDirection = FMath::Lerp<FVector>(Axis, ReachDirection, Ease);
 
-			FQuat ReachRotation = FQuat::FindBetween(Axis, ReachDirection);
+			FQuat ReachRotation = FControlRigMathLibrary::FindQuatBetweenVectors(Axis, ReachDirection);
 			Rotation = (ReachRotation * Rotation).GetNormalized();
 		}
 
@@ -201,14 +202,14 @@ FRigUnit_ChainHarmonicsPerItem_Execute()
 			VelocityLines[Index * 2 + 0] = PendulumPosition[Index];
 			VelocityLines[Index * 2 + 1] = PendulumPosition[Index] + PendulumVelocity[Index] * 0.1f;
 
-			FQuat PendulumRotation = FQuat::FindBetween(Rotation.RotateVector(LocalTip[Index]), PendulumPosition[Index] - GlobalTransform.GetLocation());
+			FQuat PendulumRotation = FControlRigMathLibrary::FindQuatBetweenVectors(Rotation.RotateVector(LocalTip[Index]), PendulumPosition[Index] - GlobalTransform.GetLocation());
 			Rotation = (PendulumRotation * Rotation).GetNormalized();
 
 			float Unwind = FMath::Lerp<float>(Pendulum.UnwindMinimum, Pendulum.UnwindMaximum, Ratio[Index]);
 			FVector CurrentUpvector = Rotation.RotateVector(Pendulum.UnwindAxis);
 			CurrentUpvector = CurrentUpvector - FVector::DotProduct(CurrentUpvector, Rotation.RotateVector(LocalTip[Index]).GetSafeNormal());
 			CurrentUpvector = FMath::Lerp<FVector>(Upvector, CurrentUpvector, Unwind);
-			FQuat UnwindRotation = FQuat::FindBetween(CurrentUpvector, Upvector);
+			FQuat UnwindRotation = FControlRigMathLibrary::FindQuatBetweenVectors(CurrentUpvector, Upvector);
 			Rotation = (UnwindRotation * Rotation).GetNormalized();
 
 			Rotation = FQuat::Slerp(NonSimulatedRotation, Rotation, FMath::Clamp<float>(Ease, 0.f, 1.f));
diff --git a/Engine/Plugins/Animation/ControlRig/Source/ControlRig/Private/Units/Highlevel/Hierarchy/RigUnit_AimBone.cpp b/Engine/Plugins/Animation/ControlRig/Source/ControlRig/Private/Units/Highlevel/Hierarchy/RigUnit_AimBone.cpp
--- a/Engine/Plugins/Animation/ControlRig/Source/ControlRig/Private/Units/Highlevel/Hierarchy/RigUnit_AimBone.cpp
+++ b/Engine/Plugins/Animation/ControlRig/Source/ControlRig/Private/Units/Highlevel/Hierarchy/RigUnit_AimBone.cpp
@@ -2,6 +2,7 @@
 
 #include "Units/Highlevel/Hierarchy/RigUnit_AimBone.h"
 #include "Units/RigUnitContext.h"
+#include "Math/ControlRigMathLibrary.h"
 
 FRigUnit_AimBoneMath_Execute()
 {
@@ -26,7 +27,7 @@ FRigUnit_AimBoneMath_Execute()
 	{
 		return;
 	}
-
+	
 	if (Primary.Weight > SMALL_NUMBER)
 	{
 		FVector Target = Primary.Target;
@@ -72,7 +73,7 @@ FRigUnit_AimBoneMath_Execute()
 			{
 				Target = FMath::Lerp<FVector>(Axis, Target, T).GetSafeNormal();
 			}
-			FQuat Rotation = FQuat::FindBetweenNormals(Axis, Target);
+			FQuat Rotation = FControlRigMathLibrary::FindQuatBetweenVectors(Axis, Target);
 			Result.SetRotation((Rotation * Result.GetRotation()).GetNormalized());
 		}
 		else
@@ -133,7 +134,7 @@ FRigUnit_AimBoneMath_Execute()
 			{
 				Target = FMath::Lerp<FVector>(Axis, Target, T).GetSafeNormal();
 			}
-			FQuat Rotation = FQuat::FindBetweenNormals(Axis, Target);
+			FQuat Rotation = FControlRigMathLibrary::FindQuatBetweenVectors(Axis, Target);
 			Result.SetRotation((Rotation * Result.GetRotation()).GetNormalized());
 		}
 		else
diff --git a/Engine/Plugins/Animation/ControlRig/Source/ControlRig/Private/Units/Math/RigUnit_MathQuaternion.cpp b/Engine/Plugins/Animation/ControlRig/Source/ControlRig/Private/Units/Math/RigUnit_MathQuaternion.cpp
--- a/Engine/Plugins/Animation/ControlRig/Source/ControlRig/Private/Units/Math/RigUnit_MathQuaternion.cpp
+++ b/Engine/Plugins/Animation/ControlRig/Source/ControlRig/Private/Units/Math/RigUnit_MathQuaternion.cpp
@@ -4,6 +4,7 @@
 #include "Units/Math/RigUnit_MathTransform.h"
 #include "Units/RigUnitContext.h"
 #include "AnimationCoreLibrary.h"
+#include "Math/ControlRigMathLibrary.h"
 
 FRigUnit_MathQuaternionFromAxisAndAngle_Execute()
 {
@@ -53,7 +54,7 @@ FRigUnit_MathQuaternionFromTwoVectors_Execute()
 		Result = FQuat::Identity;
 		return;
 	}
-	Result = FQuat::FindBetweenVectors(A, B).GetNormalized();
+	Result = FControlRigMathLibrary::FindQuatBetweenVectors(A, B);
 }
 
 FRigUnit_MathQuaternionToAxisAndAngle_Execute()
diff --git a/Engine/Plugins/Animation/ControlRig/Source/ControlRig/Private/Units/Math/RigUnit_MathVector.cpp b/Engine/Plugins/Animation/ControlRig/Source/ControlRig/Private/Units/Math/RigUnit_MathVector.cpp
--- a/Engine/Plugins/Animation/ControlRig/Source/ControlRig/Private/Units/Math/RigUnit_MathVector.cpp
+++ b/Engine/Plugins/Animation/ControlRig/Source/ControlRig/Private/Units/Math/RigUnit_MathVector.cpp
@@ -341,7 +341,7 @@ FRigUnit_MathVectorAngle_Execute()
 		Result = 0.f;
 		return;
 	}
-	Result = FQuat::FindBetween(A, B).GetAngle();
+	Result = FControlRigMathLibrary::FindQuatBetweenVectors(A, B).GetAngle();
 }
 
 FRigUnit_MathVectorParallel_Execute()
diff --git a/Engine/Plugins/Animation/ControlRig/Source/ControlRig/Public/Math/ControlRigMathLibrary.h b/Engine/Plugins/Animation/ControlRig/Source/ControlRig/Public/Math/ControlRigMathLibrary.h
--- a/Engine/Plugins/Animation/ControlRig/Source/ControlRig/Public/Math/ControlRigMathLibrary.h
+++ b/Engine/Plugins/Animation/ControlRig/Source/ControlRig/Public/Math/ControlRigMathLibrary.h
@@ -75,4 +75,6 @@ public:
 	static FTransform LerpTransform(const FTransform& A, const FTransform& B, float T);
 	static void SolveBasicTwoBoneIK(FTransform& BoneA, FTransform& BoneB, FTransform& Effector, const FVector& PoleVector, const FVector& PrimaryAxis, const FVector& SecondaryAxis, float SecondaryAxisWeight, float BoneALength, float BoneBLength, bool bEnableStretch, float StretchStartRatio, float StretchMaxRatio);
 	static FVector ClampSpatially(const FVector& Value, EAxis::Type Axis, EControlRigClampSpatialMode::Type Type, float Minimum, float Maximum, FTransform Space);
+	static FQuat FindQuatBetweenVectors(const FVector& A, const FVector& B);
+	static FQuat FindQuatBetweenNormals(const FVector& A, const FVector& B);
 };
\ No newline at end of file
