avoid using NodeToJumpTo if nullptr

diff --git a/Engine/Plugins/Animation/ControlRig/Source/ControlRigEditor/Private/Editor/ControlRigEditor.cpp b/Engine/Plugins/Animation/ControlRig/Source/ControlRigEditor/Private/Editor/ControlRigEditor.cpp
--- a/Engine/Plugins/Animation/ControlRig/Source/ControlRigEditor/Private/Editor/ControlRigEditor.cpp
+++ b/Engine/Plugins/Animation/ControlRig/Source/ControlRigEditor/Private/Editor/ControlRigEditor.cpp
@@ -4707,26 +4707,25 @@ void FControlRigEditor::HandleJumpToHyperlink(const UObject* InSubject)
 		}
 	}
 
-	if(GraphToJumpTo)
+	if (GraphToJumpTo && NodeToJumpTo)
 	{
 		if(UControlRigGraph* EdGraph = Cast<UControlRigGraph>(RigBlueprint->GetEdGraph(NodeToJumpTo->GetGraph())))
 		{
-			if(NodeToJumpTo)
+			if(const UControlRigGraphNode* EdGraphNode = Cast<UControlRigGraphNode>(EdGraph->FindNodeForModelNodeName(NodeToJumpTo->GetFName())))
 			{
-				if(const UControlRigGraphNode* EdGraphNode = Cast<UControlRigGraphNode>(EdGraph->FindNodeForModelNodeName(NodeToJumpTo->GetFName())))
+				if(PinToJumpTo)
 				{
-					if(PinToJumpTo)
+					if(const UEdGraphPin* EdGraphPin = EdGraphNode->FindPin(PinToJumpTo->GetSegmentPath(true)))
 					{
-						if(const UEdGraphPin* EdGraphPin = EdGraphNode->FindPin(PinToJumpTo->GetSegmentPath(true)))
-						{
-							JumpToPin(EdGraphPin);
-							return;
-						}
+						JumpToPin(EdGraphPin);
+						return;
 					}
-					JumpToNode(EdGraphNode);
-					return;
 				}
+				
+				JumpToNode(EdGraphNode);
+				return;
 			}
+			
 			JumpToHyperlink(EdGraph);
 		}
 	}
