Sequencer: Control Rig:  Actor snapping got Yaw/Roll, Z/X swapped, so rotations are busted.

diff --git a/Engine/Plugins/Animation/ControlRig/Source/ControlRigEditor/Private/Tools/ControlRigSnapper.cpp b/Engine/Plugins/Animation/ControlRig/Source/ControlRigEditor/Private/Tools/ControlRigSnapper.cpp
--- a/Engine/Plugins/Animation/ControlRig/Source/ControlRigEditor/Private/Tools/ControlRigSnapper.cpp
+++ b/Engine/Plugins/Animation/ControlRig/Source/ControlRigEditor/Private/Tools/ControlRigSnapper.cpp
@@ -270,7 +270,7 @@ struct FGuidAndActor
 				{
 					if (!Channels[3]->GetDefault().IsSet())
 					{
-						Channels[3]->SetDefault(Rotation.Yaw);
+						Channels[3]->SetDefault(Rotation.Roll);
 					}
 					if (!Channels[4]->GetDefault().IsSet())
 					{
@@ -278,7 +278,7 @@ struct FGuidAndActor
 					}
 					if (!Channels[5]->GetDefault().IsSet())
 					{
-						Channels[5]->SetDefault(Rotation.Roll);
+						Channels[5]->SetDefault(Rotation.Yaw);
 					}
 				}
 				if (SnapSettings->bSnapScale)
@@ -311,11 +311,11 @@ struct FGuidAndActor
 			{
 				//todo winding
 				TMovieSceneChannelData<FMovieSceneDoubleValue> ChannelData = Channels[3]->GetData();
-				MovieSceneToolHelpers::SetOrAddKey(ChannelData, Frame, Rotation.Yaw);
+				MovieSceneToolHelpers::SetOrAddKey(ChannelData, Frame, Rotation.Roll);
 				ChannelData = Channels[4]->GetData();
 				MovieSceneToolHelpers::SetOrAddKey(ChannelData, Frame, Rotation.Pitch);
 				ChannelData = Channels[5]->GetData();
-				MovieSceneToolHelpers::SetOrAddKey(ChannelData, Frame, Rotation.Roll);
+				MovieSceneToolHelpers::SetOrAddKey(ChannelData, Frame, Rotation.Yaw);
 			}
 			if (SnapSettings->bSnapScale)
 			{
