Control Rig: Read the edited default value from the right instance

diff --git a/Engine/Plugins/Animation/ControlRig/Source/ControlRigEditor/Private/Editor/ControlRigEditor.cpp b/Engine/Plugins/Animation/ControlRig/Source/ControlRigEditor/Private/Editor/ControlRigEditor.cpp
--- a/Engine/Plugins/Animation/ControlRig/Source/ControlRigEditor/Private/Editor/ControlRigEditor.cpp
+++ b/Engine/Plugins/Animation/ControlRig/Source/ControlRigEditor/Private/Editor/ControlRigEditor.cpp
@@ -4351,12 +4351,18 @@ void FControlRigEditor::NotifyPostChange(const FPropertyChangedEvent& PropertyCh
 	UControlRigBlueprint* ControlRigBP = GetControlRigBlueprint();
 	if(ControlRig && ControlRigBP)
 	{
+		bool bUseCDO = false; 
+		if(PropertyChangedEvent.GetNumObjectsBeingEdited() == 1)
+		{
+			bUseCDO = PropertyChangedEvent.GetObjectBeingEdited(0)->HasAnyFlags(RF_ClassDefaultObject);
+		}
+		
 		const FName VarName = PropertyChangedEvent.MemberProperty->GetFName();
 		for(FBPVariableDescription& NewVariable : ControlRigBP->NewVariables)
 		{
 			if(NewVariable.VarName == VarName)
 			{
-				UpdateDefaultValueForVariable(NewVariable, false);
+				UpdateDefaultValueForVariable(NewVariable, bUseCDO);
 				break;
 			}
 		}
