Control Rig: Don't auto upgrade injected nodes to template nodes

diff --git a/Engine/Plugins/Animation/ControlRig/Source/ControlRigDeveloper/Private/ControlRigBlueprint.cpp b/Engine/Plugins/Animation/ControlRig/Source/ControlRigDeveloper/Private/ControlRigBlueprint.cpp
--- a/Engine/Plugins/Animation/ControlRig/Source/ControlRigDeveloper/Private/ControlRigBlueprint.cpp
+++ b/Engine/Plugins/Animation/ControlRig/Source/ControlRigDeveloper/Private/ControlRigBlueprint.cpp
@@ -3849,6 +3849,10 @@ void UControlRigBlueprint::ConvertUnitNodesToTemplateNodes()
 		TArray<URigVMNode*> Nodes = Graph->GetNodes();
 		for (URigVMNode* Node : Nodes)
 		{
+			if(Node->IsInjected())
+			{
+				continue;
+			}
 			Controller->ReplaceUnitNodeWithTemplateNode(Node->GetFName(), false);
 		}
 	}
