Control Rig: Avoid dual compiles and refreshing graph editor unnecessarily

Performance numbers in Jira

diff --git a/Engine/Plugins/Animation/ControlRig/Source/ControlRigEditor/Private/Editor/ControlRigEditor.cpp b/Engine/Plugins/Animation/ControlRig/Source/ControlRigEditor/Private/Editor/ControlRigEditor.cpp
--- a/Engine/Plugins/Animation/ControlRig/Source/ControlRigEditor/Private/Editor/ControlRigEditor.cpp
+++ b/Engine/Plugins/Animation/ControlRig/Source/ControlRigEditor/Private/Editor/ControlRigEditor.cpp
@@ -209,6 +209,7 @@ FControlRigEditor::FControlRigEditor()
 	, bIsSettingObjectBeingDebugged(false)
 	, bExecutionControlRig(true)
 	, bSetupModeEnabled(false)
+	, bIsCompilingThroughUI(false)
 	, bAnyErrorsLeft(false)
 	, LastEventQueue(EControlRigEditorEventQueue::Setup)
 	, ExecutionMode(EControlRigExecutionModeType::EControlRigExecutionModeType_Release)
@@ -1932,6 +1933,7 @@ void FControlRigEditor::Compile()
 		ResetAllBoneModification(); 
 		
 		{
+			TGuardValue<bool> GuardCompileReEntry(bIsCompilingThroughUI, true);
 			FBlueprintEditor::Compile();
 		}
 
@@ -3281,6 +3283,17 @@ void FControlRigEditor::OnBlueprintChangedImpl(UBlueprint* InBlueprint, bool bIs
 	}
 }
 
+void FControlRigEditor::RefreshEditors(ERefreshBlueprintEditorReason::Type Reason)
+{
+	if(Reason == ERefreshBlueprintEditorReason::UnknownReason)
+	{
+		// we mark the reason as just compiled since we don't want to
+		// update the graph(s) all the time during compilation
+		Reason = ERefreshBlueprintEditorReason::BlueprintCompiled;
+	}
+	IControlRigEditor::RefreshEditors(Reason);
+}
+
 void FControlRigEditor::HandleViewportCreated(const TSharedRef<class IPersonaViewport>& InViewport)
 {
 	DECLARE_SCOPE_HIERARCHICAL_COUNTER_FUNC()
@@ -4021,9 +4034,12 @@ void FControlRigEditor::UpdateControlRig()
 				EditMode->SetObjects(ControlRig, EditorSkelComp,nullptr);
 			}
 
-			Blueprint->SetFlags(RF_Transient);
-			Blueprint->RecompileVM();
-			Blueprint->ClearFlags(RF_Transient);
+			if(!bIsCompilingThroughUI)
+			{
+				Blueprint->SetFlags(RF_Transient);
+				Blueprint->RecompileVM();
+				Blueprint->ClearFlags(RF_Transient);
+			}
 
 			ControlRig->OnInitialized_AnyThread().AddSP(this, &FControlRigEditor::HandleControlRigExecutedEvent);
 			ControlRig->OnExecuted_AnyThread().AddSP(this, &FControlRigEditor::HandleControlRigExecutedEvent);
diff --git a/Engine/Plugins/Animation/ControlRig/Source/ControlRigEditor/Private/Editor/ControlRigEditor.h b/Engine/Plugins/Animation/ControlRig/Source/ControlRigEditor/Private/Editor/ControlRigEditor.h
--- a/Engine/Plugins/Animation/ControlRig/Source/ControlRigEditor/Private/Editor/ControlRigEditor.h
+++ b/Engine/Plugins/Animation/ControlRig/Source/ControlRigEditor/Private/Editor/ControlRigEditor.h
@@ -247,6 +247,7 @@ protected:
 	virtual void OnActiveTabChanged( TSharedPtr<SDockTab> PreviouslyActive, TSharedPtr<SDockTab> NewlyActivated ) override;
 	virtual void OnSelectedNodesChangedImpl(const TSet<class UObject*>& NewSelection) override;
 	virtual void OnBlueprintChangedImpl(UBlueprint* InBlueprint, bool bIsJustBeingCompiled) override;
+	virtual void RefreshEditors(ERefreshBlueprintEditorReason::Type Reason = ERefreshBlueprintEditorReason::UnknownReason) override;
 	virtual void SetupGraphEditorEvents(UEdGraph* InGraph, SGraphEditor::FGraphEditorEvents& InEvents) override;
 	virtual void FocusInspectorOnGraphSelection(const TSet<class UObject*>& NewSelection, bool bForceRefresh = false) override;
 
@@ -469,6 +470,9 @@ protected:
 	/** Are we currently in setup mode */
 	bool bSetupModeEnabled;
 
+	/** Are we currently compiling through the user interface */
+	bool bIsCompilingThroughUI;
+
 	FPreviewControlRigUpdated PreviewControlRigUpdated;
 
 	TSharedPtr<SControlRigGraphPinNameListValueWidget> PinControlNameList;
