Control Rig: Guard against notification storms for the hierarchy

* perform lazy updates for the edit mode (avoid recreation of all actors all the time)
* avoid duplicate nofications - perform filtering on queued notifs
* bracket notifications with open and close brackets to give clients a chance to optimize

diff --git a/Engine/Plugins/Animation/ControlRig/Source/ControlRig/Private/ControlRig.cpp b/Engine/Plugins/Animation/ControlRig/Source/ControlRig/Private/ControlRig.cpp
--- a/Engine/Plugins/Animation/ControlRig/Source/ControlRig/Private/ControlRig.cpp
+++ b/Engine/Plugins/Animation/ControlRig/Source/ControlRig/Private/ControlRig.cpp
@@ -1230,14 +1230,14 @@ void UControlRig::ExecuteUnits(FRigUnitContext& InOutContext, const FName& InEve
 			}
 #endif
 
-#if WITH_EDITOR
 			URigHierarchy* Hierarchy = GetHierarchy();
+#if WITH_EDITOR
 			Hierarchy->ReadTransformsPerInstructionPerSlice.Reset();
 			Hierarchy->WrittenTransformsPerInstructionPerSlice.Reset();
 			Hierarchy->ReadTransformsPerInstructionPerSlice.AddZeroed(VM->GetByteCode().GetNumInstructions());
 			Hierarchy->WrittenTransformsPerInstructionPerSlice.AddZeroed(VM->GetByteCode().GetNumInstructions());
-			TGuardValue<const FRigVMExtendedExecuteContext*> HierarchyContextGuard(Hierarchy->ExecuteContext, &VM->GetContext());
 #endif
+			FRigHierarchyExecuteContextBracket HierarchyContextGuard(Hierarchy, &VM->GetContext());
 
 			VM->Execute(LocalMemory, AdditionalArguments, InEventName);
 		}
diff --git a/Engine/Plugins/Animation/ControlRig/Source/ControlRig/Private/ControlRigGizmoActor.cpp b/Engine/Plugins/Animation/ControlRig/Source/ControlRig/Private/ControlRigGizmoActor.cpp
--- a/Engine/Plugins/Animation/ControlRig/Source/ControlRig/Private/ControlRigGizmoActor.cpp
+++ b/Engine/Plugins/Animation/ControlRig/Source/ControlRig/Private/ControlRigGizmoActor.cpp
@@ -1,6 +1,7 @@
 // Copyright Epic Games, Inc. All Rights Reserved.
 
 #include "ControlRigGizmoActor.h"
+#include "ControlRig.h"
 #include "Engine/World.h"
 #include "Engine/StaticMesh.h"
 #include "Components/StaticMeshComponent.h"
@@ -11,6 +12,7 @@ AControlRigShapeActor::AControlRigShapeActor(const FObjectInitializer& ObjectIni
 	: Super(ObjectInitializer)
 	, ControlRigIndex(INDEX_NONE)
 	, ControlName(NAME_None)
+	, ShapeName(NAME_None)
 	, bEnabled(true)
 	, bSelected(false)
 	, bSelectable(true)
@@ -108,6 +110,86 @@ void AControlRigShapeActor::SetShapeColor(const FLinearColor& InColor)
 	}
 }
 
+bool AControlRigShapeActor::UpdateControlSettings(
+	ERigHierarchyNotification InNotif,
+	UControlRig* InControlRig,
+	const FRigControlElement* InControlElement,
+	bool bHideManipulators,
+	bool bIsInLevelEditor)
+{
+	check(InControlElement);
+
+	const FRigControlSettings& ControlSettings = InControlElement->Settings;
+	
+	// if this actor is not supposed to exist
+	if(!ControlSettings.bShapeEnabled)
+	{
+		return false;
+	}
+
+	const bool bShapeNameUpdated = ShapeName != ControlSettings.ShapeName;
+	bool bShapeTransformChanged = InNotif == ERigHierarchyNotification::ControlShapeTransformChanged;
+	const bool bLookupShape = bShapeNameUpdated || bShapeTransformChanged;
+	
+	FTransform MeshTransform = FTransform::Identity;
+
+	// update the shape used for the control
+	if(bLookupShape)
+	{
+		const TArray<TSoftObjectPtr<UControlRigShapeLibrary>> ShapeLibraries = InControlRig->GetShapeLibraries();
+		if (const FControlRigShapeDefinition* ShapeDef = UControlRigShapeLibrary::GetShapeByName(ControlSettings.ShapeName, ShapeLibraries))
+		{
+			MeshTransform = ShapeDef->Transform;
+
+			if(bShapeNameUpdated)
+			{
+				if(ShapeDef->StaticMesh.IsValid())
+				{
+					if(UStaticMesh* StaticMesh = ShapeDef->StaticMesh.Get())
+					{
+						if(StaticMesh != StaticMeshComponent->GetStaticMesh())
+						{
+							StaticMeshComponent->SetStaticMesh(StaticMesh);
+							bShapeTransformChanged = true;
+						}
+					}
+					else
+					{
+						return false;
+					}
+				}
+				else
+				{
+					return false;
+				}
+			}
+		}
+	}
+
+	// update the shape transform
+	if(bShapeTransformChanged)
+	{
+		const FTransform ShapeTransform = InControlElement->Shape.Get(ERigTransformType::InitialLocal);
+		StaticMeshComponent->SetRelativeTransform(MeshTransform * ShapeTransform);
+	}
+	
+	// update the shape color
+	SetShapeColor(ControlSettings.ShapeColor);
+
+#if WITH_EDITOR
+	// update visibility
+	SetIsTemporarilyHiddenInEditor(!ControlSettings.bShapeVisible || bHideManipulators);
+#endif
+	
+	// update selectability state
+	if (!bIsInLevelEditor) //don't change this in level editor otherwise we can never select it
+	{
+		SetSelectable(ControlSettings.bShapeVisible && !bHideManipulators && ControlSettings.bAnimatable);
+	}
+
+	return true;
+}
+
 // FControlRigShapeHelper START
 
 namespace FControlRigShapeHelper
@@ -165,6 +247,7 @@ namespace FControlRigShapeHelper
 		{
 			ShapeActor->ControlRigIndex = CreationParam.ControlRigIndex;
 			ShapeActor->ControlName = CreationParam.ControlName;
+			ShapeActor->ShapeName = CreationParam.ShapeName;
 			ShapeActor->SetSelectable(CreationParam.bSelectable);
 			ShapeActor->SetActorTransform(CreationParam.SpawnTransform);
 
@@ -215,4 +298,5 @@ FTransform AControlRigShapeActor::GetGlobalTransform() const
 
 	return FTransform::Identity;
 }
-// FControlRigShapeHelper END
\ No newline at end of file
+
+// FControlRigShapeHelper END
diff --git a/Engine/Plugins/Animation/ControlRig/Source/ControlRig/Private/Rigs/RigHierarchy.cpp b/Engine/Plugins/Animation/ControlRig/Source/ControlRig/Private/Rigs/RigHierarchy.cpp
--- a/Engine/Plugins/Animation/ControlRig/Source/ControlRig/Private/Rigs/RigHierarchy.cpp
+++ b/Engine/Plugins/Animation/ControlRig/Source/ControlRig/Private/Rigs/RigHierarchy.cpp
@@ -112,8 +112,8 @@ URigHierarchy::URigHierarchy()
 #endif
 , bEnableCacheValidityCheck(bEnableValidityCheckbyDefault)
 , HierarchyForCacheValidation()
-#if WITH_EDITOR
 , ExecuteContext(nullptr)
+#if WITH_EDITOR
 , bRecordTransformsPerInstruction(true)
 #endif
 {
@@ -2127,7 +2127,23 @@ void URigHierarchy::Notify(ERigHierarchyNotification InNotifType, const FRigBase
 	{
 		return;
 	}
-	ModifiedEvent.Broadcast(InNotifType, this, InElement);
+
+	// if we are running a VM right now
+	if(ExecuteContext != nullptr)
+	{
+		QueueNotification(InNotifType, InElement);
+		return;
+	}
+
+	if(QueuedNotifications.IsEmpty())
+	{
+		ModifiedEvent.Broadcast(InNotifType, this, InElement);
+	}
+	else
+	{
+		QueueNotification(InNotifType, InElement);
+		SendQueuedNotifications();
+	}
 
 #if WITH_EDITOR
 
@@ -2166,6 +2182,148 @@ void URigHierarchy::Notify(ERigHierarchyNotification InNotifType, const FRigBase
 #endif
 }
 
+void URigHierarchy::QueueNotification(ERigHierarchyNotification InNotification, const FRigBaseElement* InElement)
+{
+	FQueuedNotification Entry;
+	Entry.Type = InNotification;
+	Entry.Key = InElement ? InElement->GetKey() : FRigElementKey();
+	QueuedNotifications.Enqueue(Entry);
+}
+
+void URigHierarchy::SendQueuedNotifications()
+{
+	if(bSuspendNotifications)
+	{
+		QueuedNotifications.Empty();
+		return;
+	}
+
+	if(QueuedNotifications.IsEmpty())
+	{
+		return;
+	}
+
+	if(ExecuteContext != nullptr)
+	{
+		return;
+	}
+
+	// enable access to the controller during this method
+	FRigHierarchyEnableControllerBracket EnableController(this, true);
+
+	// we'll collect all notifications and will clean them up
+	// to guard against notification storms.
+	TArray<FQueuedNotification> AllNotifications;
+	FQueuedNotification EntryFromQueue;
+	while(QueuedNotifications.Dequeue(EntryFromQueue))
+	{
+		AllNotifications.Add(EntryFromQueue);
+	}
+	QueuedNotifications.Empty();
+
+	// now we'll filter the notifications. we'll go through them in
+	// reverse and will skip any aggregates (like change color multiple times on the same thing)
+	// as well as collapse pairs such as select and deselect
+	TArray<FQueuedNotification> FilteredNotifications;
+	TArray<FQueuedNotification> UniqueNotifications;
+	for(int32 Index = AllNotifications.Num() - 1; Index >= 0; Index--)
+	{
+		const FQueuedNotification& Entry = AllNotifications[Index];
+
+		bool bSkipNotification = false;
+		switch(Entry.Type)
+		{
+			case ERigHierarchyNotification::HierarchyReset:
+			case ERigHierarchyNotification::ElementRemoved:
+			case ERigHierarchyNotification::ElementRenamed:
+			{
+				// we don't allow these to happen during the run of a VM
+				if(const URigHierarchyController* Controller = GetController())
+				{
+					static constexpr TCHAR InvalidNotificationReceivedMessage[] =
+						TEXT("Found invalid queued notification %s - %s. Skipping notification.");
+					const FString NotificationText = StaticEnum<ERigHierarchyNotification>()->GetNameStringByValue((int64)Entry.Type);
+					
+					Controller->ReportErrorf(InvalidNotificationReceivedMessage, *NotificationText, *Entry.Key.ToString());	
+				}
+				bSkipNotification = true;
+				break;
+			}
+			case ERigHierarchyNotification::ControlSettingChanged:
+			case ERigHierarchyNotification::ControlVisibilityChanged:
+			case ERigHierarchyNotification::ControlShapeTransformChanged:
+			case ERigHierarchyNotification::ParentChanged:
+			case ERigHierarchyNotification::ParentWeightsChanged:
+			{
+				// these notifications are aggregates - they don't need to happen
+				// more than once during an update
+				bSkipNotification = UniqueNotifications.Contains(Entry);
+				break;
+			}
+			case ERigHierarchyNotification::ElementSelected:
+			case ERigHierarchyNotification::ElementDeselected:
+			{
+				FQueuedNotification OppositeEntry;
+				OppositeEntry.Type = (Entry.Type == ERigHierarchyNotification::ElementSelected) ?
+					ERigHierarchyNotification::ElementDeselected :
+					ERigHierarchyNotification::ElementSelected;
+				OppositeEntry.Key = Entry.Key;
+
+				// we don't need to add this if we already performed the selection or
+				// deselection of the same item before
+				bSkipNotification = UniqueNotifications.Contains(Entry) ||
+					UniqueNotifications.Contains(OppositeEntry);
+				break;
+			}
+			case ERigHierarchyNotification::Max:
+			{
+				bSkipNotification = true;
+				break;
+			}
+			case ERigHierarchyNotification::ElementAdded:
+			case ERigHierarchyNotification::InteractionBracketOpened:
+			case ERigHierarchyNotification::InteractionBracketClosed:
+			default:
+			{
+				break;
+			}
+		}
+
+		UniqueNotifications.AddUnique(Entry);
+		if(!bSkipNotification)
+		{
+			FilteredNotifications.Add(Entry);
+		}
+
+		// if we ever hit a reset then we don't need to deal with
+		// any previous notifications.
+		if(Entry.Type == ERigHierarchyNotification::HierarchyReset)
+		{
+			break;
+		}
+	}
+
+	if(FilteredNotifications.IsEmpty())
+	{
+		return;
+	}
+
+	ModifiedEvent.Broadcast(ERigHierarchyNotification::InteractionBracketOpened, this, nullptr);
+
+	// finally send all of the notifications
+	// (they have been added in the reverse order to the array)
+	for(int32 Index = FilteredNotifications.Num() - 1; Index >= 0; Index--)
+	{
+		const FQueuedNotification& Entry = FilteredNotifications[Index];
+		// we expect all elements to still exist, removal /rename
+		// // is not allowed during the run of a VM
+		const FRigBaseElement* Element = FindChecked(Entry.Key); 
+		ModifiedEvent.Broadcast(Entry.Type, this, Element);
+	}
+
+	ModifiedEvent.Broadcast(ERigHierarchyNotification::InteractionBracketClosed, this, nullptr);
+}
+
 FTransform URigHierarchy::GetTransform(FRigTransformElement* InTransformElement,
 	const ERigTransformType::Type InTransformType) const
 {
diff --git a/Engine/Plugins/Animation/ControlRig/Source/ControlRig/Private/Rigs/RigHierarchyController.cpp b/Engine/Plugins/Animation/ControlRig/Source/ControlRig/Private/Rigs/RigHierarchyController.cpp
--- a/Engine/Plugins/Animation/ControlRig/Source/ControlRig/Private/Rigs/RigHierarchyController.cpp
+++ b/Engine/Plugins/Animation/ControlRig/Source/ControlRig/Private/Rigs/RigHierarchyController.cpp
@@ -1043,7 +1043,9 @@ TArray<FRigElementKey> URigHierarchyController::ImportFromText(FString InContent
 	{
 		KeyMap.Add(Element->GetKey(), Element->GetKey());
 	}
-	
+
+	FRigHierarchyInteractionBracket InteractionBracket(Hierarchy.Get());
+
 	for(const FRigHierarchyCopyPasteContentPerElement& PerElementData : Data.Elements)
 	{
 		ErrorPipe.NumErrors = 0;
@@ -1222,6 +1224,7 @@ TArray<FRigElementKey> URigHierarchyController::ImportFromText(FString InContent
 TArray<FRigElementKey> URigHierarchyController::ImportFromHierarchyContainer(const FRigHierarchyContainer& InContainer, bool bIsCopyAndPaste)
 {
 	TMap<FRigElementKey, FRigElementKey> KeyMap;;
+	FRigHierarchyInteractionBracket InteractionBracket(Hierarchy.Get());
 
 	for(const FRigBone& Bone : InContainer.BoneHierarchy)
 	{
@@ -2426,6 +2429,8 @@ TArray<FRigElementKey> URigHierarchyController::DuplicateElements(TArray<FRigEle
 
 TArray<FRigElementKey> URigHierarchyController::MirrorElements(TArray<FRigElementKey> InKeys, FRigMirrorSettings InSettings, bool bSelectNewElements, bool bSetupUndo, bool bPrintPythonCommands)
 {
+	FRigHierarchyInteractionBracket InteractionBracket(Hierarchy.Get());
+
 	TArray<FRigElementKey> OriginalKeys = Hierarchy->SortKeys(InKeys);
 	TArray<FRigElementKey> DuplicatedKeys = DuplicateElements(OriginalKeys, bSelectNewElements, bSetupUndo);
 
diff --git a/Engine/Plugins/Animation/ControlRig/Source/ControlRig/Public/ControlRigGizmoActor.h b/Engine/Plugins/Animation/ControlRig/Source/ControlRig/Public/ControlRigGizmoActor.h
--- a/Engine/Plugins/Animation/ControlRig/Source/ControlRig/Public/ControlRigGizmoActor.h
+++ b/Engine/Plugins/Animation/ControlRig/Source/ControlRig/Public/ControlRigGizmoActor.h
@@ -20,6 +20,7 @@ struct FControlShapeActorCreationParam
 		: ManipObj(nullptr)
 		, ControlRigIndex(INDEX_NONE)
 		, ControlName(NAME_None)
+		, ShapeName(NAME_None)
 		, SpawnTransform(FTransform::Identity)
 		, ShapeTransform(FTransform::Identity)
 		, MeshTransform(FTransform::Identity)
@@ -34,6 +35,7 @@ struct FControlShapeActorCreationParam
 	UObject*	ManipObj;
 	int32		ControlRigIndex;
 	FName		ControlName;
+	FName		ShapeName;
 	FTransform	SpawnTransform;
 	FTransform  ShapeTransform;
 	FTransform  MeshTransform;
@@ -69,6 +71,10 @@ public:
 	UPROPERTY()
 	FName ControlName;
 
+	// the name of the shape to use on this actor
+   	UPROPERTY()
+   	FName ShapeName;
+
 	// the name of the color parameter on the material
 	UPROPERTY()
 	FName ColorParameterName;
@@ -137,6 +143,15 @@ public:
 
 	UFUNCTION(BlueprintPure, Category = "ControlRig|Shape")
 	FTransform GetGlobalTransform() const;
+
+	// try to update the actor with the latest settings
+	bool UpdateControlSettings(
+		ERigHierarchyNotification InNotif,
+		UControlRig* InControlRig,
+		const FRigControlElement* InControlElement,
+		bool bHideManipulators,
+		bool bIsInLevelEditor);
+	
 private:
 	/** Whether this control is enabled */
 	UPROPERTY(BlueprintGetter = IsEnabled, BlueprintSetter= SetEnabled, Category = "ControlRig|Shape")
diff --git a/Engine/Plugins/Animation/ControlRig/Source/ControlRig/Public/Rigs/RigHierarchy.h b/Engine/Plugins/Animation/ControlRig/Source/ControlRig/Public/Rigs/RigHierarchy.h
--- a/Engine/Plugins/Animation/ControlRig/Source/ControlRig/Public/Rigs/RigHierarchy.h
+++ b/Engine/Plugins/Animation/ControlRig/Source/ControlRig/Public/Rigs/RigHierarchy.h
@@ -13,6 +13,7 @@
 #if WITH_EDITOR
 #include "RigVMPythonUtils.h"
 #endif
+#include "Containers/Queue.h"
 #include "RigHierarchy.generated.h"
 
 class URigHierarchy;
@@ -2546,7 +2547,7 @@ private:
 	/*
 	* Helper function to create an element for a given type
 	*/
-	static void DestroyElement(FRigBaseElement*& InElement); 
+	static void DestroyElement(FRigBaseElement*& InElement);
 
 	/*
 	 * Templated helper function to create an element
@@ -2972,8 +2973,9 @@ private:
 	
 	void EnsureCacheValidityImpl();
 
-#if WITH_EDITOR
 	const FRigVMExtendedExecuteContext* ExecuteContext;
+
+#if WITH_EDITOR
 	mutable bool bRecordTransformsPerInstruction;
 	mutable TArray<TArray<TArray<int32>>> ReadTransformsPerInstructionPerSlice;
 	mutable TArray<TArray<TArray<int32>>> WrittenTransformsPerInstructionPerSlice;
@@ -2988,6 +2990,21 @@ private:
 
 	static const TArray<FString>& GetTransformTypeStrings();
 
+	struct FQueuedNotification
+	{
+		ERigHierarchyNotification Type;
+		FRigElementKey Key;
+		
+		FORCEINLINE bool operator == (const FQueuedNotification& InOther) const
+		{
+			return Type == InOther.Type && Key == InOther.Key;
+		}
+	};
+	TQueue<FQueuedNotification> QueuedNotifications;
+
+	void QueueNotification(ERigHierarchyNotification InNotification, const FRigBaseElement* InElement);
+	void SendQueuedNotifications();
+	
 	friend class URigHierarchyController;
 	friend class UControlRig;
 	friend class FControlRigEditor;
@@ -2995,6 +3012,28 @@ private:
 	friend struct FRigHierarchyGlobalValidityBracket;
 	friend struct FControlRigVisualGraphUtils;
 	friend struct FRigHierarchyEnableControllerBracket;
+	friend struct FRigHierarchyExecuteContextBracket;
+};
+
+struct CONTROLRIG_API FRigHierarchyInteractionBracket
+{
+public:
+	
+	FRigHierarchyInteractionBracket(URigHierarchy* InHierarchy)
+		: Hierarchy(InHierarchy)
+	{
+		check(Hierarchy);
+		Hierarchy->Notify(ERigHierarchyNotification::InteractionBracketOpened, nullptr);
+	}
+
+	~FRigHierarchyInteractionBracket()
+	{
+		Hierarchy->Notify(ERigHierarchyNotification::InteractionBracketClosed, nullptr);
+	}
+
+private:
+
+	URigHierarchy* Hierarchy;
 };
 
 struct CONTROLRIG_API FRigHierarchyEnableControllerBracket : public TGuardValue<bool>
@@ -3012,6 +3051,29 @@ private:
 	friend struct FRigUnit_AddParent;
 };
 
+struct CONTROLRIG_API FRigHierarchyExecuteContextBracket
+{
+private:
+
+	FRigHierarchyExecuteContextBracket(URigHierarchy* InHierarchy, const FRigVMExtendedExecuteContext* InContext)
+		: Hierarchy(InHierarchy)
+		, PreviousContext(InHierarchy->ExecuteContext)
+	{
+		Hierarchy->ExecuteContext = InContext;
+	}
+
+	~FRigHierarchyExecuteContextBracket()
+	{
+		Hierarchy->ExecuteContext = PreviousContext;
+		Hierarchy->SendQueuedNotifications();
+	}
+
+	URigHierarchy* Hierarchy;
+	const FRigVMExtendedExecuteContext* PreviousContext;
+
+	friend class UControlRig;
+};
+
 struct CONTROLRIG_API FRigHierarchyValidityBracket
 {
 	public:
diff --git a/Engine/Plugins/Animation/ControlRig/Source/ControlRig/Public/Rigs/RigHierarchyDefines.h b/Engine/Plugins/Animation/ControlRig/Source/ControlRig/Public/Rigs/RigHierarchyDefines.h
--- a/Engine/Plugins/Animation/ControlRig/Source/ControlRig/Public/Rigs/RigHierarchyDefines.h
+++ b/Engine/Plugins/Animation/ControlRig/Source/ControlRig/Public/Rigs/RigHierarchyDefines.h
@@ -59,6 +59,8 @@ enum class ERigHierarchyNotification : uint8
 	ControlVisibilityChanged,
 	ControlShapeTransformChanged,
 	ParentWeightsChanged,
+	InteractionBracketOpened,
+	InteractionBracketClosed,
 
 	/** MAX - invalid */
 	Max UMETA(Hidden),
diff --git a/Engine/Plugins/Animation/ControlRig/Source/ControlRigEditor/Private/EditMode/ControlRigEditMode.cpp b/Engine/Plugins/Animation/ControlRig/Source/ControlRigEditor/Private/EditMode/ControlRigEditMode.cpp
--- a/Engine/Plugins/Animation/ControlRig/Source/ControlRigEditor/Private/EditMode/ControlRigEditMode.cpp
+++ b/Engine/Plugins/Animation/ControlRig/Source/ControlRigEditor/Private/EditMode/ControlRigEditMode.cpp
@@ -2623,6 +2623,10 @@ void FControlRigEditMode::OnHierarchyModified(ERigHierarchyNotification InNotif,
 		case ERigHierarchyNotification::ElementRemoved:
 		case ERigHierarchyNotification::ElementRenamed:
 		case ERigHierarchyNotification::HierarchyReset:
+		{
+			RequestToRecreateControlShapeActors();
+			break;
+		}
 		case ERigHierarchyNotification::ControlSettingChanged:
 		case ERigHierarchyNotification::ControlShapeTransformChanged:
 		{
@@ -2634,6 +2638,25 @@ void FControlRigEditMode::OnHierarchyModified(ERigHierarchyNotification InNotif,
 					bIsChangingControlShapeTransform = false;
 				} 
 			}
+
+			const FRigElementKey Key = InElement->GetKey();
+			if(Key.Type == ERigElementType::Control)
+			{
+				if(const FRigControlElement* ControlElement = Cast<FRigControlElement>(InElement))
+				{
+					if (AControlRigShapeActor* ShapeActor = GetControlShapeFromControlName(Key.Name))
+					{
+						// try to lazily apply the changes to the actor
+						const UControlRigEditModeSettings* Settings = GetDefault<UControlRigEditModeSettings>();
+						if(ShapeActor->UpdateControlSettings(InNotif, GetControlRig(true), ControlElement, Settings->bHideControlShapes, IsInLevelEditor()))
+						{
+							break;
+						}
+					}
+				}
+			}
+
+			// if we can't deal with this lazily, let's fall back to recreating all control shape actors
 			RequestToRecreateControlShapeActors();
 			break;
 		}
@@ -2698,6 +2721,8 @@ void FControlRigEditMode::OnHierarchyModified(ERigHierarchyNotification InNotif,
 				}
 			}
 		}
+		case ERigHierarchyNotification::InteractionBracketOpened:
+		case ERigHierarchyNotification::InteractionBracketClosed:
 		default:
 		{
 			break;
@@ -3267,6 +3292,7 @@ bool FControlRigEditMode::CreateShapeActors(UWorld* World)
 				Param.ManipObj = ControlRig;
 				Param.ControlRigIndex = ControlRigIndex;
 				Param.ControlName = ControlElement->GetName();
+				Param.ShapeName = ControlElement->Settings.ShapeName;
 				Param.SpawnTransform = ControlRig->GetControlGlobalTransform(ControlElement->GetName());
 				Param.ShapeTransform = ControlRig->GetHierarchy()->GetControlShapeTransform(ControlElement, ERigTransformType::CurrentLocal);
 				Param.bSelectable = ControlElement->Settings.bAnimatable;
diff --git a/Engine/Plugins/Animation/ControlRig/Source/ControlRigEditor/Private/Editor/ControlRigSkeletalMeshComponent.cpp b/Engine/Plugins/Animation/ControlRig/Source/ControlRigEditor/Private/Editor/ControlRigSkeletalMeshComponent.cpp
--- a/Engine/Plugins/Animation/ControlRig/Source/ControlRigEditor/Private/Editor/ControlRigSkeletalMeshComponent.cpp
+++ b/Engine/Plugins/Animation/ControlRig/Source/ControlRigEditor/Private/Editor/ControlRigSkeletalMeshComponent.cpp
@@ -9,6 +9,8 @@
 UControlRigSkeletalMeshComponent::UControlRigSkeletalMeshComponent(const FObjectInitializer& ObjectInitializer)
 	: Super(ObjectInitializer)
 	, DebugDrawSkeleton(false)
+	, HierarchyInteractionBracket(0)
+	, bRebuildDebugDrawSkeletonRequired(false)
 {
 	SetDisablePostProcessBlueprint(true);
 }
@@ -54,6 +56,8 @@ void UControlRigSkeletalMeshComponent::SetCustomDefaultPose()
 
 void UControlRigSkeletalMeshComponent::RebuildDebugDrawSkeleton()
 {
+	bRebuildDebugDrawSkeletonRequired = false;
+ 
 	UControlRigLayerInstance* ControlRigInstance = Cast<UControlRigLayerInstance>(GetAnimInstance());
 
 	if (ControlRigInstance)
@@ -186,7 +190,28 @@ void UControlRigSkeletalMeshComponent::OnHierarchyModified(ERigHierarchyNotifica
 		case ERigHierarchyNotification::ParentChanged:
 		case ERigHierarchyNotification::HierarchyReset:
 		{
-			RebuildDebugDrawSkeleton();
+			bRebuildDebugDrawSkeletonRequired = true;
+			if(HierarchyInteractionBracket == 0)
+			{
+				RebuildDebugDrawSkeleton();
+			}
+			break;
+		}
+		case ERigHierarchyNotification::InteractionBracketOpened:
+		{
+			HierarchyInteractionBracket++;
+			break;
+		}
+		case ERigHierarchyNotification::InteractionBracketClosed:
+		{
+			HierarchyInteractionBracket = FMath::Max(HierarchyInteractionBracket - 1, 0);
+			if(HierarchyInteractionBracket == 0)
+			{
+				if(bRebuildDebugDrawSkeletonRequired)
+				{
+					RebuildDebugDrawSkeleton();
+				}
+			}
 			break;
 		}
 		default:
diff --git a/Engine/Plugins/Animation/ControlRig/Source/ControlRigEditor/Private/Editor/ControlRigSkeletalMeshComponent.h b/Engine/Plugins/Animation/ControlRig/Source/ControlRigEditor/Private/Editor/ControlRigSkeletalMeshComponent.h
--- a/Engine/Plugins/Animation/ControlRig/Source/ControlRigEditor/Private/Editor/ControlRigSkeletalMeshComponent.h
+++ b/Engine/Plugins/Animation/ControlRig/Source/ControlRigEditor/Private/Editor/ControlRigSkeletalMeshComponent.h
@@ -59,4 +59,6 @@ private:
 	TArray<FBoneIndexType> DebugDrawBones;
 	TArray<int32> DebugDrawBoneIndexInHierarchy;
 	TWeakObjectPtr<UControlRig> ControlRigBeingDebuggedPtr;
+	int32 HierarchyInteractionBracket;
+	bool bRebuildDebugDrawSkeletonRequired;
 };
diff --git a/Engine/Plugins/Animation/ControlRig/Source/ControlRigEditor/Private/Editor/SRigHierarchy.cpp b/Engine/Plugins/Animation/ControlRig/Source/ControlRigEditor/Private/Editor/SRigHierarchy.cpp
--- a/Engine/Plugins/Animation/ControlRig/Source/ControlRigEditor/Private/Editor/SRigHierarchy.cpp
+++ b/Engine/Plugins/Animation/ControlRig/Source/ControlRigEditor/Private/Editor/SRigHierarchy.cpp
@@ -1471,6 +1471,8 @@ void SRigHierarchy::HandleDeleteItem()
  		// clear selection early here to make sure ControlRigEditMode can react to this deletion
  		// it cannot react to it during Controller->RemoveElement() later because bSuspendAllNotifications is true
  		Controller->ClearSelection();
+
+ 		FRigHierarchyInteractionBracket InteractionBracket(Hierarchy);
  		
 		for (const FRigElementKey& SelectedKey : SelectedKeys)
 		{
@@ -2475,6 +2477,7 @@ FReply SRigHierarchy::ReparentOrMatchTransform(const TArray<FRigElementKey>& Dra
 		TGuardValue<bool> GuardRigHierarchyChanges(bIsChangingRigHierarchy, true);
 		TGuardValue<bool> SuspendBlueprintNotifs(ControlRigBlueprint->bSuspendAllNotifications, true);
 		FScopedTransaction Transaction(LOCTEXT("HierarchyDragAndDrop", "Drag & Drop"));
+		FRigHierarchyInteractionBracket InteractionBracket(Hierarchy);
 
 		FTransform TargetGlobalTransform = DebuggedHierarchy->GetGlobalTransform(TargetKey);
 
@@ -2680,6 +2683,9 @@ void SRigHierarchy::HandleSetShapeTransformFromCurrent()
 			{
 				FScopedTransaction Transaction(LOCTEXT("HierarchySetShapeTransforms", "Set Shape Transforms"));
 
+				FRigHierarchyInteractionBracket InteractionBracket(GetHierarchy());
+				FRigHierarchyInteractionBracket DebuggedInteractionBracket(DebuggedHierarchy);
+
 				TArray<TSharedPtr<FRigTreeElement>> SelectedItems = TreeView->GetSelectedItems();
 				for (const TSharedPtr<FRigTreeElement>& SelectedItem : SelectedItems)
 				{
