Sequencer: Mutiple Control Rigs
Many changes  to move a single Control Rig to an array of Control Rigs, main change in ControlRigEditMode.
Control Rigs have aControl visibility bool to specify if they should be visible in the edit mode.
All tools, except space switching, updated to support multiple control rigs.
Control  Rig Editor supports just the first[0], Control Rig.

diff --git a/Engine/Plugins/Animation/ControlRig/Source/ControlRig/Private/ControlRigGizmoActor.cpp b/Engine/Plugins/Animation/ControlRig/Source/ControlRig/Private/ControlRigGizmoActor.cpp
--- a/Engine/Plugins/Animation/ControlRig/Source/ControlRig/Private/ControlRigGizmoActor.cpp
+++ b/Engine/Plugins/Animation/ControlRig/Source/ControlRig/Private/ControlRigGizmoActor.cpp
@@ -11,6 +11,7 @@
 AControlRigShapeActor::AControlRigShapeActor(const FObjectInitializer& ObjectInitializer)
 	: Super(ObjectInitializer)
 	, ControlRigIndex(INDEX_NONE)
+	, ControlRig(nullptr)
 	, ControlName(NAME_None)
 	, ShapeName(NAME_None)
 	, bEnabled(true)
@@ -246,6 +247,7 @@ namespace FControlRigShapeHelper
 		if (ShapeActor)
 		{
 			ShapeActor->ControlRigIndex = CreationParam.ControlRigIndex;
+			ShapeActor->ControlRig = CreationParam.ControlRig;
 			ShapeActor->ControlName = CreationParam.ControlName;
 			ShapeActor->ShapeName = CreationParam.ShapeName;
 			ShapeActor->SetSelectable(CreationParam.bSelectable);
diff --git a/Engine/Plugins/Animation/ControlRig/Source/ControlRig/Public/ControlRig.h b/Engine/Plugins/Animation/ControlRig/Source/ControlRig/Public/ControlRig.h
--- a/Engine/Plugins/Animation/ControlRig/Source/ControlRig/Public/ControlRig.h
+++ b/Engine/Plugins/Animation/ControlRig/Source/ControlRig/Public/ControlRig.h
@@ -723,6 +723,10 @@ private:
 	UPROPERTY(transient)
 	bool bIsInDebugMode;
 
+	/** Whether controls are visible */
+	UPROPERTY(transient)
+	bool bControlsVisible;
+
 #endif
 
 	float DebugBoneRadiusMultiplier;
@@ -749,6 +753,10 @@ public:
 
 	/** Creates the snapshot VM if required and returns it */
 	URigVM* GetSnapshotVM(bool bCreateIfNeeded = true);
+
+	void ToggleControlsVisible() { bControlsVisible = !bControlsVisible; }
+	void SetControlsVisible(const bool bIsVisible) { bControlsVisible = bIsVisible; }
+	bool GetControlsVisible()const { return bControlsVisible;}
 #endif	
 
 private:
diff --git a/Engine/Plugins/Animation/ControlRig/Source/ControlRig/Public/ControlRigGizmoActor.h b/Engine/Plugins/Animation/ControlRig/Source/ControlRig/Public/ControlRigGizmoActor.h
--- a/Engine/Plugins/Animation/ControlRig/Source/ControlRig/Public/ControlRigGizmoActor.h
+++ b/Engine/Plugins/Animation/ControlRig/Source/ControlRig/Public/ControlRigGizmoActor.h
@@ -9,6 +9,7 @@
 #include "Rigs/RigControlHierarchy.h"
 #include "ControlRigGizmoActor.generated.h"
 
+class UControlRig;
 struct FActorSpawnParameters;
 
 USTRUCT()
@@ -19,6 +20,7 @@ struct FControlShapeActorCreationParam
 	FControlShapeActorCreationParam()
 		: ManipObj(nullptr)
 		, ControlRigIndex(INDEX_NONE)
+		, ControlRig(nullptr)
 		, ControlName(NAME_None)
 		, ShapeName(NAME_None)
 		, SpawnTransform(FTransform::Identity)
@@ -34,6 +36,7 @@ struct FControlShapeActorCreationParam
 
 	UObject*	ManipObj;
 	int32		ControlRigIndex;
+	UControlRig* ControlRig;
 	FName		ControlName;
 	FName		ShapeName;
 	FTransform	SpawnTransform;
@@ -67,6 +70,10 @@ public:
 	UPROPERTY()
 	uint32 ControlRigIndex;
 
+	//  control rig this actor is referencing we can have multiple control rig's visible
+	UPROPERTY()
+	TWeakObjectPtr<UControlRig> ControlRig;
+
 	// the name of the control this actor is referencing
 	UPROPERTY()
 	FName ControlName;
diff --git a/Engine/Plugins/Animation/ControlRig/Source/ControlRigEditor/Private/ControlRigBlueprintActions.cpp b/Engine/Plugins/Animation/ControlRig/Source/ControlRigEditor/Private/ControlRigBlueprintActions.cpp
--- a/Engine/Plugins/Animation/ControlRig/Source/ControlRigEditor/Private/ControlRigBlueprintActions.cpp
+++ b/Engine/Plugins/Animation/ControlRig/Source/ControlRigEditor/Private/ControlRigBlueprintActions.cpp
@@ -399,7 +399,7 @@ void FControlRigBlueprintActions::OnSpawnedSkeletalMeshActorChanged(UObject* InO
 				}
 				if (ControlRigEditMode)
 				{
-					ControlRigEditMode->SetObjects(ControlRig, nullptr, WeakSequencer.Pin());
+					ControlRigEditMode->AddControlRigObject(ControlRig, WeakSequencer.Pin());
 				}
 			}
 		}
diff --git a/Engine/Plugins/Animation/ControlRig/Source/ControlRigEditor/Private/ControlRigEditorModule.cpp b/Engine/Plugins/Animation/ControlRig/Source/ControlRigEditor/Private/ControlRigEditorModule.cpp
--- a/Engine/Plugins/Animation/ControlRig/Source/ControlRigEditor/Private/ControlRigEditorModule.cpp
+++ b/Engine/Plugins/Animation/ControlRig/Source/ControlRigEditor/Private/ControlRigEditorModule.cpp
@@ -758,7 +758,7 @@ void FControlRigEditorModule::BakeToControlRig(UClass* ControlRigClass, UAnimSeq
 					}
 					if (ControlRigEditMode)
 					{
-						ControlRigEditMode->SetObjects(ControlRig, nullptr, WeakSequencer.Pin());
+						ControlRigEditMode->AddControlRigObject(ControlRig, WeakSequencer.Pin());
 					}
 
 					//create soft links to each other
diff --git a/Engine/Plugins/Animation/ControlRig/Source/ControlRigEditor/Private/ControlRigSequencerEditorLibrary.cpp b/Engine/Plugins/Animation/ControlRig/Source/ControlRigEditor/Private/ControlRigSequencerEditorLibrary.cpp
--- a/Engine/Plugins/Animation/ControlRig/Source/ControlRigEditor/Private/ControlRigSequencerEditorLibrary.cpp
+++ b/Engine/Plugins/Animation/ControlRig/Source/ControlRigEditor/Private/ControlRigSequencerEditorLibrary.cpp
@@ -43,10 +43,9 @@ TArray<UControlRig*> UControlRigSequencerEditorLibrary::GetVisibleControlRigs()
 {
 	TArray<UControlRig*> ControlRigs;
 	FControlRigEditMode* ControlRigEditMode = static_cast<FControlRigEditMode*>(GLevelEditorModeTools().GetActiveMode(FControlRigEditMode::ModeName));
-	if (ControlRigEditMode && ControlRigEditMode->GetControlRig(true) && ControlRigEditMode->GetControlRig(true)->GetObjectBinding() &&\
-		ControlRigEditMode->GetControlRig(true)->GetObjectBinding()->GetBoundObject())
+	if (ControlRigEditMode)
 	{
-		ControlRigs.Add(ControlRigEditMode->GetControlRig(true));
+		ControlRigs = ControlRigEditMode->GetControlRigsArray(true /*bIsVisible*/);
 	}
 	return ControlRigs;
 }
@@ -250,7 +249,7 @@ static UMovieSceneControlRigParameterTrack* AddControlRig(ULevelSequence* LevelS
 			}
 			if (ControlRigEditMode)
 			{
-				ControlRigEditMode->SetObjects(ControlRig, nullptr, SharedSequencer);
+				ControlRigEditMode->AddControlRigObject(ControlRig, SharedSequencer);
 			}
 			return Track;
 		}
@@ -733,11 +732,13 @@ bool UControlRigSequencerEditorLibrary::BakeToControlRig(UWorld* World, ULevelSe
 						}
 						else
 						{
+							/* mz todo we don't unbind  will test more
 							UControlRig* OldControlRig = ControlRigEditMode->GetControlRig(false);
 							if (OldControlRig)
 							{
 								WeakSequencer.Pin()->ObjectImplicitlyRemoved(OldControlRig);
 							}
+							*/
 						}
 					}
 
@@ -787,7 +788,7 @@ bool UControlRigSequencerEditorLibrary::BakeToControlRig(UWorld* World, ULevelSe
 					//Finish Setup
 					if (ControlRigEditMode)
 					{
-						ControlRigEditMode->SetObjects(ControlRig, nullptr, WeakSequencer.Pin());
+						ControlRigEditMode->AddControlRigObject(ControlRig, WeakSequencer.Pin());
 					}
 
 					TempAnimSequence->MarkAsGarbage();
diff --git a/Engine/Plugins/Animation/ControlRig/Source/ControlRigEditor/Private/EditMode/ControlRigBaseDockableView.cpp b/Engine/Plugins/Animation/ControlRig/Source/ControlRigEditor/Private/EditMode/ControlRigBaseDockableView.cpp
--- a/Engine/Plugins/Animation/ControlRig/Source/ControlRigEditor/Private/EditMode/ControlRigBaseDockableView.cpp
+++ b/Engine/Plugins/Animation/ControlRig/Source/ControlRigEditor/Private/EditMode/ControlRigBaseDockableView.cpp
@@ -35,9 +35,13 @@ FControlRigBaseDockableView::~FControlRigBaseDockableView()
 	if (FControlRigEditMode* EditMode = static_cast<FControlRigEditMode*>(ModeTools->GetActiveMode(FControlRigEditMode::ModeName)))
 	{
 		EditMode->OnControlRigAddedOrRemoved().RemoveAll(this);
-		if (EditMode->GetControlRig(true))
+		TArrayView<TWeakObjectPtr<UControlRig>> Controls = EditMode->GetControlRigs();
+		for (TWeakObjectPtr<UControlRig>& ControlRigPtr : Controls)
 		{
-			EditMode->GetControlRig(true)->ControlSelected().RemoveAll(this);
+			if (UControlRig* ControlRig = ControlRigPtr.Get())
+			{
+				ControlRig->ControlSelected().RemoveAll(this);
+			}
 		}
 	}
 	else
@@ -58,7 +62,16 @@ UControlRig* FControlRigBaseDockableView::GetControlRig()
 	UControlRig* NewControlRig = nullptr;
 	if (FControlRigEditMode* EditMode = static_cast<FControlRigEditMode*>(ModeTools->GetActiveMode(FControlRigEditMode::ModeName)))
 	{
-		NewControlRig = EditMode->GetControlRig(true);
+		//just get first one, should only be one for control rig editor
+		TArrayView<TWeakObjectPtr<UControlRig>> Controls = EditMode->GetControlRigs();
+		for (TWeakObjectPtr<UControlRig>& ControlRigPtr : Controls)
+		{
+			if (UControlRig* ControlRig = ControlRigPtr.Get())
+			{
+				NewControlRig = ControlRig;
+				break;
+			}
+		}
 	}
 	bool bNewControlRig = NewControlRig != CurrentControlRig;
 	if (bNewControlRig)
diff --git a/Engine/Plugins/Animation/ControlRig/Source/ControlRigEditor/Private/EditMode/ControlRigControlsProxy.cpp b/Engine/Plugins/Animation/ControlRig/Source/ControlRigEditor/Private/EditMode/ControlRigControlsProxy.cpp
--- a/Engine/Plugins/Animation/ControlRig/Source/ControlRigEditor/Private/EditMode/ControlRigControlsProxy.cpp
+++ b/Engine/Plugins/Animation/ControlRig/Source/ControlRigEditor/Private/EditMode/ControlRigControlsProxy.cpp
@@ -100,20 +100,6 @@ void UControlRigControlsProxy::SelectionChanged(bool bInSelected)
 	}
 }
 
-void UControlRigControlsProxy::CheckEditModeOnSelectionChange(UControlRig *InControlRig)
-{
-#if WITH_EDITOR
-	FControlRigEditMode* ControlRigEditMode = static_cast<FControlRigEditMode*>(GLevelEditorModeTools().GetActiveMode(FControlRigEditMode::ModeName));
-	if (ControlRigEditMode)
-	{
-		if (ControlRigEditMode->GetControlRig(false) != InControlRig)
-		{
-			ControlRigEditMode->SetObjects(InControlRig, nullptr, nullptr);
-		}
-	}
-#endif
-}
-
 void UControlRigControlsProxy::PostEditChangeProperty(struct FPropertyChangedEvent& PropertyChangedEvent)
 {
 	Super::PostEditChangeProperty(PropertyChangedEvent);
@@ -122,10 +108,6 @@ void UControlRigControlsProxy::PostEditChangeProperty(struct FPropertyChangedEve
 		FRigControlElement* ControlElement = GetControlElement();
 		if (ControlElement && ControlRig.IsValid())
 		{
-			if (bSelected)
-			{
-				CheckEditModeOnSelectionChange(ControlRig.Get());
-			}
 			FControlRigInteractionScope InteractionScope(ControlRig.Get());
 			ControlRig->SelectControl(ControlName, bSelected);
 			ControlRig->Evaluate_AnyThread();
@@ -139,10 +121,6 @@ void UControlRigControlsProxy::PostEditUndo()
 	FRigControlElement* ControlElement = GetControlElement();
 	if (ControlElement && ControlRig.IsValid() && ControlRig->GetHierarchy()->Contains(FRigElementKey(ControlName, ERigElementType::Control)))
 	{
-		if (bSelected)
-		{
-			CheckEditModeOnSelectionChange(ControlRig.Get());
-		}
 		ControlRig->SelectControl(ControlName, bSelected);
 	}
 }
@@ -196,10 +174,6 @@ void UControlRigTransformControlProxy::PostEditUndo()
 	FRigControlElement* ControlElement = GetControlElement();
 	if (ControlElement && ControlRig.IsValid() && ControlRig->GetHierarchy()->Contains(FRigElementKey(ControlName, ERigElementType::Control)))
 	{
-		if (bSelected)
-		{
-			CheckEditModeOnSelectionChange(ControlRig.Get());
-		}
 		ControlRig->SelectControl(ControlName, bSelected);
 		FTransform RealTransform = Transform; //Transform is FEulerTransform
 		ControlRig->SetControlValue<FRigControlValue::FTransform_Float>(ControlName, RealTransform, true, EControlRigSetKey::Never,false);
@@ -254,10 +228,6 @@ void UControlRigTransformNoScaleControlProxy::PostEditUndo()
 	FRigControlElement* ControlElement = GetControlElement();
 	if (ControlElement && ControlRig.IsValid() && ControlRig->GetHierarchy()->Contains(FRigElementKey(ControlName, ERigElementType::Control)))
 	{
-		if (bSelected)
-		{
-			CheckEditModeOnSelectionChange(ControlRig.Get());
-		}
 		ControlRig->SelectControl(ControlName, bSelected);
 		ControlRig->SetControlValue<FRigControlValue::FTransformNoScale_Float>(ControlName, Transform, true, EControlRigSetKey::Never,false);
 	}
@@ -313,10 +283,6 @@ void UControlRigEulerTransformControlProxy::PostEditUndo()
 	FRigControlElement* ControlElement = GetControlElement();
 	if (ControlElement && ControlRig.IsValid() && ControlRig->GetHierarchy()->Contains(FRigElementKey(ControlName, ERigElementType::Control)))
 	{
-		if (bSelected)
-		{
-			CheckEditModeOnSelectionChange(ControlRig.Get());
-		}
 		ControlRig->SelectControl(ControlName, bSelected);
 		ControlRig->SetControlValue<FRigControlValue::FEulerTransform_Float>(ControlName, Transform, true, EControlRigSetKey::Never,false);
 	}
@@ -368,10 +334,6 @@ void UControlRigFloatControlProxy::PostEditUndo()
 	FRigControlElement* ControlElement = GetControlElement();
 	if (ControlElement && ControlRig.IsValid() && ControlRig->GetHierarchy()->Contains(FRigElementKey(ControlName, ERigElementType::Control)))
 	{
-		if (bSelected)
-		{
-			CheckEditModeOnSelectionChange(ControlRig.Get());
-		}
 		ControlRig->SelectControl(ControlName, bSelected);
 		ControlRig->SetControlValue<float>(ControlName, Float, true, EControlRigSetKey::Never,false);
 	}
@@ -423,10 +385,6 @@ void UControlRigIntegerControlProxy::PostEditUndo()
 	FRigControlElement* ControlElement = GetControlElement();
 	if (ControlElement && ControlRig.IsValid() && ControlRig->GetHierarchy()->Contains(FRigElementKey(ControlName, ERigElementType::Control)))
 	{
-		if (bSelected)
-		{
-			CheckEditModeOnSelectionChange(ControlRig.Get());
-		}
 		ControlRig->SelectControl(ControlName, bSelected);
 		ControlRig->SetControlValue<int32>(ControlName, Integer, true, EControlRigSetKey::Never,false);
 	}
@@ -482,10 +440,6 @@ void UControlRigEnumControlProxy::PostEditUndo()
 	FRigControlElement* ControlElement = GetControlElement();
 	if (ControlElement && ControlRig.IsValid() && ControlRig->GetHierarchy()->Contains(FRigElementKey(ControlName, ERigElementType::Control)))
 	{
-		if (bSelected)
-		{
-			CheckEditModeOnSelectionChange(ControlRig.Get());
-		}
 		ControlRig->SelectControl(ControlName, bSelected);
 		ControlRig->SetControlValue<int32>(ControlName, Enum.EnumIndex, true, EControlRigSetKey::Never,false);
 	}
@@ -538,10 +492,6 @@ void UControlRigVectorControlProxy::PostEditUndo()
 	FRigControlElement* ControlElement = GetControlElement();
 	if (ControlElement && ControlRig.IsValid() && ControlRig->GetHierarchy()->Contains(FRigElementKey(ControlName, ERigElementType::Control)))
 	{
-		if (bSelected)
-		{
-			CheckEditModeOnSelectionChange(ControlRig.Get());
-		}
 		ControlRig->SelectControl(ControlName, bSelected);
 		ControlRig->SetControlValue<FVector3f>(ControlName, (FVector3f)Vector, true, EControlRigSetKey::Never,false);
 	}
@@ -596,10 +546,6 @@ void UControlRigVector2DControlProxy::PostEditUndo()
 	FRigControlElement* ControlElement = GetControlElement();
 	if (ControlElement && ControlRig.IsValid() && ControlRig->GetHierarchy()->Contains(FRigElementKey(ControlName, ERigElementType::Control)))
 	{
-		if (bSelected)
-		{
-			CheckEditModeOnSelectionChange(ControlRig.Get());
-		}
 		ControlRig->SelectControl(ControlName, bSelected);
 		ControlRig->SetControlValue<FVector3f>(ControlName, FVector3f(Vector2D.X, Vector2D.Y, 0.f), true, EControlRigSetKey::Never,false);
 	}
@@ -652,10 +598,6 @@ void UControlRigBoolControlProxy::PostEditUndo()
 	FRigControlElement* ControlElement = GetControlElement();
 	if (ControlElement && ControlRig.IsValid() && ControlRig->GetHierarchy()->Contains(FRigElementKey(ControlName, ERigElementType::Control)))
 	{
-		if (bSelected)
-		{
-			CheckEditModeOnSelectionChange(ControlRig.Get());
-		}
 		ControlRig->SelectControl(ControlName, bSelected);
 		ControlRig->SetControlValue<bool>(ControlName, Bool, true, EControlRigSetKey::Never,false);
 	}
@@ -675,19 +617,23 @@ void UControlRigBoolControlProxy::SetKey(const IPropertyHandle& KeyedPropertyHan
 
 //////UControlDetailPanelControlProxies////////
 
-UControlRigControlsProxy* UControlRigDetailPanelControlProxies::FindProxy(const FName& Name) const
+UControlRigControlsProxy* UControlRigDetailPanelControlProxies::FindProxy(UControlRig* ControlRig, const FName& Name) const
 {
-	TObjectPtr<UControlRigControlsProxy> const* Proxy = AllProxies.Find(Name);
-	if (Proxy &&  Proxy[0])
+	const FControlToProxyMap* ControlRigProxies = AllProxies.Find(ControlRig);
+	if (ControlRigProxies)
 	{
-		return Proxy[0];
+		TObjectPtr<UControlRigControlsProxy> const* Proxy = ControlRigProxies->ControlToProxy.Find(Name);
+		if (Proxy && Proxy[0])
+		{
+			return Proxy[0];
+		}
 	}
 	return nullptr;
 }
 
-void UControlRigDetailPanelControlProxies::AddProxy(const FName& Name, UControlRig* ControlRig, FRigControlElement* ControlElement)
+void UControlRigDetailPanelControlProxies::AddProxy(UControlRig* ControlRig, const FName& Name,  FRigControlElement* ControlElement)
 {
-	UControlRigControlsProxy* Proxy = FindProxy(Name);
+	UControlRigControlsProxy* Proxy = FindProxy(ControlRig,Name);
 	if (!Proxy && ControlElement != nullptr)
 	{
 		switch(ControlElement->Settings.ControlType)
@@ -774,54 +720,94 @@ void UControlRigDetailPanelControlProxies::AddProxy(const FName& Name, UControlR
 			Proxy->SetName(Name);
 			Proxy->ControlRig = ControlRig;
 			Proxy->ValueChanged();
-			AllProxies.Add(Name, Proxy);
-		}
 
+			FControlToProxyMap* ControlRigProxies = AllProxies.Find(ControlRig);
+			if (ControlRigProxies)
+			{
+				ControlRigProxies->ControlToProxy.Add(Name, Proxy);
+			}
+			else
+			{
+				FControlToProxyMap NewControlRigProxies;
+				NewControlRigProxies.ControlToProxy.Add(Name, Proxy);
+				AllProxies.Add(ControlRig, NewControlRigProxies);
+			}
+		}
 	}
 }
 
-void UControlRigDetailPanelControlProxies::RemoveProxy(const FName& Name)
+void UControlRigDetailPanelControlProxies::RemoveProxy(UControlRig* ControlRig, const FName& Name)
 {
-	UControlRigControlsProxy* ExistingProxy = FindProxy(Name);
+	UControlRigControlsProxy* ExistingProxy = FindProxy(ControlRig,Name);
 	if (ExistingProxy)
 	{
 		ExistingProxy->Rename(nullptr, GetTransientPackage(), REN_ForceNoResetLoaders);
 		ExistingProxy->MarkAsGarbage();
 	}
-	AllProxies.Remove(Name);
+	FControlToProxyMap* ControlRigProxies = AllProxies.Find(ControlRig);
+	if (ControlRigProxies)
+	{
+		ControlRigProxies->ControlToProxy.Remove(Name);
+	}
 }
 
-void UControlRigDetailPanelControlProxies::RemoveAllProxies()
+void UControlRigDetailPanelControlProxies::RemoveAllProxies(UControlRig* ControlRig)
 {
-	for (TPair<FName, TObjectPtr<UControlRigControlsProxy> >Pair : AllProxies)
+	//no control rig remove all
+	if (ControlRig == nullptr)
+	{
+		for (TPair<UControlRig*, FControlToProxyMap>& ControlRigProxies : AllProxies)
+		{
+			for (TPair<FName, TObjectPtr<UControlRigControlsProxy> >& Pair : ControlRigProxies.Value.ControlToProxy)
+			{
+				UControlRigControlsProxy* ExistingProxy = Pair.Value;
+				if (ExistingProxy)
+				{
+					ExistingProxy->Rename(nullptr, GetTransientPackage(), REN_ForceNoResetLoaders);
+					ExistingProxy->MarkAsGarbage();
+				}
+			}
+		}
+		AllProxies.Empty();
+		SelectedProxies.SetNum(0);
+	}
+	else
 	{
-		UControlRigControlsProxy* ExistingProxy = Pair.Value;
-		if (ExistingProxy)
+		FControlToProxyMap* ControlRigProxies = AllProxies.Find(ControlRig);
+		if (ControlRigProxies)
 		{
-			ExistingProxy->Rename(nullptr, GetTransientPackage(), REN_ForceNoResetLoaders);
-			ExistingProxy->MarkAsGarbage();
+			for (TPair<FName, TObjectPtr<UControlRigControlsProxy>>& Pair : ControlRigProxies->ControlToProxy)
+			{
+				UControlRigControlsProxy* ExistingProxy = Pair.Value;
+				if (ExistingProxy)
+				{
+					SelectedProxies.Remove(ExistingProxy);
+					ExistingProxy->Rename(nullptr, GetTransientPackage(), REN_ForceNoResetLoaders);
+					ExistingProxy->MarkAsGarbage();
+				}
+			}
+
+			AllProxies.Remove(ControlRig);
 		}
 	}
-	AllProxies.Empty();
-	SelectedProxies.SetNum(0);
 }
 
-void UControlRigDetailPanelControlProxies::RecreateAllProxies(UControlRig* InControlRig)
+void UControlRigDetailPanelControlProxies::RecreateAllProxies(UControlRig* ControlRig)
 {
-	RemoveAllProxies();
-	TArray<FRigControlElement*> Controls = InControlRig->AvailableControls();
+	RemoveAllProxies(ControlRig);
+	TArray<FRigControlElement*> Controls = ControlRig->AvailableControls();
 	for (FRigControlElement* ControlElement : Controls)
 	{
 		if(ControlElement->Settings.bShapeEnabled && ControlElement->Settings.bAnimatable)
 		{
-			AddProxy(ControlElement->GetName(), InControlRig, ControlElement);
+			AddProxy(ControlRig,ControlElement->GetName(), ControlElement);
 		}
 	}
 }
 
-void UControlRigDetailPanelControlProxies::ProxyChanged(const FName& Name)
+void UControlRigDetailPanelControlProxies::ProxyChanged(UControlRig* ControlRig, const FName& Name)
 {
-	UControlRigControlsProxy* Proxy = FindProxy(Name);
+	UControlRigControlsProxy* Proxy = FindProxy(ControlRig,Name);
 	if (Proxy)
 	{
 		Modify();
@@ -829,9 +815,9 @@ void UControlRigDetailPanelControlProxies::ProxyChanged(const FName& Name)
 	}
 }
 
-void UControlRigDetailPanelControlProxies::SelectProxy(const FName& Name, bool bSelected)
+void UControlRigDetailPanelControlProxies::SelectProxy(UControlRig* ControlRig,const FName& Name, bool bSelected)
 {
-	UControlRigControlsProxy* Proxy = FindProxy(Name);
+	UControlRigControlsProxy* Proxy = FindProxy(ControlRig,Name);
 	if (Proxy)
 	{
 		Modify();
diff --git a/Engine/Plugins/Animation/ControlRig/Source/ControlRigEditor/Private/EditMode/ControlRigControlsProxy.h b/Engine/Plugins/Animation/ControlRig/Source/ControlRigEditor/Private/EditMode/ControlRigControlsProxy.h
--- a/Engine/Plugins/Animation/ControlRig/Source/ControlRigEditor/Private/EditMode/ControlRigControlsProxy.h
+++ b/Engine/Plugins/Animation/ControlRig/Source/ControlRigEditor/Private/EditMode/ControlRigControlsProxy.h
@@ -53,7 +53,6 @@ public:
 protected:
 	bool bIsMultiple = 0;
 	FName Name;
-	void CheckEditModeOnSelectionChange(UControlRig* InControlRig);
 };
 
 UCLASS()
@@ -298,7 +297,14 @@ public:
 	bool Bool;
 };
 
+USTRUCT()
+struct FControlToProxyMap
+{
+	GENERATED_BODY();
 
+	UPROPERTY()
+	TMap <FName, TObjectPtr<UControlRigControlsProxy>> ControlToProxy;
+};
 /** Proxy in Details Panel */
 UCLASS()
 class UControlRigDetailPanelControlProxies :public UObject
@@ -309,20 +315,20 @@ class UControlRigDetailPanelControlProxies :public UObject
 protected:
 
 	UPROPERTY()
-	TMap<FName, TObjectPtr<UControlRigControlsProxy>> AllProxies;
+	TMap<UControlRig*, FControlToProxyMap> AllProxies; //proxies themselves contain weakobjectptr to the controlrig
 
 	UPROPERTY()
 	TArray< TObjectPtr<UControlRigControlsProxy>> SelectedProxies;
 
 
 public:
-	UControlRigControlsProxy* FindProxy(const FName& Name) const;
-	void AddProxy(const FName& Name, UControlRig* InControlRig, FRigControlElement* ControlElement);
-	void RemoveProxy(const FName& Name );
-	void ProxyChanged(const FName& Name);
-	void RemoveAllProxies();
+	UControlRigControlsProxy* FindProxy(UControlRig* InControlRig, const FName& Name) const;
+	void AddProxy(UControlRig* InControlRig, const FName& Name,  FRigControlElement* ControlElement);
+	void RemoveProxy(UControlRig* InControlRig, const FName& Name );
+	void ProxyChanged(UControlRig* InControlRig, const FName& Name);
+	void RemoveAllProxies(UControlRig* InControlRig);
 	void RecreateAllProxies(UControlRig* InControlRig);
-	void SelectProxy(const FName& Name, bool bSelected);
+	void SelectProxy(UControlRig* InControlRig, const FName& Name, bool bSelected);
 	const TArray<UControlRigControlsProxy*>& GetSelectedProxies() const { return SelectedProxies;}
 
 };
diff --git a/Engine/Plugins/Animation/ControlRig/Source/ControlRigEditor/Private/EditMode/ControlRigEditMode.cpp b/Engine/Plugins/Animation/ControlRig/Source/ControlRigEditor/Private/EditMode/ControlRigEditMode.cpp
--- a/Engine/Plugins/Animation/ControlRig/Source/ControlRigEditor/Private/EditMode/ControlRigEditMode.cpp
+++ b/Engine/Plugins/Animation/ControlRig/Source/ControlRigEditor/Private/EditMode/ControlRigEditMode.cpp
@@ -43,10 +43,6 @@
 #include "ControlRigComponent.h"
 #include "EngineUtils.h"
 #include "ControlRig/Private/Units/Execution/RigUnit_BeginExecution.h"
-
-//#include "IPersonaPreviewScene.h"
-//#include "Animation/DebugSkelMeshComponent.h"
-//#include "Persona/Private/AnimationEditorViewportClient.h"
 #include "IPersonaPreviewScene.h"
 #include "PersonaSelectionProxies.h"
 #include "Framework/Application/SlateApplication.h"
@@ -58,6 +54,8 @@
 #include "Sequencer/MovieSceneControlRigParameterSection.h"
 #include "SRigSpacePickerWidget.h"
 #include "ControlRigSpaceChannelEditors.h"
+#include "ControlRigSequencerEditorLibrary.h"
+#include "LevelSequence.h"
 
 void UControlRigEditModeDelegateHelper::OnPoseInitialized()
 {
@@ -132,12 +130,10 @@ namespace ControlRigSelectionConstants
 
 FControlRigEditMode::FControlRigEditMode()
 	: bIsChangingControlShapeTransform(false)
-	, InteractionScope(nullptr)
 	, bManipulatorMadeChange(false)
 	, bSelecting(false)
 	, bSelectionChanged(false)
-	, PivotTransform(FTransform::Identity)
-	, bRecreateControlShapesRequired(false)
+	, RecreateControlShapesRequired(ERecreateControlRigShape::RecreateNone)
 	, bSuspendHierarchyNotifs(false)
 	, CurrentViewportClient(nullptr)
 	, bIsChangingCoordSystem(false)
@@ -157,7 +153,7 @@ FControlRigEditMode::~FControlRigEditMode()
 {	
 	CommandBindings = nullptr;
 
-	DestroyShapesActors();
+	DestroyShapesActors(nullptr);
 	OnControlRigAddedOrRemovedDelegate.Clear();
 
 	TArray<TWeakObjectPtr<UControlRig>> PreviousRuntimeRigs = RuntimeControlRigs;
@@ -170,19 +166,68 @@ FControlRigEditMode::~FControlRigEditMode()
 	}
 	RuntimeControlRigs.Reset();
 
-	if (DelegateHelper.IsValid())
-	{
-		DelegateHelper->RemoveDelegates();
-		DelegateHelper.Reset();
-	}
-
 #if WITH_EDITOR
 	FCoreUObjectDelegates::OnObjectsReplaced.RemoveAll(this);
 #endif
 
 }
 
-void FControlRigEditMode::SetObjects(const TWeakObjectPtr<>& InSelectedObject,  UObject* BindingObject, TWeakPtr<ISequencer> InSequencer)
+bool FControlRigEditMode:: SetSequencer(TWeakPtr<ISequencer> InSequencer)
+{
+	if (InSequencer != WeakSequencer)
+	{
+		WeakSequencer = InSequencer;
+		DestroyShapesActors(nullptr);
+		TArray<TWeakObjectPtr<UControlRig>> PreviousRuntimeRigs = RuntimeControlRigs;
+		for (int32 PreviousRuntimeRigIndex = 0; PreviousRuntimeRigIndex < PreviousRuntimeRigs.Num(); PreviousRuntimeRigIndex++)
+		{
+			if (PreviousRuntimeRigs[PreviousRuntimeRigIndex].IsValid())
+			{
+				RemoveControlRig(PreviousRuntimeRigs[PreviousRuntimeRigIndex].Get());
+			}
+		}
+		RuntimeControlRigs.Reset();
+		if (InSequencer.IsValid())
+		{
+			TSharedPtr<ISequencer> Sequencer = WeakSequencer.Pin();
+			if (ULevelSequence* LevelSequence = Cast<ULevelSequence>(Sequencer->GetFocusedMovieSceneSequence()))
+			{
+				TArray<FControlRigSequencerBindingProxy> Proxies = UControlRigSequencerEditorLibrary::GetControlRigs(LevelSequence);
+				for (FControlRigSequencerBindingProxy& Proxy : Proxies)
+				{
+					if (UControlRig* ControlRig = Proxy.ControlRig.Get())
+					{
+						AddControlRigInternal(ControlRig);
+					}
+				}
+			}
+		}
+		SetObjects_Internal();
+	}
+	return false;
+}
+
+bool FControlRigEditMode::AddControlRigObject(UControlRig* ControlRig, TWeakPtr<ISequencer> InSequencer)
+{
+	if (ControlRig)
+	{
+		if (RuntimeControlRigs.Contains(ControlRig) == false)
+		{
+			if (InSequencer.IsValid())
+			{
+				if (SetSequencer(InSequencer) == false) //was already there so just add it,otherwise this function will add everything in the active 
+				{
+					AddControlRigInternal(ControlRig);
+					SetObjects_Internal();
+				}
+				return true;
+			}		
+		}
+	}
+	return false;
+}
+
+void FControlRigEditMode::SetObjects(UControlRig* ControlRig,  UObject* BindingObject, TWeakPtr<ISequencer> InSequencer)
 {
 	TArray<TWeakObjectPtr<UControlRig>> PreviousRuntimeRigs = RuntimeControlRigs;
 	for (int32 PreviousRuntimeRigIndex = 0; PreviousRuntimeRigIndex < PreviousRuntimeRigs.Num(); PreviousRuntimeRigIndex++)
@@ -194,11 +239,10 @@ void FControlRigEditMode::SetObjects(const TWeakObjectPtr<>& InSelectedObject,
 	}
 	RuntimeControlRigs.Reset();
 
-	UControlRig* ControlRig = Cast<UControlRig>(InSelectedObject.Get());
-
 	if (InSequencer.IsValid())
 	{
 		WeakSequencer = InSequencer;
+	
 	}
 	// if we get binding object, set it to control rig binding object
 	if (BindingObject && ControlRig)
@@ -211,11 +255,11 @@ void FControlRigEditMode::SetObjects(const TWeakObjectPtr<>& InSelectedObject,
 			}
 		}
 
-		AddControlRig(ControlRig);
+		AddControlRigInternal(ControlRig);
 	}
 	else if (ControlRig)
 	{
-		AddControlRig(ControlRig);
+		AddControlRigInternal(ControlRig);
 	}
 
 	SetObjects_Internal();
@@ -230,137 +274,62 @@ void FControlRigEditMode::SetUpDetailPanel()
 	if (IsInLevelEditor() && Toolkit)
 	{
 		StaticCastSharedPtr<SControlRigEditModeTools>(Toolkit->GetInlineContent())->SetSequencer(WeakSequencer.Pin());
-#if USE_LOCAL_DETAILS
-		TArray<TWeakObjectPtr<>> Eulers;
-		TArray<TWeakObjectPtr<>> Transforms;
-		TArray<TWeakObjectPtr<>> TransformNoScales;
-		TArray<TWeakObjectPtr<>> Floats;
-		TArray<TWeakObjectPtr<>> Vectors;
-		TArray<TWeakObjectPtr<>> Vector2Ds;
-		TArray<TWeakObjectPtr<>> Bools;
-		TArray<TWeakObjectPtr<>> Integers;
-		TArray<TWeakObjectPtr<>> Enums;
-		if (UControlRig* ControlRig = GetControlRig(true))
-		{
-			const TArray<UControlRigControlsProxy*>& Proxies = ControlProxy->GetSelectedProxies();
-			for (UControlRigControlsProxy* Proxy : Proxies)
-			{
-				if (Proxy->GetClass() == UControlRigTransformControlProxy::StaticClass())
-				{
-					Transforms.Add(Proxy);
-				}
-				else if (Proxy->GetClass() == UControlRigTransformNoScaleControlProxy::StaticClass())
-				{
-					TransformNoScales.Add(Proxy);
-				}
-				else if (Proxy->GetClass() == UControlRigEulerTransformControlProxy::StaticClass())
-				{
-					Eulers.Add(Proxy);
-				}
-				else if (Proxy->GetClass() == UControlRigFloatControlProxy::StaticClass())
-				{
-					Floats.Add(Proxy);
-				}
-				else if (Proxy->GetClass() == UControlRigVectorControlProxy::StaticClass())
-				{
-					Vectors.Add(Proxy);
-				}
-				else if (Proxy->GetClass() == UControlRigVector2DControlProxy::StaticClass())
-				{
-					Vector2Ds.Add(Proxy);
-				}
-				else if (Proxy->GetClass() == UControlRigBoolControlProxy::StaticClass())
-				{
-					Bools.Add(Proxy);
-				}
-				else if (Proxy->GetClass() == UControlRigEnumControlProxy::StaticClass())
-				{
-					Enums.Add(Proxy);
-				}
-				else if (Proxy->GetClass() == UControlRigIntegerControlProxy::StaticClass())
-				{
-					Integers.Add(Proxy);
-				}
-			}
-		}
-		for (TWeakObjectPtr<>& Object : Transforms)
-		{
-			UControlRigControlsProxy* Proxy = Cast<UControlRigControlsProxy>(Object.Get());
-			if (Proxy)
-			{
-				Proxy->SetIsMultiple(Transforms.Num() > 1);
-			}
-		}
-		StaticCastSharedPtr<SControlRigEditModeTools>(Toolkit->GetInlineContent())->SetTransformDetailsObjects(Transforms);
-		StaticCastSharedPtr<SControlRigEditModeTools>(Toolkit->GetInlineContent())->SetTransformNoScaleDetailsObjects(TransformNoScales);
-		StaticCastSharedPtr<SControlRigEditModeTools>(Toolkit->GetInlineContent())->SetEulerTransformDetailsObjects(Eulers);
-		StaticCastSharedPtr<SControlRigEditModeTools>(Toolkit->GetInlineContent())->SetFloatDetailsObjects(Floats);
-		StaticCastSharedPtr<SControlRigEditModeTools>(Toolkit->GetInlineContent())->SetVectorDetailsObjects(Vectors);
-		StaticCastSharedPtr<SControlRigEditModeTools>(Toolkit->GetInlineContent())->SetVector2DDetailsObjects(Vector2Ds);
-		StaticCastSharedPtr<SControlRigEditModeTools>(Toolkit->GetInlineContent())->SetBoolDetailsObjects(Bools);
-		StaticCastSharedPtr<SControlRigEditModeTools>(Toolkit->GetInlineContent())->SetIntegerDetailsObjects(Integers);
-		StaticCastSharedPtr<SControlRigEditModeTools>(Toolkit->GetInlineContent())->SetEnumDetailsObjects(Enums);
-#else 
 		StaticCastSharedPtr<SControlRigEditModeTools>(Toolkit->GetInlineContent())->SetSettingsDetailsObject(GetMutableDefault<UControlRigEditModeSettings>());	
-#endif
-
 	}
-
 }
 
 void FControlRigEditMode::SetObjects_Internal()
 {
-	for (TWeakObjectPtr<UControlRig> RuntimeRigPtr : RuntimeControlRigs)
+	bool bHasValidRuntimeControlRig = false;
+	for (TWeakObjectPtr<UControlRig>& RuntimeRigPtr : RuntimeControlRigs)
 	{
 		if (UControlRig* RuntimeControlRig = RuntimeRigPtr.Get())
 		{
-			if (UControlRig* InteractionRig = RuntimeControlRig->InteractionRig)
-			{
-				InteractionRig->ControlModified().RemoveAll(this);
-				InteractionRig->ControlModified().AddSP(this, &FControlRigEditMode::OnControlModified);
-			}
-			else
+			RuntimeControlRig->ControlModified().RemoveAll(this);
+			RuntimeControlRig->GetHierarchy()->OnModified().RemoveAll(this);
+
+			RuntimeControlRig->ControlModified().AddSP(this, &FControlRigEditMode::OnControlModified);
+			RuntimeControlRig->GetHierarchy()->OnModified().AddSP(this, &FControlRigEditMode::OnHierarchyModified);
+			if (USkeletalMeshComponent* MeshComponent = Cast<USkeletalMeshComponent>(GetHostingSceneComponent(RuntimeControlRig)))
 			{
-				RuntimeControlRig->ControlModified().RemoveAll(this);
-				RuntimeControlRig->ControlModified().AddSP(this, &FControlRigEditMode::OnControlModified);
+				TStrongObjectPtr<UControlRigEditModeDelegateHelper>* DelegateHelper = DelegateHelpers.Find(RuntimeControlRig);
+				if (!DelegateHelper)
+				{
+					DelegateHelpers.Add(RuntimeControlRig, TStrongObjectPtr<UControlRigEditModeDelegateHelper>(NewObject<UControlRigEditModeDelegateHelper>()));
+					DelegateHelper = DelegateHelpers.Find(RuntimeControlRig);
+				}
+				else if (DelegateHelper->IsValid() == false)
+				{
+					DelegateHelper->Get()->RemoveDelegates();
+					DelegateHelpers.Remove(RuntimeControlRig);
+					*DelegateHelper = TStrongObjectPtr<UControlRigEditModeDelegateHelper>(NewObject<UControlRigEditModeDelegateHelper>());
+					DelegateHelper->Get()->EditMode = this;
+					DelegateHelper->Get()->AddDelegates(MeshComponent);
+					DelegateHelpers.Add(RuntimeControlRig, *DelegateHelper);
+				}
+				
+				if (DelegateHelper && DelegateHelper->IsValid())
+				{
+					bHasValidRuntimeControlRig = true;
+				}
 			}
 		}
 	}
 
-	// currently all the manipulatable mesh component is supposed to be same
-	// if that changes, this code has to change
-	if (USkeletalMeshComponent* MeshComponent = Cast<USkeletalMeshComponent>(GetHostingSceneComponent()))
-	{
-		DelegateHelper->AddDelegates(MeshComponent);
-	}
-
-	UControlRig* RuntimeControlRig = GetControlRig(false);
-	
-	UControlRig* InteractionControlRig = GetControlRig(true);
-
 	if (UsesToolkits() && Toolkit.IsValid())
 	{
-		StaticCastSharedPtr<SControlRigEditModeTools>(Toolkit->GetInlineContent())->SetControlRig(RuntimeControlRig);
-	}
-
-	if (InteractionControlRig)
-	{
-		InteractionControlRig->GetHierarchy()->OnModified().RemoveAll(this);
-		InteractionControlRig->ControlModified().RemoveAll(this);
-			
-		InteractionControlRig->GetHierarchy()->OnModified().AddSP(this, &FControlRigEditMode::OnHierarchyModified);
-		InteractionControlRig->ControlModified().AddSP(this, &FControlRigEditMode::OnControlModified);
+		StaticCastSharedPtr<SControlRigEditModeTools>(Toolkit->GetInlineContent())->SetControlRigs(RuntimeControlRigs);
 	}
 
-	if (!RuntimeControlRig)
+	if (!bHasValidRuntimeControlRig)
 	{
-		DestroyShapesActors();
+		DestroyShapesActors(nullptr);
 		SetUpDetailPanel();
 	}
 	else
 	{
 		// create default manipulation layer
-		RequestToRecreateControlShapeActors();
+		RequestToRecreateControlShapeActors(nullptr);
 	}
 }
 
@@ -399,27 +368,22 @@ void FControlRigEditMode::Enter()
 		ModeManager->OnWidgetModeChanged().AddSP(this, &FControlRigEditMode::OnWidgetModeChanged);
 		ModeManager->OnCoordSystemChanged().AddSP(this, &FControlRigEditMode::OnCoordSystemChanged);
 	}
-
-	if (DelegateHelper.IsValid())
-	{
-		DelegateHelper->RemoveDelegates();
-		DelegateHelper.Reset();
-	}
-
-	DelegateHelper = TStrongObjectPtr<UControlRigEditModeDelegateHelper>(NewObject<UControlRigEditModeDelegateHelper>());
-	DelegateHelper->EditMode = this;
-
+	WorldPtr = GetWorld();
+	OnWorldCleanupHandle = FWorldDelegates::OnWorldCleanup.AddSP(this, &FControlRigEditMode::OnWorldCleanup);
 	SetObjects_Internal();
 }
 
 void FControlRigEditMode::Exit()
 {
-	if (UControlRig* ControlRig = GetControlRig(true))
+	for (TWeakObjectPtr<UControlRig>& RuntimeRigPtr : RuntimeControlRigs)
 	{
-		ControlRig->ClearControlSelection();
+		if (UControlRig* ControlRig = RuntimeRigPtr.Get())
+		{
+			ControlRig->ClearControlSelection();
+		}
 	}
 
-	if (InteractionScope)
+	if (InteractionScopes.Num() >0)
 	{
 
 		if (GEditor)
@@ -427,8 +391,14 @@ void FControlRigEditMode::Exit()
 			GEditor->EndTransaction();
 		}
 
-		delete InteractionScope;
-		InteractionScope = nullptr;
+		for (TPair<UControlRig*,FControlRigInteractionScope*>& InteractionScope : InteractionScopes)
+		{
+			if (InteractionScope.Value)
+			{
+				delete InteractionScope.Value;
+			}
+		}
+		InteractionScopes.Reset();
 		bManipulatorMadeChange = false;
 	}
 
@@ -438,7 +408,7 @@ void FControlRigEditMode::Exit()
 		Toolkit.Reset();
 	}
 
-	DestroyShapesActors();
+	DestroyShapesActors(nullptr);
 	OnControlRigAddedOrRemovedDelegate.Clear();
 
 	TArray<TWeakObjectPtr<UControlRig>> PreviousRuntimeRigs = RuntimeControlRigs;
@@ -451,19 +421,13 @@ void FControlRigEditMode::Exit()
 	}
 	RuntimeControlRigs.Reset();
 
-	if (DelegateHelper.IsValid())
-	{
-		DelegateHelper->RemoveDelegates();
-		DelegateHelper.Reset();
-	}
-
 	//clear delegates
 	FEditorModeTools* ModeManager = GetModeManager();
 	ModeManager->OnWidgetModeChanged().RemoveAll(this);
 	ModeManager->OnCoordSystemChanged().RemoveAll(this);
 
 	//clear proxies
-	ControlProxy->RemoveAllProxies();
+	ControlProxy->RemoveAllProxies(nullptr);
 
 	//make sure the widget is reset
 	ResetControlShapeSize();
@@ -496,31 +460,35 @@ void FControlRigEditMode::Tick(FEditorViewportClient* ViewportClient, float Delt
 	RecalcPivotTransform();
 
 	// Defer creation of shapes if manipulating the viewport
-	if (bRecreateControlShapesRequired && !(FSlateApplication::Get().HasAnyMouseCaptor() || GUnrealEd->IsUserInteracting()))
+	if (RecreateControlShapesRequired != ERecreateControlRigShape::RecreateNone && !(FSlateApplication::Get().HasAnyMouseCaptor() || GUnrealEd->IsUserInteracting()))
 	{
 		RecreateControlShapeActors();
-		TArray<FRigElementKey> SelectedRigElements = GetSelectedRigElements();
-		for (const FRigElementKey& SelectedKey : SelectedRigElements)
+		for (TWeakObjectPtr<UControlRig>& RuntimeRigPtr : RuntimeControlRigs)
 		{
-			if (SelectedKey.Type == ERigElementType::Control)
+			if (UControlRig* ControlRig = RuntimeRigPtr.Get())
 			{
-				AControlRigShapeActor* ShapeActor = GetControlShapeFromControlName(SelectedKey.Name);
-				if (ShapeActor)
+				TArray<FRigElementKey> SelectedRigElements = GetSelectedRigElements(ControlRig);
+				for (const FRigElementKey& SelectedKey : SelectedRigElements)
 				{
-					ShapeActor->SetSelected(true);
-				}
-
-				if (IsInLevelEditor())
-				{
-					if (UControlRig* ControlRig = GetControlRig(true))
+					if (SelectedKey.Type == ERigElementType::Control)
 					{
-						FRigControlElement* ControlElement = ControlRig->FindControl(SelectedKey.Name);
-						if (ControlElement)
+						AControlRigShapeActor* ShapeActor = GetControlShapeFromControlName(ControlRig,SelectedKey.Name);
+						if (ShapeActor)
 						{
-							if (!ControlRig->IsCurveControl(ControlElement))
+							ShapeActor->SetSelected(true);
+						}
+
+						if (IsInLevelEditor())
+						{
+							FRigControlElement* ControlElement = ControlRig->FindControl(SelectedKey.Name);
+							if (ControlElement)
 							{
-								ControlProxy->AddProxy(SelectedKey.Name, ControlRig, ControlElement);
+								if (!ControlRig->IsCurveControl(ControlElement))
+								{
+									ControlProxy->AddProxy(ControlRig, SelectedKey.Name, ControlElement);
+								}
 							}
+							
 						}
 					}
 				}
@@ -528,7 +496,8 @@ void FControlRigEditMode::Tick(FEditorViewportClient* ViewportClient, float Delt
 		}
 		SetUpDetailPanel();
 		HandleSelectionChanged();
-		bRecreateControlShapesRequired = false;
+		RecreateControlShapesRequired = ERecreateControlRigShape::RecreateNone;
+		ControlRigsToRecreate.SetNum(0);
 	}
 
 	// We need to tick here since changing a bone for example
@@ -536,21 +505,23 @@ void FControlRigEditMode::Tick(FEditorViewportClient* ViewportClient, float Delt
 	{
 		PostPoseUpdate();
 
-		if (UControlRig* ControlRig = GetControlRig(true))
+		for (TWeakObjectPtr<UControlRig>& RuntimeRigPtr : RuntimeControlRigs)
 		{
-			TArray<FRigElementKey> SelectedRigElements = GetSelectedRigElements();
-			const UE::Widget::EWidgetMode CurrentWidgetMode = ViewportClient->GetWidgetMode();
-			for (FRigElementKey SelectedRigElement : SelectedRigElements)
+			if (UControlRig* ControlRig = RuntimeRigPtr.Get())
 			{
-				//need to loop through the shape actors and set widget based upon the first one
-				if (AControlRigShapeActor* ShapeActor = GetControlShapeFromControlName(SelectedRigElement.Name))
+				TArray<FRigElementKey> SelectedRigElements = GetSelectedRigElements(ControlRig);
+				const UE::Widget::EWidgetMode CurrentWidgetMode = ViewportClient->GetWidgetMode();
+				for (FRigElementKey SelectedRigElement : SelectedRigElements)
 				{
-					if (!ModeSupportedByShapeActor(ShapeActor, CurrentWidgetMode))
+					//need to loop through the shape actors and set widget based upon the first one
+					if (AControlRigShapeActor* ShapeActor = GetControlShapeFromControlName(ControlRig,SelectedRigElement.Name))
 					{
-						if (FRigControlElement* ControlElement = ControlRig->FindControl(SelectedRigElement.Name))
+						if (!ModeSupportedByShapeActor(ShapeActor, CurrentWidgetMode))
 						{
-							switch (ControlElement->Settings.ControlType)
+							if (FRigControlElement* ControlElement = ControlRig->FindControl(SelectedRigElement.Name))
 							{
+								switch (ControlElement->Settings.ControlType)
+								{
 								case ERigControlType::Float:
 								case ERigControlType::Integer:
 								case ERigControlType::Vector2D:
@@ -572,20 +543,17 @@ void FControlRigEditMode::Tick(FEditorViewportClient* ViewportClient, float Delt
 									ViewportClient->SetWidgetMode(UE::Widget::WM_Scale);
 									break;
 								}
+								}
+								return; //exit if we switchted
 							}
-							return; //exit if we switchted
 						}
 					}
-					else
-					{
-						return; //exit if we are the same
-					}
-
 				}
 			}
 		}
 	}
 }
+
 //Hit proxy for FK Rigs and bones.
 struct  HFKRigBoneProxy : public HHitProxy
 {
@@ -684,210 +652,217 @@ TSet<FName> FControlRigEditMode::GetActiveControlsFromSequencer(UControlRig* Con
 
 void FControlRigEditMode::Render(const FSceneView* View, FViewport* Viewport, FPrimitiveDrawInterface* PDI)
 {	
-	UControlRig* ControlRig = GetControlRig(false);
-	if(ControlRig == nullptr)
-	{
-		return;
-	}
-
 	const UControlRigEditModeSettings* Settings = GetDefault<UControlRigEditModeSettings>();
-
 	bool bRender = !Settings->bHideControlShapes;
 
-	FTransform ComponentTransform = FTransform::Identity;
-	if(IsInLevelEditor())
-	{
-		ComponentTransform = GetHostingSceneComponentTransform();
-	}
-	
-	if (bRender)
+	for (TWeakObjectPtr<UControlRig>& ControlRigPtr : RuntimeControlRigs)
 	{
-		for (AControlRigShapeActor* Actor : ShapeActors)
+		UControlRig* ControlRig = ControlRigPtr.Get();
+		if (bRender && ControlRig && ControlRig->GetControlsVisible())
 		{
-			//Actor->SetActorHiddenInGame(bIsHidden);
-			if (GIsEditor && Actor->GetWorld() != nullptr && !Actor->GetWorld()->IsPlayInEditor())
+			FTransform ComponentTransform = FTransform::Identity;
+			if (IsInLevelEditor())
 			{
-				Actor->SetIsTemporarilyHiddenInEditor(false);
+				ComponentTransform = GetHostingSceneComponentTransform(ControlRig);
 			}
-		}
-
-		URigHierarchy* Hierarchy = ControlRig->GetHierarchy();
-		const bool bHasFKRig = (ControlRig->IsA<UAdditiveControlRig>() || ControlRig->IsA<UFKControlRig>());
-		if (Settings->bDisplayHierarchy || bHasFKRig)
-		{
-			const bool bBoolSetHitProxies = PDI && PDI->IsHitTesting() && bHasFKRig;
-			TSet<FName> ActiveControlName;
-			if (bHasFKRig)
+			const TArray<AControlRigShapeActor*>* ShapeActors = ControlRigShapeActors.Find(ControlRig);
+			if (ShapeActors)
 			{
-				ActiveControlName = GetActiveControlsFromSequencer(ControlRig);
+				for (AControlRigShapeActor* Actor : *ShapeActors)
+				{
+					//Actor->SetActorHiddenInGame(bIsHidden);
+					if (GIsEditor && Actor->GetWorld() != nullptr && !Actor->GetWorld()->IsPlayInEditor())
+					{
+						Actor->SetIsTemporarilyHiddenInEditor(false);
+					}
+				}
 			}
-			Hierarchy->ForEach<FRigTransformElement>([PDI, Hierarchy, ComponentTransform,ControlRig, bHasFKRig, bBoolSetHitProxies, ActiveControlName](FRigTransformElement* TransformElement) -> bool
-			{
-				const FTransform Transform = Hierarchy->GetTransform(TransformElement, ERigTransformType::CurrentGlobal);
 
-				FRigBaseElementParentArray Parents = Hierarchy->GetParents(TransformElement);
-				for(FRigBaseElement* ParentElement : Parents)
+			URigHierarchy* Hierarchy = ControlRig->GetHierarchy();
+			const bool bHasFKRig = (ControlRig->IsA<UAdditiveControlRig>() || ControlRig->IsA<UFKControlRig>());
+			if (Settings->bDisplayHierarchy || bHasFKRig)
+			{
+				const bool bBoolSetHitProxies = PDI && PDI->IsHitTesting() && bHasFKRig;
+				TSet<FName> ActiveControlName;
+				if (bHasFKRig)
 				{
-					if(FRigTransformElement* ParentTransformElement = Cast<FRigTransformElement>(ParentElement))
+					ActiveControlName = GetActiveControlsFromSequencer(ControlRig);
+				}
+				Hierarchy->ForEach<FRigTransformElement>([PDI, Hierarchy, ComponentTransform, ControlRig, bHasFKRig, bBoolSetHitProxies, ActiveControlName](FRigTransformElement* TransformElement) -> bool
 					{
+						const FTransform Transform = Hierarchy->GetTransform(TransformElement, ERigTransformType::CurrentGlobal);
+
+						FRigBaseElementParentArray Parents = Hierarchy->GetParents(TransformElement);
+						for (FRigBaseElement* ParentElement : Parents)
+						{
+							if (FRigTransformElement* ParentTransformElement = Cast<FRigTransformElement>(ParentElement))
+							{
+								FLinearColor Color = FLinearColor::White;
+								if (bHasFKRig)
+								{
+									FName ControlName = UFKControlRig::GetControlName(ParentTransformElement->GetName());
+									if (ActiveControlName.Num() > 0 && ActiveControlName.Contains(ControlName) == false)
+									{
+										continue;
+									}
+									if (ControlRig->IsControlSelected(ControlName))
+									{
+										Color = FLinearColor::Yellow;
+									}
+								}
+								const FTransform ParentTransform = Hierarchy->GetTransform(ParentTransformElement, ERigTransformType::CurrentGlobal);
+								const bool bHitTesting = bBoolSetHitProxies && (ParentTransformElement->GetType() == ERigElementType::Bone);
+								if (PDI)
+								{
+									if (bHitTesting)
+									{
+										PDI->SetHitProxy(new HFKRigBoneProxy(ParentTransformElement->GetName(), ControlRig));
+									}
+									PDI->DrawLine(ComponentTransform.TransformPosition(Transform.GetLocation()), ComponentTransform.TransformPosition(ParentTransform.GetLocation()), Color, SDPG_Foreground);
+									if (bHitTesting)
+									{
+										PDI->SetHitProxy(nullptr);
+									}
+								}
+							}
+						}
+
 						FLinearColor Color = FLinearColor::White;
 						if (bHasFKRig)
 						{
-							FName ControlName = UFKControlRig::GetControlName(ParentTransformElement->GetName());
+							FName ControlName = UFKControlRig::GetControlName(TransformElement->GetName());
 							if (ActiveControlName.Num() > 0 && ActiveControlName.Contains(ControlName) == false)
 							{
-								continue;
+								return true;
 							}
 							if (ControlRig->IsControlSelected(ControlName))
 							{
 								Color = FLinearColor::Yellow;
 							}
 						}
-						const FTransform ParentTransform = Hierarchy->GetTransform(ParentTransformElement, ERigTransformType::CurrentGlobal);
-						const bool bHitTesting = bBoolSetHitProxies && (ParentTransformElement->GetType() == ERigElementType::Bone);
 						if (PDI)
 						{
+							const bool bHitTesting = PDI && PDI->IsHitTesting() && bBoolSetHitProxies && (TransformElement->GetType() == ERigElementType::Bone);
 							if (bHitTesting)
 							{
-								PDI->SetHitProxy(new HFKRigBoneProxy(ParentTransformElement->GetName(), ControlRig));
+								PDI->SetHitProxy(new HFKRigBoneProxy(TransformElement->GetName(), ControlRig));
 							}
-							PDI->DrawLine(ComponentTransform.TransformPosition(Transform.GetLocation()),ComponentTransform.TransformPosition(ParentTransform.GetLocation()), Color, SDPG_Foreground);
+							PDI->DrawPoint(ComponentTransform.TransformPosition(Transform.GetLocation()), Color, 5.0f, SDPG_Foreground);
+
 							if (bHitTesting)
 							{
 								PDI->SetHitProxy(nullptr);
 							}
 						}
-					}
-                }
 
-				FLinearColor Color = FLinearColor::White;
-				if (bHasFKRig)
-				{
-					FName ControlName = UFKControlRig::GetControlName(TransformElement->GetName());
-					if (ActiveControlName.Num() > 0 && ActiveControlName.Contains(ControlName) == false)
-					{
 						return true;
-					}
-					if (ControlRig->IsControlSelected(ControlName))
-					{
-						Color = FLinearColor::Yellow;
-					}
-				}
-				const bool bHitTesting = PDI && PDI->IsHitTesting() && bBoolSetHitProxies && (TransformElement->GetType() == ERigElementType::Bone);
-				if (bHitTesting)
-				{
-					PDI->SetHitProxy(new HFKRigBoneProxy(TransformElement->GetName(), ControlRig));
-				}
-				PDI->DrawPoint(ComponentTransform.TransformPosition(Transform.GetLocation()), Color, 5.0f, SDPG_Foreground);
-
-				if (bHitTesting)
-				{
-					PDI->SetHitProxy(nullptr);
-				}
-
-				return true;
-			});
-		}
+					});
+			}
 
-		if (Settings->bDisplayNulls || ControlRig->IsSetupModeEnabled())
-		{
-			TArray<FTransform> SpaceTransforms;
-			TArray<FTransform> SelectedSpaceTransforms;
-			Hierarchy->ForEach<FRigNullElement>([&SpaceTransforms, &SelectedSpaceTransforms, Hierarchy](FRigNullElement* NullElement) -> bool
-            {
-				if(Hierarchy->IsSelected(NullElement->GetIndex()))
-				{
-					SelectedSpaceTransforms.Add(Hierarchy->GetTransform(NullElement, ERigTransformType::CurrentGlobal));
-				}
-				else
-				{
-					SpaceTransforms.Add(Hierarchy->GetTransform(NullElement, ERigTransformType::CurrentGlobal));
-				}
-				return true;
-			});
+			if (Settings->bDisplayNulls || ControlRig->IsSetupModeEnabled())
+			{
+				TArray<FTransform> SpaceTransforms;
+				TArray<FTransform> SelectedSpaceTransforms;
+				Hierarchy->ForEach<FRigNullElement>([&SpaceTransforms, &SelectedSpaceTransforms, Hierarchy](FRigNullElement* NullElement) -> bool
+					{
+						if (Hierarchy->IsSelected(NullElement->GetIndex()))
+						{
+							SelectedSpaceTransforms.Add(Hierarchy->GetTransform(NullElement, ERigTransformType::CurrentGlobal));
+						}
+						else
+						{
+							SpaceTransforms.Add(Hierarchy->GetTransform(NullElement, ERigTransformType::CurrentGlobal));
+						}
+						return true;
+					});
 
-			GetControlRig(true)->DrawInterface.DrawAxes(FTransform::Identity, SpaceTransforms, Settings->AxisScale);
-			GetControlRig(true)->DrawInterface.DrawAxes(FTransform::Identity, SelectedSpaceTransforms, FLinearColor(1.0f, 0.34f, 0.0f, 1.0f), Settings->AxisScale);
-		}
+				ControlRig->DrawInterface.DrawAxes(FTransform::Identity, SpaceTransforms, Settings->AxisScale);
+				ControlRig->DrawInterface.DrawAxes(FTransform::Identity, SelectedSpaceTransforms, FLinearColor(1.0f, 0.34f, 0.0f, 1.0f), Settings->AxisScale);
+			}
 
-		if (Settings->bDisplayAxesOnSelection && Settings->AxisScale > SMALL_NUMBER)
-		{
-			if (ControlRig->GetWorld() && ControlRig->GetWorld()->IsPreviewWorld())
+			if (Settings->bDisplayAxesOnSelection && Settings->AxisScale > SMALL_NUMBER)
 			{
-				TArray<FRigElementKey> SelectedRigElements = GetSelectedRigElements();
-				const float Scale = Settings->AxisScale;
-				PDI->AddReserveLines(SDPG_Foreground, SelectedRigElements.Num() * 3);
-
-				for (const FRigElementKey& SelectedElement : SelectedRigElements)
+				if (ControlRig->GetWorld() && ControlRig->GetWorld()->IsPreviewWorld())
 				{
-					FTransform ElementTransform = Hierarchy->GetGlobalTransform(SelectedElement);
-					ElementTransform = ElementTransform * ComponentTransform;
+					TArray<FRigElementKey> SelectedRigElements = GetSelectedRigElements(ControlRig);
+					const float Scale = Settings->AxisScale;
+					PDI->AddReserveLines(SDPG_Foreground, SelectedRigElements.Num() * 3);
+
+					for (const FRigElementKey& SelectedElement : SelectedRigElements)
+					{
+						FTransform ElementTransform = Hierarchy->GetGlobalTransform(SelectedElement);
+						ElementTransform = ElementTransform * ComponentTransform;
 
-					PDI->DrawLine(ElementTransform.GetTranslation(), ElementTransform.TransformPosition(FVector(Scale, 0.f, 0.f)), FLinearColor::Red, SDPG_Foreground);
-					PDI->DrawLine(ElementTransform.GetTranslation(), ElementTransform.TransformPosition(FVector(0.f, Scale, 0.f)), FLinearColor::Green, SDPG_Foreground);
-					PDI->DrawLine(ElementTransform.GetTranslation(), ElementTransform.TransformPosition(FVector(0.f, 0.f, Scale)), FLinearColor::Blue, SDPG_Foreground);
+						PDI->DrawLine(ElementTransform.GetTranslation(), ElementTransform.TransformPosition(FVector(Scale, 0.f, 0.f)), FLinearColor::Red, SDPG_Foreground);
+						PDI->DrawLine(ElementTransform.GetTranslation(), ElementTransform.TransformPosition(FVector(0.f, Scale, 0.f)), FLinearColor::Green, SDPG_Foreground);
+						PDI->DrawLine(ElementTransform.GetTranslation(), ElementTransform.TransformPosition(FVector(0.f, 0.f, Scale)), FLinearColor::Blue, SDPG_Foreground);
+					}
 				}
 			}
-		}
-		for (const FControlRigDrawInstruction& Instruction : GetControlRig(true)->DrawInterface)
-		{
-			if (!Instruction.IsValid())
+			for (const FControlRigDrawInstruction& Instruction : ControlRig->DrawInterface)
 			{
-				continue;
-			}
+				if (!Instruction.IsValid())
+				{
+					continue;
+				}
 
-			FTransform InstructionTransform = Instruction.Transform * ComponentTransform;
-			switch (Instruction.PrimitiveType)
-			{
+				FTransform InstructionTransform = Instruction.Transform * ComponentTransform;
+				switch (Instruction.PrimitiveType)
+				{
 				case EControlRigDrawSettings::Points:
-			{
-				for (const FVector& Point : Instruction.Positions)
 				{
+					for (const FVector& Point : Instruction.Positions)
+					{
 						PDI->DrawPoint(InstructionTransform.TransformPosition(Point), Instruction.Color, Instruction.Thickness, SDPG_Foreground);
+					}
+					break;
 				}
-				break;
-			}
 				case EControlRigDrawSettings::Lines:
-			{
-				const TArray<FVector>& Points = Instruction.Positions;
-					PDI->AddReserveLines(SDPG_Foreground, Points.Num() / 2, false, Instruction.Thickness > SMALL_NUMBER);
-				for (int32 PointIndex = 0; PointIndex < Points.Num() - 1; PointIndex += 2)
 				{
+					const TArray<FVector>& Points = Instruction.Positions;
+					PDI->AddReserveLines(SDPG_Foreground, Points.Num() / 2, false, Instruction.Thickness > SMALL_NUMBER);
+					for (int32 PointIndex = 0; PointIndex < Points.Num() - 1; PointIndex += 2)
+					{
 						PDI->DrawLine(InstructionTransform.TransformPosition(Points[PointIndex]), InstructionTransform.TransformPosition(Points[PointIndex + 1]), Instruction.Color, SDPG_Foreground, Instruction.Thickness);
+					}
+					break;
 				}
-				break;
-			}
 				case EControlRigDrawSettings::LineStrip:
-			{
-				const TArray<FVector>& Points = Instruction.Positions;
-					PDI->AddReserveLines(SDPG_Foreground, Points.Num() - 1, false, Instruction.Thickness > SMALL_NUMBER);
-				for (int32 PointIndex = 0; PointIndex < Points.Num() - 1; PointIndex++)
 				{
+					const TArray<FVector>& Points = Instruction.Positions;
+					PDI->AddReserveLines(SDPG_Foreground, Points.Num() - 1, false, Instruction.Thickness > SMALL_NUMBER);
+					for (int32 PointIndex = 0; PointIndex < Points.Num() - 1; PointIndex++)
+					{
 						PDI->DrawLine(InstructionTransform.TransformPosition(Points[PointIndex]), InstructionTransform.TransformPosition(Points[PointIndex + 1]), Instruction.Color, SDPG_Foreground, Instruction.Thickness);
+					}
+					break;
 				}
-				break;
-			}
 
 				case EControlRigDrawSettings::DynamicMesh:
-			{
-				FDynamicMeshBuilder MeshBuilder(PDI->View->GetFeatureLevel());
-				MeshBuilder.AddVertices(Instruction.MeshVerts);
-				MeshBuilder.AddTriangles(Instruction.MeshIndices);
-				MeshBuilder.Draw(PDI, InstructionTransform.ToMatrixWithScale(), Instruction.MaterialRenderProxy, SDPG_World/*SDPG_Foreground*/);
-				break;
-			}
+				{
+					FDynamicMeshBuilder MeshBuilder(PDI->View->GetFeatureLevel());
+					MeshBuilder.AddVertices(Instruction.MeshVerts);
+					MeshBuilder.AddTriangles(Instruction.MeshIndices);
+					MeshBuilder.Draw(PDI, InstructionTransform.ToMatrixWithScale(), Instruction.MaterialRenderProxy, SDPG_World/*SDPG_Foreground*/);
+					break;
+				}
 
+				}
 			}
 		}
-	}
-	else
-	{
-		for (AControlRigShapeActor* Actor : ShapeActors)
+		else
 		{
-			//Actor->SetActorHiddenInGame(bIsHidden);
-			if (GIsEditor && Actor->GetWorld() != nullptr && !Actor->GetWorld()->IsPlayInEditor())
+			const TArray<AControlRigShapeActor*>* ShapeActors = ControlRigShapeActors.Find(ControlRig);
+			if (ShapeActors)
 			{
-				Actor->SetIsTemporarilyHiddenInEditor(true);
+				for (AControlRigShapeActor* Actor : *ShapeActors)
+				{
+					//Actor->SetActorHiddenInGame(bIsHidden);
+					if (GIsEditor && Actor->GetWorld() != nullptr && !Actor->GetWorld()->IsPlayInEditor())
+					{
+						Actor->SetIsTemporarilyHiddenInEditor(true);
+					}
+				}
 			}
 		}
 	}
@@ -911,7 +886,8 @@ bool FControlRigEditMode::InputKey(FEditorViewportClient* InViewportClient, FVie
 
 bool FControlRigEditMode::EndTracking(FEditorViewportClient* InViewportClient, FViewport* InViewport)
 {
-	if (InteractionScope)
+
+	if (InteractionScopes.Num() > 0)
 	{
 		if (bManipulatorMadeChange)
 		{
@@ -919,8 +895,15 @@ bool FControlRigEditMode::EndTracking(FEditorViewportClient* InViewportClient, F
 			GEditor->EndTransaction();
 		}
 
-		delete InteractionScope;
-		InteractionScope = nullptr;
+
+		for (TPair<UControlRig*, FControlRigInteractionScope*>& InteractionScope : InteractionScopes)
+		{
+			if (InteractionScope.Value)
+			{
+				delete InteractionScope.Value; 
+			}
+		}
+		InteractionScopes.Reset();
 
 		return true;
 	}
@@ -932,109 +915,151 @@ bool FControlRigEditMode::EndTracking(FEditorViewportClient* InViewportClient, F
 
 bool FControlRigEditMode::StartTracking(FEditorViewportClient* InViewportClient, FViewport* InViewport)
 {
-	if (InteractionScope == nullptr)
+	if (InteractionScopes.Num() == 0)
 	{
-	    if (UControlRig* ControlRig = GetControlRig(true))
+		bool bShouldModify = IsInLevelEditor();
+		if (!bShouldModify)
 		{
-			bool bShouldModify = IsInLevelEditor();
-			if (!bShouldModify)
+			for (TWeakObjectPtr<UControlRig>& RuntimeRigPtr : RuntimeControlRigs)
 			{
-				TArray<FRigElementKey> SelectedRigElements = GetSelectedRigElements();
-				for (const FRigElementKey& Key : SelectedRigElements)
+				if (UControlRig* ControlRig = RuntimeRigPtr.Get())
 				{
-					if (Key.Type != ERigElementType::Control)
+					TArray<FRigElementKey> SelectedRigElements = GetSelectedRigElements(ControlRig);
+					for (const FRigElementKey& Key : SelectedRigElements)
 					{
-						bShouldModify = true;
+						if (Key.Type != ERigElementType::Control)
+						{
+							bShouldModify = true;
+							break;
+						}
 					}
 				}
 			}
+		}
 
-			if (!IsInLevelEditor())
+		if (!IsInLevelEditor())
+		{
+			for (TWeakObjectPtr<UControlRig>& RuntimeRigPtr : RuntimeControlRigs)
 			{
-				UObject* Blueprint = ControlRig->GetClass()->ClassGeneratedBy;
-				if (Blueprint)
+				if (UControlRig* ControlRig = RuntimeRigPtr.Get())
 				{
-					Blueprint->SetFlags(RF_Transactional);
-					if (bShouldModify)
+
+					UObject* Blueprint = ControlRig->GetClass()->ClassGeneratedBy;
+					if (Blueprint)
 					{
-						Blueprint->Modify();
-					}
+						Blueprint->SetFlags(RF_Transactional);
+						if (bShouldModify)
+						{
+							Blueprint->Modify();
+						}
+					}
+					ControlRig->SetFlags(RF_Transactional);
+					if (bShouldModify)
+					{
+						ControlRig->Modify();
+					}
 				}
 			}
-
-			ControlRig->SetFlags(RF_Transactional);
-			if (bShouldModify)
-			{
-				ControlRig->Modify();
-			}
 		}
 
-		//in level editor only transact if we have at least one control selected, in editor we only select CR stuff so always transact
-		if (UControlRig* ControlRig = GetControlRig(true))
+	}
+
+	//in level editor only transact if we have at least one control selected, in editor we only select CR stuff so always transact
+
+	if (IsInLevelEditor())
+	{
+		for (TWeakObjectPtr<UControlRig>& RuntimeRigPtr : RuntimeControlRigs)
 		{
-			if (IsInLevelEditor())
+			if (UControlRig* ControlRig = RuntimeRigPtr.Get())
 			{
-				if (AreRigElementSelectedAndMovable())
+				if (AreRigElementSelectedAndMovable(ControlRig))
 				{
-					InteractionScope = new FControlRigInteractionScope(ControlRig);
+					//todo need to add multiple
+					FControlRigInteractionScope* InteractionScope = new FControlRigInteractionScope(ControlRig);
+					InteractionScopes.Add(ControlRig,InteractionScope);
+				}
+				else
+				{
+					bManipulatorMadeChange = false;
 				}
 			}
-			else if(UControlRigEditorSettings::Get()->bEnableUndoForPoseInteraction)
-			{
-				InteractionScope = new FControlRigInteractionScope(ControlRig);
-			}
-
-			bManipulatorMadeChange = false;
 		}
-
-		return InteractionScope != nullptr;
 	}
-
-	return false;
+	else if(UControlRigEditorSettings::Get()->bEnableUndoForPoseInteraction)
+	{
+		UControlRig* ControlRig = RuntimeControlRigs[0].Get();
+		FControlRigInteractionScope* InteractionScope = new FControlRigInteractionScope(ControlRig);
+		InteractionScopes.Add(ControlRig,InteractionScope);
+	}
+	else
+	{
+		bManipulatorMadeChange = false;
+	}
+	return InteractionScopes.Num() != 0;
 }
 
 bool FControlRigEditMode::UsesTransformWidget() const
 {
-	for (const AControlRigShapeActor* ShapeActor : ShapeActors)
+	for (const TPair<UControlRig*, TArray<AControlRigShapeActor*>>& Pairs : ControlRigShapeActors)
 	{
-		if (ShapeActor->IsSelected())
+		for (const AControlRigShapeActor* ShapeActor : Pairs.Value)
+		{
+			if (ShapeActor->IsSelected())
+			{
+				return true;
+			}
+		}
+		if (AreRigElementSelectedAndMovable(Pairs.Key))
 		{
 			return true;
 		}
 	}
-
-	if (AreRigElementSelectedAndMovable())
-	{
-		return true;
-	}
-
 	return FEdMode::UsesTransformWidget();
 }
 
 bool FControlRigEditMode::UsesTransformWidget(UE::Widget::EWidgetMode CheckMode) const
 {
-	for (const AControlRigShapeActor* ShapeActor : ShapeActors)
+	for (const TPair<UControlRig*, TArray<AControlRigShapeActor*>>& Pairs : ControlRigShapeActors)
 	{
-		if (ShapeActor->IsSelected())
+		for (const AControlRigShapeActor* ShapeActor : Pairs.Value)
 		{
-			return ModeSupportedByShapeActor(ShapeActor, CheckMode);
+			if (ShapeActor->IsSelected())
+			{
+				return ModeSupportedByShapeActor(ShapeActor, CheckMode);
+			}
+		}
+		if (AreRigElementSelectedAndMovable(Pairs.Key))
+		{
+			return true;
 		}
 	}
-
-	if (AreRigElementSelectedAndMovable())
-	{
-		return true;
-	}
-
 	return FEdMode::UsesTransformWidget(CheckMode);
 }
 
 FVector FControlRigEditMode::GetWidgetLocation() const
 {
-	if (AreRigElementSelectedAndMovable())
+	FVector PivotLocation(0.0, 0.0, 0.0);
+	int NumSelected = 0;
+	for (const TPair<UControlRig*, TArray<AControlRigShapeActor*>>& Pairs : ControlRigShapeActors)
 	{
-		FTransform ComponentTransform = GetHostingSceneComponentTransform();
-		return ComponentTransform.TransformPosition(PivotTransform.GetLocation());
+		if (AreRigElementSelectedAndMovable(Pairs.Key))
+		{
+			if (const FTransform* PivotTransform = PivotTransforms.Find(Pairs.Key))
+			{
+				FTransform ComponentTransform = GetHostingSceneComponentTransform(Pairs.Key);
+				PivotLocation += ComponentTransform.TransformPosition(PivotTransform->GetLocation());
+				++NumSelected;
+			}
+		}
+	}	
+	if (NumSelected > 0)
+	{
+		PivotLocation /= (NumSelected);
+		return PivotLocation;
+	}
+	else
+	{
+		PivotLocation = FVector(0, 0, 0);
 	}
 
 	return FEdMode::GetWidgetLocation();
@@ -1042,12 +1067,18 @@ FVector FControlRigEditMode::GetWidgetLocation() const
 
 bool FControlRigEditMode::GetCustomDrawingCoordinateSystem(FMatrix& OutMatrix, void* InData)
 {
-	if (AreRigElementSelectedAndMovable())
+	//since we strip translation just want the first one
+	for (const TPair<UControlRig*, TArray<AControlRigShapeActor*>>& Pairs : ControlRigShapeActors)
 	{
-		OutMatrix = PivotTransform.ToMatrixNoScale().RemoveTranslation();
-		return true;
+		if (AreRigElementSelectedAndMovable(Pairs.Key))
+		{
+			if (const FTransform* PivotTransform = PivotTransforms.Find(Pairs.Key))
+			{
+				OutMatrix = PivotTransform->ToMatrixNoScale().RemoveTranslation();
+				return true;
+			}
+		}
 	}
-
 	return false;
 }
 
@@ -1065,30 +1096,33 @@ bool FControlRigEditMode::HandleClick(FEditorViewportClient* InViewportClient, H
 			if (ActorHitProxy->Actor->IsA<AControlRigShapeActor>())
 			{
 				AControlRigShapeActor* ShapeActor = CastChecked<AControlRigShapeActor>(ActorHitProxy->Actor);
-				if (ShapeActor->IsSelectable())
+				if (ShapeActor->IsSelectable() && ShapeActor->ControlRig.IsValid())
 				{
 					FScopedTransaction ScopedTransaction(LOCTEXT("SelectControlTransaction", "Select Control"), IsInLevelEditor() && !GIsTransacting);
 
 					// temporarily disable the interaction scope
-					const bool bInteractionScopePresent = InteractionScope != nullptr; 
+					FControlRigInteractionScope** InteractionScope = InteractionScopes.Find(ShapeActor->ControlRig.Get());
+					const bool bInteractionScopePresent = InteractionScope != nullptr;
 					if (bInteractionScopePresent)
 					{
-						delete InteractionScope;
-						InteractionScope = nullptr;
+						delete* InteractionScope;
+						InteractionScopes.Remove(ShapeActor->ControlRig.Get());
 					}
 					
 					const FName& ControlName = ShapeActor->ControlName;
 					if (Click.IsShiftDown()) //guess we just select
 					{
-						SetRigElementSelection(ERigElementType::Control, ControlName, true);
+						SetRigElementSelection(ShapeActor->ControlRig.Get(),ERigElementType::Control, ControlName, true);
 					}
 					else if(Click.IsControlDown()) //if ctrl we toggle selection
 					{
-						UControlRig* InteractionRig = GetControlRig(true);
-						if (InteractionRig)
+						for (TWeakObjectPtr<UControlRig>& RuntimeRigPtr : RuntimeControlRigs)
 						{
-							bool bIsSelected = InteractionRig->IsControlSelected(ControlName);
-							SetRigElementSelection(ERigElementType::Control, ControlName, !bIsSelected);
+							if (UControlRig* ControlRig = RuntimeRigPtr.Get())
+							{
+								bool bIsSelected = ControlRig->IsControlSelected(ControlName);
+								SetRigElementSelection(ShapeActor->ControlRig.Get(), ERigElementType::Control, ControlName, !bIsSelected);
+							}
 						}
 					}
 					else
@@ -1104,15 +1138,14 @@ bool FControlRigEditMode::HandleClick(FEditorViewportClient* InViewportClient, H
 							}
 						}
 						ClearRigElementSelection(FRigElementTypeHelper::ToMask(ERigElementType::Control));
-						SetRigElementSelection(ERigElementType::Control, ControlName, true);
+						SetRigElementSelection(ShapeActor->ControlRig.Get(),ERigElementType::Control, ControlName, true);
 					}
 
 					if (bInteractionScopePresent)
 					{
-						if(UControlRig* ControlRig = GetControlRig(true))
-						{
-							InteractionScope = new FControlRigInteractionScope(ControlRig);
-						}
+						*InteractionScope = new FControlRigInteractionScope(ShapeActor->ControlRig.Get());
+						InteractionScopes.Add(ShapeActor->ControlRig.Get(), *InteractionScope);
+
 					}
 
 					// for now we show this menu all the time if body is selected
@@ -1127,49 +1160,52 @@ bool FControlRigEditMode::HandleClick(FEditorViewportClient* InViewportClient, H
 
 				return true;
 			}
-			else if(UControlRig* ControlRig = GetControlRig(false))
+			else 
 			{ 
-				//if we have an additive or fk control rig active select the control based upon the selected bone.
-				UAdditiveControlRig* AdditiveControlRig = Cast<UAdditiveControlRig>(ControlRig);
-				UFKControlRig* FKControlRig = Cast<UFKControlRig>(ControlRig);
-
-				if (AdditiveControlRig || FKControlRig)
+				for (TWeakObjectPtr<UControlRig>& RuntimeRigPtr : RuntimeControlRigs)
 				{
-					if (USkeletalMeshComponent* RigMeshComp = Cast<USkeletalMeshComponent>(ControlRig->GetObjectBinding()->GetBoundObject()))
+					if (UControlRig* ControlRig = RuntimeRigPtr.Get())
 					{
-						const USkeletalMeshComponent* SkelComp = Cast<USkeletalMeshComponent>(ActorHitProxy->PrimComponent);
 
-						if (SkelComp == RigMeshComp)
-						{
-							FHitResult Result(1.0f);
-							bool bHit = RigMeshComp->LineTraceComponent(Result, Click.GetOrigin(), Click.GetOrigin() + Click.GetDirection() * ControlRigSelectionConstants::BodyTraceDistance, FCollisionQueryParams(NAME_None, FCollisionQueryParams::GetUnknownStatId(), true));
+						//if we have an additive or fk control rig active select the control based upon the selected bone.
+						UAdditiveControlRig* AdditiveControlRig = Cast<UAdditiveControlRig>(ControlRig);
+						UFKControlRig* FKControlRig = Cast<UFKControlRig>(ControlRig);
 
-							if (bHit)
+						if (AdditiveControlRig || FKControlRig)
+						{
+							if (USkeletalMeshComponent* RigMeshComp = Cast<USkeletalMeshComponent>(ControlRig->GetObjectBinding()->GetBoundObject()))
 							{
-								FName ControlName(*(Result.BoneName.ToString() + TEXT("_CONTROL")));
-								if (ControlRig->FindControl(ControlName))
+								const USkeletalMeshComponent* SkelComp = Cast<USkeletalMeshComponent>(ActorHitProxy->PrimComponent);
+
+								if (SkelComp == RigMeshComp)
 								{
-									FScopedTransaction ScopedTransaction(LOCTEXT("SelectControlTransaction", "Select Control"), IsInLevelEditor() && !GIsTransacting);
+									FHitResult Result(1.0f);
+									bool bHit = RigMeshComp->LineTraceComponent(Result, Click.GetOrigin(), Click.GetOrigin() + Click.GetDirection() * ControlRigSelectionConstants::BodyTraceDistance, FCollisionQueryParams(NAME_None, FCollisionQueryParams::GetUnknownStatId(), true));
 
-									if (Click.IsShiftDown()) //guess we just select
-									{
-										SetRigElementSelection(ERigElementType::Control, ControlName, true);
-									}
-									else if (Click.IsControlDown()) //if ctrl we toggle selection
+									if (bHit)
 									{
-										UControlRig* InteractionRig = GetControlRig(true);
-										if (InteractionRig)
+										FName ControlName(*(Result.BoneName.ToString() + TEXT("_CONTROL")));
+										if (ControlRig->FindControl(ControlName))
 										{
-											bool bIsSelected = InteractionRig->IsControlSelected(ControlName);
-											SetRigElementSelection(ERigElementType::Control, ControlName, !bIsSelected);
+											FScopedTransaction ScopedTransaction(LOCTEXT("SelectControlTransaction", "Select Control"), IsInLevelEditor() && !GIsTransacting);
+
+											if (Click.IsShiftDown()) //guess we just select
+											{
+												SetRigElementSelection(ControlRig,ERigElementType::Control, ControlName, true);
+											}
+											else if (Click.IsControlDown()) //if ctrl we toggle selection
+											{
+												bool bIsSelected = ControlRig->IsControlSelected(ControlName);
+												SetRigElementSelection(ControlRig, ERigElementType::Control, ControlName, !bIsSelected);
+											}
+											else
+											{
+												ClearRigElementSelection(FRigElementTypeHelper::ToMask(ERigElementType::Control));
+												SetRigElementSelection(ControlRig,ERigElementType::Control, ControlName, true);
+											}
+											return true;
 										}
 									}
-									else
-									{
-										ClearRigElementSelection(FRigElementTypeHelper::ToMask(ERigElementType::Control));
-										SetRigElementSelection(ERigElementType::Control, ControlName, true);
-									}
-									return true;
 								}
 							}
 						}
@@ -1187,53 +1223,58 @@ bool FControlRigEditMode::HandleClick(FEditorViewportClient* InViewportClient, H
 
 			if (Click.IsShiftDown()) //guess we just select
 			{
-				SetRigElementSelection(ERigElementType::Control, ControlName, true);
+				SetRigElementSelection(FKBoneProxy->ControlRig,ERigElementType::Control, ControlName, true);
 			}
 			else if (Click.IsControlDown()) //if ctrl we toggle selection
 			{
-				UControlRig* InteractionRig = GetControlRig(true);
-				if (InteractionRig)
+				for (TWeakObjectPtr<UControlRig>& RuntimeRigPtr : RuntimeControlRigs)
 				{
-					bool bIsSelected = InteractionRig->IsControlSelected(ControlName);
-					SetRigElementSelection(ERigElementType::Control, ControlName, !bIsSelected);
+					if (UControlRig* ControlRig = RuntimeRigPtr.Get())
+					{
+						bool bIsSelected = ControlRig->IsControlSelected(ControlName);
+						SetRigElementSelection(FKBoneProxy->ControlRig, ERigElementType::Control, ControlName, !bIsSelected);
+					}
 				}
 			}
 			else
 			{
 				ClearRigElementSelection(FRigElementTypeHelper::ToMask(ERigElementType::Control));
-				SetRigElementSelection(ERigElementType::Control, ControlName, true);
+				SetRigElementSelection(FKBoneProxy->ControlRig,ERigElementType::Control, ControlName, true);
 			}
 			return true;
 		}
 	}
 	else if (HPersonaBoneHitProxy* BoneHitProxy = HitProxyCast<HPersonaBoneHitProxy>(HitProxy))
 	{
-		if (UControlRig* DebuggedControlRig = GetControlRig(false))
+		if (RuntimeControlRigs.Num() > 0)
 		{
-			URigHierarchy* Hierarchy = DebuggedControlRig->GetHierarchy();
-
-			// Cache mapping?
-			for (int32 Index = 0; Index < Hierarchy->Num(); Index++)
+			if (UControlRig* DebuggedControlRig = RuntimeControlRigs[0].Get())
 			{
-				const FRigElementKey ElementToSelect = Hierarchy->GetKey(Index);
-				if (ElementToSelect.Type == ERigElementType::Bone && ElementToSelect.Name == BoneHitProxy->BoneName)
+				URigHierarchy* Hierarchy = DebuggedControlRig->GetHierarchy();
+
+				// Cache mapping?
+				for (int32 Index = 0; Index < Hierarchy->Num(); Index++)
 				{
-					if (FSlateApplication::Get().GetModifierKeys().IsShiftDown())
-					{
-						Hierarchy->GetController()->SelectElement(ElementToSelect, true);
-					}
-					else if (FSlateApplication::Get().GetModifierKeys().IsControlDown())
+					const FRigElementKey ElementToSelect = Hierarchy->GetKey(Index);
+					if (ElementToSelect.Type == ERigElementType::Bone && ElementToSelect.Name == BoneHitProxy->BoneName)
 					{
-						const bool bSelect = !Hierarchy->IsSelected(ElementToSelect);
-						Hierarchy->GetController()->SelectElement(ElementToSelect, bSelect);
-					}
-					else
-					{
-						TArray<FRigElementKey> NewSelection;
-						NewSelection.Add(ElementToSelect);
-						Hierarchy->GetController()->SetSelection(NewSelection);
+						if (FSlateApplication::Get().GetModifierKeys().IsShiftDown())
+						{
+							Hierarchy->GetController()->SelectElement(ElementToSelect, true);
+						}
+						else if (FSlateApplication::Get().GetModifierKeys().IsControlDown())
+						{
+							const bool bSelect = !Hierarchy->IsSelected(ElementToSelect);
+							Hierarchy->GetController()->SelectElement(ElementToSelect, bSelect);
+						}
+						else
+						{
+							TArray<FRigElementKey> NewSelection;
+							NewSelection.Add(ElementToSelect);
+							Hierarchy->GetController()->SetSelection(NewSelection);
+						}
+						return true;
 					}
-					return true;
 				}
 			}
 		}
@@ -1332,21 +1373,24 @@ void FControlRigEditMode::OpenContextMenu(FEditorViewportClient* InViewportClien
 
 bool FControlRigEditMode::IntersectSelect(bool InSelect, const TFunctionRef<bool(const AControlRigShapeActor*, const FTransform&)>& Intersects)
 {
-	FTransform ComponentTransform = GetHostingSceneComponentTransform();
-
 	bool bSelected = false;
-	for (AControlRigShapeActor* ShapeActor : ShapeActors)
+
+	for (TPair<UControlRig*, TArray<AControlRigShapeActor*>>& Pairs : ControlRigShapeActors)
 	{
-		if (ShapeActor->IsHiddenEd())
+		FTransform ComponentTransform = GetHostingSceneComponentTransform(Pairs.Key);
+		for (AControlRigShapeActor* ShapeActor : Pairs.Value)
 		{
-			continue;
-		}
+			if (ShapeActor->IsHiddenEd())
+			{
+				continue;
+			}
 
-		const FTransform ControlTransform = ShapeActor->GetGlobalTransform() * ComponentTransform;
-		if (Intersects(ShapeActor, ControlTransform))
-		{
-			SetRigElementSelection(ERigElementType::Control, ShapeActor->ControlName, InSelect);
-			bSelected = true;
+			const FTransform ControlTransform = ShapeActor->GetGlobalTransform() * ComponentTransform;
+			if (Intersects(ShapeActor, ControlTransform))
+			{
+				SetRigElementSelection(Pairs.Key,ERigElementType::Control, ShapeActor->ControlName, InSelect);
+				bSelected = true;
+			}
 		}
 	}
 
@@ -1430,7 +1474,7 @@ bool FControlRigEditMode::BoxSelect(FBox& InBox, bool InSelect)
 		}
 
 		AControlRigShapeActor* ShapeActor = CastChecked<AControlRigShapeActor>(Actor);
-		if (!ShapeActor->IsSelectable())
+		if (!ShapeActor->IsSelectable() || ShapeActor->ControlRig.IsValid () == false)
 		{
 			continue;
 		}
@@ -1439,14 +1483,14 @@ bool FControlRigEditMode::BoxSelect(FBox& InBox, bool InSelect)
 		{
 			bSomethingSelected = true;
 			const FName& ControlName = ShapeActor->ControlName;
-			SetRigElementSelection(ERigElementType::Control, ControlName, true);
+			SetRigElementSelection(ShapeActor->ControlRig.Get(),ERigElementType::Control, ControlName, true);
 
 			if (bShiftDown)
 			{
 			}
 			else
 			{
-				SetRigElementSelection(ERigElementType::Control, ControlName, true);
+				SetRigElementSelection(ShapeActor->ControlRig.Get(),ERigElementType::Control, ControlName, true);
 			}
 		}
 	}
@@ -1469,20 +1513,23 @@ bool FControlRigEditMode::FrustumSelect(const FConvexVolume& InFrustum, FEditorV
 		ClearRigElementSelection(FRigElementTypeHelper::ToMask(ERigElementType::Control));
 	}
 
-	for (AControlRigShapeActor* ShapeActor : ShapeActors)
+	for (TPair<UControlRig*, TArray<AControlRigShapeActor*>>& Pairs : ControlRigShapeActors)
 	{
-		for (UActorComponent* Component : ShapeActor->GetComponents())
+		for (AControlRigShapeActor* ShapeActor : Pairs.Value)
 		{
-			UPrimitiveComponent* PrimitiveComponent = Cast<UPrimitiveComponent>(Component);
-			if (PrimitiveComponent && PrimitiveComponent->IsRegistered() && PrimitiveComponent->IsVisibleInEditor())
+			for (UActorComponent* Component : ShapeActor->GetComponents())
 			{
-				if (PrimitiveComponent->ComponentIsTouchingSelectionFrustum(InFrustum, InViewportClient->EngineShowFlags, false /*only bsp*/, false/*encompass entire*/))
+				UPrimitiveComponent* PrimitiveComponent = Cast<UPrimitiveComponent>(Component);
+				if (PrimitiveComponent && PrimitiveComponent->IsRegistered() && PrimitiveComponent->IsVisibleInEditor())
 				{
-					if (ShapeActor->IsSelectable())
+					if (PrimitiveComponent->ComponentIsTouchingSelectionFrustum(InFrustum, InViewportClient->EngineShowFlags, false /*only bsp*/, false/*encompass entire*/))
 					{
-						bSomethingSelected = true;
-						const FName& ControlName = ShapeActor->ControlName;
-						SetRigElementSelection(ERigElementType::Control, ControlName, true);
+						if (ShapeActor->IsSelectable())
+						{
+							bSomethingSelected = true;
+							const FName& ControlName = ShapeActor->ControlName;
+							SetRigElementSelection(Pairs.Key,ERigElementType::Control, ControlName, true);
+						}
 					}
 				}
 			}
@@ -1522,144 +1569,149 @@ bool FControlRigEditMode::InputDelta(FEditorViewportClient* InViewportClient, FV
 	const bool bDoScale = !Scale.IsZero() && WidgetMode == UE::Widget::WM_Scale;
 
 
-	if (InteractionScope != nullptr && bMouseButtonDown && !bCtrlDown && !bShiftDown && CurrentAxis != EAxisList::None
+	if (InteractionScopes.Num() > 0 && bMouseButtonDown && !bCtrlDown && !bShiftDown && CurrentAxis != EAxisList::None
 		&& (bDoRotation || bDoTranslation || bDoScale))
 	{
-			if (AreRigElementsSelected(FRigElementTypeHelper::ToMask(ERigElementType::Control)))
+		for (TPair<UControlRig*, TArray<AControlRigShapeActor*>>& Pairs : ControlRigShapeActors)
 		{
-			FTransform ComponentTransform = GetHostingSceneComponentTransform();
-
-			if (bIsChangingControlShapeTransform)
+			if (AreRigElementsSelected(FRigElementTypeHelper::ToMask(ERigElementType::Control), Pairs.Key))
 			{
-				for (AControlRigShapeActor* ShapeActor : ShapeActors)
+				FTransform ComponentTransform = GetHostingSceneComponentTransform(Pairs.Key);
+
+				if (bIsChangingControlShapeTransform)
 				{
-					if (ShapeActor->IsSelected())
+					for (AControlRigShapeActor* ShapeActor : Pairs.Value)
 					{
-						if (bManipulatorMadeChange == false)
+						if (ShapeActor->IsSelected())
 						{
-							GEditor->BeginTransaction(LOCTEXT("ChangeControlShapeTransaction", "Change Control Shape Transform"));
-						}
+							if (bManipulatorMadeChange == false)
+							{
+								GEditor->BeginTransaction(LOCTEXT("ChangeControlShapeTransaction", "Change Control Shape Transform"));
+							}
 
-						ChangeControlShapeTransform(ShapeActor, bDoTranslation, InDrag, bDoRotation, InRot, bDoScale, InScale, ComponentTransform);
-						bManipulatorMadeChange = true;
+							ChangeControlShapeTransform(ShapeActor, bDoTranslation, InDrag, bDoRotation, InRot, bDoScale, InScale, ComponentTransform);
+							bManipulatorMadeChange = true;
 
-						// break here since we only support changing shape transform of a single control at a time
-						break;
+							// break here since we only support changing shape transform of a single control at a time
+							break;
+						}
 					}
 				}
-			}
-			else
-			{
-				const UControlRigEditModeSettings* Settings = GetDefault<UControlRigEditModeSettings>();
-				bool bDoLocal = (CoordSystem == ECoordSystem::COORD_Local && Settings && Settings->bLocalTransformsInEachLocalSpace);
-				bool bUseLocal = false;
-				bool bCalcLocal = bDoLocal;
-				bool bFirstTime = true;
-				FTransform InOutLocal = FTransform::Identity;
-				for (AControlRigShapeActor* ShapeActor : ShapeActors)
+				else
 				{
-					if (ShapeActor->IsSelected())
+					const UControlRigEditModeSettings* Settings = GetDefault<UControlRigEditModeSettings>();
+					bool bDoLocal = (CoordSystem == ECoordSystem::COORD_Local && Settings && Settings->bLocalTransformsInEachLocalSpace);
+					bool bUseLocal = false;
+					bool bCalcLocal = bDoLocal;
+					bool bFirstTime = true;
+					FTransform InOutLocal = FTransform::Identity;
+
+					for (AControlRigShapeActor* ShapeActor : Pairs.Value)
 					{
-						// test local vs global
-						if (bManipulatorMadeChange == false)
-						{
-							GEditor->BeginTransaction(LOCTEXT("MoveControlTransaction", "Move Control"));
-						}
-						if (bFirstTime)
+						if (ShapeActor->IsSelected())
 						{
-							bFirstTime = false;
-						}
-						else
-						{
-							if (bDoLocal)
+							// test local vs global
+							if (bManipulatorMadeChange == false)
 							{
-								bUseLocal = true;
-								bDoLocal = false;
+								GEditor->BeginTransaction(LOCTEXT("MoveControlTransaction", "Move Control"));
+							}
+							if (bFirstTime)
+							{
+								bFirstTime = false;
+							}
+							else
+							{
+								if (bDoLocal)
+								{
+									bUseLocal = true;
+									bDoLocal = false;
+								}
 							}
-						}
 
-						MoveControlShape(ShapeActor, bDoTranslation, InDrag, bDoRotation, InRot, bDoScale, InScale, ComponentTransform,
-							bUseLocal, bDoLocal, InOutLocal);
-						bManipulatorMadeChange = true;
+							MoveControlShape(ShapeActor, bDoTranslation, InDrag, bDoRotation, InRot, bDoScale, InScale, ComponentTransform,
+								bUseLocal, bDoLocal, InOutLocal);
+							bManipulatorMadeChange = true;
+						}
 					}
 				}
 			}
-
-			RecalcPivotTransform();
-
-			if (bManipulatorMadeChange)
+			else if (AreRigElementSelectedAndMovable(Pairs.Key))
 			{
-				TickManipulatableObjects(0.f);
-			}
-			return true;
-		}
-		else if (AreRigElementSelectedAndMovable())
-		{
-			FTransform ComponentTransform = GetHostingSceneComponentTransform();
-
-			// set Bone transform
-			// that will set initial Bone transform
-			TArray<FRigElementKey> SelectedRigElements = GetSelectedRigElements();
+				FTransform ComponentTransform = GetHostingSceneComponentTransform(Pairs.Key);
 
-			for (int32 Index = 0; Index < SelectedRigElements.Num(); ++Index)
-			{
-				const ERigElementType SelectedRigElementType = SelectedRigElements[Index].Type;
+				// set Bone transform
+				// that will set initial Bone transform
+				TArray<FRigElementKey> SelectedRigElements = GetSelectedRigElements(Pairs.Key);
 
-				if (SelectedRigElementType == ERigElementType::Control)
+				for (int32 Index = 0; Index < SelectedRigElements.Num(); ++Index)
 				{
-					FTransform NewWorldTransform = OnGetRigElementTransformDelegate.Execute(SelectedRigElements[Index], false, true) * ComponentTransform;
-					bool bTransformChanged = false;
-					if (bDoRotation)
-					{
-						FQuat CurrentRotation = NewWorldTransform.GetRotation();
-						CurrentRotation = (Rot.Quaternion() * CurrentRotation);
-						NewWorldTransform.SetRotation(CurrentRotation);
-						bTransformChanged = true;
-					}
+					const ERigElementType SelectedRigElementType = SelectedRigElements[Index].Type;
 
-					if (bDoTranslation)
+					if (SelectedRigElementType == ERigElementType::Control)
 					{
-						FVector CurrentLocation = NewWorldTransform.GetLocation();
-						CurrentLocation = CurrentLocation + Drag;
-						NewWorldTransform.SetLocation(CurrentLocation);
-						bTransformChanged = true;
-					}
+						FTransform NewWorldTransform = OnGetRigElementTransformDelegate.Execute(SelectedRigElements[Index], false, true) * ComponentTransform;
+						bool bTransformChanged = false;
+						if (bDoRotation)
+						{
+							FQuat CurrentRotation = NewWorldTransform.GetRotation();
+							CurrentRotation = (Rot.Quaternion() * CurrentRotation);
+							NewWorldTransform.SetRotation(CurrentRotation);
+							bTransformChanged = true;
+						}
 
-					if (bDoScale)
-					{
-						FVector CurrentScale = NewWorldTransform.GetScale3D();
-						CurrentScale = CurrentScale + Scale;
-						NewWorldTransform.SetScale3D(CurrentScale);
-						bTransformChanged = true;
-					}
+						if (bDoTranslation)
+						{
+							FVector CurrentLocation = NewWorldTransform.GetLocation();
+							CurrentLocation = CurrentLocation + Drag;
+							NewWorldTransform.SetLocation(CurrentLocation);
+							bTransformChanged = true;
+						}
 
-					if (bTransformChanged)
-					{
-						if (bManipulatorMadeChange == false)
+						if (bDoScale)
 						{
-							GEditor->BeginTransaction(LOCTEXT("MoveControlTransaction", "Move Control"));
+							FVector CurrentScale = NewWorldTransform.GetScale3D();
+							CurrentScale = CurrentScale + Scale;
+							NewWorldTransform.SetScale3D(CurrentScale);
+							bTransformChanged = true;
+						}
+
+						if (bTransformChanged)
+						{
+							if (bManipulatorMadeChange == false)
+							{
+								GEditor->BeginTransaction(LOCTEXT("MoveControlTransaction", "Move Control"));
+							}
+							FTransform NewComponentTransform = NewWorldTransform.GetRelativeTransform(ComponentTransform);
+							OnSetRigElementTransformDelegate.Execute(SelectedRigElements[Index], NewComponentTransform, false);
+							bManipulatorMadeChange = true;
 						}
-						FTransform NewComponentTransform = NewWorldTransform.GetRelativeTransform(ComponentTransform);
-						OnSetRigElementTransformDelegate.Execute(SelectedRigElements[Index], NewComponentTransform, false);
-						bManipulatorMadeChange = true;
 					}
 				}
 			}
-
-			// not sure this makes sense @rethink
-			return bManipulatorMadeChange;
 		}
 	}
-	return false;
+
+	RecalcPivotTransform();
+
+	if (bManipulatorMadeChange)
+	{
+		TickManipulatableObjects(0.f);
+	}
+	return bManipulatorMadeChange;
 }
 
 bool FControlRigEditMode::ShouldDrawWidget() const
 {
-	if (AreRigElementSelectedAndMovable())
+	for (const TWeakObjectPtr<UControlRig>& RuntimeRigPtr : RuntimeControlRigs)
 	{
-		return true;
+		if (UControlRig* ControlRig = RuntimeRigPtr.Get())
+		{
+			if (AreRigElementSelectedAndMovable(ControlRig))
+			{
+				return true;
+			}
+		}
 	}
-
 	return FEdMode::ShouldDrawWidget();
 }
 
@@ -1670,10 +1722,10 @@ bool FControlRigEditMode::IsCompatibleWith(FEditorModeID OtherModeID) const
 
 void FControlRigEditMode::AddReferencedObjects( FReferenceCollector& Collector )
 {
-	if (ShapeActors.Num() > 0)
+	for (TPair<UControlRig*, TArray<AControlRigShapeActor*>>& ShapeActors : ControlRigShapeActors)
 	{
-		for (AControlRigShapeActor* ShapeActor : ShapeActors)
-		{
+		for (AControlRigShapeActor* ShapeActor : ShapeActors.Value)
+		{		
 			Collector.AddReferencedObject(ShapeActor);
 		}
 	}
@@ -1685,39 +1737,32 @@ void FControlRigEditMode::AddReferencedObjects( FReferenceCollector& Collector )
 
 void FControlRigEditMode::ClearRigElementSelection(uint32 InTypes)
 {
-	UControlRig* InteractionRig = GetControlRig(true);
-	if (InteractionRig == nullptr)
+	for (TWeakObjectPtr<UControlRig>& RuntimeRigPtr : RuntimeControlRigs)
 	{
-		return;
-	}
-
-	UControlRigBlueprint* Blueprint = Cast<UControlRigBlueprint>(InteractionRig->GetClass()->ClassGeneratedBy);
-	if (IsInLevelEditor())
-	{
-		if(URigHierarchyController* Controller = InteractionRig->GetHierarchy()->GetController())
+		if (UControlRig* ControlRig = RuntimeRigPtr.Get())
 		{
-			Controller->ClearSelection();
+			if (IsInLevelEditor())
+			{
+				if (URigHierarchyController* Controller = ControlRig->GetHierarchy()->GetController())
+				{
+					Controller->ClearSelection();
+				}
+			}
+			else if (UControlRigBlueprint* Blueprint = Cast<UControlRigBlueprint>(ControlRig->GetClass()->ClassGeneratedBy))
+			{
+				Blueprint->GetHierarchyController()->ClearSelection();
+			}
 		}
 	}
-	else if (Blueprint)
-	{
-		Blueprint->GetHierarchyController()->ClearSelection();
-	}
 }
 
 // internal private function that doesn't use guarding.
-void FControlRigEditMode::SetRigElementSelectionInternal(ERigElementType Type, const FName& InRigElementName, bool bSelected)
+void FControlRigEditMode::SetRigElementSelectionInternal(UControlRig* ControlRig, ERigElementType Type, const FName& InRigElementName, bool bSelected)
 {
-	UControlRig* InteractionRig = GetControlRig(true);
-	if (InteractionRig == nullptr)
-	{
-		return;
-	}
-
-	UControlRigBlueprint* Blueprint = Cast<UControlRigBlueprint>(InteractionRig->GetClass()->ClassGeneratedBy);
+	UControlRigBlueprint* Blueprint = Cast<UControlRigBlueprint>(ControlRig->GetClass()->ClassGeneratedBy);
 	if (IsInLevelEditor())
 	{
-		if(URigHierarchyController* Controller = InteractionRig->GetHierarchy()->GetController())
+		if(URigHierarchyController* Controller = ControlRig->GetHierarchy()->GetController())
 		{
 			Controller->SelectElement(FRigElementKey(InRigElementName, Type), bSelected);
 		}
@@ -1728,19 +1773,19 @@ void FControlRigEditMode::SetRigElementSelectionInternal(ERigElementType Type, c
 	}
 }
 
-void FControlRigEditMode::SetRigElementSelection(ERigElementType Type, const FName& InRigElementName, bool bSelected)
+void FControlRigEditMode::SetRigElementSelection(UControlRig* ControlRig, ERigElementType Type, const FName& InRigElementName, bool bSelected)
 {
 	if (!bSelecting)
 	{
 		TGuardValue<bool> ReentrantGuard(bSelecting, true);
 
-		SetRigElementSelectionInternal(Type, InRigElementName, bSelected);
+		SetRigElementSelectionInternal(ControlRig,Type, InRigElementName, bSelected);
 
 		HandleSelectionChanged();
 	}
 }
 
-void FControlRigEditMode::SetRigElementSelection(ERigElementType Type, const TArray<FName>& InRigElementNames, bool bSelected)
+void FControlRigEditMode::SetRigElementSelection(UControlRig* ControlRig, ERigElementType Type, const TArray<FName>& InRigElementNames, bool bSelected)
 {
 	if (!bSelecting)
 	{
@@ -1748,63 +1793,64 @@ void FControlRigEditMode::SetRigElementSelection(ERigElementType Type, const TAr
 
 		for (const FName& ElementName : InRigElementNames)
 		{
-			SetRigElementSelectionInternal(Type, ElementName, bSelected);
+			SetRigElementSelectionInternal(ControlRig, Type, ElementName, bSelected);
 		}
 
 		HandleSelectionChanged();
 	}
 }
 
-TArray<FRigElementKey> FControlRigEditMode::GetSelectedRigElements() const
+TArray<FRigElementKey> FControlRigEditMode::GetSelectedRigElements(UControlRig* ControlRig) const
 {
+	if (ControlRig == nullptr && GetControlRigs().Num() > 0)
+	{
+		ControlRig = GetControlRigs()[0].Get();
+	}
+
 	TArray<FRigElementKey> SelectedKeys;
-	
-	if (UControlRig* ControlRig = GetControlRig(true))
+
+	if (ControlRig->GetHierarchy())
 	{
-		if (ControlRig->GetHierarchy())
-		{
-			SelectedKeys = ControlRig->GetHierarchy()->GetSelectedKeys();
-		}
+		SelectedKeys = ControlRig->GetHierarchy()->GetSelectedKeys();
+	}
 
-		// currently only 1 transient control is allowed at a time
-		// Transient Control's bSelected flag is never set to true, probably to avoid confusing other parts of the system
-		// But since Edit Mode directly deals with transient controls, its selection status is given special treatment here.
-		// So basically, whenever a bone is selected, and there is a transient control present, we consider both selected.
-		if (SelectedKeys.Num() == 1)
+	// currently only 1 transient control is allowed at a time
+	// Transient Control's bSelected flag is never set to true, probably to avoid confusing other parts of the system
+	// But since Edit Mode directly deals with transient controls, its selection status is given special treatment here.
+	// So basically, whenever a bone is selected, and there is a transient control present, we consider both selected.
+	if (SelectedKeys.Num() == 1)
+	{
+		if (SelectedKeys[0].Type == ERigElementType::Bone || SelectedKeys[0].Type == ERigElementType::Null)
 		{
-			if (SelectedKeys[0].Type == ERigElementType::Bone || SelectedKeys[0].Type == ERigElementType::Null)
+			const FName ControlName = UControlRig::GetNameForTransientControl(SelectedKeys[0]);
+			const FRigElementKey TransientControlKey = FRigElementKey(ControlName, ERigElementType::Control);
+			if(ControlRig->GetHierarchy()->Contains(TransientControlKey))
 			{
-				const FName ControlName = UControlRig::GetNameForTransientControl(SelectedKeys[0]);
-				const FRigElementKey TransientControlKey = FRigElementKey(ControlName, ERigElementType::Control);
-				if(ControlRig->GetHierarchy()->Contains(TransientControlKey))
-				{
-					SelectedKeys.Add(TransientControlKey);
-				}
-
+				SelectedKeys.Add(TransientControlKey);
 			}
+
 		}
-		else
-		{
-			// check if there is a pin value transient control active
-			// when a pin control is active, all existing selection should have been cleared
-			TArray<FRigControlElement*> TransientControls = ControlRig->GetHierarchy()->GetTransientControls();
+	}
+	else
+	{
+		// check if there is a pin value transient control active
+		// when a pin control is active, all existing selection should have been cleared
+		TArray<FRigControlElement*> TransientControls = ControlRig->GetHierarchy()->GetTransientControls();
 
-			if (TransientControls.Num() > 0)
+		if (TransientControls.Num() > 0)
+		{
+			if (ensure(SelectedKeys.Num() == 0))
 			{
-				if (ensure(SelectedKeys.Num() == 0))
-				{
-					SelectedKeys.Add(TransientControls[0]->GetKey());
-				}
+				SelectedKeys.Add(TransientControls[0]->GetKey());
 			}
 		}
 	}
-
 	return SelectedKeys;
 }
 
-bool FControlRigEditMode::AreRigElementsSelected(uint32 InTypes) const
+bool FControlRigEditMode::AreRigElementsSelected(uint32 InTypes, UControlRig* InControlRig) const
 {
-	TArray<FRigElementKey> SelectedRigElements = GetSelectedRigElements();
+	TArray<FRigElementKey> SelectedRigElements = GetSelectedRigElements(InControlRig);
 
 	for (const FRigElementKey& Ele : SelectedRigElements)
 	{
@@ -1817,9 +1863,9 @@ bool FControlRigEditMode::AreRigElementsSelected(uint32 InTypes) const
 	return false;
 }
 
-int32 FControlRigEditMode::GetNumSelectedRigElements(uint32 InTypes) const
+int32 FControlRigEditMode::GetNumSelectedRigElements(uint32 InTypes, UControlRig* InControlRig) const
 {
-	TArray<FRigElementKey> SelectedRigElements = GetSelectedRigElements();
+	TArray<FRigElementKey> SelectedRigElements = GetSelectedRigElements(InControlRig);
 	if (FRigElementTypeHelper::DoesHave(InTypes, ERigElementType::All))
 	{
 		return SelectedRigElements.Num();
@@ -1848,13 +1894,16 @@ void FControlRigEditMode::RefreshObjects()
 
 bool FControlRigEditMode::CanRemoveFromPreviewScene(const USceneComponent* InComponent)
 {
-	for (AControlRigShapeActor* ShapeActor : ShapeActors)
+	for (TPair<UControlRig*, TArray<AControlRigShapeActor*>>& ShapeActors : ControlRigShapeActors)
 	{
-		TInlineComponentArray<USceneComponent*> SceneComponents;
-		ShapeActor->GetComponents(SceneComponents, true);
-		if (SceneComponents.Contains(InComponent))
+		for (AControlRigShapeActor* ShapeActor : ShapeActors.Value)
 		{
-			return false;
+			TInlineComponentArray<USceneComponent*> SceneComponents;
+			ShapeActor->GetComponents(SceneComponents, true);
+			if (SceneComponents.Contains(InComponent))
+			{
+				return false;
+			}
 		}
 	}
 
@@ -1864,120 +1913,135 @@ bool FControlRigEditMode::CanRemoveFromPreviewScene(const USceneComponent* InCom
 
 void FControlRigEditMode::RecalcPivotTransform()
 {
-	PivotTransform = FTransform::Identity;
-
-	// @todo: support bones also
-	TArray<FRigElementKey> SelectedRigElements = GetSelectedRigElements();
-	if (AreRigElementsSelected(FRigElementTypeHelper::ToMask(ERigElementType::Control)))
+	PivotTransforms.Reset();
+	for (TWeakObjectPtr<UControlRig>& RuntimeRigPtr : RuntimeControlRigs)
 	{
-		FTransform LastTransform = FTransform::Identity;
+		if (UControlRig* ControlRig = RuntimeRigPtr.Get())
+		{
+			FTransform PivotTransform = FTransform::Identity;
+			// Use average location as pivot location
+			FVector PivotLocation = FVector::ZeroVector;
 
-		// recalc coord system too
-		FTransform ComponentTransform = GetHostingSceneComponentTransform();
+			// @todo: support bones also
+			TArray<FRigElementKey> SelectedRigElements = GetSelectedRigElements(ControlRig);
+			if (AreRigElementsSelected(FRigElementTypeHelper::ToMask(ERigElementType::Control),ControlRig))
+			{
+				FTransform LastTransform = FTransform::Identity;
 
-		// Use average location as pivot location
-		FVector PivotLocation = FVector::ZeroVector;
+				// recalc coord system too
+				FTransform ComponentTransform = GetHostingSceneComponentTransform(ControlRig);
 
-		int32 NumSelectedControls = 0;
-		for (int32 Index = 0; Index < SelectedRigElements.Num(); ++Index)
-		{
-			if (SelectedRigElements[Index].Type == ERigElementType::Control)
-			{
-				// todo?
-			}
-		}
-		
-		if (bIsChangingControlShapeTransform)
-		{
-			if (UControlRig* ControlRig = GetControlRig(true))
-			{ 
-				for (const AControlRigShapeActor* ShapeActor : ShapeActors)
+
+				int32 NumSelectedControls = 0;
+				for (int32 Index = 0; Index < SelectedRigElements.Num(); ++Index)
 				{
-					if (ShapeActor->IsSelected())
+					if (SelectedRigElements[Index].Type == ERigElementType::Control)
 					{
-						if (FRigControlElement* ControlElement = ControlRig->GetHierarchy()->Find<FRigControlElement>(FRigElementKey(ShapeActor->ControlName, ERigElementType::Control)))
+						// todo?
+					}
+				}
+
+				if (bIsChangingControlShapeTransform)
+				{
+					if (TArray<AControlRigShapeActor* >* ShapeActors = ControlRigShapeActors.Find(ControlRig))
+					{
+						for (const AControlRigShapeActor* ShapeActor : *ShapeActors)
 						{
-							PivotTransform = ControlRig->GetHierarchy()->GetControlShapeTransform(ControlElement, ERigTransformType::CurrentGlobal);
+							if (ShapeActor->IsSelected())
+							{
+								if (FRigControlElement* ControlElement = ControlRig->GetHierarchy()->Find<FRigControlElement>(FRigElementKey(ShapeActor->ControlName, ERigElementType::Control)))
+								{
+									PivotTransform = ControlRig->GetHierarchy()->GetControlShapeTransform(ControlElement, ERigTransformType::CurrentGlobal);
+								}
+
+								// break here since we don't want to change the shape transform of multiple controls.
+								break;
+							}
 						}
+					}
+				}
+				else
+				{
+					const UControlRigEditModeSettings* Settings = GetDefault<UControlRigEditModeSettings>();
 
-						// break here since we don't want to change the shape transform of multiple controls.
-						break;
+					if (TArray<AControlRigShapeActor* >* ShapeActors = ControlRigShapeActors.Find(ControlRig))
+					{
+						for (const AControlRigShapeActor* ShapeActor : *ShapeActors)
+						{
+							if (ShapeActor->IsSelected())
+							{
+								LastTransform = ShapeActor->GetActorTransform().GetRelativeTransform(ComponentTransform);
+								PivotLocation += LastTransform.GetLocation();
+								++NumSelectedControls;
+								if (Settings && Settings->bLocalTransformsInEachLocalSpace) //if in local just use first actors transform
+								{
+									break;
+								}
+
+							}
+						}
 					}
+
+					PivotLocation /= (float)FMath::Max(1, NumSelectedControls);
+					PivotTransform.SetLocation(PivotLocation);
+
+					// just use last rotation
+					FTransform WorldTransform = LastTransform * ComponentTransform;
+					PivotTransform.SetRotation(WorldTransform.GetRotation());
 				}
+				PivotTransforms.Add(ControlRig, PivotTransform);
 			}
-		}
-		else
-		{
-			const UControlRigEditModeSettings* Settings = GetDefault<UControlRigEditModeSettings>();
-			for (const AControlRigShapeActor* ShapeActor : ShapeActors)
+			else if (AreRigElementSelectedAndMovable(ControlRig))
 			{
-				if (ShapeActor->IsSelected())
+				// recalc coord system too
+				FTransform ComponentTransform = GetHostingSceneComponentTransform(ControlRig);
+
+				// Use average location as pivot location
+				PivotLocation = FVector::ZeroVector;
+				int32 NumSelection = 0;
+				FTransform LastTransform = FTransform::Identity;
+				for (int32 Index = 0; Index < SelectedRigElements.Num(); ++Index)
 				{
-					LastTransform = ShapeActor->GetActorTransform().GetRelativeTransform(ComponentTransform);
-					PivotLocation += LastTransform.GetLocation();
-					++NumSelectedControls;
-					if (Settings && Settings->bLocalTransformsInEachLocalSpace) //if in local just use first actors transform
+					if (SelectedRigElements[Index].Type == ERigElementType::Control)
 					{
-						break;
+						LastTransform = OnGetRigElementTransformDelegate.Execute(SelectedRigElements[Index], false, true);
+						PivotLocation += LastTransform.GetLocation();
+						++NumSelection;
 					}
-					
 				}
-			}
 
-			PivotLocation /= (float)FMath::Max(1, NumSelectedControls);
-			PivotTransform.SetLocation(PivotLocation);
-			
-			// just use last rotation
-			FTransform WorldTransform = LastTransform * ComponentTransform;
-			PivotTransform.SetRotation(WorldTransform.GetRotation());
-		}
-	}
-	else if (AreRigElementSelectedAndMovable())
-	{
-		// recalc coord system too
-		FTransform ComponentTransform = GetHostingSceneComponentTransform();
+				PivotLocation /= (float)FMath::Max(1, NumSelection);
+				PivotTransform.SetLocation(PivotLocation);
 
-		// Use average location as pivot location
-		FVector PivotLocation = FVector::ZeroVector;
-		int32 NumSelection = 0;
-		FTransform LastTransform = FTransform::Identity;
-		for (int32 Index = 0; Index < SelectedRigElements.Num(); ++Index)
-		{
-			if (SelectedRigElements[Index].Type == ERigElementType::Control)
-			{
-				LastTransform = OnGetRigElementTransformDelegate.Execute(SelectedRigElements[Index], false, true);
-				PivotLocation += LastTransform.GetLocation();
-				++NumSelection;
+				if (NumSelection == 1)
+				{
+					// A single Bone just uses its own transform
+					FTransform WorldTransform = LastTransform * ComponentTransform;
+					PivotTransform.SetRotation(WorldTransform.GetRotation());
+				}
+				else if (NumSelection > 1)
+				{
+					// If we have more than one Bone selected, use the coordinate space of the component
+					PivotTransform.SetRotation(ComponentTransform.GetRotation());
+				}
+				PivotTransforms.Add(ControlRig, PivotTransform);
 			}
 		}
-
-		PivotLocation /= (float)FMath::Max(1, NumSelection);
-		PivotTransform.SetLocation(PivotLocation);
-
-		if (NumSelection == 1)
-		{
-			// A single Bone just uses its own transform
-			FTransform WorldTransform = LastTransform * ComponentTransform;
-			PivotTransform.SetRotation(WorldTransform.GetRotation());
-		}
-		else if (NumSelection > 1)
-		{
-			// If we have more than one Bone selected, use the coordinate space of the component
-			PivotTransform.SetRotation(ComponentTransform.GetRotation());
-		}
 	}
 }
 
 void FControlRigEditMode::HandleSelectionChanged()
 {
-
-	for (AControlRigShapeActor* ShapeActor : ShapeActors)
+	for (TPair<UControlRig*, TArray<AControlRigShapeActor*>>& ShapeActors : ControlRigShapeActors)
 	{
-		TInlineComponentArray<UPrimitiveComponent*> PrimitiveComponents;
-		ShapeActor->GetComponents(PrimitiveComponents, true);
-		for (UPrimitiveComponent* PrimitiveComponent : PrimitiveComponents)
+		for (AControlRigShapeActor* ShapeActor : ShapeActors.Value)
 		{
-			PrimitiveComponent->PushSelectionToProxy();
+			TInlineComponentArray<UPrimitiveComponent*> PrimitiveComponents;
+			ShapeActor->GetComponents(PrimitiveComponents, true);
+			for (UPrimitiveComponent* PrimitiveComponent : PrimitiveComponents)
+			{
+				PrimitiveComponent->PushSelectionToProxy();
+			}
 		}
 	}
 
@@ -2001,6 +2065,9 @@ void FControlRigEditMode::BindCommands()
 	CommandBindings->MapAction(
 		Commands.ToggleManipulators,
 		FExecuteAction::CreateRaw(this, &FControlRigEditMode::ToggleManipulators));
+	CommandBindings->MapAction(
+		Commands.ToggleAllManipulators,
+		FExecuteAction::CreateRaw(this, &FControlRigEditMode::ToggleAllManipulators));
 	CommandBindings->MapAction(
 		Commands.ResetTransforms,
 		FExecuteAction::CreateRaw(this, &FControlRigEditMode::ResetTransforms, true));
@@ -2041,39 +2108,32 @@ void FControlRigEditMode::BindCommands()
 bool FControlRigEditMode::IsControlSelected() const
 {
 	static uint32 TypeFlag = (uint32)ERigElementType::Control;
-	return (AreRigElementsSelected(TypeFlag));
-}
-
-bool FControlRigEditMode::GetRigElementGlobalTransform(const FRigElementKey& InElement, FTransform& OutGlobalTransform) const
-{
-	// if control, go through manipulation layer
-	if (InElement.Type == ERigElementType::Control)
+	for (const TWeakObjectPtr<UControlRig>& RuntimeRigPtr : RuntimeControlRigs)
 	{
-		// this code is weird. Need to set this info in manipulation layer
-		AControlRigShapeActor* ShapeActor = GetControlShapeFromControlName(InElement.Name);
-		if (ShapeActor && ensure(ShapeActor->IsSelected()))
+		if (UControlRig* ControlRig = RuntimeRigPtr.Get())
 		{
-			OutGlobalTransform = GetControlShapeTransform(ShapeActor);
-			return true;
-		}
-	}
-	else if (AreRigElementSelectedAndMovable())
-	{
-		// @tood: we often just cross ControlRig here without manipulation layer
-		// should we clean this up?
-		if (UControlRig* ControlRig = GetControlRig(true))
-		{
-			OutGlobalTransform = ControlRig->GetHierarchy()->GetGlobalTransform(InElement);
-			return true;
+			if (AreRigElementsSelected(TypeFlag,ControlRig))
+			{
+				return true;
+			}
 		}
 	}
-
 	return false;
 }
 
 bool FControlRigEditMode::CanFrameSelection()
 {
-	return  GetSelectedRigElements().Num() > 0;
+	for (TWeakObjectPtr<UControlRig>& RuntimeRigPtr : RuntimeControlRigs)
+	{
+		if (UControlRig* ControlRig = RuntimeRigPtr.Get())
+		{
+			if (GetSelectedRigElements(ControlRig).Num() > 0)
+			{
+				return true;
+			}
+		}
+	}
+	return  false;;
 }
 
 void FControlRigEditMode::ClearSelection()
@@ -2102,15 +2162,21 @@ void FControlRigEditMode::FrameSelection()
     }
 
 	TArray<AActor*> Actors;
-	TArray<FRigElementKey> SelectedRigElements = GetSelectedRigElements();
-	for (const FRigElementKey& SelectedKey : SelectedRigElements)
+	for (TWeakObjectPtr<UControlRig>& RuntimeRigPtr : RuntimeControlRigs)
 	{
-		if (SelectedKey.Type == ERigElementType::Control)
+		if (UControlRig* ControlRig = RuntimeRigPtr.Get())
 		{
-			AControlRigShapeActor* ShapeActor = GetControlShapeFromControlName(SelectedKey.Name);
-			if (ShapeActor)
+			TArray<FRigElementKey> SelectedRigElements = GetSelectedRigElements(ControlRig);
+			for (const FRigElementKey& SelectedKey : SelectedRigElements)
 			{
-				Actors.Add(ShapeActor);
+				if (SelectedKey.Type == ERigElementType::Control)
+				{
+					AControlRigShapeActor* ShapeActor = GetControlShapeFromControlName(ControlRig,SelectedKey.Name);
+					if (ShapeActor)
+					{
+						Actors.Add(ShapeActor);
+					}
+				}
 			}
 		}
 	}
@@ -2193,16 +2259,45 @@ void FControlRigEditMode::ToggleControlShapeTransformEdit()
 	}
 }
 
+void FControlRigEditMode::GetAllSelectedControls(TMap<UControlRig*, TArray<FRigElementKey>>& OutSelectedControls)
+{
+	for (TWeakObjectPtr<UControlRig>& RuntimeRigPtr : RuntimeControlRigs)
+	{
+		if (UControlRig* ControlRig = RuntimeRigPtr.Get())
+		{
+			if (URigHierarchy* Hierarchy = ControlRig->GetHierarchy())
+			{
+				TArray<FRigElementKey> SelectedControls = Hierarchy->GetSelectedKeys(ERigElementType::Control);
+				if (SelectedControls.Num() > 0)
+				{
+					OutSelectedControls.Add(ControlRig, SelectedControls);
+				}
+			}
+		}
+	}
+}
+
 void FControlRigEditMode::OpenSpacePickerWidget()
 {
-	UControlRig* RuntimeRig = GetControlRig(false); //helge space picker only works on runtime rig?
-	if (RuntimeRig == nullptr)
+	TMap<UControlRig*, TArray<FRigElementKey>> SelectedControlRigsAndControls;
+	GetAllSelectedControls(SelectedControlRigsAndControls);
+
+	if (SelectedControlRigsAndControls.Num() < 1)
 	{
 		return;
 	}
 
+	TArray<UControlRig*> ControlRigs;
+	TArray<TArray<FRigElementKey>> AllSelectedControls;
+	SelectedControlRigsAndControls.GenerateKeyArray(ControlRigs);
+	SelectedControlRigsAndControls.GenerateValueArray(AllSelectedControls);
+
+
+	//mz todo handle multiple control rigs with space picker
+	UControlRig* RuntimeRig = ControlRigs[0];
+	TArray<FRigElementKey>& SelectedControls = AllSelectedControls[0];
+
 	URigHierarchy* Hierarchy = RuntimeRig->GetHierarchy();
-	TArray<FRigElementKey> SelectedControls = Hierarchy->GetSelectedKeys(ERigElementType::Control);
 
 	TSharedRef<SRigSpacePickerWidget> PickerWidget =
 	SNew(SRigSpacePickerWidget)
@@ -2325,129 +2420,165 @@ FText FControlRigEditMode::GetToggleControlShapeTransformEditHotKey() const
 
 void FControlRigEditMode::ToggleManipulators()
 {
-	// Toggle flag (is used in drawing code)
-	UControlRigEditModeSettings* Settings = GetMutableDefault<UControlRigEditModeSettings>();
-	Settings->bHideControlShapes = !Settings->bHideControlShapes;
-}
-
-void FControlRigEditMode::ResetTransforms(bool bSelectionOnly)
-{
-	if (UControlRig* ControlRig = GetControlRig(true))
+	if (IsInLevelEditor())
 	{
-		TArray<FRigElementKey> SelectedRigElements = GetSelectedRigElements();
-		TArray<FRigElementKey> ControlsToReset = SelectedRigElements;
-		if (!bSelectionOnly)
+		TMap<UControlRig*, TArray<FRigElementKey>> SelectedControls;
+		GetAllSelectedControls(SelectedControls);
+		TArray<UControlRig*> ControlRigs;
+		SelectedControls.GenerateKeyArray(ControlRigs);
+		for (UControlRig* ControlRig : ControlRigs)
 		{
-			TArray<FRigControlElement*> Controls;
-			ControlRig->GetControlsInOrder(Controls);
-			ControlsToReset.SetNum(0);
-			for (const FRigControlElement* Control : Controls)
+			if (ControlRig)
 			{
-				ControlsToReset.Add(Control->GetKey());
-
+				ControlRig->ToggleControlsVisible();
 			}
 		}
-		bool bHasNonDefaultParent = false;
-		TArray<FRigElementKey> Parents;
-		for (const FRigElementKey& ControlKey : ControlsToReset)
+	}
+	else
+	{
+		UControlRigEditModeSettings* Settings = GetMutableDefault<UControlRigEditModeSettings>();
+		Settings->bHideControlShapes = !Settings->bHideControlShapes;
+	}
+}
+
+void FControlRigEditMode::ToggleAllManipulators()
+{	
+	UControlRigEditModeSettings* Settings = GetMutableDefault<UControlRigEditModeSettings>();
+	Settings->bHideControlShapes = !Settings->bHideControlShapes;
+
+	//turn on all if in level editor in case any where off
+	if (IsInLevelEditor() && Settings->bHideControlShapes)
+	{
+		for (TWeakObjectPtr<UControlRig>& RuntimeRigPtr : RuntimeControlRigs)
 		{
-			FRigElementKey SpaceKey = ControlRig->GetHierarchy()->GetActiveParent(ControlKey);
-			Parents.Add(SpaceKey);
-			if (SpaceKey != ControlRig->GetHierarchy()->GetDefaultParentKey())
+			if (UControlRig* ControlRig = RuntimeRigPtr.Get())
 			{
-				bHasNonDefaultParent = true;
+				ControlRig->SetControlsVisible(true);
 			}
 		}
+	}
+}
 
-		FScopedTransaction Transaction(LOCTEXT("HierarchyResetTransforms", "Reset Transforms"));
-
-		for (const FRigElementKey& ControlToReset : ControlsToReset)
+void FControlRigEditMode::ResetTransforms(bool bSelectionOnly)
+{
+	for (TWeakObjectPtr<UControlRig>& RuntimeRigPtr : RuntimeControlRigs)
+	{
+		if (UControlRig* ControlRig = RuntimeRigPtr.Get())
 		{
-			if (ControlToReset.Type == ERigElementType::Control)
+
+			TArray<FRigElementKey> SelectedRigElements = GetSelectedRigElements(ControlRig);
+			TArray<FRigElementKey> ControlsToReset = SelectedRigElements;
+			if (!bSelectionOnly)
 			{
-				FRigControlElement* ControlElement = ControlRig->FindControl(ControlToReset.Name);
-				if (ControlElement && !ControlElement->Settings.bIsTransientControl)
+				TArray<FRigControlElement*> Controls;
+				ControlRig->GetControlsInOrder(Controls);
+				ControlsToReset.SetNum(0);
+				for (const FRigControlElement* Control : Controls)
 				{
-					const FTransform InitialLocalTransform = ControlRig->GetHierarchy()->GetInitialLocalTransform(ControlToReset);
-					ControlRig->Modify();
-					if (bHasNonDefaultParent == true) //possibly not at default parent so switch to it
-					{
-						ControlRig->GetHierarchy()->SwitchToDefaultParent(ControlElement->GetKey());
-					}
-					ControlRig->GetHierarchy()->SetLocalTransform(ControlToReset, InitialLocalTransform);
-					if (bHasNonDefaultParent == false)
-					{
-						ControlRig->ControlModified().Broadcast(ControlRig, ControlElement, EControlRigSetKey::DoNotCare);
-					}
+					ControlsToReset.Add(Control->GetKey());
 
-					//@helge not sure what to do if the non-default parent
-					if (UControlRigBlueprint* Blueprint = Cast<UControlRigBlueprint>(ControlRig->GetClass()->ClassGeneratedBy))
-					{
-						Blueprint->Hierarchy->SetLocalTransform(ControlToReset, InitialLocalTransform);
-					}
 				}
 			}
-		}
-
-		if (bHasNonDefaultParent == true) //now we have the initial pose setup we need to get the global transforms as specified now then set them in the current parent space
-		{
-			ControlRig->Evaluate_AnyThread();
-
-			//get global transforms
-			TArray<FTransform> GlobalTransforms;
-			for (const FRigElementKey& ControlToReset : ControlsToReset)
+			bool bHasNonDefaultParent = false;
+			TArray<FRigElementKey> Parents;
+			for (const FRigElementKey& ControlKey : ControlsToReset)
 			{
-				FRigControlElement* ControlElement = ControlRig->FindControl(ControlToReset.Name);
-				if (ControlElement && !ControlElement->Settings.bIsTransientControl)
+				FRigElementKey SpaceKey = ControlRig->GetHierarchy()->GetActiveParent(ControlKey);
+				Parents.Add(SpaceKey);
+				if (SpaceKey != ControlRig->GetHierarchy()->GetDefaultParentKey())
 				{
-					FTransform GlobalTransform = ControlRig->GetHierarchy()->GetGlobalTransform(ControlToReset);
-					GlobalTransforms.Add(GlobalTransform);
+					bHasNonDefaultParent = true;
 				}
 			}
-			//switch back to original parent space
-			int32 Index = 0;
+
+			FScopedTransaction Transaction(LOCTEXT("HierarchyResetTransforms", "Reset Transforms"));
+
 			for (const FRigElementKey& ControlToReset : ControlsToReset)
 			{
-				FRigControlElement* ControlElement = ControlRig->FindControl(ControlToReset.Name);
-				if (ControlElement && !ControlElement->Settings.bIsTransientControl)
+				if (ControlToReset.Type == ERigElementType::Control)
 				{
-					ControlRig->GetHierarchy()->SwitchToParent(ControlToReset,Parents[Index]);
-					++Index;
+					FRigControlElement* ControlElement = ControlRig->FindControl(ControlToReset.Name);
+					if (ControlElement && !ControlElement->Settings.bIsTransientControl)
+					{
+						const FTransform InitialLocalTransform = ControlRig->GetHierarchy()->GetInitialLocalTransform(ControlToReset);
+						ControlRig->Modify();
+						if (bHasNonDefaultParent == true) //possibly not at default parent so switch to it
+						{
+							ControlRig->GetHierarchy()->SwitchToDefaultParent(ControlElement->GetKey());
+						}
+						ControlRig->GetHierarchy()->SetLocalTransform(ControlToReset, InitialLocalTransform);
+						if (bHasNonDefaultParent == false)
+						{
+							ControlRig->ControlModified().Broadcast(ControlRig, ControlElement, EControlRigSetKey::DoNotCare);
+						}
+
+						//@helge not sure what to do if the non-default parent
+						if (UControlRigBlueprint* Blueprint = Cast<UControlRigBlueprint>(ControlRig->GetClass()->ClassGeneratedBy))
+						{
+							Blueprint->Hierarchy->SetLocalTransform(ControlToReset, InitialLocalTransform);
+						}
+					}
 				}
 			}
-			//set global transforms in this space // do it twice since ControlsInOrder is not really always in order
-			for (int32 SetHack = 0; SetHack < 2; ++SetHack)
+
+			if (bHasNonDefaultParent == true) //now we have the initial pose setup we need to get the global transforms as specified now then set them in the current parent space
 			{
 				ControlRig->Evaluate_AnyThread();
-				Index = 0;
+
+				//get global transforms
+				TArray<FTransform> GlobalTransforms;
 				for (const FRigElementKey& ControlToReset : ControlsToReset)
+				{
+					FRigControlElement* ControlElement = ControlRig->FindControl(ControlToReset.Name);
+					if (ControlElement && !ControlElement->Settings.bIsTransientControl)
+					{
+						FTransform GlobalTransform = ControlRig->GetHierarchy()->GetGlobalTransform(ControlToReset);
+						GlobalTransforms.Add(GlobalTransform);
+					}
+				}
+				//switch back to original parent space
+				int32 Index = 0;
+				for (const FRigElementKey& ControlToReset : ControlsToReset)
+				{
+					FRigControlElement* ControlElement = ControlRig->FindControl(ControlToReset.Name);
+					if (ControlElement && !ControlElement->Settings.bIsTransientControl)
+					{
+						ControlRig->GetHierarchy()->SwitchToParent(ControlToReset, Parents[Index]);
+						++Index;
+					}
+				}
+				//set global transforms in this space // do it twice since ControlsInOrder is not really always in order
+				for (int32 SetHack = 0; SetHack < 2; ++SetHack)
+				{
+					ControlRig->Evaluate_AnyThread();
+					Index = 0;
+					for (const FRigElementKey& ControlToReset : ControlsToReset)
+
+					{
+						FRigControlElement* ControlElement = ControlRig->FindControl(ControlToReset.Name);
+						if (ControlElement && !ControlElement->Settings.bIsTransientControl)
+						{
+							ControlRig->GetHierarchy()->SetGlobalTransform(ControlToReset, GlobalTransforms[Index]);
+							ControlRig->Evaluate_AnyThread();
+							++Index;
+						}
+					}
+				}
+				//send notifies
 
+				for (const FRigElementKey& ControlToReset : ControlsToReset)
 				{
 					FRigControlElement* ControlElement = ControlRig->FindControl(ControlToReset.Name);
 					if (ControlElement && !ControlElement->Settings.bIsTransientControl)
 					{
-						ControlRig->GetHierarchy()->SetGlobalTransform(ControlToReset, GlobalTransforms[Index]);
-						ControlRig->Evaluate_AnyThread();
-						++Index;
+						ControlRig->ControlModified().Broadcast(ControlRig, ControlElement, EControlRigSetKey::DoNotCare);
 					}
 				}
 			}
-			//send notifies
-
-			for (const FRigElementKey& ControlToReset : ControlsToReset)
-			{
-				FRigControlElement* ControlElement = ControlRig->FindControl(ControlToReset.Name);
-				if (ControlElement && !ControlElement->Settings.bIsTransientControl)
-				{
-					ControlRig->ControlModified().Broadcast(ControlRig, ControlElement, EControlRigSetKey::DoNotCare);
-				}
-			}
-
 		}
 	}
 }
 
-
 bool FControlRigEditMode::MouseMove(FEditorViewportClient* ViewportClient, FViewport* Viewport, int32 x, int32 y)
 {
 	// Inform units of hover state
@@ -2456,9 +2587,12 @@ bool FControlRigEditMode::MouseMove(FEditorViewportClient* ViewportClient, FView
 	{
 		if (ActorHitProxy->Actor->IsA<AControlRigShapeActor>())
 		{
-			for (AControlRigShapeActor* ShapeActor : ShapeActors)
+			for (TPair<UControlRig*, TArray<AControlRigShapeActor*>>& ShapeActors : ControlRigShapeActors)
 			{
-				ShapeActor->SetHovered(ShapeActor == ActorHitProxy->Actor);
+				for (AControlRigShapeActor* ShapeActor : ShapeActors.Value)
+				{
+					ShapeActor->SetHovered(ShapeActor == ActorHitProxy->Actor);
+				}
 			}
 		}
 	}
@@ -2468,9 +2602,12 @@ bool FControlRigEditMode::MouseMove(FEditorViewportClient* ViewportClient, FView
 
 bool FControlRigEditMode::MouseLeave(FEditorViewportClient* ViewportClient, FViewport* Viewport)
 {
-	for (AControlRigShapeActor* ShapeActor : ShapeActors)
+	for (TPair<UControlRig*, TArray<AControlRigShapeActor*>>& ShapeActors : ControlRigShapeActors)
 	{
-		ShapeActor->SetHovered(false);
+		for (AControlRigShapeActor* ShapeActor : ShapeActors.Value)
+		{
+			ShapeActor->SetHovered(false);
+		}
 	}
 
 	return false;
@@ -2478,34 +2615,183 @@ bool FControlRigEditMode::MouseLeave(FEditorViewportClient* ViewportClient, FVie
 
 void FControlRigEditMode::PostUndo()
 {
-	UControlRig* RuntimeControlRig = GetControlRig(false);
-	if (!RuntimeControlRig)
+	if (WeakSequencer.IsValid())
+	{
+		for (TWeakObjectPtr<UControlRig>& RuntimeRigPtr : RuntimeControlRigs)
+		{
+			if (RuntimeRigPtr.IsValid() == false)
+			{
+				DestroyShapesActors(RuntimeRigPtr.Get());
+			}
+		}
+		TSharedPtr<ISequencer> Sequencer = WeakSequencer.Pin();
+		if (ULevelSequence* LevelSequence = Cast<ULevelSequence>(Sequencer->GetFocusedMovieSceneSequence()))
+		{
+			bool bSomethingAdded = false;
+			TArray<FControlRigSequencerBindingProxy> Proxies = UControlRigSequencerEditorLibrary::GetControlRigs(LevelSequence);
+			for (FControlRigSequencerBindingProxy& Proxy : Proxies)
+			{
+				if (UControlRig* ControlRig = Proxy.ControlRig.Get())
+				{
+					if (RuntimeControlRigs.Contains(ControlRig) == false)
+					{
+						AddControlRigInternal(ControlRig);
+						bSomethingAdded = true;
+
+					}
+				}
+			}
+			if (bSomethingAdded)
+			{
+				SetObjects_Internal();
+			}
+		}
+	}
+	else
 	{
-		DestroyShapesActors();
+		for (TWeakObjectPtr<UControlRig>& RuntimeRigPtr : RuntimeControlRigs)
+		{
+			if (RuntimeRigPtr.IsValid() == false)
+			{
+				DestroyShapesActors(RuntimeRigPtr.Get());
+			}
+		}
+	}
+}
+
+void FControlRigEditMode::RequestToRecreateControlShapeActors(UControlRig* ControlRig)
+{ 
+	if (ControlRig)
+	{
+		if (RecreateControlShapesRequired != ERecreateControlRigShape::RecreateAll)
+		{
+			RecreateControlShapesRequired = ERecreateControlRigShape::RecreateSpecified;
+			if (!ControlRigsToRecreate.Find(ControlRig))
+			{
+				ControlRigsToRecreate.Add(ControlRig);
+			}
+		}
+	}
+	else
+	{
+		RecreateControlShapesRequired = ERecreateControlRigShape::RecreateAll;
 	}
 }
 
 void FControlRigEditMode::RecreateControlShapeActors(const TArray<FRigElementKey>& InSelectedElements)
 {
-	if (ControlProxy)
+	if (RecreateControlShapesRequired == ERecreateControlRigShape::RecreateAll)
+	{
+		for (TWeakObjectPtr<UControlRig>& RuntimeRigPtr : RuntimeControlRigs)
+		{
+			if (UControlRig* RuntimeControlRig = RuntimeRigPtr.Get())
+			{
+				DestroyShapesActors(RuntimeControlRig);
+				CreateShapeActors(RuntimeControlRig);
+			}
+		}
+	}
+	else
 	{
-		ControlProxy->RemoveAllProxies();
+		for (UControlRig* ControlRig : ControlRigsToRecreate)
+		{
+			DestroyShapesActors(ControlRig);
+			CreateShapeActors(ControlRig);
+		}
 	}
+}
 
-	if (UControlRig* ControlRig = GetControlRig(false))
+// temporarily we just support following types of gizmo
+bool IsSupportedControlType(const ERigControlType ControlType)
+{
+	switch (ControlType)
+	{
+	case ERigControlType::Float:
+	case ERigControlType::Integer:
+	case ERigControlType::Vector2D:
+	case ERigControlType::Position:
+	case ERigControlType::Scale:
+	case ERigControlType::Rotator:
+	case ERigControlType::Transform:
+	case ERigControlType::TransformNoScale:
+	case ERigControlType::EulerTransform:
 	{
-		UControlRig* InteractionRig = ControlRig->GetInteractionRig();
-		InteractionRig = InteractionRig == nullptr ? ControlRig : InteractionRig;
+		return true;
+	}
+	default:
+	{
+		break;
+	}
+	}
+
+	return false;
+}
 
-		// create gizmo actors
-		CreateShapeActors(GetWorld());
+void FControlRigEditMode::CreateShapeActors(UControlRig* ControlRig)
+{
+	// create gizmo actors
+	FActorSpawnParameters ActorSpawnParameters;
+	ActorSpawnParameters.bTemporaryEditorActor = true;
 
-		USceneComponent* Component = GetHostingSceneComponent();
-		if (Component)
+	TArray<FRigControlElement*> Controls = ControlRig->AvailableControls();
+	const TArray<TSoftObjectPtr<UControlRigShapeLibrary>> ShapeLibraries = ControlRig->GetShapeLibraries();
+	int32 ControlRigIndex = RuntimeControlRigs.Find(ControlRig);
+	for (FRigControlElement* ControlElement : Controls)
+	{
+		if (!ControlElement->Settings.bShapeEnabled)
 		{
-			AActor* PreviewActor = Component->GetOwner();
+			continue;
+		}
+		if (IsSupportedControlType(ControlElement->Settings.ControlType))
+		{
+			FControlShapeActorCreationParam Param;
+			Param.ManipObj = ControlRig;
+			Param.ControlRigIndex = ControlRigIndex;
+			Param.ControlRig = ControlRig;
+			Param.ControlName = ControlElement->GetName();
+			Param.ShapeName = ControlElement->Settings.ShapeName;
+			Param.SpawnTransform = ControlRig->GetControlGlobalTransform(ControlElement->GetName());
+			Param.ShapeTransform = ControlRig->GetHierarchy()->GetControlShapeTransform(ControlElement, ERigTransformType::CurrentLocal);
+			Param.bSelectable = ControlElement->Settings.bAnimatable;
+
+			if (const FControlRigShapeDefinition* ShapeDef = UControlRigShapeLibrary::GetShapeByName(ControlElement->Settings.ShapeName, ShapeLibraries))
+			{
+				Param.MeshTransform = ShapeDef->Transform;
+				Param.StaticMesh = ShapeDef->StaticMesh;
+				Param.Material = ShapeDef->Library->DefaultMaterial;
+				Param.ColorParameterName = ShapeDef->Library->MaterialColorParameter;
+			}
+
+			Param.Color = ControlElement->Settings.ShapeColor;
+
+			AControlRigShapeActor* ShapeActor = FControlRigShapeHelper::CreateDefaultShapeActor(WorldPtr, Param);
+			if (ShapeActor)
+			{
+				TArray<AControlRigShapeActor*>* ShapeActors = ControlRigShapeActors.Find(ControlRig);
+				if (ShapeActors)
+				{
+					ShapeActors->Add(ShapeActor);
+				}
+				else
+				{
+					TArray<AControlRigShapeActor*> NewShapeActors;
+					NewShapeActors.Add(ShapeActor);
+					ControlRigShapeActors.Add(ControlRig, NewShapeActors);
+				}
+			}
+		}
+	}
+
+
+	USceneComponent* Component = GetHostingSceneComponent(ControlRig);
+	if (Component)
+	{
+		AActor* PreviewActor = Component->GetOwner();
 
-			for (AControlRigShapeActor* ShapeActor : ShapeActors)
+		const TArray<AControlRigShapeActor*>* ShapeActors = ControlRigShapeActors.Find(ControlRig);
+		if (ShapeActors)
+		{
+			for (AControlRigShapeActor* ShapeActor : *ShapeActors)
 			{
 				// attach to preview actor, so that we can communicate via relative transfrom from the previewactor
 				ShapeActor->AttachToActor(PreviewActor, FAttachmentTransformRules::KeepWorldTransform);
@@ -2519,22 +2805,25 @@ void FControlRigEditMode::RecreateControlShapeActors(const TArray<FRigElementKey
 				}
 			}
 		}
-		if (IsInLevelEditor())
+	}
+	if (IsInLevelEditor())
+	{
+
+		if (ControlProxy)
 		{
-			if (ControlProxy)
-			{
-				ControlProxy->RecreateAllProxies(InteractionRig);
-			}
+			ControlProxy->RecreateAllProxies(ControlRig);
 		}
-
-		for (const FRigElementKey& SelectedElement : InSelectedElements)
+	}
+	/** MZ got rid of this make sure it's okay
+	for (const FRigElementKey& SelectedElement : InSelectedElements)
+	{
+		if(FRigControlElement* ControlElement = ControlRig->FindControl(SelectedElement.Name))
 		{
-			if(FRigControlElement* ControlElement = ControlRig->FindControl(SelectedElement.Name))
-			{
-				OnHierarchyModified(ERigHierarchyNotification::ElementSelected, ControlRig->GetHierarchy(), ControlElement);
-			}
+			OnHierarchyModified(ERigHierarchyNotification::ElementSelected, ControlRig->GetHierarchy(), ControlElement);
 		}
 	}
+	*/
+	
 }
 
 FControlRigEditMode* FControlRigEditMode::GetEditModeFromWorldContext(UWorld* InWorldContext)
@@ -2577,7 +2866,7 @@ void FControlRigEditMode::OnObjectsReplaced(const TMap<UObject*, UObject*>& OldT
 			RuntimeControlRigs.Reset();
 
 			UControlRig* NewRig = Cast<UControlRig>(NewObject);
-			AddControlRig(NewRig);
+			AddControlRigInternal(NewRig);
 
 			NewRig->Initialize();
 
@@ -2591,9 +2880,11 @@ bool FControlRigEditMode::IsTransformDelegateAvailable() const
 	return (OnGetRigElementTransformDelegate.IsBound() && OnSetRigElementTransformDelegate.IsBound());
 }
 
-bool FControlRigEditMode::AreRigElementSelectedAndMovable() const
+bool FControlRigEditMode::AreRigElementSelectedAndMovable(UControlRig* ControlRig) const
 {
-	if (!AreRigElementsSelected(FRigElementTypeHelper::ToMask(ERigElementType::Control)))
+	const UControlRigEditModeSettings* Settings = GetDefault<UControlRigEditModeSettings>();
+
+	if (!Settings || Settings->bHideControlShapes || !ControlRig || ControlRig->GetControlsVisible() == false || !AreRigElementsSelected(FRigElementTypeHelper::ToMask(ERigElementType::Control), ControlRig))
 	{
 		return false;
 	}
@@ -2607,9 +2898,24 @@ bool FControlRigEditMode::AreRigElementSelectedAndMovable() const
 		}
 	}
 
-		return true;
+	return true;
 }
 
+void FControlRigEditMode::ReplaceControlRig(UControlRig* OldControlRig, UControlRig* NewControlRig)
+{
+	for (int32 Index = 0; Index < RuntimeControlRigs.Num(); ++Index)
+	{
+		TWeakObjectPtr<UControlRig>& RuntimeRigPtr = RuntimeControlRigs[Index];
+		if (UControlRig* RuntimeControlRig = RuntimeRigPtr.Get())
+		{
+			if (RuntimeControlRig == OldControlRig)
+			{
+				RuntimeControlRigs[Index] = NewControlRig;
+				break;
+			}
+		}
+	}
+}
 void FControlRigEditMode::OnHierarchyModified(ERigHierarchyNotification InNotif, URigHierarchy* InHierarchy, const FRigBaseElement* InElement)
 {
 	if(bSuspendHierarchyNotifs)
@@ -2624,31 +2930,24 @@ void FControlRigEditMode::OnHierarchyModified(ERigHierarchyNotification InNotif,
 		case ERigHierarchyNotification::ElementRenamed:
 		case ERigHierarchyNotification::HierarchyReset:
 		{
-			RequestToRecreateControlShapeActors();
+			UControlRig* ControlRig = InHierarchy->GetTypedOuter<UControlRig>();
+			RequestToRecreateControlShapeActors(ControlRig);
 			break;
 		}
 		case ERigHierarchyNotification::ControlSettingChanged:
 		case ERigHierarchyNotification::ControlShapeTransformChanged:
 		{
-			// in case the gizmo is turned off, automatically exit gizmo transform edit mode
-			if (bIsChangingControlShapeTransform)
-			{
-				if (!CanChangeControlShapeTransform())
-				{
-					bIsChangingControlShapeTransform = false;
-				} 
-			}
-
 			const FRigElementKey Key = InElement->GetKey();
-			if(Key.Type == ERigElementType::Control)
+			UControlRig* ControlRig = InHierarchy->GetTypedOuter<UControlRig>();
+			if (Key.Type == ERigElementType::Control)
 			{
-				if(const FRigControlElement* ControlElement = Cast<FRigControlElement>(InElement))
+				if (const FRigControlElement* ControlElement = Cast<FRigControlElement>(InElement))
 				{
-					if (AControlRigShapeActor* ShapeActor = GetControlShapeFromControlName(Key.Name))
+					if (AControlRigShapeActor* ShapeActor = GetControlShapeFromControlName(ControlRig,Key.Name))
 					{
 						// try to lazily apply the changes to the actor
 						const UControlRigEditModeSettings* Settings = GetDefault<UControlRigEditModeSettings>();
-						if(ShapeActor->UpdateControlSettings(InNotif, GetControlRig(true), ControlElement, Settings->bHideControlShapes, IsInLevelEditor()))
+						if (ShapeActor->UpdateControlSettings(InNotif, ControlRig, ControlElement, Settings->bHideControlShapes, IsInLevelEditor()))
 						{
 							break;
 						}
@@ -2657,8 +2956,8 @@ void FControlRigEditMode::OnHierarchyModified(ERigHierarchyNotification InNotif,
 			}
 
 			// if we can't deal with this lazily, let's fall back to recreating all control shape actors
-			RequestToRecreateControlShapeActors();
-			break;
+			RequestToRecreateControlShapeActors(ControlRig);
+
 		}
 		case ERigHierarchyNotification::ElementSelected:
 		case ERigHierarchyNotification::ElementDeselected:
@@ -2686,7 +2985,15 @@ void FControlRigEditMode::OnHierarchyModified(ERigHierarchyNotification InNotif,
 						}
 						// users may select gizmo and control rig units, so we have to let them go through both of them if they do
 						// first go through gizmo actor
-						AControlRigShapeActor* ShapeActor = GetControlShapeFromControlName(Key.Name);
+						UControlRig* ControlRig= InHierarchy->GetTypedOuter<UControlRig>();
+						if (ControlRig == nullptr)
+						{
+							if (RuntimeControlRigs.Num() > 0)
+							{
+								ControlRig = RuntimeControlRigs[0].Get();
+							}
+						}
+						AControlRigShapeActor* ShapeActor = GetControlShapeFromControlName(ControlRig,Key.Name);
 						if (ShapeActor)
 						{
 							ShapeActor->SetSelected(bSelected);
@@ -2696,17 +3003,14 @@ void FControlRigEditMode::OnHierarchyModified(ERigHierarchyNotification InNotif,
 						{
 							if (bSelected)
 							{
-								if (UControlRig* ControlRig = GetControlRig(true))
+								if (ControlRig->GetHierarchy()->Find<FRigControlElement>(Key))
 								{
-									if (ControlRig->GetHierarchy()->Find<FRigControlElement>(Key))
-									{
-										ControlProxy->SelectProxy(Key.Name, true);
-									}
+									ControlProxy->SelectProxy(ControlRig,Key.Name, true);
 								}
 							}
 							else
 							{
-								ControlProxy->SelectProxy(Key.Name, false);
+								ControlProxy->SelectProxy(ControlRig,Key.Name, false);
 							}
 						}
 					}
@@ -2733,7 +3037,7 @@ void FControlRigEditMode::OnHierarchyModified(ERigHierarchyNotification InNotif,
 void FControlRigEditMode::OnControlModified(UControlRig* Subject, FRigControlElement* InControlElement, const FRigControlModifiedContext& Context)
 {
 	//this makes sure the details panel ui get's updated, don't remove
-	ControlProxy->ProxyChanged(InControlElement->GetName());
+	ControlProxy->ProxyChanged(Subject,InControlElement->GetName());
 
 	/*
 	FScopedTransaction ScopedTransaction(LOCTEXT("ModifyControlTransaction", "Modify Control"),!GIsTransacting && Context.SetKey != EControlRigSetKey::Never);
@@ -2784,30 +3088,34 @@ bool FControlRigEditMode::CanChangeControlShapeTransform()
 {
 	if (!IsInLevelEditor())
 	{
-		TArray<FRigElementKey> SelectedRigElements = GetSelectedRigElements();
-		// do not allow multi-select
-		if (SelectedRigElements.Num() == 1)
+		for (TWeakObjectPtr<UControlRig> RuntimeRigPtr : RuntimeControlRigs)
 		{
-			if (AreRigElementsSelected(FRigElementTypeHelper::ToMask(ERigElementType::Control)))
+			if (UControlRig* ControlRig = RuntimeRigPtr.Get())
 			{
-				if (UControlRig* ControlRig = GetControlRig(true))
+				TArray<FRigElementKey> SelectedRigElements = GetSelectedRigElements(ControlRig);
+				// do not allow multi-select
+				if (SelectedRigElements.Num() == 1)
 				{
-					// only enable for a Control with Gizmo enabled and visible
-					if (FRigControlElement* ControlElement = ControlRig->GetHierarchy()->Find<FRigControlElement>(SelectedRigElements[0]))
+					if (AreRigElementsSelected(FRigElementTypeHelper::ToMask(ERigElementType::Control),ControlRig))
 					{
-						if (ControlElement->Settings.bShapeEnabled && ControlElement->Settings.bShapeVisible)
+						// only enable for a Control with Gizmo enabled and visible
+						if (FRigControlElement* ControlElement = ControlRig->GetHierarchy()->Find<FRigControlElement>(SelectedRigElements[0]))
 						{
-							if (AControlRigShapeActor* ShapeActor = GetControlShapeFromControlName(SelectedRigElements[0].Name))
+							if (ControlElement->Settings.bShapeEnabled && ControlElement->Settings.bShapeVisible)
 							{
-								if (ensure(ShapeActor->IsSelected()))
+								if (AControlRigShapeActor* ShapeActor = GetControlShapeFromControlName(ControlRig,SelectedRigElements[0].Name))
 								{
-									return true;
+									if (ensure(ShapeActor->IsSelected()))
+									{
+										return true;
+									}
 								}
 							}
 						}
+						
 					}
 				}
-			} 
+			}
 		}
 	}
 
@@ -2816,7 +3124,7 @@ bool FControlRigEditMode::CanChangeControlShapeTransform()
 
 void FControlRigEditMode::SetControlShapeTransform(AControlRigShapeActor* ShapeActor, const FTransform& InTransform)
 {
-	if (UControlRig* ControlRig = GetControlRig(true, ShapeActor->ControlRigIndex))
+	if (UControlRig* ControlRig = ShapeActor->ControlRig.Get())
 	{
 		ControlRig->SetControlGlobalTransform(ShapeActor->ControlName, InTransform);
 	}
@@ -2824,7 +3132,7 @@ void FControlRigEditMode::SetControlShapeTransform(AControlRigShapeActor* ShapeA
 
 FTransform FControlRigEditMode::GetControlShapeTransform(AControlRigShapeActor* ShapeActor) const
 {
-	if (UControlRig* ControlRig = GetControlRig(true, ShapeActor->ControlRigIndex))
+	if (UControlRig* ControlRig = ShapeActor->ControlRig.Get())
 	{
 		return ControlRig->GetControlGlobalTransform(ShapeActor->ControlName);
 	}
@@ -2840,13 +3148,11 @@ void FControlRigEditMode::MoveControlShape(AControlRigShapeActor* ShapeActor, co
 	//first case is where we do all controls by the local diff.
 	if (bUseLocal)
 	{
-		if (UControlRig* RuntimeControlRig = GetControlRig(false, ShapeActor->ControlRigIndex))
+		if (UControlRig* ControlRig = ShapeActor->ControlRig.Get())
 		{
-			UControlRig* InteractionControlRig = GetControlRig(true, ShapeActor->ControlRigIndex);
-
 			FRigControlModifiedContext Context;
 			Context.EventName = FRigUnit_BeginExecution::EventName;
-			FTransform CurrentLocalTransform = InteractionControlRig->GetControlLocalTransform(ShapeActor->ControlName);
+			FTransform CurrentLocalTransform = ControlRig->GetControlLocalTransform(ShapeActor->ControlName);
 			if (bRotation)
 			{
 
@@ -2866,17 +3172,14 @@ void FControlRigEditMode::MoveControlShape(AControlRigShapeActor* ShapeActor, co
 
 			if (bTransformChanged)
 			{
-				InteractionControlRig->SetControlLocalTransform(ShapeActor->ControlName, CurrentLocalTransform);
+				ControlRig->SetControlLocalTransform(ShapeActor->ControlName, CurrentLocalTransform);
 
-				FTransform CurrentTransform  = InteractionControlRig->GetControlGlobalTransform(ShapeActor->ControlName);			// assumes it's attached to actor
+				FTransform CurrentTransform  = ControlRig->GetControlGlobalTransform(ShapeActor->ControlName);			// assumes it's attached to actor
 				CurrentTransform = ToWorldTransform * CurrentTransform;
 
 				ShapeActor->SetGlobalTransform(CurrentTransform);
 
-				if (RuntimeControlRig->GetInteractionRig() == InteractionControlRig)
-				{
-					InteractionControlRig->Evaluate_AnyThread();
-				}
+				ControlRig->Evaluate_AnyThread();
 			}
 		}
 	}
@@ -2910,51 +3213,45 @@ void FControlRigEditMode::MoveControlShape(AControlRigShapeActor* ShapeActor, co
 
 		if (bTransformChanged)
 		{
-			if (UControlRig* RuntimeControlRig = GetControlRig(false, ShapeActor->ControlRigIndex))
+			if (UControlRig* ControlRig = ShapeActor->ControlRig.Get())
 			{
-				UControlRig* InteractionControlRig = GetControlRig(true, ShapeActor->ControlRigIndex);
-
 				FTransform NewTransform = CurrentTransform.GetRelativeTransform(ToWorldTransform);
 				FRigControlModifiedContext Context;
 				Context.EventName = FRigUnit_BeginExecution::EventName;
 				if (bCalcLocal)
 				{
-					InOutLocal = InteractionControlRig->GetControlLocalTransform(ShapeActor->ControlName);
+					InOutLocal = ControlRig->GetControlLocalTransform(ShapeActor->ControlName);
 				}
 
 				bool bPrintPythonCommands = false;
-				if (UWorld* World = InteractionControlRig->GetWorld())
+				if (UWorld* World = ControlRig->GetWorld())
 				{
 					bPrintPythonCommands = World->IsPreviewWorld();
 				}
-				InteractionControlRig->SetControlGlobalTransform(ShapeActor->ControlName, NewTransform, true, Context, true, bPrintPythonCommands);			// assumes it's attached to actor
+				ControlRig->SetControlGlobalTransform(ShapeActor->ControlName, NewTransform, true, Context, true, bPrintPythonCommands);			// assumes it's attached to actor
 				ShapeActor->SetGlobalTransform(CurrentTransform);
 				if (bCalcLocal)
 				{
-					FTransform NewLocal = InteractionControlRig->GetControlLocalTransform(ShapeActor->ControlName);
+					FTransform NewLocal = ControlRig->GetControlLocalTransform(ShapeActor->ControlName);
 					InOutLocal = NewLocal.GetRelativeTransform(InOutLocal);
 
 				}
-
-				if (RuntimeControlRig->GetInteractionRig() == InteractionControlRig)
-				{
-					InteractionControlRig->Evaluate_AnyThread();
-				}
+				ControlRig->Evaluate_AnyThread();		
 			}
 		}
 	}
 #if WITH_EDITOR
 	if (bTransformChanged)
 	{
-		if (UControlRig* RuntimeControlRig = GetControlRig(false, ShapeActor->ControlRigIndex))
+		if (UControlRig* ControlRig = ShapeActor->ControlRig.Get())
 		{
-			if (UWorld* World = RuntimeControlRig->GetWorld())
+			if (UWorld* World = ControlRig->GetWorld())
 			{
 				if (World->IsPreviewWorld())
 				{
-					if (UControlRigBlueprint* Blueprint = Cast<UControlRigBlueprint>(RuntimeControlRig->GetClass()->ClassGeneratedBy))
+					if (UControlRigBlueprint* Blueprint = Cast<UControlRigBlueprint>(ControlRig->GetClass()->ClassGeneratedBy))
 					{
-						Blueprint->PropagatePoseFromInstanceToBP(RuntimeControlRig);
+						Blueprint->PropagatePoseFromInstanceToBP(ControlRig);
 					}
 				}
 			}
@@ -2970,7 +3267,7 @@ void FControlRigEditMode::ChangeControlShapeTransform(AControlRigShapeActor* Sha
 
 	FTransform CurrentTransform = FTransform::Identity;
 
-	if (UControlRig* ControlRig = GetControlRig(true))
+	if (UControlRig* ControlRig = ShapeActor->ControlRig.Get())
 	{
 		if (FRigControlElement* ControlElement = ControlRig->GetHierarchy()->Find<FRigControlElement>(FRigElementKey(ShapeActor->ControlName, ERigElementType::Control)))
 		{
@@ -3005,67 +3302,36 @@ void FControlRigEditMode::ChangeControlShapeTransform(AControlRigShapeActor* Sha
 
 	if (bTransformChanged)
 	{
-		if (UControlRig* RuntimeControlRig = GetControlRig(false, ShapeActor->ControlRigIndex))
+		if (UControlRig* ControlRig = ShapeActor->ControlRig.Get())
 		{
-			UControlRig* InteractionControlRig = GetControlRig(true, ShapeActor->ControlRigIndex);
 
 			FTransform NewTransform = CurrentTransform.GetRelativeTransform(ToWorldTransform);
 			FRigControlModifiedContext Context;
 
-			if (FRigControlElement* ControlElement = InteractionControlRig->GetHierarchy()->Find<FRigControlElement>(FRigElementKey(ShapeActor->ControlName, ERigElementType::Control)))
+			if (FRigControlElement* ControlElement = ControlRig->GetHierarchy()->Find<FRigControlElement>(FRigElementKey(ShapeActor->ControlName, ERigElementType::Control)))
 			{
 				// do not setup undo for this first step since it is just used to calculate the local transform
-				InteractionControlRig->GetHierarchy()->SetControlShapeTransform(ControlElement, NewTransform, ERigTransformType::CurrentGlobal, false);
-				FTransform CurrentLocalShapeTransform = InteractionControlRig->GetHierarchy()->GetControlShapeTransform(ControlElement, ERigTransformType::CurrentLocal);
+				ControlRig->GetHierarchy()->SetControlShapeTransform(ControlElement, NewTransform, ERigTransformType::CurrentGlobal, false);
+				FTransform CurrentLocalShapeTransform = ControlRig->GetHierarchy()->GetControlShapeTransform(ControlElement, ERigTransformType::CurrentLocal);
 				// this call should trigger an instance-to-BP update in ControlRigEditor
-				InteractionControlRig->GetHierarchy()->SetControlShapeTransform(ControlElement, CurrentLocalShapeTransform, ERigTransformType::InitialLocal, true);
+				ControlRig->GetHierarchy()->SetControlShapeTransform(ControlElement, CurrentLocalShapeTransform, ERigTransformType::InitialLocal, true);
 
 				FTransform MeshTransform = FTransform::Identity;
 				FTransform ShapeTransform = CurrentLocalShapeTransform;
-
-				if (UControlRig* ControlRig = GetControlRig(true))
+				if (const FControlRigShapeDefinition* Gizmo = UControlRigShapeLibrary::GetShapeByName(ControlElement->Settings.ShapeName, ControlRig->GetShapeLibraries()))
 				{
-					if (const FControlRigShapeDefinition* Gizmo = UControlRigShapeLibrary::GetShapeByName(ControlElement->Settings.ShapeName, GetControlRig(true)->GetShapeLibraries()))
-					{
-						MeshTransform = Gizmo->Transform;
-					}
+					MeshTransform = Gizmo->Transform;
 				}
-
 				ShapeActor->StaticMeshComponent->SetRelativeTransform(MeshTransform * ShapeTransform);
 			}
 		}
 	} 
 }
 
-// temporarily we just support following types of gizmo
-bool IsSupportedControlType(const ERigControlType ControlType)
-{
-	switch (ControlType)
-	{
-		case ERigControlType::Float:
-		case ERigControlType::Integer:
-		case ERigControlType::Vector2D:
-		case ERigControlType::Position:
-		case ERigControlType::Scale:
-		case ERigControlType::Rotator:
-		case ERigControlType::Transform:
-		case ERigControlType::TransformNoScale:
-		case ERigControlType::EulerTransform:
-		{
-			return true;
-		}
-		default:
-		{
-			break;
-		}
-	}
-
-	return false;
-}
 
 bool FControlRigEditMode::ModeSupportedByShapeActor(const AControlRigShapeActor* ShapeActor, UE::Widget::EWidgetMode InMode) const
 {
-	if (UControlRig* ControlRig = GetControlRig(true, ShapeActor->ControlRigIndex))
+	if (UControlRig* ControlRig = ShapeActor->ControlRig.Get())
 	{
 		const FRigControlElement* ControlElement = ControlRig->FindControl(ShapeActor->ControlName);
 		if (ControlElement)
@@ -3163,10 +3429,10 @@ bool FControlRigEditMode::ModeSupportedByShapeActor(const AControlRigShapeActor*
 
 void FControlRigEditMode::TickControlShape(AControlRigShapeActor* ShapeActor, const FTransform& ComponentTransform)
 {
-	UControlRigEditModeSettings* Settings = GetMutableDefault<UControlRigEditModeSettings>();
+	const UControlRigEditModeSettings* Settings = GetDefault<UControlRigEditModeSettings>();
 	if (ShapeActor)
 	{
-		if (UControlRig* ControlRig = GetControlRig(true, ShapeActor->ControlRigIndex))
+		if (UControlRig* ControlRig = ShapeActor->ControlRig.Get())
 		{
 			const FTransform Transform = ControlRig->GetControlGlobalTransform(ShapeActor->ControlName);
 			ShapeActor->SetActorTransform(Transform * ComponentTransform);
@@ -3174,33 +3440,38 @@ void FControlRigEditMode::TickControlShape(AControlRigShapeActor* ShapeActor, co
 			if (FRigControlElement* ControlElement = ControlRig->FindControl(ShapeActor->ControlName))
 			{
 				ShapeActor->SetShapeColor(ControlElement->Settings.ShapeColor);
-				ShapeActor->SetIsTemporarilyHiddenInEditor(!ControlElement->Settings.bShapeVisible || Settings->bHideControlShapes);
+				ShapeActor->SetIsTemporarilyHiddenInEditor(!ControlElement->Settings.bShapeVisible || Settings->bHideControlShapes || !ControlRig->GetControlsVisible());
 				if (!IsInLevelEditor()) //don't change this in level editor otherwise we can never select it
 				{
-					ShapeActor->SetSelectable(ControlElement->Settings.bShapeVisible && !Settings->bHideControlShapes && ControlElement->Settings.bAnimatable);
+					ShapeActor->SetSelectable(ControlElement->Settings.bShapeVisible && !Settings->bHideControlShapes && ControlElement->Settings.bAnimatable && ControlRig->GetControlsVisible());
 				}
 			}
 		}
 	}
 }
 
-AControlRigShapeActor* FControlRigEditMode::GetControlShapeFromControlName(const FName& ControlName) const
+AControlRigShapeActor* FControlRigEditMode::GetControlShapeFromControlName(UControlRig* InControlRig,const FName& ControlName) const
 {
-	for (AControlRigShapeActor* ShapeActor : ShapeActors)
+	const TArray<AControlRigShapeActor*>* ShapeActors = ControlRigShapeActors.Find(InControlRig);
+	if (ShapeActors)
 	{
-		if (ShapeActor->ControlName == ControlName)
+		for (AControlRigShapeActor* ShapeActor : *ShapeActors)
 		{
-			return ShapeActor;
+			if (ShapeActor->ControlName == ControlName)
+			{
+				return ShapeActor;
+			}
 		}
 	}
 
 	return nullptr;
 }
 
-void FControlRigEditMode::AddControlRig(UControlRig* InControlRig)
+void FControlRigEditMode::AddControlRigInternal(UControlRig* InControlRig)
 {
 	RuntimeControlRigs.AddUnique(InControlRig);
 
+	InControlRig->SetControlsVisible(true);
 	InControlRig->PostInitInstanceIfRequired();
 	InControlRig->GetHierarchy()->OnModified().RemoveAll(this);
 	InControlRig->GetHierarchy()->OnModified().AddSP(this, &FControlRigEditMode::OnHierarchyModified);
@@ -3208,184 +3479,164 @@ void FControlRigEditMode::AddControlRig(UControlRig* InControlRig)
 	OnControlRigAddedOrRemovedDelegate.Broadcast(InControlRig, true);
 }
 
-UControlRig* FControlRigEditMode::GetControlRig(bool bInteractionRig, int32 InIndex) const
+TArrayView<const TWeakObjectPtr<UControlRig>> FControlRigEditMode::GetControlRigs() const
 {
-	if (!RuntimeControlRigs.IsValidIndex(InIndex))
-	{
-		return nullptr;
-	}
+	return MakeArrayView(RuntimeControlRigs);
+}
+
+TArrayView<TWeakObjectPtr<UControlRig>> FControlRigEditMode::GetControlRigs() 
+{
+	return MakeArrayView(RuntimeControlRigs);
+}
 
-	UControlRig* RuntimeControlRig = RuntimeControlRigs[InIndex].Get();
-	if (bInteractionRig && RuntimeControlRig)
+TArray<UControlRig*> FControlRigEditMode::GetControlRigsArray(bool bIsVisible)
+{
+	TArray < UControlRig*> ControlRigs;
+	for (TWeakObjectPtr<UControlRig> ControlRigPtr : RuntimeControlRigs)
 	{
-		if (UControlRig* InteractionControlRig = RuntimeControlRig->GetInteractionRig())
+		if (ControlRigPtr.IsValid() && ControlRigPtr.Get() != nullptr && (bIsVisible == false ||ControlRigPtr.Get()->GetControlsVisible()))
 		{
-			return InteractionControlRig;
+			ControlRigs.Add(ControlRigPtr.Get());
 		}
 	}
-	return RuntimeControlRig;
+	return ControlRigs;
 }
 
-void FControlRigEditMode::RemoveControlRig(UControlRig* InControlRig)
+TArray<const UControlRig*> FControlRigEditMode::GetControlRigsArray(bool bIsVisible) const
 {
-	int32 Index = RuntimeControlRigs.Find(InControlRig);
-	if (RuntimeControlRigs.IsValidIndex(Index))
+	TArray<const UControlRig*> ControlRigs;
+	for (const TWeakObjectPtr<UControlRig> ControlRigPtr : RuntimeControlRigs)
 	{
-		OnControlRigAddedOrRemovedDelegate.Broadcast(InControlRig, false);
-		RuntimeControlRigs[Index]->ControlModified().RemoveAll(this);
-		RuntimeControlRigs[Index]->GetHierarchy()->OnModified().RemoveAll(this);
-		RuntimeControlRigs.RemoveAt(Index);
-
-		DelegateHelper->RemoveDelegates();
+		if (ControlRigPtr.IsValid() && ControlRigPtr.Get() != nullptr && (bIsVisible == false || ControlRigPtr.Get()->GetControlsVisible()))
+		{
+			ControlRigs.Add(ControlRigPtr.Get());
+		}
 	}
+	return ControlRigs;
 }
 
-void FControlRigEditMode::TickManipulatableObjects(float DeltaTime)
+void FControlRigEditMode::RemoveControlRig(UControlRig* InControlRig)
 {
-	// tick skeletalmeshcomponent, that's how they update their transform from rig change
-	USceneComponent* SceneComponent = GetHostingSceneComponent();
-	if (UControlRigComponent* ControlRigComponent = Cast<UControlRigComponent>(SceneComponent))
+	InControlRig->ControlModified().RemoveAll(this);
+	InControlRig->GetHierarchy()->OnModified().RemoveAll(this);
+	int32 Index = RuntimeControlRigs.Find(InControlRig);
+	TStrongObjectPtr<UControlRigEditModeDelegateHelper>* DelegateHelper = DelegateHelpers.Find(InControlRig);
+	if (DelegateHelper && DelegateHelper->IsValid())
 	{
-		ControlRigComponent->Update();
+		DelegateHelper->Get()->RemoveDelegates();
+		DelegateHelper->Reset();
+		DelegateHelpers.Remove(InControlRig);
 	}
-	else if (USkeletalMeshComponent* MeshComponent = Cast<USkeletalMeshComponent>(SceneComponent))
+	DestroyShapesActors(InControlRig);
+	if (RuntimeControlRigs.IsValidIndex(Index))
 	{
-		MeshComponent->RefreshBoneTransforms();
-		MeshComponent->RefreshSlaveComponents();
-		MeshComponent->UpdateComponentToWorld();
-		MeshComponent->FinalizeBoneTransform();
-		MeshComponent->MarkRenderTransformDirty();
-		MeshComponent->MarkRenderDynamicDataDirty();
+		RuntimeControlRigs.RemoveAt(Index);
 	}
-
-	PostPoseUpdate();
+	OnControlRigAddedOrRemovedDelegate.Broadcast(InControlRig, false);
 }
 
-bool FControlRigEditMode::CreateShapeActors(UWorld* World)
+void FControlRigEditMode::TickManipulatableObjects(float DeltaTime)
 {
-	DestroyShapesActors();
-
-	FActorSpawnParameters ActorSpawnParameters;
-	ActorSpawnParameters.bTemporaryEditorActor = true;
-
-	// for now we only support FTransform
-	for (int32 ControlRigIndex = 0; ControlRigIndex < RuntimeControlRigs.Num(); ControlRigIndex++)
+	for (TWeakObjectPtr<UControlRig> RuntimeRigPtr : RuntimeControlRigs)
 	{
-		UControlRig* ControlRig = GetControlRig(true, ControlRigIndex);
-		if (ControlRig == nullptr)
-		{
-			continue;
-		}
-
-		TArray<FRigControlElement*> Controls = ControlRig->AvailableControls();
-		const TArray<TSoftObjectPtr<UControlRigShapeLibrary>> ShapeLibraries = ControlRig->GetShapeLibraries();
-
-		for (FRigControlElement* ControlElement : Controls)
+		if (UControlRig* ControlRig = RuntimeRigPtr.Get())
 		{
-			if (!ControlElement->Settings.bShapeEnabled)
+			// tick skeletalmeshcomponent, that's how they update their transform from rig change
+			USceneComponent* SceneComponent = GetHostingSceneComponent(ControlRig);
+			if (UControlRigComponent* ControlRigComponent = Cast<UControlRigComponent>(SceneComponent))
 			{
-				continue;
+				ControlRigComponent->Update();
 			}
-			if (IsSupportedControlType(ControlElement->Settings.ControlType))
+			else if (USkeletalMeshComponent* MeshComponent = Cast<USkeletalMeshComponent>(SceneComponent))
 			{
-				FControlShapeActorCreationParam Param;
-				Param.ManipObj = ControlRig;
-				Param.ControlRigIndex = ControlRigIndex;
-				Param.ControlName = ControlElement->GetName();
-				Param.ShapeName = ControlElement->Settings.ShapeName;
-				Param.SpawnTransform = ControlRig->GetControlGlobalTransform(ControlElement->GetName());
-				Param.ShapeTransform = ControlRig->GetHierarchy()->GetControlShapeTransform(ControlElement, ERigTransformType::CurrentLocal);
-				Param.bSelectable = ControlElement->Settings.bAnimatable;
-
-				if (const FControlRigShapeDefinition* ShapeDef = UControlRigShapeLibrary::GetShapeByName(ControlElement->Settings.ShapeName, ShapeLibraries))
-				{
-					Param.MeshTransform = ShapeDef->Transform;
-					Param.StaticMesh = ShapeDef->StaticMesh;
-					Param.Material = ShapeDef->Library->DefaultMaterial;
-					Param.ColorParameterName = ShapeDef->Library->MaterialColorParameter;
-				}
-
-				Param.Color = ControlElement->Settings.ShapeColor;
-
-				AControlRigShapeActor* ShapeActor = FControlRigShapeHelper::CreateDefaultShapeActor(World, Param);
-				if (ShapeActor)
-				{
-					ShapeActors.Add(ShapeActor);
-				}
+				MeshComponent->RefreshBoneTransforms();
+				MeshComponent->RefreshSlaveComponents();
+				MeshComponent->UpdateComponentToWorld();
+				MeshComponent->FinalizeBoneTransform();
+				MeshComponent->MarkRenderTransformDirty();
+				MeshComponent->MarkRenderDynamicDataDirty();
 			}
 		}
 	}
-
-	WorldPtr = World;
-	OnWorldCleanupHandle = FWorldDelegates::OnWorldCleanup.AddSP(this, &FControlRigEditMode::OnWorldCleanup);
-	return (ShapeActors.Num() > 0);
+	PostPoseUpdate();
 }
 
+
 void FControlRigEditMode::OnWorldCleanup(UWorld* World, bool bSessionEnded, bool bCleanupResources)
 {
 	// if world gets cleaned up first, we destroy gizmo actors
 	if (WorldPtr == World)
 	{
-		DestroyShapesActors();
+		DestroyShapesActors(nullptr);
 	}
 }
 
-void FControlRigEditMode::DestroyShapesActors()
+void FControlRigEditMode::DestroyShapesActors(UControlRig* ControlRig)
 {
-	for (AControlRigShapeActor* ShapeActor : ShapeActors)
+	if (ControlRig == nullptr)
 	{
-		UWorld* World = ShapeActor->GetWorld();
-		if (World)
+		for(TPair<UControlRig*, TArray<AControlRigShapeActor*>>& ShapeActors: ControlRigShapeActors)
 		{
-			World->DestroyActor(ShapeActor);
+			for (AControlRigShapeActor* ShapeActor : ShapeActors.Value)
+			{
+				UWorld* World = ShapeActor->GetWorld();
+				if (World)
+				{
+					World->DestroyActor(ShapeActor);
+				}
+			}
 		}
-	}
-	ShapeActors.Reset();
-
-	if (OnWorldCleanupHandle.IsValid())
-	{
-		FWorldDelegates::OnWorldCleanup.Remove(OnWorldCleanupHandle);
-	}
-}
+		ControlRigShapeActors.Reset();
 
-TSharedPtr<IControlRigObjectBinding> FControlRigEditMode::GetObjectBinding() const
-{
-	for (TWeakObjectPtr<UControlRig> ControlRig : RuntimeControlRigs)
-	{
-		if (ControlRig.IsValid())
+		if (OnWorldCleanupHandle.IsValid())
 		{
-			return ControlRig->GetObjectBinding();
+			FWorldDelegates::OnWorldCleanup.Remove(OnWorldCleanupHandle);
 		}
 	}
-
-	return TSharedPtr<IControlRigObjectBinding>();
-}
-
-void FControlRigEditMode::SetObjectBinding(TSharedPtr<IControlRigObjectBinding> InObjectBinding)
-{
-	for (TWeakObjectPtr<UControlRig> ControlRig : RuntimeControlRigs)
+	else
 	{
-		if (ControlRig.IsValid())
+		const TArray<AControlRigShapeActor*>* ShapeActors = ControlRigShapeActors.Find(ControlRig);
+		if (ShapeActors)
 		{
-			ControlRig->SetObjectBinding(InObjectBinding);
+			for (AControlRigShapeActor* ShapeActor : *ShapeActors)
+			{
+				UWorld* World = ShapeActor->GetWorld();
+				if (World)
+				{
+					World->DestroyActor(ShapeActor);
+				}
+			}
+			ControlRigShapeActors.Remove(ControlRig);
 		}
 	}
 }
 
-USceneComponent* FControlRigEditMode::GetHostingSceneComponent() const
+USceneComponent* FControlRigEditMode::GetHostingSceneComponent(const UControlRig* ControlRig) const
 {
-	TSharedPtr<IControlRigObjectBinding> ObjectBinding = GetObjectBinding();
-	if (ObjectBinding.IsValid())
+	
+	if (ControlRig == nullptr && GetControlRigs().Num() > 0)
+	{
+		ControlRig = GetControlRigs()[0].Get();
+	}
+	if (ControlRig)
 	{
-		return Cast<USceneComponent>(ObjectBinding->GetBoundObject());
+		TSharedPtr<IControlRigObjectBinding> ObjectBinding = ControlRig->GetObjectBinding();
+		if (ObjectBinding.IsValid())
+		{
+			return Cast<USceneComponent>(ObjectBinding->GetBoundObject());
+		}
 	}
 
 	return nullptr;
 }
 
-FTransform FControlRigEditMode::GetHostingSceneComponentTransform() const
+FTransform FControlRigEditMode::GetHostingSceneComponentTransform(const UControlRig* ControlRig) const
 {
-	USceneComponent* HostingComponent = GetHostingSceneComponent();
+	if (ControlRig == nullptr && GetControlRigs().Num() > 0)
+	{
+		ControlRig = GetControlRigs()[0].Get();
+	}
+	USceneComponent* HostingComponent = GetHostingSceneComponent(ControlRig);
 	return HostingComponent ? HostingComponent->GetComponentTransform() : FTransform::Identity;
 }
 
@@ -3396,15 +3647,17 @@ void FControlRigEditMode::OnPoseInitialized()
 
 void FControlRigEditMode::PostPoseUpdate()
 {
-	FTransform ComponentTransform = FTransform::Identity;
-	if(IsInLevelEditor())
+	for (TPair<UControlRig*, TArray<AControlRigShapeActor*>>& ShapeActors : ControlRigShapeActors)
 	{
-		ComponentTransform = GetHostingSceneComponentTransform();
-	}
-	
-	for(AControlRigShapeActor* ShapeActor : ShapeActors)
-	{
-		TickControlShape(ShapeActor, ComponentTransform);
+		FTransform ComponentTransform = FTransform::Identity;
+		if (IsInLevelEditor())
+		{
+			ComponentTransform = GetHostingSceneComponentTransform(ShapeActors.Key);
+		}
+		for (AControlRigShapeActor* ShapeActor : ShapeActors.Value)
+		{
+			TickControlShape(ShapeActor, ComponentTransform);	
+		}
 	}
 
 }
diff --git a/Engine/Plugins/Animation/ControlRig/Source/ControlRigEditor/Private/EditMode/ControlRigEditMode.h b/Engine/Plugins/Animation/ControlRig/Source/ControlRigEditor/Private/EditMode/ControlRigEditMode.h
--- a/Engine/Plugins/Animation/ControlRig/Source/ControlRigEditor/Private/EditMode/ControlRigEditMode.h
+++ b/Engine/Plugins/Animation/ControlRig/Source/ControlRigEditor/Private/EditMode/ControlRigEditMode.h
@@ -11,7 +11,7 @@
 #include "Framework/MultiBox/MultiBoxBuilder.h"
 #include "UObject/StrongObjectPtr.h"
 #include "UnrealWidgetFwd.h"
-#include "IControlRigEditMode.h"
+#include "IPersonaEditMode.h"
 #include "ControlRigEditMode.generated.h"
 
 
@@ -47,6 +47,14 @@ DECLARE_DELEGATE_RetVal(UToolMenu*, FOnGetContextMenu);
 
 class FControlRigEditMode;
 
+enum class ERecreateControlRigShape
+{
+	RecreateNone,
+	RecreateAll,
+	RecreateSpecified
+};
+
+
 UCLASS()
 class UControlRigEditModeDelegateHelper : public UObject
 {
@@ -70,7 +78,7 @@ private:
 	FDelegateHandle OnBoneTransformsFinalizedHandle;
 };
 
-class FControlRigEditMode : public IControlRigEditMode
+class FControlRigEditMode : public IPersonaEditMode
 {
 public:
 	static FName ModeName;
@@ -78,8 +86,18 @@ public:
 	FControlRigEditMode();
 	~FControlRigEditMode();
 
-	/** Set the objects to be displayed in the details panel */
-	virtual void SetObjects(const TWeakObjectPtr<>& InSelectedObject,  UObject* BindingObject, TWeakPtr<ISequencer> InSequencer) override;
+	/** Set the Control Rig Object to be active in the edit mode. You set both the Control Rig and a possible binding together with an optional Sequencer
+	 This will remove all other control rigs present and should be called for stand alone editors, like the Control Rig Editor*/
+	void SetObjects(UControlRig* InControlRig, UObject* BindingObject, TWeakPtr<ISequencer> InSequencer);
+
+	/** Add a Control Rig object if it doesn't exist, will return true if it was added, false if it wasn't since it's already there. You can also set the Sequencer.*/
+	bool AddControlRigObject(UControlRig* InControlRig, TWeakPtr<ISequencer> InSequencer);
+
+	/* Remove control rig */
+	void RemoveControlRig(UControlRig* InControlRig);
+
+	/*Replace old Control Rig with the New Control Rig, perhaps from a recompile in the level editor*/
+	void ReplaceControlRig(UControlRig* OldControlRig, UControlRig* NewControlRig);
 
 	/** This edit mode is re-used between the level editor and the control rig editor. Calling this indicates which context we are in */
 	bool IsInLevelEditor() const;
@@ -145,14 +163,17 @@ public:
 
 	FUICommandList* GetCommandBindings() const { return CommandBindings.Get(); }
 
-	/** Requests to recreate the shape actors in the next tick */
-	void RequestToRecreateControlShapeActors() { bRecreateControlShapesRequired = true; }
+	/** Requests to recreate the shape actors in the next tick. Will recreate only the ones for the specified
+	Control Rig, otherwise will recreate all of them*/
+	void RequestToRecreateControlShapeActors(UControlRig* ControlRig = nullptr); 
+
+
 
 protected:
 
 	// shape related functions wrt enable/selection
 	/** Get the node name from the property path */
-	AControlRigShapeActor* GetControlShapeFromControlName(const FName& InControlName) const;
+	AControlRigShapeActor* GetControlShapeFromControlName(UControlRig* InControlRig,const FName& InControlName) const;
 
 protected:
 	/** Helper function: set ControlRigs array to the details panel */
@@ -170,9 +191,12 @@ protected:
 	/** Handle selection internally */
 	void HandleSelectionChanged();
 
-	/** Toggles visibility of manipulators in the viewport */
+	/** Toggles visibility of acive control rig shapes in the viewport */
 	void ToggleManipulators();
 
+	/** Toggles visibility of all  control rig shapes in the viewport */
+	void ToggleAllManipulators();
+
 public:
 	
 	/** Clear Selection*/
@@ -220,7 +244,6 @@ private:
 	/** It creates if it doesn't have it */
 	void RecreateControlShapeActors(const TArray<FRigElementKey>& InSelectedElements = TArray<FRigElementKey>());
 
-
 	/** Let the preview scene know how we want to select components */
 	bool ShapeSelectionOverride(const UPrimitiveComponent* InComponent) const;
 
@@ -231,8 +254,8 @@ protected:
 
 	TWeakPtr<ISequencer> WeakSequencer;
 
-	/** The scope for the interaction */
-	FControlRigInteractionScope* InteractionScope;
+	/** The scope for the interaction, one per manipulated Control rig */
+	TMap<UControlRig*,FControlRigInteractionScope*> InteractionScopes;
 
 	/** Whether a manipulator actually made a change when transacting */
 	bool bManipulatorMadeChange;
@@ -243,8 +266,8 @@ protected:
 	/** If selection was changed, we set up proxies on next tick */
 	bool bSelectionChanged;
 
-	/** Cached transform of pivot point for selected Bones */
-	FTransform PivotTransform;
+	/** Cached transform of pivot point for selected objects for each Control Rig */
+	TMap<UControlRig*,FTransform> PivotTransforms;
 
 	/** Command bindings for keyboard shortcuts */
 	TSharedPtr<FUICommandList> CommandBindings;
@@ -262,51 +285,55 @@ protected:
 	FSimpleMulticastDelegate OnAnimSystemInitializedDelegate;
 	FControlRigAddedOrRemoved OnControlRigAddedOrRemovedDelegate;
 
-	/** GetSelectedRigElements */
-	TArray<FRigElementKey> GetSelectedRigElements() const;
+	/** GetSelectedRigElements, if InControlRig is nullptr get the first one */
+	TArray<FRigElementKey> GetSelectedRigElements(UControlRig* InControlRig = nullptr) const;
 
 	/* Flag to recreate shapes during tick */
-	bool bRecreateControlShapesRequired;
+	ERecreateControlRigShape RecreateControlShapesRequired;
+	/* List of Control Rigs we should recreate*/
+	TArray<UControlRig*> ControlRigsToRecreate;
 
 	/* Flag to temporarily disable handling notifs from the hierarchy */
 	bool bSuspendHierarchyNotifs;
 
 	/** Shape actors */
-	TArray<AControlRigShapeActor*> ShapeActors;
+	TMap<UControlRig*,TArray<AControlRigShapeActor*>> ControlRigShapeActors;
 	UControlRigDetailPanelControlProxies* ControlProxy;
 
 	/** Utility functions for UI/Some other viewport manipulation*/
 	bool IsControlSelected() const;
-	bool AreRigElementSelectedAndMovable() const;
+	bool AreRigElementSelectedAndMovable(UControlRig* InControlRig) const;
 	
 	/** Set initial transform handlers */
 	void OpenContextMenu(FEditorViewportClient* InViewportClient);
-
-	/* Set initial transform helpers*/
-	bool GetRigElementGlobalTransform(const FRigElementKey& InElement, FTransform& OutGlobalTransform) const;
 	
 public: 
 	/** Clear all selected RigElements */
 	void ClearRigElementSelection(uint32 InTypes);
 
 	/** Set a RigElement's selection state */
-	void SetRigElementSelection(ERigElementType Type, const FName& InRigElementName, bool bSelected);
-
-	void SetSelectedRigElement(const FName& InElementName, ERigElementType InRigElementType);
+	void SetRigElementSelection(UControlRig* ControlRig, ERigElementType Type, const FName& InRigElementName, bool bSelected);
 
 	/** Set multiple RigElement's selection states */
-	void SetRigElementSelection(ERigElementType Type, const TArray<FName>& InRigElementNames, bool bSelected);
+	void SetRigElementSelection(UControlRig* ControlRig, ERigElementType Type, const TArray<FName>& InRigElementNames, bool bSelected);
 
 	/** Check if any RigElements are selected */
-	bool AreRigElementsSelected(uint32 InTypes) const;
+	bool AreRigElementsSelected(uint32 InTypes, UControlRig* InControlRig) const;
 
 	/** Get the number of selected RigElements */
-	int32 GetNumSelectedRigElements(uint32 InTypes) const;
+	int32 GetNumSelectedRigElements(uint32 InTypes, UControlRig* InControlRig) const;
 
-	/** Get the control rig*/
-	UControlRig* GetControlRig(bool bInteractionRig, int32 InIndex = 0) const;
+	/** Get all of the selected Controls*/
+	void GetAllSelectedControls(TMap<UControlRig*, TArray<FRigElementKey>>& OutSelectedControls);
 
-	/** Get the detail proxiescontrol rig*/
+	/** Get all of the ControlRigs, maybe not valid anymore */
+	TArrayView<const TWeakObjectPtr<UControlRig>> GetControlRigs() const;
+	TArrayView<TWeakObjectPtr<UControlRig>> GetControlRigs();
+	/* Get valid  Control Rigs possibly just visible*/
+	TArray<UControlRig*> GetControlRigsArray(bool bIsVisible);
+	TArray<const UControlRig*> GetControlRigsArray(bool bIsVisible) const;
+
+	/** Get the detail proxies control rig*/
 	UControlRigDetailPanelControlProxies* GetDetailProxies() { return ControlProxy; }
 
 	/** Get Sequencer Driving This*/
@@ -315,8 +342,9 @@ public:
 	/** Suspend Rig Hierarchy Notifies*/
 	void SuspendHierarchyNotifs(bool bVal) { bSuspendHierarchyNotifs = bVal; }
 private:
+
 	/** Set a RigElement's selection state */
-	void SetRigElementSelectionInternal(ERigElementType Type, const FName& InRigElementName, bool bSelected);
+	void SetRigElementSelectionInternal(UControlRig* ControlRig, ERigElementType Type, const FName& InRigElementName, bool bSelected);
 	
 	FEditorViewportClient* CurrentViewportClient;
 
@@ -335,11 +363,14 @@ public:
 
 private:
 	TSet<FName> GetActiveControlsFromSequencer(UControlRig* ControlRig);
-	bool CreateShapeActors(UWorld* World);
-	void DestroyShapesActors();
+	bool SetSequencer(TWeakPtr<ISequencer> InSequencer);
 
-	void AddControlRig(UControlRig* InControlRig);
-	void RemoveControlRig(UControlRig* InControlRig);
+	/** Create/Delete for the specified ControlRig*/
+	void CreateShapeActors(UControlRig* InControlRig);
+	void DestroyShapesActors(UControlRig* InControlRig);
+
+	/*Internal function for adding ControlRig*/
+	void AddControlRigInternal(UControlRig* InControlRig);
 	void TickManipulatableObjects(float DeltaTime);
 
 	void SetControlShapeTransform(AControlRigShapeActor* ShapeActor, const FTransform& InTransform);
@@ -354,17 +385,13 @@ private:
 	void TickControlShape(AControlRigShapeActor* ShapeActor, const FTransform& ComponentTransform);
 	bool ModeSupportedByShapeActor(const AControlRigShapeActor* ShapeActor, UE::Widget::EWidgetMode InMode) const;
 
-	// Object binding
-	/** Setup bindings to a runtime object (or clear by passing in nullptr). */
-	void SetObjectBinding(TSharedPtr<IControlRigObjectBinding> InObjectBinding);
 
 protected:
 	
 	/** Get bindings to a runtime object */
-	TSharedPtr<IControlRigObjectBinding> GetObjectBinding() const;
-
-	USceneComponent* GetHostingSceneComponent() const;
-	FTransform	GetHostingSceneComponentTransform() const;
+	//If the passed in ControlRig is nullptr we use the first Control Rig(this can happen from the BP Editors).
+	USceneComponent* GetHostingSceneComponent(const UControlRig* ControlRig = nullptr) const;
+	FTransform	GetHostingSceneComponentTransform(const UControlRig* ControlRig =  nullptr) const;
 
 private:
 
@@ -377,9 +404,9 @@ private:
 	void OnWorldCleanup(UWorld* World, bool bSessionEnded, bool bCleanupResources);
 	UWorld* WorldPtr = nullptr;
 
+	
 	TArray<TWeakObjectPtr<UControlRig>> RuntimeControlRigs;
-
-	TStrongObjectPtr<UControlRigEditModeDelegateHelper> DelegateHelper;
+	TMap<UControlRig*,TStrongObjectPtr<UControlRigEditModeDelegateHelper>> DelegateHelpers;
 
 	TArray<FRigElementKey> DeferredItemsToFrame;
 
diff --git a/Engine/Plugins/Animation/ControlRig/Source/ControlRigEditor/Private/EditMode/ControlRigEditModeCommands.cpp b/Engine/Plugins/Animation/ControlRig/Source/ControlRigEditor/Private/EditMode/ControlRigEditModeCommands.cpp
--- a/Engine/Plugins/Animation/ControlRig/Source/ControlRigEditor/Private/EditMode/ControlRigEditModeCommands.cpp
+++ b/Engine/Plugins/Animation/ControlRig/Source/ControlRigEditor/Private/EditMode/ControlRigEditModeCommands.cpp
@@ -8,7 +8,8 @@ void FControlRigEditModeCommands::RegisterCommands()
 {
 	UI_COMMAND(ResetTransforms, "Reset Transform", "Reset the Controls Transforms", EUserInterfaceActionType::Button, FInputChord(EKeys::G, EModifierKey::Control));
 	UI_COMMAND(ResetAllTransforms, "Reset All Transform", "Reset all of the Controls Transforms", EUserInterfaceActionType::Button, FInputChord(EKeys::G, EModifierKey::Control | EModifierKey::Shift));
-	UI_COMMAND(ToggleManipulators, "Toggle Manipulators", "Toggles visibility of manipulators in the viewport", EUserInterfaceActionType::Button, FInputChord(EKeys::T));
+	UI_COMMAND(ToggleManipulators, "Toggle Shapes", "Toggles visibility of active control rig shapes in the viewport", EUserInterfaceActionType::Button, FInputChord(EKeys::T));
+	UI_COMMAND(ToggleAllManipulators, "Toggle All Shapes", "Toggles visibility of all control rig shapes in the viewport", EUserInterfaceActionType::Button, FInputChord(EKeys::T, EModifierKey::Control | EModifierKey::Shift));
 	UI_COMMAND(FrameSelection, "Frame Selection", "Focus the viewport on the current selection", EUserInterfaceActionType::Button, FInputChord(EKeys::F));
 	UI_COMMAND(ClearSelection, "Clear Selection", "Clear Selection", EUserInterfaceActionType::Button, FInputChord(EKeys::Escape));
 
diff --git a/Engine/Plugins/Animation/ControlRig/Source/ControlRigEditor/Private/EditMode/ControlRigEditModeCommands.h b/Engine/Plugins/Animation/ControlRig/Source/ControlRigEditor/Private/EditMode/ControlRigEditModeCommands.h
--- a/Engine/Plugins/Animation/ControlRig/Source/ControlRigEditor/Private/EditMode/ControlRigEditModeCommands.h
+++ b/Engine/Plugins/Animation/ControlRig/Source/ControlRigEditor/Private/EditMode/ControlRigEditModeCommands.h
@@ -20,9 +20,12 @@ public:
 	{}
 	
 
-	/** Toggles hiding all manipulators in the viewport */
+	/** Toggles hiding all manipulators on active control rig in the viewport */
 	TSharedPtr< FUICommandInfo > ToggleManipulators;
 
+	/** Toggles hiding all manipulators on all control rigsin the viewport */
+	TSharedPtr< FUICommandInfo > ToggleAllManipulators;
+
 	/** Reset Transforms for Controls */
 	TSharedPtr< FUICommandInfo > ResetTransforms;
 
diff --git a/Engine/Plugins/Animation/ControlRig/Source/ControlRigEditor/Private/EditMode/SControlRigBaseListWidget.cpp b/Engine/Plugins/Animation/ControlRig/Source/ControlRigEditor/Private/EditMode/SControlRigBaseListWidget.cpp
--- a/Engine/Plugins/Animation/ControlRig/Source/ControlRigEditor/Private/EditMode/SControlRigBaseListWidget.cpp
+++ b/Engine/Plugins/Animation/ControlRig/Source/ControlRigEditor/Private/EditMode/SControlRigBaseListWidget.cpp
@@ -417,15 +417,16 @@ void SControlRigPoseAnimSelectionToolbar::MakeControlRigAssetDialog(FControlRigA
 	{
 		return;
 	}
+
+	TMap<UControlRig*, TArray<FRigElementKey>> AllSelectedControls;
+	ControlRigEditMode->GetAllSelectedControls(AllSelectedControls);
 	
-	UControlRig* ControlRig = ControlRigEditMode->GetControlRig(true);
-	if (!ControlRig)
+	if (AllSelectedControls.Num() > 1)
 	{
 		return;
 	}
 
-	TArray<FName> SelectedControls = ControlRig->CurrentControlSelection();
-	if (SelectedControls.Num() <= 0)
+	if (AllSelectedControls.Num() <= 0)
 	{
 		FText ConfirmDelete = LOCTEXT("ConfirmNoSelectedControls", "You are saving a Pose with no selected Controls - are you sure?");
 
@@ -447,26 +448,34 @@ void SControlRigPoseAnimSelectionToolbar::MakeControlRigAssetDialog(FControlRigA
 			FString Path = OwningControlRigWidget->GetCurrentlySelectedPath();
 
 			FControlRigEditMode* ControlRigEditMode = static_cast<FControlRigEditMode*>(GLevelEditorModeTools().GetActiveMode(FControlRigEditMode::ModeName));
-			if (ControlRigEditMode && ControlRigEditMode->GetControlRig(true))
-			{	
-				UObject* NewAsset = nullptr;
-				switch (Type)
-				{
-				case FControlRigAssetType::ControlRigPose:
-					NewAsset = FControlRigToolAsset::SaveAsset<UControlRigPoseAsset>(ControlRigEditMode->GetControlRig(true), Path, AssetName, bSelectAll);
-
-					break;
-				case FControlRigAssetType::ControlRigAnimation:
-					break;
-				case FControlRigAssetType::ControlRigSelectionSet:
-					break;
-				default:
-					break;
-				};
-				if (NewAsset)
+			if (ControlRigEditMode )
+			{
+				TMap<UControlRig*, TArray<FRigElementKey>> AllSelectedControls;
+				ControlRigEditMode->GetAllSelectedControls(AllSelectedControls);
+				if (AllSelectedControls.Num() == 1)
 				{
-					FControlRigView::CaptureThumbnail(NewAsset);
-					OwningControlRigWidget->SelectThisAsset(NewAsset);
+					TArray<UControlRig*> ControlRigs;
+					AllSelectedControls.GenerateKeyArray(ControlRigs);
+					UControlRig* ControlRig =ControlRigs[0];
+					UObject* NewAsset = nullptr;
+					switch (Type)
+					{
+					case FControlRigAssetType::ControlRigPose:
+						NewAsset = FControlRigToolAsset::SaveAsset<UControlRigPoseAsset>(ControlRig, Path, AssetName, bSelectAll);
+
+						break;
+					case FControlRigAssetType::ControlRigAnimation:
+						break;
+					case FControlRigAssetType::ControlRigSelectionSet:
+						break;
+					default:
+						break;
+					};
+					if (NewAsset)
+					{
+						FControlRigView::CaptureThumbnail(NewAsset);
+						OwningControlRigWidget->SelectThisAsset(NewAsset);
+					}
 				}
 			}
 		}
@@ -483,8 +492,10 @@ bool SControlRigPoseAnimSelectionToolbar::CanExecuteMakeControlRigAsset()
 		return false;
 	}
 	
-	UControlRig* ControlRig = ControlRigEditMode->GetControlRig(true);
-	if (!ControlRig)
+	TMap<UControlRig*, TArray<FRigElementKey>> AllSelectedControls;
+	ControlRigEditMode->GetAllSelectedControls(AllSelectedControls);
+
+	if (AllSelectedControls.Num() != 1)
 	{
 		return false;
 	}
@@ -828,14 +839,6 @@ void SControlRigBaseListWidget::NotifyUser(FNotificationInfo& NotificationInfo)
 	}
 }
 
-UControlRig* SControlRigBaseListWidget::GetControlRig()
-{
-	if (FControlRigEditMode* EditMode = GetEditMode())
-	{
-		return EditMode->GetControlRig(true);
-	}
-	return nullptr;
-}
 FControlRigEditMode* SControlRigBaseListWidget::GetEditMode()
 {
 	FControlRigEditMode* ControlRigEditMode = static_cast<FControlRigEditMode*>(GLevelEditorModeTools().GetActiveMode(FControlRigEditMode::ModeName));
@@ -958,10 +961,19 @@ void SControlRigBaseListWidget::OnAssetsActivated(const TArray<FAssetData>& Sele
 				// If alt is down, select controls
 				if (FSlateApplication::Get().GetModifierKeys().IsAltDown())
 				{
-					UControlRig* ControlRig = SControlRigPoseView::GetFirstControlRigInLevelSequence(GetControlRig());
-					if (ControlRig)
+					FControlRigEditMode* ControlRigEditMode = static_cast<FControlRigEditMode*>(GLevelEditorModeTools().GetActiveMode(FControlRigEditMode::ModeName));
+					if (ControlRigEditMode)
 					{
-						PoseAsset->SelectControls(ControlRig, SControlRigPoseView::IsMirror());
+						TArray<UControlRig*> ControlRigs = ControlRigEditMode->GetControlRigsArray(true /*bIsVisible*/);
+						if (ControlRigs.Num() > 0)
+						{
+							const FScopedTransaction Transaction(LOCTEXT("SelectControls", "Select Controls"));
+							for (UControlRig* ControlRig : ControlRigs)
+							{
+								ControlRig->Modify();
+								PoseAsset->SelectControls(ControlRig, SControlRigPoseView::IsMirror());
+							}
+						}
 					}
 				}
 			}
@@ -1328,11 +1340,18 @@ void SControlRigBaseListWidget::ExecutePastePose(UControlRigPoseAsset* PoseAsset
 	if (PoseAsset)
 	{
 		FControlRigEditMode* ControlRigEditMode = static_cast<FControlRigEditMode*>(GLevelEditorModeTools().GetActiveMode(FControlRigEditMode::ModeName));
-		if (ControlRigEditMode && ControlRigEditMode->GetControlRig(true))
+		if (ControlRigEditMode)
 		{
-			const FScopedTransaction Transaction(LOCTEXT("PastePose", "Paste Pose"));
-			ControlRigEditMode->GetControlRig(true)->Modify();
-			PoseAsset->PastePose(ControlRigEditMode->GetControlRig(true));
+			TArray<UControlRig*> ControlRigs = ControlRigEditMode->GetControlRigsArray(true /*bIsVisible*/);
+			if (ControlRigs.Num() > 0)
+			{
+				const FScopedTransaction Transaction(LOCTEXT("PastePose", "Paste Pose"));
+				for (UControlRig* ControlRig : ControlRigs)
+				{
+					ControlRig->Modify();
+					PoseAsset->PastePose(ControlRig);
+				}
+			}
 		}
 	}
 }
@@ -1350,24 +1369,32 @@ void SControlRigBaseListWidget::ExecuteUpdatePose(UControlRigPoseAsset* PoseAsse
 void SControlRigBaseListWidget::ExecuteSelectControls(UControlRigPoseAsset* PoseAsset)
 {
 	FControlRigEditMode* ControlRigEditMode = static_cast<FControlRigEditMode*>(GLevelEditorModeTools().GetActiveMode(FControlRigEditMode::ModeName));
-	if (ControlRigEditMode && ControlRigEditMode->GetControlRig(true))
+	if (ControlRigEditMode)
 	{
-		const FScopedTransaction Transaction(LOCTEXT("SelectControls", "Select Controls"));
-		ControlRigEditMode->GetControlRig(true)->Modify();
-		PoseAsset->SelectControls(ControlRigEditMode->GetControlRig(true));
+		TArray<UControlRig*> ControlRigs = ControlRigEditMode->GetControlRigsArray(true /*bIsVisible*/);
+		if (ControlRigs.Num() > 0)
+		{
+			const FScopedTransaction Transaction(LOCTEXT("SelectControls", "Select Controls"));
+			for (UControlRig* ControlRig : ControlRigs)
+			{
+				ControlRig->Modify();
+				PoseAsset->SelectControls(ControlRig);
+			}
+		}
 	}
 	else
 	{
 		ULevelSequence *LevelSequence = ULevelSequenceEditorBlueprintLibrary::GetCurrentLevelSequence();
 		if (LevelSequence)
 		{
+			const FScopedTransaction Transaction(LOCTEXT("SelectControls", "Select Controls"));
 			TArray<FControlRigSequencerBindingProxy> Proxies = UControlRigSequencerEditorLibrary::GetControlRigs(LevelSequence);
-			if (Proxies.Num() > 0)
+			for (FControlRigSequencerBindingProxy& Proxy: Proxies)
 			{
-				//MZ TODO when we have Mutliple Control Rig's active select more than one.
-				if (Proxies[0].ControlRig)
+				if (Proxy.ControlRig)
 				{
-					PoseAsset->SelectControls(Proxies[0].ControlRig);
+					Proxy.ControlRig->Modify();
+					PoseAsset->SelectControls(Proxy.ControlRig);
 				}
 			}
 		}
@@ -1379,11 +1406,18 @@ void SControlRigBaseListWidget::ExecutePasteMirrorPose(UControlRigPoseAsset* Pos
 	if (PoseAsset)
 	{
 		FControlRigEditMode* ControlRigEditMode = static_cast<FControlRigEditMode*>(GLevelEditorModeTools().GetActiveMode(FControlRigEditMode::ModeName));
-		if (ControlRigEditMode && ControlRigEditMode->GetControlRig(true))
+		if (ControlRigEditMode)
 		{
-			const FScopedTransaction Transaction(LOCTEXT("PasteMirrorPose", "Paste Mirror Pose"));
-			ControlRigEditMode->GetControlRig(true)->Modify();
-			PoseAsset->PastePose(ControlRigEditMode->GetControlRig(true), false, true);
+			TArray<UControlRig*> ControlRigs = ControlRigEditMode->GetControlRigsArray(true /*bIsVisible*/);
+			if (ControlRigs.Num() > 0)
+			{
+				const FScopedTransaction Transaction(LOCTEXT("PasteMirrorPose", "Paste Mirror Pose"));
+				for (UControlRig* ControlRig : ControlRigs)
+				{
+					ControlRig->Modify();
+					PoseAsset->PastePose(ControlRig, false, true);
+				}
+			}
 		}
 	}
 }
diff --git a/Engine/Plugins/Animation/ControlRig/Source/ControlRigEditor/Private/EditMode/SControlRigBaseListWidget.h b/Engine/Plugins/Animation/ControlRig/Source/ControlRigEditor/Private/EditMode/SControlRigBaseListWidget.h
--- a/Engine/Plugins/Animation/ControlRig/Source/ControlRigEditor/Private/EditMode/SControlRigBaseListWidget.h
+++ b/Engine/Plugins/Animation/ControlRig/Source/ControlRigEditor/Private/EditMode/SControlRigBaseListWidget.h
@@ -54,7 +54,6 @@ public:
 	~SControlRigBaseListWidget();
 
 	FControlRigEditMode* GetEditMode();
-	UControlRig* GetControlRig(); 
 
 	/*Notification that filter has changed, this includes new folder selected*/
 	void FilterChanged();
diff --git a/Engine/Plugins/Animation/ControlRig/Source/ControlRigEditor/Private/EditMode/SControlRigControlViews.cpp b/Engine/Plugins/Animation/ControlRig/Source/ControlRigEditor/Private/EditMode/SControlRigControlViews.cpp
--- a/Engine/Plugins/Animation/ControlRig/Source/ControlRigEditor/Private/EditMode/SControlRigControlViews.cpp
+++ b/Engine/Plugins/Animation/ControlRig/Source/ControlRigEditor/Private/EditMode/SControlRigControlViews.cpp
@@ -437,14 +437,16 @@ void SControlRigPoseView::Construct(const FArguments& InArgs)
 				]
 				*/
 			]
-			
-
 	];
 
 	if (FControlRigEditMode* EditMode = static_cast<FControlRigEditMode*>(GLevelEditorModeTools().GetActiveMode(FControlRigEditMode::ModeName)))
 	{
 		EditMode->OnControlRigAddedOrRemoved().AddRaw(this, &SControlRigPoseView::HandleControlAdded);
-		HandleControlAdded(GetControlRig(), true);
+		TArray<UControlRig*> ControlRigs = GetControlRigs();
+		for (UControlRig* ControlRig : ControlRigs)
+		{
+			HandleControlAdded(ControlRig, true);
+		}
 	}
 }
 
@@ -453,16 +455,23 @@ SControlRigPoseView::~SControlRigPoseView()
 	if (FControlRigEditMode* EditMode = static_cast<FControlRigEditMode*>(GLevelEditorModeTools().GetActiveMode(FControlRigEditMode::ModeName)))
 	{
 		EditMode->OnControlRigAddedOrRemoved().RemoveAll(this);
-		if (EditMode->GetControlRig(true))
+		TArray<UControlRig*> EditModeRigs = EditMode->GetControlRigsArray(false /*bIsVisible*/);
+		for (UControlRig* ControlRig : EditModeRigs)
 		{
-			EditMode->GetControlRig(true)->ControlSelected().RemoveAll(this);
+			if (ControlRig)
+			{
+				ControlRig->ControlSelected().RemoveAll(this);
+			}
 		}
 	}
 	else
 	{
-		if (CurrentControlRig.IsValid())
+		for (TWeakObjectPtr<UControlRig>& CurrentControlRig: CurrentControlRigs)
 		{
-			(CurrentControlRig.Get())->ControlSelected().RemoveAll(this);
+			if (CurrentControlRig.IsValid())
+			{
+				(CurrentControlRig.Get())->ControlSelected().RemoveAll(this);
+			}
 		}
 	}
 }
@@ -523,9 +532,16 @@ bool SControlRigPoseView::IsMirrorEnabled() const
 
 FReply SControlRigPoseView::OnPastePose()
 {
-	if (GetControlRig() && PoseAsset.IsValid())
+	if (PoseAsset.IsValid())
 	{
-		PoseAsset->PastePose(GetControlRig(), bIsKey,bIsMirror);
+		TArray<UControlRig*> ControlRigs = GetControlRigs();
+		for (UControlRig* ControlRig : ControlRigs)
+		{
+			if (ControlRig)
+			{
+				PoseAsset->PastePose(ControlRig, bIsKey, bIsMirror);
+			}
+		}
 	}
 	return FReply::Handled();
 }
@@ -534,10 +550,13 @@ FReply SControlRigPoseView::OnSelectControls()
 {	
 	if (PoseAsset.IsValid())
 	{
-		UControlRig* ControlRig = GetFirstControlRigInLevelSequence(GetControlRig());
-		if (ControlRig)
+		TArray<UControlRig*> ControlRigs = GetControlRigs();
+		for (UControlRig* ControlRig : ControlRigs)
 		{
-			PoseAsset->SelectControls(ControlRig, bIsMirror);
+			if (ControlRig)
+			{
+				PoseAsset->SelectControls(ControlRig, bIsMirror);
+			}
 		}
 	}
 	return FReply::Handled();
@@ -545,20 +564,24 @@ FReply SControlRigPoseView::OnSelectControls()
 
 void SControlRigPoseView::OnPoseBlendChanged(float ChangedVal)
 {
-	UControlRig* ControlRig = GetControlRig();
-	if (ControlRig && PoseAsset.IsValid())
+	if (PoseAsset.IsValid())
 	{
-
-		PoseBlendValue = ChangedVal;
-		if (!bIsBlending)
+		TArray<UControlRig*> ControlRigs = GetControlRigs();
+		for (UControlRig* ControlRig : ControlRigs)
 		{
-			bIsBlending = true;
-			PoseAsset->GetCurrentPose(ControlRig,TempPose);
-		}
-
-		PoseAsset->BlendWithInitialPoses(TempPose, ControlRig, false, bIsMirror, PoseBlendValue);
-		PoseAsset->BlendWithInitialPoses(TempPose, ControlRig, false, bIsMirror, PoseBlendValue);
+			if (ControlRig)
+			{
+				PoseBlendValue = ChangedVal;
+				if (!bIsBlending)
+				{
+					bIsBlending = true;
+					PoseAsset->GetCurrentPose(ControlRig, TempPose);
+				}
 
+				PoseAsset->BlendWithInitialPoses(TempPose, ControlRig, false, bIsMirror, PoseBlendValue);
+				PoseAsset->BlendWithInitialPoses(TempPose, ControlRig, false, bIsMirror, PoseBlendValue);
+			}
+		}
 	}
 }
 void SControlRigPoseView::OnBeginSliderMovement()
@@ -581,15 +604,25 @@ void SControlRigPoseView::OnEndSliderMovement(float NewValue)
 
 void SControlRigPoseView::OnPoseBlendCommited(float ChangedVal, ETextCommit::Type Type)
 {
-	UControlRig* ControlRig = GetControlRig();
-	if (ControlRig && PoseAsset.IsValid())
+	if (PoseAsset.IsValid())
 	{
-		FScopedTransaction ScopedTransaction(LOCTEXT("PastePoseTransaction", "Paste Pose"));
-		PoseBlendValue = ChangedVal;
-		PoseAsset->BlendWithInitialPoses(TempPose, ControlRig, bIsKey, bIsMirror, PoseBlendValue);
-		PoseAsset->BlendWithInitialPoses(TempPose, ControlRig, bIsKey, bIsMirror, PoseBlendValue);
-		bIsBlending = false;
-		PoseBlendValue = 0.0f;
+		TArray<UControlRig*> ControlRigs = GetControlRigs();
+		if (ControlRigs.Num() > 0)
+		{
+			FScopedTransaction ScopedTransaction(LOCTEXT("PastePoseTransaction", "Paste Pose"));
+			for (UControlRig* ControlRig : ControlRigs)
+			{
+				if (ControlRig)
+				{
+					PoseBlendValue = ChangedVal;
+					PoseAsset->BlendWithInitialPoses(TempPose, ControlRig, bIsKey, bIsMirror, PoseBlendValue);
+					PoseAsset->BlendWithInitialPoses(TempPose, ControlRig, bIsKey, bIsMirror, PoseBlendValue);
+					bIsBlending = false;
+					PoseBlendValue = 0.0f;
+
+				}
+			}
+		}
 	}
 }
 
@@ -625,41 +658,30 @@ TSharedRef<SWidget> SControlRigPoseView::GetThumbnailWidget()
 		];
 }
 
-UControlRig* SControlRigPoseView::GetFirstControlRigInLevelSequence(UControlRig* ControlRig)
+TArray<UControlRig*> SControlRigPoseView::GetControlRigs()
 {
-	if (ControlRig == nullptr)
+	FControlRigEditMode* EditMode = static_cast<FControlRigEditMode*>(GLevelEditorModeTools().GetActiveMode(FControlRigEditMode::ModeName));
+	TArray<UControlRig*> NewControlRigs;
+	if (EditMode)
+	{
+		NewControlRigs =  EditMode->GetControlRigsArray(false /*bIsVisible*/);
+	}
+	for (TWeakObjectPtr<UControlRig> ControlRigPtr : CurrentControlRigs)
 	{
-		ULevelSequence* LevelSequence = ULevelSequenceEditorBlueprintLibrary::GetCurrentLevelSequence();
-		if (LevelSequence)
+		if (ControlRigPtr.IsValid())
 		{
-			TArray<FControlRigSequencerBindingProxy> Proxies = UControlRigSequencerEditorLibrary::GetControlRigs(LevelSequence);
-			if (Proxies.Num() > 0)
+			if (NewControlRigs.Contains(ControlRigPtr.Get()) == false)
 			{
-				//MZ TODO when we have Mutliple Control Rig's active select more than one.
-				ControlRig = Proxies[0].ControlRig;
+				(ControlRigPtr.Get())->ControlSelected().RemoveAll(this);
 			}
 		}
 	}
-	return ControlRig;
-}
-
-UControlRig* SControlRigPoseView::GetControlRig()
-{
-	UControlRig* NewControlRig = nullptr;
-	if (FControlRigEditMode* EditMode = static_cast<FControlRigEditMode*>(GLevelEditorModeTools().GetActiveMode(FControlRigEditMode::ModeName)))
-	{
-		NewControlRig =  EditMode->GetControlRig(true);
-	}
-	if (NewControlRig != CurrentControlRig)
+	if (EditMode)
 	{
-		if (CurrentControlRig.IsValid())
-		{
-			(CurrentControlRig.Get())->ControlSelected().RemoveAll(this);
-
-		}
+		CurrentControlRigs = EditMode->GetControlRigs();
 	}
-	CurrentControlRig = NewControlRig;
-	return NewControlRig;
+	
+	return NewControlRigs;
 }
 
 /* We may want to list the Controls in it (design said no but animators said yes)
@@ -683,7 +705,6 @@ void SControlRigPoseView::HandleControlAdded(UControlRig* ControlRig, bool bIsAd
 		{
 			(ControlRig)->ControlSelected().RemoveAll(this);
 			(ControlRig)->ControlSelected().AddRaw(this, &SControlRigPoseView::HandleControlSelected);
-			CurrentControlRig = ControlRig;
 		}
 		else
 		{
@@ -700,41 +721,47 @@ void SControlRigPoseView::HandleControlSelected(UControlRig* Subject, FRigContro
 
 void SControlRigPoseView::UpdateStatusBlocks()
 {
-	UControlRig* ControlRig = GetControlRig();
 	FText StatusText1;
 	FText StatusText2;
-
-	if (PoseAsset.IsValid() && ControlRig)
+	TArray<UControlRig*> ControlRigs = GetControlRigs();
+	if (PoseAsset.IsValid() && ControlRigs.Num() > 0)
 	{
+
 		FFormatNamedArguments NamedArgs;
 		TArray<FName> ControlNames = PoseAsset->GetControlNames();
 		NamedArgs.Add("Total", ControlNames.Num());
-		TArray<FName> SelectedNames = ControlRig->CurrentControlSelection();
-		NamedArgs.Add("Selected", SelectedNames.Num());
+		int32 TotalSelected = 0;
 		uint32 Matching = 0;
 		uint32 MirrorMatching = 0;
-		
-		for (const FName& ControlName : ControlNames)
+		for (UControlRig* ControlRig : ControlRigs)
 		{
-			for (const FName& SelectedName : SelectedNames)
+
+			TArray<FName> SelectedNames = ControlRig->CurrentControlSelection();
+			TotalSelected += SelectedNames.Num();
+			for (const FName& ControlName : ControlNames)
 			{
-				if (SelectedName == ControlName)
+				for (const FName& SelectedName : SelectedNames)
 				{
-					++Matching;
-					if (SControlRigPoseView::bIsMirror)
+					if (SelectedName == ControlName)
 					{
-						if (PoseAsset->DoesMirrorMatch(ControlRig,ControlName))
+						++Matching;
+						if (SControlRigPoseView::bIsMirror)
 						{
-							++MirrorMatching;
+							if (PoseAsset->DoesMirrorMatch(ControlRig, ControlName))
+							{
+								++MirrorMatching;
+							}
 						}
 					}
 				}
 			}
 		}
+
+		NamedArgs.Add("Selected", TotalSelected);//SelectedNames.Num());
 		NamedArgs.Add("Matching", Matching);
 		NamedArgs.Add("MirrorMatching", MirrorMatching);
 
-		if(SControlRigPoseView::bIsMirror)
+		if (SControlRigPoseView::bIsMirror)
 		{
 			StatusText1 = FText::Format(LOCTEXT("NumberControlsAndMatch", "{Total} Controls Matching {Matching} of {Selected} Selected"), NamedArgs);
 			StatusText2 = FText::Format(LOCTEXT("NumberMirroredMatch", " {MirrorMatching} Mirror String Matches"), NamedArgs);
diff --git a/Engine/Plugins/Animation/ControlRig/Source/ControlRigEditor/Private/EditMode/SControlRigControlViews.h b/Engine/Plugins/Animation/ControlRig/Source/ControlRigEditor/Private/EditMode/SControlRigControlViews.h
--- a/Engine/Plugins/Animation/ControlRig/Source/ControlRigEditor/Private/EditMode/SControlRigControlViews.h
+++ b/Engine/Plugins/Animation/ControlRig/Source/ControlRigEditor/Private/EditMode/SControlRigControlViews.h
@@ -43,8 +43,6 @@ class SControlRigPoseView : public SCompoundWidget, public FNotifyHook
 private:
 	void UpdateStatusBlocks();
 
-	UControlRig* GetControlRig();
-
 	/*
 	* Delegates and Helpers
 	*/
@@ -63,6 +61,7 @@ private:
 	FReply OnCaptureThumbnail();
 	void HandleControlSelected(UControlRig* Subject, FRigControlElement* InControl, bool bSelected);
 	void HandleControlAdded(UControlRig* ControlRig, bool bIsAdded);
+	TArray<UControlRig*> GetControlRigs();
 
 	static bool bIsKey;
 	static bool bIsMirror;
@@ -74,7 +73,7 @@ private:
 	TSharedRef<SWidget> GetThumbnailWidget();
 
 	TWeakObjectPtr<UControlRigPoseAsset> PoseAsset;
-	TWeakObjectPtr<UControlRig> CurrentControlRig;
+	TArray<TWeakObjectPtr<UControlRig>> CurrentControlRigs;
 
 	TSharedPtr<STextBlock> TextStatusBlock1;
 	TSharedPtr<STextBlock> TextStatusBlock2;
@@ -100,6 +99,5 @@ private:
 
 public:
 	static bool IsMirror() { return bIsMirror; }
-	static UControlRig* GetFirstControlRigInLevelSequence(UControlRig* ControlRig);
 };
 
diff --git a/Engine/Plugins/Animation/ControlRig/Source/ControlRigEditor/Private/EditMode/SControlRigDetails.cpp b/Engine/Plugins/Animation/ControlRig/Source/ControlRigEditor/Private/EditMode/SControlRigDetails.cpp
--- a/Engine/Plugins/Animation/ControlRig/Source/ControlRigEditor/Private/EditMode/SControlRigDetails.cpp
+++ b/Engine/Plugins/Animation/ControlRig/Source/ControlRigEditor/Private/EditMode/SControlRigDetails.cpp
@@ -157,7 +157,8 @@ void FControlRigEditModeGenericDetails::CustomizeDetails(class IDetailLayoutBuil
 				{
 					if (FControlRigEditMode* EditMode = static_cast<FControlRigEditMode*>(ModeTools->GetActiveMode(FControlRigEditMode::ModeName)))
 					{
-						if (UObject* NestedProxy = EditMode->ControlProxy->FindProxy(ControlElement->GetName()))
+						UControlRig* ControlRig = Proxy->ControlRig.Get();
+						if (UObject* NestedProxy = EditMode->ControlProxy->FindProxy(ControlRig,ControlElement->GetName()))
 						{
 							FName PropertyName(NAME_None);
 							switch (ControlElement->Settings.ControlType)
diff --git a/Engine/Plugins/Animation/ControlRig/Source/ControlRigEditor/Private/EditMode/SControlRigEditModeTools.cpp b/Engine/Plugins/Animation/ControlRig/Source/ControlRigEditor/Private/EditMode/SControlRigEditModeTools.cpp
--- a/Engine/Plugins/Animation/ControlRig/Source/ControlRigEditor/Private/EditMode/SControlRigEditModeTools.cpp
+++ b/Engine/Plugins/Animation/ControlRig/Source/ControlRigEditor/Private/EditMode/SControlRigEditModeTools.cpp
@@ -41,45 +41,43 @@
 #define LOCTEXT_NAMESPACE "ControlRigEditModeTools"
 
 
-
-void SControlRigEditModeTools::SetControlRig(UControlRig* ControlRig)
+void SControlRigEditModeTools::SetControlRigs(const TArrayView<TWeakObjectPtr<UControlRig>>& InControlRigs)
 {
-	if(ViewportRig.IsValid())
+	for (TWeakObjectPtr<UControlRig>& ControlRig : ControlRigs)
 	{
-		ViewportRig->ControlSelected().RemoveAll(this);
+		if (ControlRig.IsValid())
+		{
+			ControlRig->ControlSelected().RemoveAll(this);
+		}
 	}
-
-	SequencerRig = ControlRig;
-	ViewportRig = ControlRig;
-	if (SequencerRig.IsValid())
+	ControlRigs = InControlRigs;
+	for (TWeakObjectPtr<UControlRig>& InControlRig : InControlRigs)
 	{
-		if (UControlRig* InteractionRig = SequencerRig->GetInteractionRig())
+		if (InControlRig.IsValid())
 		{
-			ViewportRig = InteractionRig;
+			InControlRig->ControlSelected().AddRaw(this, &SControlRigEditModeTools::OnRigElementSelected);
 		}
 	}
 
+	//mz todo handle multiple rigs
+	UControlRig* Rig = ControlRigs.Num() > 0 ? ControlRigs[0].Get() : nullptr;
 	TArray<TWeakObjectPtr<>> Objects;
-	Objects.Add(SequencerRig);
+	Objects.Add(Rig);
 	RigOptionsDetailsView->SetObjects(Objects);
+
 #if USE_LOCAL_DETAILS
 	HierarchyTreeView->RefreshTreeView(true);
 #endif
-	if (ViewportRig.IsValid())
-	{
-		ViewportRig->ControlSelected().AddRaw(this, &SControlRigEditModeTools::OnRigElementSelected);
-	}
+
 }
 
 const URigHierarchy* SControlRigEditModeTools::GetHierarchy() const
 {
-	if(ViewportRig.IsValid())
+	//mz todo handle multiple rigs
+	UControlRig* Rig = ControlRigs.Num() > 0 ? ControlRigs[0].Get() : nullptr;
+	if (Rig)
 	{
-		return ViewportRig->GetHierarchy();
-	}
-	if(SequencerRig.IsValid())
-	{
-		return ViewportRig->GetHierarchy();
+		Rig->GetHierarchy();
 	}
 	return nullptr;
 }
@@ -677,20 +675,24 @@ void SControlRigEditModeTools::OnRigElementSelected(UControlRig* Subject, FRigCo
 		}
 	}
 #endif
-	if (UControlRig* ControlRig = SequencerRig.Get())
+	//mz todo handle multiple rigs
+	UControlRig* Rig = ControlRigs.Num() > 0 ? ControlRigs[0].Get() : nullptr;
+	if (Rig)
 	{
 		// get the selected controls
-		TArray<FRigElementKey> SelectedControls = ControlRig->GetHierarchy()->GetSelectedKeys(ERigElementType::Control);
-		SpacePickerWidget->SetControls(ControlRig->GetHierarchy(), SelectedControls);
+		TArray<FRigElementKey> SelectedControls = Rig->GetHierarchy()->GetSelectedKeys(ERigElementType::Control);
+		SpacePickerWidget->SetControls(Rig->GetHierarchy(), SelectedControls);
 	}
 }
 
 
 const FRigControlElementCustomization* SControlRigEditModeTools::HandleGetControlElementCustomization(URigHierarchy* InHierarchy, const FRigElementKey& InControlKey)
 {
-	if (UControlRig* ControlRig = SequencerRig.Get())
+	//mz todo handle multiple rigs
+	UControlRig* Rig = ControlRigs.Num() > 0 ? ControlRigs[0].Get() : nullptr;
+	if (Rig)
 	{
-		return ControlRig->GetControlCustomization(InControlKey);
+		return Rig->GetControlCustomization(InControlKey);
 	}
 	return nullptr;
 }
@@ -701,38 +703,41 @@ void SControlRigEditModeTools::HandleActiveSpaceChanged(URigHierarchy* InHierarc
 
 	if (WeakSequencer.IsValid())
 	{
-		if (UControlRig* ControlRig = SequencerRig.Get())
+		for (TWeakObjectPtr<UControlRig>& ControlRig : ControlRigs)
 		{
-			FString FailureReason;
-			URigHierarchy::TElementDependencyMap DependencyMap = InHierarchy->GetDependenciesForVM(ControlRig->GetVM());
-			if(!InHierarchy->CanSwitchToParent(InControlKey, InSpaceKey, DependencyMap, &FailureReason))
+			if (ControlRig.IsValid() && ControlRig->GetHierarchy() == InHierarchy)
 			{
-				// notification
-				FNotificationInfo Info(FText::FromString(FailureReason));
-				Info.bFireAndForget = true;
-				Info.FadeOutDuration = 2.0f;
-				Info.ExpireDuration = 8.0f;
-
-				const TSharedPtr<SNotificationItem> NotificationPtr = FSlateNotificationManager::Get().AddNotification(Info);
-				NotificationPtr->SetCompletionState(SNotificationItem::CS_Fail);
-				return;
-			}
+				FString FailureReason;
+				URigHierarchy::TElementDependencyMap DependencyMap = InHierarchy->GetDependenciesForVM(ControlRig->GetVM());
+				if(!InHierarchy->CanSwitchToParent(InControlKey, InSpaceKey, DependencyMap, &FailureReason))
+				{
+					// notification
+					FNotificationInfo Info(FText::FromString(FailureReason));
+					Info.bFireAndForget = true;
+					Info.FadeOutDuration = 2.0f;
+					Info.ExpireDuration = 8.0f;
+
+					const TSharedPtr<SNotificationItem> NotificationPtr = FSlateNotificationManager::Get().AddNotification(Info);
+					NotificationPtr->SetCompletionState(SNotificationItem::CS_Fail);
+					return;
+				}
 			
-			if (const FRigControlElement* ControlElement = InHierarchy->Find<FRigControlElement>(InControlKey))
-			{
-				ISequencer* Sequencer = WeakSequencer.Pin().Get();
-				if (Sequencer)
+				if (const FRigControlElement* ControlElement = InHierarchy->Find<FRigControlElement>(InControlKey))
 				{
-					FScopedTransaction Transaction(LOCTEXT("KeyControlRigSpace", "Key Control Rig Space"));
-					ControlRig->Modify();
-
-					FSpaceChannelAndSection SpaceChannelAndSection = FControlRigSpaceChannelHelpers::FindSpaceChannelAndSectionForControl(ControlRig, InControlKey.Name, Sequencer, true /*bCreateIfNeeded*/);
-					if (SpaceChannelAndSection.SpaceChannel)
+					ISequencer* Sequencer = WeakSequencer.Pin().Get();
+					if (Sequencer)
 					{
-						const FFrameRate TickResolution = Sequencer->GetFocusedTickResolution();
-						const FFrameTime FrameTime = Sequencer->GetLocalTime().ConvertTo(TickResolution);
-						FFrameNumber CurrentTime = FrameTime.GetFrame();
-						FControlRigSpaceChannelHelpers::SequencerKeyControlRigSpaceChannel(ControlRig, Sequencer, SpaceChannelAndSection.SpaceChannel, SpaceChannelAndSection.SectionToKey, CurrentTime, InHierarchy, InControlKey, InSpaceKey);
+						FScopedTransaction Transaction(LOCTEXT("KeyControlRigSpace", "Key Control Rig Space"));
+						ControlRig->Modify();
+
+						FSpaceChannelAndSection SpaceChannelAndSection = FControlRigSpaceChannelHelpers::FindSpaceChannelAndSectionForControl(ControlRig.Get(), InControlKey.Name, Sequencer, true /*bCreateIfNeeded*/);
+						if (SpaceChannelAndSection.SpaceChannel)
+						{
+							const FFrameRate TickResolution = Sequencer->GetFocusedTickResolution();
+							const FFrameTime FrameTime = Sequencer->GetLocalTime().ConvertTo(TickResolution);
+							FFrameNumber CurrentTime = FrameTime.GetFrame();
+							FControlRigSpaceChannelHelpers::SequencerKeyControlRigSpaceChannel(ControlRig.Get(), Sequencer, SpaceChannelAndSection.SpaceChannel, SpaceChannelAndSection.SectionToKey, CurrentTime, InHierarchy, InControlKey, InSpaceKey);
+						}
 					}
 				}
 			}
@@ -743,36 +748,40 @@ void SControlRigEditModeTools::HandleActiveSpaceChanged(URigHierarchy* InHierarc
 void SControlRigEditModeTools::HandleSpaceListChanged(URigHierarchy* InHierarchy, const FRigElementKey& InControlKey,
 	const TArray<FRigElementKey>& InSpaceList)
 {
-	if (UControlRig* ControlRig = SequencerRig.Get())
+	for (TWeakObjectPtr<UControlRig>& ControlRig : ControlRigs)
 	{
-		if (const FRigControlElement* ControlElement = InHierarchy->Find<FRigControlElement>(InControlKey))
+		if (ControlRig.IsValid() && ControlRig->GetHierarchy() == InHierarchy)
 		{
-			FRigControlElementCustomization ControlCustomization = *ControlRig->GetControlCustomization(InControlKey);
-			ControlCustomization.AvailableSpaces = InSpaceList;
-			ControlCustomization.RemovedSpaces.Reset();
 
-			// remember  the elements which are in the asset's available list but removed by the user
-			for (const FRigElementKey& AvailableSpace : ControlElement->Settings.Customization.AvailableSpaces)
+			if (const FRigControlElement* ControlElement = InHierarchy->Find<FRigControlElement>(InControlKey))
 			{
-				if (!ControlCustomization.AvailableSpaces.Contains(AvailableSpace))
+				FRigControlElementCustomization ControlCustomization = *ControlRig->GetControlCustomization(InControlKey);
+				ControlCustomization.AvailableSpaces = InSpaceList;
+				ControlCustomization.RemovedSpaces.Reset();
+
+				// remember  the elements which are in the asset's available list but removed by the user
+				for (const FRigElementKey& AvailableSpace : ControlElement->Settings.Customization.AvailableSpaces)
 				{
-					ControlCustomization.RemovedSpaces.Add(AvailableSpace);
+					if (!ControlCustomization.AvailableSpaces.Contains(AvailableSpace))
+					{
+						ControlCustomization.RemovedSpaces.Add(AvailableSpace);
+					}
 				}
-			}
 
-			ControlRig->SetControlCustomization(InControlKey, ControlCustomization);
+				ControlRig->SetControlCustomization(InControlKey, ControlCustomization);
 
-			if (FControlRigEditMode* EditMode = static_cast<FControlRigEditMode*>(ModeTools->GetActiveMode(FControlRigEditMode::ModeName)))
-			{
-				const TGuardValue<bool> SuspendGuard(EditMode->bSuspendHierarchyNotifs, true);
-				InHierarchy->Notify(ERigHierarchyNotification::ControlSettingChanged, ControlElement);
-			}
-			else
-			{
-				InHierarchy->Notify(ERigHierarchyNotification::ControlSettingChanged, ControlElement);
-			}
+				if (FControlRigEditMode* EditMode = static_cast<FControlRigEditMode*>(ModeTools->GetActiveMode(FControlRigEditMode::ModeName)))
+				{
+					const TGuardValue<bool> SuspendGuard(EditMode->bSuspendHierarchyNotifs, true);
+					InHierarchy->Notify(ERigHierarchyNotification::ControlSettingChanged, ControlElement);
+				}
+				else
+				{
+					InHierarchy->Notify(ERigHierarchyNotification::ControlSettingChanged, ControlElement);
+				}
 
-			SpacePickerWidget->RefreshContents();
+				SpacePickerWidget->RefreshContents();
+			}
 		}
 	}
 }
@@ -792,7 +801,18 @@ FReply SControlRigEditModeTools::OnBakeControlsToNewSpaceButtonClicked()
 	{
 		return FReply::Unhandled();
 	}
-	if (!SequencerRig.Get())
+
+	bool bNoValidControlRig = true;
+	for (TWeakObjectPtr<UControlRig>& ControlRig : ControlRigs)
+	{
+		if (ControlRig.IsValid())
+		{
+			bNoValidControlRig = false;
+			break;
+		}
+	}
+
+	if (bNoValidControlRig)
 	{
 		return FReply::Unhandled();
 	}
@@ -801,88 +821,98 @@ FReply SControlRigEditModeTools::OnBakeControlsToNewSpaceButtonClicked()
 	{
 		return FReply::Unhandled();
 	}
-	UControlRig* ControlRig = SequencerRig.Get();
-
-	FRigSpacePickerBakeSettings Settings;
-	//Find default target space, just use first control and find space at current sequencer time
-	//Then Find range
-
-	// FindSpaceChannelAndSectionForControl() will trigger RecreateCurveEditor(), which will deselect the controls
-	// but in theory the selection will be recovered in the next tick, so here we just cache the selected controls
-	// and use it throughout this function. If this deselection is causing other problems, this part could use a revisit.
-	TArray<FRigElementKey> ControlKeys = SpacePickerWidget->GetControls();
-
-	FSpaceChannelAndSection SpaceChannelAndSection = FControlRigSpaceChannelHelpers::FindSpaceChannelAndSectionForControl(ControlRig, ControlKeys[0].Name, Sequencer, true /*bCreateIfNeeded*/);
-	if (SpaceChannelAndSection.SpaceChannel != nullptr)
+	for (TWeakObjectPtr<UControlRig>& ControlRig : ControlRigs)
 	{
-		const FFrameRate TickResolution = Sequencer->GetFocusedTickResolution();
-		const FFrameTime FrameTime = Sequencer->GetLocalTime().ConvertTo(TickResolution);
-		FFrameNumber CurrentTime = FrameTime.GetFrame();
-		FMovieSceneControlRigSpaceBaseKey Value;
-		using namespace UE::MovieScene;
-		URigHierarchy* RigHierarchy = SpacePickerWidget->GetHierarchy();
-		Settings.TargetSpace = URigHierarchy::GetDefaultParentKey();
-		
-		TRange<FFrameNumber> Range = Sequencer->GetFocusedMovieSceneSequence()->GetMovieScene()->GetPlaybackRange();
-		TArray<FFrameNumber> Keys;
-		TArray < FKeyHandle> KeyHandles;
-
-		Settings.StartFrame = Range.GetLowerBoundValue();
-		Settings.EndFrame = Range.GetUpperBoundValue();
-		if (Keys.Num() > 0)
+		if (ControlRig.IsValid())
 		{
-			int32 Index = Algo::LowerBound(Keys, CurrentTime);
-			if (Index >= 0 && Index < (Keys.Num() - 1))
-			{
-				Settings.StartFrame = Keys[Index];
-				Settings.EndFrame = Keys[Index + 1];
 
-			}
-		}
+			FRigSpacePickerBakeSettings Settings;
+			//Find default target space, just use first control and find space at current sequencer time
+			//Then Find range
+
+			// FindSpaceChannelAndSectionForControl() will trigger RecreateCurveEditor(), which will deselect the controls
+			// but in theory the selection will be recovered in the next tick, so here we just cache the selected controls
+			// and use it throughout this function. If this deselection is causing other problems, this part could use a revisit.
+			TArray<FRigElementKey> ControlKeys = SpacePickerWidget->GetControls();
 
-		TSharedRef<SRigSpacePickerBakeWidget> BakeWidget =
-			SNew(SRigSpacePickerBakeWidget)
-			.Settings(Settings)
-			.Hierarchy(SpacePickerWidget->GetHierarchy())
-			.Controls(ControlKeys) // use the cached controls here since the selection is not recovered until next tick.
-			.Sequencer(Sequencer)
-			.GetControlCustomization(this, &SControlRigEditModeTools::HandleGetControlElementCustomization)
-			.OnBake_Lambda([Sequencer, ControlRig,TickResolution](URigHierarchy* InHierarchy, TArray<FRigElementKey> InControls, FRigSpacePickerBakeSettings InSettings)
+			FSpaceChannelAndSection SpaceChannelAndSection = FControlRigSpaceChannelHelpers::FindSpaceChannelAndSectionForControl(ControlRig.Get(), ControlKeys[0].Name, Sequencer, true /*bCreateIfNeeded*/);
+			if (SpaceChannelAndSection.SpaceChannel != nullptr)
+			{
+				const FFrameRate TickResolution = Sequencer->GetFocusedTickResolution();
+				const FFrameTime FrameTime = Sequencer->GetLocalTime().ConvertTo(TickResolution);
+				FFrameNumber CurrentTime = FrameTime.GetFrame();
+				FMovieSceneControlRigSpaceBaseKey Value;
+				using namespace UE::MovieScene;
+				URigHierarchy* RigHierarchy = SpacePickerWidget->GetHierarchy();
+				Settings.TargetSpace = URigHierarchy::GetDefaultParentKey();
+
+				TRange<FFrameNumber> Range = Sequencer->GetFocusedMovieSceneSequence()->GetMovieScene()->GetPlaybackRange();
+				TArray<FFrameNumber> Keys;
+				TArray < FKeyHandle> KeyHandles;
+
+				Settings.StartFrame = Range.GetLowerBoundValue();
+				Settings.EndFrame = Range.GetUpperBoundValue();
+				if (Keys.Num() > 0)
 				{
-					TArray<FFrameNumber> Frames;
-					
-					const FFrameRate& FrameRate = Sequencer->GetFocusedDisplayRate();
-					FFrameNumber FrameRateInFrameNumber = TickResolution.AsFrameNumber(FrameRate.AsInterval());
-					for (FFrameNumber& Frame = InSettings.StartFrame; Frame <= InSettings.EndFrame; Frame += FrameRateInFrameNumber)
+					int32 Index = Algo::LowerBound(Keys, CurrentTime);
+					if (Index >= 0 && Index < (Keys.Num() - 1))
 					{
-						Frames.Add(Frame);
+						Settings.StartFrame = Keys[Index];
+						Settings.EndFrame = Keys[Index + 1];
+
 					}
-					FScopedTransaction Transaction(LOCTEXT("BakeControlToSpace", "Bake Control In Space"));
-					for (const FRigElementKey& ControlKey : InControls)
-					{
-						//when baking we will now create a channel if one doesn't exist, was causing confusion
-						FSpaceChannelAndSection SpaceChannelAndSection = FControlRigSpaceChannelHelpers::FindSpaceChannelAndSectionForControl(ControlRig, ControlKey.Name, Sequencer, true /*bCreateIfNeeded*/);
-						if (SpaceChannelAndSection.SpaceChannel)
+				}
+
+				TSharedRef<SRigSpacePickerBakeWidget> BakeWidget =
+					SNew(SRigSpacePickerBakeWidget)
+					.Settings(Settings)
+					.Hierarchy(SpacePickerWidget->GetHierarchy())
+					.Controls(ControlKeys) // use the cached controls here since the selection is not recovered until next tick.
+					.Sequencer(Sequencer)
+					.GetControlCustomization(this, &SControlRigEditModeTools::HandleGetControlElementCustomization)
+					.OnBake_Lambda([Sequencer, ControlRig, TickResolution](URigHierarchy* InHierarchy, TArray<FRigElementKey> InControls, FRigSpacePickerBakeSettings InSettings)
 						{
-							FControlRigSpaceChannelHelpers::SequencerBakeControlInSpace(ControlRig, Sequencer, SpaceChannelAndSection.SpaceChannel, SpaceChannelAndSection.SectionToKey,
-								Frames, InHierarchy, ControlKey, InSettings);
-						}
-					}
-					return FReply::Handled();
-				});
+							TArray<FFrameNumber> Frames;
+
+							const FFrameRate& FrameRate = Sequencer->GetFocusedDisplayRate();
+							FFrameNumber FrameRateInFrameNumber = TickResolution.AsFrameNumber(FrameRate.AsInterval());
+							for (FFrameNumber& Frame = InSettings.StartFrame; Frame <= InSettings.EndFrame; Frame += FrameRateInFrameNumber)
+							{
+								Frames.Add(Frame);
+							}
+							FScopedTransaction Transaction(LOCTEXT("BakeControlToSpace", "Bake Control In Space"));
+							for (const FRigElementKey& ControlKey : InControls)
+							{
+								//when baking we will now create a channel if one doesn't exist, was causing confusion
+								FSpaceChannelAndSection SpaceChannelAndSection = FControlRigSpaceChannelHelpers::FindSpaceChannelAndSectionForControl(ControlRig.Get(), ControlKey.Name, Sequencer, true /*bCreateIfNeeded*/);
+								if (SpaceChannelAndSection.SpaceChannel)
+								{
+									FControlRigSpaceChannelHelpers::SequencerBakeControlInSpace(ControlRig.Get(), Sequencer, SpaceChannelAndSection.SpaceChannel, SpaceChannelAndSection.SectionToKey,
+										Frames, InHierarchy, ControlKey, InSettings);
+								}
+							}
+							return FReply::Handled();
+						});
+
+				return BakeWidget->OpenDialog(true);
+			}
+			break; //mz todo need baketo handle more than one
+		}
 
-		return BakeWidget->OpenDialog(true);
 	}
 	return FReply::Unhandled();
 }
 
 EVisibility SControlRigEditModeTools::GetRigOptionExpanderVisibility() const
 {
-	if (UControlRig* ControlRig = SequencerRig.Get())
+	for (const TWeakObjectPtr<UControlRig>& ControlRig : ControlRigs)
 	{
-		if (Cast<UFKControlRig>(ControlRig))
+		if (ControlRig.IsValid())
 		{
-			return EVisibility::Visible;
+			if (Cast<UFKControlRig>(ControlRig))
+			{
+				return EVisibility::Visible;
+			}
 		}
 	}
 	return EVisibility::Hidden;
@@ -890,7 +920,8 @@ EVisibility SControlRigEditModeTools::GetRigOptionExpanderVisibility() const
 
 void SControlRigEditModeTools::OnRigOptionFinishedChange(const FPropertyChangedEvent& PropertyChangedEvent)
 {
-	SetControlRig(SequencerRig.Get());
+	TArray<TWeakObjectPtr<UControlRig>> ControlRigsCopy = ControlRigs;
+	SetControlRigs(ControlRigsCopy);
 
 	if (FControlRigEditMode* EditMode = static_cast<FControlRigEditMode*>(ModeTools->GetActiveMode(FControlRigEditMode::ModeName)))
 	{
diff --git a/Engine/Plugins/Animation/ControlRig/Source/ControlRigEditor/Private/EditMode/SControlRigEditModeTools.h b/Engine/Plugins/Animation/ControlRig/Source/ControlRigEditor/Private/EditMode/SControlRigEditModeTools.h
--- a/Engine/Plugins/Animation/ControlRig/Source/ControlRigEditor/Private/EditMode/SControlRigEditModeTools.h
+++ b/Engine/Plugins/Animation/ControlRig/Source/ControlRigEditor/Private/EditMode/SControlRigEditModeTools.h
@@ -49,7 +49,7 @@ public:
 	void SetSequencer(TWeakPtr<ISequencer> InSequencer);
 
 	/** Set The Control Rig we are using*/
-	void SetControlRig(UControlRig* ControlRig);
+	void SetControlRigs(const TArrayView<TWeakObjectPtr<UControlRig>>& InControlRigs);
 
 	/** Returns the hierarchy currently being used */
 	const URigHierarchy* GetHierarchy() const;
@@ -88,9 +88,8 @@ private:
 	TSharedPtr<SExpandableArea> PickerExpander;
 	TSharedPtr<SRigSpacePickerWidget> SpacePickerWidget;
 
-	/** Storage for both sequencer and viewport rigs */
-	TWeakObjectPtr<UControlRig> SequencerRig;
-	TWeakObjectPtr<UControlRig> ViewportRig;
+	/** Storage for control rigs */
+	TArray<TWeakObjectPtr<UControlRig>> ControlRigs;
 
 	/** Display or edit set up for property */
 	bool ShouldShowPropertyOnDetailCustomization(const struct FPropertyAndParent& InPropertyAndParent) const;
diff --git a/Engine/Plugins/Animation/ControlRig/Source/ControlRigEditor/Private/EditMode/SControlRigOutliner.cpp b/Engine/Plugins/Animation/ControlRig/Source/ControlRigEditor/Private/EditMode/SControlRigOutliner.cpp
--- a/Engine/Plugins/Animation/ControlRig/Source/ControlRigEditor/Private/EditMode/SControlRigOutliner.cpp
+++ b/Engine/Plugins/Animation/ControlRig/Source/ControlRigEditor/Private/EditMode/SControlRigOutliner.cpp
@@ -31,11 +31,14 @@
 #include "TimerManager.h"
 #include "SRigHierarchyTreeView.h"
 #include "Rigs/RigHierarchyController.h"
+#include "ControlRigObjectBinding.h"
 
 #define LOCTEXT_NAMESPACE "ControlRigOutliner"
 
-void SControlRigOutliner::Construct(const FArguments& InArgs, FControlRigEditMode& InEditMode)
+
+void SControlRigOutlinerItem::Construct(const FArguments& InArgs)
 {
+
 	bIsChangingRigHierarchy = false;
 
 	DisplaySettings.bShowBones = false;
@@ -47,10 +50,28 @@ void SControlRigOutliner::Construct(const FArguments& InArgs, FControlRigEditMod
 	DisplaySettings.bFlattenHierarchyOnFilter = true;
 
 	FRigTreeDelegates RigTreeDelegates;
-	RigTreeDelegates.OnGetHierarchy = FOnGetRigTreeHierarchy::CreateSP(this, &SControlRigOutliner::GetHierarchy);
-	RigTreeDelegates.OnGetDisplaySettings = FOnGetRigTreeDisplaySettings::CreateSP(this, &SControlRigOutliner::GetDisplaySettings);
-	RigTreeDelegates.OnSelectionChanged = FOnRigTreeSelectionChanged::CreateSP(this, &SControlRigOutliner::HandleSelectionChanged);
+	RigTreeDelegates.OnGetHierarchy = FOnGetRigTreeHierarchy::CreateSP(this, &SControlRigOutlinerItem::GetHierarchy);
+	RigTreeDelegates.OnGetDisplaySettings = FOnGetRigTreeDisplaySettings::CreateSP(this, &SControlRigOutlinerItem::GetDisplaySettings);
+	RigTreeDelegates.OnSelectionChanged = FOnRigTreeSelectionChanged::CreateSP(this, &SControlRigOutlinerItem::HandleSelectionChanged);
 
+	FText AreaTitle;
+	if (InArgs._ControlRig)
+	{
+		FString ControlRigName = InArgs._ControlRig->GetName();
+		FString BoundObjectName;
+		if (TSharedPtr<IControlRigObjectBinding> ObjectBinding = InArgs._ControlRig->GetObjectBinding())
+		{
+			if (ObjectBinding->GetBoundObject())
+			{
+				AActor* Actor = ObjectBinding->GetBoundObject()->GetTypedOuter<AActor>();
+				if (Actor)
+				{
+					BoundObjectName = Actor->GetActorLabel();
+				}
+			}
+		}
+		AreaTitle = FText::Format(LOCTEXT("ControlTitle", "{0}  ({1})"), FText::AsCultureInvariant(ControlRigName), FText::AsCultureInvariant((BoundObjectName)));
+	}
 	ChildSlot
 		[
 			SNew(SScrollBox)
@@ -62,43 +83,118 @@ void SControlRigOutliner::Construct(const FArguments& InArgs, FControlRigEditMod
 				[
 					SAssignNew(PickerExpander, SExpandableArea)
 					.InitiallyCollapsed(false)
-					.AreaTitle(LOCTEXT("Picker_Header", "Controls"))
-					.AreaTitleFont(FEditorStyle::GetFontStyle("DetailsView.CategoryFontStyle"))
+					//.AreaTitle(AreaTitle)
+					//.AreaTitleFont(FEditorStyle::GetFontStyle("DetailsView.CategoryFontStyle"))
 					.BorderBackgroundColor(FLinearColor(.6f, .6f, .6f))
 					.BodyContent()
 					[
 						SAssignNew(HierarchyTreeView, SSearchableRigHierarchyTreeView)
 						.RigTreeDelegates(RigTreeDelegates)
 					]
+					.HeaderContent()
+						[
+							SNew(SHorizontalBox)
+							+ SHorizontalBox::Slot()
+							.AutoWidth()
+							.VAlign(VAlign_Center)
+							[
+								SNew(SButton)
+								.ContentPadding(2)
+								.ButtonStyle(FEditorStyle::Get(), "NoBorder")
+								.OnClicked(this, &SControlRigOutlinerItem::OnToggleVisibility)
+								.ToolTipText(LOCTEXT("ControlRigShapesVisibility", "Control Rig Shapes Visibility"))
+								.IsEnabled(this, &SControlRigOutlinerItem::VisibilityToggleEnabled)
+								.HAlign(HAlign_Center)
+								.VAlign(VAlign_Center)
+								.Content()
+								[
+									SNew(SImage)
+									.Image(this, &SControlRigOutlinerItem::GetVisibilityBrushForElement)
+								]
+							]
+							+ SHorizontalBox::Slot()
+							.AutoWidth()
+							.VAlign(VAlign_Center)
+							[
+								SNew(STextBlock)
+								.Text(AreaTitle)
+								.Justification(ETextJustify::Left)
+								//.TextStyle(FSourceFilterStyle::Get(), "SourceFilter.TextStyle")
+							]
+						]
 				]
 		
 			]
 		];
+	NewControlRigSet(InArgs._ControlRig);
+}
 
-	SetEditMode(InEditMode);
-	HierarchyTreeView->GetTreeView()->RefreshTreeView(true);
+FReply SControlRigOutlinerItem::OnToggleVisibility()
+{
+	if (UControlRig *ControlRig  = CurrentControlRig.Get())
+	{
+		ControlRig->ToggleControlsVisible();
+	}
+	return FReply::Handled();
 }
 
-SControlRigOutliner::~SControlRigOutliner()
+bool SControlRigOutlinerItem::VisibilityToggleEnabled() const
 {
-	//base class handles control rig related cleanup
+	return CurrentControlRig.IsValid();
+}
+
+const FSlateBrush* SControlRigOutlinerItem::GetVisibilityBrushForElement() const
+{
+	if (UControlRig* ControlRig = CurrentControlRig.Get())
+	{
+		if (ControlRig->GetControlsVisible())
+		{
+			return IsHovered() ? FEditorStyle::GetBrush("Level.VisibleHighlightIcon16x") :
+				FEditorStyle::GetBrush("Level.VisibleIcon16x");
+		}
+	}
+	return IsHovered() ? FEditorStyle::GetBrush("Level.NotVisibleHighlightIcon16x") :
+		FEditorStyle::GetBrush("Level.NotVisibleIcon16x");
 }
 
+void SControlRigOutlinerItem::OnObjectsReplaced(const TMap<UObject*, UObject*>& OldToNewInstanceMap)
+{
+	if (CurrentControlRig.IsValid())
+	{
+		UObject* OldObject = CurrentControlRig.Get();
+		UObject* NewObject = OldToNewInstanceMap.FindRef(OldObject);
+		if (NewObject)
+		{
+			if (UControlRig* ControlRig = Cast<UControlRig>(NewObject))
+			{
+				NewControlRigSet(ControlRig);
+			}
+		}
+	}
+}
 
-void SControlRigOutliner::HandleControlAdded(UControlRig* ControlRig, bool bIsAdded)
+void SControlRigOutlinerItem::NewControlRigSet(UControlRig* ControlRig)
 {
-	FControlRigBaseDockableView::HandleControlAdded(ControlRig, bIsAdded);
+	if (CurrentControlRig.IsValid())
+	{
+		(CurrentControlRig.Get())->ControlSelected().RemoveAll(this);
+	}
+	CurrentControlRig = ControlRig;
+	if (ControlRig)
+	{
+		ControlRig->ControlSelected().RemoveAll(this);
+		ControlRig->ControlSelected().AddRaw(this, &SControlRigOutlinerItem::HandleControlSelected);
+	}
 	HierarchyTreeView->GetTreeView()->RefreshTreeView(true);
 }
 
-void SControlRigOutliner::NewControlRigSet(UControlRig* ControlRig)
+SControlRigOutlinerItem::SControlRigOutlinerItem()
 {
-	FControlRigBaseDockableView::NewControlRigSet(ControlRig);
-	HierarchyTreeView->GetTreeView()->RefreshTreeView(true);
+	FCoreUObjectDelegates::OnObjectsReplaced.AddRaw(this, &SControlRigOutlinerItem::OnObjectsReplaced);
 }
-void SControlRigOutliner::HandleControlSelected(UControlRig* Subject, FRigControlElement* ControlElement, bool bSelected)
+
+void SControlRigOutlinerItem::HandleControlSelected(UControlRig* Subject, FRigControlElement* ControlElement, bool bSelected)
 {
-	FControlRigBaseDockableView::HandleControlSelected(Subject, ControlElement, bSelected);
 	const FRigElementKey Key = ControlElement->GetKey();
 	for (int32 RootIndex = 0; RootIndex < HierarchyTreeView->GetTreeView()->GetRootElements().Num(); ++RootIndex)
 	{
@@ -121,7 +217,7 @@ void SControlRigOutliner::HandleControlSelected(UControlRig* Subject, FRigContro
 	}
 }
 
-const URigHierarchy* SControlRigOutliner::GetHierarchy() const
+const URigHierarchy* SControlRigOutlinerItem::GetHierarchy() const
 {
 	if (CurrentControlRig.IsValid())
 	{
@@ -130,7 +226,7 @@ const URigHierarchy* SControlRigOutliner::GetHierarchy() const
 	return nullptr;
 }
 
-void SControlRigOutliner::HandleSelectionChanged(TSharedPtr<FRigTreeElement> Selection, ESelectInfo::Type SelectInfo)
+void SControlRigOutlinerItem::HandleSelectionChanged(TSharedPtr<FRigTreeElement> Selection, ESelectInfo::Type SelectInfo)
 {
 	if (bIsChangingRigHierarchy)
 	{
@@ -153,5 +249,75 @@ void SControlRigOutliner::HandleSelectionChanged(TSharedPtr<FRigTreeElement> Sel
 	}
 }
 
+SControlRigOutlinerItem::~SControlRigOutlinerItem()
+{
+	if (CurrentControlRig.IsValid())
+	{
+		(CurrentControlRig.Get())->ControlSelected().RemoveAll(this);
+	}
+	CurrentControlRig = nullptr;
+	FCoreUObjectDelegates::OnObjectsReplaced.RemoveAll(this);
+}
+
+void SControlRigOutliner::SetEditMode(FControlRigEditMode& InEditMode)
+{
+	ModeTools = InEditMode.GetModeManager();
+	if (FControlRigEditMode* EditMode = static_cast<FControlRigEditMode*>(ModeTools->GetActiveMode(FControlRigEditMode::ModeName)))
+	{
+		EditMode->OnControlRigAddedOrRemoved().AddRaw(this, &SControlRigOutliner::HandleControlAdded);
+	}
+}
+void SControlRigOutliner::Construct(const FArguments& InArgs, FControlRigEditMode& InEditMode)
+{
+	ChildSlot
+		[
+			SNew(SScrollBox)
+			+ SScrollBox::Slot()
+			[
+				SAssignNew(MainBoxPtr, SVerticalBox)
+			]
+		];
+
+	SetEditMode(InEditMode);
+	Rebuild();
+}
+
+SControlRigOutliner::~SControlRigOutliner()
+{
+	if (FControlRigEditMode* EditMode = static_cast<FControlRigEditMode*>(ModeTools->GetActiveMode(FControlRigEditMode::ModeName)))
+	{
+		EditMode->OnControlRigAddedOrRemoved().RemoveAll(this);
+	}
+	//base class handles control rig related cleanup
+}
+
+
+void SControlRigOutliner::HandleControlAdded(UControlRig* ControlRig, bool bIsAdded)
+{
+	Rebuild();
+
+}
+void SControlRigOutliner::Rebuild()
+{
+	MainBoxPtr->ClearChildren();
+	if (FControlRigEditMode* EditMode = static_cast<FControlRigEditMode*>(ModeTools->GetActiveMode(FControlRigEditMode::ModeName)))
+	{
+		TArray<UControlRig*> ControlRigs = EditMode->GetControlRigsArray(false /*bIsVisible*/);
+		for (UControlRig* ControlRig : ControlRigs)
+		{
+			if (ControlRig)
+			{
+				MainBoxPtr->AddSlot()
+				.AutoHeight()
+				[
+					SNew(SControlRigOutlinerItem)
+					.ControlRig(ControlRig)
+				];
+			}
+		}
+	}
+
+}
+
 
 #undef LOCTEXT_NAMESPACE
diff --git a/Engine/Plugins/Animation/ControlRig/Source/ControlRigEditor/Private/EditMode/SControlRigOutliner.h b/Engine/Plugins/Animation/ControlRig/Source/ControlRigEditor/Private/EditMode/SControlRigOutliner.h
--- a/Engine/Plugins/Animation/ControlRig/Source/ControlRigEditor/Private/EditMode/SControlRigOutliner.h
+++ b/Engine/Plugins/Animation/ControlRig/Source/ControlRigEditor/Private/EditMode/SControlRigOutliner.h
@@ -11,35 +11,63 @@
 #include "ControlRig.h"
 #include "Rigs/RigHierarchy.h"
 #include "SRigHierarchyTreeView.h"
+#include "Widgets/SBoxPanel.h"
 
 
 class ISequencer;
 class SExpandableArea;
 class SSearchableRigHierarchyTreeView;
+class UControlRig;
 
-class SControlRigOutliner: public SCompoundWidget, public FControlRigBaseDockableView
+class SControlRigOutlinerItem : public SCompoundWidget
 {
-
-	SLATE_BEGIN_ARGS(SControlRigOutliner)
-	{}
+	SLATE_BEGIN_ARGS(SControlRigOutlinerItem){}
+	SLATE_ARGUMENT(UControlRig*, ControlRig)
 	SLATE_END_ARGS()
-	~SControlRigOutliner();
-
-	void Construct(const FArguments& InArgs, FControlRigEditMode& InEditMode);
+	SControlRigOutlinerItem();
+	~SControlRigOutlinerItem();
 
+	void Construct(const FArguments& InArgs);
 private:
-	virtual void HandleControlSelected(UControlRig* Subject, FRigControlElement* InControl, bool bSelected) override;
-	virtual void HandleControlAdded(UControlRig* ControlRig, bool bIsAdded) override;
-	virtual void NewControlRigSet(UControlRig* ControlRig) override;
+	void HandleControlSelected(UControlRig* Subject, FRigControlElement* InControl, bool bSelected);
+	void OnObjectsReplaced(const TMap<UObject*, UObject*>& OldToNewInstanceMap);
+	void NewControlRigSet(UControlRig* ControlRig);
 
 	const URigHierarchy* GetHierarchy() const;
 	void HandleSelectionChanged(TSharedPtr<FRigTreeElement> Selection, ESelectInfo::Type SelectInfo);
 
+	//visibility button
+	const FSlateBrush* GetVisibilityBrushForElement() const;
+	FReply OnToggleVisibility();
+	bool VisibilityToggleEnabled() const;
+
+
 	/** Hierarchy picker for controls*/
 	TSharedPtr<SSearchableRigHierarchyTreeView> HierarchyTreeView;
 	FRigTreeDisplaySettings DisplaySettings;
 	const FRigTreeDisplaySettings& GetDisplaySettings() const { return DisplaySettings; }
 	bool bIsChangingRigHierarchy = false;
 	TSharedPtr<SExpandableArea> PickerExpander;
+
+	TWeakObjectPtr<UControlRig> CurrentControlRig;
+};
+
+class SControlRigOutliner: public SCompoundWidget
+{
+
+	SLATE_BEGIN_ARGS(SControlRigOutliner)
+	{}
+	SLATE_END_ARGS()
+	~SControlRigOutliner();
+
+	void Construct(const FArguments& InArgs, FControlRigEditMode& InEditMode);
+	void SetEditMode(FControlRigEditMode& InEditMode);
+
+private:
+	void HandleControlAdded(UControlRig* ControlRig, bool bIsAdded);
+	void Rebuild();
+	TSharedPtr<SVerticalBox> MainBoxPtr;
+	FEditorModeTools* ModeTools = nullptr;
+
 };
 
diff --git a/Engine/Plugins/Animation/ControlRig/Source/ControlRigEditor/Private/EditMode/SControlRigSnapper.cpp b/Engine/Plugins/Animation/ControlRig/Source/ControlRigEditor/Private/EditMode/SControlRigSnapper.cpp
--- a/Engine/Plugins/Animation/ControlRig/Source/ControlRigEditor/Private/EditMode/SControlRigSnapper.cpp
+++ b/Engine/Plugins/Animation/ControlRig/Source/ControlRigEditor/Private/EditMode/SControlRigSnapper.cpp
@@ -475,19 +475,23 @@ void SControlRigSnapper::SetStartEndFrames()
 FControlRigSnapperSelection SControlRigSnapper::GetSelection(bool bGetAll) 
 {
 	FControlRigSnapperSelection Selection;
-	UControlRig* ControlRig = GetControlRig();
-	if (ControlRig)
+	TArray<UControlRig*> ControlRigs;
+	GetControlRigs(ControlRigs);
+	for (UControlRig* ControlRig : ControlRigs)
 	{
-		TArray<FName> SelectedControls = ControlRig->CurrentControlSelection();
-		if (SelectedControls.Num() > 0)
+		if (ControlRig)
 		{
-			FControlRigForWorldTransforms ControlRigAndSelection;
-			ControlRigAndSelection.ControlRig = ControlRig;
-			ControlRigAndSelection.ControlNames = SelectedControls;
-			Selection.ControlRigs.Add(ControlRigAndSelection);
-			if (bGetAll == false)
+			TArray<FName> SelectedControls = ControlRig->CurrentControlSelection();
+			if (SelectedControls.Num() > 0)
 			{
-				return Selection;
+				FControlRigForWorldTransforms ControlRigAndSelection;
+				ControlRigAndSelection.ControlRig = ControlRig;
+				ControlRigAndSelection.ControlNames = SelectedControls;
+				Selection.ControlRigs.Add(ControlRigAndSelection);
+				if (bGetAll == false)
+				{
+					return Selection;
+				}
 			}
 		}
 	}
@@ -510,13 +514,20 @@ FControlRigSnapperSelection SControlRigSnapper::GetSelection(bool bGetAll)
 	return Selection;
 }
 
-UControlRig* SControlRigSnapper::GetControlRig() const
+void SControlRigSnapper::GetControlRigs(TArray<UControlRig*>& OutControlRigs) const
 {
+	OutControlRigs.SetNum(0);
 	if (FControlRigEditMode* EditMode = static_cast<FControlRigEditMode*>(GLevelEditorModeTools().GetActiveMode(FControlRigEditMode::ModeName)))
 	{
-		return EditMode->GetControlRig(true);
+		TArrayView<TWeakObjectPtr<UControlRig>> ControlRigPtrs = EditMode->GetControlRigs();
+		for (TWeakObjectPtr<UControlRig>& ControlRigPtr : ControlRigPtrs)
+		{
+			if (UControlRig* ControlRig = ControlRigPtr.Get())
+			{
+				OutControlRigs.Add(ControlRig);
+			}
+		}
 	}
-	return nullptr;
 }
 
 
diff --git a/Engine/Plugins/Animation/ControlRig/Source/ControlRigEditor/Private/EditMode/SControlRigSnapper.h b/Engine/Plugins/Animation/ControlRig/Source/ControlRigEditor/Private/EditMode/SControlRigSnapper.h
--- a/Engine/Plugins/Animation/ControlRig/Source/ControlRigEditor/Private/EditMode/SControlRigSnapper.h
+++ b/Engine/Plugins/Animation/ControlRig/Source/ControlRigEditor/Private/EditMode/SControlRigSnapper.h
@@ -31,7 +31,7 @@ class SControlRigSnapper : public SCompoundWidget
 
 private:
 
-	UControlRig* GetControlRig() const;
+	void GetControlRigs(TArray<UControlRig*>& OutControlRigs) const;
 
 	/*
 	* Delegates and Helpers
diff --git a/Engine/Plugins/Animation/ControlRig/Source/ControlRigEditor/Private/EditMode/SControlRigTweenWidget.cpp b/Engine/Plugins/Animation/ControlRig/Source/ControlRigEditor/Private/EditMode/SControlRigTweenWidget.cpp
--- a/Engine/Plugins/Animation/ControlRig/Source/ControlRigEditor/Private/EditMode/SControlRigTweenWidget.cpp
+++ b/Engine/Plugins/Animation/ControlRig/Source/ControlRigEditor/Private/EditMode/SControlRigTweenWidget.cpp
@@ -73,11 +73,19 @@ void SControlRigTweenWidget::Construct(const FArguments& InArgs)
 
 void SControlRigTweenWidget::OnPoseBlendChanged(float ChangedVal)
 {
-	UControlRig* ControlRig = GetControlRig();
-	if (ControlRig  && WeakSequencer.IsValid() && bIsBlending)
+	TArray<UControlRig*> ControlRigs = GetControlRigs();
+	bool bWeBlended = false;
+	for (UControlRig* ControlRig : ControlRigs)
+	{
+		if (ControlRig && WeakSequencer.IsValid() && bIsBlending)
+		{
+			PoseBlendValue = ChangedVal;
+			ControlsToTween.Blend(WeakSequencer, ChangedVal);
+			bWeBlended = true;
+		}
+	}
+	if(bWeBlended)
 	{
-		PoseBlendValue = ChangedVal;
-		ControlsToTween.Blend(WeakSequencer, ChangedVal);
 		WeakSequencer.Pin()->NotifyMovieSceneDataChanged(EMovieSceneDataChangeType::TrackValueChanged);
 	}
 }
@@ -99,13 +107,16 @@ void SControlRigTweenWidget::SetupControls()
 	IAssetEditorInstance* AssetEditor = GEditor->GetEditorSubsystem<UAssetEditorSubsystem>()->FindEditorForAsset(LevelSequence, false);
 	ILevelSequenceEditorToolkit* LevelSequenceEditor = static_cast<ILevelSequenceEditorToolkit*>(AssetEditor);
 	WeakSequencer = LevelSequenceEditor ? LevelSequenceEditor->GetSequencer() : nullptr;
-	UControlRig* ControlRig = GetControlRig();
-	if (ControlRig && WeakSequencer.IsValid())
+	TArray<UControlRig*> ControlRigs = GetControlRigs();
+	for (UControlRig* ControlRig : ControlRigs)
 	{
-		WeakSequencer.Pin()->GetFocusedMovieSceneSequence()->GetMovieScene()->Modify();
-		TArray<UControlRig*> SelectedControlRigs;
-		SelectedControlRigs.Add(ControlRig);
-		ControlsToTween.Setup(SelectedControlRigs, WeakSequencer);
+		if (ControlRig && WeakSequencer.IsValid() && bIsBlending)
+		{
+			WeakSequencer.Pin()->GetFocusedMovieSceneSequence()->GetMovieScene()->Modify();
+			TArray<UControlRig*> SelectedControlRigs;
+			SelectedControlRigs.Add(ControlRig);
+			ControlsToTween.Setup(SelectedControlRigs, WeakSequencer);
+		}
 	}
 }
 
@@ -154,8 +165,8 @@ void SControlRigTweenWidget::FinishDraggingWidget(const FVector2D InLocation)
 
 void SControlRigTweenWidget::OnPoseBlendCommited(float ChangedVal, ETextCommit::Type Type)
 {
-	UControlRig* ControlRig = GetControlRig();
-	if (ControlRig)
+	TArray<UControlRig*> ControlRigs = GetControlRigs();
+	if (ControlRigs.Num() > 0)
 	{
 		FScopedTransaction ScopedTransaction(LOCTEXT("TweenTransaction", "Tween"));
 		if (bIsBlending == false)
@@ -170,13 +181,14 @@ void SControlRigTweenWidget::OnPoseBlendCommited(float ChangedVal, ETextCommit::
 	}
 }
 
-UControlRig* SControlRigTweenWidget::GetControlRig()
+TArray<UControlRig*>SControlRigTweenWidget::GetControlRigs()
 {
+	TArray<UControlRig*> ControlRigs;
 	if (FControlRigEditMode* EditMode = static_cast<FControlRigEditMode*>(GLevelEditorModeTools().GetActiveMode(FControlRigEditMode::ModeName)))
 	{
-		return EditMode->GetControlRig(true);
+		ControlRigs = EditMode->GetControlRigsArray(true /*bIsVisible*/);
 	}
-	return nullptr;
+	return ControlRigs;
 }
 
 #undef LOCTEXT_NAMESPACE
diff --git a/Engine/Plugins/Animation/ControlRig/Source/ControlRigEditor/Private/EditMode/SControlRigTweenWidget.h b/Engine/Plugins/Animation/ControlRig/Source/ControlRigEditor/Private/EditMode/SControlRigTweenWidget.h
--- a/Engine/Plugins/Animation/ControlRig/Source/ControlRigEditor/Private/EditMode/SControlRigTweenWidget.h
+++ b/Engine/Plugins/Animation/ControlRig/Source/ControlRigEditor/Private/EditMode/SControlRigTweenWidget.h
@@ -34,7 +34,7 @@ class SControlRigTweenWidget : public SCompoundWidget
 
 private:
 
-	UControlRig* GetControlRig();
+	TArray<UControlRig*> GetControlRigs();
 
 	/*
 	* Delegates and Helpers
diff --git a/Engine/Plugins/Animation/ControlRig/Source/ControlRigEditor/Private/EditMode/SControlRigUpdatePose.cpp b/Engine/Plugins/Animation/ControlRig/Source/ControlRigEditor/Private/EditMode/SControlRigUpdatePose.cpp
--- a/Engine/Plugins/Animation/ControlRig/Source/ControlRigEditor/Private/EditMode/SControlRigUpdatePose.cpp
+++ b/Engine/Plugins/Animation/ControlRig/Source/ControlRigEditor/Private/EditMode/SControlRigUpdatePose.cpp
@@ -52,6 +52,7 @@ class SControlRigUpdatePoseDialog : public SCompoundWidget
 						.ContentPadding(FMargin(10, 5))
 						.Text(NSLOCTEXT("ControlRig", "Yes", "Yes"))
 						.OnClicked(this, &SControlRigUpdatePoseDialog::UpdatePose)
+						.IsEnabled(this, &SControlRigUpdatePoseDialog::CanUpdatePose)
 					]
 					+SHorizontalBox::Slot()
 					.AutoWidth()
@@ -74,16 +75,38 @@ class SControlRigUpdatePoseDialog : public SCompoundWidget
 
 private:
 
+	bool CanUpdatePose() const
+	{
+		FControlRigEditMode* ControlRigEditMode = static_cast<FControlRigEditMode*>(GLevelEditorModeTools().GetActiveMode(FControlRigEditMode::ModeName));
+		if (PoseAsset.IsValid() && ControlRigEditMode)
+		{
+			TMap<UControlRig*, TArray<FRigElementKey>> AllSelectedControls;
+			ControlRigEditMode->GetAllSelectedControls(AllSelectedControls);
+		}
+		return false;
+
+	}
 	FReply UpdatePose()
 	{
 		TSharedPtr<SWindow> Window = FSlateApplication::Get().FindWidgetWindow(AsShared());
 
 		FControlRigEditMode* ControlRigEditMode = static_cast<FControlRigEditMode*>(GLevelEditorModeTools().GetActiveMode(FControlRigEditMode::ModeName));
-		if (PoseAsset.IsValid() && ControlRigEditMode && ControlRigEditMode->GetControlRig(true))
+		if (PoseAsset.IsValid() && ControlRigEditMode)
 		{
-			const FScopedTransaction Transaction(NSLOCTEXT("ControlRig", "UpdatePose", "Update Pose"));
-			PoseAsset->Modify();
-			PoseAsset->SavePose(ControlRigEditMode->GetControlRig(true), false);
+			TMap<UControlRig*, TArray<FRigElementKey>> AllSelectedControls;
+			ControlRigEditMode->GetAllSelectedControls(AllSelectedControls);
+			if (AllSelectedControls.Num() == 1)
+			{
+				TArray<UControlRig*> ControlRigs;
+				AllSelectedControls.GenerateKeyArray(ControlRigs);
+				const FScopedTransaction Transaction(NSLOCTEXT("ControlRig", "UpdatePose", "Update Pose"));
+				PoseAsset->Modify();
+				PoseAsset->SavePose(ControlRigs[0], false);
+			}
+			else
+			{
+
+			}
 		}
 		if (Window.IsValid())
 		{
diff --git a/Engine/Plugins/Animation/ControlRig/Source/ControlRigEditor/Private/Editor/ControlRigEditorEditMode.cpp b/Engine/Plugins/Animation/ControlRig/Source/ControlRigEditor/Private/Editor/ControlRigEditorEditMode.cpp
--- a/Engine/Plugins/Animation/ControlRig/Source/ControlRigEditor/Private/Editor/ControlRigEditorEditMode.cpp
+++ b/Engine/Plugins/Animation/ControlRig/Source/ControlRigEditor/Private/Editor/ControlRigEditorEditMode.cpp
@@ -22,88 +22,92 @@ void FControlRigEditorEditMode::Render(const FSceneView* View, FViewport* Viewpo
 
 	if (bDrawHierarchyBones && BoneDrawMode != EBoneDrawMode::None)
 	{
-		if (UControlRig* ControlRig = GetControlRig(false))
+		TArray<UControlRig*> ControlRigs = GetControlRigsArray(true /*bIsVisible*/);
+		if (ControlRigs.Num() > 0)
 		{
-			URigHierarchy* Hierarchy = ControlRig->GetHierarchy();
-			URigHierarchy* HierarchyForSelection = Hierarchy;
-
-			if (UClass* ControlRigClass = ControlRig->GetClass())
+			if (UControlRig* ControlRig = ControlRigs[0]) //just one control rig in the CR asset editor
 			{
-				if (UControlRigBlueprint* RigBlueprint = Cast<UControlRigBlueprint>(ControlRigClass->ClassGeneratedBy))
-				{
-					HierarchyForSelection = RigBlueprint->Hierarchy;
-				}
-			}
+				URigHierarchy* Hierarchy = ControlRig->GetHierarchy();
+				URigHierarchy* HierarchyForSelection = Hierarchy;
 
-			if (BoneDrawMode == EBoneDrawMode::SelectedAndParents)
-			{
-				if (BoneHasSelectedChild.Num() != Hierarchy->Num())
+				if (UClass* ControlRigClass = ControlRig->GetClass())
 				{
-					BoneHasSelectedChild.SetNumZeroed(Hierarchy->Num());
+					if (UControlRigBlueprint* RigBlueprint = Cast<UControlRigBlueprint>(ControlRigClass->ClassGeneratedBy))
+					{
+						HierarchyForSelection = RigBlueprint->Hierarchy;
+					}
 				}
 
-				Hierarchy->ForEach<FRigBoneElement>([this, Hierarchy, HierarchyForSelection](FRigBoneElement* BoneElement) -> bool
+				if (BoneDrawMode == EBoneDrawMode::SelectedAndParents)
 				{
-					BoneHasSelectedChild[BoneElement->GetIndex()] = false;
-					if (!HierarchyForSelection->IsSelected(BoneElement->GetKey()))
+					if (BoneHasSelectedChild.Num() != Hierarchy->Num())
 					{
-						return true;
+						BoneHasSelectedChild.SetNumZeroed(Hierarchy->Num());
 					}
 
-					int32 ParentIndex = Hierarchy->GetFirstParent(BoneElement->GetIndex());
-					while (ParentIndex != INDEX_NONE)
-					{
-						BoneHasSelectedChild[ParentIndex] = true;
-						ParentIndex = Hierarchy->GetFirstParent(ParentIndex);
-					}
+					Hierarchy->ForEach<FRigBoneElement>([this, Hierarchy, HierarchyForSelection](FRigBoneElement* BoneElement) -> bool
+						{
+							BoneHasSelectedChild[BoneElement->GetIndex()] = false;
+							if (!HierarchyForSelection->IsSelected(BoneElement->GetKey()))
+							{
+								return true;
+							}
+
+							int32 ParentIndex = Hierarchy->GetFirstParent(BoneElement->GetIndex());
+							while (ParentIndex != INDEX_NONE)
+							{
+								BoneHasSelectedChild[ParentIndex] = true;
+								ParentIndex = Hierarchy->GetFirstParent(ParentIndex);
+							}
 
-					return true;
-				});
-			}
-			
-			Hierarchy->ForEach<FRigBoneElement>([this, PDI, Hierarchy, HierarchyForSelection, &BoneDrawMode](FRigBoneElement* BoneElement) -> bool
-			{
-				const int32 ParentIndex = Hierarchy->GetFirstParent(BoneElement->GetIndex());
-				const bool bSelected = HierarchyForSelection->IsSelected(BoneElement->GetKey());
+							return true;
+						});
+				}
 
-				if (!bSelected)
-				{
-					if (BoneDrawMode == EBoneDrawMode::Selected)
+				Hierarchy->ForEach<FRigBoneElement>([this, PDI, Hierarchy, HierarchyForSelection, &BoneDrawMode](FRigBoneElement* BoneElement) -> bool
 					{
-						return true;
-					}
-					else if (BoneDrawMode == EBoneDrawMode::SelectedAndParents)
-					{
-						if (!BoneHasSelectedChild[BoneElement->GetIndex()])
+						const int32 ParentIndex = Hierarchy->GetFirstParent(BoneElement->GetIndex());
+						const bool bSelected = HierarchyForSelection->IsSelected(BoneElement->GetKey());
+
+						if (!bSelected)
 						{
-							return true;
+							if (BoneDrawMode == EBoneDrawMode::Selected)
+							{
+								return true;
+							}
+							else if (BoneDrawMode == EBoneDrawMode::SelectedAndParents)
+							{
+								if (!BoneHasSelectedChild[BoneElement->GetIndex()])
+								{
+									return true;
+								}
+							}
 						}
-					}
-				}
 
-				const FLinearColor LineColor = bSelected ? FLinearColor(1.0f, 0.34f, 0.0f, 1.0f) : FLinearColor::White;
+						const FLinearColor LineColor = bSelected ? FLinearColor(1.0f, 0.34f, 0.0f, 1.0f) : FLinearColor::White;
 
-				FVector Start, End;
-				if (ParentIndex != INDEX_NONE)
-				{
-					Start = Hierarchy->GetGlobalTransform(ParentIndex).GetLocation();
-					End = Hierarchy->GetGlobalTransform(BoneElement->GetIndex()).GetLocation();
-				}
-				else
-				{
-					Start = FVector::ZeroVector;
-					End = Hierarchy->GetGlobalTransform(BoneElement->GetIndex()).GetLocation();
-				}
+						FVector Start, End;
+						if (ParentIndex != INDEX_NONE)
+						{
+							Start = Hierarchy->GetGlobalTransform(ParentIndex).GetLocation();
+							End = Hierarchy->GetGlobalTransform(BoneElement->GetIndex()).GetLocation();
+						}
+						else
+						{
+							Start = FVector::ZeroVector;
+							End = Hierarchy->GetGlobalTransform(BoneElement->GetIndex()).GetLocation();
+						}
 
-				const float BoneLength = (End - Start).Size();
-				const float Radius = FMath::Clamp(BoneLength * 0.05f, 0.1f, 10000.f);
+						const float BoneLength = (End - Start).Size();
+						const float Radius = FMath::Clamp(BoneLength * 0.05f, 0.1f, 10000.f);
 
-				//Render Sphere for bone end point and a cone between it and its parent.
-				PDI->SetHitProxy(new HPersonaBoneHitProxy(BoneElement->GetIndex(), BoneElement->GetName()));
-				SkeletalDebugRendering::DrawWireBone(PDI, Start, End, LineColor, SDPG_Foreground, Radius);
-				PDI->SetHitProxy(nullptr);
-				return true;
-			});
+						//Render Sphere for bone end point and a cone between it and its parent.
+						PDI->SetHitProxy(new HPersonaBoneHitProxy(BoneElement->GetIndex(), BoneElement->GetName()));
+						SkeletalDebugRendering::DrawWireBone(PDI, Start, End, LineColor, SDPG_Foreground, Radius);
+						PDI->SetHitProxy(nullptr);
+						return true;
+					});
+			}
 		}
 	}
 }
diff --git a/Engine/Plugins/Animation/ControlRig/Source/ControlRigEditor/Private/Sequencer/ControlRigParameterTrackEditor.cpp b/Engine/Plugins/Animation/ControlRig/Source/ControlRigEditor/Private/Sequencer/ControlRigParameterTrackEditor.cpp
--- a/Engine/Plugins/Animation/ControlRig/Source/ControlRigEditor/Private/Sequencer/ControlRigParameterTrackEditor.cpp
+++ b/Engine/Plugins/Animation/ControlRig/Source/ControlRigEditor/Private/Sequencer/ControlRigParameterTrackEditor.cpp
@@ -227,9 +227,23 @@ FControlRigParameterTrackEditor::FControlRigParameterTrackEditor(TSharedRef<ISeq
 				{
 					TMap<UControlRig*, UControlRig*> OldToNewControlRigs;
 					FControlRigEditMode* ControlRigEditMode = GetEditMode();
-					if (ControlRigEditMode && ControlRigEditMode->GetControlRig(true) && ControlRigEditMode->GetControlRig(true)->GetObjectBinding())
+					if (ControlRigEditMode) 
 					{
-						if (ControlRigEditMode->GetControlRig(true)->GetObjectBinding()->GetBoundObject() == nullptr)
+						TArrayView<TWeakObjectPtr<UControlRig>> ControlRigs = ControlRigEditMode->GetControlRigs();
+						bool bRequestEvaluate = false;
+						for (TWeakObjectPtr<UControlRig>& ControlRigPtr : ControlRigs)
+						{
+							if (UControlRig* ControlRig = ControlRigPtr.Get())
+							{
+								if (ControlRig->GetObjectBinding() && ControlRig->GetObjectBinding()->GetBoundObject() == nullptr)
+								{
+									bRequestEvaluate = true;
+									break;
+								}
+							}
+						}
+
+						if (bRequestEvaluate)
 						{
 							GetSequencer()->RequestEvaluate();
 						}
@@ -243,7 +257,7 @@ FControlRigParameterTrackEditor::FControlRigParameterTrackEditor(TSharedRef<ISeq
 							if (OldControlRigComponent->GetControlRig())
 							{
 								UControlRig* NewControlRig = nullptr;
-							if(NewControlRigComponent)
+								if(NewControlRigComponent)
 								{
 									NewControlRig = NewControlRigComponent->GetControlRig();
 								}
@@ -284,10 +298,9 @@ FControlRigParameterTrackEditor::FControlRigParameterTrackEditor(TSharedRef<ISeq
 								}
 								if (ControlRigEditMode)
 								{
-									if (ControlRigEditMode->GetControlRig(false) == OldControlRig)
-									{
-										ControlRigEditMode->SetObjects(*NewControlRig, nullptr, GetSequencer());
-									}
+									ControlRigEditMode->ReplaceControlRig(OldControlRig, *NewControlRig);
+									ControlRigEditMode->AddControlRigObject(*NewControlRig, GetSequencer());
+									
 									UControlRig* PtrNewControlRig = *NewControlRig;
 
 									auto UpdateSelectionDelegate = [this, SelectedControls, PtrNewControlRig]()
@@ -996,11 +1009,13 @@ void FControlRigParameterTrackEditor::BakeToControlRig(UClass* InClass, FGuid Ob
 					}
 					else
 					{
+						/* mz todo we don't unbind  will test more
 						UControlRig* OldControlRig = ControlRigEditMode->GetControlRig(false);
 						if (OldControlRig)
 						{
 							UnbindControlRig(OldControlRig);
 						}
+						*/
 					}
 
 					if (bReuseControlRig == false)
@@ -1055,7 +1070,7 @@ void FControlRigParameterTrackEditor::BakeToControlRig(UClass* InClass, FGuid Ob
 					//Finish Setup
 					if (ControlRigEditMode)
 					{
-						ControlRigEditMode->SetObjects(ControlRig, nullptr, GetSequencer());
+						ControlRigEditMode->AddControlRigObject(ControlRig, GetSequencer());
 					}
 					BindControlRig(ControlRig);
 
@@ -1354,7 +1369,7 @@ void FControlRigParameterTrackEditor::AddControlRig(UClass* InClass, UObject* Bo
 
 			if (ControlRigEditMode)
 			{
-				ControlRigEditMode->SetObjects(ControlRig, nullptr, GetSequencer());
+				ControlRigEditMode->AddControlRigObject(ControlRig, GetSequencer());
 			}
 			BindControlRig(ControlRig);
 
@@ -1424,16 +1439,11 @@ bool FControlRigParameterTrackEditor::CanAddTransformKeysForSelectedObjects() co
 	}
 
 	FControlRigEditMode* ControlRigEditMode = GetEditMode();
-	if (ControlRigEditMode && ControlRigEditMode->GetControlRig(false))
+	if (ControlRigEditMode)
 	{
-		UControlRig* ControlRig = ControlRigEditMode->GetControlRig(false);
-		FString OurName = ControlRig->GetName();
-		FName Name(*OurName);
-		if (TSharedPtr<IControlRigObjectBinding> ObjectBinding = ControlRig->GetObjectBinding())
-		{
-			TArray<FName> ControlNames = ControlRig->CurrentControlSelection();
-			return (ControlNames.Num() > 0);
-		}
+		TMap<UControlRig*, TArray<FRigElementKey>> SelectedControls;
+		ControlRigEditMode->GetAllSelectedControls(SelectedControls);
+		return (SelectedControls.Num() > 0);
 	}
 	return false;
 }
@@ -1446,20 +1456,26 @@ void FControlRigParameterTrackEditor::OnAddTransformKeysForSelectedObjects(EMovi
 	}
 
 	FControlRigEditMode* ControlRigEditMode = GetEditMode();
-	if (ControlRigEditMode && ControlRigEditMode->GetControlRig(false))
+	if (ControlRigEditMode)
 	{
-		UControlRig* ControlRig = ControlRigEditMode->GetControlRig(false);
-		FString OurName = ControlRig->GetName();
-		FName Name(*OurName);
-		if (TSharedPtr<IControlRigObjectBinding> ObjectBinding = ControlRig->GetObjectBinding())
+		TMap<UControlRig*, TArray<FRigElementKey>> SelectedControls;
+		ControlRigEditMode->GetAllSelectedControls(SelectedControls);
+
+		for (TPair<UControlRig*, TArray<FRigElementKey>>& Selection : SelectedControls)
 		{
-			TArray<FName> ControlNames = ControlRig->CurrentControlSelection();
-			for (const FName& ControlName : ControlNames)
+			UControlRig* ControlRig = Selection.Key;
+			FString OurName = ControlRig->GetName();
+			FName Name(*OurName);
+			if (TSharedPtr<IControlRigObjectBinding> ObjectBinding = ControlRig->GetObjectBinding())
 			{
-				USceneComponent* Component = Cast<USceneComponent>(ObjectBinding->GetBoundObject());
-				if (Component)
+				TArray<FName> ControlNames = ControlRig->CurrentControlSelection();
+				for (const FName& ControlName : ControlNames)
 				{
-					AddControlKeys(Component, ControlRig, Name, ControlName, (EControlRigContextChannelToKey) Channel, ESequencerKeyMode::ManualKeyForced, FLT_MAX);
+					USceneComponent* Component = Cast<USceneComponent>(ObjectBinding->GetBoundObject());
+					if (Component)
+					{
+						AddControlKeys(Component, ControlRig, Name, ControlName, (EControlRigContextChannelToKey)Channel, ESequencerKeyMode::ManualKeyForced, FLT_MAX);
+					}
 				}
 			}
 		}
@@ -1714,27 +1730,40 @@ void FControlRigParameterTrackEditor::OnSequencerDataChanged(EMovieSceneDataChan
 
 	//if we have a valid control rig edit mode need to check and see the control rig in that mode is still in a track
 	//if not we get rid of it.
-	if (ControlRigEditMode && ControlRigEditMode->GetControlRig(false) != nullptr && MovieScene && (DataChangeType == EMovieSceneDataChangeType::MovieSceneStructureItemRemoved ||
+	if (ControlRigEditMode && ControlRigEditMode->GetControlRigsArray(false /*bIsVisible*/).Num() != 0 && MovieScene && (DataChangeType == EMovieSceneDataChangeType::MovieSceneStructureItemRemoved ||
 		DataChangeType == EMovieSceneDataChangeType::Unknown))
 	{
+		TArray<UControlRig*> ControlRigs = ControlRigEditMode->GetControlRigsArray(false /*bIsVisible*/);
 		float FPS = 1.f / (float)GetSequencer()->GetFocusedDisplayRate().AsInterval();
-		ControlRigEditMode->GetControlRig(false)->SetFramesPerSecond(FPS);
-
-		const TArray<FMovieSceneBinding>& Bindings = MovieScene->GetBindings();
-		for (const FMovieSceneBinding& Binding : Bindings)
+		for (UControlRig* ControlRig : ControlRigs)
 		{
-			UMovieSceneControlRigParameterTrack* Track = Cast<UMovieSceneControlRigParameterTrack>(MovieScene->FindTrack(UMovieSceneControlRigParameterTrack::StaticClass(), Binding.GetObjectGuid(), NAME_None));
-			if (Track && Track->GetControlRig() == ControlRigEditMode->GetControlRig(false))
+			if (ControlRig)
 			{
-				return; //just exit out we still have a good track
+				ControlRig->SetFramesPerSecond(FPS);
+
+				const TArray<FMovieSceneBinding>& Bindings = MovieScene->GetBindings();
+				bool bControlRigInTrack = false;
+				for (const FMovieSceneBinding& Binding : Bindings)
+				{
+					UMovieSceneControlRigParameterTrack* Track = Cast<UMovieSceneControlRigParameterTrack>(MovieScene->FindTrack(UMovieSceneControlRigParameterTrack::StaticClass(), Binding.GetObjectGuid(), NAME_None));
+					if (Track && Track->GetControlRig() == ControlRig)
+					{
+						bControlRigInTrack = true;
+						break;; //continue, we still have a good track
+					}
+				}
+				/* Nope don't do this anymore, todo mz
+				if (FEditorModeTools* Tools = GetEditorModeTools())
+				{
+					Tools->DeactivateMode(FControlRigEditMode::ModeName);
+				}
+				*/
+				if (bControlRigInTrack == false)
+				{
+					ControlRigEditMode->RemoveControlRig(ControlRig);
+				}
 			}
 		}
-		//okay no good track so deactive it and delete it's Control Rig and bindings.
-		if (FEditorModeTools* Tools = GetEditorModeTools())
-		{
-			Tools->DeactivateMode(FControlRigEditMode::ModeName);
-		}
-		ControlRigEditMode->SetObjects(nullptr, nullptr, GetSequencer());
 	}
 }
 
@@ -1767,15 +1796,16 @@ void FControlRigParameterTrackEditor::OnCurveDisplayChanged(FCurveModel* CurveMo
 				{
 					if (ControlRigEditMode)
 					{
-						ControlRigEditMode->SetObjects(ControlRig, nullptr, GetSequencer());
+						ControlRigEditMode->AddControlRigObject(ControlRig, GetSequencer());
 					}
 				}
 			}
 			else
 			{
-				if (ControlRigEditMode->GetControlRig(false) != ControlRig)
+				TArray<UControlRig*> ControlRigs = ControlRigEditMode->GetControlRigsArray(false /*bIsVisible*/);
+				if (ControlRigs.Contains(ControlRig)== false)
 				{
-					ControlRigEditMode->SetObjects(ControlRig, nullptr, GetSequencer());
+					ControlRigEditMode->AddControlRigObject(ControlRig, GetSequencer());
 				}
 			}
 			//Not 100% safe but for now it is since that's all we show in the curve editor
@@ -1892,10 +1922,14 @@ void FControlRigParameterTrackEditor::OnSelectionChanged(TArray<UMovieSceneTrack
 	{
 		if (ControlRigEditMode)
 		{
-			ControlRig = ControlRigEditMode->GetControlRig(false);
-			if (ControlRig)
+			TMap<UControlRig*, TArray<FRigElementKey>> AllSelectedControls;
+			for (TPair<UControlRig*, TArray<FRigElementKey>>& SelectedControl : AllSelectedControls)
 			{
-				ControlRig->ClearControlSelection();
+				ControlRig = SelectedControl.Key;
+				if (ControlRig)
+				{
+					ControlRig->ClearControlSelection();
+				}
 			}
 		}
 		for (UMovieSceneTrack* Track : InTracks)
@@ -1908,11 +1942,7 @@ void FControlRigParameterTrackEditor::OnSelectionChanged(TArray<UMovieSceneTrack
 				{
 					if (ControlRigEditMode)
 					{
-						ControlRig = ControlRigEditMode->GetControlRig(false);
-						if (ControlRig != TrackControlRig)
-						{
-							ControlRigEditMode->SetObjects(TrackControlRig, nullptr, GetSequencer());
-						}
+						ControlRigEditMode->AddControlRigObject(TrackControlRig, GetSequencer());
 						break;
 					}
 					else
@@ -1922,7 +1952,7 @@ void FControlRigParameterTrackEditor::OnSelectionChanged(TArray<UMovieSceneTrack
 						{
 							if (ControlRigEditMode)
 							{
-								ControlRigEditMode->SetObjects(TrackControlRig, nullptr, GetSequencer());
+								ControlRigEditMode->AddControlRigObject(TrackControlRig, GetSequencer());
 							}
 						}
 					}
@@ -1955,19 +1985,14 @@ void FControlRigParameterTrackEditor::SelectRigsAndControls(UControlRig* Control
 				{
 					if (ControlRigEditMode)
 					{
-						ControlRigEditMode->SetObjects(ControlRig, nullptr, GetSequencer());
+						ControlRigEditMode->AddControlRigObject(ControlRig, GetSequencer());
 					}
 				}
 			}
 			else
 			{
-				if (ControlRigEditMode->GetControlRig(false) != ControlRig)
+				if (ControlRigEditMode->AddControlRigObject(ControlRig, GetSequencer()))
 				{
-					if (ControlRigEditMode->GetControlRig(false))
-					{
-						ControlRigEditMode->GetControlRig(false)->ClearControlSelection();
-					}
-					ControlRigEditMode->SetObjects(ControlRig, nullptr, GetSequencer());
 					//force an evaluation, this will get the control rig setup so edit mode looks good.
 					if (GetSequencer().IsValid())
 					{
@@ -1995,10 +2020,15 @@ void FControlRigParameterTrackEditor::SelectRigsAndControls(UControlRig* Control
 	//Always clear the control rig(s) in the edit mode.
 	if (ControlRigEditMode)
 	{
-		ControlRig = ControlRigEditMode->GetControlRig(false);
-		if (ControlRig)
+		TMap<UControlRig*, TArray<FRigElementKey>> SelectedControls;
+		ControlRigEditMode->GetAllSelectedControls(SelectedControls);
+		for (TPair<UControlRig*, TArray<FRigElementKey>>& Selection : SelectedControls)
 		{
-			ControlRig->ClearControlSelection();
+			ControlRig = Selection.Key;
+			if (ControlRig)
+			{
+				ControlRig->ClearControlSelection();
+			}
 		}
 	}
 	for (TPair<UControlRig *, TSet<FName>> Pair : RigsAndControls)
@@ -2252,19 +2282,25 @@ void FControlRigParameterTrackEditor::SetUpEditModeIfNeeded(UControlRig* Control
 		{
 			if (ControlRigEditMode)
 			{
-				ControlRigEditMode->SetObjects(ControlRig, nullptr, GetSequencer());
+				ControlRigEditMode->AddControlRigObject(ControlRig, GetSequencer());
 			}
 		}
 	}
 	else
 	{
-		if (ControlRigEditMode->GetControlRig(false) != ControlRig)
+
+		TMap<UControlRig*, TArray<FRigElementKey>> SelectedControls;
+		ControlRigEditMode->GetAllSelectedControls(SelectedControls);
+		for (TPair<UControlRig*, TArray<FRigElementKey>>& Selection : SelectedControls)
 		{
-			if (ControlRigEditMode->GetControlRig(false))
+			if(Selection.Key)
 			{
-				ControlRigEditMode->GetControlRig(false)->ClearControlSelection();
+				Selection.Key->ClearControlSelection();
 			}
-			ControlRigEditMode->SetObjects(ControlRig, nullptr, GetSequencer());
+		}
+		
+		if (ControlRigEditMode->AddControlRigObject(ControlRig, GetSequencer()))
+		{
 			//force an evaluation, this will get the control rig setup so edit mode looks good.
 			if (GetSequencer().IsValid())
 			{
diff --git a/Engine/Plugins/Animation/ControlRig/Source/ControlRigEditor/Private/Tools/AssetTypeActions_ControlRigPose.cpp b/Engine/Plugins/Animation/ControlRig/Source/ControlRigEditor/Private/Tools/AssetTypeActions_ControlRigPose.cpp
--- a/Engine/Plugins/Animation/ControlRig/Source/ControlRigEditor/Private/Tools/AssetTypeActions_ControlRigPose.cpp
+++ b/Engine/Plugins/Animation/ControlRig/Source/ControlRigEditor/Private/Tools/AssetTypeActions_ControlRigPose.cpp
@@ -79,11 +79,13 @@ void FAssetTypeActions_ControlRigPose::GetActions(const TArray<UObject*>& InObje
 		if (PoseAsset)
 		{
 			FControlRigEditMode* ControlRigEditMode = static_cast<FControlRigEditMode*>(GLevelEditorModeTools().GetActiveMode(FControlRigEditMode::ModeName));
-			if (ControlRigEditMode && ControlRigEditMode->GetControlRig(true))
+			if (ControlRigEditMode)
 			{ 
-				UControlRig* ControlRig = ControlRigEditMode->GetControlRig(true);
-
-				Section.AddDynamicEntry("Control Rig Pose Actions", FNewToolMenuSectionDelegate::CreateLambda([ControlRig,PoseAsset](FToolMenuSection& InSection)
+				TMap<UControlRig*, TArray<FRigElementKey>> AllSelectedControls;
+				ControlRigEditMode->GetAllSelectedControls(AllSelectedControls);
+				TArray<UControlRig*> ControlRigs;
+				AllSelectedControls.GenerateKeyArray(ControlRigs);
+				Section.AddDynamicEntry("Control Rig Pose Actions", FNewToolMenuSectionDelegate::CreateLambda([ControlRigs,PoseAsset](FToolMenuSection& InSection)
 					{
 						{
 							const FText Label = LOCTEXT("PastePoseButton", "Paste Pose");
@@ -94,9 +96,12 @@ void FAssetTypeActions_ControlRigPose::GetActions(const TArray<UObject*>& InObje
 								ToolTipText,
 								FSlateIcon(),
 								FUIAction(
-									FExecuteAction::CreateLambda([ControlRig,PoseAsset]()
+									FExecuteAction::CreateLambda([ControlRigs,PoseAsset]()
 										{
-											PoseAsset->PastePose(ControlRig, false, false);
+											for (UControlRig* ControlRig : ControlRigs)
+											{
+												PoseAsset->PastePose(ControlRig, false, false);
+											}
 										})));
 						}
 						{
@@ -108,24 +113,30 @@ void FAssetTypeActions_ControlRigPose::GetActions(const TArray<UObject*>& InObje
 								ToolTipText,
 								FSlateIcon(),
 								FUIAction(
-									FExecuteAction::CreateLambda([ControlRig, PoseAsset]()
+									FExecuteAction::CreateLambda([ControlRigs, PoseAsset]()
 										{
-											PoseAsset->SelectControls(ControlRig);
+											for (UControlRig* ControlRig : ControlRigs)
+											{
+												PoseAsset->SelectControls(ControlRig);
+											}
 										})));
 						}
 						{
-							const FText Label = LOCTEXT("UpdatePose", "Update Pose");
-							const FText ToolTipText = LOCTEXT("UpdatePoseTooltip", "Update the pose based upon current control rig pose and selected controls");
-							InSection.AddMenuEntry(
-								"UpdatePose",
-								Label,
-								ToolTipText,
-								FSlateIcon(),
-								FUIAction(
-									FExecuteAction::CreateLambda([ControlRig,PoseAsset]()
-										{
-											PoseAsset->SavePose(ControlRig,false);
-										})));
+							if (ControlRigs.Num() == 1)
+							{
+								const FText Label = LOCTEXT("UpdatePose", "Update Pose");
+								const FText ToolTipText = LOCTEXT("UpdatePoseTooltip", "Update the pose based upon current control rig pose and selected controls");
+								InSection.AddMenuEntry(
+									"UpdatePose",
+									Label,
+									ToolTipText,
+									FSlateIcon(),
+									FUIAction(
+										FExecuteAction::CreateLambda([ControlRigs, PoseAsset]()
+											{
+												PoseAsset->SavePose(ControlRigs[0], false);
+											})));
+							}
 						}
 						{
 
