Control Rig: Guard against importing meshes with no skeleton assigned

diff --git a/Engine/Plugins/Animation/ControlRig/Source/ControlRigEditor/Private/Editor/SRigHierarchy.cpp b/Engine/Plugins/Animation/ControlRig/Source/ControlRigEditor/Private/Editor/SRigHierarchy.cpp
--- a/Engine/Plugins/Animation/ControlRig/Source/ControlRigEditor/Private/Editor/SRigHierarchy.cpp
+++ b/Engine/Plugins/Animation/ControlRig/Source/ControlRigEditor/Private/Editor/SRigHierarchy.cpp
@@ -43,6 +43,8 @@
 #include "SRigSpacePickerWidget.h"
 #include "Settings/ControlRigSettings.h"
 #include "Slate/Private/Widgets/Views/SListPanel.h"
+#include "Framework/Notifications/NotificationManager.h"
+#include "Widgets/Notifications/SNotificationList.h"
 
 #define LOCTEXT_NAMESPACE "SRigHierarchy"
 
@@ -1336,6 +1338,24 @@ void SRigHierarchy::ImportHierarchy(const FAssetData& InAssetData)
 	USkeletalMesh* Mesh = Cast<USkeletalMesh>(InAssetData.GetAsset());
 	if (Mesh && Hierarchy)
 	{
+		// filter out meshes that don't contain a skeleton
+		if(Mesh->GetSkeleton() == nullptr)
+		{
+			FNotificationInfo Info(LOCTEXT("SkeletalMeshHasNoSkeleton", "Chosen Skeletal Mesh has no assigned skeleton. This needs to fixed before the mesh can be used for a Control Rig."));
+			Info.bUseSuccessFailIcons = true;
+			Info.Image = FEditorStyle::GetBrush(TEXT("MessageLog.Warning"));
+			Info.bFireAndForget = true;
+			Info.bUseThrobber = true;
+			Info.FadeOutDuration = 2.f;
+			Info.ExpireDuration = 8.f;;
+			TSharedPtr<SNotificationItem> NotificationPtr = FSlateNotificationManager::Get().AddNotification(Info);
+			if (NotificationPtr)
+			{
+				NotificationPtr->SetCompletionState(SNotificationItem::CS_Fail);
+			}
+			return;
+		}
+		
 		TGuardValue<bool> SuspendBlueprintNotifs(ControlRigBlueprint->bSuspendAllNotifications, true);
 
 		FScopedTransaction Transaction(LOCTEXT("HierarchyImport", "Import Hierarchy"));
