CIS Fix for shootergame, TObjectPtr used.
https://horde.devtools.epicgames.com/job/62336de7fe20b5f25d443ee2?step=3bec&issue=154189

diff --git a/Engine/Plugins/Animation/ControlRig/Source/ControlRigEditor/Private/EditMode/ControlRigControlsProxy.cpp b/Engine/Plugins/Animation/ControlRig/Source/ControlRigEditor/Private/EditMode/ControlRigControlsProxy.cpp
--- a/Engine/Plugins/Animation/ControlRig/Source/ControlRigEditor/Private/EditMode/ControlRigControlsProxy.cpp
+++ b/Engine/Plugins/Animation/ControlRig/Source/ControlRigEditor/Private/EditMode/ControlRigControlsProxy.cpp
@@ -756,7 +756,7 @@ void UControlRigDetailPanelControlProxies::RemoveAllProxies(UControlRig* Control
 	//no control rig remove all
 	if (ControlRig == nullptr)
 	{
-		for (TPair<UControlRig*, FControlToProxyMap>& ControlRigProxies : AllProxies)
+		for (TPair<TObjectPtr<UControlRig>, FControlToProxyMap>& ControlRigProxies : AllProxies)
 		{
 			for (TPair<FName, TObjectPtr<UControlRigControlsProxy> >& Pair : ControlRigProxies.Value.ControlToProxy)
 			{
diff --git a/Engine/Plugins/Animation/ControlRig/Source/ControlRigEditor/Private/EditMode/ControlRigControlsProxy.h b/Engine/Plugins/Animation/ControlRig/Source/ControlRigEditor/Private/EditMode/ControlRigControlsProxy.h
--- a/Engine/Plugins/Animation/ControlRig/Source/ControlRigEditor/Private/EditMode/ControlRigControlsProxy.h
+++ b/Engine/Plugins/Animation/ControlRig/Source/ControlRigEditor/Private/EditMode/ControlRigControlsProxy.h
@@ -315,7 +315,7 @@ class UControlRigDetailPanelControlProxies :public UObject
 protected:
 
 	UPROPERTY()
-	TMap<UControlRig*, FControlToProxyMap> AllProxies; //proxies themselves contain weakobjectptr to the controlrig
+	TMap<TObjectPtr<UControlRig>, FControlToProxyMap> AllProxies; //proxies themselves contain weakobjectptr to the controlrig
 
 	UPROPERTY()
 	TArray< TObjectPtr<UControlRigControlsProxy>> SelectedProxies;
