Sequencer: Control Rig: Use anim sequence sampling rate as sequencer display rate when editing/baking with an anim sequence with a control rig.

diff --git a/Engine/Plugins/Animation/ControlRig/Source/ControlRigEditor/Private/ControlRigEditorModule.cpp b/Engine/Plugins/Animation/ControlRig/Source/ControlRigEditor/Private/ControlRigEditorModule.cpp
--- a/Engine/Plugins/Animation/ControlRig/Source/ControlRigEditor/Private/ControlRigEditorModule.cpp
+++ b/Engine/Plugins/Animation/ControlRig/Source/ControlRigEditor/Private/ControlRigEditorModule.cpp
@@ -657,7 +657,10 @@ void FControlRigEditorModule::BakeToControlRig(UClass* ControlRigClass, UAnimSeq
 
 		FFrameRate TickResolution = MovieScene->GetTickResolution();
 		float Duration = AnimSequence->GetPlayLength();
-		LevelSequence->GetMovieScene()->SetPlaybackRange(0, (Duration * TickResolution).FloorToFrame().Value);
+		MovieScene->SetPlaybackRange(0, (Duration * TickResolution).FloorToFrame().Value);
+		FFrameRate SequenceFrameRate = AnimSequence->GetSamplingFrameRate();
+		MovieScene->SetDisplayRate(SequenceFrameRate);
+
 
 		GEditor->GetEditorSubsystem<UAssetEditorSubsystem>()->OpenEditorForAsset(LevelSequence);
 
@@ -665,8 +668,11 @@ void FControlRigEditorModule::BakeToControlRig(UClass* ControlRigClass, UAnimSeq
 		ILevelSequenceEditorToolkit* LevelSequenceEditor = static_cast<ILevelSequenceEditorToolkit*>(AssetEditor);
 		TWeakPtr<ISequencer> WeakSequencer = LevelSequenceEditor ? LevelSequenceEditor->GetSequencer() : nullptr;
 
+
+
 		if (WeakSequencer.IsValid())
 		{
+
 			ASkeletalMeshActor* MeshActor = World->SpawnActor<ASkeletalMeshActor>(ASkeletalMeshActor::StaticClass(), FTransform::Identity);
 			MeshActor->SetActorLabel(AnimSequence->GetName());
 
