Control Rig: Make debug enabled flag non-constant

Only meta data change

diff --git a/Engine/Plugins/Animation/ControlRig/Source/ControlRig/Private/Units/Debug/RigUnit_DebugHierarchy.h b/Engine/Plugins/Animation/ControlRig/Source/ControlRig/Private/Units/Debug/RigUnit_DebugHierarchy.h
--- a/Engine/Plugins/Animation/ControlRig/Source/ControlRig/Private/Units/Debug/RigUnit_DebugHierarchy.h
+++ b/Engine/Plugins/Animation/ControlRig/Source/ControlRig/Private/Units/Debug/RigUnit_DebugHierarchy.h
@@ -38,7 +38,7 @@ struct CONTROLRIG_API FRigUnit_DebugHierarchy : public FRigUnit_DebugBaseMutable
 	UPROPERTY(meta = (Input))
 	FTransform WorldOffset;
 
-	UPROPERTY(meta = (Input, Constant))
+	UPROPERTY(meta = (Input))
 	bool bEnabled;
 };
 
@@ -77,6 +77,6 @@ struct CONTROLRIG_API FRigUnit_DebugPose : public FRigUnit_DebugBaseMutable
 	UPROPERTY(meta = (Input))
 	FTransform WorldOffset;
 
-	UPROPERTY(meta = (Input, Constant))
+	UPROPERTY(meta = (Input))
 	bool bEnabled;
 };
\ No newline at end of file
diff --git a/Engine/Plugins/Animation/ControlRig/Source/ControlRig/Private/Units/Debug/RigUnit_DebugLine.h b/Engine/Plugins/Animation/ControlRig/Source/ControlRig/Private/Units/Debug/RigUnit_DebugLine.h
--- a/Engine/Plugins/Animation/ControlRig/Source/ControlRig/Private/Units/Debug/RigUnit_DebugLine.h
+++ b/Engine/Plugins/Animation/ControlRig/Source/ControlRig/Private/Units/Debug/RigUnit_DebugLine.h
@@ -86,6 +86,6 @@ struct CONTROLRIG_API FRigUnit_DebugLineItemSpace : public FRigUnit_DebugBaseMut
 	UPROPERTY(meta = (Input))
 	FTransform WorldOffset;
 
-	UPROPERTY(meta = (Input, Constant))
+	UPROPERTY(meta = (Input))
 	bool bEnabled;
 };
\ No newline at end of file
diff --git a/Engine/Plugins/Animation/ControlRig/Source/ControlRig/Private/Units/Debug/RigUnit_DebugLineStrip.h b/Engine/Plugins/Animation/ControlRig/Source/ControlRig/Private/Units/Debug/RigUnit_DebugLineStrip.h
--- a/Engine/Plugins/Animation/ControlRig/Source/ControlRig/Private/Units/Debug/RigUnit_DebugLineStrip.h
+++ b/Engine/Plugins/Animation/ControlRig/Source/ControlRig/Private/Units/Debug/RigUnit_DebugLineStrip.h
@@ -78,6 +78,6 @@ struct CONTROLRIG_API FRigUnit_DebugLineStripItemSpace : public FRigUnit_DebugBa
 	UPROPERTY(meta = (Input))
 	FTransform WorldOffset;
 	
-	UPROPERTY(meta = (Input, Constant))
+	UPROPERTY(meta = (Input))
 	bool bEnabled;
 };
diff --git a/Engine/Plugins/Animation/ControlRig/Source/ControlRig/Private/Units/Debug/RigUnit_DebugPrimitives.h b/Engine/Plugins/Animation/ControlRig/Source/ControlRig/Private/Units/Debug/RigUnit_DebugPrimitives.h
--- a/Engine/Plugins/Animation/ControlRig/Source/ControlRig/Private/Units/Debug/RigUnit_DebugPrimitives.h
+++ b/Engine/Plugins/Animation/ControlRig/Source/ControlRig/Private/Units/Debug/RigUnit_DebugPrimitives.h
@@ -88,7 +88,7 @@ struct CONTROLRIG_API FRigUnit_DebugRectangleItemSpace : public FRigUnit_DebugBa
 	UPROPERTY(meta = (Input))
 	FTransform WorldOffset;
 
-	UPROPERTY(meta = (Input, Constant))
+	UPROPERTY(meta = (Input))
 	bool bEnabled;
 };
 
@@ -198,6 +198,6 @@ struct CONTROLRIG_API FRigUnit_DebugArcItemSpace : public FRigUnit_DebugBaseMuta
 	UPROPERTY(meta = (Input))
 	FTransform WorldOffset;
 
-	UPROPERTY(meta = (Input, Constant))
+	UPROPERTY(meta = (Input))
 	bool bEnabled;
 };
\ No newline at end of file
diff --git a/Engine/Plugins/Animation/ControlRig/Source/ControlRig/Private/Units/Debug/RigUnit_DebugTransform.h b/Engine/Plugins/Animation/ControlRig/Source/ControlRig/Private/Units/Debug/RigUnit_DebugTransform.h
--- a/Engine/Plugins/Animation/ControlRig/Source/ControlRig/Private/Units/Debug/RigUnit_DebugTransform.h
+++ b/Engine/Plugins/Animation/ControlRig/Source/ControlRig/Private/Units/Debug/RigUnit_DebugTransform.h
@@ -156,7 +156,7 @@ struct CONTROLRIG_API FRigUnit_DebugTransformMutableItemSpace : public FRigUnit_
 	UPROPERTY(meta = (Input))
 	FTransform WorldOffset;
 
-	UPROPERTY(meta = (Input, Constant))
+	UPROPERTY(meta = (Input))
 	bool bEnabled;
 };
 
@@ -264,6 +264,6 @@ struct CONTROLRIG_API FRigUnit_DebugTransformArrayMutableItemSpace : public FRig
 	UPROPERTY(meta = (Input))
 	FTransform WorldOffset;
 
-	UPROPERTY(meta = (Input, Constant))
+	UPROPERTY(meta = (Input))
 	bool bEnabled;
 };
