Control Rig: Show Node Instruction Index on Node

diff --git a/Engine/Plugins/Animation/ControlRig/Source/ControlRigDeveloper/Public/ControlRigBlueprint.h b/Engine/Plugins/Animation/ControlRig/Source/ControlRigDeveloper/Public/ControlRigBlueprint.h
--- a/Engine/Plugins/Animation/ControlRig/Source/ControlRigDeveloper/Public/ControlRigBlueprint.h
+++ b/Engine/Plugins/Animation/ControlRig/Source/ControlRigDeveloper/Public/ControlRigBlueprint.h
@@ -112,7 +112,8 @@ struct CONTROLRIGDEVELOPER_API FRigGraphDisplaySettings
 	GENERATED_BODY();
 
 	FRigGraphDisplaySettings()
-		: bShowNodeRunCounts(false)
+		: bShowNodeInstructionIndex(false)
+		, bShowNodeRunCounts(false)
 		, NodeRunLowerBound(1)
 		, NodeRunLimit(64)
 		, MinMicroSeconds(0.0)
@@ -126,6 +127,11 @@ struct CONTROLRIGDEVELOPER_API FRigGraphDisplaySettings
 	{
 	}
 
+	// When enabled shows the first node instruction index
+	// matching the execution stack window.
+	UPROPERTY(EditAnywhere, Category = "Graph Display Settings")
+	bool bShowNodeInstructionIndex;
+
 	// When enabled shows the node counts both in the graph view as
 	// we as in the execution stack window.
 	// The number on each node represents how often the node has been run.
diff --git a/Engine/Plugins/Animation/ControlRig/Source/ControlRigEditor/Private/Graph/SControlRigGraphNode.cpp b/Engine/Plugins/Animation/ControlRig/Source/ControlRigEditor/Private/Graph/SControlRigGraphNode.cpp
--- a/Engine/Plugins/Animation/ControlRig/Source/ControlRigEditor/Private/Graph/SControlRigGraphNode.cpp
+++ b/Engine/Plugins/Animation/ControlRig/Source/ControlRigEditor/Private/Graph/SControlRigGraphNode.cpp
@@ -1115,17 +1115,18 @@ TArray<FOverlayWidgetInfo> SControlRigGraphNode::GetOverlayWidgets(bool bSelecte
 
 		if(Blueprint.IsValid())
 		{
+			const bool bShowInstructionIndex = Blueprint->RigGraphDisplaySettings.bShowNodeInstructionIndex;
 			const bool bShowNodeCounts = Blueprint->RigGraphDisplaySettings.bShowNodeRunCounts;
 			const bool bEnableProfiling = Blueprint->VMRuntimeSettings.bEnableProfiling;
 			
-			if(bShowNodeCounts || bEnableProfiling)
+			if(bShowNodeCounts || bShowInstructionIndex || bEnableProfiling)
 			{
 				if(UControlRig* DebuggedControlRig = Cast<UControlRig>(Blueprint->GetObjectBeingDebugged()))
 				{
-					if(bShowNodeCounts)
+					if(bShowNodeCounts || bShowInstructionIndex)
 					{
 						const int32 Count = ModelNode->GetInstructionVisitedCount(DebuggedControlRig->GetVM(), FRigVMASTProxy());
-						if(Count > Blueprint->RigGraphDisplaySettings.NodeRunLowerBound)
+						if((Count > Blueprint->RigGraphDisplaySettings.NodeRunLowerBound) || bShowInstructionIndex)
 						{
 							const int32 VOffset = bSelected ? -2 : 2;
 							const FVector2D TextSize = InstructionCountTextBlockWidget->GetDesiredSize();
@@ -1199,16 +1200,54 @@ FText SControlRigGraphNode::GetInstructionCountText() const
 {
 	if(Blueprint.IsValid())
 	{
-		if(Blueprint->RigGraphDisplaySettings.bShowNodeRunCounts)
+		bool bShowInstructionIndex = Blueprint->RigGraphDisplaySettings.bShowNodeInstructionIndex;
+		bool bShowNodeRunCount = Blueprint->RigGraphDisplaySettings.bShowNodeRunCounts;
+		if(bShowInstructionIndex || bShowNodeRunCount)
 		{
 			if (ModelNode.IsValid())
 			{
 				if(UControlRig* DebuggedControlRig = Cast<UControlRig>(Blueprint->GetObjectBeingDebugged()))
 				{
-					const int32 Count = ModelNode->GetInstructionVisitedCount(DebuggedControlRig->GetVM(), FRigVMASTProxy());
-					if(Count > Blueprint->RigGraphDisplaySettings.NodeRunLowerBound)
+					int32 RunCount = 0;
+					int32 FirstInstructionIndex = INDEX_NONE;
+					if(bShowNodeRunCount)
 					{
-						return FText::FromString(FString::FromInt(Count));
+						RunCount = ModelNode->GetInstructionVisitedCount(DebuggedControlRig->GetVM(), FRigVMASTProxy());
+						bShowNodeRunCount = RunCount > Blueprint->RigGraphDisplaySettings.NodeRunLowerBound;
+					}
+
+					if(bShowInstructionIndex)
+					{
+						const TArray<int32> Instructions = ModelNode->GetInstructionsForVM(DebuggedControlRig->GetVM());
+						bShowInstructionIndex = Instructions.Num() > 0;
+						if(bShowInstructionIndex)
+						{
+							FirstInstructionIndex = Instructions[0];
+						}
+					}
+
+					if(bShowInstructionIndex || bShowNodeRunCount)
+					{
+						FText NodeRunCountText, NodeInstructionIndexText;
+						if(bShowNodeRunCount)
+						{
+							NodeRunCountText = FText::FromString(FString::FromInt(RunCount));
+							if(!bShowInstructionIndex)
+							{
+								return NodeRunCountText;
+							}
+						}
+
+						if(bShowInstructionIndex)
+						{
+							NodeInstructionIndexText = FText::FromString(FString::FromInt(FirstInstructionIndex));
+							if(!bShowNodeRunCount)
+							{
+								return NodeInstructionIndexText;
+							}
+						}
+
+						return FText::Format(LOCTEXT("SControlRigGraphNodeCombinedNodeCountText", "{0}: {1}"), NodeInstructionIndexText, NodeRunCountText);
 					}
 				}
 			}
