Sequencer: Control Rig: Tween tools works with multiple control rigs.
#trivial

diff --git a/Engine/Plugins/Animation/ControlRig/Source/ControlRigEditor/Private/EditMode/SControlRigTweenWidget.cpp b/Engine/Plugins/Animation/ControlRig/Source/ControlRigEditor/Private/EditMode/SControlRigTweenWidget.cpp
--- a/Engine/Plugins/Animation/ControlRig/Source/ControlRigEditor/Private/EditMode/SControlRigTweenWidget.cpp
+++ b/Engine/Plugins/Animation/ControlRig/Source/ControlRigEditor/Private/EditMode/SControlRigTweenWidget.cpp
@@ -74,20 +74,13 @@ void SControlRigTweenWidget::Construct(const FArguments& InArgs)
 void SControlRigTweenWidget::OnPoseBlendChanged(float ChangedVal)
 {
 	TArray<UControlRig*> ControlRigs = GetControlRigs();
-	bool bWeBlended = false;
-	for (UControlRig* ControlRig : ControlRigs)
-	{
-		if (ControlRig && WeakSequencer.IsValid() && bIsBlending)
-		{
-			PoseBlendValue = ChangedVal;
-			ControlsToTween.Blend(WeakSequencer, ChangedVal);
-			bWeBlended = true;
-		}
-	}
-	if(bWeBlended)
+	if (ControlRigs.Num() > 0 && WeakSequencer.IsValid() && bIsBlending)
 	{
+		PoseBlendValue = ChangedVal;
+		ControlsToTween.Blend(WeakSequencer, ChangedVal);
 		WeakSequencer.Pin()->NotifyMovieSceneDataChanged(EMovieSceneDataChangeType::TrackValueChanged);
 	}
+
 }
 
 void SControlRigTweenWidget::OnBeginSliderMovement()
@@ -108,15 +101,10 @@ void SControlRigTweenWidget::SetupControls()
 	ILevelSequenceEditorToolkit* LevelSequenceEditor = static_cast<ILevelSequenceEditorToolkit*>(AssetEditor);
 	WeakSequencer = LevelSequenceEditor ? LevelSequenceEditor->GetSequencer() : nullptr;
 	TArray<UControlRig*> ControlRigs = GetControlRigs();
-	for (UControlRig* ControlRig : ControlRigs)
+	if (ControlRigs.Num() > 0)
 	{
-		if (ControlRig && WeakSequencer.IsValid() && bIsBlending)
-		{
-			WeakSequencer.Pin()->GetFocusedMovieSceneSequence()->GetMovieScene()->Modify();
-			TArray<UControlRig*> SelectedControlRigs;
-			SelectedControlRigs.Add(ControlRig);
-			ControlsToTween.Setup(SelectedControlRigs, WeakSequencer);
-		}
+		WeakSequencer.Pin()->GetFocusedMovieSceneSequence()->GetMovieScene()->Modify();
+		ControlsToTween.Setup(ControlRigs, WeakSequencer);
 	}
 }
 
@@ -186,7 +174,12 @@ TArray<UControlRig*>SControlRigTweenWidget::GetControlRigs()
 	TArray<UControlRig*> ControlRigs;
 	if (FControlRigEditMode* EditMode = static_cast<FControlRigEditMode*>(GLevelEditorModeTools().GetActiveMode(FControlRigEditMode::ModeName)))
 	{
-		ControlRigs = EditMode->GetControlRigsArray(true /*bIsVisible*/);
+		TMap<UControlRig*, TArray<FRigElementKey>> SelectedControls;
+		EditMode->GetAllSelectedControls(SelectedControls);
+		for (TPair<UControlRig*, TArray<FRigElementKey>>& Selected: SelectedControls)
+		{
+			ControlRigs.Add(Selected.Key);
+		}
 	}
 	return ControlRigs;
 }
diff --git a/Engine/Plugins/Animation/ControlRig/Source/ControlRigEditor/Private/Tools/ControlRigTweener.cpp b/Engine/Plugins/Animation/ControlRig/Source/ControlRigEditor/Private/Tools/ControlRigTweener.cpp
--- a/Engine/Plugins/Animation/ControlRig/Source/ControlRigEditor/Private/Tools/ControlRigTweener.cpp
+++ b/Engine/Plugins/Animation/ControlRig/Source/ControlRigEditor/Private/Tools/ControlRigTweener.cpp
@@ -102,7 +102,10 @@ void FControlsToTween::Setup(const TArray<UControlRig*>& SelectedControlRigs, TW
 							}
 							if (ControlRigChannels.NumChannels > 0)
 							{
-								ControlRigChannelsMap.Add(ControlElement->GetName(), ControlRigChannels);
+								TPair<UControlRig*, FName> Pair;
+								Pair.Key = ControlRig;
+								Pair.Value = ControlElement->GetName();
+								ControlRigChannelsMap.Add(Pair, ControlRigChannels);
 							}
 						}
 					}
@@ -124,7 +127,7 @@ void FControlsToTween::Blend(TWeakPtr<ISequencer>& InSequencer, float BlendValue
 	Sequencer->GetFocusedMovieSceneSequence()->GetMovieScene()->Modify();
 	//BlendValue of 0.0 = half, 1.0f = Second, -1.0f = FirtValue;
 	//NormalizedBlendValue goes from 0 to 1.0f
-	for (const TPair<FName, FControlRigChannels>& Pair : ControlRigChannelsMap)
+	for (const TPair<TPair<UControlRig*,FName>, FControlRigChannels>& Pair : ControlRigChannelsMap)
 	{
 		for (int Index = 0; Index < 9; ++Index)
 		{
diff --git a/Engine/Plugins/Animation/ControlRig/Source/ControlRigEditor/Public/Tools/ControlRigTweener.h b/Engine/Plugins/Animation/ControlRig/Source/ControlRigEditor/Public/Tools/ControlRigTweener.h
--- a/Engine/Plugins/Animation/ControlRig/Source/ControlRigEditor/Public/Tools/ControlRigTweener.h
+++ b/Engine/Plugins/Animation/ControlRig/Source/ControlRigEditor/Public/Tools/ControlRigTweener.h
@@ -55,7 +55,7 @@ private:
 	void SetupControlRigChannel(FFrameNumber CurrentFrame, TArray<FFrameNumber>& KeyTimes, TArray<FKeyHandle>& Handles, FMovieSceneFloatChannel* FloatChannel,
 		FChannelKeyBounds& KeyBounds);
 
-	TMap<FName, FControlRigChannels>  ControlRigChannelsMap;
+	TMap<TPair<UControlRig*,FName>, FControlRigChannels>  ControlRigChannelsMap;
 };
 
 
