Removed a non desired export of the symbols for TGuardValue<bool,bool>

#prefligfht 62389b0b9304d833f4e5a064

diff --git a/Engine/Plugins/Animation/ControlRig/Source/ControlRig/Public/Rigs/RigHierarchy.h b/Engine/Plugins/Animation/ControlRig/Source/ControlRig/Public/Rigs/RigHierarchy.h
--- a/Engine/Plugins/Animation/ControlRig/Source/ControlRig/Public/Rigs/RigHierarchy.h
+++ b/Engine/Plugins/Animation/ControlRig/Source/ControlRig/Public/Rigs/RigHierarchy.h
@@ -3036,11 +3036,11 @@ private:
 	URigHierarchy* Hierarchy;
 };
 
-struct CONTROLRIG_API FRigHierarchyEnableControllerBracket : public TGuardValue<bool>
+struct CONTROLRIG_API FRigHierarchyEnableControllerBracket
 {
 private:
 	FRigHierarchyEnableControllerBracket(URigHierarchy* InHierarchy, bool bEnable)
-		: TGuardValue<bool>(InHierarchy->bIsControllerAvailable, bEnable)
+		: GuardIsControllerAvailable(InHierarchy->bIsControllerAvailable, bEnable)
 	{
 	}
 
@@ -3049,6 +3049,9 @@ private:
 
 	// certain units are allowed to use this
 	friend struct FRigUnit_AddParent;
+
+private:
+	TGuardValue<bool> GuardIsControllerAvailable;
 };
 
 struct CONTROLRIG_API FRigHierarchyExecuteContextBracket
