Control Rig: Add boolean toggle nodes

diff --git a/Engine/Plugins/Animation/ControlRig/Source/ControlRig/Private/Units/Math/RigUnit_MathBool.cpp b/Engine/Plugins/Animation/ControlRig/Source/ControlRig/Private/Units/Math/RigUnit_MathBool.cpp
--- a/Engine/Plugins/Animation/ControlRig/Source/ControlRig/Private/Units/Math/RigUnit_MathBool.cpp
+++ b/Engine/Plugins/Animation/ControlRig/Source/ControlRig/Private/Units/Math/RigUnit_MathBool.cpp
@@ -74,3 +74,59 @@ FRigUnit_MathBoolToggled_Execute()
 	LastValue = Value;
 }
 
+FRigUnit_MathBoolFlipFlop_Execute()
+{
+	DECLARE_SCOPE_HIERARCHICAL_COUNTER_RIGUNIT()
+
+	if(Context.State == EControlRigState::Init)
+	{
+		Result = LastValue = !StartValue;
+		TimeLeft = 0.f;
+		return;
+	}
+
+	TimeLeft = TimeLeft - Context.DeltaTime;
+	if (TimeLeft <= 0.f)
+	{
+		LastValue = !LastValue;
+		TimeLeft = Duration;
+	}
+
+	Result = LastValue;
+}
+
+FRigUnit_MathBoolOnce_Execute()
+{
+	DECLARE_SCOPE_HIERARCHICAL_COUNTER_RIGUNIT()
+
+	if(Context.State == EControlRigState::Init)
+	{
+		Result = false;
+		LastValue = true;
+		TimeLeft = 0.f;
+		return;
+	}
+
+	if(!Result && LastValue)
+	{
+		TimeLeft = Duration;
+	}
+
+	Result = LastValue;
+
+	TimeLeft -= Context.DeltaTime;
+	if(TimeLeft <= 0.f)
+	{
+		LastValue = false;
+	}
+}
+
+FRigUnit_MathBoolToFloat_Execute()
+{
+	Result = Value ? 1.f : 0.f;
+}
+
+FRigUnit_MathBoolToInteger_Execute()
+{
+	Result = Value ? 1 : 0;
+}
\ No newline at end of file
diff --git a/Engine/Plugins/Animation/ControlRig/Source/ControlRig/Private/Units/Math/RigUnit_MathBool.h b/Engine/Plugins/Animation/ControlRig/Source/ControlRig/Private/Units/Math/RigUnit_MathBool.h
--- a/Engine/Plugins/Animation/ControlRig/Source/ControlRig/Private/Units/Math/RigUnit_MathBool.h
+++ b/Engine/Plugins/Animation/ControlRig/Source/ControlRig/Private/Units/Math/RigUnit_MathBool.h
@@ -223,3 +223,119 @@ struct CONTROLRIG_API FRigUnit_MathBoolToggled : public FRigUnit_MathBoolBase
 	UPROPERTY()
 	bool LastValue;
 };
+
+/**
+ * Returns true and false as a sequence.
+ */
+USTRUCT(meta=(DisplayName="FlipFlop", Keywords="Toggle,Changed,Different", Varying))
+struct CONTROLRIG_API FRigUnit_MathBoolFlipFlop : public FRigUnit_MathBoolBase
+{
+	GENERATED_BODY()
+	RIGVM_METHOD()
+	virtual void Execute(const FRigUnitContext& Context) override;
+
+	FRigUnit_MathBoolFlipFlop()
+	{
+		Duration = TimeLeft = 0.f;
+		Result = LastValue = StartValue = false;
+	}
+
+	// The initial value to use for the flag
+	UPROPERTY(meta=(Visible))
+	bool StartValue;
+
+	/**
+	 * The duration in seconds at which the result won't change.
+     * Use 0 for a different result every time.
+	 */
+	UPROPERTY(meta = (Input))
+	float Duration;
+
+	UPROPERTY(meta=(Output))
+	bool Result;
+
+	UPROPERTY()
+	bool LastValue;
+
+	UPROPERTY()
+	float TimeLeft;
+};
+
+/**
+ * Returns true once the first time this node is hit
+ */
+USTRUCT(meta=(DisplayName="Once", Keywords="FlipFlop,Toggle,Changed,Different", Varying))
+struct CONTROLRIG_API FRigUnit_MathBoolOnce : public FRigUnit_MathBoolBase
+{
+	GENERATED_BODY()
+	RIGVM_METHOD()
+	virtual void Execute(const FRigUnitContext& Context) override;
+
+	FRigUnit_MathBoolOnce()
+	{
+		Duration = TimeLeft = 0.f;
+		Result = LastValue = false;
+	}
+
+	/**
+	 * The duration in seconds at which the result is true
+	 * Use 0 for a different result every time.
+	 */
+	UPROPERTY(meta = (Visible))
+	float Duration;
+
+	UPROPERTY(meta=(Output))
+	bool Result;
+
+	UPROPERTY()
+	bool LastValue;
+
+	UPROPERTY()
+	float TimeLeft;
+};
+
+/**
+ * Turns the given bool into a float value
+ */
+USTRUCT(meta=(DisplayName="To Float", TemplateName="Cast"))
+struct CONTROLRIG_API FRigUnit_MathBoolToFloat : public FRigUnit_MathBoolBase
+{
+	GENERATED_BODY()
+	RIGVM_METHOD()
+	virtual void Execute(const FRigUnitContext& Context) override;
+
+	FRigUnit_MathBoolToFloat()
+	{
+		Value = false;
+		Result = 0.f;
+	}
+
+	UPROPERTY(meta = (Input))
+	bool Value;
+
+	UPROPERTY(meta= (Output))
+	float Result;
+};
+
+/**
+ * Turns the given bool into an integer value
+ */
+USTRUCT(meta=(DisplayName="To Integer", TemplateName="Cast"))
+struct CONTROLRIG_API FRigUnit_MathBoolToInteger : public FRigUnit_MathBoolBase
+{
+	GENERATED_BODY()
+	RIGVM_METHOD()
+	virtual void Execute(const FRigUnitContext& Context) override;
+
+	FRigUnit_MathBoolToInteger()
+	{
+		Value = false;
+		Result = 0;
+	}
+
+	UPROPERTY(meta = (Input))
+	bool Value;
+
+	UPROPERTY(meta= (Output))
+	int32 Result;
+};
diff --git a/Engine/Plugins/Animation/ControlRig/Source/ControlRig/Private/Units/Simulation/RigUnit_Timeline.cpp b/Engine/Plugins/Animation/ControlRig/Source/ControlRig/Private/Units/Simulation/RigUnit_Timeline.cpp
--- a/Engine/Plugins/Animation/ControlRig/Source/ControlRig/Private/Units/Simulation/RigUnit_Timeline.cpp
+++ b/Engine/Plugins/Animation/ControlRig/Source/ControlRig/Private/Units/Simulation/RigUnit_Timeline.cpp
@@ -35,4 +35,45 @@ IMPLEMENT_RIGUNIT_AUTOMATION_TEST(FRigUnit_Timeline)
 	return true;
 }
 
-#endif
\ No newline at end of file
+#endif
+
+FRigUnit_TimeLoop_Execute()
+{
+	DECLARE_SCOPE_HIERARCHICAL_COUNTER_RIGUNIT()
+	if (Context.State == EControlRigState::Init)
+	{
+		Absolute = Relative = FlipFlop = AccumulatedAbsolute = AccumulatedRelative = 0.f;
+		NumIterations = 0;
+		Even = false;
+		return;
+	}
+
+	const float DurationClamped = FMath::Max(Duration, 0.0001f);
+	const float Increment = Context.DeltaTime * Speed;
+	Absolute = AccumulatedAbsolute = AccumulatedAbsolute + Increment;
+
+	AccumulatedRelative = AccumulatedRelative + Increment;
+	while(AccumulatedRelative > DurationClamped)
+	{
+		AccumulatedRelative -= DurationClamped;
+		NumIterations++;
+	}
+
+	Relative = AccumulatedRelative;
+
+	Even = (NumIterations & 1) == 0;
+	if(Even) // check is this is even or odd
+	{
+		FlipFlop = Relative;
+	}
+	else
+	{
+		FlipFlop = Duration - Relative;
+	}
+
+	if(Normalize)
+	{
+		Relative = Relative / DurationClamped;
+		FlipFlop = FlipFlop / DurationClamped;
+	}
+}
\ No newline at end of file
diff --git a/Engine/Plugins/Animation/ControlRig/Source/ControlRig/Private/Units/Simulation/RigUnit_Timeline.h b/Engine/Plugins/Animation/ControlRig/Source/ControlRig/Private/Units/Simulation/RigUnit_Timeline.h
--- a/Engine/Plugins/Animation/ControlRig/Source/ControlRig/Private/Units/Simulation/RigUnit_Timeline.h
+++ b/Engine/Plugins/Animation/ControlRig/Source/ControlRig/Private/Units/Simulation/RigUnit_Timeline.h
@@ -31,3 +31,64 @@ struct CONTROLRIG_API FRigUnit_Timeline : public FRigUnit_SimBase
 	UPROPERTY()
 	float AccumulatedValue;
 };
+
+/**
+ * Simulates a time value - and outputs loop information
+ */
+USTRUCT(meta=(DisplayName="Time Loop", Category = "Simulation|Accumulate", Keywords="Playback,Pause,Timeline"))
+struct CONTROLRIG_API FRigUnit_TimeLoop : public FRigUnit_SimBase
+{
+	GENERATED_BODY()
+	
+	FRigUnit_TimeLoop()
+	{
+		Speed = 1.f;
+		Duration = 1.f;
+		Normalize = false;
+		Absolute = Relative = FlipFlop = AccumulatedAbsolute = AccumulatedRelative = 0.f;
+		NumIterations = 0;
+		Even = false;
+	}
+
+	RIGVM_METHOD()
+	virtual void Execute(const FRigUnitContext& Context) override;
+
+	UPROPERTY(meta = (Input))
+	float Speed;
+
+	// the duration of a single loop in seconds
+	UPROPERTY(meta = (Input))
+	float Duration;
+
+	// if set to true the output relative and flipflop
+	// will be normalized over the duration.
+	UPROPERTY(meta = (Input))
+	bool Normalize;
+
+	// the overall time in seconds
+	UPROPERTY(meta=(Output))
+	float Absolute;
+
+	// the relative time in seconds (within the loop)
+	UPROPERTY(meta=(Output))
+	float Relative;
+
+	// the relative time in seconds (within the loop),
+	// going from 0 to duration and then back from duration to 0,
+	// or 0 to 1 and 1 to 0 if Normalize is turned on
+	UPROPERTY(meta=(Output))
+	float FlipFlop;
+
+	// true if the iteration of the loop is even
+	UPROPERTY(meta=(Output))
+	bool Even;
+
+	UPROPERTY()
+	float AccumulatedAbsolute;
+
+	UPROPERTY()
+	float AccumulatedRelative;
+
+	UPROPERTY()
+	int32 NumIterations;
+};
\ No newline at end of file
