Control Rig: Improve access to interacting elements

diff --git a/Engine/Plugins/Animation/ControlRig/Source/ControlRig/Private/ControlRig.cpp b/Engine/Plugins/Animation/ControlRig/Source/ControlRig/Private/ControlRig.cpp
--- a/Engine/Plugins/Animation/ControlRig/Source/ControlRig/Private/ControlRig.cpp
+++ b/Engine/Plugins/Animation/ControlRig/Source/ControlRig/Private/ControlRig.cpp
@@ -76,6 +76,7 @@ UControlRig::UControlRig(const FObjectInitializer& ObjectInitializer)
 	, PostSetupBracket(0)
 	, InteractionBracket(0)
 	, InterRigSyncBracket(0)
+	, InteractionType((uint8)EControlRigInteractionType::None)
 #if WITH_EDITORONLY_DATA
 	, VMSnapshotBeforeExecution(nullptr)
 #endif
@@ -661,7 +662,8 @@ void UControlRig::Execute(const EControlRigState InState, const FName& InEventNa
 	Context.DeltaTime = DeltaTime;
 	Context.AbsoluteTime = AbsoluteTime;
 	Context.FramesPerSecond = GetCurrentFramesPerSecond();
-	Context.bDuringInteraction = IsInteracting();
+	Context.InteractionType = InteractionType;
+	Context.ElementsBeingInteracted = ElementsBeingInteracted;
 	Context.State = InState;
 	
 	Context.Hierarchy = GetHierarchy();
diff --git a/Engine/Plugins/Animation/ControlRig/Source/ControlRig/Private/Units/Execution/RigUnit_IsInteracting.cpp b/Engine/Plugins/Animation/ControlRig/Source/ControlRig/Private/Units/Execution/RigUnit_IsInteracting.cpp
--- a/Engine/Plugins/Animation/ControlRig/Source/ControlRig/Private/Units/Execution/RigUnit_IsInteracting.cpp
+++ b/Engine/Plugins/Animation/ControlRig/Source/ControlRig/Private/Units/Execution/RigUnit_IsInteracting.cpp
@@ -7,5 +7,9 @@ FRigUnit_IsInteracting_Execute()
 {
     DECLARE_SCOPE_HIERARCHICAL_COUNTER_RIGUNIT()
 
-	bIsInteracting = Context.bDuringInteraction;
+	bIsInteracting = Context.IsInteracting();
+	bIsTranslating = ( Context.InteractionType & (uint8)EControlRigInteractionType::Translate) != 0;
+	bIsRotating = ( Context.InteractionType & (uint8)EControlRigInteractionType::Rotate) != 0;
+	bIsScaling = ( Context.InteractionType & (uint8)EControlRigInteractionType::Scale) != 0;
+	Items = Context.ElementsBeingInteracted;
 }
diff --git a/Engine/Plugins/Animation/ControlRig/Source/ControlRig/Private/Units/Execution/RigUnit_IsInteracting.h b/Engine/Plugins/Animation/ControlRig/Source/ControlRig/Private/Units/Execution/RigUnit_IsInteracting.h
--- a/Engine/Plugins/Animation/ControlRig/Source/ControlRig/Private/Units/Execution/RigUnit_IsInteracting.h
+++ b/Engine/Plugins/Animation/ControlRig/Source/ControlRig/Private/Units/Execution/RigUnit_IsInteracting.h
@@ -9,7 +9,7 @@
 /**
  * Returns true if the Control Rig is being interacted
  */
-USTRUCT(meta=(DisplayName="Is Interacting", Category="Execution", TitleColor="1 0 0", NodeColor="1 1 1", Keywords="Gizmo,Manipulation,Interaction", Varying))
+USTRUCT(meta=(DisplayName="Is Interacting", Category="Execution", TitleColor="1 0 0", NodeColor="1 1 1", Keywords="IsInteracting,Gizmo,Manipulation,Interaction,Tracking", Varying))
 struct CONTROLRIG_API FRigUnit_IsInteracting : public FRigUnit
 {
 	GENERATED_BODY()
@@ -17,7 +17,23 @@ struct CONTROLRIG_API FRigUnit_IsInteracting : public FRigUnit
 	RIGVM_METHOD()
 	virtual void Execute(const FRigUnitContext& Context) override;
 
-	// The execution result
+	// True if there is currently an interaction happening
 	UPROPERTY(EditAnywhere, Transient, DisplayName = "Interacting", Category = "Execution", meta = (Output))
 	bool bIsInteracting = false;
+
+	// True if the current interaction is a translation
+	UPROPERTY(EditAnywhere, Transient, Category = "Execution", meta = (Output))
+	bool bIsTranslating = false;
+
+	// True if the current interaction is a rotation
+	UPROPERTY(EditAnywhere, Transient, Category = "Execution", meta = (Output))
+	bool bIsRotating = false;
+
+	// True if the current interaction is scaling
+	UPROPERTY(EditAnywhere, Transient, Category = "Execution", meta = (Output))
+	bool bIsScaling = false;
+
+	// The items being interacted on
+	UPROPERTY(EditAnywhere, Transient, Category = "Execution", meta = (Output))
+	TArray<FRigElementKey> Items;
 };
diff --git a/Engine/Plugins/Animation/ControlRig/Source/ControlRig/Private/Units/Hierarchy/RigUnit_SendEvent.cpp b/Engine/Plugins/Animation/ControlRig/Source/ControlRig/Private/Units/Hierarchy/RigUnit_SendEvent.cpp
--- a/Engine/Plugins/Animation/ControlRig/Source/ControlRig/Private/Units/Hierarchy/RigUnit_SendEvent.cpp
+++ b/Engine/Plugins/Animation/ControlRig/Source/ControlRig/Private/Units/Hierarchy/RigUnit_SendEvent.cpp
@@ -12,7 +12,7 @@ FRigUnit_SendEvent_Execute()
     	return;
     }
 
-	if (bOnlyDuringInteraction && !Context.bDuringInteraction)
+	if (bOnlyDuringInteraction && !Context.IsInteracting())
 	{
 		return;
 	}
diff --git a/Engine/Plugins/Animation/ControlRig/Source/ControlRig/Public/ControlRig.h b/Engine/Plugins/Animation/ControlRig/Source/ControlRig/Public/ControlRig.h
--- a/Engine/Plugins/Animation/ControlRig/Source/ControlRig/Public/ControlRig.h
+++ b/Engine/Plugins/Animation/ControlRig/Source/ControlRig/Public/ControlRig.h
@@ -658,6 +658,8 @@ protected:
 	int32 PostForwardsSolveBracket;
 	int32 InteractionBracket;
 	int32 InterRigSyncBracket;
+	uint8 InteractionType;
+	TArray<FRigElementKey> ElementsBeingInteracted;
 
 	TWeakObjectPtr<USceneComponent> OuterSceneComponent;
 
@@ -686,6 +688,11 @@ protected:
 		return InteractionBracket > 0;
 	}
 
+	FORCEINLINE uint8 GetInteractionType() const
+	{
+		return InteractionType;
+	}
+
 	FORCEINLINE bool IsSyncingWithOtherRig() const
 	{
 		return InterRigSyncBracket > 0;
@@ -893,6 +900,8 @@ public:
 		if(ensure(ControlRig.IsValid()))
 		{
 			ControlRig->GetHierarchy()->EndInteraction();
+			ControlRig->InteractionType = (uint8)EControlRigInteractionType::None;
+			ControlRig->ElementsBeingInteracted.Reset();
 		}
 	}
 
diff --git a/Engine/Plugins/Animation/ControlRig/Source/ControlRig/Public/ControlRigGizmoActor.h b/Engine/Plugins/Animation/ControlRig/Source/ControlRig/Public/ControlRigGizmoActor.h
--- a/Engine/Plugins/Animation/ControlRig/Source/ControlRig/Public/ControlRigGizmoActor.h
+++ b/Engine/Plugins/Animation/ControlRig/Source/ControlRig/Public/ControlRigGizmoActor.h
@@ -117,6 +117,9 @@ public:
 	/** Get whether the control is hovered */
 	virtual bool IsHovered() const;
 
+	/* Returns the key of the control element this shape actor represents */
+	FRigElementKey GetElementKey() const { return FRigElementKey(ControlName, ERigElementType::Control); }
+
 	/** Called from the edit mode each tick */
 	virtual void TickControl() {};
 
diff --git a/Engine/Plugins/Animation/ControlRig/Source/ControlRig/Public/Units/RigUnitContext.h b/Engine/Plugins/Animation/ControlRig/Source/ControlRig/Public/Units/RigUnitContext.h
--- a/Engine/Plugins/Animation/ControlRig/Source/ControlRig/Public/Units/RigUnitContext.h
+++ b/Engine/Plugins/Animation/ControlRig/Source/ControlRig/Public/Units/RigUnitContext.h
@@ -23,6 +23,19 @@ enum class EControlRigState : uint8
 	Invalid,
 };
 
+/**
+ * The type of interaction happening on a rig
+ */
+UENUM()
+enum class EControlRigInteractionType : uint8
+{
+	None = 0,
+	Translate = (1 << 0),
+	Rotate = (1 << 1),
+	Scale = (1 << 2),
+	All = Translate | Rotate | Scale
+};
+
 /** Execution context that rig units use */
 struct FRigUnitContext
 {
@@ -35,7 +48,8 @@ struct FRigUnitContext
 		, AbsoluteTime(0.f)
 		, State(EControlRigState::Invalid)
 		, Hierarchy(nullptr)
-		, bDuringInteraction(false)
+		, InteractionType((uint8)EControlRigInteractionType::None)
+		, ElementsBeingInteracted()
 		, ToWorldSpaceTransform(FTransform::Identity)
 		, OwningComponent(nullptr)
 		, OwningActor(nullptr)
@@ -71,8 +85,11 @@ struct FRigUnitContext
 	/** The current hierarchy being executed */
 	URigHierarchy* Hierarchy;
 
-	/** True if the rig is executing during an interaction */
-	bool bDuringInteraction;
+	/** The type of interaction currently happening on the rig (0 == None) */
+	uint8 InteractionType;
+
+	/** The elements being interacted with. */
+	TArray<FRigElementKey> ElementsBeingInteracted;
 
 	/** The current transform going from rig (global) space to world space */
 	FTransform ToWorldSpaceTransform;
@@ -157,6 +174,14 @@ struct FRigUnitContext
 	{
 		return ToWorldSpaceTransform.InverseTransformRotation(InRotation);
 	}
+
+	/**
+	 * Returns true if this context is currently being interacted on
+	 */
+	FORCEINLINE bool IsInteracting() const
+	{
+		return InteractionType != (uint8)EControlRigInteractionType::None;
+	}
 };
 
 #if WITH_EDITOR
diff --git a/Engine/Plugins/Animation/ControlRig/Source/ControlRigEditor/Private/EditMode/ControlRigEditMode.cpp b/Engine/Plugins/Animation/ControlRig/Source/ControlRigEditor/Private/EditMode/ControlRigEditMode.cpp
--- a/Engine/Plugins/Animation/ControlRig/Source/ControlRigEditor/Private/EditMode/ControlRigEditMode.cpp
+++ b/Engine/Plugins/Animation/ControlRig/Source/ControlRigEditor/Private/EditMode/ControlRigEditMode.cpp
@@ -137,6 +137,7 @@ FControlRigEditMode::FControlRigEditMode()
 	, bSuspendHierarchyNotifs(false)
 	, CurrentViewportClient(nullptr)
 	, bIsChangingCoordSystem(false)
+	, InteractionType((uint8)EControlRigInteractionType::None)
 {
 	ControlProxy = NewObject<UControlRigDetailPanelControlProxies>(GetTransientPackage(), NAME_None);
 	ControlProxy->SetFlags(RF_Transactional);
@@ -886,6 +887,7 @@ bool FControlRigEditMode::InputKey(FEditorViewportClient* InViewportClient, FVie
 
 bool FControlRigEditMode::EndTracking(FEditorViewportClient* InViewportClient, FViewport* InViewport)
 {
+	InteractionType = (uint8)EControlRigInteractionType::None;
 
 	if (InteractionScopes.Num() > 0)
 	{
@@ -915,6 +917,8 @@ bool FControlRigEditMode::EndTracking(FEditorViewportClient* InViewportClient, F
 
 bool FControlRigEditMode::StartTracking(FEditorViewportClient* InViewportClient, FViewport* InViewport)
 {
+	InteractionType = GetInteractionType(InViewportClient);
+	
 	if (InteractionScopes.Num() == 0)
 	{
 		bool bShouldModify = IsInLevelEditor();
@@ -1089,6 +1093,8 @@ bool FControlRigEditMode::GetCustomInputCoordinateSystem(FMatrix& OutMatrix, voi
 
 bool FControlRigEditMode::HandleClick(FEditorViewportClient* InViewportClient, HHitProxy *HitProxy, const FViewportClick &Click)
 {
+	InteractionType = GetInteractionType(InViewportClient);
+	
 	if(HActor* ActorHitProxy = HitProxyCast<HActor>(HitProxy))
 	{
 		if(ActorHitProxy->Actor)
@@ -1279,6 +1285,10 @@ bool FControlRigEditMode::HandleClick(FEditorViewportClient* InViewportClient, H
 			}
 		}
 	}
+	else
+	{
+		InteractionType = (uint8)EControlRigInteractionType::None;
+	}
 
 	// for now we show this menu all the time if body is selected
 	// if we want some global menu, we'll have to move this
@@ -2247,6 +2257,30 @@ void FControlRigEditMode::ResetControlShapeSize()
 	GetModeManager()->SetWidgetScale(Settings->GizmoScale);
 }
 
+uint8 FControlRigEditMode::GetInteractionType(FEditorViewportClient* InViewportClient)
+{
+	uint8 Result = (uint8)EControlRigInteractionType::None;
+	if(InViewportClient->IsMovingCamera())
+	{
+		return Result;
+	}
+	
+	const UE::Widget::EWidgetMode WidgetMode = InViewportClient->GetWidgetMode();
+	if(WidgetMode == UE::Widget::WM_Translate || WidgetMode == UE::Widget::WM_TranslateRotateZ)
+	{
+		Result |= (uint8)EControlRigInteractionType::Translate;
+	}
+	if(WidgetMode == UE::Widget::WM_Rotate || WidgetMode == UE::Widget::WM_TranslateRotateZ)
+	{
+		Result |= (uint8)EControlRigInteractionType::Rotate;
+	}
+	if(WidgetMode == UE::Widget::WM_Scale)
+	{
+		Result |= (uint8)EControlRigInteractionType::Scale;
+	}
+	return Result;	
+}
+
 void FControlRigEditMode::ToggleControlShapeTransformEdit()
 { 
 	if (bIsChangingControlShapeTransform)
@@ -3179,6 +3213,8 @@ void FControlRigEditMode::MoveControlShape(AControlRigShapeActor* ShapeActor, co
 
 				ShapeActor->SetGlobalTransform(CurrentTransform);
 
+				ControlRig->InteractionType = InteractionType;
+				ControlRig->ElementsBeingInteracted.AddUnique(ShapeActor->GetElementKey());
 				ControlRig->Evaluate_AnyThread();
 			}
 		}
@@ -3236,6 +3272,9 @@ void FControlRigEditMode::MoveControlShape(AControlRigShapeActor* ShapeActor, co
 					InOutLocal = NewLocal.GetRelativeTransform(InOutLocal);
 
 				}
+
+				ControlRig->InteractionType = InteractionType;
+				ControlRig->ElementsBeingInteracted.AddUnique(ShapeActor->GetElementKey());
 				ControlRig->Evaluate_AnyThread();		
 			}
 		}
@@ -3269,7 +3308,7 @@ void FControlRigEditMode::ChangeControlShapeTransform(AControlRigShapeActor* Sha
 
 	if (UControlRig* ControlRig = ShapeActor->ControlRig.Get())
 	{
-		if (FRigControlElement* ControlElement = ControlRig->GetHierarchy()->Find<FRigControlElement>(FRigElementKey(ShapeActor->ControlName, ERigElementType::Control)))
+		if (FRigControlElement* ControlElement = ControlRig->GetHierarchy()->Find<FRigControlElement>(ShapeActor->GetElementKey()))
 		{
 			CurrentTransform = ControlRig->GetHierarchy()->GetControlShapeTransform(ControlElement, ERigTransformType::CurrentGlobal);
 			CurrentTransform = CurrentTransform * ToWorldTransform;
@@ -3308,7 +3347,7 @@ void FControlRigEditMode::ChangeControlShapeTransform(AControlRigShapeActor* Sha
 			FTransform NewTransform = CurrentTransform.GetRelativeTransform(ToWorldTransform);
 			FRigControlModifiedContext Context;
 
-			if (FRigControlElement* ControlElement = ControlRig->GetHierarchy()->Find<FRigControlElement>(FRigElementKey(ShapeActor->ControlName, ERigElementType::Control)))
+			if (FRigControlElement* ControlElement = ControlRig->GetHierarchy()->Find<FRigControlElement>(ShapeActor->GetElementKey()))
 			{
 				// do not setup undo for this first step since it is just used to calculate the local transform
 				ControlRig->GetHierarchy()->SetControlShapeTransform(ControlElement, NewTransform, ERigTransformType::CurrentGlobal, false);
diff --git a/Engine/Plugins/Animation/ControlRig/Source/ControlRigEditor/Private/EditMode/ControlRigEditMode.h b/Engine/Plugins/Animation/ControlRig/Source/ControlRigEditor/Private/EditMode/ControlRigEditMode.h
--- a/Engine/Plugins/Animation/ControlRig/Source/ControlRigEditor/Private/EditMode/ControlRigEditMode.h
+++ b/Engine/Plugins/Animation/ControlRig/Source/ControlRigEditor/Private/EditMode/ControlRigEditMode.h
@@ -410,6 +410,10 @@ private:
 
 	TArray<FRigElementKey> DeferredItemsToFrame;
 
+	/** Computes the current interaction types based on the widget mode */
+	static uint8 GetInteractionType(FEditorViewportClient* InViewportClient);
+	uint8 InteractionType;
+
 	friend class FControlRigEditorModule;
 	friend class FControlRigEditor;
 	friend class FControlRigEditModeGenericDetails;
