Control Rig: Edit Mode: Need to invalidate level editor viewport when undoing if we have something selected, otherwise proxy's will be wrong.

diff --git a/Engine/Plugins/Animation/ControlRig/Source/ControlRigEditor/Private/EditMode/ControlRigEditMode.cpp b/Engine/Plugins/Animation/ControlRig/Source/ControlRigEditor/Private/EditMode/ControlRigEditMode.cpp
--- a/Engine/Plugins/Animation/ControlRig/Source/ControlRigEditor/Private/EditMode/ControlRigEditMode.cpp
+++ b/Engine/Plugins/Animation/ControlRig/Source/ControlRigEditor/Private/EditMode/ControlRigEditMode.cpp
@@ -2649,6 +2649,7 @@ bool FControlRigEditMode::MouseLeave(FEditorViewportClient* ViewportClient, FVie
 
 void FControlRigEditMode::PostUndo()
 {
+	bool bInvalidateViewport = false;
 	if (WeakSequencer.IsValid())
 	{
 		for (TWeakObjectPtr<UControlRig>& RuntimeRigPtr : RuntimeControlRigs)
@@ -2656,6 +2657,7 @@ void FControlRigEditMode::PostUndo()
 			if (RuntimeRigPtr.IsValid() == false)
 			{
 				DestroyShapesActors(RuntimeRigPtr.Get());
+				bInvalidateViewport = true;
 			}
 		}
 		TSharedPtr<ISequencer> Sequencer = WeakSequencer.Pin();
@@ -2678,6 +2680,7 @@ void FControlRigEditMode::PostUndo()
 			if (bSomethingAdded)
 			{
 				SetObjects_Internal();
+				bInvalidateViewport = true;
 			}
 		}
 	}
@@ -2688,9 +2691,23 @@ void FControlRigEditMode::PostUndo()
 			if (RuntimeRigPtr.IsValid() == false)
 			{
 				DestroyShapesActors(RuntimeRigPtr.Get());
+				bInvalidateViewport = true;
 			}
 		}
 	}
+
+	//normal actor undo will force the redraw, so we need to do the same for our transients/controls.
+	if (IsInLevelEditor() && (bInvalidateViewport || UsesTransformWidget()))
+	{
+		GEditor->GetTimerManager()->SetTimerForNextTick([this]()
+		{
+			//due to tick ordering need to manually make sure we get everything done in correct order.
+			PostPoseUpdate();
+			RecalcPivotTransform();
+			GEditor->RedrawLevelEditingViewports(true);
+		});
+	}
+
 }
 
 void FControlRigEditMode::RequestToRecreateControlShapeActors(UControlRig* ControlRig)
