Control Rig: Fix static analysis issue with PDI

diff --git a/Engine/Plugins/Animation/ControlRig/Source/ControlRigEditor/Private/EditMode/ControlRigEditMode.cpp b/Engine/Plugins/Animation/ControlRig/Source/ControlRigEditor/Private/EditMode/ControlRigEditMode.cpp
--- a/Engine/Plugins/Animation/ControlRig/Source/ControlRigEditor/Private/EditMode/ControlRigEditMode.cpp
+++ b/Engine/Plugins/Animation/ControlRig/Source/ControlRigEditor/Private/EditMode/ControlRigEditMode.cpp
@@ -743,7 +743,7 @@ void FControlRigEditMode::Render(const FSceneView* View, FViewport* Viewport, FP
 						}
 						if (PDI)
 						{
-							const bool bHitTesting = PDI && PDI->IsHitTesting() && bBoolSetHitProxies && (TransformElement->GetType() == ERigElementType::Bone);
+							const bool bHitTesting = PDI->IsHitTesting() && bBoolSetHitProxies && (TransformElement->GetType() == ERigElementType::Bone);
 							if (bHitTesting)
 							{
 								PDI->SetHitProxy(new HFKRigBoneProxy(TransformElement->GetName(), ControlRig));
