Control Rig: Improve performance for cycle checking transforms

diff --git a/Engine/Plugins/Animation/ControlRig/Source/ControlRig/Private/ControlRig.cpp b/Engine/Plugins/Animation/ControlRig/Source/ControlRig/Private/ControlRig.cpp
--- a/Engine/Plugins/Animation/ControlRig/Source/ControlRig/Private/ControlRig.cpp
+++ b/Engine/Plugins/Animation/ControlRig/Source/ControlRig/Private/ControlRig.cpp
@@ -1234,10 +1234,13 @@ void UControlRig::ExecuteUnits(FRigUnitContext& InOutContext, const FName& InEve
 
 			URigHierarchy* Hierarchy = GetHierarchy();
 #if WITH_EDITOR
-			Hierarchy->ReadTransformsPerInstructionPerSlice.Reset();
-			Hierarchy->WrittenTransformsPerInstructionPerSlice.Reset();
-			Hierarchy->ReadTransformsPerInstructionPerSlice.AddZeroed(VM->GetByteCode().GetNumInstructions());
-			Hierarchy->WrittenTransformsPerInstructionPerSlice.AddZeroed(VM->GetByteCode().GetNumInstructions());
+
+			if(Hierarchy->bRecordTransformsAtRuntime)
+			{
+				Hierarchy->ReadTransformsAtRuntime.Reset();
+				Hierarchy->WrittenTransformsAtRuntime.Reset();
+			}
+			
 #endif
 			FRigHierarchyExecuteContextBracket HierarchyContextGuard(Hierarchy, &VM->GetContext());
 
diff --git a/Engine/Plugins/Animation/ControlRig/Source/ControlRig/Private/Rigs/RigHierarchy.cpp b/Engine/Plugins/Animation/ControlRig/Source/ControlRig/Private/Rigs/RigHierarchy.cpp
--- a/Engine/Plugins/Animation/ControlRig/Source/ControlRig/Private/Rigs/RigHierarchy.cpp
+++ b/Engine/Plugins/Animation/ControlRig/Source/ControlRig/Private/Rigs/RigHierarchy.cpp
@@ -114,7 +114,7 @@ URigHierarchy::URigHierarchy()
 , HierarchyForCacheValidation()
 , ExecuteContext(nullptr)
 #if WITH_EDITOR
-, bRecordTransformsPerInstruction(true)
+, bRecordTransformsAtRuntime(true)
 #endif
 {
 	Reset();
@@ -2335,16 +2335,16 @@ FTransform URigHierarchy::GetTransform(FRigTransformElement* InTransformElement,
 
 #if WITH_EDITOR
 
-	if(bRecordTransformsPerInstruction && ExecuteContext)
+	if(bRecordTransformsAtRuntime && ExecuteContext)
 	{
-		TArray<TArray<int32>>& ReadTransformsPerSlice = ReadTransformsPerInstructionPerSlice[ExecuteContext->PublicData.GetInstructionIndex()];
-		while(ReadTransformsPerSlice.Num() < ExecuteContext->GetSlice().TotalNum())
-		{
-			ReadTransformsPerSlice.Add(TArray<int32>());
-		}
-		ReadTransformsPerSlice[ExecuteContext->GetSlice().GetIndex()].Add(InTransformElement->GetIndex());
+		ReadTransformsAtRuntime.Emplace(
+			ExecuteContext->PublicData.GetInstructionIndex(),
+			ExecuteContext->GetSlice().GetIndex(),
+			InTransformElement->GetIndex(),
+			InTransformType
+		);
 	}
-	TGuardValue<bool> RecordTransformsPerInstructionGuard(bRecordTransformsPerInstruction, false);
+	TGuardValue<bool> RecordTransformsPerInstructionGuard(bRecordTransformsAtRuntime, false);
 	
 #endif
 	
@@ -2440,16 +2440,16 @@ void URigHierarchy::SetTransform(FRigTransformElement* InTransformElement, const
 
 #if WITH_EDITOR
 
-	if(bRecordTransformsPerInstruction && ExecuteContext)
+	if(bRecordTransformsAtRuntime && ExecuteContext)
 	{
-		TArray<TArray<int32>>& WrittenTransformsPerSlice = WrittenTransformsPerInstructionPerSlice[ExecuteContext->PublicData.GetInstructionIndex()];
-		while(WrittenTransformsPerSlice.Num() < ExecuteContext->GetSlice().TotalNum())
-		{
-			WrittenTransformsPerSlice.Add(TArray<int32>());
-		}
-		WrittenTransformsPerSlice[ExecuteContext->GetSlice().GetIndex()].Add(InTransformElement->GetIndex());
+		WrittenTransformsAtRuntime.Emplace(
+			ExecuteContext->PublicData.GetInstructionIndex(),
+			ExecuteContext->GetSlice().GetIndex(),
+			InTransformElement->GetIndex(),
+			InTransformType
+		);
 	}
-	TGuardValue<bool> RecordTransformsPerInstructionGuard(bRecordTransformsPerInstruction, false);
+	TGuardValue<bool> RecordTransformsPerInstructionGuard(bRecordTransformsAtRuntime, false);
 	
 #endif
 
@@ -2568,16 +2568,16 @@ FTransform URigHierarchy::GetControlOffsetTransform(FRigControlElement* InContro
 
 #if WITH_EDITOR
 
-	if(bRecordTransformsPerInstruction && ExecuteContext)
+	if(bRecordTransformsAtRuntime && ExecuteContext)
 	{
-		TArray<TArray<int32>>& ReadTransformsPerSlice = ReadTransformsPerInstructionPerSlice[ExecuteContext->PublicData.GetInstructionIndex()];
-		while(ReadTransformsPerSlice.Num() < ExecuteContext->GetSlice().TotalNum())
-		{
-			ReadTransformsPerSlice.Add(TArray<int32>());
-		}
-		ReadTransformsPerSlice[ExecuteContext->GetSlice().GetIndex()].Add(InControlElement->GetIndex());
+		ReadTransformsAtRuntime.Emplace(
+			ExecuteContext->PublicData.GetInstructionIndex(),
+			ExecuteContext->GetSlice().GetIndex(),
+			InControlElement->GetIndex(),
+			InTransformType
+		);
 	}
-	TGuardValue<bool> RecordTransformsPerInstructionGuard(bRecordTransformsPerInstruction, false);
+	TGuardValue<bool> RecordTransformsPerInstructionGuard(bRecordTransformsAtRuntime, false);
 	
 #endif
 
@@ -2651,16 +2651,16 @@ void URigHierarchy::SetControlOffsetTransform(FRigControlElement* InControlEleme
 
 #if WITH_EDITOR
 
-	if(bRecordTransformsPerInstruction && ExecuteContext)
+	if(bRecordTransformsAtRuntime && ExecuteContext)
 	{
-		TArray<TArray<int32>>& WrittenTransformsPerSlice = WrittenTransformsPerInstructionPerSlice[ExecuteContext->PublicData.GetInstructionIndex()];
-		while(WrittenTransformsPerSlice.Num() < ExecuteContext->GetSlice().TotalNum())
-		{
-			WrittenTransformsPerSlice.Add(TArray<int32>());
-		}
-		WrittenTransformsPerSlice[ExecuteContext->GetSlice().GetIndex()].Add(InControlElement->GetIndex());
+		WrittenTransformsAtRuntime.Emplace(
+			ExecuteContext->PublicData.GetInstructionIndex(),
+			ExecuteContext->GetSlice().GetIndex(),
+			InControlElement->GetIndex(),
+			InTransformType
+		);
 	}
-	TGuardValue<bool> RecordTransformsPerInstructionGuard(bRecordTransformsPerInstruction, false);
+	TGuardValue<bool> RecordTransformsPerInstructionGuard(bRecordTransformsAtRuntime, false);
 	
 #endif
 
@@ -4209,12 +4209,6 @@ URigHierarchy::TElementDependencyMap URigHierarchy::GetDependenciesForVM(const U
 	URigHierarchy::TElementDependencyMap Dependencies;
 	const FRigVMInstructionArray Instructions = InVM->GetByteCode().GetInstructions();
 
-	// make sure the vm matches our cached data
-	if(ReadTransformsPerInstructionPerSlice.Num() != Instructions.Num())
-	{
-		return Dependencies;
-	}
-
 	// if the VM doesn't implement the given event
 	if(!InVM->ContainsEntry(InEventName))
 	{
@@ -4243,8 +4237,32 @@ URigHierarchy::TElementDependencyMap URigHierarchy::GetDependenciesForVM(const U
 		}
 	}
 
+	typedef TTuple<int32, int32> TInt32Tuple;
+	TArray<TArray<TInt32Tuple>> ReadTransformPerInstruction, WrittenTransformsPerInstruction;
+	ReadTransformPerInstruction.AddZeroed(Instructions.Num());
+	WrittenTransformsPerInstruction.AddZeroed(Instructions.Num());
+
+	// fill lookup tables per instruction / element
+	for(int32 RecordType = 0; RecordType < 2; RecordType++)
+	{
+		const TArray<TInstructionSliceElement>& Records = RecordType == 0 ? ReadTransformsAtRuntime : WrittenTransformsAtRuntime; 
+		TArray<TArray<TInt32Tuple>>& PerInstruction = RecordType == 0 ? ReadTransformPerInstruction : WrittenTransformsPerInstruction;
+
+		for(int32 RecordIndex = 0; RecordIndex < Records.Num(); RecordIndex++)
+		{
+			const TInstructionSliceElement& Record = Records[RecordIndex];
+			const int32& InstructionIndex = Record.Get<0>();
+			const int32& SliceIndex = Record.Get<1>();
+			const int32& TransformIndex = Record.Get<2>();
+			PerInstruction[InstructionIndex].Emplace(SliceIndex, TransformIndex);
+		}
+	}
+
 	// for each read transform on an instruction
 	// follow the operands to the next instruction affected by it
+	TArray<TInt32Tuple> FilteredTransforms;
+	TArray<int32> InstructionsToVisit;
+
 	for(int32 InstructionIndex = EntryInstructionIndex; InstructionIndex < Instructions.Num(); InstructionIndex++)
 	{
 		// early exit since following instructions belong to another event
@@ -4253,42 +4271,41 @@ URigHierarchy::TElementDependencyMap URigHierarchy::GetDependenciesForVM(const U
 			break;
 		}
 
-		const TArray<TArray<int32>>& ReadTransformsPerSlice = ReadTransformsPerInstructionPerSlice[InstructionIndex];
-
-		for(int32 SliceIndex = 0; SliceIndex < ReadTransformsPerSlice.Num(); SliceIndex++)
+		const TArray<TInt32Tuple>& ReadTransforms = ReadTransformPerInstruction[InstructionIndex];
+		if(ReadTransforms.IsEmpty())
 		{
-			const TArray<int32>& ReadTransforms = ReadTransformsPerSlice[SliceIndex];
-			if(ReadTransforms.IsEmpty())
-			{
-				continue;
-			}
+			continue;
+		}
 
-			TArray<int32> InstructionsToVisit;
-			InstructionsToVisit.Add(InstructionIndex);
+		FilteredTransforms.Reset();
+
+		for(int32 ReadTransformIndex = 0; ReadTransformIndex < ReadTransforms.Num(); ReadTransformIndex++)
+		{
+			const TInt32Tuple& ReadTransform = ReadTransforms[ReadTransformIndex];
 			
-			TArray<int32> WrittenTransforms;
+			InstructionsToVisit.Reset();
+			InstructionsToVisit.Add(InstructionIndex);
 
 			for(int32 InstructionToVisitIndex = 0; InstructionToVisitIndex < InstructionsToVisit.Num(); InstructionToVisitIndex++)
 			{
 				const int32 InstructionToVisit = InstructionsToVisit[InstructionToVisitIndex];
-				const TArray<TArray<int32>>& WrittenTransformsPerSlice = WrittenTransformsPerInstructionPerSlice[InstructionToVisit];
-				if(WrittenTransformsPerSlice.IsValidIndex(SliceIndex))
+
+				const TArray<TInt32Tuple>& WrittenTransforms = WrittenTransformsPerInstruction[InstructionToVisit];
+				for(int32 WrittenTransformIndex = 0; WrittenTransformIndex < WrittenTransforms.Num(); WrittenTransformIndex++)
 				{
-					const TArray<int32>& WrittenTransformsForSlice = WrittenTransformsPerSlice[SliceIndex];
-					for(int32 WrittenTransform : WrittenTransformsForSlice)
+					const TInt32Tuple& WrittenTransform = WrittenTransforms[WrittenTransformIndex];
+
+					// for the first instruction in this pass let's only care about
+					// written transforms which have not been read before
+					if(InstructionToVisitIndex == 0)
 					{
-						// for the first instruction in this pass let's only care about
-						// written transforms which have not been read before
-						if(InstructionToVisit == InstructionIndex)
+						if(ReadTransforms.Contains(WrittenTransform))
 						{
-							if(ReadTransforms.Contains(WrittenTransform))
-							{
-								continue;
-							}
+							continue;
 						}
-						
-						WrittenTransforms.AddUnique(WrittenTransform);
 					}
+					
+					FilteredTransforms.AddUnique(WrittenTransform);
 				}
 
 				FRigVMOperandArray OutputOperands = InVM->GetByteCode().GetOutputOperands(InstructionToVisit);
@@ -4305,15 +4322,16 @@ URigHierarchy::TElementDependencyMap URigHierarchy::GetDependenciesForVM(const U
 					}
 				}
 			}
-
-			for(const int32 ReadTransform : ReadTransforms)
+		}
+		
+		for(const TInt32Tuple& ReadTransform : ReadTransforms)
+		{
+			for(const TInt32Tuple& FilteredTransform : FilteredTransforms)
 			{
-				for(const int32 WrittenTransform : WrittenTransforms)
+				// only create dependencies for reads and writes that are on the same slice
+				if(ReadTransform != FilteredTransform && ReadTransform.Get<0>() == FilteredTransform.Get<0>())
 				{
-					if(ReadTransform != WrittenTransform)
-					{
-						Dependencies.FindOrAdd(WrittenTransform).AddUnique(ReadTransform);
-					}
+					Dependencies.FindOrAdd(ReadTransform.Get<1>()).AddUnique(FilteredTransform.Get<1>());
 				}
 			}
 		}
@@ -4350,7 +4368,7 @@ void URigHierarchy::PushTransformToStack(const FRigElementKey& InKey, ERigTransf
 
 	static const FText TransformPoseTitle = NSLOCTEXT("RigHierarchy", "Set Pose Transform", "Set Pose Transform");
 	static const FText ControlOffsetTitle = NSLOCTEXT("RigHierarchy", "Set Control Offset", "Set Control Offset");
-	static const FText ControlShapeTitle = NSLOCTEXT("RigHierarchy", "Set Control Gizo", "Set Control Gizo");
+	static const FText ControlShapeTitle = NSLOCTEXT("RigHierarchy", "Set Control Shape", "Set Control Shape");
 	static const FText CurveValueTitle = NSLOCTEXT("RigHierarchy", "Set Curve Value", "Set Curve Value");
 	
 	FText Title;
diff --git a/Engine/Plugins/Animation/ControlRig/Source/ControlRig/Public/Rigs/RigHierarchy.h b/Engine/Plugins/Animation/ControlRig/Source/ControlRig/Public/Rigs/RigHierarchy.h
--- a/Engine/Plugins/Animation/ControlRig/Source/ControlRig/Public/Rigs/RigHierarchy.h
+++ b/Engine/Plugins/Animation/ControlRig/Source/ControlRig/Public/Rigs/RigHierarchy.h
@@ -94,6 +94,7 @@ public:
 
 	typedef TMap<int32, TArray<int32>> TElementDependencyMap;
 	typedef TPair<int32, TArray<int32>> TElementDependencyMapPair;
+	typedef TTuple<int32, int32, int32, ERigTransformType::Type> TInstructionSliceElement;
 
 	URigHierarchy();
 
@@ -2976,10 +2977,9 @@ private:
 	const FRigVMExtendedExecuteContext* ExecuteContext;
 
 #if WITH_EDITOR
-	mutable bool bRecordTransformsPerInstruction;
-	mutable TArray<TArray<TArray<int32>>> ReadTransformsPerInstructionPerSlice;
-	mutable TArray<TArray<TArray<int32>>> WrittenTransformsPerInstructionPerSlice;
-
+	mutable bool bRecordTransformsAtRuntime;
+	mutable TArray<TInstructionSliceElement> ReadTransformsAtRuntime;
+	mutable TArray<TInstructionSliceElement> WrittenTransformsAtRuntime;
 public:
 
 	TElementDependencyMap GetDependenciesForVM(const URigVM* InVM, FName InEventName = NAME_None) const;
