Control Rig: Fix crash when changing variable type from vector to int

diff --git a/Engine/Plugins/Animation/ControlRig/Source/ControlRigEditor/Private/ControlRigGraphDetails.cpp b/Engine/Plugins/Animation/ControlRig/Source/ControlRigEditor/Private/ControlRigGraphDetails.cpp
--- a/Engine/Plugins/Animation/ControlRig/Source/ControlRigEditor/Private/ControlRigGraphDetails.cpp
+++ b/Engine/Plugins/Animation/ControlRig/Source/ControlRigEditor/Private/ControlRigGraphDetails.cpp
@@ -1804,7 +1804,7 @@ void FControlRigGraphMathTypeDetails::CustomizeHeader(TSharedRef<IPropertyHandle
 
 	for (UObject* Object : Objects)
 	{
-		ObjectsBeingCustomized.Add(Object);
+		ObjectsBeingCustomized.Add(MakeWeakObjectPtr(Object));
 
 		if(BlueprintBeingCustomized == nullptr)
 		{
@@ -1817,8 +1817,6 @@ void FControlRigGraphMathTypeDetails::CustomizeHeader(TSharedRef<IPropertyHandle
 		}
 	}
 
-	ObjectBeingCustomizedView = TArrayView<const UObject* const>(ObjectsBeingCustomized.GetData(), ObjectsBeingCustomized.Num());
-
 	FProperty* Property = InPropertyHandle->GetProperty();
 	const FStructProperty* StructProperty = CastField<FStructProperty>(Property);
 	ScriptStruct = StructProperty->Struct;
diff --git a/Engine/Plugins/Animation/ControlRig/Source/ControlRigEditor/Private/ControlRigGraphDetails.h b/Engine/Plugins/Animation/ControlRig/Source/ControlRigEditor/Private/ControlRigGraphDetails.h
--- a/Engine/Plugins/Animation/ControlRig/Source/ControlRigEditor/Private/ControlRigGraphDetails.h
+++ b/Engine/Plugins/Animation/ControlRig/Source/ControlRigEditor/Private/ControlRigGraphDetails.h
@@ -312,12 +312,12 @@ protected:
 	TOptional<NumericType> GetVectorComponent(TSharedRef<class IPropertyHandle> InPropertyHandle, int32 InComponent) 
 	{
 		TOptional<NumericType> Result;
-		for(UObject* Object : ObjectsBeingCustomized)
+		for(TWeakObjectPtr<UObject> Object : ObjectsBeingCustomized)
 		{
-			if(InPropertyHandle->IsValidHandle())
+			if(Object.IsValid() && InPropertyHandle->IsValidHandle())
 			{
 				static VectorType ZeroVector = VectorType();
-				const VectorType& Vector = ContainerUObjectToValueRef<VectorType>(Object, ZeroVector);
+				const VectorType& Vector = ContainerUObjectToValueRef<VectorType>(Object.Get(), ZeroVector);
 				NumericType Component = Vector[InComponent];
 				if(Result.IsSet())
 				{
@@ -339,7 +339,16 @@ protected:
 	template<typename VectorType, typename NumericType>
 	void OnVectorComponentChanged(TSharedRef<class IPropertyHandle> InPropertyHandle, int32 InComponent, NumericType InValue, bool bIsCommit, ETextCommit::Type InCommitType = ETextCommit::Default)
 	{
-		FPropertyChangedEvent PropertyChangedEvent(InPropertyHandle->GetProperty(), bIsCommit ? EPropertyChangeType::ValueSet : EPropertyChangeType::Interactive, ObjectsBeingCustomized);
+		TArray<UObject*> ObjectsView;
+		for(int32 Index = 0; Index < ObjectsBeingCustomized.Num(); Index++)
+		{
+			TWeakObjectPtr<UObject> Object = ObjectsBeingCustomized[Index];
+			if (Object.IsValid())
+			{
+				ObjectsView.Add(Object.Get());
+			}
+		}
+		FPropertyChangedEvent PropertyChangedEvent(InPropertyHandle->GetProperty(), bIsCommit ? EPropertyChangeType::ValueSet : EPropertyChangeType::Interactive, ObjectsView);
 		FPropertyChangedChainEvent PropertyChangedChainEvent(PropertyChain, PropertyChangedEvent);
 
 		URigVMController* Controller = nullptr;
@@ -354,11 +363,11 @@ protected:
 
 		for(int32 Index = 0; Index < ObjectsBeingCustomized.Num(); Index++)
 		{
-			UObject* Object = ObjectsBeingCustomized[Index];
-			if(InPropertyHandle->IsValidHandle())
+			TWeakObjectPtr<UObject> Object = ObjectsBeingCustomized[Index];
+			if(Object.IsValid() && InPropertyHandle->IsValidHandle())
 			{
 				static VectorType ZeroVector = VectorType();
-				VectorType& Vector = ContainerUObjectToValueRef<VectorType>(Object, ZeroVector);
+				VectorType& Vector = ContainerUObjectToValueRef<VectorType>(Object.Get(), ZeroVector);
 				VectorType PreviousVector = Vector;
 				Vector[InComponent] = InValue;
 					
@@ -439,12 +448,12 @@ protected:
 	TOptional<RotationType> GetRotation(TSharedRef<class IPropertyHandle> InPropertyHandle)
 	{
 		TOptional<RotationType> Result;
-		for(UObject* Object : ObjectsBeingCustomized)
+		for(TWeakObjectPtr<UObject> Object : ObjectsBeingCustomized)
 		{
-			if(InPropertyHandle->IsValidHandle())
+			if(Object.IsValid() && InPropertyHandle->IsValidHandle())
 			{
 				static RotationType ZeroRotation = RotationType();
-				const RotationType& Rotation = ContainerUObjectToValueRef<RotationType>(Object, ZeroRotation);
+				const RotationType& Rotation = ContainerUObjectToValueRef<RotationType>(Object.Get(), ZeroRotation);
 				if(Result.IsSet())
 				{
 					if(!Rotation.Equals(Result.GetValue()))
@@ -465,7 +474,16 @@ protected:
 	template<typename RotationType>
 	void OnRotationChanged(TSharedRef<class IPropertyHandle> InPropertyHandle, RotationType InValue, bool bIsCommit, ETextCommit::Type InCommitType = ETextCommit::Default)
 	{
-		FPropertyChangedEvent PropertyChangedEvent(InPropertyHandle->GetProperty(), bIsCommit ? EPropertyChangeType::ValueSet : EPropertyChangeType::Interactive, ObjectsBeingCustomized);
+		TArray<UObject*> ObjectsView;
+		for(int32 Index = 0; Index < ObjectsBeingCustomized.Num(); Index++)
+		{
+			TWeakObjectPtr<UObject> Object = ObjectsBeingCustomized[Index];
+			if (Object.IsValid())
+			{
+				ObjectsView.Add(Object.Get());
+			}
+		}
+		FPropertyChangedEvent PropertyChangedEvent(InPropertyHandle->GetProperty(), bIsCommit ? EPropertyChangeType::ValueSet : EPropertyChangeType::Interactive, ObjectsView);
 		FPropertyChangedChainEvent PropertyChangedChainEvent(PropertyChain, PropertyChangedEvent);
 
 		URigVMController* Controller = nullptr;
@@ -480,11 +498,11 @@ protected:
 
 		for(int32 Index = 0; Index < ObjectsBeingCustomized.Num(); Index++)
 		{
-			UObject* Object = ObjectsBeingCustomized[Index];
-			if(InPropertyHandle->IsValidHandle())
+			TWeakObjectPtr<UObject> Object = ObjectsBeingCustomized[Index];
+			if(Object.IsValid() && InPropertyHandle->IsValidHandle())
 			{
 				static RotationType ZeroRotation = RotationType();
-				RotationType& Rotation = ContainerUObjectToValueRef<RotationType>(Object, ZeroRotation);
+				RotationType& Rotation = ContainerUObjectToValueRef<RotationType>(Object.Get(), ZeroRotation);
 				RotationType PreviousRotation = Rotation;
 				Rotation = InValue;
 					
@@ -551,11 +569,11 @@ protected:
 
 		WidgetArgs.DiffersFromDefault_Lambda([this, InPropertyHandle, DefaultValue](ESlateTransformComponent::Type InTransformComponent) -> bool
 		{
-			for(UObject* Object : ObjectsBeingCustomized)
+			for(TWeakObjectPtr<UObject> Object : ObjectsBeingCustomized)
 			{
-				if(InPropertyHandle->IsValidHandle())
+				if(Object.IsValid() && InPropertyHandle->IsValidHandle())
 				{
-					const TransformType& Transform = ContainerUObjectToValueRef<TransformType>(Object, Identity);
+					const TransformType& Transform = ContainerUObjectToValueRef<TransformType>(Object.Get(), Identity);
 
 					switch(InTransformComponent)
 					{
@@ -599,11 +617,11 @@ protected:
 			ESlateTransformSubComponent::Type InTransformSubComponent) -> TOptional<FReal>
 		{
 			TOptional<FReal> Result;
-			for(UObject* Object : ObjectsBeingCustomized)
+			for(TWeakObjectPtr<UObject> Object : ObjectsBeingCustomized)
 			{
-				if(InPropertyHandle->IsValidHandle())
+				if(Object.IsValid() && InPropertyHandle->IsValidHandle())
 				{
-					const TransformType& Transform = ContainerUObjectToValueRef<TransformType>(Object, Identity);
+					const TransformType& Transform = ContainerUObjectToValueRef<TransformType>(Object.Get(), Identity);
 					
 					TOptional<FReal> Value = SAdvancedTransformInputBox<TransformType>::GetNumericValueFromTransform(
 						Transform,
@@ -640,7 +658,16 @@ protected:
 			bool bIsCommit,
 			ETextCommit::Type InCommitType = ETextCommit::Default)
 		{
-			FPropertyChangedEvent PropertyChangedEvent(InPropertyHandle->GetProperty(), bIsCommit ? EPropertyChangeType::ValueSet : EPropertyChangeType::Interactive, ObjectsBeingCustomized);
+			TArray<UObject*> ObjectsView;
+			for(int32 Index = 0; Index < ObjectsBeingCustomized.Num(); Index++)
+			{
+				TWeakObjectPtr<UObject> Object = ObjectsBeingCustomized[Index];
+				if (Object.IsValid())
+				{
+					ObjectsView.Add(Object.Get());
+				}
+			}
+			FPropertyChangedEvent PropertyChangedEvent(InPropertyHandle->GetProperty(), bIsCommit ? EPropertyChangeType::ValueSet : EPropertyChangeType::Interactive, ObjectsView);
 			FPropertyChangedChainEvent PropertyChangedChainEvent(PropertyChain, PropertyChangedEvent);
 
 			URigVMController* Controller = nullptr;
@@ -655,10 +682,10 @@ protected:
 
 			for(int32 Index = 0; Index < ObjectsBeingCustomized.Num(); Index++)
 			{
-				UObject* Object = ObjectsBeingCustomized[Index];
-				if(InPropertyHandle->IsValidHandle())
+				TWeakObjectPtr<UObject> Object = ObjectsBeingCustomized[Index];
+				if(Object.IsValid() && InPropertyHandle->IsValidHandle())
 				{
-					TransformType& Transform = ContainerUObjectToValueRef<TransformType>(Object, Identity);
+					TransformType& Transform = ContainerUObjectToValueRef<TransformType>(Object.Get(), Identity);
 					TransformType PreviousTransform = Transform;
 					
 					SAdvancedTransformInputBox<TransformType>::ApplyNumericValueChange(
@@ -713,15 +740,24 @@ protected:
 				}
 			}
 
-			FPropertyChangedEvent PropertyChangedEvent(InPropertyHandle->GetProperty(), EPropertyChangeType::ValueSet, ObjectsBeingCustomized);
+			TArray<UObject*> ObjectsView;
+			for(int32 Index = 0; Index < ObjectsBeingCustomized.Num(); Index++)
+			{
+				TWeakObjectPtr<UObject> Object = ObjectsBeingCustomized[Index];
+				if (Object.IsValid())
+				{
+					ObjectsView.Add(Object.Get());
+				}
+			}
+			FPropertyChangedEvent PropertyChangedEvent(InPropertyHandle->GetProperty(), EPropertyChangeType::ValueSet, ObjectsView);
 			FPropertyChangedChainEvent PropertyChangedChainEvent(PropertyChain, PropertyChangedEvent);
 			
 			for(int32 Index = 0; Index < ObjectsBeingCustomized.Num(); Index++)
 			{
-				UObject* Object = ObjectsBeingCustomized[Index];
-				if(InPropertyHandle->IsValidHandle())
+				TWeakObjectPtr<UObject> Object = ObjectsBeingCustomized[Index];
+				if(Object.IsValid() && InPropertyHandle->IsValidHandle())
 				{
-					TransformType& Transform = ContainerUObjectToValueRef<TransformType>(Object, Identity);
+					TransformType& Transform = ContainerUObjectToValueRef<TransformType>(Object.Get(), Identity);
 					TransformType PreviousTransform = Transform;
 
 					switch(InTransformComponent)
@@ -774,8 +810,7 @@ protected:
 	UScriptStruct* ScriptStruct;
 	UControlRigBlueprint* BlueprintBeingCustomized;
 	URigVMGraph* GraphBeingCustomized;
-	TArray<UObject*> ObjectsBeingCustomized;
-	TArrayView<const UObject* const> ObjectBeingCustomizedView;
+	TArray<TWeakObjectPtr<UObject>> ObjectsBeingCustomized; 
 	FEditPropertyChain PropertyChain;
 	TArray<int32> PropertyArrayIndices;
 	bool bEnabled;
