Sequencer: Control Rig: Misc minor tweaks for mutiple control rigs.

diff --git a/Engine/Plugins/Animation/ControlRig/Source/ControlRigEditor/Private/ControlRigBlueprintActions.cpp b/Engine/Plugins/Animation/ControlRig/Source/ControlRigEditor/Private/ControlRigBlueprintActions.cpp
--- a/Engine/Plugins/Animation/ControlRig/Source/ControlRigEditor/Private/ControlRigBlueprintActions.cpp
+++ b/Engine/Plugins/Animation/ControlRig/Source/ControlRigEditor/Private/ControlRigBlueprintActions.cpp
@@ -347,12 +347,22 @@ void FControlRigBlueprintActions::OnSpawnedSkeletalMeshActorChanged(UObject* InO
 		TArray<TWeakObjectPtr<AActor> > ActorsToAdd;
 		ActorsToAdd.Add(MeshActor);
 		TArray<FGuid> ActorTracks = WeakSequencer.Pin()->AddActors(ActorsToAdd, false);
+		FControlRigEditMode* ControlRigEditMode = static_cast<FControlRigEditMode*>(GLevelEditorModeTools().GetActiveMode(FControlRigEditMode::ModeName));
+
 		for (FGuid ActorTrackGuid : ActorTracks)
 		{
 			//Delete binding from default animating rig
 			FGuid CompGuid = WeakSequencer.Pin()->FindObjectId(*(MeshActor->GetSkeletalMeshComponent()), WeakSequencer.Pin()->GetFocusedTemplateID());
 			if (CompGuid.IsValid())
 			{
+				if (ControlRigEditMode)
+				{
+					UMovieSceneControlRigParameterTrack* Track = Cast<UMovieSceneControlRigParameterTrack>(MovieScene->FindTrack(UMovieSceneControlRigParameterTrack::StaticClass(), CompGuid, NAME_None));
+					if (Track && Track->GetControlRig())
+					{
+						ControlRigEditMode->RemoveControlRig(Track->GetControlRig());
+					}
+				}
 				if (!MovieScene->RemovePossessable(CompGuid))
 				{
 					MovieScene->RemoveSpawnable(CompGuid);
@@ -390,12 +400,10 @@ void FControlRigBlueprintActions::OnSpawnedSkeletalMeshActorChanged(UObject* InO
 				FText Name = LOCTEXT("SequenceTrackFilter_ControlRigControls", "Control Rig Controls");
 				WeakSequencer.Pin()->SetTrackFilterEnabled(Name, true);
 				WeakSequencer.Pin()->NotifyMovieSceneDataChanged(EMovieSceneDataChangeType::MovieSceneStructureItemAdded);
-				FControlRigEditMode* ControlRigEditMode = static_cast<FControlRigEditMode*>(GLevelEditorModeTools().GetActiveMode(FControlRigEditMode::ModeName));
 				if (!ControlRigEditMode)
 				{
 					GLevelEditorModeTools().ActivateMode(FControlRigEditMode::ModeName);
 					ControlRigEditMode = static_cast<FControlRigEditMode*>(GLevelEditorModeTools().GetActiveMode(FControlRigEditMode::ModeName));
-
 				}
 				if (ControlRigEditMode)
 				{
diff --git a/Engine/Plugins/Animation/ControlRig/Source/ControlRigEditor/Private/EditMode/ControlRigEditMode.cpp b/Engine/Plugins/Animation/ControlRig/Source/ControlRigEditor/Private/EditMode/ControlRigEditMode.cpp
--- a/Engine/Plugins/Animation/ControlRig/Source/ControlRigEditor/Private/EditMode/ControlRigEditMode.cpp
+++ b/Engine/Plugins/Animation/ControlRig/Source/ControlRigEditor/Private/EditMode/ControlRigEditMode.cpp
@@ -2679,6 +2679,7 @@ void FControlRigEditMode::PostUndo()
 			}
 			if (bSomethingAdded)
 			{
+				Sequencer->ForceEvaluate();
 				SetObjects_Internal();
 				bInvalidateViewport = true;
 			}
diff --git a/Engine/Plugins/Animation/ControlRig/Source/ControlRigEditor/Private/Sequencer/ControlRigParameterTrackEditor.cpp b/Engine/Plugins/Animation/ControlRig/Source/ControlRigEditor/Private/Sequencer/ControlRigParameterTrackEditor.cpp
--- a/Engine/Plugins/Animation/ControlRig/Source/ControlRigEditor/Private/Sequencer/ControlRigParameterTrackEditor.cpp
+++ b/Engine/Plugins/Animation/ControlRig/Source/ControlRigEditor/Private/Sequencer/ControlRigParameterTrackEditor.cpp
@@ -2316,8 +2316,7 @@ void FControlRigParameterTrackEditor::HandleControlSelected(UControlRig* Subject
 	TArray<FRigControl> Controls;
 
 	URigHierarchy* Hierarchy = Subject->GetHierarchy();
-	if (bSelected)
-	{
+	
 	if (URigHierarchyController* Controller = Hierarchy->GetController())
 	{
 		Hierarchy->ForEach<FRigControlElement>([ControlElement, Controller, bSelected](FRigControlElement* OtherControlElement) -> bool
@@ -2340,7 +2339,7 @@ void FControlRigParameterTrackEditor::HandleControlSelected(UControlRig* Subject
 			});
 	}
 
-	}
+
 	if (bIsDoingSelection)
 	{
 		return;
