Control Rig: Fixed a bug with AimMath, it could not handle a specific combination of primary/seconardy targets

previously when we tried to aim the secondary axis towards the secondary target when they are 180degrees apart, we did not respect the fact that it can only rotate around the primary axis

diff --git a/Engine/Plugins/Animation/ControlRig/Source/ControlRig/Private/Units/Highlevel/Hierarchy/RigUnit_AimBone.cpp b/Engine/Plugins/Animation/ControlRig/Source/ControlRig/Private/Units/Highlevel/Hierarchy/RigUnit_AimBone.cpp
--- a/Engine/Plugins/Animation/ControlRig/Source/ControlRig/Private/Units/Highlevel/Hierarchy/RigUnit_AimBone.cpp
+++ b/Engine/Plugins/Animation/ControlRig/Source/ControlRig/Private/Units/Highlevel/Hierarchy/RigUnit_AimBone.cpp
@@ -118,9 +118,10 @@ FRigUnit_AimBoneMath_Execute()
 			Target = Target - Result.GetLocation();
 		}
 
-		if (!Primary.Axis.IsNearlyZero())
+		FVector PrimaryAxis = Primary.Axis;
+		if (!PrimaryAxis.IsNearlyZero())
 		{
-			FVector PrimaryAxis = Result.TransformVectorNoScale(Primary.Axis).GetSafeNormal();
+			PrimaryAxis = Result.TransformVectorNoScale(Primary.Axis).GetSafeNormal();
 			Target = Target - FVector::DotProduct(Target, PrimaryAxis) * PrimaryAxis;
 		}
 
@@ -134,7 +135,17 @@ FRigUnit_AimBoneMath_Execute()
 			{
 				Target = FMath::Lerp<FVector>(Axis, Target, T).GetSafeNormal();
 			}
-			FQuat Rotation = FControlRigMathLibrary::FindQuatBetweenVectors(Axis, Target);
+			
+			FQuat Rotation;
+			if (FVector::DotProduct(Axis,Target) + 1.f < SMALL_NUMBER && !PrimaryAxis.IsNearlyZero())
+			{
+				// special case, when the axis and target and 180 degrees apart and there is a primary axis
+				 Rotation = FQuat(PrimaryAxis, PI);
+			}
+			else
+			{
+				Rotation = FControlRigMathLibrary::FindQuatBetweenVectors(Axis, Target);
+			}
 			Result.SetRotation((Rotation * Result.GetRotation()).GetNormalized());
 		}
 		else
