Control Rig: Improve performance of constraint nodes

* Guarded conditional branch code paths
* Replaced name / key based lookups with cached keys.

diff --git a/Engine/Plugins/Animation/ControlRig/Source/ControlRig/Private/Units/Highlevel/Hierarchy/RigUnit_SphericalPoseReader.cpp b/Engine/Plugins/Animation/ControlRig/Source/ControlRig/Private/Units/Highlevel/Hierarchy/RigUnit_SphericalPoseReader.cpp
--- a/Engine/Plugins/Animation/ControlRig/Source/ControlRig/Private/Units/Highlevel/Hierarchy/RigUnit_SphericalPoseReader.cpp
+++ b/Engine/Plugins/Animation/ControlRig/Source/ControlRig/Private/Units/Highlevel/Hierarchy/RigUnit_SphericalPoseReader.cpp
@@ -16,13 +16,17 @@ FRigUnit_SphericalPoseReader_Execute()
 
 	if (Context.State == EControlRigState::Init)
 	{
+		DriverCache.Reset();
+		OptionalParentCache.Reset();
 		return;
 	}
 
-	if (!DriverItem.IsValid())
+	if (!DriverCache.UpdateCache(DriverItem, Hierarchy))
 	{
 		return;
 	}
+	
+	OptionalParentCache.UpdateCache(OptionalParentItem, Hierarchy);
 
 	// remap/clamp inputs
 	RemapAndConvertInputs(
@@ -38,30 +42,30 @@ FRigUnit_SphericalPoseReader_Execute()
 	// get parent global space
 	FTransform GlobalDriverParentTransform;
 	FTransform GlobalDriverParentTransformInit;
-	if (OptionalParentItem.IsValid())
+	if (OptionalParentCache.IsValid())
 	{
-		GlobalDriverParentTransform = Hierarchy->GetGlobalTransform(OptionalParentItem);
-		GlobalDriverParentTransformInit = Hierarchy->GetInitialGlobalTransform(OptionalParentItem);
+		GlobalDriverParentTransform = Hierarchy->GetGlobalTransformByIndex(OptionalParentCache, false);
+		GlobalDriverParentTransformInit = Hierarchy->GetGlobalTransformByIndex(OptionalParentCache, true);
 	}else
 	{
 		// if user does not specify a custom parent space, then simply use the driver's parent
-		GlobalDriverParentTransform = Hierarchy->GetParentTransform(DriverItem);
-		GlobalDriverParentTransformInit = Hierarchy->GetInitialGlobalTransform(DriverItem);
+		GlobalDriverParentTransform = Hierarchy->GetParentTransformByIndex(DriverCache, false);
+		GlobalDriverParentTransformInit = Hierarchy->GetParentTransformByIndex(DriverCache, true);
 	}
 	
 	// get local rotation space of driver
-	FTransform GlobalDriverTransformInit = Hierarchy->GetInitialGlobalTransform(DriverItem);
+	const FTransform GlobalDriverTransformInit = Hierarchy->GetGlobalTransformByIndex(DriverCache, true);
 	FTransform LocalDriverTransformInit = GlobalDriverTransformInit * GlobalDriverParentTransformInit.Inverse();
 	// apply static offset in local space
 	FQuat RotationOffsetQuat = FQuat::MakeFromEuler(RotationOffset);
 	LocalDriverTransformInit.SetRotation(LocalDriverTransformInit.GetRotation() * RotationOffsetQuat);
 	
 	// calculate world transform of sphere
+	const FTransform GlobalDriverTransform = Hierarchy->GetGlobalTransformByIndex(DriverCache);
 	FTransform WorldOffset = LocalDriverTransformInit * GlobalDriverParentTransform;
-	WorldOffset.SetLocation(Hierarchy->GetGlobalTransform(DriverItem).GetLocation());
+	WorldOffset.SetLocation(GlobalDriverTransform.GetLocation());
 
 	// get driver axis
-	const FTransform GlobalDriverTransform = Hierarchy->GetGlobalTransform(DriverItem);
 	const FVector CurrentGlobalDriverAxis = GlobalDriverTransform.GetRotation().RotateVector(DriverAxis);
 	DriverNormal = WorldOffset.InverseTransformVectorNoScale(CurrentGlobalDriverAxis).GetSafeNormal();
 
diff --git a/Engine/Plugins/Animation/ControlRig/Source/ControlRig/Private/Units/Highlevel/Hierarchy/RigUnit_SphericalPoseReader.h b/Engine/Plugins/Animation/ControlRig/Source/ControlRig/Private/Units/Highlevel/Hierarchy/RigUnit_SphericalPoseReader.h
--- a/Engine/Plugins/Animation/ControlRig/Source/ControlRig/Private/Units/Highlevel/Hierarchy/RigUnit_SphericalPoseReader.h
+++ b/Engine/Plugins/Animation/ControlRig/Source/ControlRig/Private/Units/Highlevel/Hierarchy/RigUnit_SphericalPoseReader.h
@@ -376,4 +376,10 @@ struct CONTROLRIG_API FRigUnit_SphericalPoseReader: public FRigUnit_HighlevelBas
 	FVector DriverNormal;
 	UPROPERTY(transient)
 	FVector Driver2D;
+	
+	UPROPERTY()
+	FCachedRigElement DriverCache;
+
+	UPROPERTY()
+	FCachedRigElement OptionalParentCache;
 };
\ No newline at end of file
diff --git a/Engine/Plugins/Animation/ControlRig/Source/ControlRig/Private/Units/Highlevel/Hierarchy/RigUnit_TransformConstraint.cpp b/Engine/Plugins/Animation/ControlRig/Source/ControlRig/Private/Units/Highlevel/Hierarchy/RigUnit_TransformConstraint.cpp
--- a/Engine/Plugins/Animation/ControlRig/Source/ControlRig/Private/Units/Highlevel/Hierarchy/RigUnit_TransformConstraint.cpp
+++ b/Engine/Plugins/Animation/ControlRig/Source/ControlRig/Private/Units/Highlevel/Hierarchy/RigUnit_TransformConstraint.cpp
@@ -194,30 +194,40 @@ FRigUnit_ParentConstraint_Execute()
 	{
 		return;
 	}
+
+	if(Context.State == EControlRigState::Init)
+	{
+		ChildCache.Reset();
+		ParentCaches.Reset();
+	}
 	
  	URigHierarchy* Hierarchy = ExecuteContext.Hierarchy;	
 	if (Hierarchy)
 	{
-		if (!Child.IsValid())
+		if (!ChildCache.UpdateCache(Child, Hierarchy))
 		{
 			return;
 		}
-		
-		FTransform ChildInitialGlobalTransform = Hierarchy->GetInitialGlobalTransform(Child);
-		FTransform ChildCurrentGlobalTransform = Hierarchy->GetGlobalTransform(Child, false);
+
+		if(ParentCaches.Num() != Parents.Num())
+		{
+			ParentCaches.SetNumZeroed(Parents.Num());
+		}
+
+		const bool bChildIsControl = Child.Type == ERigElementType::Control;
 
 		// calculate total weight
 		float OverallWeight = 0;
-		for (const FConstraintParent& Parent : Parents)
+		for (int32 ParentIndex = 0; ParentIndex < Parents.Num(); ParentIndex++)
 		{
-			const float ClampedWeight = FMath::Max(Parent.Weight, 0.f);
-
-			if (ClampedWeight < KINDA_SMALL_NUMBER)
+			const FConstraintParent& Parent = Parents[ParentIndex];
+			if (!ParentCaches[ParentIndex].UpdateCache(Parent.Item, Hierarchy))
 			{
 				continue;
 			}
 
-			if (!Parent.Item.IsValid())
+			const float ClampedWeight = FMath::Max(Parent.Weight, 0.f);
+			if (ClampedWeight < KINDA_SMALL_NUMBER)
 			{
 				continue;
 			}
@@ -229,24 +239,26 @@ FRigUnit_ParentConstraint_Execute()
 		{
 			const float WeightNormalizer = 1.0f / OverallWeight;
 
-			FTransform MixedGlobalTransform = FTransform::Identity;
+			FTransform MixedGlobalTransform;
 
 			// initial rotation needs to be (0,0,0,0) instead of (0,0,0,1) due to Quaternion Blending Math
+			MixedGlobalTransform.SetLocation(FVector::ZeroVector);
 			MixedGlobalTransform.SetRotation(FQuat(0.f, 0.f, 0.f, 0.f));
 			MixedGlobalTransform.SetScale3D(FVector::ZeroVector);
 
 			float AccumulatedWeight = 0.0f;
 
-			for (const FConstraintParent& Parent : Parents)
+			for (int32 ParentIndex = 0; ParentIndex < Parents.Num(); ParentIndex++)
 			{
-				const float ClampedWeight = FMath::Max(Parent.Weight, 0.f);
-
-				if (ClampedWeight < KINDA_SMALL_NUMBER)
+				if (!ParentCaches[ParentIndex].IsValid())
 				{
 					continue;
 				}
 
-				if (!Parent.Item.IsValid())
+				const FConstraintParent& Parent = Parents[ParentIndex];
+
+				const float ClampedWeight = FMath::Max(Parent.Weight, 0.f);
+				if (ClampedWeight < KINDA_SMALL_NUMBER)
 				{
 					continue;
 				}
@@ -254,19 +266,19 @@ FRigUnit_ParentConstraint_Execute()
 				const float NormalizedWeight = ClampedWeight * WeightNormalizer;
 				AccumulatedWeight += NormalizedWeight;
 
-				FTransform OffsetTransform = FTransform::Identity;
-				FTransform ParentCurrentGlobalTransform = Hierarchy->GetGlobalTransform(Parent.Item, false);
+				FTransform OffsetParentTransform = Hierarchy->GetGlobalTransformByIndex(ParentCaches[ParentIndex], false);
 
 				if (bMaintainOffset)
 				{
-					FTransform ParentInitialGlobalTransform = Hierarchy->GetInitialGlobalTransform(Parent.Item);
+					const FTransform ChildInitialGlobalTransform = Hierarchy->GetInitialGlobalTransform(ChildCache);
+					const FTransform ParentInitialGlobalTransform = Hierarchy->GetInitialGlobalTransform(ParentCaches[ParentIndex]);
+					
 					// offset transform is a transform that transforms parent to child
-					OffsetTransform = ChildInitialGlobalTransform.GetRelativeTransform(ParentInitialGlobalTransform);
-					OffsetTransform.NormalizeRotation();
-				}
+					const FTransform OffsetTransform = ChildInitialGlobalTransform.GetRelativeTransform(ParentInitialGlobalTransform);
 
-				FTransform OffsetParentTransform = OffsetTransform * ParentCurrentGlobalTransform;
-				OffsetParentTransform.NormalizeRotation();
+					OffsetParentTransform = OffsetTransform * OffsetParentTransform;
+					OffsetParentTransform.NormalizeRotation();
+				}
 
 				// deal with different interpolation types
 				if (AdvancedSettings.InterpolationType == EConstraintInterpType::Average)
@@ -277,7 +289,6 @@ FRigUnit_ParentConstraint_Execute()
 				else if (AdvancedSettings.InterpolationType == EConstraintInterpType::Shortest)
 				{
 					FQuat MixedGlobalQuat = MixedGlobalTransform.GetRotation();
-					FQuat OffsetParentQuat = OffsetParentTransform.GetRotation();
 
 					if (MixedGlobalQuat == FQuat(0.0f, 0.0f, 0.0f, 0.0f))
 					{
@@ -286,6 +297,7 @@ FRigUnit_ParentConstraint_Execute()
 					else
 					{
 						const float Alpha = NormalizedWeight / AccumulatedWeight;
+						const FQuat OffsetParentQuat = OffsetParentTransform.GetRotation();
 
 						MixedGlobalTransform.LerpTranslationScale3D(MixedGlobalTransform, OffsetParentTransform, ScalarRegister(Alpha));
 						MixedGlobalQuat = FQuat::Slerp(MixedGlobalQuat, OffsetParentQuat, NormalizedWeight);
@@ -296,7 +308,7 @@ FRigUnit_ParentConstraint_Execute()
 				{
 					// invalid interpolation type
 					ensure(false);
-					MixedGlobalTransform = ChildCurrentGlobalTransform;
+					MixedGlobalTransform = Hierarchy->GetGlobalTransformByIndex(ChildCache, false);
 					break;
 				}
 			}
@@ -304,66 +316,73 @@ FRigUnit_ParentConstraint_Execute()
 			MixedGlobalTransform.NormalizeRotation();
 
 			// handle filtering, performed in local(parent) space
-			FTransform ChildParentGlobalTransform = Hierarchy->GetParentTransform(Child, false);
+			const FTransform ChildParentGlobalTransform = Hierarchy->GetParentTransformByIndex(ChildCache, false);
 			FTransform MixedLocalTransform = MixedGlobalTransform.GetRelativeTransform(ChildParentGlobalTransform);
-			MixedLocalTransform.NormalizeRotation();
-			FVector MixedTranslation = MixedLocalTransform.GetTranslation();
-			FQuat MixedRotation = MixedLocalTransform.GetRotation();
-			FVector MixedEulerRotation = AnimationCore::EulerFromQuat(MixedRotation, AdvancedSettings.RotationOrderForFilter);
-			FVector MixedScale = MixedLocalTransform.GetScale3D();
-
-			FTransform ChildCurrentLocalTransform = Hierarchy->GetLocalTransform(Child, false);
+			FTransform ChildCurrentLocalTransform = Hierarchy->GetLocalTransformByIndex(ChildCache, false);
 			
 			// Controls have an offset transform built-in and thus need to be handled a bit differently
-			FTransform AdditionalOffsetTransform = FTransform::Identity;
+			FTransform AdditionalOffsetTransform;
 			
-			if (Child.Type == ERigElementType::Control)
+			if (bChildIsControl)
 			{
-				if (FRigControlElement* ChildAsControlElement = Hierarchy->Find<FRigControlElement>(Child))
+				if (FRigControlElement* ChildAsControlElement = Hierarchy->Get<FRigControlElement>(ChildCache))
 				{
 					AdditionalOffsetTransform = Hierarchy->GetControlOffsetTransform(ChildAsControlElement, ERigTransformType::CurrentLocal);
 					// Control's local(parent) space transform = control local value * offset
 					ChildCurrentLocalTransform *= AdditionalOffsetTransform;
 				}
 			}
-			
-			FVector ChildTranslation = ChildCurrentLocalTransform.GetTranslation();
-			FQuat ChildRotation = ChildCurrentLocalTransform.GetRotation();
-			FVector ChildEulerRotation = AnimationCore::EulerFromQuat(ChildRotation, AdvancedSettings.RotationOrderForFilter);
-			FVector ChildScale = ChildCurrentLocalTransform.GetScale3D();
-
-			FVector FilteredTranslation;
-			FilteredTranslation.X = Filter.TranslationFilter.bX ? MixedTranslation.X : ChildTranslation.X;
-			FilteredTranslation.Y = Filter.TranslationFilter.bY ? MixedTranslation.Y : ChildTranslation.Y;
-			FilteredTranslation.Z = Filter.TranslationFilter.bZ ? MixedTranslation.Z : ChildTranslation.Z;
-
-			FVector FilteredEulerRotation;
-			FilteredEulerRotation.X = Filter.RotationFilter.bX ? MixedEulerRotation.X: ChildEulerRotation.X;
-			FilteredEulerRotation.Y = Filter.RotationFilter.bY ? MixedEulerRotation.Y : ChildEulerRotation.Y;
-			FilteredEulerRotation.Z = Filter.RotationFilter.bZ ? MixedEulerRotation.Z : ChildEulerRotation.Z;
 
-			FVector FilteredScale;
-			FilteredScale.X = Filter.ScaleFilter.bX ? MixedScale.X : ChildScale.X;
-			FilteredScale.Y = Filter.ScaleFilter.bY ? MixedScale.Y : ChildScale.Y;
-			FilteredScale.Z = Filter.ScaleFilter.bZ ? MixedScale.Z : ChildScale.Z;
+			if(!Filter.TranslationFilter.HasNoEffect())
+			{
+				const FVector ChildTranslation = ChildCurrentLocalTransform.GetTranslation();
+				FVector MixedTranslation = MixedLocalTransform.GetTranslation();
+				MixedTranslation.X = Filter.TranslationFilter.bX ? MixedTranslation.X : ChildTranslation.X;
+				MixedTranslation.Y = Filter.TranslationFilter.bY ? MixedTranslation.Y : ChildTranslation.Y;
+				MixedTranslation.Z = Filter.TranslationFilter.bZ ? MixedTranslation.Z : ChildTranslation.Z;
+				MixedLocalTransform.SetTranslation(MixedTranslation);
+			}
 
-			FTransform FilteredMixedLocalTransform(AnimationCore::QuatFromEuler(FilteredEulerRotation, AdvancedSettings.RotationOrderForFilter), FilteredTranslation, FilteredScale);
+			if(!Filter.RotationFilter.HasNoEffect())
+			{
+				const FQuat ChildRotation = ChildCurrentLocalTransform.GetRotation();
+				const FVector ChildEulerRotation = AnimationCore::EulerFromQuat(ChildRotation, AdvancedSettings.RotationOrderForFilter);
+				const FQuat MixedRotation = MixedLocalTransform.GetRotation();
+				FVector MixedEulerRotation = AnimationCore::EulerFromQuat(MixedRotation, AdvancedSettings.RotationOrderForFilter);
+				MixedEulerRotation.X = Filter.RotationFilter.bX ? MixedEulerRotation.X: ChildEulerRotation.X;
+				MixedEulerRotation.Y = Filter.RotationFilter.bY ? MixedEulerRotation.Y : ChildEulerRotation.Y;
+				MixedEulerRotation.Z = Filter.RotationFilter.bZ ? MixedEulerRotation.Z : ChildEulerRotation.Z;
+				MixedLocalTransform.SetRotation(AnimationCore::QuatFromEuler(MixedEulerRotation, AdvancedSettings.RotationOrderForFilter));
+				MixedLocalTransform.NormalizeRotation();
+			}
 
-			FTransform FinalLocalTransform = FilteredMixedLocalTransform;
+			if(!Filter.ScaleFilter.HasNoEffect())
+			{
+				const FVector ChildScale = ChildCurrentLocalTransform.GetScale3D();
+				FVector MixedScale = MixedLocalTransform.GetScale3D();
+				MixedScale.X = Filter.ScaleFilter.bX ? MixedScale.X : ChildScale.X;
+				MixedScale.Y = Filter.ScaleFilter.bY ? MixedScale.Y : ChildScale.Y;
+				MixedScale.Z = Filter.ScaleFilter.bZ ? MixedScale.Z : ChildScale.Z;
+				MixedLocalTransform.SetScale3D(MixedScale);
+			}
 
 			if (Weight < 1.0f - KINDA_SMALL_NUMBER)
 			{
-				FinalLocalTransform = FControlRigMathLibrary::LerpTransform(ChildCurrentLocalTransform, FinalLocalTransform, Weight);
+				MixedLocalTransform = FControlRigMathLibrary::LerpTransform(ChildCurrentLocalTransform, MixedLocalTransform, Weight);
+				if(!bChildIsControl)
+				{
+					MixedLocalTransform.NormalizeRotation();
+				}
 			}
 			
-			if (Child.Type == ERigElementType::Control)
+			if (bChildIsControl)
 			{
 				// need to convert back to offset space for the actual control value
-				FinalLocalTransform = FinalLocalTransform.GetRelativeTransform(AdditionalOffsetTransform);
-				FinalLocalTransform.NormalizeRotation();
+				MixedLocalTransform = MixedLocalTransform.GetRelativeTransform(AdditionalOffsetTransform);
+				MixedLocalTransform.NormalizeRotation();
 			}
 
-			Hierarchy->SetLocalTransform(Child, FinalLocalTransform);
+			Hierarchy->SetLocalTransformByIndex(ChildCache, MixedLocalTransform);
 		}
 	} 
  }
@@ -616,39 +635,37 @@ FRigUnit_PositionConstraintLocalSpaceOffset_Execute()
 	{
 		return;
 	}
-	
+
+	if(Context.State == EControlRigState::Init)
+	{
+		ChildCache.Reset();
+		ParentCaches.Reset();
+	}
+
 	URigHierarchy* Hierarchy = ExecuteContext.Hierarchy;	
 	if (Hierarchy)
 	{
-		if (!Child.IsValid())
+		if (!ChildCache.UpdateCache(Child, Hierarchy))
 		{
 			return;
 		}
-		
-		FTransform ChildInitialLocalTransform = Hierarchy->GetInitialLocalTransform(Child);
-		
-		// Controls need to be handled a bit differently
-		if (Child.Type == ERigElementType::Control)
+
+		if(ParentCaches.Num() != Parents.Num())
 		{
-			if (FRigControlElement* ChildAsControlElement = Hierarchy->Find<FRigControlElement>(Child))
-			{
-				FTransform AdditionalOffsetTransform = Hierarchy->GetControlOffsetTransform(ChildAsControlElement, ERigTransformType::InitialLocal);
-				// Control's local(parent) space position = local * offset
-				ChildInitialLocalTransform *= AdditionalOffsetTransform;
-			}
+			ParentCaches.SetNumZeroed(Parents.Num());
 		}
-
+		
 		float OverallWeight = 0;
-		for (const FConstraintParent& Parent : Parents)
+		for (int32 ParentIndex = 0; ParentIndex < Parents.Num(); ParentIndex++)
 		{
-			const float ClampedWeight = FMath::Max(Parent.Weight, 0.f);
-
-			if (ClampedWeight < KINDA_SMALL_NUMBER)
+			const FConstraintParent& Parent = Parents[ParentIndex];
+			if (!ParentCaches[ParentIndex].UpdateCache(Parent.Item, Hierarchy))
 			{
 				continue;
 			}
 
-			if (!Parent.Item.IsValid())
+			const float ClampedWeight = FMath::Max(Parent.Weight, 0.f);
+			if (ClampedWeight < KINDA_SMALL_NUMBER)
 			{
 				continue;
 			}
@@ -657,38 +674,58 @@ FRigUnit_PositionConstraintLocalSpaceOffset_Execute()
 		}
 
 		const float WeightNormalizer = 1.0f / OverallWeight;
-		
+
 		if (OverallWeight > KINDA_SMALL_NUMBER)
 		{
+			FTransform AdditionalOffsetTransform;
+			
+			const bool bChildIsControl = Child.Type == ERigElementType::Control;
+			if(bChildIsControl)
+			{
+				if (FRigControlElement* ChildAsControlElement = Hierarchy->Get<FRigControlElement>(ChildCache))
+				{
+					AdditionalOffsetTransform = Hierarchy->GetControlOffsetTransform(ChildAsControlElement, ERigTransformType::InitialLocal);
+				}
+			}
+
 			FVector OffsetPosition = FVector::ZeroVector;
 			if (bMaintainOffset)
 			{
 				FVector MixedInitialGlobalPosition = FVector::ZeroVector;
 
-				for (const FConstraintParent& Parent : Parents)
+				for (int32 ParentIndex = 0; ParentIndex < Parents.Num(); ParentIndex++)
 				{
-					const float ClampedWeight = FMath::Max(Parent.Weight, 0.f);
-
-					if (ClampedWeight < KINDA_SMALL_NUMBER)
+					if (!ParentCaches[ParentIndex].IsValid())
 					{
 						continue;
 					}
 
-					if (!Parent.Item.IsValid())
+					const FConstraintParent& Parent = Parents[ParentIndex];
+
+					const float ClampedWeight = FMath::Max(Parent.Weight, 0.f);
+					if (ClampedWeight < KINDA_SMALL_NUMBER)
 					{
 						continue;
 					}
 
 					const float NormalizedWeight = ClampedWeight * WeightNormalizer;
 
-					FTransform ParentInitialGlobalTransform = Hierarchy->GetGlobalTransform(Parent.Item, true);
+					FTransform ParentInitialGlobalTransform = Hierarchy->GetGlobalTransformByIndex(ParentCaches[ParentIndex], true);
 
 					MixedInitialGlobalPosition += ParentInitialGlobalTransform.GetLocation() * NormalizedWeight;
 				}
 				
-				FTransform ChildParentInitialGlobalTransform = Hierarchy->GetParentTransform(Child, true);
+				const FTransform ChildParentInitialGlobalTransform = Hierarchy->GetParentTransformByIndex(ChildCache, true);
 				FVector MixedInitialLocalPosition = ChildParentInitialGlobalTransform.Inverse().TransformPosition(MixedInitialGlobalPosition);
 
+				FTransform ChildInitialLocalTransform = Hierarchy->GetLocalTransformByIndex(ChildCache, true);
+
+				// Controls need to be handled a bit differently
+				if (bChildIsControl)
+				{
+					ChildInitialLocalTransform *= AdditionalOffsetTransform;
+				}
+
 				// use initial transforms to calculate the offset
 				// offset is applied in local space, at the end of constraint solve
 				OffsetPosition = ChildInitialLocalTransform.GetLocation() - MixedInitialLocalPosition;
@@ -696,76 +733,68 @@ FRigUnit_PositionConstraintLocalSpaceOffset_Execute()
 			
 			FVector MixedGlobalPosition = FVector::ZeroVector;
 
-			for (const FConstraintParent& Parent : Parents)
+			for (int32 ParentIndex = 0; ParentIndex < Parents.Num(); ParentIndex++)
 			{
-				const float ClampedWeight = FMath::Max(Parent.Weight, 0.f);
-
-				if (ClampedWeight < KINDA_SMALL_NUMBER)
+				if (!ParentCaches[ParentIndex].IsValid())
 				{
 					continue;
 				}
 
-				if (!Parent.Item.IsValid())
+				const FConstraintParent& Parent = Parents[ParentIndex];
+				
+				const float ClampedWeight = FMath::Max(Parent.Weight, 0.f);
+				if (ClampedWeight < KINDA_SMALL_NUMBER)
 				{
 					continue;
 				}
 
 				const float NormalizedWeight = ClampedWeight * WeightNormalizer;
 
-				FTransform ParentCurrentGlobalTransform = Hierarchy->GetGlobalTransform(Parent.Item, false);
+				const FTransform ParentCurrentGlobalTransform = Hierarchy->GetGlobalTransformByIndex(ParentCaches[ParentIndex], false);
 				MixedGlobalPosition += ParentCurrentGlobalTransform.GetLocation() * NormalizedWeight;
 			}
 			
 			// handle filtering, performed in local space
-			FTransform ChildParentGlobalTransform = Hierarchy->GetParentTransform(Child, false);
-			FVector MixedPosition = ChildParentGlobalTransform.Inverse().TransformPosition(MixedGlobalPosition);
+			FTransform ChildCurrentLocalTransform = Hierarchy->GetLocalTransform(Child);
+			
+			// Controls need to be handled a bit differently
+			if (bChildIsControl)
+			{
+				ChildCurrentLocalTransform *= AdditionalOffsetTransform;
+			}
 			
+			FTransform MixedTransform = ChildCurrentLocalTransform;
+
+			const FTransform ChildParentGlobalTransform = Hierarchy->GetParentTransformByIndex(ChildCache, false);
+			FVector MixedPosition = ChildParentGlobalTransform.InverseTransformPosition(MixedGlobalPosition);
 			if (bMaintainOffset)
 			{
 				MixedPosition = MixedPosition + OffsetPosition;	
 			}
 			
-			FTransform ChildCurrentLocalTransform = Hierarchy->GetLocalTransform(Child);
-			
-			// Controls need to be handled a bit differently
-			FTransform AdditionalOffsetTransform = FTransform::Identity;
-			
-			if (Child.Type == ERigElementType::Control)
+			if(!Filter.HasNoEffect())
 			{
-				if (FRigControlElement* ChildAsControlElement = Hierarchy->Find<FRigControlElement>(Child))
-				{
-					AdditionalOffsetTransform = Hierarchy->GetControlOffsetTransform(ChildAsControlElement, ERigTransformType::CurrentLocal);
-					// Control's local(parent) space position = local * offset
-					ChildCurrentLocalTransform *= AdditionalOffsetTransform;
-				}
+				FVector ChildPosition = ChildCurrentLocalTransform.GetTranslation();
+				MixedPosition.X = Filter.bX ? MixedPosition.X : ChildPosition.X;
+				MixedPosition.Y = Filter.bY ? MixedPosition.Y : ChildPosition.Y;
+				MixedPosition.Z = Filter.bZ ? MixedPosition.Z : ChildPosition.Z;
 			}
-			
-			FVector ChildPosition = ChildCurrentLocalTransform.GetTranslation();
-
-			FVector FilteredPosition;
-			FilteredPosition.X = Filter.bX ? MixedPosition.X : ChildPosition.X;
-			FilteredPosition.Y = Filter.bY ? MixedPosition.Y : ChildPosition.Y;
-			FilteredPosition.Z = Filter.bZ ? MixedPosition.Z : ChildPosition.Z;
-
-			FTransform FilteredMixedLocalTransform = ChildCurrentLocalTransform;
 
-			FilteredMixedLocalTransform.SetTranslation(FilteredPosition); 
-
-			FTransform FinalLocalTransform = FilteredMixedLocalTransform;
+			MixedTransform.SetTranslation(MixedPosition);
 			
 			if (Weight < 1.0f - KINDA_SMALL_NUMBER)
 			{
-				FinalLocalTransform = FControlRigMathLibrary::LerpTransform(ChildCurrentLocalTransform, FinalLocalTransform, Weight);
+				MixedTransform = FControlRigMathLibrary::LerpTransform(ChildCurrentLocalTransform, MixedTransform, Weight);
 			}
 			
 			if (Child.Type == ERigElementType::Control)
 			{
 				// need to convert back to offset space for the actual control value
-				FinalLocalTransform = FinalLocalTransform.GetRelativeTransform(AdditionalOffsetTransform);
-				FinalLocalTransform.NormalizeRotation();
+				MixedTransform = MixedTransform.GetRelativeTransform(AdditionalOffsetTransform);
+				MixedTransform.NormalizeRotation();
 			}
 
-			Hierarchy->SetLocalTransform(Child, FinalLocalTransform);
+			Hierarchy->SetLocalTransformByIndex(ChildCache, MixedTransform);
 		}
 	}
 }
@@ -1089,38 +1118,53 @@ FRigUnit_RotationConstraintLocalSpaceOffset_Execute()
 		return;
 	}
 	
+	if(Context.State == EControlRigState::Init)
+	{
+		ChildCache.Reset();
+		ParentCaches.Reset();
+	}
+
 	URigHierarchy* Hierarchy = ExecuteContext.Hierarchy;	
 	if (Hierarchy)
 	{
-		if (!Child.IsValid())
+		if (!ChildCache.UpdateCache(Child, Hierarchy))
 		{
 			return;
 		}
 
-		FTransform ChildInitialLocalTransform = Hierarchy->GetInitialLocalTransform(Child);
+		if(ParentCaches.Num() != Parents.Num())
+		{
+			ParentCaches.SetNumZeroed(Parents.Num());
+		}
+
+		FTransform ChildInitialLocalTransform = Hierarchy->GetLocalTransformByIndex(ChildCache, true);
 
 		// Controls need to be handled a bit differently
-		if (Child.Type == ERigElementType::Control)
+		FTransform AdditionalOffsetTransform;
+
+		const bool bChildIsControl = Child.Type == ERigElementType::Control;
+		if (bChildIsControl)
 		{
-			if (FRigControlElement* ChildAsControlElement = Hierarchy->Find<FRigControlElement>(Child))
+			if (FRigControlElement* ChildAsControlElement = Hierarchy->Get<FRigControlElement>(ChildCache))
 			{
-				FTransform AdditionalOffsetTransform = Hierarchy->GetControlOffsetTransform(ChildAsControlElement, ERigTransformType::InitialLocal);
+				AdditionalOffsetTransform = Hierarchy->GetControlOffsetTransform(ChildAsControlElement, ERigTransformType::CurrentLocal);
+				
 				// Control's local(parent) space position = local * offset
 				ChildInitialLocalTransform *= AdditionalOffsetTransform;
 			}
 		}
 		
 		float OverallWeight = 0;
-		for (const FConstraintParent& Parent : Parents)
+		for (int32 ParentIndex = 0; ParentIndex < Parents.Num(); ParentIndex++)
 		{
-			const float ClampedWeight = FMath::Max(Parent.Weight, 0.f);
-
-			if (ClampedWeight < KINDA_SMALL_NUMBER)
+			const FConstraintParent& Parent = Parents[ParentIndex];
+			if (!ParentCaches[ParentIndex].UpdateCache(Parent.Item, Hierarchy))
 			{
 				continue;
 			}
-
-			if (!Parent.Item.IsValid())
+			
+			const float ClampedWeight = FMath::Max(Parent.Weight, 0.f);
+			if (ClampedWeight < KINDA_SMALL_NUMBER)
 			{
 				continue;
 			}
@@ -1138,23 +1182,24 @@ FRigUnit_RotationConstraintLocalSpaceOffset_Execute()
 			{
 				FQuat MixedInitialGlobalRotation = FQuat(0,0,0,0);
 
-				for (const FConstraintParent& Parent : Parents)
+				for (int32 ParentIndex = 0; ParentIndex < Parents.Num(); ParentIndex++)
 				{
-					const float ClampedWeight = FMath::Max(Parent.Weight, 0.f);
-
-					if (ClampedWeight < KINDA_SMALL_NUMBER)
+					if(!ParentCaches[ParentIndex].IsValid())
 					{
 						continue;
 					}
-
-					if (!Parent.Item.IsValid())
+					
+					const FConstraintParent& Parent = Parents[ParentIndex];
+					
+					const float ClampedWeight = FMath::Max(Parent.Weight, 0.f);
+					if (ClampedWeight < KINDA_SMALL_NUMBER)
 					{
 						continue;
 					}
 
 					const float NormalizedWeight = ClampedWeight * WeightNormalizer;
 
-					FTransform ParentInitialGlobalTransform = Hierarchy->GetGlobalTransform(Parent.Item, true);
+					FTransform ParentInitialGlobalTransform = Hierarchy->GetGlobalTransformByIndex(ParentCaches[ParentIndex], true);
 					FQuat ParentInitialGlobalRotation = ParentInitialGlobalTransform.GetRotation();
 
 					// deal with different interpolation types
@@ -1188,7 +1233,7 @@ FRigUnit_RotationConstraintLocalSpaceOffset_Execute()
 					{
 						// invalid interpolation type
 						ensure(false);
-						FTransform ChildInitialGlobalTransform = Hierarchy->GetGlobalTransform(Child, true);
+						FTransform ChildInitialGlobalTransform = Hierarchy->GetGlobalTransformByIndex(ChildCache, true);
 						MixedInitialGlobalRotation = ChildInitialGlobalTransform.GetRotation();
 						break;
 					}
@@ -1198,32 +1243,33 @@ FRigUnit_RotationConstraintLocalSpaceOffset_Execute()
 				
 				// calculate the offset rotation that rotates the initial blended result
 				// back to the initial orientation of the child
-				FTransform ChildParentInitialGlobalTransform = Hierarchy->GetParentTransform(Child, true);
-				FQuat MixedInitialLocalRotation = ChildParentInitialGlobalTransform.GetRotation().Inverse() * MixedInitialGlobalRotation;
-				FQuat ChildInitialLocalRotation = ChildInitialLocalTransform.GetRotation();
+				const FTransform ChildParentInitialGlobalTransform = Hierarchy->GetParentTransformByIndex(ChildCache, true);
+				const FQuat MixedInitialLocalRotation = ChildParentInitialGlobalTransform.GetRotation().Inverse() * MixedInitialGlobalRotation;
+				const FQuat ChildInitialLocalRotation = ChildInitialLocalTransform.GetRotation();
 					
 				OffsetRotation = MixedInitialLocalRotation.Inverse() * ChildInitialLocalRotation;
 			}
 			
 			FQuat MixedGlobalRotation = FQuat(0,0,0,0);
 
-			for (const FConstraintParent& Parent : Parents)
+			for (int32 ParentIndex = 0; ParentIndex < Parents.Num(); ParentIndex++)
 			{
-				const float ClampedWeight = FMath::Max(Parent.Weight, 0.f);
-
-				if (ClampedWeight < KINDA_SMALL_NUMBER)
+				if(!ParentCaches[ParentIndex].IsValid())
 				{
 					continue;
 				}
-
-				if (!Parent.Item.IsValid())
+					
+				const FConstraintParent& Parent = Parents[ParentIndex];
+				
+				const float ClampedWeight = FMath::Max(Parent.Weight, 0.f);
+				if (ClampedWeight < KINDA_SMALL_NUMBER)
 				{
 					continue;
 				}
 
 				const float NormalizedWeight = ClampedWeight * WeightNormalizer;
 
-				FTransform ParentCurrentGlobalTransform = Hierarchy->GetGlobalTransform(Parent.Item, false);
+				FTransform ParentCurrentGlobalTransform = Hierarchy->GetGlobalTransformByIndex(ParentCaches[ParentIndex], false);
 				FQuat ParentCurrentGlobalRotation = ParentCurrentGlobalTransform.GetRotation();
 
 				// deal with different interpolation types
@@ -1257,7 +1303,7 @@ FRigUnit_RotationConstraintLocalSpaceOffset_Execute()
 				{
 					// invalid interpolation type
 					ensure(false);
-					FTransform ChildCurrentGlobalTransform = Hierarchy->GetGlobalTransform(Child, false);
+					FTransform ChildCurrentGlobalTransform = Hierarchy->GetGlobalTransformByIndex(ChildCache, false);
 					MixedGlobalRotation = ChildCurrentGlobalTransform.GetRotation();
 					break;
 				}	
@@ -1265,60 +1311,55 @@ FRigUnit_RotationConstraintLocalSpaceOffset_Execute()
 				
 			MixedGlobalRotation.Normalize();
 			
+			FTransform ChildCurrentLocalTransform = Hierarchy->GetLocalTransformByIndex(ChildCache, false);
+			
+			// Controls need to be handled a bit differently
+			if (bChildIsControl)
+			{
+				// Control's local(parent) space = local * offset
+				ChildCurrentLocalTransform *= AdditionalOffsetTransform;
+			}
+
 			// handle filtering, performed in local space
-			FTransform ChildParentGlobalTransform = Hierarchy->GetParentTransform(Child, false);
+			FTransform ChildParentGlobalTransform = Hierarchy->GetParentTransformByIndex(ChildCache, false);
 			FQuat MixedLocalRotation = ChildParentGlobalTransform.GetRotation().Inverse() * MixedGlobalRotation;
+			FTransform MixedLocalTransform = ChildCurrentLocalTransform;
 
 			if (bMaintainOffset)
 			{
 				MixedLocalRotation = MixedLocalRotation * OffsetRotation;
 			}
-			
-			FVector MixedEulerRotation = AnimationCore::EulerFromQuat(MixedLocalRotation, AdvancedSettings.RotationOrderForFilter);
-			
-			FTransform ChildCurrentLocalTransform = Hierarchy->GetLocalTransform(Child, false);
-			
-			// Controls need to be handled a bit differently
-			FTransform AdditionalOffsetTransform = FTransform::Identity;
 
-			if (Child.Type == ERigElementType::Control)
+			if(!Filter.HasNoEffect())
 			{
-				if (FRigControlElement* ChildAsControlElement = Hierarchy->Find<FRigControlElement>(Child))
-				{
-					AdditionalOffsetTransform = Hierarchy->GetControlOffsetTransform(ChildAsControlElement, ERigTransformType::CurrentLocal);
-					// Control's local(parent) space = local * offset
-					ChildCurrentLocalTransform *= AdditionalOffsetTransform;
-				}
-			}
-
-			FQuat ChildRotation = ChildCurrentLocalTransform.GetRotation();
-			
-			FVector ChildEulerRotation = AnimationCore::EulerFromQuat(ChildRotation, AdvancedSettings.RotationOrderForFilter);	
-			
-			FVector FilteredEulerRotation;
-			FilteredEulerRotation.X = Filter.bX ? MixedEulerRotation.X : ChildEulerRotation.X;
-			FilteredEulerRotation.Y = Filter.bY ? MixedEulerRotation.Y : ChildEulerRotation.Y;
-			FilteredEulerRotation.Z = Filter.bZ ? MixedEulerRotation.Z : ChildEulerRotation.Z; 
-
-			FTransform FilteredMixedLocalTransform = ChildCurrentLocalTransform;
+				FVector MixedEulerRotation = AnimationCore::EulerFromQuat(MixedLocalRotation, AdvancedSettings.RotationOrderForFilter);
 
-			FilteredMixedLocalTransform.SetRotation(AnimationCore::QuatFromEuler(FilteredEulerRotation, AdvancedSettings.RotationOrderForFilter));
-
-			FTransform FinalLocalTransform = FilteredMixedLocalTransform;
+				const FQuat ChildRotation = ChildCurrentLocalTransform.GetRotation();
+				const FVector ChildEulerRotation = AnimationCore::EulerFromQuat(ChildRotation, AdvancedSettings.RotationOrderForFilter);
+				
+				MixedEulerRotation.X = Filter.bX ? MixedEulerRotation.X : ChildEulerRotation.X;
+				MixedEulerRotation.Y = Filter.bY ? MixedEulerRotation.Y : ChildEulerRotation.Y;
+				MixedEulerRotation.Z = Filter.bZ ? MixedEulerRotation.Z : ChildEulerRotation.Z; 
+				MixedLocalTransform.SetRotation(AnimationCore::QuatFromEuler(MixedEulerRotation, AdvancedSettings.RotationOrderForFilter));
+			}
+			else
+			{
+				MixedLocalTransform.SetRotation(MixedLocalRotation);
+			}
 
 			if (Weight < 1.0f - KINDA_SMALL_NUMBER)
 			{
-				FinalLocalTransform = FControlRigMathLibrary::LerpTransform(ChildCurrentLocalTransform, FinalLocalTransform, Weight);
+				MixedLocalTransform = FControlRigMathLibrary::LerpTransform(ChildCurrentLocalTransform, MixedLocalTransform, Weight);
 			}
 
-			if (Child.Type == ERigElementType::Control)
+			if (bChildIsControl)
 			{
 				// need to convert back to offset space for the actual control value
-				FinalLocalTransform = FinalLocalTransform.GetRelativeTransform(AdditionalOffsetTransform);
-				FinalLocalTransform.NormalizeRotation();
+				MixedLocalTransform = MixedLocalTransform.GetRelativeTransform(AdditionalOffsetTransform);
+				MixedLocalTransform.NormalizeRotation();
 			}
 			
-			Hierarchy->SetLocalTransform(Child, FinalLocalTransform);
+			Hierarchy->SetLocalTransformByIndex(ChildCache, MixedLocalTransform);
 		}
 	} 	
 }
@@ -1690,37 +1731,51 @@ FRigUnit_ScaleConstraintLocalSpaceOffset_Execute()
 		return;
 	}
 	
+	if(Context.State == EControlRigState::Init)
+	{
+		ChildCache.Reset();
+		ParentCaches.Reset();
+	}
+
 	URigHierarchy* Hierarchy = ExecuteContext.Hierarchy;	
 	if (Hierarchy)
 	{
-		if (!Child.IsValid())
+		if (!ChildCache.UpdateCache(Child, Hierarchy))
 		{
 			return;
 		}
+
+		if(ParentCaches.Num() != Parents.Num())
+		{
+			ParentCaches.SetNumZeroed(Parents.Num());
+		}
 		
-		FTransform ChildInitialLocalTransform = Hierarchy->GetInitialLocalTransform(Child);
+		FTransform ChildInitialLocalTransform = Hierarchy->GetLocalTransformByIndex(ChildCache, true);
+		FTransform AdditionalOffsetTransform;
 		
 		// Controls need to be handled a bit differently
+		const bool bChildIsControl = Child.Type == ERigElementType::Control;
 		if (Child.Type == ERigElementType::Control)
 		{
-			if (FRigControlElement* ChildAsControlElement = Hierarchy->Find<FRigControlElement>(Child))
+			if (FRigControlElement* ChildAsControlElement = Hierarchy->Get<FRigControlElement>(ChildCache))
 			{
-				FTransform AdditionalOffsetTransform = Hierarchy->GetControlOffsetTransform(ChildAsControlElement, ERigTransformType::InitialLocal);
+				AdditionalOffsetTransform = Hierarchy->GetControlOffsetTransform(ChildAsControlElement, ERigTransformType::CurrentLocal);
+
 				// Control's local(parent) space position = local * offset
 				ChildInitialLocalTransform *= AdditionalOffsetTransform;
 			}
 		}
 		FVector::FReal OverallWeight = 0;
-		for (const FConstraintParent& Parent : Parents)
+		for (int32 ParentIndex = 0; ParentIndex < Parents.Num(); ParentIndex++)
 		{
-			const FVector::FReal ClampedWeight = FMath::Max<FVector::FReal>(Parent.Weight, 0.f);
-
-			if (ClampedWeight < KINDA_SMALL_NUMBER)
+			const FConstraintParent& Parent = Parents[ParentIndex];
+			if (!ParentCaches[ParentIndex].UpdateCache(Parent.Item, Hierarchy))
 			{
 				continue;
 			}
-
-			if (!Parent.Item.IsValid())
+			
+			const FVector::FReal ClampedWeight = FMath::Max<FVector::FReal>(Parent.Weight, 0.f);
+			if (ClampedWeight < KINDA_SMALL_NUMBER)
 			{
 				continue;
 			}
@@ -1737,10 +1792,16 @@ FRigUnit_ScaleConstraintLocalSpaceOffset_Execute()
 			{
 				FVector MixedInitialGlobalScale = FVector::OneVector;
 
-				for (const FConstraintParent& Parent : Parents)
+				for (int32 ParentIndex = 0; ParentIndex < Parents.Num(); ParentIndex++)
 				{
-					const FVector::FReal ClampedWeight = FMath::Max<FVector::FReal>(Parent.Weight, 0.f);
+					if (!ParentCaches[ParentIndex].IsValid())
+					{
+						continue;
+					}
+					
+					const FConstraintParent& Parent = Parents[ParentIndex];
 
+					const FVector::FReal ClampedWeight = FMath::Max<FVector::FReal>(Parent.Weight, 0.f);
 					if (ClampedWeight < KINDA_SMALL_NUMBER)
 					{
 						continue;
@@ -1753,8 +1814,8 @@ FRigUnit_ScaleConstraintLocalSpaceOffset_Execute()
 
 					const FVector::FReal NormalizedWeight = ClampedWeight * WeightNormalizer;
 
-					FTransform ParentInitialGlobalTransform = Hierarchy->GetGlobalTransform(Parent.Item, true);
-					FVector ParentInitialGlobalScale = ParentInitialGlobalTransform.GetScale3D();
+					const FTransform ParentInitialGlobalTransform = Hierarchy->GetGlobalTransformByIndex(ParentCaches[ParentIndex], true);
+					const FVector ParentInitialGlobalScale = ParentInitialGlobalTransform.GetScale3D();
 					
 					FVector WeightedInitialParentGlobalScale;
 					WeightedInitialParentGlobalScale.X = FMath::Pow(ParentInitialGlobalScale.X, NormalizedWeight);
@@ -1765,33 +1826,34 @@ FRigUnit_ScaleConstraintLocalSpaceOffset_Execute()
 				}
 				
 				// handle filtering, performed in local space
-				FTransform ChildParentInitialGlobalTransform = Hierarchy->GetParentTransform(Child, true);
-				FVector ChildParentInitialGlobalScale = ChildParentInitialGlobalTransform.GetScale3D();
-				FVector MixedInitialLocalScale = MixedInitialGlobalScale / GetNonZeroScale(ChildParentInitialGlobalScale);
+				const FTransform ChildParentInitialGlobalTransform = Hierarchy->GetParentTransformByIndex(ChildCache, true);
+				const FVector ChildParentInitialGlobalScale = ChildParentInitialGlobalTransform.GetScale3D();
+				const FVector MixedInitialLocalScale = MixedInitialGlobalScale / GetNonZeroScale(ChildParentInitialGlobalScale);
 			
-				 OffsetScale = ChildInitialLocalTransform.GetScale3D() / GetNonZeroScale(MixedInitialLocalScale);
+				OffsetScale = ChildInitialLocalTransform.GetScale3D() / GetNonZeroScale(MixedInitialLocalScale);
 			}
 
 			
 			FVector MixedGlobalScale = FVector::OneVector;
 
-			for (const FConstraintParent& Parent : Parents)
+			for (int32 ParentIndex = 0; ParentIndex < Parents.Num(); ParentIndex++)
 			{
-				const FVector::FReal ClampedWeight = FMath::Max<FVector::FReal>(Parent.Weight, 0.f);
-
-				if (ClampedWeight < KINDA_SMALL_NUMBER)
+				if (!ParentCaches[ParentIndex].IsValid())
 				{
 					continue;
 				}
-
-				if (!Parent.Item.IsValid())
+					
+				const FConstraintParent& Parent = Parents[ParentIndex];
+				
+				const FVector::FReal ClampedWeight = FMath::Max<FVector::FReal>(Parent.Weight, 0.f);
+				if (ClampedWeight < KINDA_SMALL_NUMBER)
 				{
 					continue;
 				}
 
 				const FVector::FReal NormalizedWeight = ClampedWeight * WeightNormalizer;
 
-				FTransform ParentCurrentGlobalTransform = Hierarchy->GetGlobalTransform(Parent.Item, false);
+				FTransform ParentCurrentGlobalTransform = Hierarchy->GetGlobalTransformByIndex(ParentCaches[ParentIndex], false);
 				FVector ParentCurrentGlobalScale = ParentCurrentGlobalTransform.GetScale3D();
 
 				FVector WeightedParentGlobalScale;
@@ -1806,7 +1868,6 @@ FRigUnit_ScaleConstraintLocalSpaceOffset_Execute()
 			FTransform ChildParentGlobalTransform = Hierarchy->GetParentTransform(Child, false);
 			FVector ChildParentGlobalScale = ChildParentGlobalTransform.GetScale3D();
 			FVector MixedLocalScale = MixedGlobalScale / GetNonZeroScale(ChildParentGlobalScale);
-
 			if (bMaintainOffset)
 			{
 				MixedLocalScale = MixedLocalScale * OffsetScale;
@@ -1815,44 +1876,35 @@ FRigUnit_ScaleConstraintLocalSpaceOffset_Execute()
 			FTransform ChildCurrentLocalTransform = Hierarchy->GetLocalTransform(Child, false);
 
 			// Controls need to be handled a bit differently
-			FTransform AdditionalOffsetTransform = FTransform::Identity;
-
-			if (Child.Type == ERigElementType::Control)
+			if (bChildIsControl)
 			{
-				if (FRigControlElement* ChildAsControlElement = Hierarchy->Find<FRigControlElement>(Child))
-				{
-					AdditionalOffsetTransform = Hierarchy->GetControlOffsetTransform(ChildAsControlElement, ERigTransformType::CurrentLocal);
-					// Control's local(parent) space = local * offset
-					ChildCurrentLocalTransform *= AdditionalOffsetTransform;
-				}
+				// Control's local(parent) space = local * offset
+				ChildCurrentLocalTransform *= AdditionalOffsetTransform;
 			}
-			
-			FVector ChildLocalScale = ChildCurrentLocalTransform.GetScale3D();
-
-			FVector FilteredLocalScale;
-			FilteredLocalScale.X = Filter.bX ? MixedLocalScale.X : ChildLocalScale.X;
-			FilteredLocalScale.Y = Filter.bY ? MixedLocalScale.Y : ChildLocalScale.Y;
-			FilteredLocalScale.Z = Filter.bZ ? MixedLocalScale.Z : ChildLocalScale.Z;
-
-			FTransform FilteredMixedLocalTransform = ChildCurrentLocalTransform;
-
-			FilteredMixedLocalTransform.SetScale3D(FilteredLocalScale);
 
-			FTransform FinalLocalTransform = FilteredMixedLocalTransform;
+			FTransform MixedLocalTransform = ChildCurrentLocalTransform;
+			if(!Filter.HasNoEffect())
+			{
+				const FVector ChildLocalScale = ChildCurrentLocalTransform.GetScale3D();
+				MixedLocalScale.X = Filter.bX ? MixedLocalScale.X : ChildLocalScale.X;
+				MixedLocalScale.Y = Filter.bY ? MixedLocalScale.Y : ChildLocalScale.Y;
+				MixedLocalScale.Z = Filter.bZ ? MixedLocalScale.Z : ChildLocalScale.Z;
+			}
+			MixedLocalTransform.SetScale3D(MixedLocalScale);
 
 			if (Weight < 1.0f - KINDA_SMALL_NUMBER)
 			{
-				FinalLocalTransform = FControlRigMathLibrary::LerpTransform(ChildCurrentLocalTransform, FinalLocalTransform, Weight);
+				MixedLocalTransform = FControlRigMathLibrary::LerpTransform(ChildCurrentLocalTransform, MixedLocalTransform, Weight);
 			}
 			
-			if (Child.Type == ERigElementType::Control)
+			if (bChildIsControl)
 			{
 				// need to convert back to offset space for the actual control value
-				FinalLocalTransform = FinalLocalTransform.GetRelativeTransform(AdditionalOffsetTransform);
-				FinalLocalTransform.NormalizeRotation();
+				MixedLocalTransform = MixedLocalTransform.GetRelativeTransform(AdditionalOffsetTransform);
+				MixedLocalTransform.NormalizeRotation();
 			}
 
-			Hierarchy->SetLocalTransform(Child, FinalLocalTransform);
+			Hierarchy->SetLocalTransform(Child, MixedLocalTransform);
 		}
 	}
 }
diff --git a/Engine/Plugins/Animation/ControlRig/Source/ControlRig/Private/Units/Highlevel/Hierarchy/RigUnit_TransformConstraint.h b/Engine/Plugins/Animation/ControlRig/Source/ControlRig/Private/Units/Highlevel/Hierarchy/RigUnit_TransformConstraint.h
--- a/Engine/Plugins/Animation/ControlRig/Source/ControlRig/Private/Units/Highlevel/Hierarchy/RigUnit_TransformConstraint.h
+++ b/Engine/Plugins/Animation/ControlRig/Source/ControlRig/Private/Units/Highlevel/Hierarchy/RigUnit_TransformConstraint.h
@@ -258,6 +258,8 @@ struct FRigUnit_ParentConstraint : public FRigUnit_HighlevelBaseMutable
         : Child(FRigElementKey(NAME_None, ERigElementType::Bone))
         , bMaintainOffset(true)
 		, Weight(1.0f)
+		, ChildCache()
+		, ParentCaches()
 	{
 		Parents.Add(FConstraintParent());
 	}
@@ -282,6 +284,12 @@ struct FRigUnit_ParentConstraint : public FRigUnit_HighlevelBaseMutable
 
 	UPROPERTY(EditAnywhere, Category = "Constraint", meta = (Input))
 	float Weight;
+
+	UPROPERTY()
+	FCachedRigElement ChildCache;
+
+	UPROPERTY()
+	TArray<FCachedRigElement> ParentCaches;
 };
 
 /**
@@ -334,6 +342,8 @@ struct FRigUnit_PositionConstraintLocalSpaceOffset : public FRigUnit_HighlevelBa
 		: Child(FRigElementKey(NAME_None, ERigElementType::Bone))
 		, bMaintainOffset(true)
 		, Weight(1.0f)
+		, ChildCache()
+		, ParentCaches()
 	{
 		Parents.Add(FConstraintParent());
 	}
@@ -355,6 +365,12 @@ struct FRigUnit_PositionConstraintLocalSpaceOffset : public FRigUnit_HighlevelBa
 
 	UPROPERTY(EditAnywhere, Category = "Constraint", meta = (Input))
 	float Weight;
+
+	UPROPERTY()
+	FCachedRigElement ChildCache;
+
+	UPROPERTY()
+	TArray<FCachedRigElement> ParentCaches;
 };
 
 
@@ -434,6 +450,8 @@ struct FRigUnit_RotationConstraintLocalSpaceOffset : public FRigUnit_HighlevelBa
 		: Child(FRigElementKey(NAME_None, ERigElementType::Bone))
 		, bMaintainOffset(true)
 		, Weight(1.0f)
+		, ChildCache()
+		, ParentCaches()
 	{
 		Parents.Add(FConstraintParent());
 	}
@@ -458,6 +476,12 @@ struct FRigUnit_RotationConstraintLocalSpaceOffset : public FRigUnit_HighlevelBa
 
 	UPROPERTY(EditAnywhere, Category = "Constraint", meta = (Input))
 	float Weight;
+
+	UPROPERTY()
+	FCachedRigElement ChildCache;
+
+	UPROPERTY()
+	TArray<FCachedRigElement> ParentCaches;
 };
 
 /**
@@ -510,6 +534,8 @@ struct FRigUnit_ScaleConstraintLocalSpaceOffset : public FRigUnit_HighlevelBaseM
 		: Child(FRigElementKey(NAME_None, ERigElementType::Bone))
 		, bMaintainOffset(true)
 		, Weight(1.0f)
+		, ChildCache()
+		, ParentCaches()
 	{
 		Parents.Add(FConstraintParent());
 	}
@@ -531,4 +557,10 @@ struct FRigUnit_ScaleConstraintLocalSpaceOffset : public FRigUnit_HighlevelBaseM
 
 	UPROPERTY(EditAnywhere, Category = "Constraint", meta = (Input))
 	float Weight;
+
+	UPROPERTY()
+	FCachedRigElement ChildCache;
+
+	UPROPERTY()
+	TArray<FCachedRigElement> ParentCaches;
 };
diff --git a/Engine/Source/Runtime/AnimationCore/Public/Constraint.h b/Engine/Source/Runtime/AnimationCore/Public/Constraint.h
--- a/Engine/Source/Runtime/AnimationCore/Public/Constraint.h
+++ b/Engine/Source/Runtime/AnimationCore/Public/Constraint.h
@@ -102,6 +102,12 @@ struct ANIMATIONCORE_API FFilterOptionPerAxis
 		// if none of them is set, it's not valid
 		return bX || bY || bZ;
 	}
+	
+	bool HasNoEffect() const
+	{
+		// if all of them are set the filter won't affect anything
+		return bX && bY && bZ;
+	}
 };
 
 /** A filter for a whole transform */
