Control Rig: Adapt include paths for nativization

diff --git a/Engine/Plugins/Animation/ControlRig/Source/ControlRig/Private/Units/Animation/RigUnit_AnimEasing.h b/Engine/Plugins/Animation/ControlRig/Source/ControlRig/Private/Units/Animation/RigUnit_AnimEasing.h
--- a/Engine/Plugins/Animation/ControlRig/Source/ControlRig/Private/Units/Animation/RigUnit_AnimEasing.h
+++ b/Engine/Plugins/Animation/ControlRig/Source/ControlRig/Private/Units/Animation/RigUnit_AnimEasing.h
@@ -2,7 +2,7 @@
 
 #pragma once
 
-#include "Units/Animation/RigUnit_AnimBase.h"
+#include "RigUnit_AnimBase.h"
 #include "Math/ControlRigMathLibrary.h"
 #include "RigUnit_AnimEasing.generated.h"
 
diff --git a/Engine/Plugins/Animation/ControlRig/Source/ControlRig/Private/Units/Animation/RigUnit_AnimEvalRichCurve.h b/Engine/Plugins/Animation/ControlRig/Source/ControlRig/Private/Units/Animation/RigUnit_AnimEvalRichCurve.h
--- a/Engine/Plugins/Animation/ControlRig/Source/ControlRig/Private/Units/Animation/RigUnit_AnimEvalRichCurve.h
+++ b/Engine/Plugins/Animation/ControlRig/Source/ControlRig/Private/Units/Animation/RigUnit_AnimEvalRichCurve.h
@@ -2,7 +2,7 @@
 
 #pragma once
 
-#include "Units/Animation/RigUnit_AnimBase.h"
+#include "RigUnit_AnimBase.h"
 #include "Curves/CurveFloat.h"
 #include "RigUnit_AnimEvalRichCurve.generated.h"
 
diff --git a/Engine/Plugins/Animation/ControlRig/Source/ControlRig/Private/Units/Animation/RigUnit_AnimRichCurve.h b/Engine/Plugins/Animation/ControlRig/Source/ControlRig/Private/Units/Animation/RigUnit_AnimRichCurve.h
--- a/Engine/Plugins/Animation/ControlRig/Source/ControlRig/Private/Units/Animation/RigUnit_AnimRichCurve.h
+++ b/Engine/Plugins/Animation/ControlRig/Source/ControlRig/Private/Units/Animation/RigUnit_AnimRichCurve.h
@@ -2,7 +2,7 @@
 
 #pragma once
 
-#include "Units/Animation/RigUnit_AnimBase.h"
+#include "RigUnit_AnimBase.h"
 #include "Curves/CurveFloat.h"
 #include "RigUnit_AnimRichCurve.generated.h"
 
diff --git a/Engine/Plugins/Animation/ControlRig/Source/ControlRig/Private/Units/Animation/RigUnit_GetDeltaTime.h b/Engine/Plugins/Animation/ControlRig/Source/ControlRig/Private/Units/Animation/RigUnit_GetDeltaTime.h
--- a/Engine/Plugins/Animation/ControlRig/Source/ControlRig/Private/Units/Animation/RigUnit_GetDeltaTime.h
+++ b/Engine/Plugins/Animation/ControlRig/Source/ControlRig/Private/Units/Animation/RigUnit_GetDeltaTime.h
@@ -2,7 +2,7 @@
 
 #pragma once
 
-#include "Units/Animation/RigUnit_AnimBase.h"
+#include "RigUnit_AnimBase.h"
 #include "RigUnit_GetDeltaTime.generated.h"
 
 /**
diff --git a/Engine/Plugins/Animation/ControlRig/Source/ControlRig/Private/Units/Animation/RigUnit_GetWorldTime.h b/Engine/Plugins/Animation/ControlRig/Source/ControlRig/Private/Units/Animation/RigUnit_GetWorldTime.h
--- a/Engine/Plugins/Animation/ControlRig/Source/ControlRig/Private/Units/Animation/RigUnit_GetWorldTime.h
+++ b/Engine/Plugins/Animation/ControlRig/Source/ControlRig/Private/Units/Animation/RigUnit_GetWorldTime.h
@@ -2,7 +2,7 @@
 
 #pragma once
 
-#include "Units/Animation/RigUnit_AnimBase.h"
+#include "RigUnit_AnimBase.h"
 #include "RigUnit_GetWorldTime.generated.h"
 
 /**
diff --git a/Engine/Plugins/Animation/ControlRig/Source/ControlRig/Private/Units/Animation/RigUnit_TimeConversion.h b/Engine/Plugins/Animation/ControlRig/Source/ControlRig/Private/Units/Animation/RigUnit_TimeConversion.h
--- a/Engine/Plugins/Animation/ControlRig/Source/ControlRig/Private/Units/Animation/RigUnit_TimeConversion.h
+++ b/Engine/Plugins/Animation/ControlRig/Source/ControlRig/Private/Units/Animation/RigUnit_TimeConversion.h
@@ -2,7 +2,7 @@
 
 #pragma once
 
-#include "Units/Animation/RigUnit_AnimBase.h"
+#include "RigUnit_AnimBase.h"
 #include "RigUnit_TimeConversion.generated.h"
 
 /**
diff --git a/Engine/Plugins/Animation/ControlRig/Source/ControlRig/Private/Units/Debug/RigUnit_DebugBezier.h b/Engine/Plugins/Animation/ControlRig/Source/ControlRig/Private/Units/Debug/RigUnit_DebugBezier.h
--- a/Engine/Plugins/Animation/ControlRig/Source/ControlRig/Private/Units/Debug/RigUnit_DebugBezier.h
+++ b/Engine/Plugins/Animation/ControlRig/Source/ControlRig/Private/Units/Debug/RigUnit_DebugBezier.h
@@ -2,7 +2,7 @@
 
 #pragma once
 
-#include "Units/Debug/RigUnit_DebugBase.h"
+#include "RigUnit_DebugBase.h"
 #include "Math/ControlRigMathLibrary.h"
 #include "RigUnit_DebugBezier.generated.h"
 
diff --git a/Engine/Plugins/Animation/ControlRig/Source/ControlRig/Private/Units/Debug/RigUnit_DebugHierarchy.h b/Engine/Plugins/Animation/ControlRig/Source/ControlRig/Private/Units/Debug/RigUnit_DebugHierarchy.h
--- a/Engine/Plugins/Animation/ControlRig/Source/ControlRig/Private/Units/Debug/RigUnit_DebugHierarchy.h
+++ b/Engine/Plugins/Animation/ControlRig/Source/ControlRig/Private/Units/Debug/RigUnit_DebugHierarchy.h
@@ -2,7 +2,7 @@
 
 #pragma once
 
-#include "Units/Debug/RigUnit_DebugBase.h"
+#include "RigUnit_DebugBase.h"
 #include "Drawing/ControlRigDrawInterface.h"
 #include "RigUnit_DebugHierarchy.generated.h"
 
diff --git a/Engine/Plugins/Animation/ControlRig/Source/ControlRig/Private/Units/Debug/RigUnit_DebugLine.h b/Engine/Plugins/Animation/ControlRig/Source/ControlRig/Private/Units/Debug/RigUnit_DebugLine.h
--- a/Engine/Plugins/Animation/ControlRig/Source/ControlRig/Private/Units/Debug/RigUnit_DebugLine.h
+++ b/Engine/Plugins/Animation/ControlRig/Source/ControlRig/Private/Units/Debug/RigUnit_DebugLine.h
@@ -2,7 +2,7 @@
 
 #pragma once
 
-#include "Units/Debug/RigUnit_DebugBase.h"
+#include "RigUnit_DebugBase.h"
 #include "RigUnit_DebugLine.generated.h"
 
 USTRUCT(meta=(DisplayName="Draw Line", Deprecated = "4.25"))
diff --git a/Engine/Plugins/Animation/ControlRig/Source/ControlRig/Private/Units/Debug/RigUnit_DebugLineStrip.h b/Engine/Plugins/Animation/ControlRig/Source/ControlRig/Private/Units/Debug/RigUnit_DebugLineStrip.h
--- a/Engine/Plugins/Animation/ControlRig/Source/ControlRig/Private/Units/Debug/RigUnit_DebugLineStrip.h
+++ b/Engine/Plugins/Animation/ControlRig/Source/ControlRig/Private/Units/Debug/RigUnit_DebugLineStrip.h
@@ -2,7 +2,7 @@
 
 #pragma once
 
-#include "Units/Debug/RigUnit_DebugBase.h"
+#include "RigUnit_DebugBase.h"
 #include "RigUnit_DebugLineStrip.generated.h"
 
 USTRUCT(meta=(DisplayName="Draw Line Strip", Deprecated = "4.25"))
diff --git a/Engine/Plugins/Animation/ControlRig/Source/ControlRig/Private/Units/Debug/RigUnit_DebugPoint.h b/Engine/Plugins/Animation/ControlRig/Source/ControlRig/Private/Units/Debug/RigUnit_DebugPoint.h
--- a/Engine/Plugins/Animation/ControlRig/Source/ControlRig/Private/Units/Debug/RigUnit_DebugPoint.h
+++ b/Engine/Plugins/Animation/ControlRig/Source/ControlRig/Private/Units/Debug/RigUnit_DebugPoint.h
@@ -2,7 +2,7 @@
 
 #pragma once
 
-#include "Units/Debug/RigUnit_DebugBase.h"
+#include "RigUnit_DebugBase.h"
 #include "RigUnit_DebugPoint.generated.h"
 
 UENUM()
diff --git a/Engine/Plugins/Animation/ControlRig/Source/ControlRig/Private/Units/Debug/RigUnit_DebugPrimitives.h b/Engine/Plugins/Animation/ControlRig/Source/ControlRig/Private/Units/Debug/RigUnit_DebugPrimitives.h
--- a/Engine/Plugins/Animation/ControlRig/Source/ControlRig/Private/Units/Debug/RigUnit_DebugPrimitives.h
+++ b/Engine/Plugins/Animation/ControlRig/Source/ControlRig/Private/Units/Debug/RigUnit_DebugPrimitives.h
@@ -2,7 +2,7 @@
 
 #pragma once
 
-#include "Units/Debug/RigUnit_DebugBase.h"
+#include "RigUnit_DebugBase.h"
 #include "RigUnit_DebugPrimitives.generated.h"
 
 USTRUCT(meta=(DisplayName="Draw Rectangle", Keywords="Draw Square", Deprecated = "4.25"))
diff --git a/Engine/Plugins/Animation/ControlRig/Source/ControlRig/Private/Units/Debug/RigUnit_DebugTransform.h b/Engine/Plugins/Animation/ControlRig/Source/ControlRig/Private/Units/Debug/RigUnit_DebugTransform.h
--- a/Engine/Plugins/Animation/ControlRig/Source/ControlRig/Private/Units/Debug/RigUnit_DebugTransform.h
+++ b/Engine/Plugins/Animation/ControlRig/Source/ControlRig/Private/Units/Debug/RigUnit_DebugTransform.h
@@ -2,7 +2,7 @@
 
 #pragma once
 
-#include "Units/Debug/RigUnit_DebugBase.h"
+#include "RigUnit_DebugBase.h"
 #include "RigUnit_DebugTransform.generated.h"
 
 UENUM()
diff --git a/Engine/Plugins/Animation/ControlRig/Source/ControlRig/Private/Units/Debug/RigUnit_ProfilingBracket.h b/Engine/Plugins/Animation/ControlRig/Source/ControlRig/Private/Units/Debug/RigUnit_ProfilingBracket.h
--- a/Engine/Plugins/Animation/ControlRig/Source/ControlRig/Private/Units/Debug/RigUnit_ProfilingBracket.h
+++ b/Engine/Plugins/Animation/ControlRig/Source/ControlRig/Private/Units/Debug/RigUnit_ProfilingBracket.h
@@ -2,7 +2,7 @@
 
 #pragma once
 
-#include "Units/Debug/RigUnit_DebugBase.h"
+#include "RigUnit_DebugBase.h"
 #include "RigUnit_ProfilingBracket.generated.h"
 
 /**
diff --git a/Engine/Plugins/Animation/ControlRig/Source/ControlRig/Private/Units/Math/RigUnit_MathBool.h b/Engine/Plugins/Animation/ControlRig/Source/ControlRig/Private/Units/Math/RigUnit_MathBool.h
--- a/Engine/Plugins/Animation/ControlRig/Source/ControlRig/Private/Units/Math/RigUnit_MathBool.h
+++ b/Engine/Plugins/Animation/ControlRig/Source/ControlRig/Private/Units/Math/RigUnit_MathBool.h
@@ -2,7 +2,7 @@
 
 #pragma once
 
-#include "Units/Math/RigUnit_MathBase.h"
+#include "RigUnit_MathBase.h"
 #include "RigUnit_MathBool.generated.h"
 
 USTRUCT(meta=(Abstract, Category="Math|Boolean"))
diff --git a/Engine/Plugins/Animation/ControlRig/Source/ControlRig/Private/Units/Math/RigUnit_MathColor.h b/Engine/Plugins/Animation/ControlRig/Source/ControlRig/Private/Units/Math/RigUnit_MathColor.h
--- a/Engine/Plugins/Animation/ControlRig/Source/ControlRig/Private/Units/Math/RigUnit_MathColor.h
+++ b/Engine/Plugins/Animation/ControlRig/Source/ControlRig/Private/Units/Math/RigUnit_MathColor.h
@@ -2,7 +2,7 @@
 
 #pragma once
 
-#include "Units/Math/RigUnit_MathBase.h"
+#include "RigUnit_MathBase.h"
 #include "Math/ControlRigMathLibrary.h"
 #include "RigUnit_MathColor.generated.h"
 
diff --git a/Engine/Plugins/Animation/ControlRig/Source/ControlRig/Private/Units/Math/RigUnit_MathInt.h b/Engine/Plugins/Animation/ControlRig/Source/ControlRig/Private/Units/Math/RigUnit_MathInt.h
--- a/Engine/Plugins/Animation/ControlRig/Source/ControlRig/Private/Units/Math/RigUnit_MathInt.h
+++ b/Engine/Plugins/Animation/ControlRig/Source/ControlRig/Private/Units/Math/RigUnit_MathInt.h
@@ -2,7 +2,7 @@
 
 #pragma once
 
-#include "Units/Math/RigUnit_MathBase.h"
+#include "RigUnit_MathBase.h"
 #include "RigUnit_MathInt.generated.h"
 
 USTRUCT(meta=(Abstract, Category="Math|Int", MenuDescSuffix="(Int)"))
diff --git a/Engine/Plugins/Animation/ControlRig/Source/ControlRig/Private/Units/Math/RigUnit_MathMatrix.h b/Engine/Plugins/Animation/ControlRig/Source/ControlRig/Private/Units/Math/RigUnit_MathMatrix.h
--- a/Engine/Plugins/Animation/ControlRig/Source/ControlRig/Private/Units/Math/RigUnit_MathMatrix.h
+++ b/Engine/Plugins/Animation/ControlRig/Source/ControlRig/Private/Units/Math/RigUnit_MathMatrix.h
@@ -3,7 +3,7 @@
 #pragma once
 
 #include "EulerTransform.h"
-#include "Units/Math/RigUnit_MathBase.h"
+#include "RigUnit_MathBase.h"
 #include "Math/ControlRigMathLibrary.h"
 #include "RigUnit_MathMatrix.generated.h"
 
diff --git a/Engine/Plugins/Animation/ControlRig/Source/ControlRig/Private/Units/Math/RigUnit_MathQuaternion.h b/Engine/Plugins/Animation/ControlRig/Source/ControlRig/Private/Units/Math/RigUnit_MathQuaternion.h
--- a/Engine/Plugins/Animation/ControlRig/Source/ControlRig/Private/Units/Math/RigUnit_MathQuaternion.h
+++ b/Engine/Plugins/Animation/ControlRig/Source/ControlRig/Private/Units/Math/RigUnit_MathQuaternion.h
@@ -2,7 +2,7 @@
 
 #pragma once
 
-#include "Units/Math/RigUnit_MathBase.h"
+#include "RigUnit_MathBase.h"
 #include "Math/ControlRigMathLibrary.h"
 #include "RigUnit_MathQuaternion.generated.h"
 
diff --git a/Engine/Plugins/Animation/ControlRig/Source/ControlRig/Private/Units/Math/RigUnit_MathRBFInterpolate.h b/Engine/Plugins/Animation/ControlRig/Source/ControlRig/Private/Units/Math/RigUnit_MathRBFInterpolate.h
--- a/Engine/Plugins/Animation/ControlRig/Source/ControlRig/Private/Units/Math/RigUnit_MathRBFInterpolate.h
+++ b/Engine/Plugins/Animation/ControlRig/Source/ControlRig/Private/Units/Math/RigUnit_MathRBFInterpolate.h
@@ -2,7 +2,7 @@
 
 #pragma once
 
-#include "Units/Math/RigUnit_MathBase.h"
+#include "RigUnit_MathBase.h"
 #include "RBF/RBFInterpolator.h"
 
 #include "RigUnit_MathRBFInterpolate.generated.h"
diff --git a/Engine/Plugins/Animation/ControlRig/Source/ControlRig/Private/Units/Math/RigUnit_MathTransform.h b/Engine/Plugins/Animation/ControlRig/Source/ControlRig/Private/Units/Math/RigUnit_MathTransform.h
--- a/Engine/Plugins/Animation/ControlRig/Source/ControlRig/Private/Units/Math/RigUnit_MathTransform.h
+++ b/Engine/Plugins/Animation/ControlRig/Source/ControlRig/Private/Units/Math/RigUnit_MathTransform.h
@@ -3,7 +3,7 @@
 #pragma once
 
 #include "EulerTransform.h"
-#include "Units/Math/RigUnit_MathBase.h"
+#include "RigUnit_MathBase.h"
 #include "Math/ControlRigMathLibrary.h"
 #include "RigUnit_MathTransform.generated.h"
 
diff --git a/Engine/Plugins/Animation/ControlRig/Source/ControlRig/Private/Units/Math/RigUnit_MathVector.h b/Engine/Plugins/Animation/ControlRig/Source/ControlRig/Private/Units/Math/RigUnit_MathVector.h
--- a/Engine/Plugins/Animation/ControlRig/Source/ControlRig/Private/Units/Math/RigUnit_MathVector.h
+++ b/Engine/Plugins/Animation/ControlRig/Source/ControlRig/Private/Units/Math/RigUnit_MathVector.h
@@ -2,7 +2,7 @@
 
 #pragma once
 
-#include "Units/Math/RigUnit_MathBase.h"
+#include "RigUnit_MathBase.h"
 #include "Math/ControlRigMathLibrary.h"
 #include "RigUnit_MathVector.generated.h"
 
diff --git a/Engine/Plugins/Animation/ControlRig/Source/ControlRig/Private/Units/Math/RigUnit_Noise.h b/Engine/Plugins/Animation/ControlRig/Source/ControlRig/Private/Units/Math/RigUnit_Noise.h
--- a/Engine/Plugins/Animation/ControlRig/Source/ControlRig/Private/Units/Math/RigUnit_Noise.h
+++ b/Engine/Plugins/Animation/ControlRig/Source/ControlRig/Private/Units/Math/RigUnit_Noise.h
@@ -2,7 +2,7 @@
 
 #pragma once
 
-#include "Units/Math/RigUnit_MathBase.h"
+#include "RigUnit_MathBase.h"
 #include "RigUnit_Noise.generated.h"
 
 /**
diff --git a/Engine/Plugins/Animation/ControlRig/Source/ControlRig/Private/Units/Math/RigUnit_Random.h b/Engine/Plugins/Animation/ControlRig/Source/ControlRig/Private/Units/Math/RigUnit_Random.h
--- a/Engine/Plugins/Animation/ControlRig/Source/ControlRig/Private/Units/Math/RigUnit_Random.h
+++ b/Engine/Plugins/Animation/ControlRig/Source/ControlRig/Private/Units/Math/RigUnit_Random.h
@@ -2,7 +2,7 @@
 
 #pragma once
 
-#include "Units/Math/RigUnit_MathBase.h"
+#include "RigUnit_MathBase.h"
 #include "RigUnit_Random.generated.h"
 
 /**
diff --git a/Engine/Plugins/Animation/ControlRig/Source/ControlRig/Private/Units/Simulation/RigUnit_Accumulate.h b/Engine/Plugins/Animation/ControlRig/Source/ControlRig/Private/Units/Simulation/RigUnit_Accumulate.h
--- a/Engine/Plugins/Animation/ControlRig/Source/ControlRig/Private/Units/Simulation/RigUnit_Accumulate.h
+++ b/Engine/Plugins/Animation/ControlRig/Source/ControlRig/Private/Units/Simulation/RigUnit_Accumulate.h
@@ -2,7 +2,7 @@
 
 #pragma once
 
-#include "Units/Simulation/RigUnit_SimBase.h"
+#include "RigUnit_SimBase.h"
 #include "RigUnit_Accumulate.generated.h"
 
 USTRUCT(meta=(Category="Simulation|Accumulate"))
diff --git a/Engine/Plugins/Animation/ControlRig/Source/ControlRig/Private/Units/Simulation/RigUnit_DeltaFromPrevious.h b/Engine/Plugins/Animation/ControlRig/Source/ControlRig/Private/Units/Simulation/RigUnit_DeltaFromPrevious.h
--- a/Engine/Plugins/Animation/ControlRig/Source/ControlRig/Private/Units/Simulation/RigUnit_DeltaFromPrevious.h
+++ b/Engine/Plugins/Animation/ControlRig/Source/ControlRig/Private/Units/Simulation/RigUnit_DeltaFromPrevious.h
@@ -2,7 +2,7 @@
 
 #pragma once
 
-#include "Units/Simulation/RigUnit_SimBase.h"
+#include "RigUnit_SimBase.h"
 #include "RigUnit_DeltaFromPrevious.generated.h"
 
 /**
diff --git a/Engine/Plugins/Animation/ControlRig/Source/ControlRig/Private/Units/Simulation/RigUnit_Kalman.h b/Engine/Plugins/Animation/ControlRig/Source/ControlRig/Private/Units/Simulation/RigUnit_Kalman.h
--- a/Engine/Plugins/Animation/ControlRig/Source/ControlRig/Private/Units/Simulation/RigUnit_Kalman.h
+++ b/Engine/Plugins/Animation/ControlRig/Source/ControlRig/Private/Units/Simulation/RigUnit_Kalman.h
@@ -2,7 +2,7 @@
 
 #pragma once
 
-#include "Units/Simulation/RigUnit_SimBase.h"
+#include "RigUnit_SimBase.h"
 #include "RigUnit_Kalman.generated.h"
 
 /**
diff --git a/Engine/Plugins/Animation/ControlRig/Source/ControlRig/Private/Units/Simulation/RigUnit_PointSimulation.h b/Engine/Plugins/Animation/ControlRig/Source/ControlRig/Private/Units/Simulation/RigUnit_PointSimulation.h
--- a/Engine/Plugins/Animation/ControlRig/Source/ControlRig/Private/Units/Simulation/RigUnit_PointSimulation.h
+++ b/Engine/Plugins/Animation/ControlRig/Source/ControlRig/Private/Units/Simulation/RigUnit_PointSimulation.h
@@ -2,7 +2,7 @@
 
 #pragma once
 
-#include "Units/Simulation/RigUnit_SimBase.h"
+#include "RigUnit_SimBase.h"
 #include "Math/ControlRigMathLibrary.h"
 #include "Math/Simulation/CRSimPoint.h"
 #include "Math/Simulation/CRSimLinearSpring.h"
diff --git a/Engine/Plugins/Animation/ControlRig/Source/ControlRig/Private/Units/Simulation/RigUnit_TimeOffset.h b/Engine/Plugins/Animation/ControlRig/Source/ControlRig/Private/Units/Simulation/RigUnit_TimeOffset.h
--- a/Engine/Plugins/Animation/ControlRig/Source/ControlRig/Private/Units/Simulation/RigUnit_TimeOffset.h
+++ b/Engine/Plugins/Animation/ControlRig/Source/ControlRig/Private/Units/Simulation/RigUnit_TimeOffset.h
@@ -2,7 +2,7 @@
 
 #pragma once
 
-#include "Units/Simulation/RigUnit_SimBase.h"
+#include "RigUnit_SimBase.h"
 #include "RigUnit_TimeOffset.generated.h"
 
 /**
diff --git a/Engine/Plugins/Animation/ControlRig/Source/ControlRig/Private/Units/Simulation/RigUnit_Timeline.h b/Engine/Plugins/Animation/ControlRig/Source/ControlRig/Private/Units/Simulation/RigUnit_Timeline.h
--- a/Engine/Plugins/Animation/ControlRig/Source/ControlRig/Private/Units/Simulation/RigUnit_Timeline.h
+++ b/Engine/Plugins/Animation/ControlRig/Source/ControlRig/Private/Units/Simulation/RigUnit_Timeline.h
@@ -2,7 +2,7 @@
 
 #pragma once
 
-#include "Units/Simulation/RigUnit_SimBase.h"
+#include "RigUnit_SimBase.h"
 #include "RigUnit_Timeline.generated.h"
 
 /**
diff --git a/Engine/Plugins/Animation/ControlRig/Source/ControlRig/Private/Units/Simulation/RigUnit_Verlet.h b/Engine/Plugins/Animation/ControlRig/Source/ControlRig/Private/Units/Simulation/RigUnit_Verlet.h
--- a/Engine/Plugins/Animation/ControlRig/Source/ControlRig/Private/Units/Simulation/RigUnit_Verlet.h
+++ b/Engine/Plugins/Animation/ControlRig/Source/ControlRig/Private/Units/Simulation/RigUnit_Verlet.h
@@ -2,7 +2,7 @@
 
 #pragma once
 
-#include "Units/Simulation/RigUnit_SimBase.h"
+#include "RigUnit_SimBase.h"
 #include "Math/Simulation/CRSimPoint.h"
 #include "RigUnit_Verlet.generated.h"
 
