Sequencer: Control Rig; Multi selection fixes

diff --git a/Engine/Plugins/Animation/ControlRig/Source/ControlRigEditor/Private/EditMode/SControlRigOutliner.cpp b/Engine/Plugins/Animation/ControlRig/Source/ControlRigEditor/Private/EditMode/SControlRigOutliner.cpp
--- a/Engine/Plugins/Animation/ControlRig/Source/ControlRigEditor/Private/EditMode/SControlRigOutliner.cpp
+++ b/Engine/Plugins/Animation/ControlRig/Source/ControlRigEditor/Private/EditMode/SControlRigOutliner.cpp
@@ -201,6 +201,8 @@ void SControlRigOutlinerItem::HandleControlSelected(UControlRig* Subject, FRigCo
 		TSharedPtr<FRigTreeElement> Found = HierarchyTreeView->GetTreeView()->FindElement(Key, HierarchyTreeView->GetTreeView()->GetRootElements()[RootIndex]);
 		if (Found.IsValid())
 		{
+			TGuardValue<bool> GuardRigHierarchyChanges(bIsChangingRigHierarchy, true);
+
 			HierarchyTreeView->GetTreeView()->SetItemSelection(Found, bSelected, ESelectInfo::OnNavigation);
 
 			TArray<TSharedPtr<FRigTreeElement>> SelectedItems = HierarchyTreeView->GetTreeView()->GetSelectedItems();
@@ -240,7 +242,6 @@ void SControlRigOutlinerItem::HandleSelectionChanged(TSharedPtr<FRigTreeElement>
 		check(Controller);
 
 		TGuardValue<bool> GuardRigHierarchyChanges(bIsChangingRigHierarchy, true);
-
 		const TArray<FRigElementKey> NewSelection = HierarchyTreeView->GetTreeView()->GetSelectedKeys();
 		if (!Controller->SetSelection(NewSelection))
 		{
diff --git a/Engine/Plugins/Animation/ControlRig/Source/ControlRigEditor/Private/Sequencer/ControlRigParameterTrackEditor.cpp b/Engine/Plugins/Animation/ControlRig/Source/ControlRigEditor/Private/Sequencer/ControlRigParameterTrackEditor.cpp
--- a/Engine/Plugins/Animation/ControlRig/Source/ControlRigEditor/Private/Sequencer/ControlRigParameterTrackEditor.cpp
+++ b/Engine/Plugins/Animation/ControlRig/Source/ControlRigEditor/Private/Sequencer/ControlRigParameterTrackEditor.cpp
@@ -2287,18 +2287,7 @@ void FControlRigParameterTrackEditor::SetUpEditModeIfNeeded(UControlRig* Control
 		}
 	}
 	else
-	{
-
-		TMap<UControlRig*, TArray<FRigElementKey>> SelectedControls;
-		ControlRigEditMode->GetAllSelectedControls(SelectedControls);
-		for (TPair<UControlRig*, TArray<FRigElementKey>>& Selection : SelectedControls)
-		{
-			if(Selection.Key)
-			{
-				Selection.Key->ClearControlSelection();
-			}
-		}
-		
+	{	
 		if (ControlRigEditMode->AddControlRigObject(ControlRig, GetSequencer()))
 		{
 			//force an evaluation, this will get the control rig setup so edit mode looks good.
