avoid using possible nullptr

diff --git a/Engine/Plugins/Animation/ControlRig/Source/ControlRigEditor/Private/Graph/SControlRigGraphNode.cpp b/Engine/Plugins/Animation/ControlRig/Source/ControlRigEditor/Private/Graph/SControlRigGraphNode.cpp
--- a/Engine/Plugins/Animation/ControlRig/Source/ControlRigEditor/Private/Graph/SControlRigGraphNode.cpp
+++ b/Engine/Plugins/Animation/ControlRig/Source/ControlRigEditor/Private/Graph/SControlRigGraphNode.cpp
@@ -757,8 +757,7 @@ EVisibility SControlRigGraphNode::GetArrayPlusButtonVisibility(URigVMPin* InMode
 {
 	if(InModelPin)
 	{
-		URigVMNode* RigVMNode = InModelPin->GetNode();
-		if (RigVMNode->IsA<URigVMFunctionReturnNode>())
+		if (Cast<URigVMFunctionReturnNode>(InModelPin->GetNode()))
 		{
 			return EVisibility::Hidden;
 		}
