Control Rig: Add a rig element curve value validity flag. When the CR anim node is evaluated, we copy curve values and their validity flag into the element flag, and when evaluation is done we copy the same out. Updated the GetCurveValue to also return the validity flag, along with the value, and a new node, UnsetCurveValue, which can be used to clear a curve value, for example if we don't want a curve value to be captured by sequencer baking.

diff --git a/Engine/Plugins/Animation/ControlRig/Source/ControlRig/Private/AnimNode_ControlRigBase.cpp b/Engine/Plugins/Animation/ControlRig/Source/ControlRig/Private/AnimNode_ControlRigBase.cpp
--- a/Engine/Plugins/Animation/ControlRig/Source/ControlRig/Private/AnimNode_ControlRigBase.cpp
+++ b/Engine/Plugins/Animation/ControlRig/Source/ControlRig/Private/AnimNode_ControlRigBase.cpp
@@ -231,8 +231,16 @@ void FAnimNode_ControlRigBase::UpdateInput(UControlRig* ControlRig, const FPoseC
 				const uint16 ControlRigIndex = Pair.Key;
 				const uint16 SkeletonIndex = Pair.Value;
 
-				const float Value = InOutput.Curve.Get(SkeletonIndex);
-				ControlRig->GetHierarchy()->SetCurveValueByIndex(ControlRigIndex, Value);
+				bool bIsValid;
+				const float Value = InOutput.Curve.Get(SkeletonIndex, bIsValid);
+				if (bIsValid)
+				{
+					ControlRig->GetHierarchy()->SetCurveValueByIndex(ControlRigIndex, Value);
+				}
+				else
+				{
+					ControlRig->GetHierarchy()->UnsetCurveValueByIndex(ControlRigIndex);
+				}
 			}
 		}
 		else
@@ -243,7 +251,16 @@ void FAnimNode_ControlRigBase::UpdateInput(UControlRig* ControlRig, const FPoseC
 				const uint16 SkeletonIndex = Iter.Value();
 				const FRigElementKey Key(Name, ERigElementType::Curve);
 
-				ControlRig->GetHierarchy()->SetCurveValue(Key, InOutput.Curve.Get(SkeletonIndex));
+				bool bIsValid;
+				const float Value = InOutput.Curve.Get(SkeletonIndex, bIsValid);
+				if (bIsValid)
+				{
+					ControlRig->GetHierarchy()->SetCurveValue(Key, Value);
+				}
+				else
+				{
+					ControlRig->GetHierarchy()->UnsetCurveValue(Key);
+				}
 			}
 		}
 	}
@@ -364,14 +381,17 @@ void FAnimNode_ControlRigBase::UpdateOutput(UControlRig* ControlRig, FPoseContex
 				const uint16 ControlRigIndex = Pair.Key;
 				const uint16 SkeletonIndex = Pair.Value;
 
-				const float PreviousValue = InOutput.Curve.Get(SkeletonIndex);
-				const float Value = ControlRig->GetHierarchy()->GetCurveValueByIndex(ControlRigIndex);
-
-				if(!FMath::IsNearlyEqual(PreviousValue, Value))
+				if (ControlRig->GetHierarchy()->IsCurveValueSetByIndex(ControlRigIndex))
 				{
-					// this causes a side effect of marking the curve as "valid"
-					// so only apply it for curves that have really changed
-					InOutput.Curve.Set(SkeletonIndex, Value);
+					InOutput.Curve.Set(SkeletonIndex, ControlRig->GetHierarchy()->GetCurveValueByIndex(ControlRigIndex));
+				}
+				else
+				{
+					const int32 WeightIndex = InOutput.Curve.GetArrayIndexByUID(SkeletonIndex);
+					if (WeightIndex != INDEX_NONE)
+					{
+						InOutput.Curve.ValidCurveWeights[WeightIndex] = false;
+					}
 				}
 			}
 		}
@@ -383,14 +403,17 @@ void FAnimNode_ControlRigBase::UpdateOutput(UControlRig* ControlRig, FPoseContex
 				const uint16 Index = Iter.Value();
 				const FRigElementKey Key(Name, ERigElementType::Curve);
 
-				const float PreviousValue = InOutput.Curve.Get(Index);
-				const float Value = ControlRig->GetHierarchy()->GetCurveValue(Key);
-
-				if(!FMath::IsNearlyEqual(PreviousValue, Value))
+				if (ControlRig->GetHierarchy()->IsCurveValueSet(Key))
 				{
-					// this causes a side effect of marking the curve as "valid"
-					// so only apply it for curves that have really changed
-					InOutput.Curve.Set(Index, Value);
+					InOutput.Curve.Set(Index, ControlRig->GetHierarchy()->GetCurveValue(Key));
+				}
+				else
+				{
+					const int32 WeightIndex = InOutput.Curve.GetArrayIndexByUID(Index);
+					if (WeightIndex != INDEX_NONE)
+					{
+						InOutput.Curve.ValidCurveWeights[WeightIndex] = false;
+					}
 				}
 			}
 		}
@@ -429,7 +452,21 @@ void FAnimNode_ControlRigBase::Evaluate_AnyThread(FPoseContext& Output)
 		if (FAnimWeight::IsFullWeight(InternalBlendAlpha))
 		{
 			ExecuteControlRig(SourcePose);
+			
+			FSmartName SM;
+			if (SourcePose.AnimInstanceProxy->GetSkeleton()->GetSmartNameByName(USkeleton::AnimCurveMappingName, FName("CTRL_expressions_eyeLookDownL"), SM))
+			{
+				FPlatformMisc::LowLevelOutputDebugStringf(TEXT("SC %g [%d]\n"),
+					SourcePose.Curve.CurveWeights[SM.UID], int32(SourcePose.Curve.ValidCurveWeights[SM.UID]));
+			}
+			
 			Output = SourcePose;
+
+			if (Output.AnimInstanceProxy->GetSkeleton()->GetSmartNameByName(USkeleton::AnimCurveMappingName, FName("CTRL_expressions_eyeLookDownL"), SM))
+			{
+				FPlatformMisc::LowLevelOutputDebugStringf(TEXT("OC %g [%d]\n"),
+					Output.Curve.CurveWeights[SM.UID], int32(Output.Curve.ValidCurveWeights[SM.UID]));
+			}
 		}
 		else 
 		{
diff --git a/Engine/Plugins/Animation/ControlRig/Source/ControlRig/Private/Rigs/RigHierarchy.cpp b/Engine/Plugins/Animation/ControlRig/Source/ControlRig/Private/Rigs/RigHierarchy.cpp
--- a/Engine/Plugins/Animation/ControlRig/Source/ControlRig/Private/Rigs/RigHierarchy.cpp
+++ b/Engine/Plugins/Animation/ControlRig/Source/ControlRig/Private/Rigs/RigHierarchy.cpp
@@ -3179,15 +3179,23 @@ void URigHierarchy::SetControlVisibility(FRigControlElement* InControlElement, b
 #endif
 }
 
+
 float URigHierarchy::GetCurveValue(FRigCurveElement* InCurveElement) const
 {
 	if(InCurveElement == nullptr)
 	{
 		return 0.f;
 	}
-	return InCurveElement->Value;
+	return InCurveElement->bIsValueSet ? InCurveElement->Value : 0.f;
 }
 
+
+bool URigHierarchy::IsCurveValueSet(FRigCurveElement* InCurveElement) const
+{
+	return InCurveElement && InCurveElement->bIsValueSet;
+}
+
+
 void URigHierarchy::SetCurveValue(FRigCurveElement* InCurveElement, float InValue, bool bSetupUndo, bool bForce)
 {
 	LLM_SCOPE_BYNAME(TEXT("Animation/ControlRig"));
@@ -3196,18 +3204,20 @@ void URigHierarchy::SetCurveValue(FRigCurveElement* InCurveElement, float InValu
 		return;
 	}
 
+	const bool bPreviousIsValueSet = InCurveElement->bIsValueSet; 
 	const float PreviousValue = InCurveElement->Value;
-	if(!bForce && FMath::IsNearlyZero(PreviousValue - InValue))
+	if(!bForce && InCurveElement->bIsValueSet && FMath::IsNearlyZero(PreviousValue - InValue))
 	{
 		return;
 	}
 
+	InCurveElement->bIsValueSet = true;
 	InCurveElement->Value = InValue;
 
 #if WITH_EDITOR
 	if(bSetupUndo || IsTracingChanges())
 	{
-		PushCurveToStack(InCurveElement->GetKey(), PreviousValue, InCurveElement->Value, bSetupUndo);
+		PushCurveToStack(InCurveElement->GetKey(), PreviousValue, InCurveElement->Value, bPreviousIsValueSet, true, bSetupUndo);
 	}
 
 	if (ensure(!bPropagatingChange))
@@ -3221,12 +3231,11 @@ void URigHierarchy::SetCurveValue(FRigCurveElement* InCurveElement, float InValu
 				continue;
 			}
 
-			URigHierarchy* ListeningHierarchy = Listener.Hierarchy.Get();
-			if (ListeningHierarchy)
+			if (URigHierarchy* ListeningHierarchy = Listener.Hierarchy.Get())
 			{
 				if(FRigCurveElement* ListeningElement = Cast<FRigCurveElement>(ListeningHierarchy->Find(InCurveElement->GetKey())))
 				{
-					// bSetupUndo = false such that all listening hierarchies performs undo at the same time the root hierachy undos
+					// bSetupUndo = false such that all listening hierarchies performs undo at the same time the root hierarchy undoes
 					ListeningHierarchy->SetCurveValue(ListeningElement, InValue, false, bForce);
 				}
 			}
@@ -3235,6 +3244,54 @@ void URigHierarchy::SetCurveValue(FRigCurveElement* InCurveElement, float InValu
 #endif
 }
 
+
+void URigHierarchy::UnsetCurveValue(FRigCurveElement* InCurveElement, bool bSetupUndo, bool bForce)
+{
+	LLM_SCOPE_BYNAME(TEXT("Animation/ControlRig"));
+	if(InCurveElement == nullptr)
+	{
+		return;
+	}
+
+	const bool bPreviousIsValueSet = InCurveElement->bIsValueSet; 
+	if(!bForce && !InCurveElement->bIsValueSet)
+	{
+		return;
+	}
+
+	InCurveElement->bIsValueSet = false;
+
+#if WITH_EDITOR
+	if(bSetupUndo || IsTracingChanges())
+	{
+		PushCurveToStack(InCurveElement->GetKey(), InCurveElement->Value, InCurveElement->Value, bPreviousIsValueSet, false, bSetupUndo);
+	}
+
+	if (ensure(!bPropagatingChange))
+	{
+		TGuardValue<bool> bPropagatingChangeGuardValue(bPropagatingChange, true);
+			
+		for(FRigHierarchyListener& Listener : ListeningHierarchies)
+		{
+			if(!Listener.Hierarchy.IsValid())
+			{
+				continue;
+			}
+
+			if (URigHierarchy* ListeningHierarchy = Listener.Hierarchy.Get())
+			{
+				if(FRigCurveElement* ListeningElement = Cast<FRigCurveElement>(ListeningHierarchy->Find(InCurveElement->GetKey())))
+				{
+					// bSetupUndo = false such that all listening hierarchies performs undo at the same time the root hierarchy undoes
+					ListeningHierarchy->UnsetCurveValue(ListeningElement, false, bForce);
+				}
+			}
+		}
+	}
+#endif
+}
+
+
 FName URigHierarchy::GetPreviousName(const FRigElementKey& InKey) const
 {
 	if(const FRigElementKey* OldKeyPtr = PreviousNameMap.Find(InKey))
@@ -4475,15 +4532,15 @@ void URigHierarchy::PushTransformToStack(const FRigElementKey& InKey, ERigTransf
 #endif
 }
 
-void URigHierarchy::PushCurveToStack(const FRigElementKey& InKey, float InOldCurveValue, float InNewCurveValue, bool bModify)
+void URigHierarchy::PushCurveToStack(const FRigElementKey& InKey, float InOldCurveValue, float InNewCurveValue, bool bInOldIsCurveValueSet, bool bInNewIsCurveValueSet, bool bModify)
 {
 #if WITH_EDITOR
 
 	FTransform OldTransform = FTransform::Identity;
 	FTransform NewTransform = FTransform::Identity;
 
-	OldTransform.SetTranslation(FVector(InOldCurveValue, 0.f, 0.f));
-	NewTransform.SetTranslation(FVector(InNewCurveValue, 0.f, 0.f));
+	OldTransform.SetTranslation(FVector(InOldCurveValue, bInOldIsCurveValueSet ? 1.f : 0.f, 0.f));
+	NewTransform.SetTranslation(FVector(InNewCurveValue, bInNewIsCurveValueSet ? 1.f : 0.f, 0.f));
 
 	PushTransformToStack(InKey, ERigTransformStackEntryType::CurveValue, ERigTransformType::CurrentLocal, OldTransform, NewTransform, false, bModify);
 
diff --git a/Engine/Plugins/Animation/ControlRig/Source/ControlRig/Private/Rigs/RigHierarchyElements.cpp b/Engine/Plugins/Animation/ControlRig/Source/ControlRig/Private/Rigs/RigHierarchyElements.cpp
--- a/Engine/Plugins/Animation/ControlRig/Source/ControlRig/Private/Rigs/RigHierarchyElements.cpp
+++ b/Engine/Plugins/Animation/ControlRig/Source/ControlRig/Private/Rigs/RigHierarchyElements.cpp
@@ -760,6 +760,7 @@ void FRigCurveElement::Save(FArchive& Ar, URigHierarchy* Hierarchy, ESerializati
 
 	if(SerializationPhase == ESerializationPhase::StaticData)
 	{
+		Ar << bIsValueSet;
 		Ar << Value;
 	}
 }
@@ -770,6 +771,14 @@ void FRigCurveElement::Load(FArchive& Ar, URigHierarchy* Hierarchy, ESerializati
 
 	if(SerializationPhase == ESerializationPhase::StaticData)
 	{
+		if (Ar.CustomVer(FControlRigObjectVersion::GUID) >= FControlRigObjectVersion::CurveElementValueStateFlag)
+		{
+			Ar << bIsValueSet;
+		}
+		else
+		{
+			bIsValueSet = true;
+		}
 		Ar << Value;
 	}
 }
@@ -778,8 +787,9 @@ void FRigCurveElement::CopyPose(FRigBaseElement* InOther, bool bCurrent, bool bI
 {
 	Super::CopyPose(InOther, bCurrent, bInitial, bWeights);
 	
-	if(FRigCurveElement* Other = Cast<FRigCurveElement>(InOther))
+	if(const FRigCurveElement* Other = Cast<FRigCurveElement>(InOther))
 	{
+		bIsValueSet = Other->bIsValueSet;
 		Value = Other->Value;
 	}
 }
@@ -787,7 +797,12 @@ void FRigCurveElement::CopyPose(FRigBaseElement* InOther, bool bCurrent, bool bI
 void FRigCurveElement::CopyFrom(URigHierarchy* InHierarchy, FRigBaseElement* InOther, URigHierarchy* InOtherHierarchy)
 {
 	Super::CopyFrom(InHierarchy, InOther, InOtherHierarchy);
-	Value = CastChecked<FRigCurveElement>(InOther)->Value;
+	
+	if(const FRigCurveElement* Other = CastChecked<FRigCurveElement>(InOther))
+	{
+		bIsValueSet = Other->bIsValueSet;
+		Value = Other->Value;
+	}
 }
 
 ////////////////////////////////////////////////////////////////////////////////
diff --git a/Engine/Plugins/Animation/ControlRig/Source/ControlRig/Private/Units/Hierarchy/RigUnit_GetCurveValue.cpp b/Engine/Plugins/Animation/ControlRig/Source/ControlRig/Private/Units/Hierarchy/RigUnit_GetCurveValue.cpp
--- a/Engine/Plugins/Animation/ControlRig/Source/ControlRig/Private/Units/Hierarchy/RigUnit_GetCurveValue.cpp
+++ b/Engine/Plugins/Animation/ControlRig/Source/ControlRig/Private/Units/Hierarchy/RigUnit_GetCurveValue.cpp
@@ -6,8 +6,7 @@
 FRigUnit_GetCurveValue_Execute()
 {
     DECLARE_SCOPE_HIERARCHICAL_COUNTER_RIGUNIT()
-	URigHierarchy* Hierarchy = Context.Hierarchy;
-	if (Hierarchy)
+    if (const URigHierarchy* Hierarchy = Context.Hierarchy)
 	{
 		switch (Context.State)
 		{
@@ -19,7 +18,8 @@ FRigUnit_GetCurveValue_Execute()
 			{
 				if(CachedCurveIndex.UpdateCache(FRigElementKey(Curve, ERigElementType::Curve), Hierarchy))
 				{
-					Value = Hierarchy->GetCurveValue(CachedCurveIndex);
+					Valid = Hierarchy->IsCurveValueSetByIndex(CachedCurveIndex);
+					Value = Hierarchy->GetCurveValueByIndex(CachedCurveIndex);
 				}
 			}
 			default:
diff --git a/Engine/Plugins/Animation/ControlRig/Source/ControlRig/Private/Units/Hierarchy/RigUnit_GetCurveValue.h b/Engine/Plugins/Animation/ControlRig/Source/ControlRig/Private/Units/Hierarchy/RigUnit_GetCurveValue.h
--- a/Engine/Plugins/Animation/ControlRig/Source/ControlRig/Private/Units/Hierarchy/RigUnit_GetCurveValue.h
+++ b/Engine/Plugins/Animation/ControlRig/Source/ControlRig/Private/Units/Hierarchy/RigUnit_GetCurveValue.h
@@ -15,6 +15,7 @@ struct CONTROLRIG_API FRigUnit_GetCurveValue : public FRigUnit
 
 	FRigUnit_GetCurveValue()
 		: Curve(NAME_None)
+		, Valid(true)
 		, Value(0.f)
 		, CachedCurveIndex(FCachedRigElement())
 	{}
@@ -27,6 +28,9 @@ struct CONTROLRIG_API FRigUnit_GetCurveValue : public FRigUnit
 	 */
 	UPROPERTY(meta = (Input, CustomWidget = "CurveName"))
 	FName Curve;
+	
+	UPROPERTY(meta=(Output)) 
+	bool Valid;
 
 	// The current transform of the given Curve - or identity in case it wasn't found.
 	UPROPERTY(meta=(Output, UIMin = 0.f, UIMax = 1.f))
diff --git a/Engine/Plugins/Animation/ControlRig/Source/ControlRig/Private/Units/Hierarchy/RigUnit_SetCurveValue.cpp b/Engine/Plugins/Animation/ControlRig/Source/ControlRig/Private/Units/Hierarchy/RigUnit_SetCurveValue.cpp
--- a/Engine/Plugins/Animation/ControlRig/Source/ControlRig/Private/Units/Hierarchy/RigUnit_SetCurveValue.cpp
+++ b/Engine/Plugins/Animation/ControlRig/Source/ControlRig/Private/Units/Hierarchy/RigUnit_SetCurveValue.cpp
@@ -21,7 +21,7 @@ FRigUnit_SetCurveValue_Execute()
 				const FRigElementKey Key(Curve, ERigElementType::Curve);
 				if (CachedCurveIndex.UpdateCache(Key, Hierarchy))
 				{
-					Hierarchy->SetCurveValue(CachedCurveIndex, Value);
+					Hierarchy->SetCurveValueByIndex(CachedCurveIndex, Value);
 				}
 			}
 			default:
diff --git a/Engine/Plugins/Animation/ControlRig/Source/ControlRig/Private/Units/Hierarchy/RigUnit_UnsetCurveValue.cpp b/Engine/Plugins/Animation/ControlRig/Source/ControlRig/Private/Units/Hierarchy/RigUnit_UnsetCurveValue.cpp
new file mode 100644
--- /dev/null
+++ b/Engine/Plugins/Animation/ControlRig/Source/ControlRig/Private/Units/Hierarchy/RigUnit_UnsetCurveValue.cpp
@@ -0,0 +1,51 @@
+﻿// Copyright Epic Games, Inc. All Rights Reserved.
+
+#include "RigUnit_UnsetCurveValue.h"
+#include "Units/RigUnitContext.h"
+
+FRigUnit_UnsetCurveValue_Execute()
+{
+    DECLARE_SCOPE_HIERARCHICAL_COUNTER_RIGUNIT()
+	URigHierarchy* Hierarchy = ExecuteContext.Hierarchy;
+	if (Hierarchy)
+	{
+		switch (Context.State)
+		{
+			case EControlRigState::Init:
+			{
+				CachedCurveIndex.Reset();
+				// fall through to update
+			}
+			case EControlRigState::Update:
+			{
+				const FRigElementKey Key(Curve, ERigElementType::Curve);
+				if (CachedCurveIndex.UpdateCache(Key, Hierarchy))
+				{
+					Hierarchy->UnsetCurveValueByIndex(CachedCurveIndex);
+				}
+			}
+			default:
+			{
+				break;
+			}
+		}
+	}
+}
+
+#if WITH_DEV_AUTOMATION_TESTS
+#include "Units/RigUnitTest.h"
+
+IMPLEMENT_RIGUNIT_AUTOMATION_TEST(FRigUnit_UnsetCurveValue)
+{
+	const FRigElementKey Curve = Controller->AddCurve(TEXT("Curve"), 0.f);
+	Unit.ExecuteContext.Hierarchy = Hierarchy;
+	
+	Hierarchy->ResetCurveValues();
+	Unit.Curve = TEXT("Curve");
+	InitAndExecute();
+
+	AddErrorIfFalse(!Hierarchy->IsCurveValueSet(Curve), TEXT("unexpected value"));
+
+	return true;
+}
+#endif
diff --git a/Engine/Plugins/Animation/ControlRig/Source/ControlRig/Private/Units/Hierarchy/RigUnit_UnsetCurveValue.h b/Engine/Plugins/Animation/ControlRig/Source/ControlRig/Private/Units/Hierarchy/RigUnit_UnsetCurveValue.h
new file mode 100644
--- /dev/null
+++ b/Engine/Plugins/Animation/ControlRig/Source/ControlRig/Private/Units/Hierarchy/RigUnit_UnsetCurveValue.h
@@ -0,0 +1,34 @@
+﻿// Copyright Epic Games, Inc. All Rights Reserved.
+
+#pragma once
+
+#include "Units/RigUnit.h"
+#include "RigUnit_UnsetCurveValue.generated.h"
+
+
+/**
+ * UnsetCurveValue is used to perform a change in the curve container by invalidating a single Curve value.
+ */
+USTRUCT(meta=(DisplayName="Unset Curve Value", Category="Curve", Keywords = "UnsetCurveValue", NodeColor = "0.0 0.36470600962638855 1.0"))
+struct FRigUnit_UnsetCurveValue : public FRigUnitMutable
+{
+	GENERATED_BODY()
+
+	FRigUnit_UnsetCurveValue()
+		: CachedCurveIndex(FCachedRigElement())
+	{}
+
+	RIGVM_METHOD()
+	virtual void Execute(const FRigUnitContext& Context) override;
+
+	/**
+	 * The name of the Curve to set the Value for.
+	 */
+	UPROPERTY(meta = (Input, CustomWidget = "CurveName"))
+	FName Curve;
+
+private:
+	// Used to cache the internally used curve index
+	UPROPERTY()
+	FCachedRigElement CachedCurveIndex;
+};
diff --git a/Engine/Plugins/Animation/ControlRig/Source/ControlRig/Public/ControlRigObjectVersion.h b/Engine/Plugins/Animation/ControlRig/Source/ControlRig/Public/ControlRigObjectVersion.h
--- a/Engine/Plugins/Animation/ControlRig/Source/ControlRig/Public/ControlRigObjectVersion.h
+++ b/Engine/Plugins/Animation/ControlRig/Source/ControlRig/Public/ControlRigObjectVersion.h
@@ -57,6 +57,9 @@ struct CONTROLRIG_API FControlRigObjectVersion
 		// Deprecation of parameters
 		RemoveParameters,
 
+		// Added rig curve element value state flag
+		CurveElementValueStateFlag,
+
 		// -----<new versions can be added above this line>-------------------------------------------------
 		VersionPlusOne,
 		LatestVersion = VersionPlusOne - 1
diff --git a/Engine/Plugins/Animation/ControlRig/Source/ControlRig/Public/Rigs/RigHierarchy.h b/Engine/Plugins/Animation/ControlRig/Source/ControlRig/Public/Rigs/RigHierarchy.h
--- a/Engine/Plugins/Animation/ControlRig/Source/ControlRig/Public/Rigs/RigHierarchy.h
+++ b/Engine/Plugins/Animation/ControlRig/Source/ControlRig/Public/Rigs/RigHierarchy.h
@@ -1365,11 +1365,41 @@ public:
 		return 0.f;
 	}
 
+	// TODO: Deprecate?
 	FORCEINLINE_DEBUGGABLE float GetCurveValue(int32 InElementIndex) const
 	{
 		return GetCurveValueByIndex(InElementIndex);
 	}
 
+	/**
+	 * Returns whether a curve's value is set, given its key
+	 * @param InKey The key of the element to retrieve the value for
+	 * @return Returns true if the value is set, false otherwise.
+	 */
+	UFUNCTION(BlueprintCallable, Category = URigHierarchy)
+	FORCEINLINE_DEBUGGABLE bool IsCurveValueSet(FRigElementKey InKey) const
+	{
+		return IsCurveValueSetByIndex(GetIndex(InKey));
+	}
+
+	/**
+	 * Returns a curve's value given its index
+	 * @param InElementIndex The index of the element to retrieve the value for
+	 * @return Returns true if the value is set, false otherwise.
+	 */
+	UFUNCTION(BlueprintCallable, Category = URigHierarchy)
+	FORCEINLINE_DEBUGGABLE bool IsCurveValueSetByIndex(int32 InElementIndex) const
+	{
+		if(Elements.IsValidIndex(InElementIndex))
+		{
+			if(FRigCurveElement* CurveElement = Cast<FRigCurveElement>(Elements[InElementIndex]))
+			{
+				return IsCurveValueSet(CurveElement);
+			}
+		}
+		return false;
+	}
+	
 	/**
 	 * Sets a curve's value given its key
 	 * @param InKey The key of the element to set the value for
@@ -1400,11 +1430,42 @@ public:
 		}
 	}
 
+	// TODO: Deprecate?
 	FORCEINLINE_DEBUGGABLE void SetCurveValue(int32 InElementIndex, float InValue, bool bSetupUndo = false)
 	{
 		SetCurveValueByIndex(InElementIndex, InValue, bSetupUndo);
 	}
 
+	/**
+	 * Sets a curve's value given its key
+	 * @param InKey The key of the element to set the value for
+	 * @param InValue The value to set on the curve
+	 * @param bSetupUndo If true the transform stack will be setup for undo / redo
+	 */
+	UFUNCTION(BlueprintCallable, Category = URigHierarchy)
+    FORCEINLINE_DEBUGGABLE void UnsetCurveValue(FRigElementKey InKey, bool bSetupUndo = false)
+	{
+		UnsetCurveValueByIndex(GetIndex(InKey), bSetupUndo);
+	}
+
+	/**
+	 * Sets a curve's value given its index
+	 * @param InElementIndex The index of the element to set the value for
+	 * @param InValue The value to set on the curve
+	 * @param bSetupUndo If true the transform stack will be setup for undo / redo
+	 */
+	UFUNCTION(BlueprintCallable, Category = URigHierarchy)
+	FORCEINLINE_DEBUGGABLE void UnsetCurveValueByIndex(int32 InElementIndex, bool bSetupUndo = false)
+	{
+		if(Elements.IsValidIndex(InElementIndex))
+		{
+			if(FRigCurveElement* CurveElement = Cast<FRigCurveElement>(Elements[InElementIndex]))
+			{
+				UnsetCurveValue(CurveElement, bSetupUndo);
+			}
+		}
+	}
+
 	/**
 	 * Sets the offset transform for a given control element by key
 	 * @param InKey The key of the control element to set the offset transform for
@@ -2420,12 +2481,19 @@ public:
 	void SetControlVisibility(FRigControlElement* InControlElement, bool bVisibility);
 
 	/**
-	 * Returns a curve's value
+	 * Returns a curve's value. If the curve value is not set, returns 
 	 * @param InCurveElement The element to retrieve the value for
 	 * @return Returns the value of the curve
 	 */
 	float GetCurveValue(FRigCurveElement* InCurveElement) const;
 
+	/**
+	 * Returns whether a curve's value is set. If the curve value is not set, returns false. 
+	 * @param InCurveElement The element to retrieve the value for
+	 * @return Returns true if the value is set, false otherwise.
+	 */
+	bool IsCurveValueSet(FRigCurveElement* InCurveElement) const;
+
 	/**
 	 * Sets a curve's value
 	 * @param InCurveElement The element to set the value for
@@ -2435,6 +2503,14 @@ public:
 	 */
 	void SetCurveValue(FRigCurveElement* InCurveElement, float InValue, bool bSetupUndo = false, bool bForce = false);
 
+	/**
+	 * Unsets a curve's value. Basically the curve's value becomes meaningless.
+	 * @param InCurveElement The element to set the value for
+	 * @param bSetupUndo If true the transform stack will be setup for undo / redo
+	 * @param bForce Unset the curve even if it was already unset.
+	 */
+	void UnsetCurveValue(FRigCurveElement* InCurveElement, bool bSetupUndo = false, bool bForce = false);
+
 	/**
 	 * Returns the previous name of an element prior to a rename operation
 	 * @param InKey The key of the element to request the old name for
@@ -2675,6 +2751,8 @@ private:
             const FRigElementKey& InKey,
             float InOldCurveValue,
             float InNewCurveValue,
+            bool bInOldIsCurveValueSet,
+            bool bInNewIsCurveValueSet,
             bool bModify);
 
 	/**
diff --git a/Engine/Plugins/Animation/ControlRig/Source/ControlRig/Public/Rigs/RigHierarchyElements.h b/Engine/Plugins/Animation/ControlRig/Source/ControlRig/Public/Rigs/RigHierarchyElements.h
--- a/Engine/Plugins/Animation/ControlRig/Source/ControlRig/Public/Rigs/RigHierarchyElements.h
+++ b/Engine/Plugins/Animation/ControlRig/Source/ControlRig/Public/Rigs/RigHierarchyElements.h
@@ -1043,12 +1043,13 @@ public:
 
 	FRigCurveElement()
 		: FRigBaseElement()
+		, bIsValueSet(true)
 		, Value(0.f)
 	{
 		Key.Type = ERigElementType::Curve;
 	}
 
-	virtual ~FRigCurveElement(){}
+	virtual ~FRigCurveElement() override {}
 
 	virtual void Save(FArchive& A, URigHierarchy* Hierarchy, ESerializationPhase SerializationPhase) override;
 	virtual void Load(FArchive& Ar, URigHierarchy* Hierarchy, ESerializationPhase SerializationPhase) override;
@@ -1058,9 +1059,12 @@ private:
 	virtual void CopyFrom(URigHierarchy* InHierarchy, FRigBaseElement* InOther, URigHierarchy* InOtherHierarchy) override;
 
 public:
+	// Set to true if the value was actually set. Used to carry back and forth blend curve
+	// value validity state.
+	bool bIsValueSet;
 	
 	float Value;
-	
+
 	FORCEINLINE static bool IsClassOf(const FRigBaseElement* InElement)
 	{
 		return InElement->GetType() == ERigElementType::Curve;
