Control Rig: Remove stray debug logging.

diff --git a/Engine/Plugins/Animation/ControlRig/Source/ControlRig/Private/AnimNode_ControlRigBase.cpp b/Engine/Plugins/Animation/ControlRig/Source/ControlRig/Private/AnimNode_ControlRigBase.cpp
--- a/Engine/Plugins/Animation/ControlRig/Source/ControlRig/Private/AnimNode_ControlRigBase.cpp
+++ b/Engine/Plugins/Animation/ControlRig/Source/ControlRig/Private/AnimNode_ControlRigBase.cpp
@@ -452,21 +452,7 @@ void FAnimNode_ControlRigBase::Evaluate_AnyThread(FPoseContext& Output)
 		if (FAnimWeight::IsFullWeight(InternalBlendAlpha))
 		{
 			ExecuteControlRig(SourcePose);
-			
-			FSmartName SM;
-			if (SourcePose.AnimInstanceProxy->GetSkeleton()->GetSmartNameByName(USkeleton::AnimCurveMappingName, FName("CTRL_expressions_eyeLookDownL"), SM))
-			{
-				FPlatformMisc::LowLevelOutputDebugStringf(TEXT("SC %g [%d]\n"),
-					SourcePose.Curve.CurveWeights[SM.UID], int32(SourcePose.Curve.ValidCurveWeights[SM.UID]));
-			}
-			
 			Output = SourcePose;
-
-			if (Output.AnimInstanceProxy->GetSkeleton()->GetSmartNameByName(USkeleton::AnimCurveMappingName, FName("CTRL_expressions_eyeLookDownL"), SM))
-			{
-				FPlatformMisc::LowLevelOutputDebugStringf(TEXT("OC %g [%d]\n"),
-					Output.Curve.CurveWeights[SM.UID], int32(Output.Curve.ValidCurveWeights[SM.UID]));
-			}
 		}
 		else 
 		{
