Control Rig: Fix upgrade node for GetChildren

diff --git a/Engine/Plugins/Animation/ControlRig/Source/ControlRig/Private/Units/Execution/RigUnit_Collection.cpp b/Engine/Plugins/Animation/ControlRig/Source/ControlRig/Private/Units/Execution/RigUnit_Collection.cpp
--- a/Engine/Plugins/Animation/ControlRig/Source/ControlRig/Private/Units/Execution/RigUnit_Collection.cpp
+++ b/Engine/Plugins/Animation/ControlRig/Source/ControlRig/Private/Units/Execution/RigUnit_Collection.cpp
@@ -122,7 +122,9 @@ FRigVMStructUpgradeInfo FRigUnit_CollectionChildren::GetUpgradeInfo() const
 	NewNode.bRecursive = bRecursive;
 	NewNode.TypeToSearch = TypeToSearch;
 
-	return FRigVMStructUpgradeInfo(*this, NewNode);
+	FRigVMStructUpgradeInfo Info(*this, NewNode);
+	Info.AddRemappedPin(TEXT("Collection.Keys"), TEXT("Items"));
+	return Info;
 }
 
 FRigUnit_CollectionChildrenArray_Execute()
