Control Rig: Fix Local variable not working correctly

diff --git a/Engine/Plugins/Animation/ControlRig/Source/ControlRigEditor/Private/Editor/SControlRigStackView.cpp b/Engine/Plugins/Animation/ControlRig/Source/ControlRigEditor/Private/Editor/SControlRigStackView.cpp
--- a/Engine/Plugins/Animation/ControlRig/Source/ControlRigEditor/Private/Editor/SControlRigStackView.cpp
+++ b/Engine/Plugins/Animation/ControlRig/Source/ControlRigEditor/Private/Editor/SControlRigStackView.cpp
@@ -87,60 +87,75 @@ void SRigStackItem::Construct(const FArguments& InArgs, const TSharedRef<STableV
 		[
 			SNew(SHorizontalBox)
 			+ SHorizontalBox::Slot()
-			.MaxWidth(25.f)
-			.VAlign(VAlign_Center)
-			.HAlign(HAlign_Right)
-			[
-				SAssignNew(NumberWidget, STextBlock)
-				.Text(this, &SRigStackItem::GetIndexText)
-				.Font(IDetailLayoutBuilder::GetDetailFont())
-			]
-			+ SHorizontalBox::Slot()
-			.MaxWidth(22.f)
-			.VAlign(VAlign_Center)
-			.HAlign(HAlign_Center)
-			[
-				SNew(SImage)
-				.Image(Icon)
-			]
-			
-			+ SHorizontalBox::Slot()
-			.FillWidth(1.f)
-			.MaxWidth(200.f)
+			.AutoWidth()
 			.VAlign(VAlign_Center)
 			.HAlign(HAlign_Left)
 			[
-				SAssignNew(TextWidget, STextBlock)
-				.Text(this, &SRigStackItem::GetLabelText)
-				.Font(this, &SRigStackItem::GetLabelFont)
+				SNew(SHorizontalBox)
+				+ SHorizontalBox::Slot()
+				.MaxWidth(25.f)
+				.VAlign(VAlign_Center)
+				.HAlign(HAlign_Right)
+				[
+					SAssignNew(NumberWidget, STextBlock)
+					.Text(this, &SRigStackItem::GetIndexText)
+					.Font(IDetailLayoutBuilder::GetDetailFont())
+				]
+				+ SHorizontalBox::Slot()
+				.MaxWidth(22.f)
+				.VAlign(VAlign_Center)
+				.HAlign(HAlign_Center)
+				[
+					SNew(SImage)
+					.Image(Icon)
+				]
+				
+				+ SHorizontalBox::Slot()
+				.FillWidth(1.f)
+				.AutoWidth()
+				.VAlign(VAlign_Center)
+				.HAlign(HAlign_Fill)
+				[
+					SAssignNew(TextWidget, STextBlock)
+					.Text(this, &SRigStackItem::GetLabelText)
+					.Font(this, &SRigStackItem::GetLabelFont)
+					.Justification(ETextJustify::Left)
+				]
 			]
 
 			+ SHorizontalBox::Slot()
 			.FillWidth(1.f)
-			.MaxWidth(40.f)
 			[
 				SNew(SSpacer)
 			]
 
 			+ SHorizontalBox::Slot()
-			.MaxWidth(36.f)
-            .VAlign(VAlign_Center)
-            .HAlign(HAlign_Right)
-            [
-                SNew(STextBlock)
-                .Text(this, &SRigStackItem::GetVisitedCountText)
-                .Font(IDetailLayoutBuilder::GetDetailFont())
-            ]
-
-			+ SHorizontalBox::Slot()
-			.MaxWidth(80.f)
-            .VAlign(VAlign_Center)
-            .HAlign(HAlign_Right)
-            [
-                SNew(STextBlock)
-                .Text(this, &SRigStackItem::GetDurationText)
-                .Font(IDetailLayoutBuilder::GetDetailFont())
-            ]
+			.AutoWidth()
+			.VAlign(VAlign_Center)
+			.HAlign(HAlign_Right)
+			[
+				SNew(SHorizontalBox)
+				+ SHorizontalBox::Slot()
+				.AutoWidth()
+				.VAlign(VAlign_Center)
+	            .HAlign(HAlign_Left)
+	            [
+	                SNew(STextBlock)
+	                .Text(this, &SRigStackItem::GetVisitedCountText)
+	                .Font(IDetailLayoutBuilder::GetDetailFont())
+	            ]
+	            
+	            + SHorizontalBox::Slot()
+	            .Padding(20, 0, 0, 0)
+				.AutoWidth()
+	            .VAlign(VAlign_Center)
+	            .HAlign(HAlign_Left)
+	            [
+	                SNew(STextBlock)
+	                .Text(this, &SRigStackItem::GetDurationText)
+	                .Font(IDetailLayoutBuilder::GetDetailFont())
+	            ]
+	        ]
         ], OwnerTable);
 }
 
diff --git a/Engine/Source/Developer/RigVMDeveloper/Private/RigVMCompiler/RigVMCompiler.cpp b/Engine/Source/Developer/RigVMDeveloper/Private/RigVMCompiler/RigVMCompiler.cpp
--- a/Engine/Source/Developer/RigVMDeveloper/Private/RigVMCompiler/RigVMCompiler.cpp
+++ b/Engine/Source/Developer/RigVMDeveloper/Private/RigVMCompiler/RigVMCompiler.cpp
@@ -1989,6 +1989,11 @@ FString URigVMCompiler::GetPinHash(const URigVMPin* InPin, const FRigVMVarExprAS
 					{
 						return FString::Printf(TEXT("%sLocalVariableDefault::%s|%s%s"), *Prefix, *Node->GetGraph()->GetGraphName(), *VariableName.ToString(), *Suffix);
 					}
+					else if (InVarExpr)
+					{
+						const FString GraphPath = InVarExpr->GetProxy().GetCallstack().GetCallPath(false);
+						return FString::Printf(TEXT("%sLocalVariable::%s|%s%s"), *Prefix, *GraphPath, *VariableName.ToString(), *Suffix);					
+					}
 					else
 					{
 						return FString::Printf(TEXT("%sLocalVariable::%s|%s%s"), *Prefix, *Node->GetGraph()->GetGraphName(), *VariableName.ToString(), *Suffix);					
@@ -2086,8 +2091,15 @@ FString URigVMCompiler::GetPinHash(const URigVMPin* InPin, const FRigVMVarExprAS
 
 	if (InVarExpr)
 	{
-		FString FullPath = InVarExpr->GetProxy().GetCallstack().GetCallPath(true);
-		return FString::Printf(TEXT("%s%s%s"), *Prefix, *FullPath, *Suffix);
+		if (bUseFullNodePath)
+		{
+			FString FullPath = InVarExpr->GetProxy().GetCallstack().GetCallPath(true);
+			return FString::Printf(TEXT("%s%s%s"), *Prefix, *FullPath, *Suffix);
+		}
+		else
+		{
+			return FString::Printf(TEXT("%s%s%s"), *Prefix, *InPin->GetPinPath(), *Suffix);
+		}
 	}
 
 	FString PinPath = InPin->GetPinPath(bUseFullNodePath);
