Control Rig: Fix Control Rig Display Actor Transforms

diff --git a/Engine/Plugins/Animation/ControlRig/Source/ControlRig/Private/ControlRigControlActor.cpp b/Engine/Plugins/Animation/ControlRig/Source/ControlRig/Private/ControlRigControlActor.cpp
--- a/Engine/Plugins/Animation/ControlRig/Source/ControlRig/Private/ControlRigControlActor.cpp
+++ b/Engine/Plugins/Animation/ControlRig/Source/ControlRig/Private/ControlRigControlActor.cpp
@@ -227,7 +227,7 @@ void AControlRigControlActor::Refresh()
 						Component->SetMaterial(0, MaterialInstance);
 
 						ControlNames.Add(ControlElement->GetName());
-						ShapeTransforms.Add(Hierarchy->GetControlShapeTransform(ControlElement, ERigTransformType::CurrentLocal) * ShapeDef->Transform);
+						ShapeTransforms.Add(ShapeDef->Transform * Hierarchy->GetControlShapeTransform(ControlElement, ERigTransformType::CurrentLocal));
 						Components.Add(Component);
 						Materials.Add(MaterialInstance);
 					}
