Control Rig: Fix missing upgrade pins for setter nodes

diff --git a/Engine/Plugins/Animation/ControlRig/Source/ControlRig/Private/Units/Hierarchy/RigUnit_SetBoneRotation.cpp b/Engine/Plugins/Animation/ControlRig/Source/ControlRig/Private/Units/Hierarchy/RigUnit_SetBoneRotation.cpp
--- a/Engine/Plugins/Animation/ControlRig/Source/ControlRig/Private/Units/Hierarchy/RigUnit_SetBoneRotation.cpp
+++ b/Engine/Plugins/Animation/ControlRig/Source/ControlRig/Private/Units/Hierarchy/RigUnit_SetBoneRotation.cpp
@@ -83,6 +83,7 @@ FRigVMStructUpgradeInfo FRigUnit_SetBoneRotation::GetUpgradeInfo() const
 	NewNode.Item = FRigElementKey(Bone, ERigElementType::Bone);
 	NewNode.Space = Space;
 	NewNode.Value = Rotation;
+	NewNode.Weight = Weight;
 
 	FRigVMStructUpgradeInfo Info(*this, NewNode);
 	Info.AddRemappedPin(TEXT("Bone"), TEXT("Item.Name"));
diff --git a/Engine/Plugins/Animation/ControlRig/Source/ControlRig/Private/Units/Hierarchy/RigUnit_SetBoneTransform.cpp b/Engine/Plugins/Animation/ControlRig/Source/ControlRig/Private/Units/Hierarchy/RigUnit_SetBoneTransform.cpp
--- a/Engine/Plugins/Animation/ControlRig/Source/ControlRig/Private/Units/Hierarchy/RigUnit_SetBoneTransform.cpp
+++ b/Engine/Plugins/Animation/ControlRig/Source/ControlRig/Private/Units/Hierarchy/RigUnit_SetBoneTransform.cpp
@@ -80,6 +80,7 @@ FRigVMStructUpgradeInfo FRigUnit_SetBoneTransform::GetUpgradeInfo() const
 	NewNode.Item = FRigElementKey(Bone, ERigElementType::Bone);
 	NewNode.Space = Space;
 	NewNode.Value = Transform;
+	NewNode.Weight = Weight;
 
 	FRigVMStructUpgradeInfo Info(*this, NewNode);
 	Info.AddRemappedPin(TEXT("Bone"), TEXT("Item.Name"));
diff --git a/Engine/Plugins/Animation/ControlRig/Source/ControlRig/Private/Units/Hierarchy/RigUnit_SetBoneTranslation.cpp b/Engine/Plugins/Animation/ControlRig/Source/ControlRig/Private/Units/Hierarchy/RigUnit_SetBoneTranslation.cpp
--- a/Engine/Plugins/Animation/ControlRig/Source/ControlRig/Private/Units/Hierarchy/RigUnit_SetBoneTranslation.cpp
+++ b/Engine/Plugins/Animation/ControlRig/Source/ControlRig/Private/Units/Hierarchy/RigUnit_SetBoneTranslation.cpp
@@ -83,6 +83,7 @@ FRigVMStructUpgradeInfo FRigUnit_SetBoneTranslation::GetUpgradeInfo() const
 	NewNode.Item = FRigElementKey(Bone, ERigElementType::Bone);
 	NewNode.Space = Space;
 	NewNode.Value = Translation;
+	NewNode.Weight = Weight;	
 
 	FRigVMStructUpgradeInfo Info(*this, NewNode);
 	Info.AddRemappedPin(TEXT("Bone"), TEXT("Item.Name"));
diff --git a/Engine/Plugins/Animation/ControlRig/Source/ControlRig/Private/Units/Hierarchy/RigUnit_SetControlTransform.cpp b/Engine/Plugins/Animation/ControlRig/Source/ControlRig/Private/Units/Hierarchy/RigUnit_SetControlTransform.cpp
--- a/Engine/Plugins/Animation/ControlRig/Source/ControlRig/Private/Units/Hierarchy/RigUnit_SetControlTransform.cpp
+++ b/Engine/Plugins/Animation/ControlRig/Source/ControlRig/Private/Units/Hierarchy/RigUnit_SetControlTransform.cpp
@@ -498,6 +498,7 @@ FRigVMStructUpgradeInfo FRigUnit_SetControlTransform::GetUpgradeInfo() const
 	NewNode.Item = FRigElementKey(Control, ERigElementType::Control);
 	NewNode.Space = Space;
 	NewNode.Value = Transform;
+	NewNode.Weight = Weight;
 
 	FRigVMStructUpgradeInfo Info(*this, NewNode);
 	Info.AddRemappedPin(TEXT("Control"), TEXT("Item.Name"));
diff --git a/Engine/Plugins/Animation/ControlRig/Source/ControlRig/Private/Units/Hierarchy/RigUnit_SetRelativeBoneTransform.cpp b/Engine/Plugins/Animation/ControlRig/Source/ControlRig/Private/Units/Hierarchy/RigUnit_SetRelativeBoneTransform.cpp
--- a/Engine/Plugins/Animation/ControlRig/Source/ControlRig/Private/Units/Hierarchy/RigUnit_SetRelativeBoneTransform.cpp
+++ b/Engine/Plugins/Animation/ControlRig/Source/ControlRig/Private/Units/Hierarchy/RigUnit_SetRelativeBoneTransform.cpp
@@ -59,6 +59,7 @@ FRigVMStructUpgradeInfo FRigUnit_SetRelativeBoneTransform::GetUpgradeInfo() cons
 	NewNode.Parent = FRigElementKey(Space, ERigElementType::Bone);
 	NewNode.Child = FRigElementKey(Bone, ERigElementType::Bone);
 	NewNode.Value = Transform;
+	NewNode.Weight = Weight;
 
 	FRigVMStructUpgradeInfo Info(*this, NewNode);
 	Info.AddRemappedPin(TEXT("Space"), TEXT("Parent.Name"));
diff --git a/Engine/Plugins/Animation/ControlRig/Source/ControlRig/Private/Units/Hierarchy/RigUnit_SetSpaceTransform.cpp b/Engine/Plugins/Animation/ControlRig/Source/ControlRig/Private/Units/Hierarchy/RigUnit_SetSpaceTransform.cpp
--- a/Engine/Plugins/Animation/ControlRig/Source/ControlRig/Private/Units/Hierarchy/RigUnit_SetSpaceTransform.cpp
+++ b/Engine/Plugins/Animation/ControlRig/Source/ControlRig/Private/Units/Hierarchy/RigUnit_SetSpaceTransform.cpp
@@ -73,6 +73,7 @@ FRigVMStructUpgradeInfo FRigUnit_SetSpaceTransform::GetUpgradeInfo() const
 	NewNode.Item = FRigElementKey(Space, ERigElementType::Null);
 	NewNode.Space = SpaceType;
 	NewNode.Value = Transform;
+	NewNode.Weight = Weight;
 
 	FRigVMStructUpgradeInfo Info(*this, NewNode);
 	Info.AddRemappedPin(TEXT("Space"), TEXT("Item.Name"));
