Control Rig: Avoid recreating the control rig actors in the viewport during interaction

diff --git a/Engine/Plugins/Animation/ControlRig/Source/ControlRigEditor/Private/Editor/ControlRigEditor.cpp b/Engine/Plugins/Animation/ControlRig/Source/ControlRigEditor/Private/Editor/ControlRigEditor.cpp
--- a/Engine/Plugins/Animation/ControlRig/Source/ControlRigEditor/Private/Editor/ControlRigEditor.cpp
+++ b/Engine/Plugins/Animation/ControlRig/Source/ControlRigEditor/Private/Editor/ControlRigEditor.cpp
@@ -6038,7 +6038,12 @@ void FControlRigEditor::HandleOnControlModified(UControlRig* Subject, FRigContro
 		FRigControlElement* TargetControlElement = Blueprint->Hierarchy->Find<FRigControlElement>(ControlElement->GetKey());
 
 		TargetControlElement->Settings = SourceControlElement->Settings;
-		Blueprint->Hierarchy->OnModified().Broadcast(ERigHierarchyNotification::ControlSettingChanged, Blueprint->Hierarchy, TargetControlElement);
+
+		// only fire the setting change if the interaction is not currently ongoing
+		if(!Subject->ElementsBeingInteracted.Contains(ControlElement->GetKey()))
+		{
+			Blueprint->Hierarchy->OnModified().Broadcast(ERigHierarchyNotification::ControlSettingChanged, Blueprint->Hierarchy, TargetControlElement);
+		}
 
 		// we copy the pose including the weights since we want the topology to align during setup mode.
 		// i.e. dynamic reparenting should be reset here.
