Incorporate element type when generating an FK control - this fixes any overlaps between curve and bone names in target USkeleton
- Added upgrade path for Sequencer control curves on ControlRigParameterSection

diff --git a/Engine/Plugins/Animation/ControlRig/Source/ControlRig/Private/Rigs/FKControlRig.cpp b/Engine/Plugins/Animation/ControlRig/Source/ControlRig/Private/Rigs/FKControlRig.cpp
--- a/Engine/Plugins/Animation/ControlRig/Source/ControlRig/Private/Rigs/FKControlRig.cpp
+++ b/Engine/Plugins/Animation/ControlRig/Source/ControlRig/Private/Rigs/FKControlRig.cpp
@@ -6,6 +6,7 @@
 #include "IControlRigObjectBinding.h"
 #include "Components/SkeletalMeshComponent.h"
 #include "Execution/RigUnit_PrepareForExecution.h"
+#include "Settings/ControlRigSettings.h"
 #include "Units/Execution/RigUnit_BeginExecution.h"
 #include "Units/Execution/RigUnit_InverseExecution.h"
 
@@ -20,14 +21,50 @@ UFKControlRig::UFKControlRig(const FObjectInitializer& ObjectInitializer)
 	bResetInitialTransformsBeforeSetup = false;
 }
 
-FName UFKControlRig::GetControlName(const FName& InBoneName)
+FName UFKControlRig::GetControlName(const FName& InName, const ERigElementType& InType)
 {
-	if (InBoneName != NAME_None)
+	if (InName != NAME_None)
 	{
-		return FName(*(InBoneName.ToString() + TEXT("_CONTROL")));
+		const FString& PostFix = [InType]()
+		{
+			if (InType == ERigElementType::Curve)
+			{
+				return TEXT("_CURVE");
+			}
+
+			return TEXT("");
+		}();
+		
+		return FName(*(InName.ToString() + PostFix + TEXT("_CONTROL")));
+	}
+
+	// if control name is coming as none, we don't append the postfix
+	return NAME_None;
+}
+
+FName UFKControlRig::GetControlTargetName(const FName& InName, const ERigElementType& InType)
+{
+	if (InName != NAME_None)
+	{
+		const FString& PostFix = [InType]()
+		{
+			if (InType == ERigElementType::Curve)
+			{
+				return TEXT("_CURVE_CONTROL");
+			}
+
+			return TEXT("_CONTROL");
+		}();
+
+		FString StringName = InName.ToString();
+		if (StringName.EndsWith(PostFix, ESearchCase::CaseSensitive))
+		{
+			StringName.RemoveFromEnd(PostFix);
+			return FName(*StringName);
+		}		
 	}
 
-	// if bone name is coming as none, we don't append
+	// If incoming name is NONE, or not correctly formatted according to expected post-fix return NAME_None
 	return NAME_None;
 }
 
@@ -52,7 +89,7 @@ void UFKControlRig::ExecuteUnits(FRigUnitContext& InOutContext, const FName& InE
 	{
 		GetHierarchy()->ForEach<FRigBoneElement>([&](FRigBoneElement* BoneElement) -> bool
         {
-			const FName ControlName = GetControlName(BoneElement->GetName());
+			const FName ControlName = GetControlName(BoneElement->GetName(), BoneElement->GetType());
 			const FRigElementKey ControlKey(ControlName, ERigElementType::Control);
 			const int32 ControlIndex = GetHierarchy()->GetIndex(ControlKey);
 			if (ControlIndex != INDEX_NONE && IsControlActive[ControlIndex])
@@ -89,7 +126,7 @@ void UFKControlRig::ExecuteUnits(FRigUnitContext& InOutContext, const FName& InE
 
 		GetHierarchy()->ForEach<FRigCurveElement>([&](FRigCurveElement* CurveElement) -> bool
         {
-			const FName ControlName = GetControlName(CurveElement->GetName());
+			const FName ControlName = GetControlName(CurveElement->GetName(), CurveElement->GetType());
 			const FRigElementKey ControlKey(ControlName, ERigElementType::Control);
 			const int32 ControlIndex = GetHierarchy()->GetIndex(ControlKey);
 			
@@ -132,7 +169,7 @@ void UFKControlRig::ExecuteUnits(FRigUnitContext& InOutContext, const FName& InE
 
 			FRigBoneElement* BoneElement = CastChecked<FRigBoneElement>(InElement);
 
-			const FName ControlName = GetControlName(BoneElement->GetName());
+			const FName ControlName = GetControlName(BoneElement->GetName(), BoneElement->GetType());
 			const FRigElementKey ControlKey(ControlName, ERigElementType::Control);
 			const int32 ControlIndex = GetHierarchy()->GetIndex(ControlKey);
 			
@@ -167,7 +204,7 @@ void UFKControlRig::ExecuteUnits(FRigUnitContext& InOutContext, const FName& InE
 
 		GetHierarchy()->ForEach<FRigCurveElement>([&](FRigCurveElement* CurveElement) -> bool
         {
-            const FName ControlName = GetControlName(CurveElement->GetName());
+			const FName ControlName = GetControlName(CurveElement->GetName(), CurveElement->GetType());
 			const FRigElementKey ControlKey(ControlName, ERigElementType::Control);
 			const int32 ControlIndex = GetHierarchy()->GetIndex(ControlKey);
 
@@ -293,7 +330,7 @@ void UFKControlRig::CreateRigElements(const FReferenceSkeleton& InReferenceSkele
 		GetHierarchy()->ForEach<FRigBoneElement>([&](FRigBoneElement* BoneElement) -> bool
 		{
 			const FName BoneName = BoneElement->GetName();
-			const FName ControlName = GetControlName(BoneName); // name conflict?
+			const FName ControlName = GetControlName(BoneName, BoneElement->GetType());
 			const FRigElementKey ParentKey = GetHierarchy()->GetFirstParent(BoneElement->GetKey());
 
 			FTransform OffsetTransform = FTransform::Identity;
@@ -313,12 +350,14 @@ void UFKControlRig::CreateRigElements(const FReferenceSkeleton& InReferenceSkele
 		
 		GetHierarchy()->ForEach<FRigCurveElement>([&](FRigCurveElement* CurveElement) -> bool
 		{
-			const FName ControlName = GetControlName(CurveElement->GetName()); // name conflict?
+			const FName ControlName = GetControlName(CurveElement->GetName(), CurveElement->GetType());
 
 			FRigControlSettings Settings;
 			Settings.ControlType = ERigControlType::Float;
 			Settings.DisplayName = CurveElement->GetName();
 
+			const FName DisplayCurveControlName(*(CurveElement->GetName().ToString() + TEXT(" Curve")));
+		    Settings.DisplayName = DisplayCurveControlName;
 			Controller->AddControl(ControlName, FRigElementKey(), Settings, FRigControlValue::Make(CurveElement->Value), FTransform::Identity, FTransform::Identity, false);
 			
 			return true;
@@ -347,7 +386,7 @@ void UFKControlRig::SetControlOffsetsFromBoneInitials()
 
 		FRigBoneElement* BoneElement = CastChecked<FRigBoneElement>(InElement);
 		const FName BoneName = BoneElement->GetName();
-		const FName ControlName = GetControlName(BoneName); // name conflict?
+		const FName ControlName = GetControlName(BoneName, BoneElement->GetType());
 
 		FRigControlElement* ControlElement = GetHierarchy()->Find<FRigControlElement>(FRigElementKey(ControlName, ERigElementType::Control));
 		if(ControlElement == nullptr)
diff --git a/Engine/Plugins/Animation/ControlRig/Source/ControlRig/Private/Sequencer/MovieSceneControlRigParameterSection.cpp b/Engine/Plugins/Animation/ControlRig/Source/ControlRig/Private/Sequencer/MovieSceneControlRigParameterSection.cpp
--- a/Engine/Plugins/Animation/ControlRig/Source/ControlRig/Private/Sequencer/MovieSceneControlRigParameterSection.cpp
+++ b/Engine/Plugins/Animation/ControlRig/Source/ControlRig/Private/Sequencer/MovieSceneControlRigParameterSection.cpp
@@ -15,6 +15,7 @@
 #include "MovieSceneTimeHelpers.h"
 #include "Animation/AnimSequenceHelpers.h"
 #include "Components/SkeletalMeshComponent.h"
+#include "UObject/UE5MainStreamObjectVersion.h"
 
 #if WITH_EDITOR
 #include "AnimPose.h"
@@ -2104,6 +2105,27 @@ void UMovieSceneControlRigParameterSection::RecreateWithThisControlRig(UControlR
 	TArray<FRigControlElement*> SortedControls;
 	ControlRig->GetControlsInOrder(SortedControls);
 
+	TMap<FName, FName> CurveControlNameRemapping;
+	if (GetLinkerCustomVersion(FUE5MainStreamObjectVersion::GUID) < FUE5MainStreamObjectVersion::FKControlNamingScheme)
+	{
+		for (FRigControlElement* ControlElement : SortedControls)
+		{
+			if (ControlElement->Settings.ControlType == ERigControlType::Float)
+			{
+				const FName TargetCurveName = UFKControlRig::GetControlTargetName(ControlElement->GetName(), ERigElementType::Curve);
+				const FRigElementKey CurveKey = FRigElementKey(TargetCurveName, ERigElementType::Curve);
+				// Ensure name is valid, and curve actually exists in the hierarchy,
+				// this means we could not be renaming some controls for which the curves do not exist anymore, which ties back to comment at the top op the function
+				// with regards to non-associated curves
+				if (TargetCurveName != NAME_None && ControlRig->GetHierarchy()->Find(CurveKey))
+				{
+					// Add mapping from old to new control naming scheme (previous was using uniform naming for both bones and curves)
+					CurveControlNameRemapping.Add(ControlElement->GetName(), UFKControlRig::GetControlName(TargetCurveName, ERigElementType::Bone));
+				}
+			}
+		}
+	}
+
 	for (FRigControlElement* ControlElement : SortedControls)
 	{
 		if (!ControlElement->Settings.bAnimatable)
@@ -2117,6 +2139,11 @@ void UMovieSceneControlRigParameterSection::RecreateWithThisControlRig(UControlR
 			RenameParameterName(PreviousName, ControlElement->GetKey().Name);
 		}
 
+		if (const FName* OldCurveControlName = CurveControlNameRemapping.Find(ControlElement->GetName()))
+		{
+			RenameParameterName(*OldCurveControlName, ControlElement->GetKey().Name);
+		}
+		
 		switch (ControlElement->Settings.ControlType)
 		{
 		case ERigControlType::Float:
diff --git a/Engine/Plugins/Animation/ControlRig/Source/ControlRig/Public/Rigs/FKControlRig.h b/Engine/Plugins/Animation/ControlRig/Source/ControlRig/Public/Rigs/FKControlRig.h
--- a/Engine/Plugins/Animation/ControlRig/Source/ControlRig/Public/Rigs/FKControlRig.h
+++ b/Engine/Plugins/Animation/ControlRig/Source/ControlRig/Public/Rigs/FKControlRig.h
@@ -48,8 +48,10 @@ public:
 	virtual void SetBoneInitialTransformsFromSkeletalMeshComponent(USkeletalMeshComponent* InSkelMeshComp, bool bUseAnimInstance = false) override;
 	// END ControlRig
 
-	// utility function to 
-	static FName GetControlName(const FName& InBoneName);
+	// utility function to generate a valid control element name
+	static FName GetControlName(const FName& InName, const ERigElementType& InType);
+	// utility function to generate a target element name for control
+	static FName GetControlTargetName(const FName& InName, const ERigElementType& InType);
 
 	TArray<FName> GetControlNames();
 	bool GetControlActive(int32 Index) const;
diff --git a/Engine/Plugins/Animation/ControlRig/Source/ControlRigEditor/Private/EditMode/ControlRigEditMode.cpp b/Engine/Plugins/Animation/ControlRig/Source/ControlRigEditor/Private/EditMode/ControlRigEditMode.cpp
--- a/Engine/Plugins/Animation/ControlRig/Source/ControlRigEditor/Private/EditMode/ControlRigEditMode.cpp
+++ b/Engine/Plugins/Animation/ControlRig/Source/ControlRigEditor/Private/EditMode/ControlRigEditMode.cpp
@@ -701,7 +701,7 @@ void FControlRigEditMode::Render(const FSceneView* View, FViewport* Viewport, FP
 								FLinearColor Color = FLinearColor::White;
 								if (bHasFKRig)
 								{
-									FName ControlName = UFKControlRig::GetControlName(ParentTransformElement->GetName());
+									FName ControlName = UFKControlRig::GetControlName(ParentTransformElement->GetName(), ParentTransformElement->GetType());
 									if (ActiveControlName.Num() > 0 && ActiveControlName.Contains(ControlName) == false)
 									{
 										continue;
@@ -731,7 +731,7 @@ void FControlRigEditMode::Render(const FSceneView* View, FViewport* Viewport, FP
 						FLinearColor Color = FLinearColor::White;
 						if (bHasFKRig)
 						{
-							FName ControlName = UFKControlRig::GetControlName(TransformElement->GetName());
+							FName ControlName = UFKControlRig::GetControlName(TransformElement->GetName(), TransformElement->GetType());
 							if (ActiveControlName.Num() > 0 && ActiveControlName.Contains(ControlName) == false)
 							{
 								return true;
diff --git a/Engine/Source/Runtime/Core/Public/UObject/UE5MainStreamObjectVersion.h b/Engine/Source/Runtime/Core/Public/UObject/UE5MainStreamObjectVersion.h
--- a/Engine/Source/Runtime/Core/Public/UObject/UE5MainStreamObjectVersion.h
+++ b/Engine/Source/Runtime/Core/Public/UObject/UE5MainStreamObjectVersion.h
@@ -200,6 +200,9 @@ struct CORE_API FUE5MainStreamObjectVersion
 		// FDynamicMeshAttributeSet has Weight Maps. TDynamicAttributeBase serializes its name.
 		DynamicMeshAttributesWeightMapsAndNames,
 
+		// Switching FK control naming scheme to incorporate _CURVE for curve controls
+		FKControlNamingScheme,
+
 		// -----<new versions can be added above this line>-------------------------------------------------
 		VersionPlusOne,
 		LatestVersion = VersionPlusOne - 1
