Control Rig: Fix incorrect hierarchy dependency order

Fallout from 19494364

diff --git a/Engine/Plugins/Animation/ControlRig/Source/ControlRig/Private/Rigs/RigHierarchy.cpp b/Engine/Plugins/Animation/ControlRig/Source/ControlRig/Private/Rigs/RigHierarchy.cpp
--- a/Engine/Plugins/Animation/ControlRig/Source/ControlRig/Private/Rigs/RigHierarchy.cpp
+++ b/Engine/Plugins/Animation/ControlRig/Source/ControlRig/Private/Rigs/RigHierarchy.cpp
@@ -4388,7 +4388,7 @@ URigHierarchy::TElementDependencyMap URigHierarchy::GetDependenciesForVM(const U
 				// only create dependencies for reads and writes that are on the same slice
 				if(ReadTransform != FilteredTransform && ReadTransform.Get<0>() == FilteredTransform.Get<0>())
 				{
-					Dependencies.FindOrAdd(ReadTransform.Get<1>()).AddUnique(FilteredTransform.Get<1>());
+					Dependencies.FindOrAdd(FilteredTransform.Get<1>()).AddUnique(ReadTransform.Get<1>());
 				}
 			}
 		}
