Sequencer: Control Rig: Fix issue with not being able to change transform type.

diff --git a/Engine/Plugins/Animation/ControlRig/Source/ControlRigEditor/Private/EditMode/ControlRigEditMode.cpp b/Engine/Plugins/Animation/ControlRig/Source/ControlRigEditor/Private/EditMode/ControlRigEditMode.cpp
--- a/Engine/Plugins/Animation/ControlRig/Source/ControlRigEditor/Private/EditMode/ControlRigEditMode.cpp
+++ b/Engine/Plugins/Animation/ControlRig/Source/ControlRigEditor/Private/EditMode/ControlRigEditMode.cpp
@@ -56,6 +56,11 @@
 #include "ControlRigSpaceChannelEditors.h"
 #include "ControlRigSequencerEditorLibrary.h"
 #include "LevelSequence.h"
+#include "LevelEditor.h"
+#include "InteractiveToolManager.h"
+#include "EdModeInteractiveToolsContext.h"
+
+
 
 void UControlRigEditModeDelegateHelper::OnPoseInitialized()
 {
@@ -374,8 +379,28 @@ void FControlRigEditMode::Enter()
 	SetObjects_Internal();
 }
 
+//todo get working with Persona
+static void ClearOutAnyActiveTools()
+{
+	if (FLevelEditorModule* LevelEditorModule = FModuleManager::GetModulePtr<FLevelEditorModule>(TEXT("LevelEditor")))
+	{
+		TSharedPtr<ILevelEditor> LevelEditorPtr = LevelEditorModule->GetLevelEditorInstance().Pin();
+
+		if (LevelEditorPtr.IsValid())
+		{
+			FString ActiveToolName = LevelEditorPtr->GetEditorModeManager().GetInteractiveToolsContext()->ToolManager->GetActiveToolName(EToolSide::Left);
+			if (ActiveToolName == TEXT("SequencerPivotTool"))
+			{
+				LevelEditorPtr->GetEditorModeManager().GetInteractiveToolsContext()->ToolManager->DeactivateTool(EToolSide::Left, EToolShutdownType::Completed);
+			}
+		}
+	}
+}
+
 void FControlRigEditMode::Exit()
 {
+	ClearOutAnyActiveTools();
+
 	for (TWeakObjectPtr<UControlRig>& RuntimeRigPtr : RuntimeControlRigs)
 	{
 		if (UControlRig* ControlRig = RuntimeRigPtr.Get())
@@ -505,47 +530,49 @@ void FControlRigEditMode::Tick(FEditorViewportClient* ViewportClient, float Delt
 	// might have changed the transform of the Control
 	{
 		PostPoseUpdate();
-
-		for (TWeakObjectPtr<UControlRig>& RuntimeRigPtr : RuntimeControlRigs)
+		if (IsInLevelEditor() == false) //only do this check if not in level editor
 		{
-			if (UControlRig* ControlRig = RuntimeRigPtr.Get())
+			for (TWeakObjectPtr<UControlRig>& RuntimeRigPtr : RuntimeControlRigs)
 			{
-				TArray<FRigElementKey> SelectedRigElements = GetSelectedRigElements(ControlRig);
-				const UE::Widget::EWidgetMode CurrentWidgetMode = ViewportClient->GetWidgetMode();
-				for (FRigElementKey SelectedRigElement : SelectedRigElements)
+				if (UControlRig* ControlRig = RuntimeRigPtr.Get())
 				{
-					//need to loop through the shape actors and set widget based upon the first one
-					if (AControlRigShapeActor* ShapeActor = GetControlShapeFromControlName(ControlRig,SelectedRigElement.Name))
+					TArray<FRigElementKey> SelectedRigElements = GetSelectedRigElements(ControlRig);
+					const UE::Widget::EWidgetMode CurrentWidgetMode = ViewportClient->GetWidgetMode();
+					for (FRigElementKey SelectedRigElement : SelectedRigElements)
 					{
-						if (!ModeSupportedByShapeActor(ShapeActor, CurrentWidgetMode))
+						//need to loop through the shape actors and set widget based upon the first one
+						if (AControlRigShapeActor* ShapeActor = GetControlShapeFromControlName(ControlRig, SelectedRigElement.Name))
 						{
-							if (FRigControlElement* ControlElement = ControlRig->FindControl(SelectedRigElement.Name))
+							if (!ModeSupportedByShapeActor(ShapeActor, CurrentWidgetMode))
 							{
-								switch (ControlElement->Settings.ControlType)
+								if (FRigControlElement* ControlElement = ControlRig->FindControl(SelectedRigElement.Name))
 								{
-								case ERigControlType::Float:
-								case ERigControlType::Integer:
-								case ERigControlType::Vector2D:
-								case ERigControlType::Position:
-								case ERigControlType::Transform:
-								case ERigControlType::TransformNoScale:
-								case ERigControlType::EulerTransform:
-								{
-									ViewportClient->SetWidgetMode(UE::Widget::WM_Translate);
-									break;
-								}
-								case ERigControlType::Rotator:
-								{
-									ViewportClient->SetWidgetMode(UE::Widget::WM_Rotate);
-									break;
-								}
-								case ERigControlType::Scale:
-								{
-									ViewportClient->SetWidgetMode(UE::Widget::WM_Scale);
-									break;
-								}
+									switch (ControlElement->Settings.ControlType)
+									{
+									case ERigControlType::Float:
+									case ERigControlType::Integer:
+									case ERigControlType::Vector2D:
+									case ERigControlType::Position:
+									case ERigControlType::Transform:
+									case ERigControlType::TransformNoScale:
+									case ERigControlType::EulerTransform:
+									{
+										ViewportClient->SetWidgetMode(UE::Widget::WM_Translate);
+										break;
+									}
+									case ERigControlType::Rotator:
+									{
+										ViewportClient->SetWidgetMode(UE::Widget::WM_Rotate);
+										break;
+									}
+									case ERigControlType::Scale:
+									{
+										ViewportClient->SetWidgetMode(UE::Widget::WM_Scale);
+										break;
+									}
+									}
+									return; //exit if we switchted
 								}
-								return; //exit if we switchted
 							}
 						}
 					}
@@ -3733,6 +3760,4 @@ bool FControlRigEditMode::GetOnlySelectRigControls()const
 }
 
 
-
-
 #undef LOCTEXT_NAMESPACE
