Add const get-er for DynamicHierarchy

diff --git a/Engine/Plugins/Animation/ControlRig/Source/ControlRig/Public/ControlRig.h b/Engine/Plugins/Animation/ControlRig/Source/ControlRig/Public/ControlRig.h
--- a/Engine/Plugins/Animation/ControlRig/Source/ControlRig/Public/ControlRig.h
+++ b/Engine/Plugins/Animation/ControlRig/Source/ControlRig/Public/ControlRig.h
@@ -250,6 +250,11 @@ public:
 	{
 		return DynamicHierarchy;
 	}
+	
+	FORCEINLINE_DEBUGGABLE URigHierarchy* GetHierarchy() const
+	{
+		return DynamicHierarchy;
+	}
 
 #if WITH_EDITOR
 
