Control Rig: Fix interaction rig unit test

diff --git a/Engine/Plugins/Animation/ControlRig/Source/ControlRig/Private/ControlRig.cpp b/Engine/Plugins/Animation/ControlRig/Source/ControlRig/Private/ControlRig.cpp
--- a/Engine/Plugins/Animation/ControlRig/Source/ControlRig/Private/ControlRig.cpp
+++ b/Engine/Plugins/Animation/ControlRig/Source/ControlRig/Private/ControlRig.cpp
@@ -2485,9 +2485,7 @@ void UControlRig::HandleInteractionRigControlModified(UControlRig* Subject, FRig
 						QueuedModifiedControls.Add(AffectedControl->GetKey());
 					}
 				}
-				else if (
-					AffectedKey.Type == ERigElementType::Bone ||
-					AffectedKey.Type == ERigElementType::Curve)
+				else if (AffectedKey.Type == ERigElementType::Bone)
 				{
 					// special case controls with a CONTROL suffix
 					FName BoneControlName = *FString::Printf(TEXT("%s_CONTROL"), *AffectedKey.Name.ToString());
@@ -2496,6 +2494,15 @@ void UControlRig::HandleInteractionRigControlModified(UControlRig* Subject, FRig
 						QueuedModifiedControls.Add(AffectedControl->GetKey());
 					}
 				}
+				else if(AffectedKey.Type == ERigElementType::Curve)
+				{
+					// special case controls with a CONTROL suffix
+					FName CurveControlName = *FString::Printf(TEXT("%s_CURVE_CONTROL"), *AffectedKey.Name.ToString());
+					if (FRigControlElement* AffectedControl = FindControl(CurveControlName))
+					{
+						QueuedModifiedControls.Add(AffectedControl->GetKey());
+					}
+				}
 			}
 		}
 	}
