Control Rig: Sustain order of selection in hierarchy

diff --git a/Engine/Plugins/Animation/ControlRig/Source/ControlRig/Private/ControlRig.cpp b/Engine/Plugins/Animation/ControlRig/Source/ControlRig/Private/ControlRig.cpp
--- a/Engine/Plugins/Animation/ControlRig/Source/ControlRig/Private/ControlRig.cpp
+++ b/Engine/Plugins/Animation/ControlRig/Source/ControlRig/Private/ControlRig.cpp
@@ -1732,8 +1732,8 @@ TArray<FName> UControlRig::CurrentControlSelection() const
 
 	if(DynamicHierarchy)
 	{
-		TArray<FRigBaseElement*> SelectedControls = DynamicHierarchy->GetSelectedElements(ERigElementType::Control);
-		for (FRigBaseElement* SelectedControl : SelectedControls)
+		TArray<const FRigBaseElement*> SelectedControls = DynamicHierarchy->GetSelectedElements(ERigElementType::Control);
+		for (const FRigBaseElement* SelectedControl : SelectedControls)
 		{
 			SelectedControlNames.Add(SelectedControl->GetName());
 		}
diff --git a/Engine/Plugins/Animation/ControlRig/Source/ControlRig/Private/Rigs/RigHierarchy.cpp b/Engine/Plugins/Animation/ControlRig/Source/ControlRig/Private/Rigs/RigHierarchy.cpp
--- a/Engine/Plugins/Animation/ControlRig/Source/ControlRig/Private/Rigs/RigHierarchy.cpp
+++ b/Engine/Plugins/Animation/ControlRig/Source/ControlRig/Private/Rigs/RigHierarchy.cpp
@@ -663,10 +663,10 @@ int32 URigHierarchy::Num(ERigElementType InElementType) const
 	return ElementsPerType[RigElementTypeToFlatIndex(InElementType)].Num();
 }
 
-TArray<FRigBaseElement*> URigHierarchy::GetSelectedElements(ERigElementType InTypeFilter) const
+TArray<const FRigBaseElement*> URigHierarchy::GetSelectedElements(ERigElementType InTypeFilter) const
 {
 	LLM_SCOPE_BYNAME(TEXT("Animation/ControlRig"));
-	TArray<FRigBaseElement*> Selection;
+	TArray<const FRigBaseElement*> Selection;
 
 	if(URigHierarchy* HierarchyForSelection = HierarchyForSelectionPtr.Get())
 	{
@@ -681,13 +681,13 @@ TArray<FRigBaseElement*> URigHierarchy::GetSelectedElements(ERigElementType InTy
 		return Selection;
 	}
 
-	for (int32 ElementIndex = 0; ElementIndex < Elements.Num(); ElementIndex++)
+	for (const FRigElementKey& SelectedKey : OrderedSelection)
 	{
-		FRigBaseElement* Element = Elements[ElementIndex];
-		if(Element->IsTypeOf(InTypeFilter))
+		if(SelectedKey.IsTypeOf(InTypeFilter))
 		{
-			if(IsSelected(Element))
+			if(const FRigBaseElement* Element = FindChecked(SelectedKey))
 			{
+				ensure(Element->IsSelected());
 				Selection.Add(Element);
 			}
 		}
@@ -705,17 +705,14 @@ TArray<FRigElementKey> URigHierarchy::GetSelectedKeys(ERigElementType InTypeFilt
 	}
 
 	TArray<FRigElementKey> Selection;
-	for (int32 ElementIndex = 0; ElementIndex < Elements.Num(); ElementIndex++)
+	for (const FRigElementKey& SelectedKey : OrderedSelection)
 	{
-		FRigBaseElement* Element = Elements[ElementIndex];
-		if(Element->IsTypeOf(InTypeFilter))
+		if(SelectedKey.IsTypeOf(InTypeFilter))
 		{
-			if(IsSelected(Element))
-			{
-				Selection.Add(Element->GetKey());
-			}
+			Selection.Add(SelectedKey);
 		}
 	}
+	
 	return Selection;
 }
 
@@ -3529,7 +3526,10 @@ bool URigHierarchy::IsSelected(const FRigBaseElement* InElement) const
 	{
 		return HierarchyForSelection->IsSelected(InElement->GetKey());
 	}
-	return InElement->IsSelected();
+
+	const bool bIsSelected = OrderedSelection.Contains(InElement->GetKey());
+	ensure(bIsSelected == InElement->IsSelected());
+	return OrderedSelection.Contains(InElement->GetKey());
 }
 
 void URigHierarchy::ResetCachedChildren()
diff --git a/Engine/Plugins/Animation/ControlRig/Source/ControlRig/Private/Rigs/RigHierarchyController.cpp b/Engine/Plugins/Animation/ControlRig/Source/ControlRig/Private/Rigs/RigHierarchyController.cpp
--- a/Engine/Plugins/Animation/ControlRig/Source/ControlRig/Private/Rigs/RigHierarchyController.cpp
+++ b/Engine/Plugins/Animation/ControlRig/Source/ControlRig/Private/Rigs/RigHierarchyController.cpp
@@ -78,12 +78,22 @@ bool URigHierarchyController::SelectElement(FRigElementKey InKey, bool bSelect,
 		return false;
 	}
 
+	const bool bSelectionState = Hierarchy->OrderedSelection.Contains(InKey);
+	ensure(bSelectionState == Element->bSelected);
 	if(Element->bSelected == bSelect)
 	{
 		return false;
 	}
 
 	Element->bSelected = bSelect;
+	if(bSelect)
+	{
+		Hierarchy->OrderedSelection.Add(InKey);
+	}
+	else
+	{
+		Hierarchy->OrderedSelection.Remove(InKey);
+	}
 
 	if(Element->bSelected)
 	{
diff --git a/Engine/Plugins/Animation/ControlRig/Source/ControlRig/Public/Rigs/RigHierarchy.h b/Engine/Plugins/Animation/ControlRig/Source/ControlRig/Public/Rigs/RigHierarchy.h
--- a/Engine/Plugins/Animation/ControlRig/Source/ControlRig/Public/Rigs/RigHierarchy.h
+++ b/Engine/Plugins/Animation/ControlRig/Source/ControlRig/Public/Rigs/RigHierarchy.h
@@ -753,7 +753,7 @@ public:
 	 * @InTypeFilter The types to retrieve the selection for
 	 * @return An array of the currently selected elements
 	 */
-	TArray<FRigBaseElement*> GetSelectedElements(ERigElementType InTypeFilter = ERigElementType::All) const;
+	TArray<const FRigBaseElement*> GetSelectedElements(ERigElementType InTypeFilter = ERigElementType::All) const;
 
 	/**
 	 * Returns the keys of selected elements
@@ -2786,6 +2786,7 @@ private:
 	FRigHierarchyUndoRedoTransformEvent UndoRedoEvent;
 
 	TWeakObjectPtr<URigHierarchy> HierarchyForSelectionPtr;
+	TArray<FRigElementKey> OrderedSelection;
 
 	UPROPERTY(Transient)
 	TObjectPtr<URigHierarchyController> HierarchyController;
diff --git a/Engine/Plugins/Animation/ControlRig/Source/ControlRig/Public/Rigs/RigHierarchyDefines.h b/Engine/Plugins/Animation/ControlRig/Source/ControlRig/Public/Rigs/RigHierarchyDefines.h
--- a/Engine/Plugins/Animation/ControlRig/Source/ControlRig/Public/Rigs/RigHierarchyDefines.h
+++ b/Engine/Plugins/Animation/ControlRig/Source/ControlRig/Public/Rigs/RigHierarchyDefines.h
@@ -1405,6 +1405,11 @@ public:
 		Name = NAME_None;
 	}
 
+	FORCEINLINE bool IsTypeOf(ERigElementType InElementType) const
+	{
+		return ((uint8)InElementType & (uint8)Type) == (uint8)Type;
+	}
+
 	friend FORCEINLINE uint32 GetTypeHash(const FRigElementKey& Key)
 	{
 		return GetTypeHash(Key.Name) * 10 + (uint32)Key.Type;
diff --git a/Engine/Plugins/Animation/ControlRig/Source/ControlRig/Public/Rigs/RigHierarchyElements.h b/Engine/Plugins/Animation/ControlRig/Source/ControlRig/Public/Rigs/RigHierarchyElements.h
--- a/Engine/Plugins/Animation/ControlRig/Source/ControlRig/Public/Rigs/RigHierarchyElements.h
+++ b/Engine/Plugins/Animation/ControlRig/Source/ControlRig/Public/Rigs/RigHierarchyElements.h
@@ -470,7 +470,7 @@ public:
 
 	FORCEINLINE bool IsTypeOf(ERigElementType InElementType) const
 	{
-		return ((uint8)InElementType & (uint8)Key.Type) == (uint8)Key.Type;
+		return Key.IsTypeOf(InElementType);
 	}
 
 	template<typename T>
diff --git a/Engine/Plugins/Animation/ControlRig/Source/ControlRigEditor/Private/Editor/ControlRigEditor.cpp b/Engine/Plugins/Animation/ControlRig/Source/ControlRigEditor/Private/Editor/ControlRigEditor.cpp
--- a/Engine/Plugins/Animation/ControlRig/Source/ControlRigEditor/Private/Editor/ControlRigEditor.cpp
+++ b/Engine/Plugins/Animation/ControlRig/Source/ControlRigEditor/Private/Editor/ControlRigEditor.cpp
@@ -4825,7 +4825,7 @@ void FControlRigEditor::OnHierarchyChanged()
 
 		FBlueprintEditorUtils::MarkBlueprintAsModified(GetControlRigBlueprint());
 		
-		TArray<FRigBaseElement*> SelectedElements = ControlRigBP->Hierarchy->GetSelectedElements();
+		TArray<const FRigBaseElement*> SelectedElements = ControlRigBP->Hierarchy->GetSelectedElements();
 		for(const FRigBaseElement* SelectedElement : SelectedElements)
 		{
 			ControlRigBP->Hierarchy->OnModified().Broadcast(ERigHierarchyNotification::ElementSelected, ControlRigBP->Hierarchy, SelectedElement);
@@ -5179,7 +5179,7 @@ void FControlRigEditor::SynchronizeViewportBoneSelection()
 	{
 		EditorSkelComp->BonesOfInterest.Reset();
 
-		TArray<FRigBaseElement*> SelectedBones = RigBlueprint->Hierarchy->GetSelectedElements(ERigElementType::Bone);
+		TArray<const FRigBaseElement*> SelectedBones = RigBlueprint->Hierarchy->GetSelectedElements(ERigElementType::Bone);
 		for (const FRigBaseElement* SelectedBone : SelectedBones)
 		{
  			const int32 BoneIndex = EditorSkelComp->GetReferenceSkeleton().FindBoneIndex(SelectedBone->GetName());
