Control Rig: Integrate node / pin workflows into editor code

diff --git a/Engine/Plugins/Animation/ControlRig/Source/ControlRig/Private/Units/ControlRigNodeWorkflow.cpp b/Engine/Plugins/Animation/ControlRig/Source/ControlRig/Private/Units/ControlRigNodeWorkflow.cpp
new file mode 100644
--- /dev/null
+++ b/Engine/Plugins/Animation/ControlRig/Source/ControlRig/Private/Units/ControlRigNodeWorkflow.cpp
@@ -0,0 +1,65 @@
+// Copyright Epic Games, Inc. All Rights Reserved.
+
+#include "Units/ControlRigNodeWorkflow.h"
+#include "RigVMCore/RigVMStruct.h"
+
+bool UControlRigWorkflowOptions::EnsureAtLeastOneRigElementSelected() const
+{
+	if(Selection.IsEmpty())
+	{
+		static constexpr TCHAR SelectAtLeastOneRigElementMessage[] = TEXT("Please select at least one element in the hierarchy!");
+		Reportf(EMessageSeverity::Error, SelectAtLeastOneRigElementMessage);
+		return false;
+	}
+	return true;
+}
+
+// provides the default workflows for any pin
+TArray<FRigVMUserWorkflow> UControlRigTransformWorkflowOptions::ProvideWorkflows(const UObject* InSubject)
+{
+	TArray<FRigVMUserWorkflow> Workflows;
+
+	if(const URigVMPin* Pin = Cast<URigVMPin>(InSubject))
+	{
+		if(!Pin->IsArray())
+		{
+			if(Pin->GetCPPType() == TBaseStructure<FTransform>::Get()->GetStructCPPName())
+			{
+				Workflows.Emplace(
+					TEXT("Set from hierarchy"),
+					TEXT("Sets the pin to match the global transform of the selected element in the hierarchy"),
+					ERigVMUserWorkflowType::PinContext,
+					FRigVMWorkflowGetActionsDelegate::CreateStatic(&UControlRigTransformWorkflowOptions::ProvideTransformWorkflow),
+					StaticClass()
+				);
+			}
+		}
+	}
+	return Workflows;
+}
+
+TArray<FRigVMUserWorkflowAction> UControlRigTransformWorkflowOptions::ProvideTransformWorkflow(
+	const URigVMUserWorkflowOptions* InOptions)
+{
+	TArray<FRigVMUserWorkflowAction> Actions;
+
+	if(const UControlRigTransformWorkflowOptions* Options = Cast<UControlRigTransformWorkflowOptions>(InOptions))
+	{
+		if(Options->EnsureAtLeastOneRigElementSelected())
+		{
+			const FRigElementKey& Key = Options->Selection[0];
+			if(FRigTransformElement* TransformElement = (FRigTransformElement*)Options->Hierarchy->Find<FRigTransformElement>(Key))
+			{
+				const FTransform Transform = Options->Hierarchy->GetTransform(TransformElement, Options->TransformType);
+				
+				Actions.Emplace(
+					ERigVMUserWorkflowActionType::SetPinDefaultValue,
+					InOptions->GetSubject<URigVMPin>(),
+					FRigVMStruct::ExportToFullyQualifiedText<FTransform>(Transform)
+				);
+			}
+		}
+	}
+	
+	return Actions;
+}
diff --git a/Engine/Plugins/Animation/ControlRig/Source/ControlRig/Public/Units/ControlRigNodeWorkflow.h b/Engine/Plugins/Animation/ControlRig/Source/ControlRig/Public/Units/ControlRigNodeWorkflow.h
new file mode 100644
--- /dev/null
+++ b/Engine/Plugins/Animation/ControlRig/Source/ControlRig/Public/Units/ControlRigNodeWorkflow.h
@@ -0,0 +1,44 @@
+// Copyright Epic Games, Inc. All Rights Reserved.
+
+#pragma once
+
+#include "Rigs/RigHierarchy.h"
+#include "RigVMCore/RigVMUserWorkflow.h"
+#include "RigVMModel/RigVMNode.h"
+#include "ControlRigNodeWorkflow.generated.h"
+
+UCLASS(BlueprintType)
+class CONTROLRIG_API UControlRigWorkflowOptions : public URigVMUserWorkflowOptions
+{
+	GENERATED_BODY()
+
+public:
+	
+	UPROPERTY(BlueprintReadWrite, EditDefaultsOnly, Category = Options)
+	TObjectPtr<const URigHierarchy> Hierarchy;
+
+	UPROPERTY(BlueprintReadWrite, EditDefaultsOnly, Category = Options)
+	TArray<FRigElementKey> Selection;
+
+	UFUNCTION(BlueprintCallable, Category = "Options")
+	bool EnsureAtLeastOneRigElementSelected() const;
+};
+
+UCLASS(BlueprintType)
+class CONTROLRIG_API UControlRigTransformWorkflowOptions : public UControlRigWorkflowOptions
+{
+	GENERATED_BODY()
+
+public:
+
+	// The type of transform to retrieve from the hierarchy
+	UPROPERTY(BlueprintReadWrite, EditAnywhere, Category = Options)
+	TEnumAsByte<ERigTransformType::Type> TransformType = ERigTransformType::CurrentGlobal;
+
+	UFUNCTION()
+	TArray<FRigVMUserWorkflow> ProvideWorkflows(const UObject* InSubject);
+
+protected:
+
+	static TArray<FRigVMUserWorkflowAction> ProvideTransformWorkflow(const URigVMUserWorkflowOptions* InOptions);
+};
diff --git a/Engine/Plugins/Animation/ControlRig/Source/ControlRigDeveloper/Private/ControlRigBlueprint.cpp b/Engine/Plugins/Animation/ControlRig/Source/ControlRigDeveloper/Private/ControlRigBlueprint.cpp
--- a/Engine/Plugins/Animation/ControlRig/Source/ControlRigDeveloper/Private/ControlRigBlueprint.cpp
+++ b/Engine/Plugins/Animation/ControlRig/Source/ControlRigDeveloper/Private/ControlRigBlueprint.cpp
@@ -28,6 +28,8 @@
 #include "RigVMPythonUtils.h"
 #include "RigVMTypeUtils.h"
 #include "RigVMModel/Nodes/RigVMAggregateNode.h"
+#include "Units/ControlRigNodeWorkflow.h"
+#include "Rigs/RigControlHierarchy.h"
 
 #if WITH_EDITOR
 #include "IControlRigEditorModule.h"
@@ -1601,6 +1603,24 @@ URigVMController* UControlRigBlueprint::GetOrCreateController(URigVMGraph* InGra
 		}
 	});
 
+	Controller->ConfigureWorkflowOptionsDelegate.BindLambda([WeakThis](URigVMUserWorkflowOptions* Options)
+	{
+		if(UControlRigWorkflowOptions* ControlRigNodeWorkflowOptions = Cast<UControlRigWorkflowOptions>(Options))
+		{
+			ControlRigNodeWorkflowOptions->Hierarchy = nullptr;
+			ControlRigNodeWorkflowOptions->Selection.Reset();
+			
+			if(const UControlRigBlueprint* StrongThis = WeakThis.Get())
+			{
+				if(UControlRig* ControlRig = Cast<UControlRig>(StrongThis->GetObjectBeingDebugged()))
+				{
+					ControlRigNodeWorkflowOptions->Hierarchy = ControlRig->GetHierarchy();
+				}
+				ControlRigNodeWorkflowOptions->Selection = StrongThis->Hierarchy->GetSelectedKeys();
+			}
+		}
+	});
+
 #endif
 
 	Controller->RemoveStaleNodes();
diff --git a/Engine/Plugins/Animation/ControlRig/Source/ControlRigDeveloper/Private/Graph/ControlRigGraph.cpp b/Engine/Plugins/Animation/ControlRig/Source/ControlRigDeveloper/Private/Graph/ControlRigGraph.cpp
--- a/Engine/Plugins/Animation/ControlRig/Source/ControlRigDeveloper/Private/Graph/ControlRigGraph.cpp
+++ b/Engine/Plugins/Animation/ControlRig/Source/ControlRigDeveloper/Private/Graph/ControlRigGraph.cpp
@@ -595,6 +595,16 @@ void UControlRigGraph::HandleModifiedEvent(ERigVMGraphNotifType InNotifType, URi
 					else if (Cast<URigVMUnitNode>(ModelPin->GetNode()))
 					{
 						RigNode->InvalidateNodeTitle();
+
+						// if the node contains a rig element key - invalidate the node
+						if(RigNode->GetAllPins().ContainsByPredicate([](UEdGraphPin* Pin) -> bool
+						{
+							return Pin->PinType.PinSubCategoryObject == FRigElementKey::StaticStruct();
+						}))
+						{
+							// we do this to enforce the refresh of the element name widgets
+							RigNode->ReconstructNode_Internal(true);
+						}
 					}
 				}
 				else if (URigVMInjectionInfo* Injection = ModelPin->GetNode()->GetInjectionInfo())
diff --git a/Engine/Plugins/Animation/ControlRig/Source/ControlRigEditor/Private/ControlRigEditorModule.cpp b/Engine/Plugins/Animation/ControlRig/Source/ControlRigEditor/Private/ControlRigEditorModule.cpp
--- a/Engine/Plugins/Animation/ControlRig/Source/ControlRigEditor/Private/ControlRigEditorModule.cpp
+++ b/Engine/Plugins/Animation/ControlRig/Source/ControlRigEditor/Private/ControlRigEditorModule.cpp
@@ -126,6 +126,8 @@
 #include "UserDefinedStructure/UserDefinedStructEditorData.h"
 #include "UObject/FieldIterator.h"
 #include "RigVMModel/Nodes/RigVMAggregateNode.h"
+#include "RigVMUserWorkflowRegistry.h"
+#include "Units/ControlRigNodeWorkflow.h"
 
 #define LOCTEXT_NAMESPACE "ControlRigEditorModule"
 
@@ -251,11 +253,17 @@ void FControlRigEditorModule::StartupModule()
 	//UThumbnailManager::Get().RegisterCustomRenderer(UControlRigPoseAsset::StaticClass(), UControlRigPoseThumbnailRenderer::StaticClass());
 
 	bFilterAssetBySkeleton = true;
+
+	URigVMUserWorkflowRegistry* WorkflowRegistry = URigVMUserWorkflowRegistry::Get();
+
+	// register the workflow provider for ANY node
+	FRigVMUserWorkflowProvider Provider;
+	Provider.BindUFunction(UControlRigTransformWorkflowOptions::StaticClass()->GetDefaultObject(), TEXT("ProvideWorkflows"));
+	WorkflowHandles.Add(WorkflowRegistry->RegisterProvider(nullptr, Provider));
 }
 
 void FControlRigEditorModule::ShutdownModule()
 {
-
 	if (ICurveEditorModule* CurveEditorModule = FModuleManager::GetModulePtr<ICurveEditorModule>("CurveEditor"))
 	{
 		CurveEditorModule->UnregisterView(FControlRigSpaceChannelCurveModel::ViewID);
@@ -321,6 +329,15 @@ void FControlRigEditorModule::ShutdownModule()
 		typedef IAnimationEditorModule::FAnimationEditorToolbarExtender DelegateType;
 		AnimationEditorModule->GetAllAnimationEditorToolbarExtenders().RemoveAll([=](const DelegateType& In) { return In.GetHandle() == AnimationEditorExtenderHandle; });
 	}
+
+	for(const int32 WorkflowHandle : WorkflowHandles)
+	{
+		if(URigVMUserWorkflowRegistry::StaticClass()->GetDefaultObject(false) != nullptr)
+		{
+			URigVMUserWorkflowRegistry::Get()->UnregisterProvider(WorkflowHandle);
+		}
+	}
+	WorkflowHandles.Reset();
 }
 
 TSharedRef<FExtender> FControlRigEditorModule::GetAnimationEditorToolbarExtender(const TSharedRef<FUICommandList> CommandList, TSharedRef<IAnimationEditor> InAnimationEditor)
@@ -1218,10 +1235,75 @@ void FControlRigEditorModule::GetContextMenuActions(const UControlRigGraphSchema
 	{
 		Schema->UEdGraphSchema::GetContextMenuActions(Menu, Context);
 
+		auto ShowWorkflowOptionsDialog = [](URigVMUserWorkflowOptions* InOptions)
+		{
+			FPropertyEditorModule& EditModule = FModuleManager::Get().GetModuleChecked<FPropertyEditorModule>("PropertyEditor");
+
+			FDetailsViewArgs DetailsViewArgs;
+			DetailsViewArgs.NameAreaSettings = FDetailsViewArgs::HideNameArea;
+			DetailsViewArgs.bHideSelectionTip = true;
+			DetailsViewArgs.DefaultsOnlyVisibility = EEditDefaultsOnlyNodeVisibility::Hide;
+			DetailsViewArgs.bAllowSearch = false;
+
+			const TSharedRef<class IDetailsView> PropertyView = EditModule.CreateDetailView( DetailsViewArgs );
+			PropertyView->SetObject(InOptions);
+
+			TSharedRef<SCustomDialog> OptionsDialog = SNew(SCustomDialog)
+				.Title(FText(LOCTEXT("ControlRigWorkflowOptions", "Options")))
+				.Content()
+				[
+					PropertyView
+				]
+				.Buttons({
+					SCustomDialog::FButton(LOCTEXT("OK", "OK")),
+					SCustomDialog::FButton(LOCTEXT("Cancel", "Cancel"))
+			});
+			return OptionsDialog->ShowModal() == 0;
+		};
+
 		if (UEdGraphPin* InGraphPin = (UEdGraphPin* )Context->Pin)
 		{
 			UEdGraph* Graph = InGraphPin->GetOwningNode()->GetGraph();
 
+			// per pin workflows
+			if(UControlRigBlueprint* RigBlueprint = Cast<UControlRigBlueprint>((UBlueprint*)Context->Blueprint))
+			{
+				if (URigVMPin* ModelPin = RigBlueprint->GetModel(Graph)->FindPin(InGraphPin->GetName()))
+				{
+					const TArray<FRigVMUserWorkflow> Workflows = ModelPin->GetNode()->GetSupportedWorkflows(ERigVMUserWorkflowType::PinContext, ModelPin);
+					if(!Workflows.IsEmpty())
+					{
+						URigVMController* Controller = RigBlueprint->GetController(ModelPin->GetGraph());
+
+						FToolMenuSection& SettingsSection = Menu->AddSection("EdGraphSchemaWorkflow", LOCTEXT("WorkflowHeader", "Workflow"));
+
+						for(const FRigVMUserWorkflow& Workflow : Workflows)
+						{
+							SettingsSection.AddMenuEntry(
+								*Workflow.GetTitle(),
+								FText::FromString(Workflow.GetTitle()),
+								FText::FromString(Workflow.GetTooltip()),
+								FSlateIcon(),
+								FUIAction(FExecuteAction::CreateLambda([Controller, Workflow, ModelPin, ShowWorkflowOptionsDialog]()
+								{
+									URigVMUserWorkflowOptions* Options = Controller->MakeOptionsForWorkflow(ModelPin, Workflow);
+
+									bool bPerform = true;
+									if(Options->RequiresDialog())
+									{
+										bPerform = ShowWorkflowOptionsDialog(Options);
+									}
+									if(bPerform)
+									{
+										Controller->PerformUserWorkflow(Workflow, Options);
+									}
+								}))
+							);
+						}
+					}
+				}
+			}
+
 			bool bIsExecutePin = false;
 			bool bIsUnknownPin = false;
 			if(UScriptStruct* ScriptStruct = Cast<UScriptStruct>(InGraphPin->PinType.PinSubCategoryObject))
@@ -1977,6 +2059,43 @@ void FControlRigEditorModule::GetContextMenuActions(const UControlRigGraphSchema
 					));
 				}
 
+				// per node workflows
+				if (const UControlRigGraphNode* RigNode = Cast<const UControlRigGraphNode>(Context->Node))
+				{
+					if (URigVMNode* ModelNode = RigNode->GetModelNode())
+					{
+						const TArray<FRigVMUserWorkflow> Workflows = ModelNode->GetSupportedWorkflows(ERigVMUserWorkflowType::NodeContext, ModelNode);
+						if(!Workflows.IsEmpty())
+						{
+							FToolMenuSection& SettingsSection = Menu->AddSection("EdGraphSchemaWorkflow", LOCTEXT("WorkflowHeader", "Workflow"));
+
+							for(const FRigVMUserWorkflow& Workflow : Workflows)
+							{
+								SettingsSection.AddMenuEntry(
+									*Workflow.GetTitle(),
+									FText::FromString(Workflow.GetTitle()),
+									FText::FromString(Workflow.GetTooltip()),
+									FSlateIcon(),
+									FUIAction(FExecuteAction::CreateLambda([Controller, Workflow, ModelNode, ShowWorkflowOptionsDialog]()
+									{
+										URigVMUserWorkflowOptions* Options = Controller->MakeOptionsForWorkflow(ModelNode, Workflow);
+
+										bool bPerform = true;
+										if(Options->RequiresDialog())
+										{
+											bPerform = ShowWorkflowOptionsDialog(Options);
+										}
+										if(bPerform)
+										{
+											Controller->PerformUserWorkflow(Workflow, Options);
+										}
+									}))
+								);
+							}
+						}
+					}
+				}
+
 				if (const UControlRigGraphNode* RigNode = Cast<const UControlRigGraphNode>(Context->Node))
 				{
 					if (URigVMUnitNode* UnitNode = Cast<URigVMUnitNode>(RigNode->GetModelNode()))
diff --git a/Engine/Plugins/Animation/ControlRig/Source/ControlRigEditor/Private/ControlRigEditorModule.h b/Engine/Plugins/Animation/ControlRig/Source/ControlRigEditor/Private/ControlRigEditorModule.h
--- a/Engine/Plugins/Animation/ControlRig/Source/ControlRigEditor/Private/ControlRigEditorModule.h
+++ b/Engine/Plugins/Animation/ControlRig/Source/ControlRigEditor/Private/ControlRigEditorModule.h
@@ -94,4 +94,6 @@ private:
 	/** Param to hold Filter Result to pass to Filter*/
 	bool bFilterAssetBySkeleton;
 
+	/** Handles for all registered workflows */
+	TArray<int32> WorkflowHandles;
 };
\ No newline at end of file
