Control Rig: Fix build break introduced in 19625128

diff --git a/Engine/Plugins/Animation/ControlRig/Source/ControlRig/Private/Units/ControlRigNodeWorkflow.cpp b/Engine/Plugins/Animation/ControlRig/Source/ControlRig/Private/Units/ControlRigNodeWorkflow.cpp
--- a/Engine/Plugins/Animation/ControlRig/Source/ControlRig/Private/Units/ControlRigNodeWorkflow.cpp
+++ b/Engine/Plugins/Animation/ControlRig/Source/ControlRig/Private/Units/ControlRigNodeWorkflow.cpp
@@ -3,6 +3,10 @@
 #include "Units/ControlRigNodeWorkflow.h"
 #include "RigVMCore/RigVMStruct.h"
 
+#if WITH_EDITOR
+#include "RigVMModel/RigVMNode.h"
+#endif
+
 bool UControlRigWorkflowOptions::EnsureAtLeastOneRigElementSelected() const
 {
 	if(Selection.IsEmpty())
@@ -18,7 +22,7 @@ bool UControlRigWorkflowOptions::EnsureAtLeastOneRigElementSelected() const
 TArray<FRigVMUserWorkflow> UControlRigTransformWorkflowOptions::ProvideWorkflows(const UObject* InSubject)
 {
 	TArray<FRigVMUserWorkflow> Workflows;
-
+#if WITH_EDITOR
 	if(const URigVMPin* Pin = Cast<URigVMPin>(InSubject))
 	{
 		if(!Pin->IsArray())
@@ -35,9 +39,11 @@ TArray<FRigVMUserWorkflow> UControlRigTransformWorkflowOptions::ProvideWorkflows
 			}
 		}
 	}
+#endif
 	return Workflows;
 }
 
+#if WITH_EDITOR
 TArray<FRigVMUserWorkflowAction> UControlRigTransformWorkflowOptions::ProvideTransformWorkflow(
 	const URigVMUserWorkflowOptions* InOptions)
 {
@@ -63,3 +69,4 @@ TArray<FRigVMUserWorkflowAction> UControlRigTransformWorkflowOptions::ProvideTra
 	
 	return Actions;
 }
+#endif
\ No newline at end of file
diff --git a/Engine/Plugins/Animation/ControlRig/Source/ControlRig/Public/Units/ControlRigNodeWorkflow.h b/Engine/Plugins/Animation/ControlRig/Source/ControlRig/Public/Units/ControlRigNodeWorkflow.h
--- a/Engine/Plugins/Animation/ControlRig/Source/ControlRig/Public/Units/ControlRigNodeWorkflow.h
+++ b/Engine/Plugins/Animation/ControlRig/Source/ControlRig/Public/Units/ControlRigNodeWorkflow.h
@@ -4,7 +4,6 @@
 
 #include "Rigs/RigHierarchy.h"
 #include "RigVMCore/RigVMUserWorkflow.h"
-#include "RigVMModel/RigVMNode.h"
 #include "ControlRigNodeWorkflow.generated.h"
 
 UCLASS(BlueprintType)
@@ -40,5 +39,7 @@ public:
 
 protected:
 
+#if WITH_EDITOR
 	static TArray<FRigVMUserWorkflowAction> ProvideTransformWorkflow(const URigVMUserWorkflowOptions* InOptions);
+#endif
 };
