Control Rig: Hierarchy ordered selection fix for serialize

diff --git a/Engine/Plugins/Animation/ControlRig/Source/ControlRig/Private/Rigs/RigHierarchy.cpp b/Engine/Plugins/Animation/ControlRig/Source/ControlRig/Private/Rigs/RigHierarchy.cpp
--- a/Engine/Plugins/Animation/ControlRig/Source/ControlRig/Private/Rigs/RigHierarchy.cpp
+++ b/Engine/Plugins/Animation/ControlRig/Source/ControlRig/Private/Rigs/RigHierarchy.cpp
@@ -272,6 +272,7 @@ void URigHierarchy::Load(FArchive& Ar)
 			if(FRigBaseElement* Element = Find<FRigBaseElement>(SelectedKey))
 			{
 				Element->bSelected = true;
+				OrderedSelection.Add(SelectedKey);
 			}
 		}
 	}
@@ -311,6 +312,7 @@ void URigHierarchy::Reset()
 	ResetPoseHash = INDEX_NONE;
 	ResetPoseIsFilteredOut.Reset();
 	DefaultParentPerElement.Reset();
+	OrderedSelection.Reset();
 
 	if(!IsGarbageCollecting())
 	{
