Control Rig: Disable stacked hierarchy view by default

diff --git a/Engine/Plugins/Animation/ControlRig/Source/ControlRig/Private/Settings/ControlRigSettings.cpp b/Engine/Plugins/Animation/ControlRig/Source/ControlRig/Private/Settings/ControlRigSettings.cpp
--- a/Engine/Plugins/Animation/ControlRig/Source/ControlRig/Private/Settings/ControlRigSettings.cpp
+++ b/Engine/Plugins/Animation/ControlRig/Source/ControlRig/Private/Settings/ControlRigSettings.cpp
@@ -31,7 +31,7 @@ UControlRigEditorSettings::UControlRigEditorSettings(const FObjectInitializer& O
 	BackwardsSolveBorderColor = FLinearColor::Yellow;
 	BackwardsAndForwardsBorderColor = FLinearColor::Blue;
 	DynamicHierarchyBorderColor = FLinearColor::Blue;
-	bShowStackedHierarchy = true;
+	bShowStackedHierarchy = false;
 	MaxStackSize = 16;
 #endif
 }
