Control rig and RigVM changes to allow for data interface prototype to be further developed

This removes some of the dependencies on UControlRigBlueprint and replaces them with 'host' interfaces for URigVMController and URigVMGraph.
Also adds support for TScriptInterfaces in RigVM. This includes a modification of the RigVM parts of UHT.
Some of these changes are due to be reworked later to allow RigVM-hosting systems to not have to be also implemented in terms of control rig.

diff --git a/Engine/Plugins/Animation/ControlRig/Source/ControlRigDeveloper/Private/ControlRigBlueprint.cpp b/Engine/Plugins/Animation/ControlRig/Source/ControlRigDeveloper/Private/ControlRigBlueprint.cpp
--- a/Engine/Plugins/Animation/ControlRig/Source/ControlRigDeveloper/Private/ControlRigBlueprint.cpp
+++ b/Engine/Plugins/Animation/ControlRig/Source/ControlRigDeveloper/Private/ControlRigBlueprint.cpp
@@ -1373,7 +1373,7 @@ URigVMFunctionLibrary* UControlRigBlueprint::GetLocalFunctionLibrary() const
 	return FunctionLibrary;
 }
 
-URigVMController* UControlRigBlueprint::GetController(URigVMGraph* InGraph) const
+URigVMController* UControlRigBlueprint::GetController(const URigVMGraph* InGraph) const
 {
 	if (InGraph == nullptr)
 	{
diff --git a/Engine/Plugins/Animation/ControlRig/Source/ControlRigDeveloper/Private/Graph/ControlRigGraph.cpp b/Engine/Plugins/Animation/ControlRig/Source/ControlRigDeveloper/Private/Graph/ControlRigGraph.cpp
--- a/Engine/Plugins/Animation/ControlRig/Source/ControlRigDeveloper/Private/Graph/ControlRigGraph.cpp
+++ b/Engine/Plugins/Animation/ControlRig/Source/ControlRigDeveloper/Private/Graph/ControlRigGraph.cpp
@@ -71,7 +71,10 @@ void UControlRigGraph::Serialize(FArchive& Ar)
 
 	if (Ar.IsLoading())
 	{
-		Schema = UControlRigGraphSchema::StaticClass();
+		if(Schema == nullptr || !Schema->IsChildOf(UControlRigGraphSchema::StaticClass()))
+		{
+			Schema = UControlRigGraphSchema::StaticClass();
+		}
 	}
 }
 #endif
@@ -351,7 +354,7 @@ void UControlRigGraph::HandleModifiedEvent(ERigVMGraphNotifType InNotifType, URi
 				}
 				else if (URigVMRerouteNode* RerouteModelNode = Cast<URigVMRerouteNode>(ModelNode))
 				{
-					UControlRigGraphNode* NewNode = NewObject<UControlRigGraphNode>(this, ModelNode->GetFName());
+					UControlRigGraphNode* NewNode = NewObject<UControlRigGraphNode>(this, GetControlRigGraphSchema()->GetGraphNodeClass(), ModelNode->GetFName());
 					AddNode(NewNode, false, false);
 
 					NewNode->ModelNodePath = ModelNode->GetNodePath();
@@ -372,7 +375,7 @@ void UControlRigGraph::HandleModifiedEvent(ERigVMGraphNotifType InNotifType, URi
 				}
 				else // struct, library, parameter + variable
 				{
-					UControlRigGraphNode* NewNode = NewObject<UControlRigGraphNode>(this, ModelNode->GetFName());
+					UControlRigGraphNode* NewNode = NewObject<UControlRigGraphNode>(this, GetControlRigGraphSchema()->GetGraphNodeClass(), ModelNode->GetFName());
 					AddNode(NewNode, false, false);
 
 					NewNode->ModelNodePath = ModelNode->GetNodePath();
@@ -908,18 +911,18 @@ UControlRigBlueprint* UControlRigGraph::GetBlueprint() const
 
 URigVMGraph* UControlRigGraph::GetModel() const
 {
-	if (UControlRigBlueprint* Blueprint = GetBlueprint())
+	if (const IRigVMGraphHost* GraphHost = GetImplementingOuter<IRigVMGraphHost>())
 	{
-		return Blueprint->GetModel(this);
+		return GraphHost->GetRigVMGraph(this);
 	}
 	return nullptr;
 }
 
 URigVMController* UControlRigGraph::GetController() const
 {
-	if (UControlRigBlueprint* Blueprint = GetBlueprint())
+	if (IRigVMControllerHost* ControllerHost = GetImplementingOuter<IRigVMControllerHost>())
 	{
-		return Blueprint->GetOrCreateController(this);
+		return ControllerHost->GetOrCreateRigVMController(this);
 	}
 	return nullptr;
 }
@@ -964,7 +967,7 @@ URigVMController* UControlRigGraph::GetTemplateController()
 	return TemplateController;
 }
 
-void UControlRigGraph::HandleVMCompiledEvent(UBlueprint* InBlueprint, URigVM* InVM)
+void UControlRigGraph::HandleVMCompiledEvent(UObject* InCompiledObject, URigVM* InVM)
 {
 	CachedInstructionIndices.Reset();
 }
diff --git a/Engine/Plugins/Animation/ControlRig/Source/ControlRigDeveloper/Private/Graph/ControlRigGraphNode.cpp b/Engine/Plugins/Animation/ControlRig/Source/ControlRigDeveloper/Private/Graph/ControlRigGraphNode.cpp
--- a/Engine/Plugins/Animation/ControlRig/Source/ControlRigDeveloper/Private/Graph/ControlRigGraphNode.cpp
+++ b/Engine/Plugins/Animation/ControlRig/Source/ControlRigDeveloper/Private/Graph/ControlRigGraphNode.cpp
@@ -895,16 +895,17 @@ void UControlRigGraphNode::DestroyNode()
 
 	if(UControlRigGraph* Graph = Cast<UControlRigGraph>(GetOuter()))
 	{
+		BreakAllNodeLinks();
+		
 		UControlRigBlueprint* ControlRigBlueprint = Cast<UControlRigBlueprint>(Graph->GetOuter());
 		if(ControlRigBlueprint)
-	{
-			BreakAllNodeLinks();
-			if(PropertyName_DEPRECATED.IsValid())
 		{
+			if(PropertyName_DEPRECATED.IsValid())
+			{
 				FControlRigBlueprintUtils::RemoveMemberVariableIfNotUsed(ControlRigBlueprint, PropertyName_DEPRECATED, this);
 			}
 		}
-		}
+	}
 
 	UEdGraphNode::DestroyNode();
 }
@@ -936,7 +937,8 @@ void UControlRigGraphNode::CopyPinDefaultsToModel(UEdGraphPin* Pin, bool bUndo,
 		if(DefaultValue.IsEmpty() && (
 			Pin->PinType.PinCategory == UEdGraphSchema_K2::PC_Object ||
 			Pin->PinType.PinCategory == UEdGraphSchema_K2::PC_SoftObject ||
-			Pin->PinType.PinCategory == UEdGraphSchema_K2::AllObjectTypes
+			Pin->PinType.PinCategory == UEdGraphSchema_K2::AllObjectTypes ||
+			Pin->PinType.PinCategory == UEdGraphSchema_K2::PC_Interface
 			))
 		{
 			if(Pin->DefaultObject)
diff --git a/Engine/Plugins/Animation/ControlRig/Source/ControlRigDeveloper/Private/Graph/ControlRigGraphSchema.cpp b/Engine/Plugins/Animation/ControlRig/Source/ControlRigDeveloper/Private/Graph/ControlRigGraphSchema.cpp
--- a/Engine/Plugins/Animation/ControlRig/Source/ControlRigDeveloper/Private/Graph/ControlRigGraphSchema.cpp
+++ b/Engine/Plugins/Animation/ControlRig/Source/ControlRigDeveloper/Private/Graph/ControlRigGraphSchema.cpp
@@ -556,11 +556,10 @@ bool UControlRigGraphSchema::TryCreateConnection(UEdGraphPin* PinA, UEdGraphPin*
 
 	UControlRigGraphSchema* MutableThis = (UControlRigGraphSchema*)this;
 
-	UBlueprint* Blueprint = FBlueprintEditorUtils::FindBlueprintForNodeChecked(PinA->GetOwningNode());
-	UControlRigBlueprint* RigBlueprint = Cast<UControlRigBlueprint>(Blueprint);
-	if (RigBlueprint != nullptr)
+	IRigVMControllerHost* ControllerHost = PinA->GetOwningNode()->GetImplementingOuter<IRigVMControllerHost>();
+	if (ControllerHost != nullptr)
 	{
-		if (URigVMController* Controller = RigBlueprint->GetOrCreateController(PinA->GetOwningNode()->GetGraph()))
+		if (URigVMController* Controller = ControllerHost->GetOrCreateRigVMController(PinA->GetOwningNode()->GetGraph()))
 		{
 			ERigVMPinDirection UserLinkDirection = ERigVMPinDirection::Output;
 			if (PinA->Direction == EGPD_Input)
@@ -675,45 +674,40 @@ static bool HasChildConnection_Recursive(const UEdGraphPin* InPin)
 
 const FPinConnectionResponse UControlRigGraphSchema::CanCreateConnection(const UEdGraphPin* A, const UEdGraphPin* B) const
 {
-	UBlueprint* Blueprint = FBlueprintEditorUtils::FindBlueprintForNodeChecked(A->GetOwningNode());
-	UControlRigBlueprint* RigBlueprint = Cast<UControlRigBlueprint>(Blueprint);
-	if (RigBlueprint != nullptr)
-	{
-		UControlRigGraphNode* RigNodeA = Cast<UControlRigGraphNode>(A->GetOwningNode());
-		UControlRigGraphNode* RigNodeB = Cast<UControlRigGraphNode>(B->GetOwningNode());
+	UControlRigGraphNode* RigNodeA = Cast<UControlRigGraphNode>(A->GetOwningNode());
+	UControlRigGraphNode* RigNodeB = Cast<UControlRigGraphNode>(B->GetOwningNode());
 
-		if (RigNodeA && RigNodeB && RigNodeA != RigNodeB)
+	if (RigNodeA && RigNodeB && RigNodeA != RigNodeB)
+	{
+		URigVMPin* PinA = RigNodeA->GetModelPinFromPinPath(A->GetName());
+		if (PinA)
 		{
-			URigVMPin* PinA = RigNodeA->GetModelPinFromPinPath(A->GetName());
-			if (PinA)
-			{
-				PinA = PinA->GetPinForLink();
-				RigNodeA->GetModel()->PrepareCycleChecking(PinA, A->Direction == EGPD_Input);
-			}
+			PinA = PinA->GetPinForLink();
+			RigNodeA->GetModel()->PrepareCycleChecking(PinA, A->Direction == EGPD_Input);
+		}
 
-			URigVMPin* PinB = RigNodeB->GetModelPinFromPinPath(B->GetName());
-			if (PinB)
-			{
-				PinB = PinB->GetPinForLink();
-			}
+		URigVMPin* PinB = RigNodeB->GetModelPinFromPinPath(B->GetName());
+		if (PinB)
+		{
+			PinB = PinB->GetPinForLink();
+		}
 
-			ERigVMPinDirection UserLinkDirection = ERigVMPinDirection::Output;
-			if (A->Direction == EGPD_Input)
-			{
-				Swap(PinA, PinB);
-				UserLinkDirection = ERigVMPinDirection::Input;
-			}
+		ERigVMPinDirection UserLinkDirection = ERigVMPinDirection::Output;
+		if (A->Direction == EGPD_Input)
+		{
+			Swap(PinA, PinB);
+			UserLinkDirection = ERigVMPinDirection::Input;
+		}
 
-			const FRigVMByteCode* ByteCode = RigNodeA->GetController()->GetCurrentByteCode();
+		const FRigVMByteCode* ByteCode = RigNodeA->GetController()->GetCurrentByteCode();
 
-			FString FailureReason;
-			bool bResult = RigNodeA->GetModel()->CanLink(PinA, PinB, &FailureReason, ByteCode, UserLinkDirection);
-			if (!bResult)
-			{
-				return FPinConnectionResponse(CONNECT_RESPONSE_DISALLOW, FText::FromString(FailureReason));
-			}
-			return FPinConnectionResponse(CONNECT_RESPONSE_MAKE, LOCTEXT("ConnectResponse_Allowed", "Connect"));
+		FString FailureReason;
+		bool bResult = RigNodeA->GetModel()->CanLink(PinA, PinB, &FailureReason, ByteCode, UserLinkDirection);
+		if (!bResult)
+		{
+			return FPinConnectionResponse(CONNECT_RESPONSE_DISALLOW, FText::FromString(FailureReason));
 		}
+		return FPinConnectionResponse(CONNECT_RESPONSE_MAKE, LOCTEXT("ConnectResponse_Allowed", "Connect"));
 	}
 
 	return FPinConnectionResponse(CONNECT_RESPONSE_DISALLOW, LOCTEXT("ConnectResponse_Disallowed_Unexpected", "Unexpected error"));
@@ -859,6 +853,10 @@ bool UControlRigGraphSchema::SupportsPinType(UScriptStruct* ScriptStruct) const
 		{
 			continue;
 		}
+		else if (CastField<FInterfaceProperty>(Property) && RigVMCore::SupportsUInterfaces())
+		{
+			continue;
+		}
 	
 		return false;
 	}
@@ -898,6 +896,17 @@ bool UControlRigGraphSchema::SupportsPinType(TWeakPtr<const FEdGraphSchemaAction
 		}
 	}
 
+	if (RigVMCore::SupportsUInterfaces())
+	{
+		if (PinType.PinCategory == UEdGraphSchema_K2::PC_Interface)
+		{
+			if (PinType.PinSubCategoryObject.IsValid())
+			{
+				return PinType.PinSubCategoryObject->IsA<UInterface>();
+			}
+		}
+	}
+
 	if (PinType.PinCategory == UEdGraphSchema_K2::PC_Struct)
 	{
 		if (UScriptStruct* ScriptStruct = Cast<UScriptStruct>(PinType.PinSubCategoryObject))
diff --git a/Engine/Plugins/Animation/ControlRig/Source/ControlRigDeveloper/Public/ControlRigBlueprint.h b/Engine/Plugins/Animation/ControlRig/Source/ControlRigDeveloper/Public/ControlRigBlueprint.h
--- a/Engine/Plugins/Animation/ControlRig/Source/ControlRigDeveloper/Public/ControlRigBlueprint.h
+++ b/Engine/Plugins/Animation/ControlRig/Source/ControlRigDeveloper/Public/ControlRigBlueprint.h
@@ -32,7 +32,7 @@ class USkeletalMesh;
 class UControlRigGraph;
 struct FEndLoadPackageContext;
 
-DECLARE_EVENT_TwoParams(UControlRigBlueprint, FOnVMCompiledEvent, UBlueprint*, URigVM*);
+DECLARE_EVENT_TwoParams(UControlRigBlueprint, FOnVMCompiledEvent, UObject*, URigVM*);
 DECLARE_EVENT_OneParam(UControlRigBlueprint, FOnRefreshEditorEvent, UControlRigBlueprint*);
 DECLARE_EVENT_FourParams(UControlRigBlueprint, FOnVariableDroppedEvent, UObject*, FProperty*, const FVector2D&, const FVector2D&);
 DECLARE_EVENT_OneParam(UControlRigBlueprint, FOnExternalVariablesChanged, const TArray<FRigVMExternalVariable>&);
@@ -195,7 +195,7 @@ struct CONTROLRIGDEVELOPER_API FControlRigPythonSettings
 };
 
 UCLASS(BlueprintType, meta=(IgnoreClassThumbnail))
-class CONTROLRIGDEVELOPER_API UControlRigBlueprint : public UBlueprint, public IInterface_PreviewMeshProvider
+class CONTROLRIGDEVELOPER_API UControlRigBlueprint : public UBlueprint, public IInterface_PreviewMeshProvider, public IRigVMGraphHost, public IRigVMControllerHost
 {
 	GENERATED_UCLASS_BODY()
 
@@ -248,6 +248,15 @@ public:
 	// UObject interface
 	virtual void PostEditChangeChainProperty(struct FPropertyChangedChainEvent& PropertyChangedEvent) override;
 
+	// IRigVMGraphHost interface
+	virtual URigVMGraph* GetRigVMGraph(const UObject* InEditorObject) const override { return GetModel(CastChecked<UEdGraph>(InEditorObject)); }
+
+	// IRigVMControllerHost interface
+	virtual URigVMController* GetRigVMController(const URigVMGraph* InRigVMGraph) const override { return GetController(InRigVMGraph); }
+	virtual URigVMController* GetRigVMController(const UObject* InEditorObject) const override { return GetController(CastChecked<UEdGraph>(InEditorObject)); }
+	virtual URigVMController* GetOrCreateRigVMController(URigVMGraph* InRigVMGraph) override { return GetOrCreateController(InRigVMGraph); }
+	virtual URigVMController* GetOrCreateRigVMController(const UObject* InEditorObject) override { return GetOrCreateController(CastChecked<UEdGraph>(InEditorObject)); }
+	
 	FOnRequestInspectObject& OnRequestInspectObject() { return OnRequestInspectObjectEvent; }
 	void RequestInspectObject(const TArray<UObject*>& InObjects) { OnRequestInspectObjectEvent.Broadcast(InObjects); }
 
@@ -294,7 +303,7 @@ public:
 	URigVMFunctionLibrary* GetLocalFunctionLibrary() const;
 
 	UFUNCTION(BlueprintCallable, Category = "Control Rig Blueprint")
-	URigVMController* GetController(URigVMGraph* InGraph = nullptr) const;
+	URigVMController* GetController(const URigVMGraph* InGraph = nullptr) const;
 
 	UFUNCTION(BlueprintCallable, Category = "Control Rig Blueprint")
 	URigVMController* GetControllerByName(const FString InGraphName) const;
diff --git a/Engine/Plugins/Animation/ControlRig/Source/ControlRigDeveloper/Public/Graph/ControlRigGraph.h b/Engine/Plugins/Animation/ControlRig/Source/ControlRigDeveloper/Public/Graph/ControlRigGraph.h
--- a/Engine/Plugins/Animation/ControlRig/Source/ControlRigDeveloper/Public/Graph/ControlRigGraph.h
+++ b/Engine/Plugins/Animation/ControlRig/Source/ControlRigDeveloper/Public/Graph/ControlRigGraph.h
@@ -69,8 +69,8 @@ public:
 	UEdGraphNode* FindNodeForModelNodeName(const FName& InModelNodeName, const bool bCacheIfRequired = true);
 
 	UControlRigBlueprint* GetBlueprint() const;
-	URigVMGraph* GetModel() const;
-	URigVMController* GetController() const;
+	virtual URigVMGraph* GetModel() const;
+	virtual URigVMController* GetController() const;
 
 	void HandleModifiedEvent(ERigVMGraphNotifType InNotifType, URigVMGraph* InGraph, UObject* InSubject);
 
@@ -147,8 +147,11 @@ private:
 #if WITH_EDITOR
 
 	URigVMController* GetTemplateController();
-	void HandleVMCompiledEvent(UBlueprint* InBlueprint, URigVM* InVM);
 
+protected:
+	void HandleVMCompiledEvent(UObject* InCompiledObject, URigVM* InVM);
+
+private:
 	static TArray<TSharedPtr<FString>> EmptyElementNameList;
 	TMap<FName, UEdGraphNode*> ModelNodePathToEdNode;
 	
diff --git a/Engine/Plugins/Animation/ControlRig/Source/ControlRigDeveloper/Public/Graph/ControlRigGraphSchema.h b/Engine/Plugins/Animation/ControlRig/Source/ControlRigDeveloper/Public/Graph/ControlRigGraphSchema.h
--- a/Engine/Plugins/Animation/ControlRig/Source/ControlRigDeveloper/Public/Graph/ControlRigGraphSchema.h
+++ b/Engine/Plugins/Animation/ControlRig/Source/ControlRigDeveloper/Public/Graph/ControlRigGraphSchema.h
@@ -2,6 +2,7 @@
 
 #pragma once
 
+#include "ControlRigGraphNode.h"
 #include "EdGraph/EdGraphSchema.h"
 #include "GraphEditorDragDropAction.h"
 #include "RigVMModel/RigVMGraph.h"
@@ -261,6 +262,9 @@ public:
 
 	void HandleModifiedEvent(ERigVMGraphNotifType InNotifType, URigVMGraph* InGraph, UObject* InSubject);
 
+	// Allow derived classes to spawn derived node classes
+	virtual TSubclassOf<UControlRigGraphNode> GetGraphNodeClass() const { return UControlRigGraphNode::StaticClass(); }
+
 private:
 
 	const UEdGraphPin* LastPinForCompatibleCheck = nullptr;
diff --git a/Engine/Plugins/Animation/ControlRig/Source/ControlRigEditor/Private/ControlRigCompilerDetails.h b/Engine/Plugins/Animation/ControlRig/Source/ControlRigEditor/Private/ControlRigCompilerDetails.h
--- a/Engine/Plugins/Animation/ControlRig/Source/ControlRigEditor/Private/ControlRigCompilerDetails.h
+++ b/Engine/Plugins/Animation/ControlRig/Source/ControlRigEditor/Private/ControlRigCompilerDetails.h
@@ -30,5 +30,5 @@ protected:
 	FReply OnCopyHierarchyGraphClicked();
 	FReply OnCopyGeneratedCodeClicked();
 	
-	UControlRigBlueprint* BlueprintBeingCustomized;
+	UControlRigBlueprint* BlueprintBeingCustomized = nullptr;
 };
diff --git a/Engine/Plugins/Animation/ControlRig/Source/ControlRigEditor/Private/Editor/ControlRigEditor.cpp b/Engine/Plugins/Animation/ControlRig/Source/ControlRigEditor/Private/Editor/ControlRigEditor.cpp
--- a/Engine/Plugins/Animation/ControlRig/Source/ControlRigEditor/Private/Editor/ControlRigEditor.cpp
+++ b/Engine/Plugins/Animation/ControlRig/Source/ControlRigEditor/Private/Editor/ControlRigEditor.cpp
@@ -2820,9 +2820,9 @@ void FControlRigEditor::HandleModifiedEvent(ERigVMGraphNotifType InNotifType, UR
 	}
 }
 
-void FControlRigEditor::HandleVMCompiledEvent(UBlueprint* InBlueprint, URigVM* InVM)
+void FControlRigEditor::HandleVMCompiledEvent(UObject* InCompiledObject, URigVM* InVM)
 {
-	if(UControlRigBlueprint* RigBlueprint = Cast<UControlRigBlueprint>(InBlueprint))
+	if(UControlRigBlueprint* RigBlueprint = Cast<UControlRigBlueprint>(InCompiledObject))
 	{
 		CompilerResultsListing->ClearMessages();
 		CompilerResultsListing->AddMessages(RigBlueprint->CompileLog.Messages);
diff --git a/Engine/Plugins/Animation/ControlRig/Source/ControlRigEditor/Private/Editor/ControlRigEditor.h b/Engine/Plugins/Animation/ControlRig/Source/ControlRigEditor/Private/Editor/ControlRigEditor.h
--- a/Engine/Plugins/Animation/ControlRig/Source/ControlRigEditor/Private/Editor/ControlRigEditor.h
+++ b/Engine/Plugins/Animation/ControlRig/Source/ControlRigEditor/Private/Editor/ControlRigEditor.h
@@ -252,7 +252,7 @@ protected:
 	virtual void FocusInspectorOnGraphSelection(const TSet<class UObject*>& NewSelection, bool bForceRefresh = false) override;
 
 	void HandleModifiedEvent(ERigVMGraphNotifType InNotifType, URigVMGraph* InGraph, UObject* InSubject);
-	void HandleVMCompiledEvent(UBlueprint* InBlueprint, URigVM* InVM);
+	void HandleVMCompiledEvent(UObject* InCompiledObject, URigVM* InVM);
 	void HandleControlRigExecutedEvent(UControlRig* InControlRig, const EControlRigState InState, const FName& InEventName);
 	void HandleControlRigExecutionHalted(const int32 InstructionIndex, UObject* InNode, const FName& InEntryName);
 	void SetHaltedNode(URigVMNode* Node);
diff --git a/Engine/Plugins/Animation/ControlRig/Source/ControlRigEditor/Private/Editor/SControlRigStackView.cpp b/Engine/Plugins/Animation/ControlRig/Source/ControlRigEditor/Private/Editor/SControlRigStackView.cpp
--- a/Engine/Plugins/Animation/ControlRig/Source/ControlRigEditor/Private/Editor/SControlRigStackView.cpp
+++ b/Engine/Plugins/Animation/ControlRig/Source/ControlRigEditor/Private/Editor/SControlRigStackView.cpp
@@ -848,7 +848,7 @@ void SControlRigStackView::HandleFocusOnSelectedGraphNode()
 	}
 }
 
-void SControlRigStackView::OnVMCompiled(UBlueprint* InCompiledBlueprint, URigVM* InCompiledVM)
+void SControlRigStackView::OnVMCompiled(UObject* InCompiledObject, URigVM* InCompiledVM)
 {
 	RefreshTreeView(InCompiledVM);
 
diff --git a/Engine/Plugins/Animation/ControlRig/Source/ControlRigEditor/Private/Editor/SControlRigStackView.h b/Engine/Plugins/Animation/ControlRig/Source/ControlRigEditor/Private/Editor/SControlRigStackView.h
--- a/Engine/Plugins/Animation/ControlRig/Source/ControlRigEditor/Private/Editor/SControlRigStackView.h
+++ b/Engine/Plugins/Animation/ControlRig/Source/ControlRigEditor/Private/Editor/SControlRigStackView.h
@@ -96,7 +96,7 @@ private:
 	/** Focus on the selected operator in the graph*/
 	void HandleFocusOnSelectedGraphNode();
 
-	void OnVMCompiled(UBlueprint* InCompiledBlueprint, URigVM* InCompiledVM);
+	void OnVMCompiled(UObject* InCompiledObject, URigVM* InCompiledVM);
 
 	//* Focus on the instruction when the execution is halted */
 	void HandleExecutionHalted(const int32 HaltedAtInstruction, UObject* InNode, const FName& InEntryName);
diff --git a/Engine/Plugins/Animation/ControlRig/Source/ControlRigEditor/Private/Graph/ControlRigConnectionDrawingPolicy.cpp b/Engine/Plugins/Animation/ControlRig/Source/ControlRigEditor/Private/Graph/ControlRigConnectionDrawingPolicy.cpp
--- a/Engine/Plugins/Animation/ControlRig/Source/ControlRigEditor/Private/Graph/ControlRigConnectionDrawingPolicy.cpp
+++ b/Engine/Plugins/Animation/ControlRig/Source/ControlRigEditor/Private/Graph/ControlRigConnectionDrawingPolicy.cpp
@@ -34,11 +34,9 @@ void FControlRigConnectionDrawingPolicy::ResetIncompatiblePinDrawState(const TSe
 		UEdGraphPin* Pin = PinWidget->GetPinObj();
 		if (Pin != nullptr)
 		{
-			UBlueprint* Blueprint = FBlueprintEditorUtils::FindBlueprintForNodeChecked(Pin->GetOwningNode());
-			UControlRigBlueprint* RigBlueprint = Cast<UControlRigBlueprint>(Blueprint);
-			if (RigBlueprint != nullptr)
+			if(UControlRigGraphNode* RigNode = Cast<UControlRigGraphNode>(Pin->GetOwningNode()))
 			{
-				RigBlueprint->GetModel(Pin->GetOwningNode()->GetGraph())->PrepareCycleChecking(nullptr, true);
+				RigNode->GetModel()->PrepareCycleChecking(nullptr, true);
 			}
 		}
 	}
@@ -313,7 +311,7 @@ void FControlRigConnectionDrawingPolicy::DetermineWiringStyle(UEdGraphPin* Outpu
 
 			if (OutputInstructionIndex != INDEX_NONE && InputInstructionIndex != INDEX_NONE)
 			{
-				if (UControlRigBlueprint* RigBlueprint = Cast<UControlRigBlueprint>(FBlueprintEditorUtils::FindBlueprintForNodeChecked(OutputNode)))
+				if (UControlRigBlueprint* RigBlueprint = Cast<UControlRigBlueprint>(FBlueprintEditorUtils::FindBlueprintForNode(OutputNode)))
 				{
 					if (UControlRig* ControlRig = Cast<UControlRig>(RigBlueprint->GetObjectBeingDebugged()))
 					{
diff --git a/Engine/Source/Developer/RigVMDeveloper/Private/RigVMCompiler/RigVMCompiler.cpp b/Engine/Source/Developer/RigVMDeveloper/Private/RigVMCompiler/RigVMCompiler.cpp
--- a/Engine/Source/Developer/RigVMDeveloper/Private/RigVMCompiler/RigVMCompiler.cpp
+++ b/Engine/Source/Developer/RigVMDeveloper/Private/RigVMCompiler/RigVMCompiler.cpp
@@ -292,10 +292,20 @@ bool URigVMCompiler::Compile(URigVMGraph* InGraph, URigVMController* InControlle
 
 				if(VariableDescription.CPPTypeObject && !RigVMCore::SupportsUObjects())
 				{
-					if(VariableDescription.CPPTypeObject->IsA<UClass>() ||
-						VariableDescription.CPPTypeObject->IsA<UInterface>())
+					if(VariableDescription.CPPTypeObject->IsA<UClass>())
 					{
-						static const FString InvalidObjectTypeMessage = TEXT("Variable Node @@ uses an unsupported UClass / UInterface type.");
+						static const FString InvalidObjectTypeMessage = TEXT("Variable Node @@ uses an unsupported UClass type.");
+						Settings.ASTSettings.Report(EMessageSeverity::Error, ModelNode, InvalidObjectTypeMessage);
+						bEncounteredGraphError = true;
+					}
+				}
+
+
+				if (VariableDescription.CPPTypeObject && !RigVMCore::SupportsUInterfaces())
+				{
+					if (VariableDescription.CPPTypeObject->IsA<UInterface>())
+					{
+						static const FString InvalidObjectTypeMessage = TEXT("Variable Node @@ uses an unsupported UInterface type.");
 						Settings.ASTSettings.Report(EMessageSeverity::Error, ModelNode, InvalidObjectTypeMessage);
 						bEncounteredGraphError = true;
 					}
diff --git a/Engine/Source/Developer/RigVMDeveloper/Private/RigVMModel/Nodes/RigVMUnitNode.cpp b/Engine/Source/Developer/RigVMDeveloper/Private/RigVMModel/Nodes/RigVMUnitNode.cpp
--- a/Engine/Source/Developer/RigVMDeveloper/Private/RigVMModel/Nodes/RigVMUnitNode.cpp
+++ b/Engine/Source/Developer/RigVMDeveloper/Private/RigVMModel/Nodes/RigVMUnitNode.cpp
@@ -342,7 +342,7 @@ TSharedPtr<FStructOnScope> URigVMUnitNode::ConstructStructInstance(bool bUseDefa
 		else
 		{
 			FString StructDefaultValue = GetStructDefaultValue();
-			Struct->ImportText(*StructDefaultValue, StructMemory, nullptr, PPF_None, nullptr, Struct->GetName());
+			Struct->ImportText(*StructDefaultValue, StructMemory, nullptr, PPF_IncludeTransient, GLog, Struct->GetName());
 		}
 		return StructOnScope;
 	}
diff --git a/Engine/Source/Developer/RigVMDeveloper/Private/RigVMModel/RigVMController.cpp b/Engine/Source/Developer/RigVMDeveloper/Private/RigVMModel/RigVMController.cpp
--- a/Engine/Source/Developer/RigVMDeveloper/Private/RigVMModel/RigVMController.cpp
+++ b/Engine/Source/Developer/RigVMDeveloper/Private/RigVMModel/RigVMController.cpp
@@ -12875,6 +12875,19 @@ void URigVMController::ConfigurePinFromProperty(FProperty* InProperty, URigVMPin
 			InOutPin->CPPTypeObject = nullptr;
 		}
 	}
+	else if (FInterfaceProperty* InterfaceProperty = CastField<FInterfaceProperty>(PropertyForType))
+	{
+		if (RigVMCore::SupportsUInterfaces())
+		{
+			InOutPin->CPPTypeObject = InterfaceProperty->InterfaceClass;
+		}
+		else
+		{
+			ReportErrorf(TEXT("Unsupported type '%s' for pin."), *InterfaceProperty->InterfaceClass->GetName(), *InOutPin->GetName());
+			InOutPin->CPPType = FString();
+			InOutPin->CPPTypeObject = nullptr;
+		}
+	}
 	else if (FEnumProperty* EnumProperty = CastField<FEnumProperty>(PropertyForType))
 	{
 		InOutPin->CPPTypeObject = EnumProperty->GetEnum();
diff --git a/Engine/Source/Developer/RigVMDeveloper/Private/RigVMModel/RigVMPin.cpp b/Engine/Source/Developer/RigVMDeveloper/Private/RigVMModel/RigVMPin.cpp
--- a/Engine/Source/Developer/RigVMDeveloper/Private/RigVMModel/RigVMPin.cpp
+++ b/Engine/Source/Developer/RigVMDeveloper/Private/RigVMModel/RigVMPin.cpp
@@ -521,6 +521,11 @@ bool URigVMPin::IsUObject() const
 	return RigVMTypeUtils::IsUObjectType(CPPType);
 }
 
+bool URigVMPin::IsInterface() const
+{
+	return RigVMTypeUtils::IsInterfaceType(CPPType);
+}
+
 bool URigVMPin::IsArray() const
 {
 	return RigVMTypeUtils::IsArrayType(CPPType);
@@ -778,6 +783,7 @@ bool URigVMPin::IsValidDefaultValue(const FString& InDefaultValue) const
 	FString BaseCPPType = GetCPPType()
 		.Replace(RigVMTypeUtils::TArrayPrefix, TEXT(""))
 		.Replace(RigVMTypeUtils::TObjectPtrPrefix, TEXT(""))
+		.Replace(RigVMTypeUtils::TScriptInterfacePrefix, TEXT(""))
 		.Replace(TEXT(">"), TEXT(""));
 
 	for (const FString& Value : DefaultValues)
@@ -908,6 +914,7 @@ FString URigVMPin::ClampDefaultValueFromMetaData(const FString& InDefaultValue)
 		FString BaseCPPType = GetCPPType()
 			.Replace(RigVMTypeUtils::TArrayPrefix, TEXT(""))
 			.Replace(RigVMTypeUtils::TObjectPtrPrefix, TEXT(""))
+			.Replace(RigVMTypeUtils::TScriptInterfacePrefix, TEXT(""))
 			.Replace(TEXT(">"), TEXT(""));
 
 		RetVals.SetNumZeroed(DefaultValues.Num());
diff --git a/Engine/Source/Developer/RigVMDeveloper/Public/RigVMModel/RigVMController.h b/Engine/Source/Developer/RigVMDeveloper/Public/RigVMModel/RigVMController.h
--- a/Engine/Source/Developer/RigVMDeveloper/Public/RigVMModel/RigVMController.h
+++ b/Engine/Source/Developer/RigVMDeveloper/Public/RigVMModel/RigVMController.h
@@ -19,6 +19,7 @@
 #include "RigVMModel/Nodes/RigVMArrayNode.h"
 #include "RigVMModel/RigVMBuildData.h"
 #include "RigVMCore/RigVMUserWorkflow.h"
+#include "UObject/Interface.h"
 #include "RigVMController.generated.h"
 
 #ifndef UE_RIGVM_ENABLE_TEMPLATE_NODES
@@ -90,6 +91,32 @@ DECLARE_DELEGATE_RetVal_TwoParams(FString, FRigVMController_PinPathRemapDelegate
 DECLARE_DELEGATE_OneParam(FRigVMController_RequestJumpToHyperlinkDelegate, const UObject* InSubject);
 DECLARE_DELEGATE_OneParam(FRigVMController_ConfigureWorkflowOptionsDelegate, URigVMUserWorkflowOptions* InOutOptions);
 
+UINTERFACE()
+class RIGVMDEVELOPER_API URigVMControllerHost : public UInterface
+{
+	GENERATED_BODY()
+};
+
+// Interface that allows an object to host a rig VM controller. Used by graph edting code to interact with the controller.
+class RIGVMDEVELOPER_API IRigVMControllerHost
+{
+	GENERATED_BODY()
+
+public:
+	// Get an existing Rig VM controller corresponding to a URigVMGraph  
+	virtual URigVMController* GetRigVMController(const URigVMGraph* InRigVMGraph) const = 0;
+	
+	// Get an existing Rig VM controller corresponding to the supplied editor object that represents it (usually a UEdGraph) 
+	virtual URigVMController* GetRigVMController(const UObject* InEditorObject) const = 0;
+
+	// Get or create Rig VM controller corresponding to a URigVMGraph
+	virtual URigVMController* GetOrCreateRigVMController(URigVMGraph* InRigVMGraph) = 0;
+	
+	// Get or create a Rig VM controller corresponding to the supplied editor object that represents it (usually a UEdGraph) 
+	virtual URigVMController* GetOrCreateRigVMController(const UObject* InEditorObject) = 0;
+};
+
+
 /**
  * The Controller is the sole authority to perform changes
  * on the Graph. The Controller itself is stateless.
diff --git a/Engine/Source/Developer/RigVMDeveloper/Public/RigVMModel/RigVMGraph.h b/Engine/Source/Developer/RigVMDeveloper/Public/RigVMModel/RigVMGraph.h
--- a/Engine/Source/Developer/RigVMDeveloper/Public/RigVMModel/RigVMGraph.h
+++ b/Engine/Source/Developer/RigVMDeveloper/Public/RigVMModel/RigVMGraph.h
@@ -10,10 +10,27 @@
 #include "Nodes/RigVMFunctionEntryNode.h"
 #include "Nodes/RigVMFunctionReturnNode.h"
 #include "RigVMCompiler/RigVMAST.h"
+#include "UObject/Interface.h"
 #include "RigVMGraph.generated.h"
 
 class URigVMFunctionLibrary;
 
+UINTERFACE()
+class RIGVMDEVELOPER_API URigVMGraphHost : public UInterface
+{
+	GENERATED_BODY()
+};
+
+// Interface that allows an object to host a rig VM graph. Used by graph edting code to interact with the graph.
+class RIGVMDEVELOPER_API IRigVMGraphHost
+{
+	GENERATED_BODY()
+
+public:
+	// Get a Rig VM graph corresponding to the supplied editor object that represents it (usually a UEdGraph)
+	virtual URigVMGraph* GetRigVMGraph(const UObject* InEditorObject) const = 0;
+};
+
 /**
  * The Graph represents a Function definition
  * using Nodes as statements.
diff --git a/Engine/Source/Developer/RigVMDeveloper/Public/RigVMModel/RigVMPin.h b/Engine/Source/Developer/RigVMDeveloper/Public/RigVMModel/RigVMPin.h
--- a/Engine/Source/Developer/RigVMDeveloper/Public/RigVMModel/RigVMPin.h
+++ b/Engine/Source/Developer/RigVMDeveloper/Public/RigVMModel/RigVMPin.h
@@ -193,6 +193,10 @@ public:
 	UFUNCTION(BlueprintCallable, Category = RigVMPin)
 	bool IsUObject() const;
 
+	// Returns true if the data type of the Pin is a interface
+	UFUNCTION(BlueprintCallable, Category = RigVMPin)
+	bool IsInterface() const;
+
 	// Returns true if the data type of the Pin is an array
 	UFUNCTION(BlueprintCallable, Category = RigVMPin)
 	bool IsArray() const;
diff --git a/Engine/Source/Programs/UnrealHeaderTool/Private/HeaderParser.cpp b/Engine/Source/Programs/UnrealHeaderTool/Private/HeaderParser.cpp
--- a/Engine/Source/Programs/UnrealHeaderTool/Private/HeaderParser.cpp
+++ b/Engine/Source/Programs/UnrealHeaderTool/Private/HeaderParser.cpp
@@ -6501,6 +6501,7 @@ const TCHAR* FHeaderParser::FTArrayText = TEXT("TArray");
 const TCHAR* FHeaderParser::FTArrayViewText = TEXT("TArrayView");
 const TCHAR* FHeaderParser::GetArrayText = TEXT("GetArray");
 const TCHAR* FHeaderParser::GetArrayViewText = TEXT("GetArrayView");
+const TCHAR* FHeaderParser::FTScriptInterfaceText = TEXT("TScriptInterface");
 
 void FHeaderParser::ParseRigVMMethodParameters(FUnrealStructDefinitionInfo& StructDef)
 {
@@ -6546,16 +6547,22 @@ void FHeaderParser::ParseRigVMMethodParameters(FUnrealStructDefinitionInfo& Stru
 		}
 
 #if !UE_RIGVM_UOBJECT_PROPERTIES_ENABLED
-		if(PropertyDef->GetPropertyBase().IsObjectOrInterface())
+		if(PropertyDef->GetPropertyBase().IsObject())
 		{
-			LogError(TEXT("RigVM Struct '%s' - Member '%s' is a UObject or interface - object types are not allowed on structs with RIGVM_METHOD."), *StructDef.GetName(), *Parameter.Name);
+			LogError(TEXT("RigVM Struct '%s' - Member '%s' is a UObject - object types are not allowed on structs with RIGVM_METHOD."), *StructDef.GetName(), *Parameter.Name);
+		}
+#endif
+#if !UE_RIGVM_UINTERFACE_PROPERTIES_ENABLED
+		if (PropertyDef->GetPropertyBase().IsInterface())
+		{
+			LogError(TEXT("RigVM Struct '%s' - Member '%s' is a UInterface - interface types are not allowed on structs with RIGVM_METHOD."), *StructDef.GetName(), *Parameter.Name);
 		}
 #endif
 
 		if (!ExtendedCPPType.IsEmpty())
 		{
-			// we only support arrays - no maps or similar data structures
-			if (MemberCPPType != TArrayText && MemberCPPType != TEnumAsByteText)
+			// we only support arrays & script interfaces - no maps or similar data structures
+			if (MemberCPPType != TArrayText && MemberCPPType != TEnumAsByteText && MemberCPPType != FTScriptInterfaceText)
 			{
 				LogError(TEXT("RigVM Struct '%s' - Member '%s' type '%s' not supported by RigVM."), *StructDef.GetName(), *Parameter.Name, *MemberCPPType);
 				continue;
diff --git a/Engine/Source/Programs/UnrealHeaderTool/Private/HeaderParser.h b/Engine/Source/Programs/UnrealHeaderTool/Private/HeaderParser.h
--- a/Engine/Source/Programs/UnrealHeaderTool/Private/HeaderParser.h
+++ b/Engine/Source/Programs/UnrealHeaderTool/Private/HeaderParser.h
@@ -668,6 +668,7 @@ public:
 	static const TCHAR* FTArrayViewText;
 	static const TCHAR* GetArrayText;
 	static const TCHAR* GetArrayViewText;
+	static const TCHAR* FTScriptInterfaceText;
 };
 
 /////////////////////////////////////////////////////
diff --git a/Engine/Source/Programs/UnrealHeaderTool/Private/ParserHelper.h b/Engine/Source/Programs/UnrealHeaderTool/Private/ParserHelper.h
--- a/Engine/Source/Programs/UnrealHeaderTool/Private/ParserHelper.h
+++ b/Engine/Source/Programs/UnrealHeaderTool/Private/ParserHelper.h
@@ -162,6 +162,16 @@ inline bool IsObjectOrInterface(EPropertyType Type)
 	return Type == CPT_ObjectReference || Type == CPT_Interface || Type == CPT_WeakObjectReference || Type == CPT_LazyObjectReference || Type == CPT_ObjectPtrReference || Type == CPT_SoftObjectReference;
 }
 
+inline bool IsObject(EPropertyType Type)
+{
+	return Type == CPT_ObjectReference || Type == CPT_WeakObjectReference || Type == CPT_LazyObjectReference || Type == CPT_ObjectPtrReference || Type == CPT_SoftObjectReference;
+}
+
+inline bool IsInterface(EPropertyType Type)
+{
+	return Type == CPT_Interface;
+}
+
 /**
  * Basic information describing a type.
  */
@@ -255,6 +265,16 @@ public:
 		return ::IsObjectOrInterface(Type);
 	}
 
+	bool IsObject() const
+	{
+		return ::IsObject(Type);
+	}
+
+	bool IsInterface() const
+	{
+		return ::IsInterface(Type);
+	}
+
 	bool IsBool() const
 	{
 		return ::IsBool(Type);
diff --git a/Engine/Source/Runtime/RigVM/Private/RigVMCore/RigVMMemoryStorage.cpp b/Engine/Source/Runtime/RigVM/Private/RigVMCore/RigVMMemoryStorage.cpp
--- a/Engine/Source/Runtime/RigVM/Private/RigVMCore/RigVMMemoryStorage.cpp
+++ b/Engine/Source/Runtime/RigVM/Private/RigVMCore/RigVMMemoryStorage.cpp
@@ -67,6 +67,12 @@ FRigVMPropertyDescription::FRigVMPropertyDescription(const FProperty* InProperty
 		{
 			CPPTypeObject = ObjectProperty->PropertyClass;
 		}
+#endif
+#if UE_RIGVM_UINTERFACE_PROPERTIES_ENABLED
+		else if (const FInterfaceProperty* InterfaceProperty = CastField<FInterfaceProperty>(ValueProperty))
+		{
+			CPPTypeObject = InterfaceProperty->InterfaceClass;
+		}
 #endif
 		else if(const FEnumProperty* EnumProperty = CastField<FEnumProperty>(ValueProperty))
 		{
@@ -87,10 +93,14 @@ FRigVMPropertyDescription::FRigVMPropertyDescription(const FProperty* InProperty
 	do
 	{
 		if(ChildProperty->IsA<FObjectProperty>() ||
-			ChildProperty->IsA<FSoftObjectProperty>() ||
-			ChildProperty->IsA<FInterfaceProperty>())
+			ChildProperty->IsA<FSoftObjectProperty>())
+		{
+			checkf(RigVMCore::SupportsUObjects(), TEXT("UClass types are not supported."));
+		}
+
+		if (ChildProperty->IsA<FInterfaceProperty>())
 		{
-			checkf(RigVMCore::SupportsUObjects(), TEXT("UClass / UInterface types are not supported."));
+			checkf(RigVMCore::SupportsUInterfaces(), TEXT("UInterface types are not supported."));
 		}
 
 		if(const FArrayProperty* ArrayProperty = CastField<FArrayProperty>(ChildProperty))
@@ -131,10 +141,12 @@ FRigVMPropertyDescription::FRigVMPropertyDescription(const FName& InName, const
 
 	if(CPPTypeObject)
 	{
-		if(CPPTypeObject->IsA<UClass>() ||
-			CPPTypeObject->IsA<UInterface>())
+		if(CPPTypeObject->IsA<UClass>())
 		{
-			checkf(RigVMCore::SupportsUObjects(), TEXT("UClass / UInterface types are not supported."));
+			if (!RigVMCore::SupportsUInterfaces() || !Cast<UClass>(CPPTypeObject)->IsChildOf<UInterface>())
+			{
+				checkf(RigVMCore::SupportsUObjects(), TEXT("UClass types are not supported."));
+			}
 		}
 	}
 
@@ -247,7 +259,8 @@ bool FRigVMPropertyDescription::RequiresCPPTypeObject(const FString& InCPPType)
 		TEXT("U"), 
 		TEXT("TArray<U"),
 		TEXT("TObjectPtr<"), 
-		TEXT("TArray<TObjectPtr<")
+		TEXT("TArray<TObjectPtr<"),
+		TEXT("TScriptInterface<")
 	};
 	static const TArray<FString> CPPTypesNotRequiringCPPTypeObject = {
 		TEXT("FString"), 
diff --git a/Engine/Source/Runtime/RigVM/Private/RigVMModule.cpp b/Engine/Source/Runtime/RigVM/Private/RigVMModule.cpp
--- a/Engine/Source/Runtime/RigVM/Private/RigVMModule.cpp
+++ b/Engine/Source/Runtime/RigVM/Private/RigVMModule.cpp
@@ -6,6 +6,7 @@
 
 #include "RigVMModule.h"
 #include "Modules/ModuleManager.h"
+#include "HAL/IConsoleManager.h"
 
 IMPLEMENT_MODULE(FDefaultModuleImpl, RigVM);
 
@@ -15,6 +16,10 @@ DEFINE_LOG_CATEGORY(LogRigVM);
 TAutoConsoleVariable<bool> CVarRigVMEnableUObjects(TEXT("RigVM.UObjectSupport"), true, TEXT("When true the RigVMCompiler will allow UObjects."));
 #endif
 
+#if UE_RIGVM_UINTERFACE_PROPERTIES_ENABLED
+TAutoConsoleVariable<bool> CVarRigVMEnableUInterfaces(TEXT("RigVM.UInterfaceSupport"), true, TEXT("When true the RigVMCompiler will allow UInterfaces."));
+#endif
+
 bool RigVMCore::SupportsUObjects()
 {
 #if UE_RIGVM_UOBJECT_PROPERTIES_ENABLED
@@ -23,3 +28,12 @@ bool RigVMCore::SupportsUObjects()
 	return false;
 #endif
 }
+
+bool RigVMCore::SupportsUInterfaces()
+{
+#if UE_RIGVM_UINTERFACE_PROPERTIES_ENABLED
+	return CVarRigVMEnableUInterfaces.GetValueOnGameThread();
+#else
+	return false;
+#endif
+}
\ No newline at end of file
diff --git a/Engine/Source/Runtime/RigVM/Public/RigVMCore/RigVMExecuteContext.h b/Engine/Source/Runtime/RigVM/Public/RigVMCore/RigVMExecuteContext.h
--- a/Engine/Source/Runtime/RigVM/Public/RigVMCore/RigVMExecuteContext.h
+++ b/Engine/Source/Runtime/RigVM/Public/RigVMCore/RigVMExecuteContext.h
@@ -242,13 +242,7 @@ struct RIGVM_API FRigVMExecuteContext
 	FORCEINLINE FName GetEventName() const { return EventName; }
 
 
-private:
-
-	FORCEINLINE void Reset()
-	{
-		InstructionIndex = 0;
-	}
-
+protected:
 	FORCEINLINE void CopyFrom(const FRigVMExecuteContext& Other)
 	{
 		EventName = Other.EventName;
@@ -256,6 +250,13 @@ private:
 		InstructionIndex = Other.InstructionIndex;
 		RuntimeSettings = Other.RuntimeSettings;
 	}
+
+private:
+
+	FORCEINLINE void Reset()
+	{
+		InstructionIndex = 0;
+	}
 	
 	FName EventName;
 	
diff --git a/Engine/Source/Runtime/RigVM/Public/RigVMCore/RigVMExternalVariable.h b/Engine/Source/Runtime/RigVM/Public/RigVMCore/RigVMExternalVariable.h
--- a/Engine/Source/Runtime/RigVM/Public/RigVMCore/RigVMExternalVariable.h
+++ b/Engine/Source/Runtime/RigVM/Public/RigVMCore/RigVMExternalVariable.h
@@ -94,6 +94,21 @@ struct RIGVM_API FRigVMExternalVariable
 				OutTypeObject = nullptr;
 			}
 		}
+		else if (const FInterfaceProperty* InterfaceProperty = CastField<FInterfaceProperty>(InProperty))
+		{
+			if (RigVMCore::SupportsUInterfaces())
+			{
+				OutTypeName = *FString::Printf(TEXT("TScriptInterface<%s%s>"),
+					InterfaceProperty->InterfaceClass->GetPrefixCPP(),
+					*InterfaceProperty->InterfaceClass->GetName());
+				OutTypeObject = InterfaceProperty->InterfaceClass;
+			}
+			else
+			{
+				OutTypeName = NAME_None;
+				OutTypeObject = nullptr;
+			}
+		}
 		else
 		{
 			checkNoEntry();
diff --git a/Engine/Source/Runtime/RigVM/Public/RigVMDefines.h b/Engine/Source/Runtime/RigVM/Public/RigVMDefines.h
--- a/Engine/Source/Runtime/RigVM/Public/RigVMDefines.h
+++ b/Engine/Source/Runtime/RigVM/Public/RigVMDefines.h
@@ -14,6 +14,10 @@
 #define UE_RIGVM_UOBJECT_PROPERTIES_ENABLED 0
 #endif
 
+#ifndef UE_RIGVM_UINTERFACE_PROPERTIES_ENABLED
+#define UE_RIGVM_UINTERFACE_PROPERTIES_ENABLED 1
+#endif
+
 #ifndef UE_RIGVM_AGGREGATE_NODES_ENABLED
 #define UE_RIGVM_AGGREGATE_NODES_ENABLED 0
 #endif
\ No newline at end of file
diff --git a/Engine/Source/Runtime/RigVM/Public/RigVMModule.h b/Engine/Source/Runtime/RigVM/Public/RigVMModule.h
--- a/Engine/Source/Runtime/RigVM/Public/RigVMModule.h
+++ b/Engine/Source/Runtime/RigVM/Public/RigVMModule.h
@@ -14,4 +14,5 @@ RIGVM_API DECLARE_LOG_CATEGORY_EXTERN(LogRigVM, Log, All);
 namespace RigVMCore
 {
 	RIGVM_API bool SupportsUObjects();
+	RIGVM_API bool SupportsUInterfaces();
 }
diff --git a/Engine/Source/Runtime/RigVM/Public/RigVMTypeUtils.h b/Engine/Source/Runtime/RigVM/Public/RigVMTypeUtils.h
--- a/Engine/Source/Runtime/RigVM/Public/RigVMTypeUtils.h
+++ b/Engine/Source/Runtime/RigVM/Public/RigVMTypeUtils.h
@@ -4,13 +4,16 @@
 
 #include "CoreMinimal.h"
 #include "RigVMCore/RigVMUnknownType.h"
+#include "UObject/Interface.h"
 
 namespace RigVMTypeUtils
 {
 	const TCHAR TArrayPrefix[] = TEXT("TArray<");
 	const TCHAR TObjectPtrPrefix[] = TEXT("TObjectPtr<");
+	const TCHAR TScriptInterfacePrefix[] = TEXT("TScriptInterface<");
 	const TCHAR TArrayTemplate[] = TEXT("TArray<%s>");
 	const TCHAR TObjectPtrTemplate[] = TEXT("TObjectPtr<%s%s>");
+	const TCHAR TScriptInterfaceTemplate[] = TEXT("TScriptInterface<%s%s>");
 
 	const FString BoolType = TEXT("bool");
 	const FString FloatType = TEXT("float");
@@ -75,6 +78,11 @@ namespace RigVMTypeUtils
 		return InCPPType.StartsWith(TObjectPtrPrefix);
 	}
 
+	FORCEINLINE bool IsInterfaceType(const FString& InCPPType)
+	{
+		return InCPPType.StartsWith(TScriptInterfacePrefix);
+	}
+
 	static UScriptStruct* GetWildCardCPPTypeObject()
 	{
 		static UScriptStruct* WildCardTypeObject = FRigVMUnknownType::StaticStruct();
@@ -99,7 +107,14 @@ namespace RigVMTypeUtils
 	
 		if (const UClass* Class = Cast<UClass>(InCPPTypeObject))
 		{
-			CPPType = FString::Printf(RigVMTypeUtils::TObjectPtrTemplate, Class->GetPrefixCPP(), *Class->GetName());
+			if (Class->IsChildOf(UInterface::StaticClass()))
+			{
+				CPPType = FString::Printf(RigVMTypeUtils::TScriptInterfaceTemplate, TEXT("I"), *Class->GetName());
+			}
+			else
+			{
+				CPPType = FString::Printf(RigVMTypeUtils::TObjectPtrTemplate, Class->GetPrefixCPP(), *Class->GetName());
+			}
 		}
 		else if (const UScriptStruct* ScriptStruct = Cast<UScriptStruct>(InCPPTypeObject))
 		{
