Control Rig: Make sure we don't try to reference UObjects if they've been shut down.

#rnx

diff --git a/Engine/Plugins/Animation/ControlRig/Source/ControlRigEditor/Private/ControlRigEditorModule.cpp b/Engine/Plugins/Animation/ControlRig/Source/ControlRigEditor/Private/ControlRigEditorModule.cpp
--- a/Engine/Plugins/Animation/ControlRig/Source/ControlRigEditor/Private/ControlRigEditorModule.cpp
+++ b/Engine/Plugins/Animation/ControlRig/Source/ControlRigEditor/Private/ControlRigEditorModule.cpp
@@ -330,11 +330,14 @@ void FControlRigEditorModule::ShutdownModule()
 		AnimationEditorModule->GetAllAnimationEditorToolbarExtenders().RemoveAll([=](const DelegateType& In) { return In.GetHandle() == AnimationEditorExtenderHandle; });
 	}
 
-	for(const int32 WorkflowHandle : WorkflowHandles)
+	if (UObjectInitialized())
 	{
-		if(URigVMUserWorkflowRegistry::StaticClass()->GetDefaultObject(false) != nullptr)
+		for(const int32 WorkflowHandle : WorkflowHandles)
 		{
-			URigVMUserWorkflowRegistry::Get()->UnregisterProvider(WorkflowHandle);
+			if(URigVMUserWorkflowRegistry::StaticClass()->GetDefaultObject(false) != nullptr)
+			{
+				URigVMUserWorkflowRegistry::Get()->UnregisterProvider(WorkflowHandle);
+			}
 		}
 	}
 	WorkflowHandles.Reset();
