FKControlRig changes
- Add ability to generate rig elements from bound USkeleton rather than USkeletalMeshComponent
- Disable retargeting when baking animation sequence to CR
- Ensure raw animation data is sampled when baking animation sequence to CR
- Only set control defaults for frame 0
- Added RefreshActiveControls (for re-use w/ follow-up AnimationSequencer changes)

diff --git a/Engine/Plugins/Animation/ControlRig/Source/ControlRig/Private/Rigs/FKControlRig.cpp b/Engine/Plugins/Animation/ControlRig/Source/ControlRig/Private/Rigs/FKControlRig.cpp
--- a/Engine/Plugins/Animation/ControlRig/Source/ControlRig/Private/Rigs/FKControlRig.cpp
+++ b/Engine/Plugins/Animation/ControlRig/Source/ControlRig/Private/Rigs/FKControlRig.cpp
@@ -262,6 +262,11 @@ void UFKControlRig::Initialize(bool bInitRigUnits /*= true*/)
 	{
 		CreateRigElements(SkeletalMeshComponent->SkeletalMesh);
 	}
+	else if (USkeleton* Skeleton = Cast<USkeleton>(GetObjectBinding()->GetBoundObject()))
+	{
+		const FReferenceSkeleton& RefSkeleton = Skeleton->GetReferenceSkeleton();
+		CreateRigElements(RefSkeleton, (Skeleton) ? Skeleton->GetSmartNameContainer(USkeleton::AnimCurveMappingName) : nullptr);
+	}
 
 	// execute init
 	Execute(EControlRigState::Init, FRigUnit_BeginExecution::EventName);
@@ -363,14 +368,7 @@ void UFKControlRig::CreateRigElements(const FReferenceSkeleton& InReferenceSkele
 			return true;
 		});
 
-		if (IsControlActive.Num() != GetHierarchy()->Num())
-		{
-			IsControlActive.SetNum(GetHierarchy()->Num());
-			for (bool& bIsActive : IsControlActive)
-			{
-				 bIsActive = true;
-			}
-		}
+		RefreshActiveControls();
 	}
 }
 
@@ -415,6 +413,19 @@ void UFKControlRig::SetControlOffsetsFromBoneInitials()
 	}, true);
 }
 
+void UFKControlRig::RefreshActiveControls()
+{
+	if (IsControlActive.Num() != GetHierarchy()->Num())
+	{
+		IsControlActive.Empty(GetHierarchy()->Num());
+		IsControlActive.SetNum(GetHierarchy()->Num());
+		for (bool& bIsActive : IsControlActive)
+		{
+			bIsActive = true;
+		}
+	}
+}
+
 void UFKControlRig::CreateRigElements(const USkeletalMesh* InReferenceMesh)
 {
 	if (InReferenceMesh)
diff --git a/Engine/Plugins/Animation/ControlRig/Source/ControlRig/Private/Sequencer/MovieSceneControlRigParameterSection.cpp b/Engine/Plugins/Animation/ControlRig/Source/ControlRig/Private/Sequencer/MovieSceneControlRigParameterSection.cpp
--- a/Engine/Plugins/Animation/ControlRig/Source/ControlRig/Private/Sequencer/MovieSceneControlRigParameterSection.cpp
+++ b/Engine/Plugins/Animation/ControlRig/Source/ControlRig/Private/Sequencer/MovieSceneControlRigParameterSection.cpp
@@ -2558,10 +2558,17 @@ void UMovieSceneControlRigParameterSection::RecordControlRigKey(FFrameNumber Fra
 	}
 }
 
-bool UMovieSceneControlRigParameterSection::LoadAnimSequenceIntoThisSection(UAnimSequence* AnimSequence, UMovieScene* MovieScene,USkeletalMeshComponent* SkelMeshComp,
-	bool bKeyReduce, float Tolerance, FFrameNumber InStartFrame)
+bool UMovieSceneControlRigParameterSection::LoadAnimSequenceIntoThisSection(UAnimSequence* AnimSequence, UMovieScene* MovieScene, UObject* BoundObject, bool bKeyReduce, float Tolerance, FFrameNumber InStartFrame)
 {
-	if (SkelMeshComp == nullptr || SkelMeshComp->SkeletalMesh == nullptr || SkelMeshComp->SkeletalMesh->GetSkeleton() == nullptr)
+	USkeletalMeshComponent* SkelMeshComp = Cast<USkeletalMeshComponent>(BoundObject);
+	
+	if (SkelMeshComp != nullptr && (SkelMeshComp->SkeletalMesh == nullptr || SkelMeshComp->SkeletalMesh->GetSkeleton() == nullptr))
+	{
+		return false;
+	}
+	
+	USkeleton* Skeleton = SkelMeshComp ? SkelMeshComp->SkeletalMesh->GetSkeleton() : Cast<USkeleton>(BoundObject);	
+	if (Skeleton == nullptr)
 	{
 		return false;
 	}
@@ -2625,7 +2632,14 @@ bool UMovieSceneControlRigParameterSection::LoadAnimSequenceIntoThisSection(UAni
 	const FAnimationCurveData& CurveData = DataModel->GetCurveData();
 	const TArray<FBoneAnimationTrack>& BoneAnimationTracks = DataModel->GetBoneAnimationTracks();
 
-	ControlRig->SetBoneInitialTransformsFromSkeletalMeshComponent(SkelMeshComp, true);
+	if (SkelMeshComp)
+	{
+		ControlRig->SetBoneInitialTransformsFromSkeletalMeshComponent(SkelMeshComp, true);
+	}
+	else
+	{
+		ControlRig->SetBoneInitialTransformsFromRefSkeleton(Skeleton->GetReferenceSkeleton());
+	}
 	ControlRig->RequestSetup();
 	ControlRig->Evaluate_AnyThread();
 
@@ -2635,6 +2649,7 @@ bool UMovieSceneControlRigParameterSection::LoadAnimSequenceIntoThisSection(UAni
 		FFrameNumber FrameNumber = StartFrame + (FrameRateInFrameNumber * Index);
 
 		ControlRig->GetHierarchy()->ResetPoseToInitial();
+		ControlRig->GetHierarchy()->ResetCurveValues();
 
 		for (const FFloatCurve& Curve : CurveData.FloatCurves)
 		{
@@ -2645,7 +2660,10 @@ bool UMovieSceneControlRigParameterSection::LoadAnimSequenceIntoThisSection(UAni
 		// retrieve the pose using the services that persona and sequencer rely on
 		// rather than accessing the low level raw tracks.
 		FAnimPoseEvaluationOptions EvaluationOptions;
-		EvaluationOptions.OptionalSkeletalMesh = SkelMeshComp->SkeletalMesh;
+		EvaluationOptions.OptionalSkeletalMesh = SkelMeshComp ? SkelMeshComp->SkeletalMesh : nullptr;
+		EvaluationOptions.bShouldRetarget = false;
+		EvaluationOptions.EvaluationType = EAnimDataEvalType::Raw;
+
 		FAnimPose AnimPose;
 		UAnimPoseExtensions::GetAnimPoseAtTime(AnimSequence, SequenceSecond, EvaluationOptions, AnimPose);
 
@@ -2668,7 +2686,7 @@ bool UMovieSceneControlRigParameterSection::LoadAnimSequenceIntoThisSection(UAni
 		ControlRig->Execute(EControlRigState::Update, FRigUnit_InverseExecution::EventName);
 
 		const ERichCurveInterpMode InterpMode = bKeyReduce ? RCIM_Cubic : RCIM_Linear;
-		RecordControlRigKey(FrameNumber, true, InterpMode);
+		RecordControlRigKey(FrameNumber, Index == 0, InterpMode);
 		Progress.EnterProgressFrame(1);
 		if (Progress.ShouldCancel())
 		{
diff --git a/Engine/Plugins/Animation/ControlRig/Source/ControlRig/Public/Rigs/FKControlRig.h b/Engine/Plugins/Animation/ControlRig/Source/ControlRig/Public/Rigs/FKControlRig.h
--- a/Engine/Plugins/Animation/ControlRig/Source/ControlRig/Public/Rigs/FKControlRig.h
+++ b/Engine/Plugins/Animation/ControlRig/Source/ControlRig/Public/Rigs/FKControlRig.h
@@ -63,6 +63,8 @@ public:
 	bool CanToggleApplyMode() const { return true; }
 	bool IsApplyModeAdditive() const { return ApplyMode == EControlRigFKRigExecuteMode::Additive; }
 
+	// Ensures that controls mask is updated according to contained ControlRig (control) elements
+	void RefreshActiveControls();
 private:
 
 	/** Create RigElements - bone hierarchy and curves - from incoming skeleton */
diff --git a/Engine/Plugins/Animation/ControlRig/Source/ControlRig/Public/Sequencer/MovieSceneControlRigParameterSection.h b/Engine/Plugins/Animation/ControlRig/Source/ControlRig/Public/Sequencer/MovieSceneControlRigParameterSection.h
--- a/Engine/Plugins/Animation/ControlRig/Source/ControlRig/Public/Sequencer/MovieSceneControlRigParameterSection.h
+++ b/Engine/Plugins/Animation/ControlRig/Source/ControlRig/Public/Sequencer/MovieSceneControlRigParameterSection.h
@@ -258,8 +258,7 @@ public:
 
 	//Function to load an Anim Sequence into this section. It will automatically resize to the section size.
 	//Will return false if fails or is canceled
-	virtual bool LoadAnimSequenceIntoThisSection(UAnimSequence* Sequence, UMovieScene* MovieScene, USkeletalMeshComponent* SkelMeshComp,
-		bool bKeyReduce, float Tolerance, FFrameNumber InStartFrame = 0);
+	virtual bool LoadAnimSequenceIntoThisSection(UAnimSequence* Sequence, UMovieScene* MovieScene, UObject* BoundObject, bool bKeyReduce, float Tolerance, FFrameNumber InStartFrame = 0);
 #endif
 	const TArray<bool>& GetControlsMask() const
 	{
