Control Rig: Fix error when trying to delete rig created through a python script

diff --git a/Engine/Plugins/Animation/ControlRig/Source/ControlRigDeveloper/Private/ControlRigBlueprint.cpp b/Engine/Plugins/Animation/ControlRig/Source/ControlRigDeveloper/Private/ControlRigBlueprint.cpp
--- a/Engine/Plugins/Animation/ControlRig/Source/ControlRigDeveloper/Private/ControlRigBlueprint.cpp
+++ b/Engine/Plugins/Animation/ControlRig/Source/ControlRigDeveloper/Private/ControlRigBlueprint.cpp
@@ -1640,19 +1640,17 @@ URigVMController* UControlRigBlueprint::GetOrCreateController(const UEdGraph* In
 
 TArray<FString> UControlRigBlueprint::GeneratePythonCommands(const FString InNewBlueprintName)
 {
-	TArray<FString> Commands;
-	Commands.Add(FString::Printf(TEXT("import unreal\n"
-			"unreal.load_module('ControlRigDeveloper')\n"
-			"factory = unreal.ControlRigBlueprintFactory\n"
-			"blueprint = factory.create_new_control_rig_asset(desired_package_path = '%s')\n"
-			"library = blueprint.get_local_function_library()\n"
-			"library_controller = blueprint.get_controller(library)\n"
-			"hierarchy = blueprint.hierarchy\n"
-			"hierarchy_controller = hierarchy.get_controller()\n")
-			, *InNewBlueprintName));
+	TArray<FString> InternalCommands;
+	InternalCommands.Add(TEXT("unreal.load_module('ControlRigDeveloper')"));
+	InternalCommands.Add(TEXT("factory = unreal.ControlRigBlueprintFactory"));
+	InternalCommands.Add(FString::Printf(TEXT("blueprint = factory.create_new_control_rig_asset(desired_package_path = '%s')"), *InNewBlueprintName));
+	InternalCommands.Add(TEXT("library = blueprint.get_local_function_library()"));
+	InternalCommands.Add(TEXT("library_controller = blueprint.get_controller(library)"));
+	InternalCommands.Add(TEXT("hierarchy = blueprint.hierarchy"));
+	InternalCommands.Add(TEXT("hierarchy_controller = hierarchy.get_controller()"));
 
 	// Hierarchy
-	Commands.Append(Hierarchy->GetController(true)->GeneratePythonCommands());
+	InternalCommands.Append(Hierarchy->GetController(true)->GeneratePythonCommands());
 		
 	// Add variables
 	for (const FBPVariableDescription& Variable : NewVariables)
@@ -1673,7 +1671,7 @@ TArray<FString> UControlRigBlueprint::GeneratePythonCommands(const FString InNew
 			}
 		}
 		// FName AddMemberVariable(const FName& InName, const FString& InCPPType, bool bIsPublic = false, bool bIsReadOnly = false, FString InDefaultValue = TEXT(""));
-		Commands.Add(FString::Printf(TEXT("blueprint.add_member_variable('%s', '%s', %s, %s)"),
+		InternalCommands.Add(FString::Printf(TEXT("blueprint.add_member_variable('%s', '%s', %s, %s)"),
 					*ExternalVariable.Name.ToString(),
 					*CPPType,
 					ExternalVariable.bIsPublic ? TEXT("True") : TEXT("False"),
@@ -1741,7 +1739,7 @@ TArray<FString> UControlRigBlueprint::GeneratePythonCommands(const FString InNew
 				// The controller will deal with deleting collapsed graph function when it creates the collapse node
 				{						
 					// Add Function
-					Commands.Add(FString::Printf(TEXT("function_%s = library_controller.add_function_to_library('%s', mutable=%s)\ngraph = function_%s.get_contained_graph()"),
+					InternalCommands.Add(FString::Printf(TEXT("function_%s = library_controller.add_function_to_library('%s', mutable=%s)\ngraph = function_%s.get_contained_graph()"),
 							*RigVMPythonUtils::NameToPep8(Graph->GetGraphName()),
 							*Graph->GetGraphName(),
 							Graph->GetEntryNode()->IsMutable() ? TEXT("True") : TEXT("False"),
@@ -1753,12 +1751,12 @@ TArray<FString> UControlRigBlueprint::GeneratePythonCommands(const FString InNew
 					// Set Entry and Return nodes in the correct position
 					{
 						//bool SetNodePositionByName(const FName& InNodeName, const FVector2D& InPosition, bool bSetupUndoRedo = true, bool bMergeUndoAction = false);
-						Commands.Add(FString::Printf(TEXT("blueprint.get_controller_by_name('%s').set_node_position_by_name('Entry', unreal.Vector2D(%f, %f))"),
+						InternalCommands.Add(FString::Printf(TEXT("blueprint.get_controller_by_name('%s').set_node_position_by_name('Entry', unreal.Vector2D(%f, %f))"),
 								*Graph->GetGraphName(),
 								EntryNode->GetPosition().X, 
 								EntryNode->GetPosition().Y));
 
-						Commands.Add(FString::Printf(TEXT("blueprint.get_controller_by_name('%s').set_node_position_by_name('Return', unreal.Vector2D(%f, %f))"),
+						InternalCommands.Add(FString::Printf(TEXT("blueprint.get_controller_by_name('%s').set_node_position_by_name('Return', unreal.Vector2D(%f, %f))"),
 								*Graph->GetGraphName(),
 								ReturnNode->GetPosition().X, 
 								ReturnNode->GetPosition().Y));
@@ -1774,7 +1772,7 @@ TArray<FString> UControlRigBlueprint::GeneratePythonCommands(const FString InNew
 							}
 
 							// FName AddExposedPin(const FName& InPinName, ERigVMPinDirection InDirection, const FString& InCPPType, const FName& InCPPTypeObjectPath, const FString& InDefaultValue, bool bSetupUndoRedo = true);
-							Commands.Add(FString::Printf(TEXT("blueprint.get_controller_by_name('%s').add_exposed_pin('%s', unreal.RigVMPinDirection.INPUT, '%s', '%s', '%s')"),
+							InternalCommands.Add(FString::Printf(TEXT("blueprint.get_controller_by_name('%s').add_exposed_pin('%s', unreal.RigVMPinDirection.INPUT, '%s', '%s', '%s')"),
 									*Graph->GetGraphName(),
 									*Pin->GetName(),
 									*Pin->GetCPPType(),
@@ -1790,7 +1788,7 @@ TArray<FString> UControlRigBlueprint::GeneratePythonCommands(const FString InNew
 							}
 
 							// FName AddExposedPin(const FName& InPinName, ERigVMPinDirection InDirection, const FString& InCPPType, const FName& InCPPTypeObjectPath, const FString& InDefaultValue, bool bSetupUndoRedo = true);
-							Commands.Add(FString::Printf(TEXT("blueprint.get_controller_by_name('%s').add_exposed_pin('%s', unreal.RigVMPinDirection.OUTPUT, '%s', '%s', '%s')"),
+							InternalCommands.Add(FString::Printf(TEXT("blueprint.get_controller_by_name('%s').add_exposed_pin('%s', unreal.RigVMPinDirection.OUTPUT, '%s', '%s', '%s')"),
 									*Graph->GetGraphName(),
 									*Pin->GetName(), 
 									*Pin->GetCPPType(),
@@ -1801,16 +1799,39 @@ TArray<FString> UControlRigBlueprint::GeneratePythonCommands(const FString InNew
 				}
 			}
 								
-			Commands.Append(Controller->GeneratePythonCommands());
+			InternalCommands.Append(Controller->GeneratePythonCommands());
 		}
 	}
 
 #if WITH_EDITORONLY_DATA
 	FString PreviewMeshPath = GetPreviewMesh()->GetPathName();
-	Commands.Add(FString::Printf(TEXT("blueprint.set_preview_mesh(unreal.load_object(name='%s', outer=None))"),
+	InternalCommands.Add(FString::Printf(TEXT("blueprint.set_preview_mesh(unreal.load_object(name='%s', outer=None))"),
 		*PreviewMeshPath));
 #endif
 
+	// Split multiple commands into different array elements
+	TArray<FString> InnerFunctionCmds;
+	for (FString Cmd : InternalCommands)
+	{
+		FString Left, Right=Cmd;
+		while (Right.Split(TEXT("\n"), &Left, &Right))
+		{
+			InnerFunctionCmds.Add(Left);			
+		}
+		InnerFunctionCmds.Add(Right);
+	}
+
+	// Define a function and insert all the commands
+	// We do not want to pollute the global state with our definitions
+	TArray<FString> Commands;
+	Commands.Add(FString::Printf(TEXT("import unreal\n"
+		"def create_control_rig():\n")));
+	for (const FString& InnerCmd : InnerFunctionCmds)
+	{
+		Commands.Add(FString::Printf(TEXT("\t%s"), *InnerCmd));
+	}
+
+	Commands.Add(TEXT("create_control_rig()\n"));
 	return Commands;
 }
 
