Sequencer: Multiple Control Rig fixes for when control rigs are dynamically added via cut/paste, undo/redo, etc..

diff --git a/Engine/Plugins/Animation/ControlRig/Source/ControlRigEditor/Private/EditMode/ControlRigEditMode.cpp b/Engine/Plugins/Animation/ControlRig/Source/ControlRigEditor/Private/EditMode/ControlRigEditMode.cpp
--- a/Engine/Plugins/Animation/ControlRig/Source/ControlRigEditor/Private/EditMode/ControlRigEditMode.cpp
+++ b/Engine/Plugins/Animation/ControlRig/Source/ControlRigEditor/Private/EditMode/ControlRigEditMode.cpp
@@ -207,6 +207,7 @@ bool FControlRigEditMode:: SetSequencer(TWeakPtr<ISequencer> InSequencer)
 					}
 				}
 			}
+			LastMovieSceneSig = Sequencer->GetFocusedMovieSceneSequence()->GetMovieScene()->GetSignature();
 		}
 		SetObjects_Internal();
 	}
@@ -348,7 +349,7 @@ void FControlRigEditMode::Enter()
 {
 	// Call parent implementation
 	FEdMode::Enter();
-
+	LastMovieSceneSig = FGuid();
 	if(UsesToolkits())
 	{
 		if (!Toolkit.IsValid())
@@ -465,6 +466,17 @@ void FControlRigEditMode::Exit()
 void FControlRigEditMode::Tick(FEditorViewportClient* ViewportClient, float DeltaTime)
 {
 	FEdMode::Tick(ViewportClient, DeltaTime);
+	
+	CheckMovieSceneSig();
+
+	if (bool* GameView = ViewportToGameView.Find(ViewportClient->Viewport))
+	{
+		*GameView = ViewportClient->IsInGameView();
+	} 
+	else
+	{
+		ViewportToGameView.Add(ViewportClient->Viewport, ViewportClient->IsInGameView());
+	}
 
 	if(DeferredItemsToFrame.Num() > 0)
 	{
@@ -681,11 +693,12 @@ TSet<FName> FControlRigEditMode::GetActiveControlsFromSequencer(UControlRig* Con
 void FControlRigEditMode::Render(const FSceneView* View, FViewport* Viewport, FPrimitiveDrawInterface* PDI)
 {	
 	const UControlRigEditModeSettings* Settings = GetDefault<UControlRigEditModeSettings>();
-	const bool bIsInGameView = IsInLevelEditor() ? (GCurrentLevelEditingViewportClient && GCurrentLevelEditingViewportClient->IsInGameView()) : false;
-	bool bRender = !Settings->bHideControlShapes && !bIsInGameView;
+	const bool bIsInGameView = IsInLevelEditor() ? (ViewportToGameView.Find(Viewport) && ViewportToGameView[Viewport]) : false;
+	bool bRender = !Settings->bHideControlShapes;
 	for (TWeakObjectPtr<UControlRig>& ControlRigPtr : RuntimeControlRigs)
 	{
 		UControlRig* ControlRig = ControlRigPtr.Get();
+		//actor game view drawing is handled by not drawing in game via SetActorHiddenInGame().
 		if (bRender && ControlRig && ControlRig->GetControlsVisible())
 		{
 			FTransform ComponentTransform = FTransform::Identity;
@@ -698,183 +711,185 @@ void FControlRigEditMode::Render(const FSceneView* View, FViewport* Viewport, FP
 			{
 				for (AControlRigShapeActor* Actor : *ShapeActors)
 				{
-					//Actor->SetActorHiddenInGame(bIsHidden);
 					if (GIsEditor && Actor->GetWorld() != nullptr && !Actor->GetWorld()->IsPlayInEditor())
 					{
 						Actor->SetIsTemporarilyHiddenInEditor(false);
 					}
 				}
 			}
-
-			URigHierarchy* Hierarchy = ControlRig->GetHierarchy();
-			const bool bHasFKRig = (ControlRig->IsA<UAdditiveControlRig>() || ControlRig->IsA<UFKControlRig>());
-			if (Settings->bDisplayHierarchy || bHasFKRig)
+			//only draw stuff if not in game view
+			if (!bIsInGameView)
 			{
-				const bool bBoolSetHitProxies = PDI && PDI->IsHitTesting() && bHasFKRig;
-				TSet<FName> ActiveControlName;
-				if (bHasFKRig)
+				URigHierarchy* Hierarchy = ControlRig->GetHierarchy();
+				const bool bHasFKRig = (ControlRig->IsA<UAdditiveControlRig>() || ControlRig->IsA<UFKControlRig>());
+				if (Settings->bDisplayHierarchy || bHasFKRig)
 				{
-					ActiveControlName = GetActiveControlsFromSequencer(ControlRig);
-				}
-				Hierarchy->ForEach<FRigTransformElement>([PDI, Hierarchy, ComponentTransform, ControlRig, bHasFKRig, bBoolSetHitProxies, ActiveControlName](FRigTransformElement* TransformElement) -> bool
+					const bool bBoolSetHitProxies = PDI && PDI->IsHitTesting() && bHasFKRig;
+					TSet<FName> ActiveControlName;
+					if (bHasFKRig)
 					{
-						const FTransform Transform = Hierarchy->GetTransform(TransformElement, ERigTransformType::CurrentGlobal);
-
-						FRigBaseElementParentArray Parents = Hierarchy->GetParents(TransformElement);
-						for (FRigBaseElement* ParentElement : Parents)
+						ActiveControlName = GetActiveControlsFromSequencer(ControlRig);
+					}
+					Hierarchy->ForEach<FRigTransformElement>([PDI, Hierarchy, ComponentTransform, ControlRig, bHasFKRig, bBoolSetHitProxies, ActiveControlName](FRigTransformElement* TransformElement) -> bool
 						{
-							if (FRigTransformElement* ParentTransformElement = Cast<FRigTransformElement>(ParentElement))
+							const FTransform Transform = Hierarchy->GetTransform(TransformElement, ERigTransformType::CurrentGlobal);
+
+							FRigBaseElementParentArray Parents = Hierarchy->GetParents(TransformElement);
+							for (FRigBaseElement* ParentElement : Parents)
 							{
-								FLinearColor Color = FLinearColor::White;
-								if (bHasFKRig)
-								{
-									FName ControlName = UFKControlRig::GetControlName(ParentTransformElement->GetName(), ParentTransformElement->GetType());
-									if (ActiveControlName.Num() > 0 && ActiveControlName.Contains(ControlName) == false)
-									{
-										continue;
-									}
-									if (ControlRig->IsControlSelected(ControlName))
-									{
-										Color = FLinearColor::Yellow;
-									}
-								}
-								const FTransform ParentTransform = Hierarchy->GetTransform(ParentTransformElement, ERigTransformType::CurrentGlobal);
-								const bool bHitTesting = bBoolSetHitProxies && (ParentTransformElement->GetType() == ERigElementType::Bone);
-								if (PDI)
+								if (FRigTransformElement* ParentTransformElement = Cast<FRigTransformElement>(ParentElement))
 								{
-									if (bHitTesting)
+									FLinearColor Color = FLinearColor::White;
+									if (bHasFKRig)
 									{
-										PDI->SetHitProxy(new HFKRigBoneProxy(ParentTransformElement->GetName(), ControlRig));
+										FName ControlName = UFKControlRig::GetControlName(ParentTransformElement->GetName(), ParentTransformElement->GetType());
+										if (ActiveControlName.Num() > 0 && ActiveControlName.Contains(ControlName) == false)
+										{
+											continue;
+										}
+										if (ControlRig->IsControlSelected(ControlName))
+										{
+											Color = FLinearColor::Yellow;
+										}
 									}
-									PDI->DrawLine(ComponentTransform.TransformPosition(Transform.GetLocation()), ComponentTransform.TransformPosition(ParentTransform.GetLocation()), Color, SDPG_Foreground);
-									if (bHitTesting)
+									const FTransform ParentTransform = Hierarchy->GetTransform(ParentTransformElement, ERigTransformType::CurrentGlobal);
+									const bool bHitTesting = bBoolSetHitProxies && (ParentTransformElement->GetType() == ERigElementType::Bone);
+									if (PDI)
 									{
-										PDI->SetHitProxy(nullptr);
+										if (bHitTesting)
+										{
+											PDI->SetHitProxy(new HFKRigBoneProxy(ParentTransformElement->GetName(), ControlRig));
+										}
+										PDI->DrawLine(ComponentTransform.TransformPosition(Transform.GetLocation()), ComponentTransform.TransformPosition(ParentTransform.GetLocation()), Color, SDPG_Foreground);
+										if (bHitTesting)
+										{
+											PDI->SetHitProxy(nullptr);
+										}
 									}
 								}
 							}
-						}
 
-						FLinearColor Color = FLinearColor::White;
-						if (bHasFKRig)
-						{
-							FName ControlName = UFKControlRig::GetControlName(TransformElement->GetName(), TransformElement->GetType());
-							if (ActiveControlName.Num() > 0 && ActiveControlName.Contains(ControlName) == false)
+							FLinearColor Color = FLinearColor::White;
+							if (bHasFKRig)
 							{
-								return true;
+								FName ControlName = UFKControlRig::GetControlName(TransformElement->GetName(), TransformElement->GetType());
+								if (ActiveControlName.Num() > 0 && ActiveControlName.Contains(ControlName) == false)
+								{
+									return true;
+								}
+								if (ControlRig->IsControlSelected(ControlName))
+								{
+									Color = FLinearColor::Yellow;
+								}
 							}
-							if (ControlRig->IsControlSelected(ControlName))
+							if (PDI)
 							{
-								Color = FLinearColor::Yellow;
+								const bool bHitTesting = PDI->IsHitTesting() && bBoolSetHitProxies && (TransformElement->GetType() == ERigElementType::Bone);
+								if (bHitTesting)
+								{
+									PDI->SetHitProxy(new HFKRigBoneProxy(TransformElement->GetName(), ControlRig));
+								}
+								PDI->DrawPoint(ComponentTransform.TransformPosition(Transform.GetLocation()), Color, 5.0f, SDPG_Foreground);
+
+								if (bHitTesting)
+								{
+									PDI->SetHitProxy(nullptr);
+								}
 							}
-						}
-						if (PDI)
+
+							return true;
+						});
+				}
+
+				if (Settings->bDisplayNulls || ControlRig->IsSetupModeEnabled())
+				{
+					TArray<FTransform> SpaceTransforms;
+					TArray<FTransform> SelectedSpaceTransforms;
+					Hierarchy->ForEach<FRigNullElement>([&SpaceTransforms, &SelectedSpaceTransforms, Hierarchy](FRigNullElement* NullElement) -> bool
 						{
-							const bool bHitTesting = PDI->IsHitTesting() && bBoolSetHitProxies && (TransformElement->GetType() == ERigElementType::Bone);
-							if (bHitTesting)
+							if (Hierarchy->IsSelected(NullElement->GetIndex()))
 							{
-								PDI->SetHitProxy(new HFKRigBoneProxy(TransformElement->GetName(), ControlRig));
+								SelectedSpaceTransforms.Add(Hierarchy->GetTransform(NullElement, ERigTransformType::CurrentGlobal));
 							}
-							PDI->DrawPoint(ComponentTransform.TransformPosition(Transform.GetLocation()), Color, 5.0f, SDPG_Foreground);
-
-							if (bHitTesting)
+							else
 							{
-								PDI->SetHitProxy(nullptr);
+								SpaceTransforms.Add(Hierarchy->GetTransform(NullElement, ERigTransformType::CurrentGlobal));
 							}
-						}
-
-						return true;
-					});
-			}
-
-			if (Settings->bDisplayNulls || ControlRig->IsSetupModeEnabled())
-			{
-				TArray<FTransform> SpaceTransforms;
-				TArray<FTransform> SelectedSpaceTransforms;
-				Hierarchy->ForEach<FRigNullElement>([&SpaceTransforms, &SelectedSpaceTransforms, Hierarchy](FRigNullElement* NullElement) -> bool
-					{
-						if (Hierarchy->IsSelected(NullElement->GetIndex()))
-						{
-							SelectedSpaceTransforms.Add(Hierarchy->GetTransform(NullElement, ERigTransformType::CurrentGlobal));
-						}
-						else
-						{
-							SpaceTransforms.Add(Hierarchy->GetTransform(NullElement, ERigTransformType::CurrentGlobal));
-						}
-						return true;
-					});
+							return true;
+						});
 
-				ControlRig->DrawInterface.DrawAxes(FTransform::Identity, SpaceTransforms, Settings->AxisScale);
-				ControlRig->DrawInterface.DrawAxes(FTransform::Identity, SelectedSpaceTransforms, FLinearColor(1.0f, 0.34f, 0.0f, 1.0f), Settings->AxisScale);
-			}
+					ControlRig->DrawInterface.DrawAxes(FTransform::Identity, SpaceTransforms, Settings->AxisScale);
+					ControlRig->DrawInterface.DrawAxes(FTransform::Identity, SelectedSpaceTransforms, FLinearColor(1.0f, 0.34f, 0.0f, 1.0f), Settings->AxisScale);
+				}
 
-			if (Settings->bDisplayAxesOnSelection && Settings->AxisScale > SMALL_NUMBER)
-			{
-				if (ControlRig->GetWorld() && ControlRig->GetWorld()->IsPreviewWorld())
+				if (Settings->bDisplayAxesOnSelection && Settings->AxisScale > SMALL_NUMBER)
 				{
-					TArray<FRigElementKey> SelectedRigElements = GetSelectedRigElements(ControlRig);
-					const float Scale = Settings->AxisScale;
-					PDI->AddReserveLines(SDPG_Foreground, SelectedRigElements.Num() * 3);
-
-					for (const FRigElementKey& SelectedElement : SelectedRigElements)
+					if (ControlRig->GetWorld() && ControlRig->GetWorld()->IsPreviewWorld())
 					{
-						FTransform ElementTransform = Hierarchy->GetGlobalTransform(SelectedElement);
-						ElementTransform = ElementTransform * ComponentTransform;
+						TArray<FRigElementKey> SelectedRigElements = GetSelectedRigElements(ControlRig);
+						const float Scale = Settings->AxisScale;
+						PDI->AddReserveLines(SDPG_Foreground, SelectedRigElements.Num() * 3);
 
-						PDI->DrawLine(ElementTransform.GetTranslation(), ElementTransform.TransformPosition(FVector(Scale, 0.f, 0.f)), FLinearColor::Red, SDPG_Foreground);
-						PDI->DrawLine(ElementTransform.GetTranslation(), ElementTransform.TransformPosition(FVector(0.f, Scale, 0.f)), FLinearColor::Green, SDPG_Foreground);
-						PDI->DrawLine(ElementTransform.GetTranslation(), ElementTransform.TransformPosition(FVector(0.f, 0.f, Scale)), FLinearColor::Blue, SDPG_Foreground);
+						for (const FRigElementKey& SelectedElement : SelectedRigElements)
+						{
+							FTransform ElementTransform = Hierarchy->GetGlobalTransform(SelectedElement);
+							ElementTransform = ElementTransform * ComponentTransform;
+
+							PDI->DrawLine(ElementTransform.GetTranslation(), ElementTransform.TransformPosition(FVector(Scale, 0.f, 0.f)), FLinearColor::Red, SDPG_Foreground);
+							PDI->DrawLine(ElementTransform.GetTranslation(), ElementTransform.TransformPosition(FVector(0.f, Scale, 0.f)), FLinearColor::Green, SDPG_Foreground);
+							PDI->DrawLine(ElementTransform.GetTranslation(), ElementTransform.TransformPosition(FVector(0.f, 0.f, Scale)), FLinearColor::Blue, SDPG_Foreground);
+						}
 					}
 				}
-			}
-			for (const FControlRigDrawInstruction& Instruction : ControlRig->DrawInterface)
-			{
-				if (!Instruction.IsValid())
+				for (const FControlRigDrawInstruction& Instruction : ControlRig->DrawInterface)
 				{
-					continue;
-				}
+					if (!Instruction.IsValid())
+					{
+						continue;
+					}
 
-				FTransform InstructionTransform = Instruction.Transform * ComponentTransform;
-				switch (Instruction.PrimitiveType)
-				{
-				case EControlRigDrawSettings::Points:
-				{
-					for (const FVector& Point : Instruction.Positions)
+					FTransform InstructionTransform = Instruction.Transform * ComponentTransform;
+					switch (Instruction.PrimitiveType)
 					{
-						PDI->DrawPoint(InstructionTransform.TransformPosition(Point), Instruction.Color, Instruction.Thickness, SDPG_Foreground);
+					case EControlRigDrawSettings::Points:
+					{
+						for (const FVector& Point : Instruction.Positions)
+						{
+							PDI->DrawPoint(InstructionTransform.TransformPosition(Point), Instruction.Color, Instruction.Thickness, SDPG_Foreground);
+						}
+						break;
 					}
-					break;
-				}
-				case EControlRigDrawSettings::Lines:
-				{
-					const TArray<FVector>& Points = Instruction.Positions;
-					PDI->AddReserveLines(SDPG_Foreground, Points.Num() / 2, false, Instruction.Thickness > SMALL_NUMBER);
-					for (int32 PointIndex = 0; PointIndex < Points.Num() - 1; PointIndex += 2)
+					case EControlRigDrawSettings::Lines:
 					{
-						PDI->DrawLine(InstructionTransform.TransformPosition(Points[PointIndex]), InstructionTransform.TransformPosition(Points[PointIndex + 1]), Instruction.Color, SDPG_Foreground, Instruction.Thickness);
+						const TArray<FVector>& Points = Instruction.Positions;
+						PDI->AddReserveLines(SDPG_Foreground, Points.Num() / 2, false, Instruction.Thickness > SMALL_NUMBER);
+						for (int32 PointIndex = 0; PointIndex < Points.Num() - 1; PointIndex += 2)
+						{
+							PDI->DrawLine(InstructionTransform.TransformPosition(Points[PointIndex]), InstructionTransform.TransformPosition(Points[PointIndex + 1]), Instruction.Color, SDPG_Foreground, Instruction.Thickness);
+						}
+						break;
 					}
-					break;
-				}
-				case EControlRigDrawSettings::LineStrip:
-				{
-					const TArray<FVector>& Points = Instruction.Positions;
-					PDI->AddReserveLines(SDPG_Foreground, Points.Num() - 1, false, Instruction.Thickness > SMALL_NUMBER);
-					for (int32 PointIndex = 0; PointIndex < Points.Num() - 1; PointIndex++)
+					case EControlRigDrawSettings::LineStrip:
 					{
-						PDI->DrawLine(InstructionTransform.TransformPosition(Points[PointIndex]), InstructionTransform.TransformPosition(Points[PointIndex + 1]), Instruction.Color, SDPG_Foreground, Instruction.Thickness);
+						const TArray<FVector>& Points = Instruction.Positions;
+						PDI->AddReserveLines(SDPG_Foreground, Points.Num() - 1, false, Instruction.Thickness > SMALL_NUMBER);
+						for (int32 PointIndex = 0; PointIndex < Points.Num() - 1; PointIndex++)
+						{
+							PDI->DrawLine(InstructionTransform.TransformPosition(Points[PointIndex]), InstructionTransform.TransformPosition(Points[PointIndex + 1]), Instruction.Color, SDPG_Foreground, Instruction.Thickness);
+						}
+						break;
 					}
-					break;
-				}
 
-				case EControlRigDrawSettings::DynamicMesh:
-				{
-					FDynamicMeshBuilder MeshBuilder(PDI->View->GetFeatureLevel());
-					MeshBuilder.AddVertices(Instruction.MeshVerts);
-					MeshBuilder.AddTriangles(Instruction.MeshIndices);
-					MeshBuilder.Draw(PDI, InstructionTransform.ToMatrixWithScale(), Instruction.MaterialRenderProxy, SDPG_World/*SDPG_Foreground*/);
-					break;
-				}
+					case EControlRigDrawSettings::DynamicMesh:
+					{
+						FDynamicMeshBuilder MeshBuilder(PDI->View->GetFeatureLevel());
+						MeshBuilder.AddVertices(Instruction.MeshVerts);
+						MeshBuilder.AddTriangles(Instruction.MeshIndices);
+						MeshBuilder.Draw(PDI, InstructionTransform.ToMatrixWithScale(), Instruction.MaterialRenderProxy, SDPG_World/*SDPG_Foreground*/);
+						break;
+					}
 
+					}
 				}
 			}
 		}
@@ -885,7 +900,6 @@ void FControlRigEditMode::Render(const FSceneView* View, FViewport* Viewport, FP
 			{
 				for (AControlRigShapeActor* Actor : *ShapeActors)
 				{
-					//Actor->SetActorHiddenInGame(bIsHidden);
 					if (GIsEditor && Actor->GetWorld() != nullptr && !Actor->GetWorld()->IsPlayInEditor())
 					{
 						Actor->SetIsTemporarilyHiddenInEditor(true);
@@ -1489,6 +1503,10 @@ bool IntersectsBox( AActor& InActor, const FBox& InBox, FLevelEditorViewportClie
 bool FControlRigEditMode::BoxSelect(FBox& InBox, bool InSelect)
 {
 	FLevelEditorViewportClient* LevelViewportClient = GCurrentLevelEditingViewportClient;
+	if (LevelViewportClient->IsInGameView() == true)
+	{
+		return  FEdMode::BoxSelect(InBox, InSelect);
+	}
 	const bool bStrictDragSelection = GetDefault<ULevelEditorViewportSettings>()->bStrictBoxSelection;
 
 	FScopedTransaction ScopedTransaction(LOCTEXT("SelectControlTransaction", "Select Control"), IsInLevelEditor() && !GIsTransacting);
@@ -1511,7 +1529,7 @@ bool FControlRigEditMode::BoxSelect(FBox& InBox, bool InSelect)
 		}
 
 		AControlRigShapeActor* ShapeActor = CastChecked<AControlRigShapeActor>(Actor);
-		if (!ShapeActor->IsSelectable() || ShapeActor->ControlRig.IsValid () == false)
+		if (!ShapeActor->IsSelectable() || ShapeActor->ControlRig.IsValid () == false || ShapeActor->ControlRig->GetControlsVisible() == false)
 		{
 			continue;
 		}
@@ -1542,6 +1560,11 @@ bool FControlRigEditMode::BoxSelect(FBox& InBox, bool InSelect)
 
 bool FControlRigEditMode::FrustumSelect(const FConvexVolume& InFrustum, FEditorViewportClient* InViewportClient, bool InSelect)
 {
+	if (InViewportClient->IsInGameView() == true)
+	{
+		return FEdMode::FrustumSelect(InFrustum, InViewportClient, InSelect);
+	}
+
 	FScopedTransaction ScopedTransaction(LOCTEXT("SelectControlTransaction", "Select Control"), IsInLevelEditor() && !GIsTransacting);
 	bool bSomethingSelected(false);
 	const bool bShiftDown = InViewportClient->Viewport->KeyState(EKeys::LeftShift) || InViewportClient->Viewport->KeyState(EKeys::RightShift);
@@ -1561,7 +1584,7 @@ bool FControlRigEditMode::FrustumSelect(const FConvexVolume& InFrustum, FEditorV
 				{
 					if (PrimitiveComponent->ComponentIsTouchingSelectionFrustum(InFrustum, InViewportClient->EngineShowFlags, false /*only bsp*/, false/*encompass entire*/))
 					{
-						if (ShapeActor->IsSelectable())
+						if (ShapeActor->IsSelectable() && ShapeActor->ControlRig.IsValid() &&  ShapeActor->ControlRig->GetControlsVisible() )
 						{
 							bSomethingSelected = true;
 							const FName& ControlName = ShapeActor->ControlName;
@@ -2674,18 +2697,86 @@ bool FControlRigEditMode::MouseLeave(FEditorViewportClient* ViewportClient, FVie
 	return false;
 }
 
+bool FControlRigEditMode::CheckMovieSceneSig()
+{
+	bool bSomethingChanged = false;
+	if (WeakSequencer.IsValid())
+	{
+		TSharedPtr<ISequencer> Sequencer = WeakSequencer.Pin();
+		if (Sequencer->GetFocusedMovieSceneSequence() && Sequencer->GetFocusedMovieSceneSequence()->GetMovieScene())
+		{
+			FGuid CurrentMovieSceneSig = Sequencer->GetFocusedMovieSceneSequence()->GetMovieScene()->GetSignature();
+			if (LastMovieSceneSig != CurrentMovieSceneSig)
+			{
+				if (ULevelSequence* LevelSequence = Cast<ULevelSequence>(Sequencer->GetFocusedMovieSceneSequence()))
+				{
+					TArray<TWeakObjectPtr<UControlRig>> CurrentControlRigs;
+					TArray<FControlRigSequencerBindingProxy> Proxies = UControlRigSequencerEditorLibrary::GetControlRigs(LevelSequence);
+					for (FControlRigSequencerBindingProxy& Proxy : Proxies)
+					{
+						if (UControlRig* ControlRig = Proxy.ControlRig.Get())
+						{
+							CurrentControlRigs.Add(ControlRig);
+							if (RuntimeControlRigs.Contains(ControlRig) == false)
+							{
+								AddControlRigInternal(ControlRig);
+								bSomethingChanged = true;
+							}
+						}
+					}
+					TArray<TWeakObjectPtr<UControlRig>> ControlRigsToRemove;
+					for (TWeakObjectPtr<UControlRig>& RuntimeRigPtr : RuntimeControlRigs)
+					{
+						if (CurrentControlRigs.Contains(RuntimeRigPtr) == false)
+						{
+							ControlRigsToRemove.Add(RuntimeRigPtr);
+						}
+					}
+					for (TWeakObjectPtr<UControlRig>& OldRuntimeRigPtr : ControlRigsToRemove)
+					{
+						RemoveControlRig(OldRuntimeRigPtr.Get());
+					}
+				}
+				LastMovieSceneSig = CurrentMovieSceneSig;
+				if (bSomethingChanged)
+				{
+					SetObjects_Internal();
+				}
+			}
+		}
+	}
+	return bSomethingChanged;
+}
+
 void FControlRigEditMode::PostUndo()
 {
 	bool bInvalidateViewport = false;
 	if (WeakSequencer.IsValid())
 	{
+		bool bHaveInvalidControlRig = false;
 		for (TWeakObjectPtr<UControlRig>& RuntimeRigPtr : RuntimeControlRigs)
 		{
 			if (RuntimeRigPtr.IsValid() == false)
+			{				
+				bHaveInvalidControlRig = bInvalidateViewport = true;
+				break;
+			}
+		}
+		//if one is invalid we need to clear everything,since no longer have ptr to selectively delete
+		if (bHaveInvalidControlRig == true)
+		{
+			TArray<TWeakObjectPtr<UControlRig>> PreviousRuntimeRigs = RuntimeControlRigs;
+			for (int32 PreviousRuntimeRigIndex = 0; PreviousRuntimeRigIndex < PreviousRuntimeRigs.Num(); PreviousRuntimeRigIndex++)
 			{
-				DestroyShapesActors(RuntimeRigPtr.Get());
-				bInvalidateViewport = true;
+				if (PreviousRuntimeRigs[PreviousRuntimeRigIndex].IsValid())
+				{
+					RemoveControlRig(PreviousRuntimeRigs[PreviousRuntimeRigIndex].Get());
+				}
 			}
+			RuntimeControlRigs.Reset();
+			DestroyShapesActors(nullptr);
+			DelegateHelpers.Reset();
+			RuntimeControlRigs.Reset();
 		}
 		TSharedPtr<ISequencer> Sequencer = WeakSequencer.Pin();
 		if (ULevelSequence* LevelSequence = Cast<ULevelSequence>(Sequencer->GetFocusedMovieSceneSequence()))
@@ -2846,6 +2937,8 @@ void FControlRigEditMode::CreateShapeActors(UControlRig* ControlRig)
 			AControlRigShapeActor* ShapeActor = FControlRigShapeHelper::CreateDefaultShapeActor(WorldPtr, Param);
 			if (ShapeActor)
 			{
+				//not drawn in game or in game view.
+				ShapeActor->SetActorHiddenInGame(true);
 				TArray<AControlRigShapeActor*>* ShapeActors = ControlRigShapeActors.Find(ControlRig);
 				if (ShapeActors)
 				{
@@ -3526,11 +3619,9 @@ void FControlRigEditMode::TickControlShape(AControlRigShapeActor* ShapeActor, co
 			if (FRigControlElement* ControlElement = ControlRig->FindControl(ShapeActor->ControlName))
 			{
 				ShapeActor->SetShapeColor(ControlElement->Settings.ShapeColor);
-				const bool bIsInGameView = IsInLevelEditor() ? (GCurrentLevelEditingViewportClient && GCurrentLevelEditingViewportClient->IsInGameView()) : false;
-
-				ShapeActor->SetIsTemporarilyHiddenInEditor(bIsInGameView || !ControlElement->Settings.bShapeVisible || Settings->bHideControlShapes || !ControlRig->GetControlsVisible());
-				if (!IsInLevelEditor()) //don't change this in level editor otherwise we can never select it
+				if (!IsInLevelEditor()) //don't change this in level editor otherwise we can never select it and render has to be based on viewport.
 				{
+					ShapeActor->SetIsTemporarilyHiddenInEditor(!ControlElement->Settings.bShapeVisible || Settings->bHideControlShapes || !ControlRig->GetControlsVisible());
 					ShapeActor->SetSelectable(ControlElement->Settings.bShapeVisible && !Settings->bHideControlShapes && ControlElement->Settings.bAnimatable && ControlRig->GetControlsVisible());
 				}
 			}
@@ -3564,6 +3655,12 @@ void FControlRigEditMode::AddControlRigInternal(UControlRig* InControlRig)
 	InControlRig->GetHierarchy()->OnModified().RemoveAll(this);
 	InControlRig->GetHierarchy()->OnModified().AddSP(this, &FControlRigEditMode::OnHierarchyModified);
 
+	//needed for the control rig track editor delegates to get hooked up
+	if (WeakSequencer.IsValid())
+	{
+		TSharedPtr<ISequencer> Sequencer = WeakSequencer.Pin();
+		Sequencer->ObjectImplicitlyAdded(InControlRig);
+	}
 	OnControlRigAddedOrRemovedDelegate.Broadcast(InControlRig, true);
 }
 
@@ -3620,6 +3717,12 @@ void FControlRigEditMode::RemoveControlRig(UControlRig* InControlRig)
 	{
 		RuntimeControlRigs.RemoveAt(Index);
 	}
+	//needed for the control rig track editor delegates to get removed
+	if (WeakSequencer.IsValid())
+	{
+		TSharedPtr<ISequencer> Sequencer = WeakSequencer.Pin();
+		Sequencer->ObjectImplicitlyRemoved(InControlRig);
+	}
 	OnControlRigAddedOrRemovedDelegate.Broadcast(InControlRig, false);
 }
 
diff --git a/Engine/Plugins/Animation/ControlRig/Source/ControlRigEditor/Private/EditMode/ControlRigEditMode.h b/Engine/Plugins/Animation/ControlRig/Source/ControlRigEditor/Private/EditMode/ControlRigEditMode.h
--- a/Engine/Plugins/Animation/ControlRig/Source/ControlRigEditor/Private/EditMode/ControlRigEditMode.h
+++ b/Engine/Plugins/Animation/ControlRig/Source/ControlRigEditor/Private/EditMode/ControlRigEditMode.h
@@ -12,6 +12,7 @@
 #include "UObject/StrongObjectPtr.h"
 #include "UnrealWidgetFwd.h"
 #include "IPersonaEditMode.h"
+#include "Misc/Guid.h"
 #include "ControlRigEditMode.generated.h"
 
 
@@ -253,6 +254,7 @@ private:
 protected:
 
 	TWeakPtr<ISequencer> WeakSequencer;
+	FGuid LastMovieSceneSig;
 
 	/** The scope for the interaction, one per manipulated Control rig */
 	TMap<UControlRig*,FControlRigInteractionScope*> InteractionScopes;
@@ -373,6 +375,8 @@ private:
 	void AddControlRigInternal(UControlRig* InControlRig);
 	void TickManipulatableObjects(float DeltaTime);
 
+	/* Check on tick to see if movie scene has changed, returns true if it has*/
+	bool CheckMovieSceneSig();
 	void SetControlShapeTransform(AControlRigShapeActor* ShapeActor, const FTransform& InTransform);
 	FTransform GetControlShapeTransform(AControlRigShapeActor* ShapeActor) const;
 	void MoveControlShape(AControlRigShapeActor* ShapeActor, const bool bTranslation, FVector& InDrag, 
@@ -408,6 +412,9 @@ private:
 	TArray<TWeakObjectPtr<UControlRig>> RuntimeControlRigs;
 	TMap<UControlRig*,TStrongObjectPtr<UControlRigEditModeDelegateHelper>> DelegateHelpers;
 
+	//hack since we can't get the viewport client from the viewport, so in the tick we set the gameview bool and then in render/tickcontrolshapes we use it.
+	TMap<FViewport*, bool>  ViewportToGameView;
+
 	TArray<FRigElementKey> DeferredItemsToFrame;
 
 	/** Computes the current interaction types based on the widget mode */
diff --git a/Engine/Plugins/Animation/ControlRig/Source/ControlRigEditor/Private/EditMode/SControlRigOutliner.cpp b/Engine/Plugins/Animation/ControlRig/Source/ControlRigEditor/Private/EditMode/SControlRigOutliner.cpp
--- a/Engine/Plugins/Animation/ControlRig/Source/ControlRigEditor/Private/EditMode/SControlRigOutliner.cpp
+++ b/Engine/Plugins/Animation/ControlRig/Source/ControlRigEditor/Private/EditMode/SControlRigOutliner.cpp
@@ -38,7 +38,7 @@
 
 void SControlRigOutlinerItem::Construct(const FArguments& InArgs)
 {
-
+	ControlRigEditMode = InArgs._EditMode;
 	bIsChangingRigHierarchy = false;
 
 	DisplaySettings.bShowBones = false;
@@ -243,6 +243,22 @@ void SControlRigOutlinerItem::HandleSelectionChanged(TSharedPtr<FRigTreeElement>
 
 		TGuardValue<bool> GuardRigHierarchyChanges(bIsChangingRigHierarchy, true);
 		const TArray<FRigElementKey> NewSelection = HierarchyTreeView->GetTreeView()->GetSelectedKeys();
+		//due to how Sequencer Tree View will redo selection on next tick if we aren't keeping or toggling selection we need to clear it out
+		if (FSlateApplication::Get().GetModifierKeys().IsShiftDown() == false || FSlateApplication::Get().GetModifierKeys().IsControlDown() == false)
+		{
+			if (ControlRigEditMode)
+			{
+				TMap<UControlRig*, TArray<FRigElementKey>> SelectedControls;
+				ControlRigEditMode->GetAllSelectedControls(SelectedControls);
+				for (TPair<UControlRig*, TArray<FRigElementKey>>& CurrentSelection : SelectedControls)
+				{
+					if (CurrentSelection.Key)
+					{
+						CurrentSelection.Key->ClearControlSelection();
+					}
+				}
+			}
+		}
 		if (!Controller->SetSelection(NewSelection))
 		{
 			return;
@@ -313,6 +329,7 @@ void SControlRigOutliner::Rebuild()
 				[
 					SNew(SControlRigOutlinerItem)
 					.ControlRig(ControlRig)
+					.EditMode(EditMode)
 				];
 			}
 		}
diff --git a/Engine/Plugins/Animation/ControlRig/Source/ControlRigEditor/Private/EditMode/SControlRigOutliner.h b/Engine/Plugins/Animation/ControlRig/Source/ControlRigEditor/Private/EditMode/SControlRigOutliner.h
--- a/Engine/Plugins/Animation/ControlRig/Source/ControlRigEditor/Private/EditMode/SControlRigOutliner.h
+++ b/Engine/Plugins/Animation/ControlRig/Source/ControlRigEditor/Private/EditMode/SControlRigOutliner.h
@@ -23,6 +23,7 @@ class SControlRigOutlinerItem : public SCompoundWidget
 {
 	SLATE_BEGIN_ARGS(SControlRigOutlinerItem){}
 	SLATE_ARGUMENT(UControlRig*, ControlRig)
+	SLATE_ARGUMENT(FControlRigEditMode*, EditMode)
 	SLATE_END_ARGS()
 	SControlRigOutlinerItem();
 	~SControlRigOutlinerItem();
@@ -50,6 +51,8 @@ private:
 	TSharedPtr<SExpandableArea> PickerExpander;
 
 	TWeakObjectPtr<UControlRig> CurrentControlRig;
+	FControlRigEditMode* ControlRigEditMode = nullptr;
+
 };
 
 class SControlRigOutliner: public SCompoundWidget
