Control Rig: Perform OnEndLoadPackage work correctly if -game is specified

original CL 19666423 by [REDACTED_PERSON]

diff --git a/Engine/Plugins/Animation/ControlRig/Source/ControlRigDeveloper/Private/ControlRigBlueprint.cpp b/Engine/Plugins/Animation/ControlRig/Source/ControlRigDeveloper/Private/ControlRigBlueprint.cpp
--- a/Engine/Plugins/Animation/ControlRig/Source/ControlRigDeveloper/Private/ControlRigBlueprint.cpp
+++ b/Engine/Plugins/Animation/ControlRig/Source/ControlRigDeveloper/Private/ControlRigBlueprint.cpp
@@ -590,7 +590,10 @@ void UControlRigBlueprint::PostLoad()
 
 #if WITH_EDITOR
 	// delay compilation until the package has been loaded
-	FCoreUObjectDelegates::OnEndLoadPackage.AddUObject(this, &UControlRigBlueprint::HandlePackageDone);
+	if(GIsEditor)
+	{
+		FCoreUObjectDelegates::OnEndLoadPackage.AddUObject(this, &UControlRigBlueprint::HandlePackageDone);
+	}
 #else
 	RecompileVMIfRequired();
 #endif
@@ -605,6 +608,19 @@ void UControlRigBlueprint::PostLoad()
 	{
 		Package->SetDirtyFlag(bDirtyDuringLoad);
 	}
+
+#if WITH_EDITOR
+	// if we are running with -game we are in editor code,
+	// but GIsEditor is turned off
+	if(!GIsEditor)
+	{
+		FEndLoadPackageContext Context;
+		Context.bSynchronous = true;
+		Context.LoadedPackages = {GetPackage()};
+		Context.RecursiveDepth = 0;
+		HandlePackageDone(Context);
+	}
+#endif
 }
 
 #if WITH_EDITOR
