Control Rig: Ensure to match the target hierarchy exactly prior to baking

Original CL 19666486 by [REDACTED_PERSON]

diff --git a/Engine/Plugins/Animation/ControlRig/Source/ControlRig/Private/Sequencer/MovieSceneControlRigParameterSection.cpp b/Engine/Plugins/Animation/ControlRig/Source/ControlRig/Private/Sequencer/MovieSceneControlRigParameterSection.cpp
--- a/Engine/Plugins/Animation/ControlRig/Source/ControlRig/Private/Sequencer/MovieSceneControlRigParameterSection.cpp
+++ b/Engine/Plugins/Animation/ControlRig/Source/ControlRig/Private/Sequencer/MovieSceneControlRigParameterSection.cpp
@@ -2632,6 +2632,14 @@ bool UMovieSceneControlRigParameterSection::LoadAnimSequenceIntoThisSection(UAni
 	const FAnimationCurveData& CurveData = DataModel->GetCurveData();
 	const TArray<FBoneAnimationTrack>& BoneAnimationTracks = DataModel->GetBoneAnimationTracks();
 
+	// copy the hierarchy from the CDO into the target control rig.
+	// this ensures that the topology version matches in case of a dynamic hierarchy
+	if (UControlRig* CDO = Cast<UControlRig>(ControlRig->GetClass()->GetDefaultObject()))
+	{
+		ControlRig->GetHierarchy()->CopyHierarchy(CDO->GetHierarchy());
+	}
+
+	// now set the hierarchies initial transforms based on the currently used skeletal mesh
 	if (SkelMeshComp)
 	{
 		ControlRig->SetBoneInitialTransformsFromSkeletalMeshComponent(SkelMeshComp, true);
