ControlRig: Improve robustness of memory snapshot hash

diff --git a/Engine/Plugins/Animation/ControlRig/Source/ControlRig/Private/ControlRig.cpp b/Engine/Plugins/Animation/ControlRig/Source/ControlRig/Private/ControlRig.cpp
--- a/Engine/Plugins/Animation/ControlRig/Source/ControlRig/Private/ControlRig.cpp
+++ b/Engine/Plugins/Animation/ControlRig/Source/ControlRig/Private/ControlRig.cpp
@@ -180,6 +180,20 @@ void UControlRig::Initialize(bool bInitRigUnits)
 		return;
 	}
 
+	// recompute the hash used to differentiate VMs based on their memory layout 
+	if (HasAnyFlags(RF_ClassDefaultObject))
+	{
+		CachedMemoryHash = INDEX_NONE;
+
+		if(VM)
+		{
+			CachedMemoryHash = HashCombine(
+				VM->GetLiteralMemory()->GetMemoryHash(),
+				VM->GetWorkMemory()->GetMemoryHash()
+			);
+		}
+	}
+
 	if (IsTemplate())
 	{
 		// don't initialize template class 
@@ -1204,21 +1218,26 @@ void UControlRig::ExecuteUnits(FRigUnitContext& InOutContext, const FName& InEve
 				UControlRig* CDO = Cast<UControlRig>(GetClass()->GetDefaultObject());
 
 				bool bIsValidSnapshot = false;
-				TObjectPtr<URigVM>* InitializedVMSnapshotPtr = CDO->InitializedVMSnapshots.Find(SnapshotHash);
-				if(InitializedVMSnapshotPtr && !InitializedVMSnapshotPtr->IsNull())
+				if(SnapshotHash != INDEX_NONE)
 				{
-					const URigVM* InitializedVMSnapshot = InitializedVMSnapshotPtr->Get();
-
-					if(VM->WorkMemoryStorageObject->GetClass() == InitializedVMSnapshot->WorkMemoryStorageObject->GetClass())
-					{
-						VM->WorkMemoryStorageObject->CopyFrom(InitializedVMSnapshot->WorkMemoryStorageObject);
-						VM->InvalidateCachedMemory();
-						VM->Initialize(LocalMemory, AdditionalArguments, false);
-						bIsValidSnapshot = true;
-					}
-					else
+					TObjectPtr<URigVM>* InitializedVMSnapshotPtr = CDO->InitializedVMSnapshots.Find(SnapshotHash);
+					if(InitializedVMSnapshotPtr && !InitializedVMSnapshotPtr->IsNull())
 					{
-						CDO->InitializedVMSnapshots.Remove(SnapshotHash);
+						const URigVM* InitializedVMSnapshot = InitializedVMSnapshotPtr->Get();
+
+						if(VM->WorkMemoryStorageObject->GetClass() == InitializedVMSnapshot->WorkMemoryStorageObject->GetClass())
+						{
+							InitializedVMSnapshots.Reset();
+
+							VM->WorkMemoryStorageObject->CopyFrom(InitializedVMSnapshot->WorkMemoryStorageObject);
+							VM->InvalidateCachedMemory();
+							VM->Initialize(LocalMemory, AdditionalArguments, false);
+							bIsValidSnapshot = true;
+						}
+						else
+						{
+							CDO->InitializedVMSnapshots.Remove(SnapshotHash);
+						}
 					}
 				}
 				
@@ -2975,6 +2994,11 @@ void UControlRig::PostInitInstance(UControlRig* InCDO)
 
 uint32 UControlRig::GetHashForInitializeVMSnapShot()
 {
+	if(CachedMemoryHash == INDEX_NONE)
+	{
+		return INDEX_NONE;
+	}
+	
 	uint32 Hash = GetHierarchy()->GetNameHash();
 
 	const TArray<FRigVMExternalVariable> ExternalVariables = GetExternalVariablesImpl(false);
@@ -2986,6 +3010,8 @@ uint32 UControlRig::GetHashForInitializeVMSnapShot()
 		Hash = HashCombine(ExternalVariable.GetTypeHash(), Hash);
 	}
 
+	Hash = HashCombine(CachedMemoryHash, Hash);
+
 	return Hash;
 }
 
diff --git a/Engine/Plugins/Animation/ControlRig/Source/ControlRig/Public/ControlRig.h b/Engine/Plugins/Animation/ControlRig/Source/ControlRig/Public/ControlRig.h
--- a/Engine/Plugins/Animation/ControlRig/Source/ControlRig/Public/ControlRig.h
+++ b/Engine/Plugins/Animation/ControlRig/Source/ControlRig/Public/ControlRig.h
@@ -463,6 +463,7 @@ private:
 	// and stores an initialized VM per hash.
 	UPROPERTY(transient)
 	TMap<uint32, TObjectPtr<URigVM>> InitializedVMSnapshots;
+	uint32 CachedMemoryHash = INDEX_NONE;
 
 	// computes the hash used to store / find VM snapshots
 	uint32 GetHashForInitializeVMSnapShot();
