Control Rig: Clear shape name widget prior to destroying the details panel

diff --git a/Engine/Plugins/Animation/ControlRig/Source/ControlRigEditor/Private/ControlRigElementDetails.cpp b/Engine/Plugins/Animation/ControlRig/Source/ControlRigEditor/Private/ControlRigElementDetails.cpp
--- a/Engine/Plugins/Animation/ControlRig/Source/ControlRigEditor/Private/ControlRigElementDetails.cpp
+++ b/Engine/Plugins/Animation/ControlRig/Source/ControlRigEditor/Private/ControlRigElementDetails.cpp
@@ -8,7 +8,6 @@
 #include "Widgets/Input/SVectorInputBox.h"
 #include "Widgets/Input/SCheckBox.h"
 #include "Widgets/Input/SButton.h"
-#include "Editor/SControlRigGizmoNameList.h"
 #include "ControlRigBlueprint.h"
 #include "Graph/ControlRigGraph.h"
 #include "PropertyCustomizationHelpers.h"
@@ -631,6 +630,8 @@ void FRigBaseElementDetails::SetName(const FText& InNewText, ETextCommit::Type I
 
 	if (Hierarchy)
 	{
+		BeginDestroy();
+		
 		URigHierarchyController* Controller = Hierarchy->GetController(true);
 		check(Controller);
 		Controller->RenameElement(GetElementKey(), *InNewText.ToString(), true, true);
@@ -2936,7 +2937,7 @@ void FRigControlElementDetails::CustomizeShape(IDetailLayoutBuilder& DetailBuild
 	]
 	.ValueContent()
 	[
-		SNew(SControlRigShapeNameList, ControlElements, BlueprintBeingCustomized)
+		SAssignNew(ShapeNameListWidget, SControlRigShapeNameList, ControlElements, BlueprintBeingCustomized)
 		.OnGetNameListContent(this, &FRigControlElementDetails::GetShapeNameList)
 		.IsEnabled(this, &FRigControlElementDetails::IsShapeEnabled)
 	];
@@ -2946,6 +2947,16 @@ void FRigControlElementDetails::CustomizeShape(IDetailLayoutBuilder& DetailBuild
 	.DisplayName(FText::FromString(TEXT("Color")));
 }
 
+void FRigControlElementDetails::BeginDestroy()
+{
+	FRigTransformElementDetails::BeginDestroy();
+
+	if(ShapeNameListWidget.IsValid())
+	{
+		ShapeNameListWidget->BeginDestroy();
+	}
+}
+
 void FRigControlElementDetails::RegisterSectionMappings(FPropertyEditorModule& PropertyEditorModule, UClass* InClass)
 {
 	FRigTransformElementDetails::RegisterSectionMappings(PropertyEditorModule, InClass);
diff --git a/Engine/Plugins/Animation/ControlRig/Source/ControlRigEditor/Private/ControlRigElementDetails.h b/Engine/Plugins/Animation/ControlRig/Source/ControlRigEditor/Private/ControlRigElementDetails.h
--- a/Engine/Plugins/Animation/ControlRig/Source/ControlRigEditor/Private/ControlRigElementDetails.h
+++ b/Engine/Plugins/Animation/ControlRig/Source/ControlRigEditor/Private/ControlRigElementDetails.h
@@ -11,6 +11,7 @@
 #include "ControlRigBlueprint.h"
 #include "DetailsViewWrapperObject.h"
 #include "Graph/SControlRigGraphPinNameListValueWidget.h"
+#include "Editor/SControlRigGizmoNameList.h"
 #include "Styling/SlateTypes.h"
 #include "IPropertyUtilities.h"
 #include "SSearchableComboBox.h"
@@ -105,6 +106,7 @@ public:
 
 	/** IDetailCustomization interface */
 	virtual void CustomizeDetails(IDetailLayoutBuilder& DetailBuilder) override;
+	virtual void BeginDestroy() {};
 
 	URigHierarchy* GetHierarchy() const { return HierarchyBeingCustomized; }
 	URigHierarchy* GetHierarchyBeingDebugged() const;
@@ -245,6 +247,7 @@ public:
 	void CustomizeValue(IDetailLayoutBuilder& DetailBuilder);
 	void CustomizeControl(IDetailLayoutBuilder& DetailBuilder);
 	void CustomizeShape(IDetailLayoutBuilder& DetailBuilder);
+	virtual void BeginDestroy() override;
 
 	/** FRigBaseElementDetails interface */
 	virtual void RegisterSectionMappings(FPropertyEditorModule& PropertyEditorModule, UClass* InClass) override;
@@ -556,7 +559,7 @@ private:
 
 	TArray<FRigControlElement> ControlElements;
 	TArray<UDetailsViewWrapperObject*> ObjectPerControl;
-
+	TSharedPtr<SControlRigShapeNameList> ShapeNameListWidget; 
 	static TSharedPtr<TArray<ERigControlValueType>> PickedValueTypes;
 };
 
diff --git a/Engine/Plugins/Animation/ControlRig/Source/ControlRigEditor/Private/Editor/SControlRigGizmoNameList.cpp b/Engine/Plugins/Animation/ControlRig/Source/ControlRigEditor/Private/Editor/SControlRigGizmoNameList.cpp
--- a/Engine/Plugins/Animation/ControlRig/Source/ControlRigEditor/Private/Editor/SControlRigGizmoNameList.cpp
+++ b/Engine/Plugins/Animation/ControlRig/Source/ControlRigEditor/Private/Editor/SControlRigGizmoNameList.cpp
@@ -45,6 +45,14 @@ void SControlRigShapeNameList::Construct(const FArguments& InArgs, TArray<FRigCo
 	ConstructCommon();
 }
 
+void SControlRigShapeNameList::BeginDestroy()
+{
+	if(NameListComboBox.IsValid())
+	{
+		NameListComboBox->SetOptionsSource(&GetEmptyList());
+	}
+}
+
 void SControlRigShapeNameList::ConstructCommon()
 {
 	SBox::Construct(SBox::FArguments());
@@ -84,7 +92,7 @@ const TArray<TSharedPtr<FString>>& SControlRigShapeNameList::GetNameList() const
 	{
 		return OnGetNameListContent.Execute();
 	}
-	return EmptyList;
+	return GetEmptyList();
 }
 
 FText SControlRigShapeNameList::GetNameListText() const
diff --git a/Engine/Plugins/Animation/ControlRig/Source/ControlRigEditor/Private/Editor/SControlRigGizmoNameList.h b/Engine/Plugins/Animation/ControlRig/Source/ControlRigEditor/Private/Editor/SControlRigGizmoNameList.h
--- a/Engine/Plugins/Animation/ControlRig/Source/ControlRigEditor/Private/Editor/SControlRigGizmoNameList.h
+++ b/Engine/Plugins/Animation/ControlRig/Source/ControlRigEditor/Private/Editor/SControlRigGizmoNameList.h
@@ -25,6 +25,7 @@ public:
 	void Construct(const FArguments& InArgs, FRigControlElement* ControlElement, UControlRigBlueprint* InBlueprint);
 	void Construct(const FArguments& InArgs, TArray<FRigControlElement*> ControlElements, UControlRigBlueprint* InBlueprint);
 	void Construct(const FArguments& InArgs, TArray<FRigControlElement> ControlElements, UControlRigBlueprint* InBlueprint);
+	void BeginDestroy();
 
 protected:
 
@@ -38,9 +39,14 @@ protected:
 	void OnNameListChanged(TSharedPtr<FString> NewSelection, ESelectInfo::Type SelectInfo);
 	void OnNameListComboBox();
 
+	FORCEINLINE static const TArray< TSharedPtr<FString> >& GetEmptyList()
+	{
+		static const TArray< TSharedPtr<FString> > EmptyList;
+		return EmptyList;
+	}
+
 	FOnGetNameListContent OnGetNameListContent;
 	TSharedPtr<SControlRigGraphPinNameListValueWidget> NameListComboBox;
-	TArray<TSharedPtr<FString>> EmptyList;
 
 	TArray<FRigElementKey> ControlKeys;
 	UControlRigBlueprint* Blueprint;
