Control Rig: Ensure to avoid invalid VMs / memory during runtime

diff --git a/Engine/Plugins/Animation/ControlRig/Source/ControlRig/Private/ControlRig.cpp b/Engine/Plugins/Animation/ControlRig/Source/ControlRig/Private/ControlRig.cpp
--- a/Engine/Plugins/Animation/ControlRig/Source/ControlRig/Private/ControlRig.cpp
+++ b/Engine/Plugins/Animation/ControlRig/Source/ControlRig/Private/ControlRig.cpp
@@ -1196,6 +1196,16 @@ void UControlRig::ExecuteUnits(FRigUnitContext& InOutContext, const FName& InEve
 {
 	if (VM)
 	{
+		// sanity check the validity of the VM to ensure stability.
+		if(!IsValidLowLevel() ||
+			!VM->IsValidLowLevel() ||
+			!VM->GetLiteralMemory()->IsValidLowLevel() ||
+			!VM->GetWorkMemory()->IsValidLowLevel())
+		{
+			UE_LOG(LogControlRig, Warning, TEXT("%s: Invalidated VM - aborting execution."), *GetClass()->GetName());
+			return;
+		}
+		
 #if UE_CONTROLRIG_PROFILE_EXECUTE_UNITS_NUM
 		const uint64 StartCycles = FPlatformTime::Cycles64();
 		if(ProfilingRunsLeft <= 0)
@@ -1225,7 +1235,8 @@ void UControlRig::ExecuteUnits(FRigUnitContext& InOutContext, const FName& InEve
 					{
 						const URigVM* InitializedVMSnapshot = InitializedVMSnapshotPtr->Get();
 
-						if(VM->WorkMemoryStorageObject->GetClass() == InitializedVMSnapshot->WorkMemoryStorageObject->GetClass())
+						if(VM->WorkMemoryStorageObject->GetClass() == InitializedVMSnapshot->WorkMemoryStorageObject->GetClass() &&
+							InitializedVMSnapshot->WorkMemoryStorageObject->IsValidLowLevel())
 						{
 							InitializedVMSnapshots.Reset();
 
