Optimizing Spherical Pose Reader

diff --git a/Engine/Plugins/Animation/ControlRig/Source/ControlRig/Private/Units/Highlevel/Hierarchy/RigUnit_SphericalPoseReader.cpp b/Engine/Plugins/Animation/ControlRig/Source/ControlRig/Private/Units/Highlevel/Hierarchy/RigUnit_SphericalPoseReader.cpp
--- a/Engine/Plugins/Animation/ControlRig/Source/ControlRig/Private/Units/Highlevel/Hierarchy/RigUnit_SphericalPoseReader.cpp
+++ b/Engine/Plugins/Animation/ControlRig/Source/ControlRig/Private/Units/Highlevel/Hierarchy/RigUnit_SphericalPoseReader.cpp
@@ -18,6 +18,7 @@ FRigUnit_SphericalPoseReader_Execute()
 	{
 		DriverCache.Reset();
 		OptionalParentCache.Reset();
+		bCachedInitTransforms = false;
 		return;
 	}
 
@@ -25,8 +26,6 @@ FRigUnit_SphericalPoseReader_Execute()
 	{
 		return;
 	}
-	
-	OptionalParentCache.UpdateCache(OptionalParentItem, Hierarchy);
 
 	// remap/clamp inputs
 	RemapAndConvertInputs(
@@ -41,25 +40,40 @@ FRigUnit_SphericalPoseReader_Execute()
 
 	// get parent global space
 	FTransform GlobalDriverParentTransform;
-	FTransform GlobalDriverParentTransformInit;
+	OptionalParentCache.UpdateCache(OptionalParentItem, Hierarchy);
 	if (OptionalParentCache.IsValid())
 	{
 		GlobalDriverParentTransform = Hierarchy->GetGlobalTransformByIndex(OptionalParentCache, false);
-		GlobalDriverParentTransformInit = Hierarchy->GetGlobalTransformByIndex(OptionalParentCache, true);
 	}else
 	{
 		// if user does not specify a custom parent space, then simply use the driver's parent
 		GlobalDriverParentTransform = Hierarchy->GetParentTransformByIndex(DriverCache, false);
-		GlobalDriverParentTransformInit = Hierarchy->GetParentTransformByIndex(DriverCache, true);
 	}
-	
-	// get local rotation space of driver
-	const FTransform GlobalDriverTransformInit = Hierarchy->GetGlobalTransformByIndex(DriverCache, true);
-	FTransform LocalDriverTransformInit = GlobalDriverTransformInit * GlobalDriverParentTransformInit.Inverse();
-	// apply static offset in local space
-	FQuat RotationOffsetQuat = FQuat::MakeFromEuler(RotationOffset);
-	LocalDriverTransformInit.SetRotation(LocalDriverTransformInit.GetRotation() * RotationOffsetQuat);
-	
+
+	// check if we need to update the LocalDriverTransformInit
+	if (!bCachedInitTransforms || !CachedRotationOffset.Equals(RotationOffset))
+	{
+		bCachedInitTransforms = true;
+		CachedRotationOffset = RotationOffset;
+
+		FTransform GlobalDriverParentTransformInitInverse;
+		if (OptionalParentCache.IsValid())
+		{
+			GlobalDriverParentTransformInitInverse = Hierarchy->GetGlobalTransformByIndex(OptionalParentCache, true).Inverse();
+		}else
+		{
+			// if user does not specify a custom parent space, then simply use the driver's parent
+			GlobalDriverParentTransformInitInverse = Hierarchy->GetParentTransformByIndex(DriverCache, true).Inverse();
+		}
+
+		// get local rotation space of driver
+		const FTransform GlobalDriverTransformInit = Hierarchy->GetGlobalTransformByIndex(DriverCache, true);
+		LocalDriverTransformInit = GlobalDriverTransformInit * GlobalDriverParentTransformInitInverse;
+		// apply static offset in local space
+		FQuat RotationOffsetQuat = FQuat::MakeFromEuler(RotationOffset);
+		LocalDriverTransformInit.SetRotation(LocalDriverTransformInit.GetRotation() * RotationOffsetQuat);
+	}
+
 	// calculate world transform of sphere
 	const FTransform GlobalDriverTransform = Hierarchy->GetGlobalTransformByIndex(DriverCache);
 	FTransform WorldOffset = LocalDriverTransformInit * GlobalDriverParentTransform;
@@ -88,7 +102,7 @@ FRigUnit_SphericalPoseReader_Execute()
 		Debug.DrawDebug(WorldOffset, Context.DrawInterface, InnerRegion, OuterRegion, DriverNormal, Driver2D, OutputParam);
 		return;
 	}
-
+	
 	const float PointX = Driver2D.X;
 	const float PointY = Driver2D.Y;
 
@@ -107,7 +121,7 @@ FRigUnit_SphericalPoseReader_Execute()
 
 	// calc output param
 	OutputParam = CalcOutputParam(InnerEllipseResults, OuterEllipseResults);
-	
+
 	// do all debug drawing
 	Debug.DrawDebug(WorldOffset, Context.DrawInterface, InnerRegion, OuterRegion, DriverNormal, Driver2D, OutputParam);
 }
diff --git a/Engine/Plugins/Animation/ControlRig/Source/ControlRig/Private/Units/Highlevel/Hierarchy/RigUnit_SphericalPoseReader.h b/Engine/Plugins/Animation/ControlRig/Source/ControlRig/Private/Units/Highlevel/Hierarchy/RigUnit_SphericalPoseReader.h
--- a/Engine/Plugins/Animation/ControlRig/Source/ControlRig/Private/Units/Highlevel/Hierarchy/RigUnit_SphericalPoseReader.h
+++ b/Engine/Plugins/Animation/ControlRig/Source/ControlRig/Private/Units/Highlevel/Hierarchy/RigUnit_SphericalPoseReader.h
@@ -382,4 +382,11 @@ struct CONTROLRIG_API FRigUnit_SphericalPoseReader: public FRigUnit_HighlevelBas
 
 	UPROPERTY()
 	FCachedRigElement OptionalParentCache;
+	
+	UPROPERTY()
+	FTransform LocalDriverTransformInit;
+	UPROPERTY()
+	FVector CachedRotationOffset;
+	UPROPERTY()
+	bool bCachedInitTransforms;
 };
\ No newline at end of file
