Fixes warnings in incremental cook from Spherical Pose Reader.

diff --git a/Engine/Plugins/Animation/ControlRig/Source/ControlRig/Private/Units/Highlevel/Hierarchy/RigUnit_SphericalPoseReader.h b/Engine/Plugins/Animation/ControlRig/Source/ControlRig/Private/Units/Highlevel/Hierarchy/RigUnit_SphericalPoseReader.h
--- a/Engine/Plugins/Animation/ControlRig/Source/ControlRig/Private/Units/Highlevel/Hierarchy/RigUnit_SphericalPoseReader.h
+++ b/Engine/Plugins/Animation/ControlRig/Source/ControlRig/Private/Units/Highlevel/Hierarchy/RigUnit_SphericalPoseReader.h
@@ -274,17 +274,27 @@ struct CONTROLRIG_API FRigUnit_SphericalPoseReader: public FRigUnit_HighlevelBas
 	GENERATED_BODY()
 
 	FRigUnit_SphericalPoseReader()
-			: OutputParam(0.0f),
-			DriverItem(FRigElementKey(NAME_None, ERigElementType::Bone)),
-			DriverAxis(FVector(1.0f,0.0f,0.0f)),
-			RotationOffset(FVector(90.0f,0.0f,0.0f)),
-			ActiveRegionSize(0.1f),
-			FalloffSize(0.2f),
-			Debug(FSphericalPoseReaderDebugSettings()),
-			InnerRegion(FSphericalRegion()),
-			OuterRegion(FSphericalRegion()),
-			DriverNormal(FVector::Zero()),
-			Driver2D(FVector::Zero())
+	: OutputParam(0.0f),
+	DriverItem(FRigElementKey(NAME_None, ERigElementType::Bone)),
+	DriverAxis(FVector(1.0f,0.0f,0.0f)),
+	RotationOffset(FVector(90.0f,0.0f,0.0f)),
+	ActiveRegionSize(0.1f),
+	ActiveRegionScaleFactors(FRegionScaleFactors()),
+	FalloffSize(0.2f),
+	FalloffRegionScaleFactors(FRegionScaleFactors()),
+	FlipWidthScaling(false),
+	FlipHeightScaling(false),
+	OptionalParentItem(FRigElementKey(NAME_None, ERigElementType::Bone)),
+	Debug(FSphericalPoseReaderDebugSettings()),
+	InnerRegion(FSphericalRegion()),
+	OuterRegion(FSphericalRegion()),
+	DriverNormal(FVector::Zero()),
+	Driver2D(FVector::Zero()),
+	DriverCache(FCachedRigElement()),
+	OptionalParentCache(FCachedRigElement()),
+	LocalDriverTransformInit(FTransform::Identity),
+	CachedRotationOffset(FVector::ZeroVector),
+	bCachedInitTransforms(false)
 	{
 	}
 	
@@ -332,33 +342,33 @@ struct CONTROLRIG_API FRigUnit_SphericalPoseReader: public FRigUnit_HighlevelBas
 
 	// The axis of the driver transform that is compared against the falloff regions. Typically the axis that is pointing at the child; usually X by convention. Default is X-axis (1,0,0).
 	UPROPERTY(meta = (Input))
-	FVector DriverAxis = FVector(1.0f,0.0f,0.0f);
+	FVector DriverAxis;
 
 	// Rotate the entire falloff region to align with the desired area of effect.
 	UPROPERTY(meta = (Input))
-	FVector RotationOffset = FVector(90.0f,0.0f,0.0f);
+	FVector RotationOffset;
 
 	// The size of the active region (green) that outputs the full value (1.0). Range is 0-1. Default is 0.1.
 	UPROPERTY(meta = (Input, ClampMin = "0.0", ClampMax = "1.0", UIMin = "0.0", UIMax = "1.0"))
-    float ActiveRegionSize = 0.1f;
+    float ActiveRegionSize;
 	// The directional scaling parameters for the active region (green).
 	UPROPERTY(meta = (Input))
 	FRegionScaleFactors ActiveRegionScaleFactors;
 
 	// The size of the falloff region (yellow) that defines the start of the output range. A value of 1 wraps the entire sphere with falloff. Range is 0-1. Default is 0.2.
 	UPROPERTY(meta = (Input, ClampMin = "0.0", ClampMax = "1.0", UIMin = "0.0", UIMax = "1.0"))
-    float FalloffSize = 0.2f;
+    float FalloffSize;
 	// The directional scaling parameters for the falloff region (yellow).
 	UPROPERTY(meta = (Input))
 	FRegionScaleFactors FalloffRegionScaleFactors;
 
 	// Flip the positive / negative directions of the width scale factors.
 	UPROPERTY(meta = (Input))
-	bool FlipWidthScaling = false;
+	bool FlipWidthScaling;
 
 	// Flip the positive / negative directions of the height scale factors.
 	UPROPERTY(meta = (Input))
-	bool FlipHeightScaling = false;
+	bool FlipHeightScaling;
 
 	// An optional parent to use as a stable frame of reference for the active regions (defaults to DriverItem's parent if unset).
 	UPROPERTY(meta = (Input))
