Control Rig: Implement workflows for aim + twoboneIK in python

diff --git a/Engine/Plugins/Animation/ControlRig/Source/ControlRig/Private/Units/Highlevel/Hierarchy/RigUnit_TwoBoneIKSimple.cpp b/Engine/Plugins/Animation/ControlRig/Source/ControlRig/Private/Units/Highlevel/Hierarchy/RigUnit_TwoBoneIKSimple.cpp
--- a/Engine/Plugins/Animation/ControlRig/Source/ControlRig/Private/Units/Highlevel/Hierarchy/RigUnit_TwoBoneIKSimple.cpp
+++ b/Engine/Plugins/Animation/ControlRig/Source/ControlRig/Private/Units/Highlevel/Hierarchy/RigUnit_TwoBoneIKSimple.cpp
@@ -85,12 +85,12 @@ FRigUnit_TwoBoneIKSimplePerItem_Execute()
 	}
 
 	if (!CachedItemAIndex.UpdateCache(ItemA, Hierarchy) ||
-		!CachedItemBIndex.UpdateCache(ItemB, Hierarchy) ||
-		!CachedEffectorItemIndex.UpdateCache(EffectorItem, Hierarchy))
+		!CachedItemBIndex.UpdateCache(ItemB, Hierarchy))
 	{
 		return;
 	}
 
+	CachedEffectorItemIndex.UpdateCache(EffectorItem, Hierarchy);
 	CachedPoleVectorSpaceIndex.UpdateCache(PoleVectorSpace, Hierarchy);
 
 	if (Weight <= SMALL_NUMBER)
@@ -170,14 +170,20 @@ FRigUnit_TwoBoneIKSimplePerItem_Execute()
 		FVector PositionC = TransformB.InverseTransformPosition(TransformC.GetLocation());
 		TransformA.SetRotation(FQuat::Slerp(Hierarchy->GetGlobalTransform(CachedItemAIndex).GetRotation(), TransformA.GetRotation(), Weight));
 		TransformB.SetRotation(FQuat::Slerp(Hierarchy->GetGlobalTransform(CachedItemBIndex).GetRotation(), TransformB.GetRotation(), Weight));
-		TransformC.SetRotation(FQuat::Slerp(Hierarchy->GetGlobalTransform(CachedEffectorItemIndex).GetRotation(), TransformC.GetRotation(), Weight));
+		if(CachedEffectorItemIndex != INDEX_NONE)
+		{
+			TransformC.SetRotation(FQuat::Slerp(Hierarchy->GetGlobalTransform(CachedEffectorItemIndex).GetRotation(), TransformC.GetRotation(), Weight));
+		}
 		TransformB.SetLocation(TransformA.TransformPosition(PositionB));
 		TransformC.SetLocation(TransformB.TransformPosition(PositionC));
 	}
 
 	Hierarchy->SetGlobalTransform(CachedItemAIndex, TransformA, bPropagateToChildren);
 	Hierarchy->SetGlobalTransform(CachedItemBIndex, TransformB, bPropagateToChildren);
-	Hierarchy->SetGlobalTransform(CachedEffectorItemIndex, TransformC, bPropagateToChildren);
+	if(CachedEffectorItemIndex != INDEX_NONE)
+	{
+		Hierarchy->SetGlobalTransform(CachedEffectorItemIndex, TransformC, bPropagateToChildren);
+	}
 }
 
 FRigUnit_TwoBoneIKSimpleVectors_Execute()
