Control Rig: Focus on search when clicked Space

diff --git a/Engine/Plugins/Animation/ControlRig/Source/ControlRigEditor/Private/Editor/SRigHierarchyTreeView.cpp b/Engine/Plugins/Animation/ControlRig/Source/ControlRigEditor/Private/Editor/SRigHierarchyTreeView.cpp
--- a/Engine/Plugins/Animation/ControlRig/Source/ControlRigEditor/Private/Editor/SRigHierarchyTreeView.cpp
+++ b/Engine/Plugins/Animation/ControlRig/Source/ControlRigEditor/Private/Editor/SRigHierarchyTreeView.cpp
@@ -795,7 +795,7 @@ void SSearchableRigHierarchyTreeView::Construct(const FArguments& InArgs)
 		.HAlign(HAlign_Fill)
 		.Padding(0.0f)
 		[
-			SNew(SSearchBox)
+			SAssignNew(SearchBox, SSearchBox)
 			.InitialText(InArgs._InitialFilterText)
 			.OnTextChanged(this, &SSearchableRigHierarchyTreeView::OnFilterTextChanged)
 		]
diff --git a/Engine/Plugins/Animation/ControlRig/Source/ControlRigEditor/Private/Editor/SRigHierarchyTreeView.h b/Engine/Plugins/Animation/ControlRig/Source/ControlRigEditor/Private/Editor/SRigHierarchyTreeView.h
--- a/Engine/Plugins/Animation/ControlRig/Source/ControlRigEditor/Private/Editor/SRigHierarchyTreeView.h
+++ b/Engine/Plugins/Animation/ControlRig/Source/ControlRigEditor/Private/Editor/SRigHierarchyTreeView.h
@@ -6,6 +6,7 @@
 #include "Widgets/Views/STreeView.h"
 #include "Rigs/RigHierarchy.h"
 
+class SSearchBox;
 class SRigHierarchyTreeView;
 class SRigHierarchyItem;
 class FRigTreeElement;
@@ -307,6 +308,7 @@ public:
 
 	void Construct(const FArguments& InArgs);
 	virtual ~SSearchableRigHierarchyTreeView() {}
+	TSharedRef<SSearchBox> GetSearchBox() const { return SearchBox.ToSharedRef(); }
 	TSharedRef<SRigHierarchyTreeView> GetTreeView() const { return TreeView.ToSharedRef(); }
 	const FRigTreeDisplaySettings& GetDisplaySettings();
 
@@ -317,5 +319,6 @@ private:
 	FOnGetRigTreeDisplaySettings SuperGetRigTreeDisplaySettings;
 	FText FilterText;
 	FRigTreeDisplaySettings Settings;
+	TSharedPtr<SSearchBox> SearchBox;
 	TSharedPtr<SRigHierarchyTreeView> TreeView;
 };
diff --git a/Engine/Plugins/Animation/ControlRig/Source/ControlRigEditor/Private/Editor/SRigSpacePickerWidget.cpp b/Engine/Plugins/Animation/ControlRig/Source/ControlRigEditor/Private/Editor/SRigSpacePickerWidget.cpp
--- a/Engine/Plugins/Animation/ControlRig/Source/ControlRigEditor/Private/Editor/SRigSpacePickerWidget.cpp
+++ b/Engine/Plugins/Animation/ControlRig/Source/ControlRigEditor/Private/Editor/SRigSpacePickerWidget.cpp
@@ -6,6 +6,7 @@
 #include "Widgets/Text/STextBlock.h"
 #include "Widgets/Layout/SUniformGridPanel.h"
 #include "Widgets/Input/SButton.h"
+#include "Widgets/Input/SSearchBox.h"
 #include "SRigHierarchyTreeView.h"
 #include "ControlRigEditorStyle.h"
 #include "PropertyCustomizationHelpers.h"
@@ -659,6 +660,7 @@ FReply SRigSpacePickerWidget::HandleAddElementClicked()
 	.RigTreeDelegates(TreeDelegates);
 	SearchableTreeView->GetTreeView()->RefreshTreeView(true);
 
+	const bool bFocusImmediately = false;
 	// Create as context menu
 	TGuardValue<bool> AboutToShowMenu(bLaunchingContextMenu, true);
 	ContextMenu = FSlateApplication::Get().PushMenu(
@@ -666,7 +668,8 @@ FReply SRigSpacePickerWidget::HandleAddElementClicked()
 		FWidgetPath(),
 		SearchableTreeView.ToSharedRef(),
 		FSlateApplication::Get().GetCursorPos(),
-		FPopupTransitionEffect(FPopupTransitionEffect::ContextMenu)
+		FPopupTransitionEffect(FPopupTransitionEffect::ContextMenu),
+		bFocusImmediately
 	);
 
 	if(!ContextMenu.IsValid())
@@ -689,7 +692,7 @@ FReply SRigSpacePickerWidget::HandleAddElementClicked()
 		}
 	});
 
-	return FReply::Handled();
+	return FReply::Handled().SetUserFocus(SearchableTreeView->GetSearchBox(), EFocusCause::SetDirectly);
 }
 
 bool SRigSpacePickerWidget::IsSpaceMoveUpEnabled(FRigElementKey InKey) const
