Control Rig: Tweens button should be checked when enabled

diff --git a/Engine/Plugins/Animation/ControlRig/Source/ControlRigEditor/Private/EditMode/ControlRigEditModeToolkit.cpp b/Engine/Plugins/Animation/ControlRig/Source/ControlRigEditor/Private/EditMode/ControlRigEditModeToolkit.cpp
--- a/Engine/Plugins/Animation/ControlRig/Source/ControlRigEditor/Private/EditMode/ControlRigEditModeToolkit.cpp
+++ b/Engine/Plugins/Animation/ControlRig/Source/ControlRigEditor/Private/EditMode/ControlRigEditModeToolkit.cpp
@@ -132,6 +132,17 @@ void FControlRigEditModeToolkit::TryInvokeToolkitUI(const FName InName)
 	}
 }
 
+bool FControlRigEditModeToolkit::IsToolkitUIActive(const FName InName) const
+{
+	if (InName == TweenOverlayName)
+	{
+		return TweenWidget.IsValid();
+	}
+
+	TSharedPtr<FAssetEditorModeUILayer> ModeUILayerPtr = ModeUILayer.Pin();
+	return ModeUILayerPtr->GetTabManager()->FindExistingLiveTab(FTabId(InName)).IsValid();
+}
+
 FText FControlRigEditModeToolkit::GetActiveToolDisplayName() const
 {
 	return ModeTools->GetActiveToolName();
diff --git a/Engine/Plugins/Animation/ControlRig/Source/ControlRigEditor/Private/EditMode/ControlRigEditModeToolkit.h b/Engine/Plugins/Animation/ControlRig/Source/ControlRigEditor/Private/EditMode/ControlRigEditModeToolkit.h
--- a/Engine/Plugins/Animation/ControlRig/Source/ControlRigEditor/Private/EditMode/ControlRigEditModeToolkit.h
+++ b/Engine/Plugins/Animation/ControlRig/Source/ControlRigEditor/Private/EditMode/ControlRigEditModeToolkit.h
@@ -48,6 +48,7 @@ public:
 	virtual void OnToolPaletteChanged(FName PaletteName) override;
 
 	void TryInvokeToolkitUI(const FName InName);
+	bool IsToolkitUIActive(const FName InName) const;
 
 public:
 	static const FName PoseTabName;
diff --git a/Engine/Plugins/Animation/ControlRig/Source/ControlRigEditor/Private/EditMode/SControlRigEditModeTools.cpp b/Engine/Plugins/Animation/ControlRig/Source/ControlRigEditor/Private/EditMode/SControlRigEditModeTools.cpp
--- a/Engine/Plugins/Animation/ControlRig/Source/ControlRigEditor/Private/EditMode/SControlRigEditModeTools.cpp
+++ b/Engine/Plugins/Animation/ControlRig/Source/ControlRigEditor/Private/EditMode/SControlRigEditModeTools.cpp
@@ -974,12 +974,16 @@ void SControlRigEditModeTools::CustomizeToolBarPalette(FToolBarBuilder& ToolBarB
 
 	// Tweens
 	ToolBarBuilder.AddToolBarButton(
+		FUIAction(
 		FExecuteAction::CreateRaw(OwningToolkit.Pin().Get(), &FControlRigEditModeToolkit::TryInvokeToolkitUI, FControlRigEditModeToolkit::TweenOverlayName),
+		FCanExecuteAction(),
+		FIsActionChecked::CreateRaw(OwningToolkit.Pin().Get(), &FControlRigEditModeToolkit::IsToolkitUIActive, FControlRigEditModeToolkit::TweenOverlayName)
+		),		
 		NAME_None,
 		LOCTEXT("Tweens", "Tweens"),
 		LOCTEXT("TweensTooltip", "Create Tweens"),
 		FSlateIcon(TEXT("ControlRigEditorStyle"), TEXT("ControlRig.TweenTool")),
-		EUserInterfaceActionType::Button
+		EUserInterfaceActionType::ToggleButton
 	);
 
 	// Snap
