Control Rig: Note selection change so that viewport hit proxies will be invalidated when selecting controls

This fixes an issue where if you select a control in the anim outliner, the hit proxies (ie. gizmos) won't be updated so you can't click on a gizmo to move it immediately.

diff --git a/Engine/Plugins/Animation/ControlRig/Source/ControlRigEditor/Private/Sequencer/ControlRigParameterTrackEditor.cpp b/Engine/Plugins/Animation/ControlRig/Source/ControlRigEditor/Private/Sequencer/ControlRigParameterTrackEditor.cpp
--- a/Engine/Plugins/Animation/ControlRig/Source/ControlRigEditor/Private/Sequencer/ControlRigParameterTrackEditor.cpp
+++ b/Engine/Plugins/Animation/ControlRig/Source/ControlRigEditor/Private/Sequencer/ControlRigParameterTrackEditor.cpp
@@ -1963,6 +1963,8 @@ void FControlRigParameterTrackEditor::OnSelectionChanged(TArray<UMovieSceneTrack
 	}
 
 	SelectRigsAndControls(ControlRig, KeyAreas);
+	
+	GEditor->NoteSelectionChange();
 }
 
 void FControlRigParameterTrackEditor::SelectRigsAndControls(UControlRig* ControlRig, const TArray<const IKeyArea*>& KeyAreas)
