Control Rig: only copy hierarchy during bake for user authored rigs

diff --git a/Engine/Plugins/Animation/ControlRig/Source/ControlRig/Private/Sequencer/MovieSceneControlRigParameterSection.cpp b/Engine/Plugins/Animation/ControlRig/Source/ControlRig/Private/Sequencer/MovieSceneControlRigParameterSection.cpp
--- a/Engine/Plugins/Animation/ControlRig/Source/ControlRig/Private/Sequencer/MovieSceneControlRigParameterSection.cpp
+++ b/Engine/Plugins/Animation/ControlRig/Source/ControlRig/Private/Sequencer/MovieSceneControlRigParameterSection.cpp
@@ -2634,9 +2634,12 @@ bool UMovieSceneControlRigParameterSection::LoadAnimSequenceIntoThisSection(UAni
 
 	// copy the hierarchy from the CDO into the target control rig.
 	// this ensures that the topology version matches in case of a dynamic hierarchy
-	if (UControlRig* CDO = Cast<UControlRig>(ControlRig->GetClass()->GetDefaultObject()))
+	if(!ControlRig->IsNative())
 	{
-		ControlRig->GetHierarchy()->CopyHierarchy(CDO->GetHierarchy());
+		if (UControlRig* CDO = Cast<UControlRig>(ControlRig->GetClass()->GetDefaultObject()))
+		{
+			ControlRig->GetHierarchy()->CopyHierarchy(CDO->GetHierarchy());
+		}
 	}
 
 	// now set the hierarchies initial transforms based on the currently used skeletal mesh
