Control Rig: FIx ensure after undoing a variable deletion

diff --git a/Engine/Plugins/Animation/ControlRig/Source/ControlRigDeveloper/Private/Graph/ControlRigGraph.cpp b/Engine/Plugins/Animation/ControlRig/Source/ControlRigDeveloper/Private/Graph/ControlRigGraph.cpp
--- a/Engine/Plugins/Animation/ControlRig/Source/ControlRigDeveloper/Private/Graph/ControlRigGraph.cpp
+++ b/Engine/Plugins/Animation/ControlRig/Source/ControlRigDeveloper/Private/Graph/ControlRigGraph.cpp
@@ -929,24 +929,12 @@ URigVMController* UControlRigGraph::GetController() const
 
 void UControlRigGraph::RemoveAndDeleteNode(UEdGraphNode* InNode)
 {
+	// Make sure EdGraph is not part of the transaction
+	TGuardValue<ITransaction*> TransactionGuard(GUndo, nullptr);
+					
 	static UObject* NewOuter = GetTransientPackage();
+	InNode->Rename(nullptr, NewOuter, REN_ForceNoResetLoaders | REN_DontCreateRedirectors);
 					
-	// Rename the soon to be deleted object to a unique name, so that other objects can use
-	// the old name
-	{
-		FString DeletedName;
-		UObject* ExistingObject = nullptr;
-		static int32 DeletedIndex = FMath::Rand();
-		do
-		{
-			DeletedName = FString::Printf(TEXT("%s_Deleted_%d"), *InNode->GetName(), DeletedIndex++); 
-			ExistingObject = StaticFindObject(/*Class=*/ NULL, NewOuter, *DeletedName, true);						
-		}
-		while (ExistingObject);
-		InNode->Rename(*DeletedName, NewOuter, REN_ForceNoResetLoaders | REN_DoNotDirty | REN_DontCreateRedirectors | REN_NonTransactional);
-		InNode->MarkAsGarbage();
-	}
-
 	// Remove the node and notify of the removal
 	{
 		Nodes.Remove(InNode);
diff --git a/Engine/Source/Developer/RigVMDeveloper/Private/RigVMModel/RigVMController.cpp b/Engine/Source/Developer/RigVMDeveloper/Private/RigVMModel/RigVMController.cpp
--- a/Engine/Source/Developer/RigVMDeveloper/Private/RigVMModel/RigVMController.cpp
+++ b/Engine/Source/Developer/RigVMDeveloper/Private/RigVMModel/RigVMController.cpp
@@ -1275,6 +1275,12 @@ void URigVMController::OnExternalVariableRemoved(const FName& InVarName, bool bS
 	URigVMGraph* Graph = GetGraph();
 	check(Graph);
 
+	// When transacting, the action stack will deal with the deletion of variable nodes
+	if(GIsTransacting)
+	{
+		return;
+	}
+
 	for (const FRigVMGraphVariableDescription& LocalVariable : Graph->GetLocalVariables(true))
 	{
 		if (InVarName == LocalVariable.Name)
