ControlRig: Assertion thrown while deleting a Curve from a Control Rig asset

diff --git a/Engine/Plugins/Animation/ControlRig/Source/ControlRigEditor/Private/Editor/ControlRigEditor.cpp b/Engine/Plugins/Animation/ControlRig/Source/ControlRigEditor/Private/Editor/ControlRigEditor.cpp
--- a/Engine/Plugins/Animation/ControlRig/Source/ControlRigEditor/Private/Editor/ControlRigEditor.cpp
+++ b/Engine/Plugins/Animation/ControlRig/Source/ControlRigEditor/Private/Editor/ControlRigEditor.cpp
@@ -2886,6 +2886,11 @@ void FControlRigEditor::HandleControlRigExecutedEvent(UControlRig* InControlRig,
 						const FRigElementKey Key = WrapperObject->GetContent<FRigBaseElement>().GetKey();
 
 						FRigBaseElement* Element = Hierarchy->Find(Key);
+						if(Element == nullptr)
+						{
+							ClearDetailObject();
+							break;
+						}
 
 						if(FRigControlElement* ControlElement = Cast<FRigControlElement>(Element))
 						{
