Control Rig: Only select sub controls with no shape

diff --git a/Engine/Plugins/Animation/ControlRig/Source/ControlRigEditor/Private/Sequencer/ControlRigParameterTrackEditor.cpp b/Engine/Plugins/Animation/ControlRig/Source/ControlRigEditor/Private/Sequencer/ControlRigParameterTrackEditor.cpp
--- a/Engine/Plugins/Animation/ControlRig/Source/ControlRigEditor/Private/Sequencer/ControlRigParameterTrackEditor.cpp
+++ b/Engine/Plugins/Animation/ControlRig/Source/ControlRigEditor/Private/Sequencer/ControlRigParameterTrackEditor.cpp
@@ -2011,7 +2011,34 @@ void FControlRigParameterTrackEditor::SelectRigsAndControls(UControlRig* Control
 				String.ParseIntoArray(StringArray, TEXT("."));
 				if (StringArray.Num() > 0)
 				{
-					FName ControlName(*StringArray[0]);
+					const FName ControlName(*StringArray[0]);
+
+					// skip nested controls which have the shape enabled flag turned on
+					if(const URigHierarchy* Hierarchy = ControlRig->GetHierarchy())
+					{
+						if(const FRigControlElement* ControlElement = Hierarchy->Find<FRigControlElement>(FRigElementKey(ControlName, ERigElementType::Control)))
+						{
+							if (ControlElement->Settings.ControlType == ERigControlType::Bool ||
+								ControlElement->Settings.ControlType == ERigControlType::Float ||
+								ControlElement->Settings.ControlType == ERigControlType::Integer)
+							{
+								if(ControlElement->Settings.bShapeEnabled)
+								{
+									if(const FRigControlElement* ParentControlElement = Cast<FRigControlElement>(Hierarchy->GetFirstParent(ControlElement)))
+									{
+										if(const TSet<FName>* Controls = RigsAndControls.Find(ControlRig))
+										{
+											if(Controls->Contains(ParentControlElement->GetName()))
+											{
+												continue;
+											}
+										}
+									}
+								}
+							}
+						}
+					}
+					
 					RigsAndControls.FindOrAdd(ControlRig).Add(ControlName);
 				}
 			}
@@ -2312,13 +2339,18 @@ void FControlRigParameterTrackEditor::HandleControlSelected(UControlRig* Subject
 	{
 		Hierarchy->ForEach<FRigControlElement>([ControlElement, Controller, bSelected](FRigControlElement* OtherControlElement) -> bool
 			{
-			if(OtherControlElement->Settings.ControlType == ERigControlType::Bool ||
+				if (OtherControlElement->Settings.ControlType == ERigControlType::Bool ||
 					OtherControlElement->Settings.ControlType == ERigControlType::Float ||
 					OtherControlElement->Settings.ControlType == ERigControlType::Integer)
 				{
-				for(const FRigElementParentConstraint& ParentConstraint : OtherControlElement->ParentConstraints)
+					if(OtherControlElement->Settings.bShapeEnabled || !OtherControlElement->Settings.bAnimatable)
+					{
+						return true;
+					}
+					
+					for (const FRigElementParentConstraint& ParentConstraint : OtherControlElement->ParentConstraints)
 					{
-					if(ParentConstraint.ParentElement == ControlElement)
+						if (ParentConstraint.ParentElement == ControlElement)
 						{
 							Controller->SelectElement(OtherControlElement->GetKey(), bSelected);
 							break;
