Sequencer: Make sure we delete keys on the base section when collapsing else we will get duplicated keys. And minor from Main fix to deselect children controls when parent is deselected.

diff --git a/Engine/Plugins/Animation/ControlRig/Source/ControlRigEditor/Private/Sequencer/ControlRigParameterTrackEditor.cpp b/Engine/Plugins/Animation/ControlRig/Source/ControlRigEditor/Private/Sequencer/ControlRigParameterTrackEditor.cpp
--- a/Engine/Plugins/Animation/ControlRig/Source/ControlRigEditor/Private/Sequencer/ControlRigParameterTrackEditor.cpp
+++ b/Engine/Plugins/Animation/ControlRig/Source/ControlRigEditor/Private/Sequencer/ControlRigParameterTrackEditor.cpp
@@ -236,7 +236,7 @@ FControlRigParameterTrackEditor::FControlRigParameterTrackEditor(TSharedRef<ISeq
 							if (UControlRig* ControlRig = ControlRigPtr.Get())
 							{
 								if (ControlRig->GetObjectBinding() && ControlRig->GetObjectBinding()->GetBoundObject() == nullptr)
-								{
+					{
 									bRequestEvaluate = true;
 									break;
 								}
@@ -257,7 +257,7 @@ FControlRigParameterTrackEditor::FControlRigParameterTrackEditor(TSharedRef<ISeq
 							if (OldControlRigComponent->GetControlRig())
 							{
 								UControlRig* NewControlRig = nullptr;
-								if(NewControlRigComponent)
+								if (NewControlRigComponent)
 								{
 									NewControlRig = NewControlRigComponent->GetControlRig();
 								}
@@ -417,7 +417,7 @@ void FControlRigParameterTrackEditor::UnbindControlRig(UControlRig* ControlRig)
 }
 void FControlRigParameterTrackEditor::UnbindAllControlRigs()
 {
-	for(TWeakObjectPtr<UControlRig>& ObjectPtr: BoundControlRigs)
+	for (TWeakObjectPtr<UControlRig>& ObjectPtr : BoundControlRigs)
 	{
 		if (ObjectPtr.IsValid())
 		{
@@ -697,7 +697,7 @@ void FControlRigParameterTrackEditor::BakeToControlRigSubMenu(FMenuBuilder& Menu
 		FClassViewerInitializationOptions Options;
 		Options.bShowUnloadedBlueprints = true;
 		Options.NameTypeToDisplay = EClassViewerNameTypeToDisplay::DisplayName;
-		TSharedPtr<FControlRigClassFilter> ClassFilter = MakeShareable(new FControlRigClassFilter(bFilterAssetBySkeleton,true, true, Skeleton));
+		TSharedPtr<FControlRigClassFilter> ClassFilter = MakeShareable(new FControlRigClassFilter(bFilterAssetBySkeleton, true, true, Skeleton));
 		Options.ClassFilters.Add(ClassFilter.ToSharedRef());
 		Options.bShowNoneOption = false;
 
@@ -1090,7 +1090,7 @@ void FControlRigParameterTrackEditor::BuildObjectBindingTrackMenu(FMenuBuilder&
 	if (ObjectClass->IsChildOf(USkeletalMeshComponent::StaticClass()) || ObjectClass->IsChildOf(AActor::StaticClass()) || ObjectClass->IsChildOf(UChildActorComponent::StaticClass()))
 	{
 		const TSharedPtr<ISequencer> ParentSequencer = GetSequencer();
-		UObject *BoundObject = nullptr;
+		UObject* BoundObject = nullptr;
 		USkeleton* Skeleton = AcquireSkeletonFromObjectGuid(ObjectBindings[0], &BoundObject, ParentSequencer);
 
 		if (AActor* BoundActor = Cast<AActor>(BoundObject))
@@ -1223,7 +1223,7 @@ void FControlRigParameterTrackEditor::HandleAddControlRigSubMenu(FMenuBuilder& M
 
 
 	const TSharedPtr<ISequencer> ParentSequencer = GetSequencer();
-	UObject *BoundObject = nullptr;
+	UObject* BoundObject = nullptr;
 	//todo support multiple bindings?
 	USkeleton* Skeleton = AcquireSkeletonFromObjectGuid(ObjectBindings[0], &BoundObject, GetSequencer());
 
@@ -1407,7 +1407,7 @@ void FControlRigParameterTrackEditor::AddFKControlRig(TArray<FGuid> ObjectBindin
 {
 	for (const FGuid& ObjectBinding : ObjectBindings)
 	{
-		UObject *BoundObject = nullptr;
+		UObject* BoundObject = nullptr;
 		AcquireSkeletonFromObjectGuid(ObjectBinding, &BoundObject, GetSequencer());
 		if (BoundObject)
 		{
@@ -1444,7 +1444,7 @@ bool FControlRigParameterTrackEditor::CanAddTransformKeysForSelectedObjects() co
 		TMap<UControlRig*, TArray<FRigElementKey>> SelectedControls;
 		ControlRigEditMode->GetAllSelectedControls(SelectedControls);
 		return (SelectedControls.Num() > 0);
-	}
+		}
 	return false;
 }
 
@@ -1462,25 +1462,25 @@ void FControlRigParameterTrackEditor::OnAddTransformKeysForSelectedObjects(EMovi
 		ControlRigEditMode->GetAllSelectedControls(SelectedControls);
 
 		for (TPair<UControlRig*, TArray<FRigElementKey>>& Selection : SelectedControls)
-		{
+	{
 			UControlRig* ControlRig = Selection.Key;
-			FString OurName = ControlRig->GetName();
-			FName Name(*OurName);
-			if (TSharedPtr<IControlRigObjectBinding> ObjectBinding = ControlRig->GetObjectBinding())
+		FString OurName = ControlRig->GetName();
+		FName Name(*OurName);
+		if (TSharedPtr<IControlRigObjectBinding> ObjectBinding = ControlRig->GetObjectBinding())
+		{
+			TArray<FName> ControlNames = ControlRig->CurrentControlSelection();
+			for (const FName& ControlName : ControlNames)
 			{
-				TArray<FName> ControlNames = ControlRig->CurrentControlSelection();
-				for (const FName& ControlName : ControlNames)
+				USceneComponent* Component = Cast<USceneComponent>(ObjectBinding->GetBoundObject());
+				if (Component)
 				{
-					USceneComponent* Component = Cast<USceneComponent>(ObjectBinding->GetBoundObject());
-					if (Component)
-					{
-						AddControlKeys(Component, ControlRig, Name, ControlName, (EControlRigContextChannelToKey)Channel, ESequencerKeyMode::ManualKeyForced, FLT_MAX);
-					}
+					AddControlKeys(Component, ControlRig, Name, ControlName, (EControlRigContextChannelToKey)Channel, ESequencerKeyMode::ManualKeyForced, FLT_MAX);
 				}
 			}
 		}
 	}
 }
+}
 
 //function to evaluate a Control and Set it on the ControlRig
 static void EvaluateThisControl(UMovieSceneControlRigParameterSection* Section, const FName& ControlName, const FFrameTime& FrameTime)
@@ -1741,22 +1741,22 @@ void FControlRigParameterTrackEditor::OnSequencerDataChanged(EMovieSceneDataChan
 			{
 				ControlRig->SetFramesPerSecond(FPS);
 
-				const TArray<FMovieSceneBinding>& Bindings = MovieScene->GetBindings();
+		const TArray<FMovieSceneBinding>& Bindings = MovieScene->GetBindings();
 				bool bControlRigInTrack = false;
-				for (const FMovieSceneBinding& Binding : Bindings)
-				{
-					UMovieSceneControlRigParameterTrack* Track = Cast<UMovieSceneControlRigParameterTrack>(MovieScene->FindTrack(UMovieSceneControlRigParameterTrack::StaticClass(), Binding.GetObjectGuid(), NAME_None));
+		for (const FMovieSceneBinding& Binding : Bindings)
+		{
+			UMovieSceneControlRigParameterTrack* Track = Cast<UMovieSceneControlRigParameterTrack>(MovieScene->FindTrack(UMovieSceneControlRigParameterTrack::StaticClass(), Binding.GetObjectGuid(), NAME_None));
 					if (Track && Track->GetControlRig() == ControlRig)
-					{
+			{
 						bControlRigInTrack = true;
 						break;; //continue, we still have a good track
-					}
-				}
+			}
+		}
 				/* Nope don't do this anymore, todo mz
-				if (FEditorModeTools* Tools = GetEditorModeTools())
-				{
-					Tools->DeactivateMode(FControlRigEditMode::ModeName);
-				}
+		if (FEditorModeTools* Tools = GetEditorModeTools())
+		{
+			Tools->DeactivateMode(FControlRigEditMode::ModeName);
+		}
 				*/
 				if (bControlRigInTrack == false)
 				{
@@ -1926,12 +1926,12 @@ void FControlRigParameterTrackEditor::OnSelectionChanged(TArray<UMovieSceneTrack
 			for (TPair<UControlRig*, TArray<FRigElementKey>>& SelectedControl : AllSelectedControls)
 			{
 				ControlRig = SelectedControl.Key;
-				if (ControlRig)
-				{
-					ControlRig->ClearControlSelection();
-				}
+			if (ControlRig)
+			{
+				ControlRig->ClearControlSelection();
 			}
 		}
+		}
 		for (UMovieSceneTrack* Track : InTracks)
 		{
 			UMovieSceneControlRigParameterTrack* CRTrack = Cast<UMovieSceneControlRigParameterTrack>(Track);
@@ -1972,7 +1972,7 @@ void FControlRigParameterTrackEditor::SelectRigsAndControls(UControlRig* Control
 	FControlRigEditMode* ControlRigEditMode = GetEditMode();
 
 	TArray<FString> StringArray;
-	TMap<UControlRig *, TSet<FName>> RigsAndControls;
+	TMap<UControlRig*, TSet<FName>> RigsAndControls;
 	for (const IKeyArea* KeyArea : KeyAreas)
 	{
 		UMovieSceneControlRigParameterSection* MovieSection = Cast<UMovieSceneControlRigParameterSection>(KeyArea->GetOwningSection());
@@ -2054,13 +2054,13 @@ void FControlRigParameterTrackEditor::SelectRigsAndControls(UControlRig* Control
 		for (TPair<UControlRig*, TArray<FRigElementKey>>& Selection : SelectedControls)
 		{
 			ControlRig = Selection.Key;
-			if (ControlRig)
-			{
-				ControlRig->ClearControlSelection();
-			}
+		if (ControlRig)
+		{
+			ControlRig->ClearControlSelection();
 		}
 	}
-	for (TPair<UControlRig *, TSet<FName>> Pair : RigsAndControls)
+	}
+	for (TPair<UControlRig*, TSet<FName>> Pair : RigsAndControls)
 	{
 		if (Pair.Key != ControlRig)
 		{
@@ -2133,7 +2133,7 @@ void FControlRigParameterTrackEditor::SelectSequencerNodeInSection(UMovieSceneCo
 				if (CategoryIndex != INDEX_NONE)
 				{
 					GetSequencer()->SelectByNthCategoryNode(ParamSection, CategoryIndex, bSelected);
-			}
+				}
 			}
 			else
 			{
@@ -2316,9 +2316,9 @@ void FControlRigParameterTrackEditor::SetUpEditModeIfNeeded(UControlRig* Control
 		}
 	}
 	else
-	{	
+	{
 		if (ControlRigEditMode->AddControlRigObject(ControlRig, GetSequencer()))
-		{
+			{
 			//force an evaluation, this will get the control rig setup so edit mode looks good.
 			if (GetSequencer().IsValid())
 			{
@@ -2334,35 +2334,35 @@ void FControlRigParameterTrackEditor::HandleControlSelected(UControlRig* Subject
 	TArray<FRigControl> Controls;
 
 	URigHierarchy* Hierarchy = Subject->GetHierarchy();
-	
-	if (URigHierarchyController* Controller = Hierarchy->GetController())
-	{
-		Hierarchy->ForEach<FRigControlElement>([ControlElement, Controller, bSelected](FRigControlElement* OtherControlElement) -> bool
-			{
-				if (OtherControlElement->Settings.ControlType == ERigControlType::Bool ||
-					OtherControlElement->Settings.ControlType == ERigControlType::Float ||
-					OtherControlElement->Settings.ControlType == ERigControlType::Integer)
+
+		if (URigHierarchyController* Controller = Hierarchy->GetController())
+		{
+			Hierarchy->ForEach<FRigControlElement>([ControlElement, Controller, bSelected](FRigControlElement* OtherControlElement) -> bool
 				{
-					if(OtherControlElement->Settings.bShapeEnabled || !OtherControlElement->Settings.bAnimatable)
+					if (OtherControlElement->Settings.ControlType == ERigControlType::Bool ||
+						OtherControlElement->Settings.ControlType == ERigControlType::Float ||
+						OtherControlElement->Settings.ControlType == ERigControlType::Integer)
 					{
-						return true;
-					}
-					
-					for (const FRigElementParentConstraint& ParentConstraint : OtherControlElement->ParentConstraints)
-					{
-						if (ParentConstraint.ParentElement == ControlElement)
+						if(OtherControlElement->Settings.bShapeEnabled || !OtherControlElement->Settings.bAnimatable)
 						{
-							Controller->SelectElement(OtherControlElement->GetKey(), bSelected);
-							break;
+							return true;
+						}
+						
+						for (const FRigElementParentConstraint& ParentConstraint : OtherControlElement->ParentConstraints)
+						{
+							if (ParentConstraint.ParentElement == ControlElement)
+							{
+								Controller->SelectElement(OtherControlElement->GetKey(), bSelected);
+								break;
+							}
 						}
 					}
-				}
-
-				return true;
-			});
-	}
 
+					return true;
+				});
+		}
 
+	
 	if (bIsDoingSelection)
 	{
 		return;
@@ -2396,8 +2396,8 @@ void FControlRigParameterTrackEditor::HandleControlSelected(UControlRig* Subject
 			GetSequencer()->SuspendSelectionBroadcast();
 			//Just set in the section to key not all
 			UMovieSceneSection* Section = Track->GetSectionToKey();
-				UMovieSceneControlRigParameterSection* ParamSection = Cast<UMovieSceneControlRigParameterSection>(Section);
-				SelectSequencerNodeInSection(ParamSection, ControlElement->GetName(), bSelected);
+			UMovieSceneControlRigParameterSection* ParamSection = Cast<UMovieSceneControlRigParameterSection>(Section);
+			SelectSequencerNodeInSection(ParamSection, ControlElement->GetName(), bSelected);
 			
 			GetSequencer()->ResumeSelectionBroadcast();
 
@@ -3041,7 +3041,7 @@ void FControlRigParameterTrackEditor::BuildTrackContextMenu(FMenuBuilder& MenuBu
 				LOCTEXT("SelectBonesToAnimateToolTip", "Select which bones or curves you want to directly animate"),
 				FSlateIcon(),
 				FUIAction(
-					FExecuteAction::CreateRaw(this, &FControlRigParameterTrackEditor::SelectFKBonesToAnimate, AutoRig,Track)));
+					FExecuteAction::CreateRaw(this, &FControlRigParameterTrackEditor::SelectFKBonesToAnimate, AutoRig, Track)));
 
 			MenuBuilder.AddMenuEntry(
 				LOCTEXT("FKRigApplyMode", "Additive"),
@@ -3158,7 +3158,7 @@ public:
 
 	SLATE_BEGIN_ARGS(SFKControlRigBoneSelect) {}
 	SLATE_ATTRIBUTE(UFKControlRig*, AutoRig)
-	SLATE_ATTRIBUTE(UMovieSceneControlRigParameterTrack*,Track)
+		SLATE_ATTRIBUTE(UMovieSceneControlRigParameterTrack*, Track)
 		SLATE_ATTRIBUTE(ISequencer*, Sequencer)
 		SLATE_END_ARGS()
 
@@ -3515,6 +3515,9 @@ bool FControlRigParameterTrackEditor::CollapseAllLayers(TSharedPtr<ISequencer>&
 			}
 		}
 
+		//remove all keys, except Space Channels, from the Section.
+		ParameterSection->RemoveAllKeys(false /*bIncludedSpaceKeys*/);
+
 		FRigControlModifiedContext Context;
 		Context.SetKey = EControlRigSetKey::Always;
 
