Sequencer: Tools won't work in subsequences due to wrong API call.

diff --git a/Engine/Plugins/Animation/ControlRig/Source/ControlRigEditor/Private/ControlRigBlueprintActions.cpp b/Engine/Plugins/Animation/ControlRig/Source/ControlRigEditor/Private/ControlRigBlueprintActions.cpp
--- a/Engine/Plugins/Animation/ControlRig/Source/ControlRigEditor/Private/ControlRigBlueprintActions.cpp
+++ b/Engine/Plugins/Animation/ControlRig/Source/ControlRigEditor/Private/ControlRigBlueprintActions.cpp
@@ -291,7 +291,7 @@ void FControlRigBlueprintActions::OnSpawnedSkeletalMeshActorChanged(UObject* InO
 	UClass* ControlRigClass = RigBlueprint->GeneratedClass;
 
 	// find a level sequence in the world, if can't find that, create one
-	ULevelSequence* Sequence = ULevelSequenceEditorBlueprintLibrary::GetCurrentLevelSequence();
+	ULevelSequence* Sequence = ULevelSequenceEditorBlueprintLibrary::GetFocusedLevelSequence();
 	if (Sequence == nullptr)
 	{
 		ALevelSequenceActor* LevelSequenceActor = nullptr;
diff --git a/Engine/Plugins/Animation/ControlRig/Source/ControlRigEditor/Private/EditMode/SControlRigBaseListWidget.cpp b/Engine/Plugins/Animation/ControlRig/Source/ControlRigEditor/Private/EditMode/SControlRigBaseListWidget.cpp
--- a/Engine/Plugins/Animation/ControlRig/Source/ControlRigEditor/Private/EditMode/SControlRigBaseListWidget.cpp
+++ b/Engine/Plugins/Animation/ControlRig/Source/ControlRigEditor/Private/EditMode/SControlRigBaseListWidget.cpp
@@ -1384,7 +1384,7 @@ void SControlRigBaseListWidget::ExecuteSelectControls(UControlRigPoseAsset* Pose
 	}
 	else
 	{
-		ULevelSequence *LevelSequence = ULevelSequenceEditorBlueprintLibrary::GetCurrentLevelSequence();
+		ULevelSequence *LevelSequence = ULevelSequenceEditorBlueprintLibrary::GetFocusedLevelSequence();
 		if (LevelSequence)
 		{
 			const FScopedTransaction Transaction(LOCTEXT("SelectControls", "Select Controls"));
diff --git a/Engine/Plugins/Animation/ControlRig/Source/ControlRigEditor/Private/EditMode/SControlRigTweenWidget.cpp b/Engine/Plugins/Animation/ControlRig/Source/ControlRigEditor/Private/EditMode/SControlRigTweenWidget.cpp
--- a/Engine/Plugins/Animation/ControlRig/Source/ControlRigEditor/Private/EditMode/SControlRigTweenWidget.cpp
+++ b/Engine/Plugins/Animation/ControlRig/Source/ControlRigEditor/Private/EditMode/SControlRigTweenWidget.cpp
@@ -96,7 +96,7 @@ void SControlRigTweenWidget::OnBeginSliderMovement()
 
 void SControlRigTweenWidget::SetupControls()
 {
-	ULevelSequence* LevelSequence = ULevelSequenceEditorBlueprintLibrary::GetCurrentLevelSequence();
+	ULevelSequence* LevelSequence = ULevelSequenceEditorBlueprintLibrary::GetFocusedLevelSequence();
 	IAssetEditorInstance* AssetEditor = GEditor->GetEditorSubsystem<UAssetEditorSubsystem>()->FindEditorForAsset(LevelSequence, false);
 	ILevelSequenceEditorToolkit* LevelSequenceEditor = static_cast<ILevelSequenceEditorToolkit*>(AssetEditor);
 	WeakSequencer = LevelSequenceEditor ? LevelSequenceEditor->GetSequencer() : nullptr;
diff --git a/Engine/Plugins/Animation/ControlRig/Source/ControlRigEditor/Private/Tools/ControlRigSnapper.cpp b/Engine/Plugins/Animation/ControlRig/Source/ControlRigEditor/Private/Tools/ControlRigSnapper.cpp
--- a/Engine/Plugins/Animation/ControlRig/Source/ControlRigEditor/Private/Tools/ControlRigSnapper.cpp
+++ b/Engine/Plugins/Animation/ControlRig/Source/ControlRigEditor/Private/Tools/ControlRigSnapper.cpp
@@ -101,7 +101,7 @@ int32 FControlRigSnapperSelection::NumSelected() const
 TWeakPtr<ISequencer> FControlRigSnapper::GetSequencer()
 {
 	TWeakPtr<ISequencer> WeakSequencer = nullptr;
-	ULevelSequence* LevelSequence = ULevelSequenceEditorBlueprintLibrary::GetCurrentLevelSequence();
+	ULevelSequence* LevelSequence = ULevelSequenceEditorBlueprintLibrary::GetFocusedLevelSequence();
 	if (LevelSequence)
 	{
 		IAssetEditorInstance* AssetEditor = GEditor->GetEditorSubsystem<UAssetEditorSubsystem>()->FindEditorForAsset(LevelSequence, false);
diff --git a/Engine/Plugins/Editor/SequencerAnimTools/Source/SequencerAnimTools/Private/SequencerAnimEditPivotTool.cpp b/Engine/Plugins/Editor/SequencerAnimTools/Source/SequencerAnimTools/Private/SequencerAnimEditPivotTool.cpp
--- a/Engine/Plugins/Editor/SequencerAnimTools/Source/SequencerAnimTools/Private/SequencerAnimEditPivotTool.cpp
+++ b/Engine/Plugins/Editor/SequencerAnimTools/Source/SequencerAnimTools/Private/SequencerAnimEditPivotTool.cpp
@@ -65,7 +65,7 @@ FLastSelectedObjects USequencerPivotTool::LastSelectedObjects;
 
 static void GetControlRigsAndSequencer(TArray<TWeakObjectPtr<UControlRig>>& ControlRigs, TWeakPtr<ISequencer>& SequencerPtr, ULevelSequence** LevelSequence)
 {
-	*LevelSequence = ULevelSequenceEditorBlueprintLibrary::GetCurrentLevelSequence();
+	*LevelSequence = ULevelSequenceEditorBlueprintLibrary::GetFocusedLevelSequence();
 	if (*LevelSequence)
 	{
 		IAssetEditorInstance* AssetEditor = GEditor->GetEditorSubsystem<UAssetEditorSubsystem>()->FindEditorForAsset(*LevelSequence, false);
