Fix EControlRigRotationOrder redirector having wrong format

diff --git a/Engine/Plugins/Animation/ControlRig/Config/BaseControlRig.ini b/Engine/Plugins/Animation/ControlRig/Config/BaseControlRig.ini
--- a/Engine/Plugins/Animation/ControlRig/Config/BaseControlRig.ini
+++ b/Engine/Plugins/Animation/ControlRig/Config/BaseControlRig.ini
@@ -25,7 +25,7 @@
 +ClassRedirects=(OldName="RigVMGraph",NewName="RigVMGraph",ValueChanges=(("Model","RigVMModel")))
 
 +EnumRedirects=(OldName="ERigElementType",ValueChanges=(("Space","Null"))
-+EnumRedirects=(OldName="EControlRigRotationOrder",NewName="/Script/AnimationCore.EEulerRotationOrder")
++EnumRedirects=(OldName="/Script/ControlRig.EControlRigRotationOrder",NewName="/Script/AnimationCore.EEulerRotationOrder")
 
 
 ; renaming gizmo to shape
