Control Rig: Fix loading of enums

diff --git a/Engine/Plugins/Animation/ControlRig/Source/ControlRig/Private/Rigs/RigHierarchyElements.cpp b/Engine/Plugins/Animation/ControlRig/Source/ControlRig/Private/Rigs/RigHierarchyElements.cpp
--- a/Engine/Plugins/Animation/ControlRig/Source/ControlRig/Private/Rigs/RigHierarchyElements.cpp
+++ b/Engine/Plugins/Animation/ControlRig/Source/ControlRig/Private/Rigs/RigHierarchyElements.cpp
@@ -539,7 +539,7 @@ void FRigControlSettings::Load(FArchive& Ar)
 	ControlEnum = nullptr;
 	if(!ControlEnumPathName.IsEmpty())
 	{
-		ControlEnum = FindObject<UEnum>(ANY_PACKAGE, *ControlEnumPathName);
+		ControlEnum = LoadObject<UEnum>(nullptr, *ControlEnumPathName);
 	}
 
 	if (Ar.CustomVer(FControlRigObjectVersion::GUID) >= FControlRigObjectVersion::RigHierarchyControlSpaceFavorites)
