Control Rig: Crash: Don't record transform changes within component

diff --git a/Engine/Plugins/Animation/ControlRig/Source/ControlRig/Private/ControlRig.cpp b/Engine/Plugins/Animation/ControlRig/Source/ControlRig/Private/ControlRig.cpp
--- a/Engine/Plugins/Animation/ControlRig/Source/ControlRig/Private/ControlRig.cpp
+++ b/Engine/Plugins/Animation/ControlRig/Source/ControlRig/Private/ControlRig.cpp
@@ -26,6 +26,7 @@
 #include "Kismet2/BlueprintEditorUtils.h"
 #include "RigVMDeveloperTypeUtils.h"
 #endif// WITH_EDITOR
+#include "ControlRigComponent.h"
 
 #define LOCTEXT_NAMESPACE "ControlRig"
 
@@ -1291,6 +1292,16 @@ void UControlRig::ExecuteUnits(FRigUnitContext& InOutContext, const FName& InEve
 			URigHierarchy* Hierarchy = GetHierarchy();
 #if WITH_EDITOR
 
+			bool bRecordTransformsAtRuntime = true;
+			if(const UObject* Outer = GetOuter())
+			{
+				if(Outer->IsA<UControlRigComponent>())
+				{
+					bRecordTransformsAtRuntime = false;
+				}
+			}
+			TGuardValue<bool> RecordTransformsPerInstructionGuard(Hierarchy->bRecordTransformsAtRuntime, bRecordTransformsAtRuntime);
+			
 			if(Hierarchy->bRecordTransformsAtRuntime)
 			{
 				Hierarchy->ReadTransformsAtRuntime.Reset();
