Control Rig: Ensure to maintain hierarchy during bake for native classes

diff --git a/Engine/Plugins/Animation/ControlRig/Source/ControlRig/Private/Sequencer/MovieSceneControlRigParameterSection.cpp b/Engine/Plugins/Animation/ControlRig/Source/ControlRig/Private/Sequencer/MovieSceneControlRigParameterSection.cpp
--- a/Engine/Plugins/Animation/ControlRig/Source/ControlRig/Private/Sequencer/MovieSceneControlRigParameterSection.cpp
+++ b/Engine/Plugins/Animation/ControlRig/Source/ControlRig/Private/Sequencer/MovieSceneControlRigParameterSection.cpp
@@ -2634,7 +2634,7 @@ bool UMovieSceneControlRigParameterSection::LoadAnimSequenceIntoThisSection(UAni
 
 	// copy the hierarchy from the CDO into the target control rig.
 	// this ensures that the topology version matches in case of a dynamic hierarchy
-	if(!ControlRig->IsNative())
+	if(!ControlRig->GetClass()->IsNative())
 	{
 		if (UControlRig* CDO = Cast<UControlRig>(ControlRig->GetClass()->GetDefaultObject()))
 		{
