Control Rig: Provider python getters for shape transform

diff --git a/Engine/Plugins/Animation/ControlRig/Source/ControlRig/Public/Rigs/RigHierarchy.h b/Engine/Plugins/Animation/ControlRig/Source/ControlRig/Public/Rigs/RigHierarchy.h
--- a/Engine/Plugins/Animation/ControlRig/Source/ControlRig/Public/Rigs/RigHierarchy.h
+++ b/Engine/Plugins/Animation/ControlRig/Source/ControlRig/Public/Rigs/RigHierarchy.h
@@ -1076,6 +1076,37 @@ public:
 		return FTransform::Identity;
 	}
 
+	/**
+	 * Returns the local shape transform for a given control element.
+	 * @param InKey The key of the control to retrieve the transform for
+	 * @param bInitial If true the initial transform will be used
+	 * @return The local shape transform
+	 */
+	UFUNCTION(BlueprintCallable, Category = URigHierarchy)
+	FORCEINLINE_DEBUGGABLE FTransform GetLocalControlShapeTransform(FRigElementKey InKey, bool bInitial = false) const
+	{
+		return GetLocalControlShapeTransformByIndex(GetIndex(InKey), bInitial);
+	}
+
+	/**
+	 * Returns the local shape transform for a given control element.
+	 * @param InElementIndex The index of the control to retrieve the transform for
+	 * @param bInitial If true the initial transform will be used
+	 * @return The local shape transform
+	 */
+	UFUNCTION(BlueprintCallable, Category = URigHierarchy)
+	FORCEINLINE_DEBUGGABLE FTransform GetLocalControlShapeTransformByIndex(int32 InElementIndex, bool bInitial = false) const
+	{
+		if(Elements.IsValidIndex(InElementIndex))
+		{
+			if(FRigControlElement* ControlElement = Cast<FRigControlElement>(Elements[InElementIndex]))
+			{
+				return GetControlShapeTransform(ControlElement, bInitial ? ERigTransformType::InitialLocal : ERigTransformType::CurrentLocal);
+			}
+		}
+		return FTransform::Identity;
+	}
+
 	/**
 	 * Returns the global shape transform for a given control element.
 	 * @param InKey The key of the control to retrieve the transform for
