Sequencer: Can't select channels and widget selection after changing selection in outliner or sequencer.

diff --git a/Engine/Plugins/Animation/ControlRig/Source/ControlRigEditor/Private/EditMode/ControlRigEditMode.cpp b/Engine/Plugins/Animation/ControlRig/Source/ControlRigEditor/Private/EditMode/ControlRigEditMode.cpp
--- a/Engine/Plugins/Animation/ControlRig/Source/ControlRigEditor/Private/EditMode/ControlRigEditMode.cpp
+++ b/Engine/Plugins/Animation/ControlRig/Source/ControlRigEditor/Private/EditMode/ControlRigEditMode.cpp
@@ -2116,6 +2116,12 @@ void FControlRigEditMode::HandleSelectionChanged()
 
 	// update the pivot transform of our selected objects (they could be animating)
 	RecalcPivotTransform();
+	//need to force the redraw also
+	if (IsInLevelEditor())
+	{
+		GEditor->RedrawLevelEditingViewports(true);
+	}
+
 }
 
 void FControlRigEditMode::BindCommands()
diff --git a/Engine/Plugins/Animation/ControlRig/Source/ControlRigEditor/Private/Sequencer/ControlRigParameterTrackEditor.cpp b/Engine/Plugins/Animation/ControlRig/Source/ControlRigEditor/Private/Sequencer/ControlRigParameterTrackEditor.cpp
--- a/Engine/Plugins/Animation/ControlRig/Source/ControlRigEditor/Private/Sequencer/ControlRigParameterTrackEditor.cpp
+++ b/Engine/Plugins/Animation/ControlRig/Source/ControlRigEditor/Private/Sequencer/ControlRigParameterTrackEditor.cpp
@@ -1964,7 +1964,6 @@ void FControlRigParameterTrackEditor::OnSelectionChanged(TArray<UMovieSceneTrack
 
 	SelectRigsAndControls(ControlRig, KeyAreas);
 	
-	GEditor->NoteSelectionChange();
 }
 
 void FControlRigParameterTrackEditor::SelectRigsAndControls(UControlRig* ControlRig, const TArray<const IKeyArea*>& KeyAreas)
