Control Rig: Allow user specified mappings on ControlRigComponent

diff --git a/Engine/Plugins/Animation/ControlRig/Source/ControlRig/Private/ControlRigComponent.cpp b/Engine/Plugins/Animation/ControlRig/Source/ControlRig/Private/ControlRigComponent.cpp
--- a/Engine/Plugins/Animation/ControlRig/Source/ControlRig/Private/ControlRigComponent.cpp
+++ b/Engine/Plugins/Animation/ControlRig/Source/ControlRig/Private/ControlRigComponent.cpp
@@ -461,6 +461,7 @@ void UControlRigComponent::ClearMappedElements()
 	}
 
 	MappedElements.Reset();
+	MappedElements = UserDefinedElements;
 	ValidateMappingData();
 	Initialize();
 }
@@ -1195,6 +1196,18 @@ void UControlRigComponent::ValidateMappingData()
 			AActor* MappedOwner = MappedElement.ComponentReference.OtherActor == nullptr ? GetOwner() : ToRawPtr(MappedElement.ComponentReference.OtherActor);
 			MappedElement.SceneComponent = Cast<USceneComponent>(MappedElement.ComponentReference.GetComponent(MappedOwner));
 
+			// try again with the path to the component
+			if (MappedElement.SceneComponent == nullptr)
+			{
+				FComponentReference TempReference;
+				TempReference.PathToComponent = MappedElement.ComponentReference.ComponentProperty.ToString();
+				if (USceneComponent* TempSceneComponent = Cast<USceneComponent>(TempReference.GetComponent(MappedOwner)))
+				{
+					MappedElement.ComponentReference = TempReference;
+					MappedElement.SceneComponent = TempSceneComponent;
+				}
+			}
+
 			if (MappedElement.SceneComponent == nullptr ||
 				MappedElement.SceneComponent == this ||
 				MappedElement.ElementName.IsNone())
diff --git a/Engine/Plugins/Animation/ControlRig/Source/ControlRig/Public/ControlRigComponent.h b/Engine/Plugins/Animation/ControlRig/Source/ControlRig/Public/ControlRigComponent.h
--- a/Engine/Plugins/Animation/ControlRig/Source/ControlRig/Public/ControlRigComponent.h
+++ b/Engine/Plugins/Animation/ControlRig/Source/ControlRig/Public/ControlRigComponent.h
@@ -267,6 +267,9 @@ public:
 	void Update(float DeltaTime = 0.f);
 
 	UPROPERTY(EditAnywhere, Category = "Animation")
+	TArray<FControlRigComponentMappedElement> UserDefinedElements;
+
+	UPROPERTY(VisibleAnywhere, Category = "Animation")
 	TArray<FControlRigComponentMappedElement> MappedElements;
 
 	/** Removes all mapped elements from the component */
