Sequencer: When using tools in subsequences need to get ISequencer from the master not the current.

diff --git a/Engine/Plugins/Animation/ControlRig/Source/ControlRigEditor/Private/ControlRigSequencerEditorLibrary.cpp b/Engine/Plugins/Animation/ControlRig/Source/ControlRigEditor/Private/ControlRigSequencerEditorLibrary.cpp
--- a/Engine/Plugins/Animation/ControlRig/Source/ControlRigEditor/Private/ControlRigSequencerEditorLibrary.cpp
+++ b/Engine/Plugins/Animation/ControlRig/Source/ControlRigEditor/Private/ControlRigSequencerEditorLibrary.cpp
@@ -36,6 +36,7 @@
 #include "ScopedTransaction.h"
 #include "ControlRigParameterTrackEditor.h"
 #include "ControlRigSpaceChannelEditors.h"
+#include "LevelSequenceEditorBlueprintLibrary.h"
 
 #define LOCTEXT_NAMESPACE "ControlrigSequencerEditorLibrary"
 
@@ -173,8 +174,10 @@ static void AcquireSkeletonAndSkelMeshCompFromObject(UObject* BoundObject, USkel
 	}
 }
 
-static TSharedPtr<ISequencer> GetSequencerFromAsset(ULevelSequence* LevelSequence)
+static TSharedPtr<ISequencer> GetSequencerFromAsset()
 {
+	//if getting sequencer from level sequence need to use the current(master), not the focused
+	ULevelSequence* LevelSequence = ULevelSequenceEditorBlueprintLibrary::GetCurrentLevelSequence();
 	IAssetEditorInstance* AssetEditor = GEditor->GetEditorSubsystem<UAssetEditorSubsystem>()->FindEditorForAsset(LevelSequence, false);
 	ILevelSequenceEditorToolkit* LevelSequenceEditor = static_cast<ILevelSequenceEditorToolkit*>(AssetEditor);
 	TSharedPtr<ISequencer> Sequencer = LevelSequenceEditor ? LevelSequenceEditor->GetSequencer() : nullptr;
@@ -192,7 +195,7 @@ static UMovieSceneControlRigParameterTrack* AddControlRig(ULevelSequence* LevelS
 	if (InClass && InClass->IsChildOf(UControlRig::StaticClass()) && LevelSequence && LevelSequence->GetMovieScene())
 	{
 		UMovieScene* OwnerMovieScene = LevelSequence->GetMovieScene();
-		TSharedPtr<ISequencer> SharedSequencer = GetSequencerFromAsset(LevelSequence);
+		TSharedPtr<ISequencer> SharedSequencer = GetSequencerFromAsset();
 		ISequencer* Sequencer = nullptr; // will be valid  if we have a ISequencer AND it's focused.
 		if (SharedSequencer.IsValid() && SharedSequencer->GetFocusedMovieSceneSequence() == LevelSequence)
 		{
@@ -383,7 +386,7 @@ TArray<UMovieSceneTrack*> UControlRigSequencerEditorLibrary::FindOrCreateControl
 
 bool UControlRigSequencerEditorLibrary::TweenControlRig(ULevelSequence* LevelSequence, UControlRig* ControlRig, float TweenValue)
 {
-	TWeakPtr<ISequencer> WeakSequencer = GetSequencerFromAsset(LevelSequence);
+	TWeakPtr<ISequencer> WeakSequencer = GetSequencerFromAsset();
 	if (WeakSequencer.IsValid() && WeakSequencer.Pin()->GetFocusedMovieSceneSequence() == LevelSequence
 		&& ControlRig && LevelSequence->GetMovieScene())
 	{
@@ -444,7 +447,7 @@ static void ConvertFramesToTickResolution(UMovieScene* MovieScene, const TArray<
 
 TArray<FTransform> UControlRigSequencerEditorLibrary::GetActorWorldTransforms(ULevelSequence* LevelSequence,AActor* Actor, const TArray<FFrameNumber>& InFrames, ESequenceTimeUnit TimeUnit)
 {
-	TWeakPtr<ISequencer> WeakSequencer = GetSequencerFromAsset(LevelSequence);
+	TWeakPtr<ISequencer> WeakSequencer = GetSequencerFromAsset();
 	TArray<FTransform> OutWorldTransforms;
 	if (WeakSequencer.IsValid() && Actor)
 	{
@@ -474,7 +477,7 @@ FTransform UControlRigSequencerEditorLibrary::GetSkeletalMeshComponentWorldTrans
 TArray<FTransform> UControlRigSequencerEditorLibrary::GetSkeletalMeshComponentWorldTransforms(ULevelSequence* LevelSequence,USkeletalMeshComponent* SkeletalMeshComponent, const TArray<FFrameNumber>& InFrames,
 	ESequenceTimeUnit TimeUnit, FName SocketName)
 {
-	TWeakPtr<ISequencer> WeakSequencer = GetSequencerFromAsset(LevelSequence);
+	TWeakPtr<ISequencer> WeakSequencer = GetSequencerFromAsset();
 	TArray<FTransform> OutWorldTransforms;
 	if (WeakSequencer.IsValid() && SkeletalMeshComponent)
 	{
@@ -509,7 +512,7 @@ FTransform UControlRigSequencerEditorLibrary::GetControlRigWorldTransform(ULevel
 TArray<FTransform> UControlRigSequencerEditorLibrary::GetControlRigWorldTransforms(ULevelSequence* LevelSequence,UControlRig* ControlRig, FName ControlName, const TArray<FFrameNumber>& InFrames,
 	ESequenceTimeUnit TimeUnit)
 {
-	TWeakPtr<ISequencer> WeakSequencer = GetSequencerFromAsset(LevelSequence);
+	TWeakPtr<ISequencer> WeakSequencer = GetSequencerFromAsset();
 	TArray<FTransform> OutWorldTransforms;
 	if (WeakSequencer.IsValid() && ControlRig)
 	{
@@ -541,7 +544,7 @@ TArray<FTransform> UControlRigSequencerEditorLibrary::GetControlRigWorldTransfor
 static void LocalSetControlRigWorldTransforms(ULevelSequence* LevelSequence,UControlRig* ControlRig, FName ControlName, EControlRigSetKey SetKey, const TArray<FFrameNumber>& InFrames, 
 	const TArray<FTransform>& WorldTransforms, ESequenceTimeUnit TimeUnit)
 {
-	TWeakPtr<ISequencer> WeakSequencer = GetSequencerFromAsset(LevelSequence);
+	TWeakPtr<ISequencer> WeakSequencer = GetSequencerFromAsset();
 	if (WeakSequencer.IsValid() && ControlRig)
 	{
 		if (TSharedPtr<IControlRigObjectBinding> ObjectBinding = ControlRig->GetObjectBinding())
@@ -621,7 +624,7 @@ bool UControlRigSequencerEditorLibrary::BakeToControlRig(UWorld* World, ULevelSe
 		return false;
 	}
 	//get level sequence if one exists...
-	TWeakPtr<ISequencer> WeakSequencer = GetSequencerFromAsset(LevelSequence);
+	TWeakPtr<ISequencer> WeakSequencer = GetSequencerFromAsset();
 
 	ALevelSequenceActor* OutActor = nullptr;
 	FMovieSceneSequencePlaybackSettings Settings;
@@ -934,7 +937,7 @@ static bool GetControlRigValues(UWorld* World, ULevelSequence* LevelSequence, UC
 float UControlRigSequencerEditorLibrary::GetLocalControlRigFloat(ULevelSequence* LevelSequence,UControlRig* ControlRig, FName ControlName, FFrameNumber Frame, ESequenceTimeUnit TimeUnit)
 {
 	float Value = 0.0f;
-	TWeakPtr<ISequencer> WeakSequencer = GetSequencerFromAsset(LevelSequence);
+	TWeakPtr<ISequencer> WeakSequencer = GetSequencerFromAsset();
 
 	if (WeakSequencer.IsValid() && ControlRig)
 	{
@@ -958,7 +961,7 @@ float UControlRigSequencerEditorLibrary::GetLocalControlRigFloat(ULevelSequence*
 TArray<float> UControlRigSequencerEditorLibrary::GetLocalControlRigFloats(ULevelSequence* LevelSequence,UControlRig* ControlRig, FName ControlName, const TArray<FFrameNumber>& Frames, ESequenceTimeUnit TimeUnit)
 {
 	TArray<float> Values;
-	TWeakPtr<ISequencer> WeakSequencer = GetSequencerFromAsset(LevelSequence);
+	TWeakPtr<ISequencer> WeakSequencer = GetSequencerFromAsset();
 
 	if (WeakSequencer.IsValid() && ControlRig)
 	{
@@ -1051,7 +1054,7 @@ void UControlRigSequencerEditorLibrary::SetLocalControlRigFloats(ULevelSequence*
 bool UControlRigSequencerEditorLibrary::GetLocalControlRigBool(ULevelSequence* LevelSequence, UControlRig* ControlRig, FName ControlName, FFrameNumber Frame, ESequenceTimeUnit TimeUnit)
 {
 	bool Value = true;
-	TWeakPtr<ISequencer> WeakSequencer = GetSequencerFromAsset(LevelSequence);
+	TWeakPtr<ISequencer> WeakSequencer = GetSequencerFromAsset();
 
 	if (WeakSequencer.IsValid() && ControlRig)
 	{
@@ -1075,7 +1078,7 @@ bool UControlRigSequencerEditorLibrary::GetLocalControlRigBool(ULevelSequence* L
 TArray<bool> UControlRigSequencerEditorLibrary::GetLocalControlRigBools(ULevelSequence* LevelSequence, UControlRig* ControlRig, FName ControlName, const TArray<FFrameNumber>& Frames, ESequenceTimeUnit TimeUnit)
 {
 	TArray<bool> Values;
-	TWeakPtr<ISequencer> WeakSequencer = GetSequencerFromAsset(LevelSequence);
+	TWeakPtr<ISequencer> WeakSequencer = GetSequencerFromAsset();
 
 	if (WeakSequencer.IsValid() && ControlRig)
 	{
@@ -1166,7 +1169,7 @@ void UControlRigSequencerEditorLibrary::SetLocalControlRigBools(ULevelSequence*
 int32 UControlRigSequencerEditorLibrary::GetLocalControlRigInt(ULevelSequence* LevelSequence, UControlRig* ControlRig, FName ControlName, FFrameNumber Frame, ESequenceTimeUnit TimeUnit)
 {
 	int32 Value = 0;
-	TWeakPtr<ISequencer> WeakSequencer = GetSequencerFromAsset(LevelSequence);
+	TWeakPtr<ISequencer> WeakSequencer = GetSequencerFromAsset();
 
 	if (WeakSequencer.IsValid() && ControlRig)
 	{
@@ -1190,7 +1193,7 @@ int32 UControlRigSequencerEditorLibrary::GetLocalControlRigInt(ULevelSequence* L
 TArray<int32> UControlRigSequencerEditorLibrary::GetLocalControlRigInts(ULevelSequence* LevelSequence, UControlRig* ControlRig, FName ControlName, const TArray<FFrameNumber>& Frames, ESequenceTimeUnit TimeUnit)
 {
 	TArray<int32> Values;
-	TWeakPtr<ISequencer> WeakSequencer = GetSequencerFromAsset(LevelSequence);
+	TWeakPtr<ISequencer> WeakSequencer = GetSequencerFromAsset();
 
 	if (WeakSequencer.IsValid() && ControlRig)
 	{
@@ -1282,7 +1285,7 @@ void UControlRigSequencerEditorLibrary::SetLocalControlRigInts(ULevelSequence* L
 FVector2D UControlRigSequencerEditorLibrary::GetLocalControlRigVector2D(ULevelSequence* LevelSequence, UControlRig* ControlRig, FName ControlName, FFrameNumber Frame, ESequenceTimeUnit TimeUnit)
 {
 	FVector2D Value;
-	TWeakPtr<ISequencer> WeakSequencer = GetSequencerFromAsset(LevelSequence);
+	TWeakPtr<ISequencer> WeakSequencer = GetSequencerFromAsset();
 
 	if (WeakSequencer.IsValid() && ControlRig)
 	{
@@ -1307,7 +1310,7 @@ FVector2D UControlRigSequencerEditorLibrary::GetLocalControlRigVector2D(ULevelSe
 TArray<FVector2D> UControlRigSequencerEditorLibrary::GetLocalControlRigVector2Ds(ULevelSequence* LevelSequence, UControlRig* ControlRig, FName ControlName, const TArray<FFrameNumber>& Frames, ESequenceTimeUnit TimeUnit)
 {
 	TArray<FVector2D> Values;
-	TWeakPtr<ISequencer> WeakSequencer = GetSequencerFromAsset(LevelSequence);
+	TWeakPtr<ISequencer> WeakSequencer = GetSequencerFromAsset();
 
 	if (WeakSequencer.IsValid() && ControlRig)
 	{
@@ -1400,7 +1403,7 @@ void UControlRigSequencerEditorLibrary::SetLocalControlRigVector2Ds(ULevelSequen
 FVector UControlRigSequencerEditorLibrary::GetLocalControlRigPosition(ULevelSequence* LevelSequence, UControlRig* ControlRig, FName ControlName, FFrameNumber Frame, ESequenceTimeUnit TimeUnit)
 {
 	FVector Value;
-	TWeakPtr<ISequencer> WeakSequencer = GetSequencerFromAsset(LevelSequence);
+	TWeakPtr<ISequencer> WeakSequencer = GetSequencerFromAsset();
 
 	if (WeakSequencer.IsValid() && ControlRig)
 	{
@@ -1424,7 +1427,7 @@ FVector UControlRigSequencerEditorLibrary::GetLocalControlRigPosition(ULevelSequ
 TArray<FVector> UControlRigSequencerEditorLibrary::GetLocalControlRigPositions(ULevelSequence* LevelSequence, UControlRig* ControlRig, FName ControlName, const TArray<FFrameNumber>& Frames, ESequenceTimeUnit TimeUnit)
 {
 	TArray<FVector> Values;
-	TWeakPtr<ISequencer> WeakSequencer = GetSequencerFromAsset(LevelSequence);
+	TWeakPtr<ISequencer> WeakSequencer = GetSequencerFromAsset();
 
 	if (WeakSequencer.IsValid() && ControlRig)
 	{
@@ -1516,7 +1519,7 @@ void UControlRigSequencerEditorLibrary::SetLocalControlRigPositions(ULevelSequen
 FRotator UControlRigSequencerEditorLibrary::GetLocalControlRigRotator(ULevelSequence* LevelSequence, UControlRig* ControlRig, FName ControlName, FFrameNumber Frame, ESequenceTimeUnit TimeUnit)
 {
 	FRotator Value;
-	TWeakPtr<ISequencer> WeakSequencer = GetSequencerFromAsset(LevelSequence);
+	TWeakPtr<ISequencer> WeakSequencer = GetSequencerFromAsset();
 
 	if (WeakSequencer.IsValid() && ControlRig)
 	{
@@ -1540,7 +1543,7 @@ FRotator UControlRigSequencerEditorLibrary::GetLocalControlRigRotator(ULevelSequ
 TArray<FRotator> UControlRigSequencerEditorLibrary::GetLocalControlRigRotators(ULevelSequence* LevelSequence, UControlRig* ControlRig, FName ControlName, const TArray<FFrameNumber>& Frames, ESequenceTimeUnit TimeUnit)
 {
 	TArray<FRotator> Values;
-	TWeakPtr<ISequencer> WeakSequencer = GetSequencerFromAsset(LevelSequence);
+	TWeakPtr<ISequencer> WeakSequencer = GetSequencerFromAsset();
 
 	if (WeakSequencer.IsValid() && ControlRig)
 	{
@@ -1632,7 +1635,7 @@ void UControlRigSequencerEditorLibrary::SetLocalControlRigRotators(ULevelSequenc
 FVector UControlRigSequencerEditorLibrary::GetLocalControlRigScale(ULevelSequence* LevelSequence, UControlRig* ControlRig, FName ControlName, FFrameNumber Frame,  ESequenceTimeUnit TimeUnit)
 {
 	FVector Value;
-	TWeakPtr<ISequencer> WeakSequencer = GetSequencerFromAsset(LevelSequence);
+	TWeakPtr<ISequencer> WeakSequencer = GetSequencerFromAsset();
 
 	if (WeakSequencer.IsValid() && ControlRig)
 	{
@@ -1656,7 +1659,7 @@ FVector UControlRigSequencerEditorLibrary::GetLocalControlRigScale(ULevelSequenc
 TArray<FVector>UControlRigSequencerEditorLibrary::GetLocalControlRigScales(ULevelSequence* LevelSequence, UControlRig* ControlRig, FName ControlName, const TArray<FFrameNumber>& Frames, ESequenceTimeUnit TimeUnit)
 {
 	TArray<FVector> Values;
-	TWeakPtr<ISequencer> WeakSequencer = GetSequencerFromAsset(LevelSequence);
+	TWeakPtr<ISequencer> WeakSequencer = GetSequencerFromAsset();
 
 	if (WeakSequencer.IsValid() && ControlRig)
 	{
@@ -1748,7 +1751,7 @@ void UControlRigSequencerEditorLibrary::SetLocalControlRigScales(ULevelSequence*
 FEulerTransform UControlRigSequencerEditorLibrary::GetLocalControlRigEulerTransform(ULevelSequence* LevelSequence, UControlRig* ControlRig, FName ControlName, FFrameNumber Frame, ESequenceTimeUnit TimeUnit)
 {
 	FEulerTransform Value = FEulerTransform::Identity;
-	TWeakPtr<ISequencer> WeakSequencer = GetSequencerFromAsset(LevelSequence);
+	TWeakPtr<ISequencer> WeakSequencer = GetSequencerFromAsset();
 
 	if (WeakSequencer.IsValid() && ControlRig)
 	{
@@ -1772,7 +1775,7 @@ FEulerTransform UControlRigSequencerEditorLibrary::GetLocalControlRigEulerTransf
 TArray<FEulerTransform> UControlRigSequencerEditorLibrary::GetLocalControlRigEulerTransforms(ULevelSequence* LevelSequence, UControlRig* ControlRig, FName ControlName, const TArray<FFrameNumber>& Frames, ESequenceTimeUnit TimeUnit)
 {
 	TArray<FEulerTransform> Values;
-	TWeakPtr<ISequencer> WeakSequencer = GetSequencerFromAsset(LevelSequence);
+	TWeakPtr<ISequencer> WeakSequencer = GetSequencerFromAsset();
 
 	if (WeakSequencer.IsValid() && ControlRig)
 	{
@@ -1865,7 +1868,7 @@ void UControlRigSequencerEditorLibrary::SetLocalControlRigEulerTransforms(ULevel
 FTransformNoScale UControlRigSequencerEditorLibrary::GetLocalControlRigTransformNoScale(ULevelSequence* LevelSequence, UControlRig* ControlRig, FName ControlName, FFrameNumber Frame, ESequenceTimeUnit TimeUnit)
 {
 	FTransformNoScale Value = FTransformNoScale::Identity;
-	TWeakPtr<ISequencer> WeakSequencer = GetSequencerFromAsset(LevelSequence);
+	TWeakPtr<ISequencer> WeakSequencer = GetSequencerFromAsset();
 
 	if (WeakSequencer.IsValid() && ControlRig)
 	{
@@ -1889,7 +1892,7 @@ FTransformNoScale UControlRigSequencerEditorLibrary::GetLocalControlRigTransform
 TArray<FTransformNoScale> UControlRigSequencerEditorLibrary::GetLocalControlRigTransformNoScales(ULevelSequence* LevelSequence, UControlRig* ControlRig, FName ControlName, const TArray<FFrameNumber>& Frames, ESequenceTimeUnit TimeUnit)
 {
 	TArray<FTransformNoScale> Values;
-	TWeakPtr<ISequencer> WeakSequencer = GetSequencerFromAsset(LevelSequence);
+	TWeakPtr<ISequencer> WeakSequencer = GetSequencerFromAsset();
 
 	if (WeakSequencer.IsValid() && ControlRig)
 	{
@@ -1982,7 +1985,7 @@ void UControlRigSequencerEditorLibrary::SetLocalControlRigTransformNoScales(ULev
 FTransform UControlRigSequencerEditorLibrary::GetLocalControlRigTransform(ULevelSequence* LevelSequence, UControlRig* ControlRig, FName ControlName, FFrameNumber Frame, ESequenceTimeUnit TimeUnit)
 {
 	FTransform Value = FTransform::Identity;
-	TWeakPtr<ISequencer> WeakSequencer = GetSequencerFromAsset(LevelSequence);
+	TWeakPtr<ISequencer> WeakSequencer = GetSequencerFromAsset();
 	if (WeakSequencer.IsValid() && ControlRig)
 	{
 		if (FRigControlElement* Element = ControlRig->FindControl(ControlName))
@@ -2006,7 +2009,7 @@ FTransform UControlRigSequencerEditorLibrary::GetLocalControlRigTransform(ULevel
 TArray<FTransform> UControlRigSequencerEditorLibrary::GetLocalControlRigTransforms(ULevelSequence* LevelSequence, UControlRig* ControlRig, FName ControlName, const TArray<FFrameNumber>& Frames, ESequenceTimeUnit TimeUnit)
 {
 	TArray<FTransform> Values;
-	TWeakPtr<ISequencer> WeakSequencer = GetSequencerFromAsset(LevelSequence);
+	TWeakPtr<ISequencer> WeakSequencer = GetSequencerFromAsset();
 
 	if (WeakSequencer.IsValid() && ControlRig)
 	{
@@ -2138,7 +2141,7 @@ bool UControlRigSequencerEditorLibrary::ImportFBXToControlRigTrack(UWorld* World
 
 bool UControlRigSequencerEditorLibrary::CollapseControlRigAnimLayers(ULevelSequence* LevelSequence, UMovieSceneControlRigParameterTrack* InTrack, bool bKeyReduce, float Tolerance)
 {
-	TWeakPtr<ISequencer> WeakSequencer = GetSequencerFromAsset(LevelSequence);
+	TWeakPtr<ISequencer> WeakSequencer = GetSequencerFromAsset();
 	bool bValid = false;
 
 	if (WeakSequencer.IsValid() && InTrack)
@@ -2156,7 +2159,7 @@ bool UControlRigSequencerEditorLibrary::CollapseControlRigAnimLayers(ULevelSeque
 
 bool UControlRigSequencerEditorLibrary::SetControlRigSpace(ULevelSequence* LevelSequence, UControlRig* ControlRig, FName ControlName, const FRigElementKey& InSpaceKey, FFrameNumber InTime, ESequenceTimeUnit TimeUnit)
 {
-	TWeakPtr<ISequencer> WeakSequencer = GetSequencerFromAsset(LevelSequence);
+	TWeakPtr<ISequencer> WeakSequencer = GetSequencerFromAsset();
 	bool bValid = false;
 
 	if (WeakSequencer.IsValid() && ControlRig && ControlName != NAME_None)
@@ -2192,7 +2195,7 @@ bool UControlRigSequencerEditorLibrary::SetControlRigSpace(ULevelSequence* Level
 
 bool UControlRigSequencerEditorLibrary::BakeControlRigSpace(ULevelSequence* InSequence, UControlRig* InControlRig, const TArray<FName>& InControlNames, FRigSpacePickerBakeSettings InSettings, ESequenceTimeUnit TimeUnit)
 {
-	TWeakPtr<ISequencer> WeakSequencer = GetSequencerFromAsset(InSequence);
+	TWeakPtr<ISequencer> WeakSequencer = GetSequencerFromAsset();
 	bool bValid = false;
 
 	if (WeakSequencer.IsValid() && InControlRig && InControlNames.Num() > 0)
@@ -2232,7 +2235,7 @@ bool UControlRigSequencerEditorLibrary::BakeControlRigSpace(ULevelSequence* InSe
 
 bool UControlRigSequencerEditorLibrary::DeleteControlRigSpace(ULevelSequence* LevelSequence, UControlRig* ControlRig, FName ControlName, FFrameNumber InTime, ESequenceTimeUnit TimeUnit)
 {
-	TWeakPtr<ISequencer> WeakSequencer = GetSequencerFromAsset(LevelSequence);
+	TWeakPtr<ISequencer> WeakSequencer = GetSequencerFromAsset();
 	bool bValid = false;
 
 	if (WeakSequencer.IsValid() && ControlRig && ControlName != NAME_None)
@@ -2289,7 +2292,7 @@ bool UControlRigSequencerEditorLibrary::DeleteControlRigSpace(ULevelSequence* Le
 
 bool UControlRigSequencerEditorLibrary::MoveControlRigSpace(ULevelSequence* LevelSequence, UControlRig* ControlRig, FName ControlName, FFrameNumber InTime, FFrameNumber InNewTime, ESequenceTimeUnit TimeUnit)
 {
-	TWeakPtr<ISequencer> WeakSequencer = GetSequencerFromAsset(LevelSequence);
+	TWeakPtr<ISequencer> WeakSequencer = GetSequencerFromAsset();
 	bool bValid = false;
 
 	if (WeakSequencer.IsValid() && ControlRig && ControlName != NAME_None)
diff --git a/Engine/Plugins/Animation/ControlRig/Source/ControlRigEditor/Private/EditMode/SControlRigTweenWidget.cpp b/Engine/Plugins/Animation/ControlRig/Source/ControlRigEditor/Private/EditMode/SControlRigTweenWidget.cpp
--- a/Engine/Plugins/Animation/ControlRig/Source/ControlRigEditor/Private/EditMode/SControlRigTweenWidget.cpp
+++ b/Engine/Plugins/Animation/ControlRig/Source/ControlRigEditor/Private/EditMode/SControlRigTweenWidget.cpp
@@ -96,7 +96,8 @@ void SControlRigTweenWidget::OnBeginSliderMovement()
 
 void SControlRigTweenWidget::SetupControls()
 {
-	ULevelSequence* LevelSequence = ULevelSequenceEditorBlueprintLibrary::GetFocusedLevelSequence();
+	//if getting sequencer from level sequence need to use the current(master), not the focused
+	ULevelSequence* LevelSequence = ULevelSequenceEditorBlueprintLibrary::GetCurrentLevelSequence();
 	IAssetEditorInstance* AssetEditor = GEditor->GetEditorSubsystem<UAssetEditorSubsystem>()->FindEditorForAsset(LevelSequence, false);
 	ILevelSequenceEditorToolkit* LevelSequenceEditor = static_cast<ILevelSequenceEditorToolkit*>(AssetEditor);
 	WeakSequencer = LevelSequenceEditor ? LevelSequenceEditor->GetSequencer() : nullptr;
diff --git a/Engine/Plugins/Animation/ControlRig/Source/ControlRigEditor/Private/Tools/ControlRigSnapper.cpp b/Engine/Plugins/Animation/ControlRig/Source/ControlRigEditor/Private/Tools/ControlRigSnapper.cpp
--- a/Engine/Plugins/Animation/ControlRig/Source/ControlRigEditor/Private/Tools/ControlRigSnapper.cpp
+++ b/Engine/Plugins/Animation/ControlRig/Source/ControlRigEditor/Private/Tools/ControlRigSnapper.cpp
@@ -101,7 +101,8 @@ int32 FControlRigSnapperSelection::NumSelected() const
 TWeakPtr<ISequencer> FControlRigSnapper::GetSequencer()
 {
 	TWeakPtr<ISequencer> WeakSequencer = nullptr;
-	ULevelSequence* LevelSequence = ULevelSequenceEditorBlueprintLibrary::GetFocusedLevelSequence();
+	//if getting sequencer from level sequence need to use the current(master), not the focused
+	ULevelSequence* LevelSequence = ULevelSequenceEditorBlueprintLibrary::GetCurrentLevelSequence();
 	if (LevelSequence)
 	{
 		IAssetEditorInstance* AssetEditor = GEditor->GetEditorSubsystem<UAssetEditorSubsystem>()->FindEditorForAsset(LevelSequence, false);
diff --git a/Engine/Plugins/Editor/SequencerAnimTools/Source/SequencerAnimTools/Private/SequencerAnimEditPivotTool.cpp b/Engine/Plugins/Editor/SequencerAnimTools/Source/SequencerAnimTools/Private/SequencerAnimEditPivotTool.cpp
--- a/Engine/Plugins/Editor/SequencerAnimTools/Source/SequencerAnimTools/Private/SequencerAnimEditPivotTool.cpp
+++ b/Engine/Plugins/Editor/SequencerAnimTools/Source/SequencerAnimTools/Private/SequencerAnimEditPivotTool.cpp
@@ -66,9 +66,11 @@ FLastSelectedObjects USequencerPivotTool::LastSelectedObjects;
 static void GetControlRigsAndSequencer(TArray<TWeakObjectPtr<UControlRig>>& ControlRigs, TWeakPtr<ISequencer>& SequencerPtr, ULevelSequence** LevelSequence)
 {
 	*LevelSequence = ULevelSequenceEditorBlueprintLibrary::GetFocusedLevelSequence();
-	if (*LevelSequence)
+	//if getting sequencer from level sequence need to use the current(master), not the focused
+	ULevelSequence* SequencerLevelSequence = ULevelSequenceEditorBlueprintLibrary::GetCurrentLevelSequence();
+	if (*LevelSequence && SequencerLevelSequence)
 	{
-		IAssetEditorInstance* AssetEditor = GEditor->GetEditorSubsystem<UAssetEditorSubsystem>()->FindEditorForAsset(*LevelSequence, false);
+		IAssetEditorInstance* AssetEditor = GEditor->GetEditorSubsystem<UAssetEditorSubsystem>()->FindEditorForAsset(SequencerLevelSequence, false);
 		ILevelSequenceEditorToolkit* LevelSequenceEditor = static_cast<ILevelSequenceEditorToolkit*>(AssetEditor);
 		SequencerPtr = LevelSequenceEditor ? LevelSequenceEditor->GetSequencer() : nullptr;
 		if (SequencerPtr.IsValid())
