Control Rig: Fix static analysis issue

diff --git a/Engine/Plugins/Animation/ControlRig/Source/ControlRigDeveloper/Private/Graph/ControlRigGraphSchema.cpp b/Engine/Plugins/Animation/ControlRig/Source/ControlRigDeveloper/Private/Graph/ControlRigGraphSchema.cpp
--- a/Engine/Plugins/Animation/ControlRig/Source/ControlRigDeveloper/Private/Graph/ControlRigGraphSchema.cpp
+++ b/Engine/Plugins/Animation/ControlRig/Source/ControlRigDeveloper/Private/Graph/ControlRigGraphSchema.cpp
@@ -714,6 +714,16 @@ const FPinConnectionResponse UControlRigGraphSchema::CanCreateConnection(const U
 			UserLinkDirection = ERigVMPinDirection::Input;
 		}
 
+		if (!PinA)
+		{
+			return FPinConnectionResponse(ECanCreateConnectionResponse::CONNECT_RESPONSE_DISALLOW, FString::Printf(TEXT("Pin %s not found"), *A->GetName()));
+		}
+
+		if (!PinB)
+		{
+			return FPinConnectionResponse(ECanCreateConnectionResponse::CONNECT_RESPONSE_DISALLOW, FString::Printf(TEXT("Pin %s not found"), *B->GetName()));
+		}
+
 		if(PinA->IsWildCard() != PinB->IsWildCard())
 		{
 			// switch the user link direction if only one of the pins is a wildcard 
