Control Rig: Fix TM_ControlTypes failure due to Assertion failed

diff --git a/Engine/Plugins/Animation/ControlRig/Source/ControlRig/Private/Rigs/RigHierarchyElements.cpp b/Engine/Plugins/Animation/ControlRig/Source/ControlRig/Private/Rigs/RigHierarchyElements.cpp
--- a/Engine/Plugins/Animation/ControlRig/Source/ControlRig/Private/Rigs/RigHierarchyElements.cpp
+++ b/Engine/Plugins/Animation/ControlRig/Source/ControlRig/Private/Rigs/RigHierarchyElements.cpp
@@ -539,7 +539,14 @@ void FRigControlSettings::Load(FArchive& Ar)
 	ControlEnum = nullptr;
 	if(!ControlEnumPathName.IsEmpty())
 	{
-		ControlEnum = LoadObject<UEnum>(nullptr, *ControlEnumPathName);
+		if (IsInGameThread())
+		{
+			ControlEnum = LoadObject<UEnum>(nullptr, *ControlEnumPathName);
+		}
+		else
+		{			
+			ControlEnum = FindObject<UEnum>(ANY_PACKAGE, *ControlEnumPathName);
+		}
 	}
 
 	if (Ar.CustomVer(FControlRigObjectVersion::GUID) >= FControlRigObjectVersion::RigHierarchyControlSpaceFavorites)
