Sequencer: Control Rig: YZ Postion channels always keyed

diff --git a/Engine/Plugins/Animation/ControlRig/Source/ControlRigEditor/Private/Sequencer/ControlRigParameterTrackEditor.cpp b/Engine/Plugins/Animation/ControlRig/Source/ControlRigEditor/Private/Sequencer/ControlRigParameterTrackEditor.cpp
--- a/Engine/Plugins/Animation/ControlRig/Source/ControlRigEditor/Private/Sequencer/ControlRigParameterTrackEditor.cpp
+++ b/Engine/Plugins/Animation/ControlRig/Source/ControlRigEditor/Private/Sequencer/ControlRigParameterTrackEditor.cpp
@@ -2539,8 +2539,8 @@ void FControlRigParameterTrackEditor::GetControlRigKeys(UControlRig* InControlRi
 				if (ControlElement->Settings.ControlType == ERigControlType::Position)
 				{
 					bKeyX = bSetKey && EnumHasAnyFlags(ChannelsToKey, EControlRigContextChannelToKey::TranslationX);
-					bKeyY = EnumHasAnyFlags(ChannelsToKey, EControlRigContextChannelToKey::TranslationY);
-					bKeyZ = EnumHasAnyFlags(ChannelsToKey, EControlRigContextChannelToKey::TranslationZ);
+					bKeyY = bSetKey && EnumHasAnyFlags(ChannelsToKey, EControlRigContextChannelToKey::TranslationY);
+					bKeyZ = bSetKey && EnumHasAnyFlags(ChannelsToKey, EControlRigContextChannelToKey::TranslationZ);
 				}
 				else if(ControlElement->Settings.ControlType == ERigControlType::Rotator)
 				{
