Sequencer: Multiple Control Rig Fixes, new anim outliner, keying fixe, etc..
Revamped the outliner to have the treeview hold all of the control rigs in it so we can search on all at once. As part of that revamped the base dockable view so that it works off of the edit mode, that way have one place that holds the lis of valid control rigs and their selections.  Changed SetKey iwth proxies so it proprerly filters. Fix for when a Control rig is replaced via a compile and it wouldn't show up in the viewport.

diff --git a/Engine/Plugins/Animation/ControlRig/Source/ControlRigEditor/Private/EditMode/ControlRigBaseDockableView.cpp b/Engine/Plugins/Animation/ControlRig/Source/ControlRigEditor/Private/EditMode/ControlRigBaseDockableView.cpp
--- a/Engine/Plugins/Animation/ControlRig/Source/ControlRigEditor/Private/EditMode/ControlRigBaseDockableView.cpp
+++ b/Engine/Plugins/Animation/ControlRig/Source/ControlRigEditor/Private/EditMode/ControlRigBaseDockableView.cpp
@@ -13,21 +13,9 @@
 #include "EditorModeManager.h"
 #include "ISequencer.h"
 
-void FControlRigBaseDockableView::SetEditMode(FControlRigEditMode& InEditMode)
-{
-
-	CurrentControlRig = nullptr;
-	ModeTools = InEditMode.GetModeManager();
-	if (FControlRigEditMode* EditMode = static_cast<FControlRigEditMode*>(ModeTools->GetActiveMode(FControlRigEditMode::ModeName)))
-	{
-		EditMode->OnControlRigAddedOrRemoved().AddRaw(this, &FControlRigBaseDockableView::HandleControlAdded);
-		HandleControlAdded(GetControlRig(), true);
-	}
-}
 
 FControlRigBaseDockableView::FControlRigBaseDockableView()
 {
-	FCoreUObjectDelegates::OnObjectsReplaced.AddRaw(this, &FControlRigBaseDockableView::OnObjectsReplaced);
 }
 
 FControlRigBaseDockableView::~FControlRigBaseDockableView()
@@ -35,102 +23,56 @@ FControlRigBaseDockableView::~FControlRigBaseDockableView()
 	if (FControlRigEditMode* EditMode = static_cast<FControlRigEditMode*>(ModeTools->GetActiveMode(FControlRigEditMode::ModeName)))
 	{
 		EditMode->OnControlRigAddedOrRemoved().RemoveAll(this);
-		TArrayView<TWeakObjectPtr<UControlRig>> Controls = EditMode->GetControlRigs();
-		for (TWeakObjectPtr<UControlRig>& ControlRigPtr : Controls)
-		{
-			if (UControlRig* ControlRig = ControlRigPtr.Get())
-			{
-				ControlRig->ControlSelected().RemoveAll(this);
-			}
-		}
+		EditMode->OnControlRigSelected().RemoveAll(this);
 	}
-	else
-	{
-		if (CurrentControlRig.IsValid())
-		{
-			(CurrentControlRig.Get())->ControlSelected().RemoveAll(this);
-		}
-	}
-	CurrentControlRig = nullptr;
-
-	FCoreUObjectDelegates::OnObjectsReplaced.RemoveAll(this);
-
 }
 
-UControlRig* FControlRigBaseDockableView::GetControlRig() 
+TArray<UControlRig*> FControlRigBaseDockableView::GetControlRigs()  const
 {
-	UControlRig* NewControlRig = nullptr;
+	TArray<UControlRig*> ControlRigs;
 	if (FControlRigEditMode* EditMode = static_cast<FControlRigEditMode*>(ModeTools->GetActiveMode(FControlRigEditMode::ModeName)))
 	{
-		//just get first one, should only be one for control rig editor
-		TArrayView<TWeakObjectPtr<UControlRig>> Controls = EditMode->GetControlRigs();
-		for (TWeakObjectPtr<UControlRig>& ControlRigPtr : Controls)
-		{
-			if (UControlRig* ControlRig = ControlRigPtr.Get())
-			{
-				NewControlRig = ControlRig;
-				break;
-			}
-		}
+		ControlRigs = EditMode->GetControlRigsArray(false);
 	}
-	bool bNewControlRig = NewControlRig != CurrentControlRig;
-	if (bNewControlRig)
-	{
-		if (CurrentControlRig.IsValid())
-		{
-			(CurrentControlRig.Get())->ControlSelected().RemoveAll(this);
+	return ControlRigs;
+}
 
-		}
-		CurrentControlRig = NewControlRig;
-		NewControlRigSet(NewControlRig);
+void FControlRigBaseDockableView::SetEditMode(FControlRigEditMode& InEditMode)
+{
+
+	ModeTools = InEditMode.GetModeManager();
+	if (FControlRigEditMode* EditMode = static_cast<FControlRigEditMode*>(ModeTools->GetActiveMode(FControlRigEditMode::ModeName)))
+	{
+		EditMode->OnControlRigAddedOrRemoved().AddRaw(this, &FControlRigBaseDockableView::HandleControlAdded);
+		EditMode->OnControlRigSelected().AddRaw(this, &FControlRigBaseDockableView::HandlElementSelected);
 	}
-	return NewControlRig;
 }
 
-
-void FControlRigBaseDockableView::OnObjectsReplaced(const TMap<UObject*, UObject*>& OldToNewInstanceMap)
+void FControlRigBaseDockableView::HandleControlAdded(UControlRig* ControlRig, bool bIsAdded)
 {
-	if(CurrentControlRig.IsValid())
-	{ 
-		UObject* OldObject = CurrentControlRig.Get();
-		UObject* NewObject = OldToNewInstanceMap.FindRef(OldObject);
-		if (NewObject)
-		{
-			if (UControlRig* ControlRig = Cast<UControlRig>(NewObject))
-			{
-				NewControlRigSet(ControlRig);
-			}
-		}
-	}
 }
 
-void FControlRigBaseDockableView::NewControlRigSet(UControlRig* ControlRig)
+void FControlRigBaseDockableView::HandleControlSelected(UControlRig* ControlRig, FRigControlElement* InControl, bool bSelected)
 {
-	CurrentControlRig = ControlRig;
 }
 
-void FControlRigBaseDockableView::HandleControlAdded(UControlRig* ControlRig, bool bIsAdded)
+void FControlRigBaseDockableView::HandlElementSelected(UControlRig* ControlRig, const FRigElementKey& Key, bool bSelected)
 {
 	if (ControlRig)
 	{
-		if (bIsAdded)
-		{
-			(ControlRig)->ControlSelected().RemoveAll(this);
-			(ControlRig)->ControlSelected().AddRaw(this, &FControlRigBaseDockableView::HandleControlSelected);
-			CurrentControlRig = ControlRig;
-		}
-		else
+		if (Key.Type == ERigElementType::Control)
 		{
-			(ControlRig)->ControlSelected().RemoveAll(this);
+			if (FRigBaseElement* BaseElement = ControlRig->GetHierarchy()->Find(Key))
+			{
+				if (FRigControlElement* ControlElement = Cast<FRigControlElement>(BaseElement))
+				{
+					HandleControlSelected(ControlRig, ControlElement, bSelected);
+				}
+			}
 		}
 	}
 }
 
-void FControlRigBaseDockableView::HandleControlSelected(UControlRig* Subject, FRigControlElement* InControl, bool bSelected)
-{
-
-}
-
 ISequencer* FControlRigBaseDockableView::GetSequencer() const
 {
 	if (FControlRigEditMode* EditMode = static_cast<FControlRigEditMode*>(ModeTools->GetActiveMode(FControlRigEditMode::ModeName)))
diff --git a/Engine/Plugins/Animation/ControlRig/Source/ControlRigEditor/Private/EditMode/ControlRigBaseDockableView.h b/Engine/Plugins/Animation/ControlRig/Source/ControlRigEditor/Private/EditMode/ControlRigBaseDockableView.h
--- a/Engine/Plugins/Animation/ControlRig/Source/ControlRigEditor/Private/EditMode/ControlRigBaseDockableView.h
+++ b/Engine/Plugins/Animation/ControlRig/Source/ControlRigEditor/Private/EditMode/ControlRigBaseDockableView.h
@@ -12,25 +12,24 @@ class UControlRig;
 class ISequencer;
 class FControlRigEditMode;
 struct FRigControlElement;
+struct FRigElementKey;
 
 class FControlRigBaseDockableView 
 {
 public:
 	FControlRigBaseDockableView();
 	virtual ~FControlRigBaseDockableView();
+	TArray<UControlRig*> GetControlRigs() const;
 
-	void SetEditMode(FControlRigEditMode& InEditMode);
+	virtual void SetEditMode(FControlRigEditMode& InEditMode);
 
 protected:
 	virtual void HandleControlSelected(UControlRig* Subject, FRigControlElement* InControl, bool bSelected);
 	virtual void HandleControlAdded(UControlRig* ControlRig, bool bIsAdded);
-	virtual void NewControlRigSet(UControlRig* ControlRig);
 
-	void OnObjectsReplaced(const TMap<UObject*, UObject*>& OldToNewInstanceMap);
+	void HandlElementSelected(UControlRig* Subject, const FRigElementKey& Key, bool bSelected);
 
-	UControlRig* GetControlRig();
 	ISequencer* GetSequencer() const;
-	TWeakObjectPtr<UControlRig> CurrentControlRig;
 
 	FEditorModeTools* ModeTools = nullptr;
 
diff --git a/Engine/Plugins/Animation/ControlRig/Source/ControlRigEditor/Private/EditMode/ControlRigControlsProxy.cpp b/Engine/Plugins/Animation/ControlRig/Source/ControlRigEditor/Private/EditMode/ControlRigControlsProxy.cpp
--- a/Engine/Plugins/Animation/ControlRig/Source/ControlRigEditor/Private/EditMode/ControlRigControlsProxy.cpp
+++ b/Engine/Plugins/Animation/ControlRig/Source/ControlRigEditor/Private/EditMode/ControlRigControlsProxy.cpp
@@ -186,8 +186,22 @@ void UControlRigTransformControlProxy::SetKey(const IPropertyHandle& KeyedProper
 	FRigControlElement* ControlElement = GetControlElement();
 	if (ControlElement)
 	{
+		FRigControlModifiedContext Context;
+		Context.SetKey = EControlRigSetKey::Always;
+		if (KeyedPropertyHandle.GetProperty()->GetFName() == GET_MEMBER_NAME_CHECKED(FEulerTransform, Location))
+		{
+			Context.KeyMask = (uint32)EControlRigContextChannelToKey::Translation;
+		}
+		else if (KeyedPropertyHandle.GetProperty()->GetFName() == GET_MEMBER_NAME_CHECKED(FEulerTransform, Rotation))
+		{
+			Context.KeyMask = (uint32)EControlRigContextChannelToKey::Rotation;
+		}
+		else if (KeyedPropertyHandle.GetProperty()->GetFName() == GET_MEMBER_NAME_CHECKED(FEulerTransform, Scale))
+		{
+			Context.KeyMask = (uint32)EControlRigContextChannelToKey::Scale;
+		}
 		FTransform RealTransform = Transform; //Transform is FEulerTransform
-		ControlRig->SetControlValue<FRigControlValue::FTransform_Float>(ControlName, RealTransform, true, EControlRigSetKey::Always,false);
+		ControlRig->SetControlValue<FRigControlValue::FTransform_Float>(ControlName, RealTransform, true, Context, false);
 	}
 }
 
@@ -239,7 +253,17 @@ void UControlRigTransformNoScaleControlProxy::SetKey(const IPropertyHandle& Keye
 	FRigControlElement* ControlElement = GetControlElement();
 	if (ControlElement)
 	{
-		ControlRig->SetControlValue<FRigControlValue::FTransformNoScale_Float>(ControlName, Transform, true, EControlRigSetKey::Always,false);
+		FRigControlModifiedContext Context;
+		Context.SetKey = EControlRigSetKey::Always;
+		if (KeyedPropertyHandle.GetProperty()->GetFName() == GET_MEMBER_NAME_CHECKED(FTransformNoScale, Location))
+		{
+			Context.KeyMask = (uint32)EControlRigContextChannelToKey::Translation;
+		}
+		else if (KeyedPropertyHandle.GetProperty()->GetFName() == GET_MEMBER_NAME_CHECKED(FTransformNoScale, Rotation))
+		{
+			Context.KeyMask = (uint32)EControlRigContextChannelToKey::Rotation;
+		}
+		ControlRig->SetControlValue<FRigControlValue::FTransformNoScale_Float>(ControlName, Transform, true, Context, false);
 	}
 }
 
@@ -294,7 +318,21 @@ void UControlRigEulerTransformControlProxy::SetKey(const IPropertyHandle& KeyedP
 	FRigControlElement* ControlElement = GetControlElement();
 	if (ControlElement)
 	{
-		ControlRig->SetControlValue<FRigControlValue::FEulerTransform_Float>(ControlName, Transform, true, EControlRigSetKey::Always,false);
+		FRigControlModifiedContext Context;
+		Context.SetKey = EControlRigSetKey::Always;
+		if (KeyedPropertyHandle.GetProperty()->GetFName() == GET_MEMBER_NAME_CHECKED(FEulerTransform, Location))
+		{
+			Context.KeyMask = (uint32)EControlRigContextChannelToKey::Translation;
+		}
+		else if (KeyedPropertyHandle.GetProperty()->GetFName() == GET_MEMBER_NAME_CHECKED(FEulerTransform, Rotation))
+		{
+			Context.KeyMask = (uint32)EControlRigContextChannelToKey::Rotation;
+		}
+		else if (KeyedPropertyHandle.GetProperty()->GetFName() == GET_MEMBER_NAME_CHECKED(FEulerTransform, Scale))
+		{
+			Context.KeyMask = (uint32)EControlRigContextChannelToKey::Scale;
+		}
+		ControlRig->SetControlValue<FRigControlValue::FEulerTransform_Float>(ControlName, Transform, true, Context,false);
 	}
 }
 
@@ -466,7 +504,7 @@ void UControlRigVectorControlProxy::PostEditChangeProperty(struct FPropertyChang
 		{
 			//MUST set through ControlRig
 			FControlRigInteractionScope InteractionScope(ControlRig.Get());
-			ControlRig->SetControlValue<FVector>(ControlName, Vector, true, EControlRigSetKey::DoNotCare,false);
+			ControlRig->SetControlValue<FVector3f>(ControlName, Vector, true, EControlRigSetKey::DoNotCare,false);
 			ControlRig->Evaluate_AnyThread();
 
 		}
@@ -481,8 +519,8 @@ void UControlRigVectorControlProxy::ValueChanged()
 		Modify();
 		const FName PropertyName("Vector");
 		FTrackInstancePropertyBindings Binding(PropertyName, PropertyName.ToString());
-		const FVector Val = (FVector)ControlRig.Get()->GetHierarchy()->GetControlValue(ControlElement, ERigControlValueType::Current).Get<FVector3f>();
-		Binding.CallFunction<FVector>(*this, Val);
+		const FVector3f Val = ControlRig.Get()->GetHierarchy()->GetControlValue(ControlElement, ERigControlValueType::Current).Get<FVector3f>();
+		Binding.CallFunction<FVector3f>(*this, Val);
 	}
 }
 
@@ -504,7 +542,54 @@ void UControlRigVectorControlProxy::SetKey(const IPropertyHandle& KeyedPropertyH
 	FRigControlElement* ControlElement = GetControlElement();
 	if (ControlElement)
 	{
-		ControlRig->SetControlValue<FVector3f>(ControlName, (FVector3f)Vector, true, EControlRigSetKey::Always,false);
+		FRigControlModifiedContext Context;
+		Context.SetKey = EControlRigSetKey::Always;
+		if (KeyedPropertyHandle.GetProperty()->GetFName() == GET_MEMBER_NAME_CHECKED(FVector3f, X))
+		{
+			switch (ControlElement->Settings.ControlType)
+			{
+			case ERigControlType::Position:
+				Context.KeyMask = (uint32)EControlRigContextChannelToKey::TranslationX;
+				break;
+			case ERigControlType::Rotator:	
+				Context.KeyMask = (uint32)EControlRigContextChannelToKey::RotationX;
+				break;
+			case ERigControlType::Scale:
+				Context.KeyMask = (uint32)EControlRigContextChannelToKey::ScaleX;
+				break;
+			}
+		}
+		else if (KeyedPropertyHandle.GetProperty()->GetFName() == GET_MEMBER_NAME_CHECKED(FVector3f, Y))
+		{
+			switch (ControlElement->Settings.ControlType)
+			{
+			case ERigControlType::Position:
+				Context.KeyMask = (uint32)EControlRigContextChannelToKey::TranslationY;
+				break;
+			case ERigControlType::Rotator:
+				Context.KeyMask = (uint32)EControlRigContextChannelToKey::RotationY;
+				break;
+			case ERigControlType::Scale:
+				Context.KeyMask = (uint32)EControlRigContextChannelToKey::ScaleY;
+				break;
+			}
+		}
+		else if (KeyedPropertyHandle.GetProperty()->GetFName() == GET_MEMBER_NAME_CHECKED(FVector3f, Z))
+		{
+			switch (ControlElement->Settings.ControlType)
+			{
+			case ERigControlType::Position:
+				Context.KeyMask = (uint32)EControlRigContextChannelToKey::TranslationZ;
+				break;
+			case ERigControlType::Rotator:
+				Context.KeyMask = (uint32)EControlRigContextChannelToKey::RotationZ;
+				break;
+			case ERigControlType::Scale:
+				Context.KeyMask = (uint32)EControlRigContextChannelToKey::ScaleZ;
+				break;
+			}
+		}
+		ControlRig->SetControlValue<FVector3f>(ControlName, (FVector3f)Vector, true, Context,false);
 	}
 }
 
@@ -558,7 +643,17 @@ void UControlRigVector2DControlProxy::SetKey(const IPropertyHandle& KeyedPropert
 	FRigControlElement* ControlElement = GetControlElement();
 	if (ControlElement)
 	{
-		ControlRig->SetControlValue<FVector3f>(ControlName, FVector3f(Vector2D.X, Vector2D.Y, 0.f), true, EControlRigSetKey::Always,false);
+		FRigControlModifiedContext Context;
+		Context.SetKey = EControlRigSetKey::Always;
+		if (KeyedPropertyHandle.GetProperty()->GetFName() == GET_MEMBER_NAME_CHECKED(FVector2D, X))
+		{
+			Context.KeyMask = (uint32)EControlRigContextChannelToKey::TranslationX;
+		}
+		else if (KeyedPropertyHandle.GetProperty()->GetFName() == GET_MEMBER_NAME_CHECKED(FVector2D, Y))
+		{
+			Context.KeyMask = (uint32)EControlRigContextChannelToKey::TranslationY;
+		}
+		ControlRig->SetControlValue<FVector3f>(ControlName, FVector3f(Vector2D.X, Vector2D.Y, 0.f), true, Context,false);
 	}
 }
 
diff --git a/Engine/Plugins/Animation/ControlRig/Source/ControlRigEditor/Private/EditMode/ControlRigControlsProxy.h b/Engine/Plugins/Animation/ControlRig/Source/ControlRigEditor/Private/EditMode/ControlRigControlsProxy.h
--- a/Engine/Plugins/Animation/ControlRig/Source/ControlRigEditor/Private/EditMode/ControlRigControlsProxy.h
+++ b/Engine/Plugins/Animation/ControlRig/Source/ControlRigEditor/Private/EditMode/ControlRigControlsProxy.h
@@ -250,7 +250,7 @@ class UControlRigVectorControlProxy : public UControlRigControlsProxy
 public:
 
 	UPROPERTY(EditAnywhere, Interp, Category = "Control")
-	FVector Vector;
+	FVector3f Vector;
 };
 
 UCLASS()
diff --git a/Engine/Plugins/Animation/ControlRig/Source/ControlRigEditor/Private/EditMode/ControlRigEditMode.cpp b/Engine/Plugins/Animation/ControlRig/Source/ControlRigEditor/Private/EditMode/ControlRigEditMode.cpp
--- a/Engine/Plugins/Animation/ControlRig/Source/ControlRigEditor/Private/EditMode/ControlRigEditMode.cpp
+++ b/Engine/Plugins/Animation/ControlRig/Source/ControlRigEditor/Private/EditMode/ControlRigEditMode.cpp
@@ -161,6 +161,7 @@ FControlRigEditMode::~FControlRigEditMode()
 
 	DestroyShapesActors(nullptr);
 	OnControlRigAddedOrRemovedDelegate.Clear();
+	OnControlRigSelectedDelegate.Clear();
 
 	TArray<TWeakObjectPtr<UControlRig>> PreviousRuntimeRigs = RuntimeControlRigs;
 	for (int32 PreviousRuntimeRigIndex = 0; PreviousRuntimeRigIndex < PreviousRuntimeRigs.Num(); PreviousRuntimeRigIndex++)
@@ -437,6 +438,7 @@ void FControlRigEditMode::Exit()
 
 	DestroyShapesActors(nullptr);
 	OnControlRigAddedOrRemovedDelegate.Clear();
+	OnControlRigSelectedDelegate.Clear();
 
 	TArray<TWeakObjectPtr<UControlRig>> PreviousRuntimeRigs = RuntimeControlRigs;
 	for (int32 PreviousRuntimeRigIndex = 0; PreviousRuntimeRigIndex < PreviousRuntimeRigs.Num(); PreviousRuntimeRigIndex++)
@@ -1163,13 +1165,10 @@ bool FControlRigEditMode::HandleClick(FEditorViewportClient* InViewportClient, H
 					}
 					else if(Click.IsControlDown()) //if ctrl we toggle selection
 					{
-						for (TWeakObjectPtr<UControlRig>& RuntimeRigPtr : RuntimeControlRigs)
+						if (UControlRig* ControlRig = ShapeActor->ControlRig.Get())
 						{
-							if (UControlRig* ControlRig = RuntimeRigPtr.Get())
-							{
-								bool bIsSelected = ControlRig->IsControlSelected(ControlName);
-								SetRigElementSelection(ShapeActor->ControlRig.Get(), ERigElementType::Control, ControlName, !bIsSelected);
-							}
+							bool bIsSelected = ControlRig->IsControlSelected(ControlName);
+							SetRigElementSelection(ControlRig, ERigElementType::Control, ControlName, !bIsSelected);
 						}
 					}
 					else
@@ -1349,8 +1348,10 @@ bool FControlRigEditMode::HandleClick(FEditorViewportClient* InViewportClient, H
 	FScopedTransaction ScopedTransaction(LOCTEXT("SelectControlTransaction", "Select Control"), IsInLevelEditor() &&  !GIsTransacting);
 	
 	// clear selected controls
-	ClearRigElementSelection(FRigElementTypeHelper::ToMask(ERigElementType::All));
-
+	if (Click.IsShiftDown() ==false && Click.IsControlDown() == false)
+	{
+		ClearRigElementSelection(FRigElementTypeHelper::ToMask(ERigElementType::All));
+	}
 	/*
 	if(!InViewportClient->IsLevelEditorClient() && !InViewportClient->IsSimulateInEditorViewport())
 	{
@@ -1502,8 +1503,10 @@ bool IntersectsBox( AActor& InActor, const FBox& InBox, FLevelEditorViewportClie
 
 bool FControlRigEditMode::BoxSelect(FBox& InBox, bool InSelect)
 {
+
+	const UControlRigEditModeSettings* Settings = GetDefault<UControlRigEditModeSettings>();
 	FLevelEditorViewportClient* LevelViewportClient = GCurrentLevelEditingViewportClient;
-	if (LevelViewportClient->IsInGameView() == true)
+	if (LevelViewportClient->IsInGameView() == true || Settings->bHideControlShapes)
 	{
 		return  FEdMode::BoxSelect(InBox, InSelect);
 	}
@@ -1560,7 +1563,8 @@ bool FControlRigEditMode::BoxSelect(FBox& InBox, bool InSelect)
 
 bool FControlRigEditMode::FrustumSelect(const FConvexVolume& InFrustum, FEditorViewportClient* InViewportClient, bool InSelect)
 {
-	if (InViewportClient->IsInGameView() == true)
+	const UControlRigEditModeSettings* Settings = GetDefault<UControlRigEditModeSettings>();
+	if (InViewportClient->IsInGameView() == true || Settings->bHideControlShapes)
 	{
 		return FEdMode::FrustumSelect(InFrustum, InViewportClient, InSelect);
 	}
@@ -2520,6 +2524,8 @@ void FControlRigEditMode::ToggleManipulators()
 		{
 			if (ControlRig)
 			{
+				FScopedTransaction ScopedTransaction(LOCTEXT("ToggleControlsVisibility", "Toggle Controls Visibility"),!GIsTransacting);
+				ControlRig->Modify();
 				ControlRig->ToggleControlsVisible();
 			}
 		}
@@ -2842,7 +2848,7 @@ void FControlRigEditMode::RequestToRecreateControlShapeActors(UControlRig* Contr
 		if (RecreateControlShapesRequired != ERecreateControlRigShape::RecreateAll)
 		{
 			RecreateControlShapesRequired = ERecreateControlRigShape::RecreateSpecified;
-			if (!ControlRigsToRecreate.Find(ControlRig))
+			if (ControlRigsToRecreate.Find(ControlRig) == INDEX_NONE)
 			{
 				ControlRigsToRecreate.Add(ControlRig);
 			}
@@ -3062,7 +3068,7 @@ bool FControlRigEditMode::AreRigElementSelectedAndMovable(UControlRig* ControlRi
 {
 	const UControlRigEditModeSettings* Settings = GetDefault<UControlRigEditModeSettings>();
 
-	if (!Settings || Settings->bHideControlShapes || !ControlRig || ControlRig->GetControlsVisible() == false || !AreRigElementsSelected(FRigElementTypeHelper::ToMask(ERigElementType::Control), ControlRig))
+	if (!Settings || Settings->bHideControlShapes || !ControlRig || !AreRigElementsSelected(FRigElementTypeHelper::ToMask(ERigElementType::Control), ControlRig))
 	{
 		return false;
 	}
@@ -3081,18 +3087,14 @@ bool FControlRigEditMode::AreRigElementSelectedAndMovable(UControlRig* ControlRi
 
 void FControlRigEditMode::ReplaceControlRig(UControlRig* OldControlRig, UControlRig* NewControlRig)
 {
-	for (int32 Index = 0; Index < RuntimeControlRigs.Num(); ++Index)
+	if (OldControlRig != nullptr)
 	{
-		TWeakObjectPtr<UControlRig>& RuntimeRigPtr = RuntimeControlRigs[Index];
-		if (UControlRig* RuntimeControlRig = RuntimeRigPtr.Get())
-		{
-			if (RuntimeControlRig == OldControlRig)
-			{
-				RuntimeControlRigs[Index] = NewControlRig;
-				break;
-			}
-		}
+		RemoveControlRig(OldControlRig);
 	}
+	AddControlRigInternal(NewControlRig);
+	SetObjects_Internal();
+	RequestToRecreateControlShapeActors(NewControlRig);
+
 }
 void FControlRigEditMode::OnHierarchyModified(ERigHierarchyNotification InNotif, URigHierarchy* InHierarchy, const FRigBaseElement* InElement)
 {
@@ -3152,7 +3154,20 @@ void FControlRigEditMode::OnHierarchyModified(ERigHierarchyNotification InNotif,
             	case ERigElementType::Reference:
 				{
 					const bool bSelected = InNotif == ERigHierarchyNotification::ElementSelected;
-					
+					// users may select gizmo and control rig units, so we have to let them go through both of them if they do
+						// first go through gizmo actor
+					UControlRig* ControlRig = InHierarchy->GetTypedOuter<UControlRig>();
+					if (ControlRig == nullptr)
+					{
+						if (RuntimeControlRigs.Num() > 0)
+						{
+							ControlRig = RuntimeControlRigs[0].Get();
+						}
+					}
+					if (ControlRig)
+					{
+						OnControlRigSelectedDelegate.Broadcast(ControlRig, Key, bSelected);
+					}
 					// if it's control
 					if (Key.Type == ERigElementType::Control)
 					{
@@ -3161,16 +3176,7 @@ void FControlRigEditMode::OnHierarchyModified(ERigHierarchyNotification InNotif,
 						{
 							ControlProxy->Modify();
 						}
-						// users may select gizmo and control rig units, so we have to let them go through both of them if they do
-						// first go through gizmo actor
-						UControlRig* ControlRig= InHierarchy->GetTypedOuter<UControlRig>();
-						if (ControlRig == nullptr)
-						{
-							if (RuntimeControlRigs.Num() > 0)
-							{
-								ControlRig = RuntimeControlRigs[0].Get();
-							}
-						}
+						
 						AControlRigShapeActor* ShapeActor = GetControlShapeFromControlName(ControlRig,Key.Name);
 						if (ShapeActor)
 						{
@@ -3191,6 +3197,7 @@ void FControlRigEditMode::OnHierarchyModified(ERigHierarchyNotification InNotif,
 								ControlProxy->SelectProxy(ControlRig,Key.Name, false);
 							}
 						}
+
 					}
 					bSelectionChanged = true;
 		
diff --git a/Engine/Plugins/Animation/ControlRig/Source/ControlRigEditor/Private/EditMode/ControlRigEditMode.h b/Engine/Plugins/Animation/ControlRig/Source/ControlRigEditor/Private/EditMode/ControlRigEditMode.h
--- a/Engine/Plugins/Animation/ControlRig/Source/ControlRigEditor/Private/EditMode/ControlRigEditMode.h
+++ b/Engine/Plugins/Animation/ControlRig/Source/ControlRigEditor/Private/EditMode/ControlRigEditMode.h
@@ -44,6 +44,7 @@ DECLARE_DELEGATE_RetVal_ThreeParams(FTransform, FOnGetRigElementTransform, const
 DECLARE_DELEGATE_ThreeParams(FOnSetRigElementTransform, const FRigElementKey& /*RigElementKey*/, const FTransform& /*Transform*/, bool /*bLocal*/);
 DECLARE_DELEGATE_RetVal(TSharedPtr<FUICommandList>, FNewMenuCommandsDelegate);
 DECLARE_MULTICAST_DELEGATE_TwoParams(FControlRigAddedOrRemoved, UControlRig*, bool /*true if added, false if removed*/);
+DECLARE_MULTICAST_DELEGATE_ThreeParams(FControlRigSelected, UControlRig*, const FRigElementKey& /*RigElementKey*/,const bool /*bIsSelected*/);
 DECLARE_DELEGATE_RetVal(UToolMenu*, FOnGetContextMenu);
 
 class FControlRigEditMode;
@@ -154,6 +155,9 @@ public:
 	/* Control Rig Changed Delegate*/
 	FControlRigAddedOrRemoved& OnControlRigAddedOrRemoved() { return OnControlRigAddedOrRemovedDelegate; }
 
+	/* Control Rig Selected Delegate*/
+	FControlRigSelected& OnControlRigSelected() { return OnControlRigSelectedDelegate; }
+
 	// callback that gets called when rig element is selected in other view
 	void OnHierarchyModified(ERigHierarchyNotification InNotif, URigHierarchy* InHierarchy, const FRigBaseElement* InElement);
 	void OnControlModified(UControlRig* Subject, FRigControlElement* InControlElement, const FRigControlModifiedContext& Context);
@@ -286,6 +290,7 @@ protected:
 	FNewMenuCommandsDelegate OnContextMenuCommandsDelegate;
 	FSimpleMulticastDelegate OnAnimSystemInitializedDelegate;
 	FControlRigAddedOrRemoved OnControlRigAddedOrRemovedDelegate;
+	FControlRigSelected OnControlRigSelectedDelegate;
 
 	/** GetSelectedRigElements, if InControlRig is nullptr get the first one */
 	TArray<FRigElementKey> GetSelectedRigElements(UControlRig* InControlRig = nullptr) const;
diff --git a/Engine/Plugins/Animation/ControlRig/Source/ControlRigEditor/Private/EditMode/SControlRigDetails.cpp b/Engine/Plugins/Animation/ControlRig/Source/ControlRigEditor/Private/EditMode/SControlRigDetails.cpp
--- a/Engine/Plugins/Animation/ControlRig/Source/ControlRigEditor/Private/EditMode/SControlRigDetails.cpp
+++ b/Engine/Plugins/Animation/ControlRig/Source/ControlRigEditor/Private/EditMode/SControlRigDetails.cpp
@@ -404,13 +404,6 @@ SControlRigDetails::~SControlRigDetails()
 	//base class handles control rig related cleanup
 }
 
-
-void SControlRigDetails::HandleControlAdded(UControlRig* ControlRig, bool bIsAdded)
-{
-	FControlRigBaseDockableView::HandleControlAdded(ControlRig, bIsAdded);
-	UpdateProxies();
-}
-
 void SControlRigDetails::HandleControlSelected(UControlRig* Subject, FRigControlElement* InControl, bool bSelected)
 {
 	FControlRigBaseDockableView::HandleControlSelected(Subject, InControl, bSelected);
@@ -435,100 +428,99 @@ void SControlRigDetails::UpdateProxies()
 		TArray<TWeakObjectPtr<>> IndividualBools;
 		TArray<TWeakObjectPtr<>> IndividualIntegers;
 		TArray<TWeakObjectPtr<>> IndividualEnums;
-
-		if (UControlRig* ControlRig = GetControlRig())
+		TArray<UControlRig*> ControlRigs = GetControlRigs();
+	
+		if (FControlRigEditMode* EditMode = static_cast<FControlRigEditMode*>(ModeTools->GetActiveMode(FControlRigEditMode::ModeName)))
 		{
-			if (FControlRigEditMode* EditMode = static_cast<FControlRigEditMode*>(ModeTools->GetActiveMode(FControlRigEditMode::ModeName)))
+			if (UControlRigDetailPanelControlProxies* ControlProxy = EditMode->GetDetailProxies())
 			{
-				if (UControlRigDetailPanelControlProxies* ControlProxy = EditMode->GetDetailProxies())
+				const TArray<UControlRigControlsProxy*>& Proxies = ControlProxy->GetSelectedProxies();
+				for (UControlRigControlsProxy* Proxy : Proxies)
 				{
-					const TArray<UControlRigControlsProxy*>& Proxies = ControlProxy->GetSelectedProxies();
-					for (UControlRigControlsProxy* Proxy : Proxies)
+					if (Proxy->GetClass() == UControlRigTransformControlProxy::StaticClass())
 					{
-						if (Proxy->GetClass() == UControlRigTransformControlProxy::StaticClass())
-						{
-							Transforms.Add(Proxy);
-						}
-						else if (Proxy->GetClass() == UControlRigTransformNoScaleControlProxy::StaticClass())
+						Transforms.Add(Proxy);
+					}
+					else if (Proxy->GetClass() == UControlRigTransformNoScaleControlProxy::StaticClass())
+					{
+						TransformNoScales.Add(Proxy);
+					}
+					else if (Proxy->GetClass() == UControlRigEulerTransformControlProxy::StaticClass())
+					{
+						Eulers.Add(Proxy);
+					}
+					else if (Proxy->GetClass() == UControlRigFloatControlProxy::StaticClass())
+					{
+						if (Proxy->bIsIndividual)
 						{
-							TransformNoScales.Add(Proxy);
+							IndividualFloats.Add(Proxy);
 						}
-						else if (Proxy->GetClass() == UControlRigEulerTransformControlProxy::StaticClass())
+						else
 						{
-							Eulers.Add(Proxy);
+							Floats.Add(Proxy);
 						}
-						else if (Proxy->GetClass() == UControlRigFloatControlProxy::StaticClass())
+					}
+					else if (Proxy->GetClass() == UControlRigVectorControlProxy::StaticClass())
+					{
+						Vectors.Add(Proxy);
+					}
+					else if (Proxy->GetClass() == UControlRigVector2DControlProxy::StaticClass())
+					{
+						Vector2Ds.Add(Proxy);
+					}
+					else if (Proxy->GetClass() == UControlRigBoolControlProxy::StaticClass())
+					{
+						if (Proxy->bIsIndividual)
 						{
-							if (Proxy->bIsIndividual)
-							{
-								IndividualFloats.Add(Proxy);
-							}
-							else
-							{
-								Floats.Add(Proxy);
-							}
+							IndividualBools.Add(Proxy);
 						}
-						else if (Proxy->GetClass() == UControlRigVectorControlProxy::StaticClass())
+						else
 						{
-							Vectors.Add(Proxy);
+							Bools.Add(Proxy);
 						}
-						else if (Proxy->GetClass() == UControlRigVector2DControlProxy::StaticClass())
+					}
+					else if (Proxy->GetClass() == UControlRigEnumControlProxy::StaticClass())
+					{
+						if (Proxy->bIsIndividual)
 						{
-							Vector2Ds.Add(Proxy);
+							IndividualEnums.Add(Proxy);
 						}
-						else if (Proxy->GetClass() == UControlRigBoolControlProxy::StaticClass())
+						else
 						{
-							if (Proxy->bIsIndividual)
-							{
-								IndividualBools.Add(Proxy);
-							}
-							else
-							{
-								Bools.Add(Proxy);
-							}
+							Enums.Add(Proxy);
 						}
-						else if (Proxy->GetClass() == UControlRigEnumControlProxy::StaticClass())
+					}
+					else if (Proxy->GetClass() == UControlRigIntegerControlProxy::StaticClass())
+					{
+						if (Proxy->bIsIndividual)
 						{
-							if (Proxy->bIsIndividual)
-							{
-								IndividualEnums.Add(Proxy);
-							}
-							else
-							{
-								Enums.Add(Proxy);
-							}
+							IndividualIntegers.Add(Proxy);
 						}
-						else if (Proxy->GetClass() == UControlRigIntegerControlProxy::StaticClass())
+						else
 						{
-							if (Proxy->bIsIndividual)
-							{
-								IndividualIntegers.Add(Proxy);
-							}
-							else
-							{
-								Integers.Add(Proxy);
-							}
+							Integers.Add(Proxy);
 						}
 					}
 				}
-				for (TWeakObjectPtr<>& Object : Transforms)
+			}
+			for (TWeakObjectPtr<>& Object : Transforms)
+			{
+				UControlRigControlsProxy* Proxy = Cast<UControlRigControlsProxy>(Object.Get());
+				if (Proxy)
 				{
-					UControlRigControlsProxy* Proxy = Cast<UControlRigControlsProxy>(Object.Get());
-					if (Proxy)
-					{
-						Proxy->SetIsMultiple(Transforms.Num() > 1);
-					}
+					Proxy->SetIsMultiple(Transforms.Num() > 1);
 				}
-				for (TWeakObjectPtr<>& Object : Floats)
+			}
+			for (TWeakObjectPtr<>& Object : Floats)
+			{
+				UControlRigControlsProxy* Proxy = Cast<UControlRigControlsProxy>(Object.Get());
+				if (Proxy)
 				{
-					UControlRigControlsProxy* Proxy = Cast<UControlRigControlsProxy>(Object.Get());
-					if (Proxy)
-					{
-						Proxy->SetIsMultiple(Transforms.Num() > 1);
-					}
+					Proxy->SetIsMultiple(Transforms.Num() > 1);
 				}
 			}
 		}
+
 		SetTransformDetailsObjects(Transforms);
 		SetTransformNoScaleDetailsObjects(TransformNoScales);
 		SetEulerTransformDetailsObjects(Eulers);
diff --git a/Engine/Plugins/Animation/ControlRig/Source/ControlRigEditor/Private/EditMode/SControlRigDetails.h b/Engine/Plugins/Animation/ControlRig/Source/ControlRigEditor/Private/EditMode/SControlRigDetails.h
--- a/Engine/Plugins/Animation/ControlRig/Source/ControlRigEditor/Private/EditMode/SControlRigDetails.h
+++ b/Engine/Plugins/Animation/ControlRig/Source/ControlRigEditor/Private/EditMode/SControlRigDetails.h
@@ -76,7 +76,6 @@ private:
 
 	void UpdateProxies();
 	virtual void HandleControlSelected(UControlRig* Subject, FRigControlElement* InControl, bool bSelected) override;
-	virtual void HandleControlAdded(UControlRig* ControlRig, bool bIsAdded) override;
 
 	//these views will hold more than one of the same tyhpe
 	TSharedPtr<IDetailsView> ControlEulerTransformDetailsView;
diff --git a/Engine/Plugins/Animation/ControlRig/Source/ControlRigEditor/Private/EditMode/SControlRigOutliner.cpp b/Engine/Plugins/Animation/ControlRig/Source/ControlRigEditor/Private/EditMode/SControlRigOutliner.cpp
--- a/Engine/Plugins/Animation/ControlRig/Source/ControlRigEditor/Private/EditMode/SControlRigOutliner.cpp
+++ b/Engine/Plugins/Animation/ControlRig/Source/ControlRigEditor/Private/EditMode/SControlRigOutliner.cpp
@@ -22,44 +22,50 @@
 #include "UnrealEd/Public/Selection.h"
 #include "Editor.h"
 #include "LevelEditor.h"
-#include "Widgets/Views/SListView.h"
 #include "Widgets/Text/STextBlock.h"
-#include "Widgets/Input/SButton.h"
-#include "Widgets/Layout/SScrollBox.h"
 #include "ControlRigControlsProxy.h"
 #include "ControlRigEditModeSettings.h"
 #include "TimerManager.h"
 #include "SRigHierarchyTreeView.h"
 #include "Rigs/RigHierarchyController.h"
 #include "ControlRigObjectBinding.h"
+#include "ControlRigEditorStyle.h"
+#include "Settings/ControlRigSettings.h"
+#include "Widgets/Layout/SSpacer.h"
+#include "Widgets/Images/SImage.h"
+#include "Widgets/Input/SSearchBox.h"
+#include "Slate/Private/Widgets/Views/SListPanel.h"
+#include "PropertyCustomizationHelpers.h"
+#include "Framework/Application/SlateApplication.h"
+#include "Editor/EditorEngine.h"
+#include "HelperUtil.h"
+#include "ScopedTransaction.h"
 
 #define LOCTEXT_NAMESPACE "ControlRigOutliner"
 
+FRigTreeDisplaySettings FMultiRigTreeDelegates::DefaultDisplaySettings;
 
-void SControlRigOutlinerItem::Construct(const FArguments& InArgs)
-{
-	ControlRigEditMode = InArgs._EditMode;
-	bIsChangingRigHierarchy = false;
 
-	DisplaySettings.bShowBones = false;
-	DisplaySettings.bShowControls = true;
-	DisplaySettings.bShowNulls = false;
-	DisplaySettings.bShowReferences = false;
-	DisplaySettings.bShowRigidBodies = false;
-	DisplaySettings.bHideParentsOnFilter = true;
-	DisplaySettings.bFlattenHierarchyOnFilter = true;
+//////////////////////////////////////////////////////////////
+/// FMultiRigData
+///////////////////////////////////////////////////////////
 
-	FRigTreeDelegates RigTreeDelegates;
-	RigTreeDelegates.OnGetHierarchy = FOnGetRigTreeHierarchy::CreateSP(this, &SControlRigOutlinerItem::GetHierarchy);
-	RigTreeDelegates.OnGetDisplaySettings = FOnGetRigTreeDisplaySettings::CreateSP(this, &SControlRigOutlinerItem::GetDisplaySettings);
-	RigTreeDelegates.OnSelectionChanged = FOnRigTreeSelectionChanged::CreateSP(this, &SControlRigOutlinerItem::HandleSelectionChanged);
+uint32 GetTypeHash(const FMultiRigData& Data)
+{
+	return GetTypeHash(TTuple<const UControlRig*, FRigElementKey>(Data.ControlRig.Get(), (Data.Key.IsSet() ? Data.Key.GetValue() : FRigElementKey())));
+}
 
-	FText AreaTitle;
-	if (InArgs._ControlRig)
+FText FMultiRigData::GetName() const
+{
+	if (Key.IsSet())
 	{
-		FString ControlRigName = InArgs._ControlRig->GetName();
+		return FText::FromName(Key.GetValue().Name);
+	}
+	else if (ControlRig.IsValid())
+	{
+		FString ControlRigName = ControlRig.Get()->GetName();
 		FString BoundObjectName;
-		if (TSharedPtr<IControlRigObjectBinding> ObjectBinding = InArgs._ControlRig->GetObjectBinding())
+		if (TSharedPtr<IControlRigObjectBinding> ObjectBinding = ControlRig.Get()->GetObjectBinding())
 		{
 			if (ObjectBinding->GetBoundObject())
 			{
@@ -70,271 +76,1076 @@ void SControlRigOutlinerItem::Construct(const FArguments& InArgs)
 				}
 			}
 		}
-		AreaTitle = FText::Format(LOCTEXT("ControlTitle", "{0}  ({1})"), FText::AsCultureInvariant(ControlRigName), FText::AsCultureInvariant((BoundObjectName)));
+		FText AreaTitle = FText::Format(LOCTEXT("ControlTitle", "{0}  ({1})"), FText::AsCultureInvariant(ControlRigName), FText::AsCultureInvariant((BoundObjectName)));
+		return AreaTitle;
 	}
-	ChildSlot
-		[
-			SNew(SScrollBox)
-			+ SScrollBox::Slot()
-			[
-			SNew(SVerticalBox)
-			+ SVerticalBox::Slot()
-				.AutoHeight()
-				[
-					SAssignNew(PickerExpander, SExpandableArea)
-					.InitiallyCollapsed(false)
-					//.AreaTitle(AreaTitle)
-					//.AreaTitleFont(FEditorStyle::GetFontStyle("DetailsView.CategoryFontStyle"))
-					.BorderBackgroundColor(FLinearColor(.6f, .6f, .6f))
-					.BodyContent()
-					[
-						SAssignNew(HierarchyTreeView, SSearchableRigHierarchyTreeView)
-						.RigTreeDelegates(RigTreeDelegates)
-					]
-					.HeaderContent()
-						[
-							SNew(SHorizontalBox)
-							+ SHorizontalBox::Slot()
-							.AutoWidth()
-							.VAlign(VAlign_Center)
-							[
-								SNew(SButton)
-								.ContentPadding(2)
-								.ButtonStyle(FEditorStyle::Get(), "NoBorder")
-								.OnClicked(this, &SControlRigOutlinerItem::OnToggleVisibility)
-								.ToolTipText(LOCTEXT("ControlRigShapesVisibility", "Control Rig Shapes Visibility"))
-								.IsEnabled(this, &SControlRigOutlinerItem::VisibilityToggleEnabled)
-								.HAlign(HAlign_Center)
-								.VAlign(VAlign_Center)
-								.Content()
-								[
-									SNew(SImage)
-									.Image(this, &SControlRigOutlinerItem::GetVisibilityBrushForElement)
-								]
-							]
-							+ SHorizontalBox::Slot()
-							.AutoWidth()
-							.VAlign(VAlign_Center)
-							[
-								SNew(STextBlock)
-								.Text(AreaTitle)
-								.Justification(ETextJustify::Left)
-								//.TextStyle(FSourceFilterStyle::Get(), "SourceFilter.TextStyle")
-							]
-						]
-				]
-		
-			]
-		];
-	NewControlRigSet(InArgs._ControlRig);
+	FName None = NAME_None;
+	return FText::FromName(None);
 }
 
-FReply SControlRigOutlinerItem::OnToggleVisibility()
+bool FMultiRigData::operator== (const FMultiRigData & Other) const
 {
-	if (UControlRig *ControlRig  = CurrentControlRig.Get())
+	if (ControlRig.IsValid() == false && Other.ControlRig.IsValid() == false)
 	{
-		ControlRig->ToggleControlsVisible();
+		return true;
+	}
+	else if (ControlRig.IsValid() == true && Other.ControlRig.IsValid() == true)
+	{
+		if (ControlRig.Get() == Other.ControlRig.Get())
+		{
+			if (Key.IsSet() == false && Other.Key.IsSet() == false)
+			{
+				return true;
+			}
+			else if (Key.IsSet() == true && Other.Key.IsSet() == true)
+			{
+				return Key.GetValue() == Other.Key.GetValue();
+			}
+		}
 	}
-	return FReply::Handled();
+	return false;
 }
 
-bool SControlRigOutlinerItem::VisibilityToggleEnabled() const
+bool FMultiRigData::IsValid() const
 {
-	return CurrentControlRig.IsValid();
+	return (ControlRig.IsValid() && (Key.IsSet() == false || Key.GetValue().IsValid()));
 }
 
-const FSlateBrush* SControlRigOutlinerItem::GetVisibilityBrushForElement() const
+URigHierarchy* FMultiRigData::GetHierarchy() const
 {
-	if (UControlRig* ControlRig = CurrentControlRig.Get())
+	if (ControlRig.IsValid())
 	{
-		if (ControlRig->GetControlsVisible())
+		return ControlRig.Get()->GetHierarchy();
+	}
+	return nullptr;
+}
+
+//////////////////////////////////////////////////////////////
+/// FMultiRigTreeElement
+///////////////////////////////////////////////////////////
+
+FMultiRigTreeElement::FMultiRigTreeElement(const FMultiRigData& InData, TWeakPtr<SMultiRigHierarchyTreeView> InTreeView, ERigTreeFilterResult InFilterResult)
+{
+	Data = InData;
+	FilterResult = InFilterResult;
+
+	if (InTreeView.IsValid() && Data.ControlRig.IsValid())
+	{
+		if (const URigHierarchy* Hierarchy = Data.GetHierarchy())
 		{
-			return IsHovered() ? FEditorStyle::GetBrush("Level.VisibleHighlightIcon16x") :
-				FEditorStyle::GetBrush("Level.VisibleIcon16x");
+			const FRigTreeDisplaySettings& Settings = InTreeView.Pin()->GetTreeDelegates().GetDisplaySettings();
+			RefreshDisplaySettings(Hierarchy, Settings);
 		}
 	}
-	return IsHovered() ? FEditorStyle::GetBrush("Level.NotVisibleHighlightIcon16x") :
-		FEditorStyle::GetBrush("Level.NotVisibleIcon16x");
 }
 
-void SControlRigOutlinerItem::OnObjectsReplaced(const TMap<UObject*, UObject*>& OldToNewInstanceMap)
+TSharedRef<ITableRow> FMultiRigTreeElement::MakeTreeRowWidget(const TSharedRef<STableViewBase>& InOwnerTable, TSharedRef<FMultiRigTreeElement> InRigTreeElement, TSharedPtr<SMultiRigHierarchyTreeView> InTreeView, const FRigTreeDisplaySettings& InSettings, bool bPinned)
+{
+	return SNew(SMultiRigHierarchyItem, InOwnerTable, InRigTreeElement, InTreeView, InSettings, bPinned);
+
+}
+
+void FMultiRigTreeElement::RefreshDisplaySettings(const URigHierarchy* InHierarchy, const FRigTreeDisplaySettings& InSettings)
+{
+	const TPair<const FSlateBrush*, FSlateColor> Result = SMultiRigHierarchyItem::GetBrushForElementType(InHierarchy, Data);
+
+	IconBrush = Result.Key;
+	IconColor = Result.Value;
+	if (IconColor.IsColorSpecified() && InSettings.bShowIconColors)
+	{
+		IconColor = FilterResult == ERigTreeFilterResult::Shown ? Result.Value : FSlateColor(Result.Value.GetSpecifiedColor() * 0.5f);
+	}
+	else
+	{
+		IconColor = FilterResult == ERigTreeFilterResult::Shown ? FSlateColor::UseForeground() : FSlateColor(FLinearColor::Gray * 0.5f);
+	}
+	TextColor = FilterResult == ERigTreeFilterResult::Shown ? FSlateColor::UseForeground() : FSlateColor(FLinearColor::Gray * 0.5f);
+
+}
+//////////////////////////////////////////////////////////////
+/// SMultiRigHierarchyItem
+///////////////////////////////////////////////////////////
+void SMultiRigHierarchyItem::Construct(const FArguments& InArgs, const TSharedRef<STableViewBase>& OwnerTable, TSharedRef<FMultiRigTreeElement> InRigTreeElement, TSharedPtr<SMultiRigHierarchyTreeView> InTreeView, const FRigTreeDisplaySettings& InSettings, bool bPinned)
+{
+	WeakRigTreeElement = InRigTreeElement;
+	Delegates = InTreeView->GetTreeDelegates();
+
+	
+	if (InRigTreeElement->Data.ControlRig.IsValid() == false || (InRigTreeElement->Data.Key.IsSet() &&
+		InRigTreeElement->Data.Key.GetValue().IsValid() == false))
+	{
+		STableRow<TSharedPtr<FMultiRigTreeElement>>::Construct(
+			STableRow<TSharedPtr<FMultiRigTreeElement>>::FArguments()
+			.ShowSelection(false)
+			.Content()
+			[
+				SNew(SVerticalBox)
+				+ SVerticalBox::Slot()
+			.FillHeight(200.f)
+			[
+				SNew(SSpacer)
+			]
+			], OwnerTable);
+		return;
+	}
+
+	TSharedPtr< SInlineEditableTextBlock > InlineWidget;
+
+	STableRow<TSharedPtr<FMultiRigTreeElement>>::Construct(
+		STableRow<TSharedPtr<FMultiRigTreeElement>>::FArguments()
+		.ShowWires(true)
+		.Content()
+		[
+			SNew(SHorizontalBox)
+			+ SHorizontalBox::Slot()
+		.MaxWidth(18)
+		.FillWidth(1.0)
+		.HAlign(HAlign_Left)
+		.VAlign(VAlign_Center)
+		.Padding(FMargin(0.f, 0.f, 3.f, 0.f))
+		[
+
+
+			SNew(SButton)
+			.ButtonStyle(FEditorStyle::Get(), "NoBorder")
+			.OnClicked(this, &SMultiRigHierarchyItem::OnGetSelectedClicked)
+			[
+				SNew(SImage)
+				.Image_Lambda([this]() -> const FSlateBrush*
+					{
+						//if no key is set then it's the control rig so we get that based upon it's state
+						if (WeakRigTreeElement.Pin()->Data.Key.IsSet() == false)
+						{
+							if (UControlRig* ControlRig = WeakRigTreeElement.Pin()->Data.ControlRig.Get())
+							{
+								if (ControlRig->GetControlsVisible())
+								{
+									return (FEditorStyle::GetBrush("Level.VisibleIcon16x"));
+								}
+								else
+								{
+									return  (FEditorStyle::GetBrush("Level.NotVisibleIcon16x"));
+								}
+							}
+
+						}
+						return WeakRigTreeElement.Pin()->IconBrush;
+					})
+					.ColorAndOpacity_Lambda([this]()
+					{
+						return WeakRigTreeElement.Pin()->IconColor;
+					})
+			]
+		]
+		+ SHorizontalBox::Slot()
+			.AutoWidth()
+			.VAlign(VAlign_Center)
+			[
+				SAssignNew(InlineWidget, SInlineEditableTextBlock)
+				.Text(this, &SMultiRigHierarchyItem::GetName)
+			.MultiLine(false)
+			.ColorAndOpacity_Lambda([this]()
+				{
+					if (WeakRigTreeElement.IsValid())
+					{
+						return WeakRigTreeElement.Pin()->TextColor;
+					}
+					return FSlateColor::UseForeground();
+				})
+			]
+			], OwnerTable);
+
+}
+
+FReply SMultiRigHierarchyItem::OnGetSelectedClicked()
 {
-	if (CurrentControlRig.IsValid())
+	if (FMultiRigTreeElement* Element = WeakRigTreeElement.Pin().Get())
 	{
-		UObject* OldObject = CurrentControlRig.Get();
-		UObject* NewObject = OldToNewInstanceMap.FindRef(OldObject);
-		if (NewObject)
+		if (Element->Data.Key.IsSet() == false)
 		{
-			if (UControlRig* ControlRig = Cast<UControlRig>(NewObject))
+			if (UControlRig* ControlRig = Element->Data.ControlRig.Get())
 			{
-				NewControlRigSet(ControlRig);
+				FScopedTransaction ScopedTransaction(LOCTEXT("ToggleControlsVisibility", "Toggle Controls Visibility"), !GIsTransacting);
+				ControlRig->Modify();
+				ControlRig->ToggleControlsVisible();
+				return FReply::Handled();
 			}
 		}
 	}
+	return FReply::Unhandled();  //should flow to selection instead
 }
 
-void SControlRigOutlinerItem::NewControlRigSet(UControlRig* ControlRig)
+FText SMultiRigHierarchyItem::GetName() const
 {
-	if (CurrentControlRig.IsValid())
+	return (WeakRigTreeElement.Pin()->Data.GetName());
+}
+
+TPair<const FSlateBrush*, FSlateColor> SMultiRigHierarchyItem::GetBrushForElementType(const URigHierarchy* InHierarchy, const FMultiRigData& InData)
+{
+	const FSlateBrush* Brush = nullptr;
+	FSlateColor Color = FSlateColor::UseForeground();
+	if (InData.Key.IsSet())
 	{
-		(CurrentControlRig.Get())->ControlSelected().RemoveAll(this);
+		const FRigElementKey Key = InData.Key.GetValue();
+
+		switch (Key.Type)
+		{
+		case ERigElementType::Control:
+		{
+			Brush = FControlRigEditorStyle::Get().GetBrush("ControlRig.Tree.Control");
+			if (const FRigControlElement* Control = InHierarchy->Find<FRigControlElement>(InData.Key.GetValue()))
+			{
+				FLinearColor ShapeColor = Control->Settings.ShapeColor;
+				// ensure the alpha is always visible
+				ShapeColor.A = 1.f;
+				Color = FSlateColor(ShapeColor);
+			}
+			break;
+		}
+		case ERigElementType::Null:
+		{
+			Brush = FControlRigEditorStyle::Get().GetBrush("ControlRig.Tree.Null");
+			break;
+		}
+		case ERigElementType::Bone:
+		{
+			ERigBoneType BoneType = ERigBoneType::User;
+
+			if (InHierarchy)
+			{
+				const FRigBoneElement* BoneElement = InHierarchy->Find<FRigBoneElement>(InData.Key.GetValue());
+				if (BoneElement)
+				{
+					BoneType = BoneElement->BoneType;
+				}
+			}
+
+			switch (BoneType)
+			{
+			case ERigBoneType::Imported:
+			{
+				Brush = FControlRigEditorStyle::Get().GetBrush("ControlRig.Tree.BoneImported");
+				break;
+			}
+			case ERigBoneType::User:
+			default:
+			{
+				Brush = FControlRigEditorStyle::Get().GetBrush("ControlRig.Tree.BoneUser");
+				break;
+			}
+			}
+
+			break;
+		}
+		case ERigElementType::RigidBody:
+		{
+			Brush = FControlRigEditorStyle::Get().GetBrush("ControlRig.Tree.RigidBody");
+			break;
+		}
+		case ERigElementType::Reference:
+		{
+			Brush = FControlRigEditorStyle::Get().GetBrush("ControlRig.Tree.Socket");
+			break;
+		}
+		default:
+		{
+			break;
+		}
+		}
 	}
-	CurrentControlRig = ControlRig;
-	if (ControlRig)
+	else
 	{
-		ControlRig->ControlSelected().RemoveAll(this);
-		ControlRig->ControlSelected().AddRaw(this, &SControlRigOutlinerItem::HandleControlSelected);
+		if (UControlRig* ControlRig = InData.ControlRig.Get())
+		{
+			if (ControlRig->GetControlsVisible())
+			{
+				Brush = FEditorStyle::GetBrush("Level.VisibleIcon16x");
+			}
+			else
+			{
+				Brush = FEditorStyle::GetBrush("Level.NotVisibleIcon16x");
+			}
+		}
+		else
+		{
+			Brush = FEditorStyle::GetBrush("Level.NotVisibleIcon16x");
+		}
 	}
-	HierarchyTreeView->GetTreeView()->RefreshTreeView(true);
+
+	return TPair<const FSlateBrush*, FSlateColor>(Brush, Color);
 }
 
-SControlRigOutlinerItem::SControlRigOutlinerItem()
+//////////////////////////////////////////////////////////////
+/// SRigHierarchyTreeView
+///////////////////////////////////////////////////////////
+
+void SMultiRigHierarchyTreeView::Construct(const FArguments& InArgs)
 {
-	FCoreUObjectDelegates::OnObjectsReplaced.AddRaw(this, &SControlRigOutlinerItem::OnObjectsReplaced);
+	Delegates = InArgs._RigTreeDelegates;
+
+	STreeView<TSharedPtr<FMultiRigTreeElement>>::FArguments SuperArgs;
+	SuperArgs.TreeItemsSource(&RootElements);
+	SuperArgs.SelectionMode(ESelectionMode::Multi);
+	SuperArgs.OnGenerateRow(this, &SMultiRigHierarchyTreeView::MakeTableRowWidget, false);
+	SuperArgs.OnGetChildren(this, &SMultiRigHierarchyTreeView::HandleGetChildrenForTree);
+	SuperArgs.OnSelectionChanged(FOnMultiRigTreeSelectionChanged::CreateRaw(&Delegates, &FMultiRigTreeDelegates::HandleSelectionChanged));
+	SuperArgs.OnContextMenuOpening(Delegates.OnContextMenuOpening);
+	SuperArgs.OnMouseButtonClick(Delegates.OnMouseButtonClick);
+	SuperArgs.OnMouseButtonDoubleClick(Delegates.OnMouseButtonDoubleClick);
+	SuperArgs.OnSetExpansionRecursive(Delegates.OnSetExpansionRecursive);
+	SuperArgs.HighlightParentNodesForSelection(true);
+	SuperArgs.ItemHeight(24);
+	SuperArgs.AllowInvisibleItemSelection(true);  //without this we deselect everything when we filter or we collapse
+
+	SuperArgs.ShouldStackHierarchyHeaders_Lambda([]() -> bool {
+		return UControlRigEditorSettings::Get()->bShowStackedHierarchy;
+		});
+	SuperArgs.OnGeneratePinnedRow(this, &SMultiRigHierarchyTreeView::MakeTableRowWidget, true);
+	SuperArgs.MaxPinnedItems_Lambda([]() -> int32
+		{
+			return FMath::Max<int32>(1, UControlRigEditorSettings::Get()->MaxStackSize);
+		});
+
+	STreeView<TSharedPtr<FMultiRigTreeElement>>::Construct(SuperArgs);
 }
 
-void SControlRigOutlinerItem::HandleControlSelected(UControlRig* Subject, FRigControlElement* ControlElement, bool bSelected)
+TSharedPtr<FMultiRigTreeElement> SMultiRigHierarchyTreeView::FindElement(const FMultiRigData& InElementData, TSharedPtr<FMultiRigTreeElement> CurrentItem)
 {
-	const FRigElementKey Key = ControlElement->GetKey();
-	for (int32 RootIndex = 0; RootIndex < HierarchyTreeView->GetTreeView()->GetRootElements().Num(); ++RootIndex)
+	if (CurrentItem->Data == InElementData)
 	{
-		TSharedPtr<FRigTreeElement> Found = HierarchyTreeView->GetTreeView()->FindElement(Key, HierarchyTreeView->GetTreeView()->GetRootElements()[RootIndex]);
+		return CurrentItem;
+	}
+
+	for (int32 ChildIndex = 0; ChildIndex < CurrentItem->Children.Num(); ++ChildIndex)
+	{
+		TSharedPtr<FMultiRigTreeElement> Found = FindElement(InElementData, CurrentItem->Children[ChildIndex]);
 		if (Found.IsValid())
 		{
-			TGuardValue<bool> GuardRigHierarchyChanges(bIsChangingRigHierarchy, true);
+			return Found;
+		}
+	}
 
-			HierarchyTreeView->GetTreeView()->SetItemSelection(Found, bSelected, ESelectInfo::OnNavigation);
+	return TSharedPtr<FMultiRigTreeElement>();
+}
+
+bool SMultiRigHierarchyTreeView::AddElement(const FMultiRigData& InData, const FMultiRigData& InParentData)
+{
+	if (ElementMap.Contains(InData))
+	{
+		return false;
+	}
+
+	const FRigTreeDisplaySettings& Settings = Delegates.GetDisplaySettings();
 
-			TArray<TSharedPtr<FRigTreeElement>> SelectedItems = HierarchyTreeView->GetTreeView()->GetSelectedItems();
-			for (TSharedPtr<FRigTreeElement> SelectedItem : SelectedItems)
+	const FString FilteredString = Settings.FilterText.ToString();
+	if (FilteredString.IsEmpty() || !InData.IsValid())
+	{
+		TSharedPtr<FMultiRigTreeElement> NewItem = MakeShared<FMultiRigTreeElement>(InData, SharedThis(this), ERigTreeFilterResult::Shown);
+
+		if (InData.IsValid())
+		{
+			ElementMap.Add(InData, NewItem);
+			if (InParentData.IsValid())
 			{
-				HierarchyTreeView->GetTreeView()->SetExpansionRecursive(SelectedItem, false, true);
+				ParentMap.Add(InData, InParentData);
+				TSharedPtr<FMultiRigTreeElement>* FoundItem = ElementMap.Find(InParentData);
+				check(FoundItem);
+				FoundItem->Get()->Children.Add(NewItem);
 			}
+			else
+			{
+				RootElements.Add(NewItem);
+			}
+		}
+		else
+		{
+			RootElements.Add(NewItem);
+		}
+	}
+	else
+	{
+		FString FilteredStringUnderScores = FilteredString.Replace(TEXT(" "), TEXT("_"));
+		if (InData.GetName().ToString().Contains(FilteredString) || InData.GetName().ToString().Contains(FilteredStringUnderScores))
+		{
+			TSharedPtr<FMultiRigTreeElement> NewItem = MakeShared<FMultiRigTreeElement>(InData, SharedThis(this), ERigTreeFilterResult::Shown);
+			ElementMap.Add(InData, NewItem);
+			RootElements.Add(NewItem);
 
-			if (SelectedItems.Num() > 0)
+			if (!Settings.bFlattenHierarchyOnFilter && !Settings.bHideParentsOnFilter)
 			{
-				HierarchyTreeView->GetTreeView()->RequestScrollIntoView(SelectedItems.Last());
+				if (const URigHierarchy* Hierarchy = InData.GetHierarchy())
+				{
+					if (InData.Key.IsSet())
+					{
+						TSharedPtr<FMultiRigTreeElement> ChildItem = NewItem;
+						FRigElementKey ParentKey = Hierarchy->GetFirstParent(InData.Key.GetValue());
+						while (ParentKey.IsValid())
+						{
+							FMultiRigData ParentData(InData.ControlRig.Get(), ParentKey);
+							if (!ElementMap.Contains(ParentData))
+							{
+								TSharedPtr<FMultiRigTreeElement> ParentItem = MakeShared<FMultiRigTreeElement>(ParentData, SharedThis(this), ERigTreeFilterResult::ShownDescendant);
+								ElementMap.Add(ParentData, ParentItem);
+								RootElements.Add(ParentItem);
+
+								ReparentElement(ChildItem->Data, ParentData);
+
+								ChildItem = ParentItem;
+								ParentKey = Hierarchy->GetFirstParent(ParentKey);
+							}
+							else
+							{
+								ReparentElement(ChildItem->Data, ParentData);
+								break;
+							}
+						}
+					}
+				}
 			}
 		}
 	}
+
+	return true;
 }
 
-const URigHierarchy* SControlRigOutlinerItem::GetHierarchy() const
+bool SMultiRigHierarchyTreeView::AddElement(UControlRig* InControlRig, const FRigBaseElement* InElement)
 {
-	if (CurrentControlRig.IsValid())
+	check(InControlRig);
+	check(InElement);
+	
+	FMultiRigData Data(InControlRig, InElement->GetKey());
+
+	if (ElementMap.Contains(Data))
 	{
-		return CurrentControlRig.Get()->GetHierarchy();
+		return false;
 	}
-	return nullptr;
+
+	const FRigTreeDisplaySettings& Settings = Delegates.GetDisplaySettings();
+
+	switch (InElement->GetType())
+	{
+	case ERigElementType::Bone:
+	{
+		if (!Settings.bShowBones)
+		{
+			return false;
+		}
+
+		const FRigBoneElement* BoneElement = CastChecked<FRigBoneElement>(InElement);
+		if (!Settings.bShowImportedBones && BoneElement->BoneType == ERigBoneType::Imported)
+		{
+			return false;
+		}
+		break;
+	}
+	case ERigElementType::Null:
+	{
+		if (!Settings.bShowNulls)
+		{
+			return false;
+		}
+		break;
+	}
+	case ERigElementType::Control:
+	{
+		if (!Settings.bShowControls)
+		{
+			return false;
+		}
+		break;
+	}
+	case ERigElementType::RigidBody:
+	{
+		if (!Settings.bShowRigidBodies)
+		{
+			return false;
+		}
+		break;
+	}
+	case ERigElementType::Reference:
+	{
+		if (!Settings.bShowReferences)
+		{
+			return false;
+		}
+		break;
+	}
+	case ERigElementType::Curve:
+	{
+		return false;
+	}
+	default:
+	{
+		break;
+	}
+	}
+
+	FMultiRigData ParentData;
+	ParentData.ControlRig = InControlRig;
+
+	if (!AddElement(Data,ParentData))
+	{
+		return false;
+	}
+
+	if (ElementMap.Contains(Data))
+	{
+		if (const URigHierarchy* Hierarchy = InControlRig->GetHierarchy())
+		{
+			FRigElementKey ParentKey = Hierarchy->GetFirstParent(InElement->GetKey());
+			if (Settings.bShowDynamicHierarchy)
+			{
+				TArray<FRigElementWeight> ParentWeights = Hierarchy->GetParentWeightArray(InElement->GetKey());
+				if (ParentWeights.Num() > 0)
+				{
+					TArray<FRigElementKey> ParentKeys = Hierarchy->GetParents(InElement->GetKey());
+					check(ParentKeys.Num() == ParentWeights.Num());
+					for (int32 ParentIndex = 0; ParentIndex < ParentKeys.Num(); ParentIndex++)
+					{
+						if (ParentWeights[ParentIndex].IsAlmostZero())
+						{
+							continue;
+						}
+						ParentKey = ParentKeys[ParentIndex];
+						break;
+					}
+				}
+			}
+			if (ParentKey.IsValid())
+			{
+				if (const FRigBaseElement* ParentElement = Hierarchy->Find(ParentKey))
+				{
+					AddElement(InControlRig,ParentElement);
+
+					FMultiRigData NewParentData(InControlRig, ParentKey);
+
+					if (ElementMap.Contains(NewParentData))
+					{
+						ReparentElement(Data, NewParentData);
+					}
+				}
+			}
+		}
+	}
+
+	return true;
 }
 
-void SControlRigOutlinerItem::HandleSelectionChanged(TSharedPtr<FRigTreeElement> Selection, ESelectInfo::Type SelectInfo)
+
+bool SMultiRigHierarchyTreeView::ReparentElement(const FMultiRigData& InData, const FMultiRigData& InParentData)
 {
-	if (bIsChangingRigHierarchy)
+	if (!InData.IsValid() || InData == InParentData)
 	{
-		return;
+		return false;
 	}
 
-	const URigHierarchy* Hierarchy = GetHierarchy();
-	if (Hierarchy)
+	const FRigTreeDisplaySettings& Settings = Delegates.GetDisplaySettings();
+
+	TSharedPtr<FMultiRigTreeElement>* FoundItem = ElementMap.Find(InData);
+	if (FoundItem == nullptr)
 	{
-		URigHierarchyController* Controller = ((URigHierarchy*)Hierarchy)->GetController(true);
-		check(Controller);
+		return false;
+	}
+
+	if (!Settings.FilterText.IsEmpty() && Settings.bFlattenHierarchyOnFilter)
+	{
+		return false;
+	}
+
+	if (const FMultiRigData* ExistingParentKey = ParentMap.Find(InData))
+	{
+		if (*ExistingParentKey == InParentData)
+		{
+			return false;
+		}
+
+		if (TSharedPtr<FMultiRigTreeElement>* ExistingParent = ElementMap.Find(*ExistingParentKey))
+		{
+			(*ExistingParent)->Children.Remove(*FoundItem);
+		}
+
+		ParentMap.Remove(InData);
+	}
+	else
+	{
+		if (!InParentData.IsValid())
+		{
+			return false;
+		}
+
+		RootElements.Remove(*FoundItem);
+	}
+
+	if (InParentData.IsValid())
+	{
+		ParentMap.Add(InData, InParentData);
+
+		TSharedPtr<FMultiRigTreeElement>* FoundParent = ElementMap.Find(InParentData);
+		check(FoundParent);
+		FoundParent->Get()->Children.Add(*FoundItem);
+	}
+	else
+	{
+		RootElements.Add(*FoundItem);
+	}
+
+	return true;
+}
+
+bool SMultiRigHierarchyTreeView::RemoveElement(const FMultiRigData& InData)
+{
+	TSharedPtr<FMultiRigTreeElement>* FoundItem = ElementMap.Find(InData);
+	if (FoundItem == nullptr)
+	{
+		return false;
+	}
+
+	FMultiRigData EmptyParent(nullptr, FRigElementKey());
+	ReparentElement(InData, EmptyParent);
+
+	RootElements.Remove(*FoundItem);
+	return ElementMap.Remove(InData) > 0;
+}
+
+void SMultiRigHierarchyTreeView::RefreshTreeView(bool bRebuildContent)
+{
+	TMap<FMultiRigData, bool> ExpansionState;
+
+	if (bRebuildContent)
+	{
+		for (TPair<FMultiRigData, TSharedPtr<FMultiRigTreeElement>> Pair : ElementMap)
+		{
+			ExpansionState.FindOrAdd(Pair.Key) = IsItemExpanded(Pair.Value);
+		}
+
+		// internally save expansion states before rebuilding the tree, so the states can be restored later
+		SaveAndClearSparseItemInfos();
+
+		RootElements.Reset();
+		ElementMap.Reset();
+		ParentMap.Reset();
+	}
+
+	if (bRebuildContent)
+	{
+		for (const TWeakObjectPtr<UControlRig>& ControlRigPtr : ControlRigs)
+		{
+			if (UControlRig* ControlRig  = ControlRigPtr.Get())
+			{
+				FMultiRigData Empty(nullptr, FRigElementKey());
+				FMultiRigData CRData;
+				CRData.ControlRig = ControlRig; //leave key unset so it's valid
+
+				AddElement(CRData, Empty);
+				if (const URigHierarchy* Hierarchy = ControlRig->GetHierarchy())
+				{
+					Hierarchy->Traverse([&](FRigBaseElement* Element, bool& bContinue)
+						{
+							AddElement(ControlRig,Element);
+							bContinue = true;
+						});
+
+					// expand all elements upon the initial construction of the tree
+					if (ExpansionState.Num() == 0)
+					{
+						for (TSharedPtr<FMultiRigTreeElement> RootElement : RootElements)
+						{
+							SetExpansionRecursive(RootElement, false, true);
+						}
+					}
+					else if (ExpansionState.Num() < ElementMap.Num())
+					{
+						for (const TPair<FMultiRigData, TSharedPtr<FMultiRigTreeElement>>& Element : ElementMap)
+						{
+							if (!ExpansionState.Contains(Element.Key))
+							{
+								SetItemExpansion(Element.Value, true);
+							}
+						}
+					}
+
+					for (const auto& Pair : ElementMap)
+					{
+						RestoreSparseItemInfos(Pair.Value);
+					}
+
+				}
+			}
+		}
+	}
+	else
+	{
+		if (RootElements.Num() > 0)
+		{
+			// elements may be added at the end of the list after a spacer element
+			// we need to remove the spacer element and re-add it at the end
+			RootElements.RemoveAll([](TSharedPtr<FMultiRigTreeElement> InElement)
+				{
+					return (InElement.Get()->Data.ControlRig == nullptr && InElement.Get()->Data.Key == FRigElementKey());
+				});
+		}
+	}
+
+	RequestTreeRefresh();
+	{
+		ClearSelection();
 
-		TGuardValue<bool> GuardRigHierarchyChanges(bIsChangingRigHierarchy, true);
-		const TArray<FRigElementKey> NewSelection = HierarchyTreeView->GetTreeView()->GetSelectedKeys();
-		//due to how Sequencer Tree View will redo selection on next tick if we aren't keeping or toggling selection we need to clear it out
-		if (FSlateApplication::Get().GetModifierKeys().IsShiftDown() == false || FSlateApplication::Get().GetModifierKeys().IsControlDown() == false)
+		for (const TWeakObjectPtr<UControlRig>& ControlRigPtr : ControlRigs)
 		{
-			if (ControlRigEditMode)
+			if (UControlRig* ControlRig = ControlRigPtr.Get())
 			{
-				TMap<UControlRig*, TArray<FRigElementKey>> SelectedControls;
-				ControlRigEditMode->GetAllSelectedControls(SelectedControls);
-				for (TPair<UControlRig*, TArray<FRigElementKey>>& CurrentSelection : SelectedControls)
+				if (const URigHierarchy* Hierarchy = ControlRig->GetHierarchy())
 				{
-					if (CurrentSelection.Key)
+
+					TArray<FRigElementKey> Selection = Hierarchy->GetSelectedKeys();
+					for (const FRigElementKey& Key : Selection)
 					{
-						CurrentSelection.Key->ClearControlSelection();
+						for (int32 RootIndex = 0; RootIndex < RootElements.Num(); ++RootIndex)
+						{
+							FMultiRigData Data(ControlRig, Key);
+							TSharedPtr<FMultiRigTreeElement> Found = FindElement(Data, RootElements[RootIndex]);
+							if (Found.IsValid())
+							{
+								SetItemSelection(Found, true, ESelectInfo::OnNavigation);
+							}
+						}
 					}
 				}
 			}
 		}
-		if (!Controller->SetSelection(NewSelection))
+	}
+}
+
+void SMultiRigHierarchyTreeView::SetExpansionRecursive(TSharedPtr<FMultiRigTreeElement> InElement, bool bTowardsParent,
+	bool bShouldBeExpanded)
+{
+	SetItemExpansion(InElement, bShouldBeExpanded);
+
+	if (bTowardsParent)
+	{
+		if (const FMultiRigData* ParentKey = ParentMap.Find(InElement->Data))
+		{
+			if (TSharedPtr<FMultiRigTreeElement>* ParentItem = ElementMap.Find(*ParentKey))
+			{
+				SetExpansionRecursive(*ParentItem, bTowardsParent, bShouldBeExpanded);
+			}
+		}
+	}
+	else
+	{
+		for (int32 ChildIndex = 0; ChildIndex < InElement->Children.Num(); ++ChildIndex)
 		{
-			return;
+			SetExpansionRecursive(InElement->Children[ChildIndex], bTowardsParent, bShouldBeExpanded);
 		}
 	}
 }
 
-SControlRigOutlinerItem::~SControlRigOutlinerItem()
+TSharedRef<ITableRow> SMultiRigHierarchyTreeView::MakeTableRowWidget(TSharedPtr<FMultiRigTreeElement> InItem,
+	const TSharedRef<STableViewBase>& OwnerTable, bool bPinned)
 {
-	if (CurrentControlRig.IsValid())
+	const FRigTreeDisplaySettings& Settings = Delegates.GetDisplaySettings();
+	return InItem->MakeTreeRowWidget(OwnerTable, InItem.ToSharedRef(), SharedThis(this), Settings, bPinned);
+}
+
+void SMultiRigHierarchyTreeView::HandleGetChildrenForTree(TSharedPtr<FMultiRigTreeElement> InItem,
+	TArray<TSharedPtr<FMultiRigTreeElement>>& OutChildren)
+{
+	OutChildren = InItem.Get()->Children;
+}
+
+TArray<FMultiRigData> SMultiRigHierarchyTreeView::GetSelectedData() const
+{
+	TArray<FMultiRigData> Keys;
+	TArray<TSharedPtr<FMultiRigTreeElement>> SelectedElements = GetSelectedItems();
+	for (const TSharedPtr<FMultiRigTreeElement>& SelectedElement : SelectedElements)
 	{
-		(CurrentControlRig.Get())->ControlSelected().RemoveAll(this);
+		Keys.Add(SelectedElement->Data);
 	}
-	CurrentControlRig = nullptr;
-	FCoreUObjectDelegates::OnObjectsReplaced.RemoveAll(this);
+	return Keys;
 }
 
-void SControlRigOutliner::SetEditMode(FControlRigEditMode& InEditMode)
+TSharedPtr<FMultiRigTreeElement> SMultiRigHierarchyTreeView::FindItemAtPosition(FVector2D InScreenSpacePosition) const
 {
-	ModeTools = InEditMode.GetModeManager();
-	if (FControlRigEditMode* EditMode = static_cast<FControlRigEditMode*>(ModeTools->GetActiveMode(FControlRigEditMode::ModeName)))
+	if (ItemsPanel.IsValid() && ItemsSource != nullptr)
 	{
-		EditMode->OnControlRigAddedOrRemoved().AddRaw(this, &SControlRigOutliner::HandleControlAdded);
+		const FGeometry MyGeometry = ItemsPanel->GetCachedGeometry();
+		FArrangedChildren ArrangedChildren(EVisibility::Visible);
+		ItemsPanel->ArrangeChildren(MyGeometry, ArrangedChildren, true);
+
+		const int32 Index = ItemsPanel->FindChildUnderPosition(ArrangedChildren, InScreenSpacePosition);
+		if (ItemsSource->IsValidIndex(Index))
+		{
+			return ItemsSource->operator[](Index);
+		}
+	}
+	return TSharedPtr<FMultiRigTreeElement>();
+}
+
+TArray<URigHierarchy*> SMultiRigHierarchyTreeView::GetHierarchy() const
+{
+	TArray<URigHierarchy*> RigHierarchy;
+	for (const TWeakObjectPtr<UControlRig>& ControlRig : ControlRigs)
+	{
+		if (ControlRig.IsValid())
+		{
+			RigHierarchy.Add(ControlRig.Get()->GetHierarchy());
+		}
+	}
+	return RigHierarchy;
+}
+
+void SMultiRigHierarchyTreeView::SetControlRigs(TArrayView < TWeakObjectPtr<UControlRig>>& InControlRigs)
+{
+	ControlRigs.SetNum(0);
+	for (TWeakObjectPtr<UControlRig>& ControlRig : InControlRigs)
+	{
+		if (ControlRig.IsValid())
+		{
+			ControlRigs.Add(ControlRig.Get());
+		}
 	}
+	RefreshTreeView(true);
 }
+
+//////////////////////////////////////////////////////////////
+/// SSearchableMultiRigHierarchyTreeView
+///////////////////////////////////////////////////////////
+
+void SSearchableMultiRigHierarchyTreeView::Construct(const FArguments& InArgs)
+{
+	FMultiRigTreeDelegates TreeDelegates = InArgs._RigTreeDelegates;
+	SuperGetRigTreeDisplaySettings = TreeDelegates.OnGetDisplaySettings;
+
+	TreeDelegates.OnGetDisplaySettings.BindSP(this, &SSearchableMultiRigHierarchyTreeView::GetDisplaySettings);
+
+	ChildSlot
+		[
+			SNew(SVerticalBox)
+			+ SVerticalBox::Slot()
+				.AutoHeight()
+				.VAlign(VAlign_Top)
+				.HAlign(HAlign_Fill)
+				.Padding(0.0f)
+				[
+					SNew(SSearchBox)
+					.InitialText(InArgs._InitialFilterText)
+				.OnTextChanged(this, &SSearchableMultiRigHierarchyTreeView::OnFilterTextChanged)
+				]
+
+			+ SVerticalBox::Slot()
+				.AutoHeight()
+				.VAlign(VAlign_Top)
+				.HAlign(HAlign_Fill)
+				.Padding(0.0f, 0.0f)
+				[
+					SNew(SScrollBox)
+					+ SScrollBox::Slot()
+					[
+						SNew(SBorder)
+						.Padding(2.0f)
+						.BorderImage(FEditorStyle::GetBrush("SCSEditor.TreePanel"))
+						[
+							SAssignNew(TreeView, SMultiRigHierarchyTreeView)
+							.RigTreeDelegates(TreeDelegates)
+						]
+					]
+				]
+		];
+}
+
+const FRigTreeDisplaySettings& SSearchableMultiRigHierarchyTreeView::GetDisplaySettings()
+{
+	if (SuperGetRigTreeDisplaySettings.IsBound())
+	{
+		Settings = SuperGetRigTreeDisplaySettings.Execute();
+	}
+	Settings.FilterText = FilterText;
+	return Settings;
+}
+
+void SSearchableMultiRigHierarchyTreeView::OnFilterTextChanged(const FText& SearchText)
+{
+	FilterText = SearchText;
+	GetTreeView()->RefreshTreeView();
+}
+
+//////////////////////////////////////////////////////////////
+/// SControlRigOutliner
+///////////////////////////////////////////////////////////
+
 void SControlRigOutliner::Construct(const FArguments& InArgs, FControlRigEditMode& InEditMode)
 {
+	bIsChangingRigHierarchy = false;
+
+	DisplaySettings.bShowBones = false;
+	DisplaySettings.bShowControls = true;
+	DisplaySettings.bShowNulls = false;
+	DisplaySettings.bShowReferences = false;
+	DisplaySettings.bShowRigidBodies = false;
+	DisplaySettings.bHideParentsOnFilter = true;
+	DisplaySettings.bFlattenHierarchyOnFilter = true;
+
+	FMultiRigTreeDelegates RigTreeDelegates;
+	RigTreeDelegates.OnGetDisplaySettings = FOnGetRigTreeDisplaySettings::CreateSP(this, &SControlRigOutliner::GetDisplaySettings);
+	RigTreeDelegates.OnSelectionChanged = FOnMultiRigTreeSelectionChanged::CreateSP(this, &SControlRigOutliner::HandleSelectionChanged);
+
+	
 	ChildSlot
 		[
 			SNew(SScrollBox)
 			+ SScrollBox::Slot()
 			[
-				SAssignNew(MainBoxPtr, SVerticalBox)
+			SNew(SVerticalBox)
+			+ SVerticalBox::Slot()
+				.AutoHeight()
+				[
+					SAssignNew(PickerExpander, SExpandableArea)
+					.InitiallyCollapsed(false)
+					//.AreaTitle(AreaTitle)
+					//.AreaTitleFont(FEditorStyle::GetFontStyle("DetailsView.CategoryFontStyle"))
+					.BorderBackgroundColor(FLinearColor(.6f, .6f, .6f))
+					.BodyContent()
+					[
+						SAssignNew(HierarchyTreeView, SSearchableMultiRigHierarchyTreeView)
+						.RigTreeDelegates(RigTreeDelegates)
+					]
+					
+				]
+		
 			]
 		];
-
 	SetEditMode(InEditMode);
-	Rebuild();
 }
 
-SControlRigOutliner::~SControlRigOutliner()
+void SControlRigOutliner::OnObjectsReplaced(const TMap<UObject*, UObject*>& OldToNewInstanceMap)
 {
-	if (FControlRigEditMode* EditMode = static_cast<FControlRigEditMode*>(ModeTools->GetActiveMode(FControlRigEditMode::ModeName)))
+	//if there's a control rig recreate the tree, controls may have changed
+	bool bNewControlRig = false;
+	for (const TPair<UObject*, UObject*>& Pair : OldToNewInstanceMap)
 	{
-		EditMode->OnControlRigAddedOrRemoved().RemoveAll(this);
+		if (Pair.Key->IsA<UControlRig>() && Pair.Value->IsA<UControlRig>())
+		{
+			bNewControlRig = false;
+			break;
+		}
+	}
+	if (bNewControlRig)
+	{
+		HierarchyTreeView->GetTreeView()->RefreshTreeView(true);
 	}
-	//base class handles control rig related cleanup
 }
 
+SControlRigOutliner::SControlRigOutliner()
+{
+	FCoreUObjectDelegates::OnObjectsReplaced.AddRaw(this, &SControlRigOutliner::OnObjectsReplaced);
+}
 
-void SControlRigOutliner::HandleControlAdded(UControlRig* ControlRig, bool bIsAdded)
+SControlRigOutliner::~SControlRigOutliner()
 {
-	Rebuild();
+	FCoreUObjectDelegates::OnObjectsReplaced.RemoveAll(this);
+}
+
+void SControlRigOutliner::HandleControlSelected(UControlRig* Subject, FRigControlElement* ControlElement, bool bSelected)
+{
+	FControlRigBaseDockableView::HandleControlSelected(Subject, ControlElement, bSelected);
+	const FRigElementKey Key = ControlElement->GetKey();
+	FMultiRigData Data(Subject, Key);
+	for (int32 RootIndex = 0; RootIndex < HierarchyTreeView->GetTreeView()->GetRootElements().Num(); ++RootIndex)
+	{
+		TSharedPtr<FMultiRigTreeElement> Found = HierarchyTreeView->GetTreeView()->FindElement(Data, HierarchyTreeView->GetTreeView()->GetRootElements()[RootIndex]);
+		if (Found.IsValid())
+		{
+			TGuardValue<bool> GuardRigHierarchyChanges(bIsChangingRigHierarchy, true);
+
+			HierarchyTreeView->GetTreeView()->SetItemSelection(Found, bSelected, ESelectInfo::OnNavigation);
+
+			TArray<TSharedPtr<FMultiRigTreeElement>> SelectedItems = HierarchyTreeView->GetTreeView()->GetSelectedItems();
+			for (TSharedPtr<FMultiRigTreeElement> SelectedItem : SelectedItems)
+			{
+				HierarchyTreeView->GetTreeView()->SetExpansionRecursive(SelectedItem, false, true);
+			}
 
+			if (SelectedItems.Num() > 0)
+			{
+				HierarchyTreeView->GetTreeView()->RequestScrollIntoView(SelectedItems.Last());
+			}
+		}
+	}
 }
-void SControlRigOutliner::Rebuild()
+
+
+void SControlRigOutliner::HandleSelectionChanged(TSharedPtr<FMultiRigTreeElement> Selection, ESelectInfo::Type SelectInfo)
 {
-	MainBoxPtr->ClearChildren();
-	if (FControlRigEditMode* EditMode = static_cast<FControlRigEditMode*>(ModeTools->GetActiveMode(FControlRigEditMode::ModeName)))
+	if (bIsChangingRigHierarchy)
 	{
-		TArray<UControlRig*> ControlRigs = EditMode->GetControlRigsArray(false /*bIsVisible*/);
-		for (UControlRig* ControlRig : ControlRigs)
+		return;
+	}
+	const TArray<FMultiRigData> NewSelection = HierarchyTreeView->GetTreeView()->GetSelectedData();
+	TMap<UControlRig*, TArray<FRigElementKey>> SelectedRigAndKeys;
+	for (const FMultiRigData& Data : NewSelection)
+	{
+		if (Data.ControlRig.IsValid() && (Data.Key.IsSet() && Data.Key.GetValue() != FRigElementKey()))
 		{
-			if (ControlRig)
+			if (SelectedRigAndKeys.Find(Data.ControlRig.Get()) == nullptr)
 			{
-				MainBoxPtr->AddSlot()
-				.AutoHeight()
-				[
-					SNew(SControlRigOutlinerItem)
-					.ControlRig(ControlRig)
-					.EditMode(EditMode)
-				];
+				SelectedRigAndKeys.Add(Data.ControlRig.Get());
 			}
+			SelectedRigAndKeys[Data.ControlRig.Get()].Add(Data.Key.GetValue());
 		}
 	}
+	TGuardValue<bool> GuardRigHierarchyChanges(bIsChangingRigHierarchy, true);
+	//due to how Sequencer Tree View will redo selection on next tick if we aren't keeping or toggling selection we need to clear it out
+	if (FSlateApplication::Get().GetModifierKeys().IsShiftDown() == false || FSlateApplication::Get().GetModifierKeys().IsControlDown() == false)
+	{
+		if (FControlRigEditMode* EditMode = static_cast<FControlRigEditMode*>(ModeTools->GetActiveMode(FControlRigEditMode::ModeName)))
+		{
+			TMap<UControlRig*, TArray<FRigElementKey>> SelectedControls;
+			EditMode->GetAllSelectedControls(SelectedControls);
+			for (TPair<UControlRig*, TArray<FRigElementKey>>& CurrentSelection : SelectedControls)
+			{
+				if (CurrentSelection.Key)
+				{
+					CurrentSelection.Key->ClearControlSelection();
+				}
+			}
+		}
+	}
+
+	for(TPair<UControlRig*, TArray<FRigElementKey>>& RigAndKeys: SelectedRigAndKeys)
+	{ 
+		const URigHierarchy* Hierarchy = RigAndKeys.Key->GetHierarchy();
+		if (Hierarchy)
+		{
+			URigHierarchyController* Controller = ((URigHierarchy*)Hierarchy)->GetController(true);
+			check(Controller);
+			if (!Controller->SetSelection(RigAndKeys.Value))
+			{
+				return;
+			}
+		}
+	}
+}
 
+void SControlRigOutliner::SetEditMode(FControlRigEditMode& InEditMode)
+{
+	FControlRigBaseDockableView::SetEditMode(InEditMode);
+	ModeTools = InEditMode.GetModeManager();
+	if (FControlRigEditMode* EditMode = static_cast<FControlRigEditMode*>(ModeTools->GetActiveMode(FControlRigEditMode::ModeName)))
+	{
+		TArrayView<TWeakObjectPtr<UControlRig>> ControlRigs = EditMode->GetControlRigs();
+		HierarchyTreeView->GetTreeView()->SetControlRigs(ControlRigs); //will refresh tree
+	}
+}
+
+void SControlRigOutliner::HandleControlAdded(UControlRig* ControlRig, bool bIsAdded)
+{
+	FControlRigBaseDockableView::HandleControlAdded(ControlRig, bIsAdded);
+	if (FControlRigEditMode* EditMode = static_cast<FControlRigEditMode*>(ModeTools->GetActiveMode(FControlRigEditMode::ModeName)))
+	{
+		EditMode->OnControlRigAddedOrRemoved().AddRaw(this, &SControlRigOutliner::HandleControlAdded);
+		TArrayView<TWeakObjectPtr<UControlRig>> ControlRigs = EditMode->GetControlRigs();
+		HierarchyTreeView->GetTreeView()->SetControlRigs(ControlRigs); //will refresh tree
+	}
 }
 
 
diff --git a/Engine/Plugins/Animation/ControlRig/Source/ControlRigEditor/Private/EditMode/SControlRigOutliner.h b/Engine/Plugins/Animation/ControlRig/Source/ControlRigEditor/Private/EditMode/SControlRigOutliner.h
--- a/Engine/Plugins/Animation/ControlRig/Source/ControlRigEditor/Private/EditMode/SControlRigOutliner.h
+++ b/Engine/Plugins/Animation/ControlRig/Source/ControlRigEditor/Private/EditMode/SControlRigOutliner.h
@@ -19,58 +19,252 @@ class SExpandableArea;
 class SSearchableRigHierarchyTreeView;
 class UControlRig;
 
-class SControlRigOutlinerItem : public SCompoundWidget
+class SMultiRigHierarchyTreeView;
+class SMultiRigHierarchyItem;
+class FMultiRigTreeElement;
+class FMultiRigData;
+
+DECLARE_DELEGATE_RetVal_TwoParams(bool, FOnMultiRigTreeCompareKeys, const FMultiRigData& /*A*/, const FMultiRigData& /*B*/);
+
+typedef STreeView<TSharedPtr<FMultiRigTreeElement>>::FOnSelectionChanged FOnMultiRigTreeSelectionChanged;
+typedef STreeView<TSharedPtr<FMultiRigTreeElement>>::FOnMouseButtonClick FOnMultiRigTreeMouseButtonClick;
+typedef STreeView<TSharedPtr<FMultiRigTreeElement>>::FOnMouseButtonDoubleClick FOnMultiRigTreeMouseButtonDoubleClick;
+typedef STreeView<TSharedPtr<FMultiRigTreeElement>>::FOnSetExpansionRecursive FOnMultiRigTreeSetExpansionRecursive;
+
+uint32 GetTypeHash(const FMultiRigData& Data);
+
+struct CONTROLRIGEDITOR_API FMultiRigTreeDelegates
+{
+	FOnGetRigTreeDisplaySettings OnGetDisplaySettings;
+	FOnMultiRigTreeSelectionChanged OnSelectionChanged;
+	FOnContextMenuOpening OnContextMenuOpening;
+	FOnMultiRigTreeMouseButtonClick OnMouseButtonClick;
+	FOnMultiRigTreeMouseButtonDoubleClick OnMouseButtonDoubleClick;
+	FOnMultiRigTreeSetExpansionRecursive OnSetExpansionRecursive;
+	FOnRigTreeCompareKeys OnCompareKeys;
+
+	FMultiRigTreeDelegates()
+	{
+		bIsChangingRigHierarchy = false;
+	}
+
+
+	FORCEINLINE const FRigTreeDisplaySettings& GetDisplaySettings() const
+	{
+		if (OnGetDisplaySettings.IsBound())
+		{
+			return OnGetDisplaySettings.Execute();
+		}
+		return DefaultDisplaySettings;
+	}
+
+	FORCEINLINE void HandleSelectionChanged(TSharedPtr<FMultiRigTreeElement> Selection, ESelectInfo::Type SelectInfo)
+	{
+		if (bIsChangingRigHierarchy)
+		{
+			return;
+		}
+		TGuardValue<bool> Guard(bIsChangingRigHierarchy, true);
+		OnSelectionChanged.ExecuteIfBound(Selection, SelectInfo);
+	}
+
+	static FRigTreeDisplaySettings DefaultDisplaySettings;
+	bool bIsChangingRigHierarchy;
+};
+
+/** Data for the tree*/
+class FMultiRigData
+{
+
+public:
+	FMultiRigData() {};
+	FMultiRigData(UControlRig* InControlRig, FRigElementKey InKey) : ControlRig(InControlRig), Key(InKey) {};
+	FText GetName() const;
+	bool operator == (const FMultiRigData & Other) const;
+	bool IsValid() const;
+	URigHierarchy* GetHierarchy() const;
+public:
+	TWeakObjectPtr<UControlRig> ControlRig;
+	TOptional<FRigElementKey> Key;
+};
+
+/** An item in the tree */
+class FMultiRigTreeElement : public TSharedFromThis<FMultiRigTreeElement>
+{
+public:
+	FMultiRigTreeElement(const FMultiRigData& InData, TWeakPtr<SMultiRigHierarchyTreeView> InTreeView,ERigTreeFilterResult InFilterResult);
+public:
+	/** Element Data to display */
+	FMultiRigData Data;
+	TArray<TSharedPtr<FMultiRigTreeElement>> Children;
+
+	TSharedRef<ITableRow> MakeTreeRowWidget(const TSharedRef<STableViewBase>& InOwnerTable, TSharedRef<FMultiRigTreeElement> InRigTreeElement, TSharedPtr<SMultiRigHierarchyTreeView> InTreeView, const FRigTreeDisplaySettings& InSettings, bool bPinned);
+
+	void RefreshDisplaySettings(const URigHierarchy* InHierarchy, const FRigTreeDisplaySettings& InSettings);
+
+	/** The current filter result */
+	ERigTreeFilterResult FilterResult;
+
+	/** The brush to use when rendering an icon */
+	const FSlateBrush* IconBrush;;
+
+	/** The color to use when rendering an icon */
+	FSlateColor IconColor;
+
+	/** The color to use when rendering the label text */
+	FSlateColor TextColor;
+};
+
+class SMultiRigHierarchyItem : public STableRow<TSharedPtr<FMultiRigTreeElement>>
+{
+public:
+
+	void Construct(const FArguments& InArgs, const TSharedRef<STableViewBase>& OwnerTable, TSharedRef<FMultiRigTreeElement> InRigTreeElement, TSharedPtr<SMultiRigHierarchyTreeView> InTreeView, const FRigTreeDisplaySettings& InSettings, bool bPinned);
+	static TPair<const FSlateBrush*, FSlateColor> GetBrushForElementType(const URigHierarchy* InHierarchy, const FMultiRigData& InData);
+
+private:
+	TWeakPtr<FMultiRigTreeElement> WeakRigTreeElement;
+	FMultiRigTreeDelegates Delegates;
+
+	FText GetName() const;
+	FReply OnGetSelectedClicked();
+
+};
+
+class SMultiRigHierarchyTreeView : public STreeView<TSharedPtr<FMultiRigTreeElement>>
 {
-	SLATE_BEGIN_ARGS(SControlRigOutlinerItem){}
-	SLATE_ARGUMENT(UControlRig*, ControlRig)
-	SLATE_ARGUMENT(FControlRigEditMode*, EditMode)
+public:
+
+	SLATE_BEGIN_ARGS(SMultiRigHierarchyTreeView) {}
+	SLATE_ARGUMENT(FMultiRigTreeDelegates, RigTreeDelegates)
 	SLATE_END_ARGS()
-	SControlRigOutlinerItem();
-	~SControlRigOutlinerItem();
 
 	void Construct(const FArguments& InArgs);
+	virtual ~SMultiRigHierarchyTreeView() {}
+
+	virtual FReply OnFocusReceived(const FGeometry& MyGeometry, const FFocusEvent& InFocusEvent) override
+	{
+		FReply Reply = STreeView<TSharedPtr<FMultiRigTreeElement>>::OnFocusReceived(MyGeometry, InFocusEvent);
+		return Reply;
+	}
+
+	/** Save a snapshot of the internal map that tracks item expansion before tree reconstruction */
+	void SaveAndClearSparseItemInfos()
+	{
+		// Only save the info if there is something to save (do not overwrite info with an empty map)
+		if (!SparseItemInfos.IsEmpty())
+		{
+			OldSparseItemInfos = SparseItemInfos;
+		}
+		ClearExpandedItems();
+	}
+
+	/** Restore the expansion infos map from the saved snapshot after tree reconstruction */
+	void RestoreSparseItemInfos(TSharedPtr<FMultiRigTreeElement> ItemPtr)
+	{
+		for (const auto& Pair : OldSparseItemInfos)
+		{
+			if (Pair.Key->Data == ItemPtr->Data)
+			{
+				// the SparseItemInfos now reference the new element, but keep the same expansion state
+				SparseItemInfos.Add(ItemPtr, Pair.Value);
+				break;
+			}
+		}
+	}
+
+	static TSharedPtr<FMultiRigTreeElement> FindElement(const FMultiRigData& InData, TSharedPtr<FMultiRigTreeElement> CurrentItem);
+
+	bool AddElement(const FMultiRigData& InData, const FMultiRigData& InParentData);
+	bool AddElement(UControlRig* InControlRig, const FRigBaseElement* InElement);
+	bool ReparentElement(const FMultiRigData& InData, const FMultiRigData& InParentData);
+	bool RemoveElement(const FMultiRigData& InData);
+	void RefreshTreeView(bool bRebuildContent = true);
+	void SetExpansionRecursive(TSharedPtr<FMultiRigTreeElement> InElement, bool bTowardsParent, bool bShouldBeExpanded);
+	TSharedRef<ITableRow> MakeTableRowWidget(TSharedPtr<FMultiRigTreeElement> InItem, const TSharedRef<STableViewBase>& OwnerTable, bool bPinned);
+	void HandleGetChildrenForTree(TSharedPtr<FMultiRigTreeElement> InItem, TArray<TSharedPtr<FMultiRigTreeElement>>& OutChildren);
+
+	TArray<FMultiRigData> GetSelectedData() const;
+	const TArray<TSharedPtr<FMultiRigTreeElement>>& GetRootElements() const { return RootElements; }
+	FMultiRigTreeDelegates& GetTreeDelegates() { return Delegates; }
+
+	TSharedPtr<FMultiRigTreeElement> FindItemAtPosition(FVector2D InScreenSpacePosition) const;
+	TArray<URigHierarchy*> GetHierarchy() const;
+	void SetControlRigs(TArrayView < TWeakObjectPtr<UControlRig>>& InControlRigs);
+
 private:
-	void HandleControlSelected(UControlRig* Subject, FRigControlElement* InControl, bool bSelected);
-	void OnObjectsReplaced(const TMap<UObject*, UObject*>& OldToNewInstanceMap);
-	void NewControlRigSet(UControlRig* ControlRig);
 
-	const URigHierarchy* GetHierarchy() const;
-	void HandleSelectionChanged(TSharedPtr<FRigTreeElement> Selection, ESelectInfo::Type SelectInfo);
+	/** A temporary snapshot of the SparseItemInfos in STreeView, used during RefreshTreeView() */
+	TSparseItemMap OldSparseItemInfos;
 
-	//visibility button
-	const FSlateBrush* GetVisibilityBrushForElement() const;
-	FReply OnToggleVisibility();
-	bool VisibilityToggleEnabled() const;
+	/** Backing array for tree view */
+	TArray<TSharedPtr<FMultiRigTreeElement>> RootElements;
 
+	/** A map for looking up items based on their key */
+	TMap<FMultiRigData, TSharedPtr<FMultiRigTreeElement>> ElementMap;
 
-	/** Hierarchy picker for controls*/
-	TSharedPtr<SSearchableRigHierarchyTreeView> HierarchyTreeView;
-	FRigTreeDisplaySettings DisplaySettings;
-	const FRigTreeDisplaySettings& GetDisplaySettings() const { return DisplaySettings; }
-	bool bIsChangingRigHierarchy = false;
-	TSharedPtr<SExpandableArea> PickerExpander;
+	/** A map for looking up a parent based on their key */
+	TMap<FMultiRigData, FMultiRigData> ParentMap;
+
+	FMultiRigTreeDelegates Delegates;
 
-	TWeakObjectPtr<UControlRig> CurrentControlRig;
-	FControlRigEditMode* ControlRigEditMode = nullptr;
+	friend class SRigHierarchy;
 
+	TArray <TWeakObjectPtr<UControlRig>> ControlRigs;
 };
 
-class SControlRigOutliner: public SCompoundWidget
+class SSearchableMultiRigHierarchyTreeView : public SCompoundWidget
 {
+public:
+
+	SLATE_BEGIN_ARGS(SSearchableMultiRigHierarchyTreeView) {}
+	SLATE_ARGUMENT(FMultiRigTreeDelegates, RigTreeDelegates)
+		SLATE_ARGUMENT(FText, InitialFilterText)
+	SLATE_END_ARGS()
+
+	void Construct(const FArguments& InArgs);
+	virtual ~SSearchableMultiRigHierarchyTreeView() {}
+	TSharedRef<SMultiRigHierarchyTreeView> GetTreeView() const { return TreeView.ToSharedRef(); }
+	const FRigTreeDisplaySettings& GetDisplaySettings();
 
-	SLATE_BEGIN_ARGS(SControlRigOutliner)
-	{}
+private:
+
+	void OnFilterTextChanged(const FText& SearchText);
+
+	FOnGetRigTreeDisplaySettings SuperGetRigTreeDisplaySettings;
+	FText FilterText;
+	FRigTreeDisplaySettings Settings;
+	TSharedPtr<SMultiRigHierarchyTreeView> TreeView;
+};
+
+class SControlRigOutliner : public FControlRigBaseDockableView, public SCompoundWidget
+{
+	SLATE_BEGIN_ARGS(SControlRigOutliner){}
 	SLATE_END_ARGS()
+	SControlRigOutliner();
 	~SControlRigOutliner();
 
 	void Construct(const FArguments& InArgs, FControlRigEditMode& InEditMode);
-	void SetEditMode(FControlRigEditMode& InEditMode);
 
+	//FControlRigBaseDockableView overrides
+	virtual void SetEditMode(FControlRigEditMode& InEditMode) override;
 private:
-	void HandleControlAdded(UControlRig* ControlRig, bool bIsAdded);
-	void Rebuild();
-	TSharedPtr<SVerticalBox> MainBoxPtr;
-	FEditorModeTools* ModeTools = nullptr;
+	virtual void HandleControlAdded(UControlRig* ControlRig, bool bIsAdded) override;
+	virtual void HandleControlSelected(UControlRig* Subject, FRigControlElement* InControl, bool bSelected) override;
+
+	void OnObjectsReplaced(const TMap<UObject*, UObject*>& OldToNewInstanceMap);
+
+	const URigHierarchy* GetHierarchy() const;
+	void HandleSelectionChanged(TSharedPtr<FMultiRigTreeElement> Selection, ESelectInfo::Type SelectInfo);
+
+	/** Hierarchy picker for controls*/
+	TSharedPtr<SSearchableMultiRigHierarchyTreeView> HierarchyTreeView;
+	FRigTreeDisplaySettings DisplaySettings;
+	const FRigTreeDisplaySettings& GetDisplaySettings() const { return DisplaySettings; }
+	bool bIsChangingRigHierarchy = false;
+	TSharedPtr<SExpandableArea> PickerExpander;
+
 
 };
 
+
diff --git a/Engine/Plugins/Animation/ControlRig/Source/ControlRigEditor/Private/EditMode/SControlRigSpacePicker.cpp b/Engine/Plugins/Animation/ControlRig/Source/ControlRigEditor/Private/EditMode/SControlRigSpacePicker.cpp
--- a/Engine/Plugins/Animation/ControlRig/Source/ControlRigEditor/Private/EditMode/SControlRigSpacePicker.cpp
+++ b/Engine/Plugins/Animation/ControlRig/Source/ControlRigEditor/Private/EditMode/SControlRigSpacePicker.cpp
@@ -105,25 +105,37 @@ SControlRigSpacePicker::~SControlRigSpacePicker()
 	//base class handles control rig related cleanup
 }
 
-
-void SControlRigSpacePicker::HandleControlAdded(UControlRig* ControlRig, bool bIsAdded)
-{
-	FControlRigBaseDockableView::HandleControlAdded(ControlRig, bIsAdded);
-}
-
-void SControlRigSpacePicker::NewControlRigSet(UControlRig* ControlRig)
+UControlRig* SControlRigSpacePicker::GetControlRig()
 {
-	FControlRigBaseDockableView::NewControlRigSet(ControlRig);
+	TArray<UControlRig*> ControlRigs = GetControlRigs();
+	for(UControlRig* ControlRig: ControlRigs)
+	{
+		if (ControlRig)
+		{
+			TArray<FRigElementKey> SelectedControls = ControlRig->GetHierarchy()->GetSelectedKeys(ERigElementType::Control);
+			if(SelectedControls.Num() >0)
+			{
+				return ControlRig;
+			}
+		}
+	}
+	return nullptr;
 }
 
 void SControlRigSpacePicker::HandleControlSelected(UControlRig* Subject, FRigControlElement* ControlElement, bool bSelected)
 {
+	FControlRigBaseDockableView::HandleControlSelected(Subject, ControlElement, bSelected);
 	if (UControlRig* ControlRig = GetControlRig())
 	{
 		// get the selected controls
 		TArray<FRigElementKey> SelectedControls = ControlRig->GetHierarchy()->GetSelectedKeys(ERigElementType::Control);
 		SpacePickerWidget->SetControls(ControlRig->GetHierarchy(), SelectedControls);
 	}
+	else //set nothing
+	{
+		TArray<FRigElementKey> SelectedControls;
+		SpacePickerWidget->SetControls(ControlRig->GetHierarchy(), SelectedControls);
+	}
 }
 
 const FRigControlElementCustomization* SControlRigSpacePicker::HandleGetControlElementCustomization(URigHierarchy* InHierarchy, const FRigElementKey& InControlKey)
diff --git a/Engine/Plugins/Animation/ControlRig/Source/ControlRigEditor/Private/EditMode/SControlRigSpacePicker.h b/Engine/Plugins/Animation/ControlRig/Source/ControlRigEditor/Private/EditMode/SControlRigSpacePicker.h
--- a/Engine/Plugins/Animation/ControlRig/Source/ControlRigEditor/Private/EditMode/SControlRigSpacePicker.h
+++ b/Engine/Plugins/Animation/ControlRig/Source/ControlRigEditor/Private/EditMode/SControlRigSpacePicker.h
@@ -30,11 +30,6 @@ class SControlRigSpacePicker : public SCompoundWidget, public FControlRigBaseDoc
 
 private:
 	virtual void HandleControlSelected(UControlRig* Subject, FRigControlElement* InControl, bool bSelected) override;
-	virtual void HandleControlAdded(UControlRig* ControlRig, bool bIsAdded) override;
-	virtual void NewControlRigSet(UControlRig* ControlRig) override;
-
-	const URigHierarchy* GetHierarchy() const;
-	void HandleSelectionChanged(TSharedPtr<FRigTreeElement> Selection, ESelectInfo::Type SelectInfo);
 
 	/** Space picker widget*/
 	TSharedPtr<SRigSpacePickerWidget> SpacePickerWidget;
@@ -46,5 +41,8 @@ private:
 	FReply HandleAddSpaceClicked();
 	FReply OnBakeControlsToNewSpaceButtonClicked();
 
+	//for now picker works off of one ControlRig, this function gets the first control rig with a selection
+	UControlRig* GetControlRig();
+
 };
 
diff --git a/Engine/Plugins/Animation/ControlRig/Source/ControlRigEditor/Private/Sequencer/ControlRigParameterTrackEditor.cpp b/Engine/Plugins/Animation/ControlRig/Source/ControlRigEditor/Private/Sequencer/ControlRigParameterTrackEditor.cpp
--- a/Engine/Plugins/Animation/ControlRig/Source/ControlRigEditor/Private/Sequencer/ControlRigParameterTrackEditor.cpp
+++ b/Engine/Plugins/Animation/ControlRig/Source/ControlRigEditor/Private/Sequencer/ControlRigParameterTrackEditor.cpp
@@ -236,7 +236,7 @@ FControlRigParameterTrackEditor::FControlRigParameterTrackEditor(TSharedRef<ISeq
 							if (UControlRig* ControlRig = ControlRigPtr.Get())
 							{
 								if (ControlRig->GetObjectBinding() && ControlRig->GetObjectBinding()->GetBoundObject() == nullptr)
-					{
+								{
 									bRequestEvaluate = true;
 									break;
 								}
@@ -299,7 +299,6 @@ FControlRigParameterTrackEditor::FControlRigParameterTrackEditor(TSharedRef<ISeq
 								if (ControlRigEditMode)
 								{
 									ControlRigEditMode->ReplaceControlRig(OldControlRig, *NewControlRig);
-									ControlRigEditMode->AddControlRigObject(*NewControlRig, GetSequencer());
 									
 									UControlRig* PtrNewControlRig = *NewControlRig;
 
@@ -1926,12 +1925,12 @@ void FControlRigParameterTrackEditor::OnSelectionChanged(TArray<UMovieSceneTrack
 			for (TPair<UControlRig*, TArray<FRigElementKey>>& SelectedControl : AllSelectedControls)
 			{
 				ControlRig = SelectedControl.Key;
-			if (ControlRig)
-			{
-				ControlRig->ClearControlSelection();
+				if (ControlRig)
+				{
+					ControlRig->ClearControlSelection();
+				}
 			}
 		}
-		}
 		for (UMovieSceneTrack* Track : InTracks)
 		{
 			UMovieSceneControlRigParameterTrack* CRTrack = Cast<UMovieSceneControlRigParameterTrack>(Track);
@@ -2053,12 +2052,12 @@ void FControlRigParameterTrackEditor::SelectRigsAndControls(UControlRig* Control
 		for (TPair<UControlRig*, TArray<FRigElementKey>>& Selection : SelectedControls)
 		{
 			ControlRig = Selection.Key;
-		if (ControlRig)
-		{
-			ControlRig->ClearControlSelection();
+			if (ControlRig)
+			{
+				ControlRig->ClearControlSelection();
+			}
 		}
 	}
-	}
 	for (TPair<UControlRig*, TSet<FName>> Pair : RigsAndControls)
 	{
 		if (Pair.Key != ControlRig)
@@ -2523,10 +2522,13 @@ void FControlRigParameterTrackEditor::GetControlRigKeys(UControlRig* InControlRi
 			}
 			case ERigControlType::Vector2D:
 			{
+				//use translation x,y for key masks for vector2d
+				bool bKeyX = bSetKey && EnumHasAnyFlags(ChannelsToKey, EControlRigContextChannelToKey::TranslationX);;
+				bool bKeyY = bSetKey && EnumHasAnyFlags(ChannelsToKey, EControlRigContextChannelToKey::TranslationY);;
 				FVector3f Val = ControlValue.Get<FVector3f>();
 				pChannelIndex->GeneratedKeyIndex = OutGeneratedKeys.Num();
-				OutGeneratedKeys.Add(FMovieSceneChannelValueSetter::Create<FMovieSceneFloatChannel>(ChannelIndex++, Val.X, bSetKey));
-				OutGeneratedKeys.Add(FMovieSceneChannelValueSetter::Create<FMovieSceneFloatChannel>(ChannelIndex++, Val.Y, bSetKey));
+				OutGeneratedKeys.Add(FMovieSceneChannelValueSetter::Create<FMovieSceneFloatChannel>(ChannelIndex++, Val.X, bKeyX));
+				OutGeneratedKeys.Add(FMovieSceneChannelValueSetter::Create<FMovieSceneFloatChannel>(ChannelIndex++, Val.Y, bKeyY));
 				break;
 			}
 			case ERigControlType::Position:
