Control Rig: Disable skeletal mesh tracks for bindings on either the skeletal mesh component and the actor since you can have it on either.

diff --git a/Engine/Plugins/Animation/ControlRig/Source/ControlRigEditor/Private/Sequencer/ControlRigParameterTrackEditor.cpp b/Engine/Plugins/Animation/ControlRig/Source/ControlRigEditor/Private/Sequencer/ControlRigParameterTrackEditor.cpp
--- a/Engine/Plugins/Animation/ControlRig/Source/ControlRigEditor/Private/Sequencer/ControlRigParameterTrackEditor.cpp
+++ b/Engine/Plugins/Animation/ControlRig/Source/ControlRigEditor/Private/Sequencer/ControlRigParameterTrackEditor.cpp
@@ -1045,27 +1045,44 @@ void FControlRigParameterTrackEditor::BakeToControlRig(UClass* InClass, FGuid Ob
 						BakeSettings->bReduceKeys, BakeSettings->Tolerance);
 
 					//Turn Off Any Skeletal Animation Tracks
-					const FMovieSceneBinding* Binding = OwnerMovieScene->FindBinding(ObjectBinding);
-					if (Binding)
+					TArray<UMovieSceneSkeletalAnimationTrack*> SkelAnimationTracks;
+					if (const FMovieSceneBinding* Binding = OwnerMovieScene->FindBinding(ObjectBinding))
 					{
 						for (UMovieSceneTrack* MovieSceneTrack : Binding->GetTracks())
 						{
 							if (UMovieSceneSkeletalAnimationTrack* SkelTrack = Cast<UMovieSceneSkeletalAnimationTrack>(MovieSceneTrack))
 							{
-								SkelTrack->Modify();
-								//can't just turn off the track so need to mute the sections
-								const TArray<UMovieSceneSection*>& Sections = SkelTrack->GetAllSections();
-								for (UMovieSceneSection* Section : Sections)
-								{
-									if (Section)
-									{
-										Section->TryModify();
-										Section->SetIsActive(false);
-									}
-								}
+								SkelAnimationTracks.Add(SkelTrack);
+							}
+						}
+					}
+					FGuid SkelMeshGuid = GetSequencer()->FindObjectId(*SkelMeshComp, GetSequencer()->GetFocusedTemplateID());
+					if (const FMovieSceneBinding* Binding = OwnerMovieScene->FindBinding(SkelMeshGuid))
+					{
+						for (UMovieSceneTrack* MovieSceneTrack : Binding->GetTracks())
+						{
+							if (UMovieSceneSkeletalAnimationTrack* SkelTrack = Cast<UMovieSceneSkeletalAnimationTrack>(MovieSceneTrack))
+							{
+								SkelAnimationTracks.Add(SkelTrack);
+							}
+						}
+					}
+
+					for (UMovieSceneSkeletalAnimationTrack* SkelTrack : SkelAnimationTracks)
+					{
+						SkelTrack->Modify();
+						//can't just turn off the track so need to mute the sections
+						const TArray<UMovieSceneSection*>& Sections = SkelTrack->GetAllSections();
+						for (UMovieSceneSection* Section : Sections)
+						{
+							if (Section)
+							{
+								Section->TryModify();
+								Section->SetIsActive(false);
 							}
 						}
 					}
+
 					//Finish Setup
 					if (ControlRigEditMode)
 					{
