Control Rig: Set from Selected Controls should only be active if there are selected controls

diff --git a/Engine/Plugins/Animation/ControlRig/Source/ControlRigEditor/Private/Sequencer/ControlRigParameterTrackEditor.cpp b/Engine/Plugins/Animation/ControlRig/Source/ControlRigEditor/Private/Sequencer/ControlRigParameterTrackEditor.cpp
--- a/Engine/Plugins/Animation/ControlRig/Source/ControlRigEditor/Private/Sequencer/ControlRigParameterTrackEditor.cpp
+++ b/Engine/Plugins/Animation/ControlRig/Source/ControlRigEditor/Private/Sequencer/ControlRigParameterTrackEditor.cpp
@@ -3816,7 +3816,10 @@ void FControlRigParameterSection::BuildSectionContextMenu(FMenuBuilder& MenuBuil
 				LOCTEXT("SetFromSelectedControls", "Set From Selected Controls"),
 				LOCTEXT("SetFromSelectedControls_ToolTip", "Set active channels from the current control selection"),
 				FSlateIcon(),
-				FUIAction(FExecuteAction::CreateLambda([=] { ShowSelectedControlsChannels(); }))
+				FUIAction(
+					FExecuteAction::CreateLambda([=] { ShowSelectedControlsChannels(); }),
+					FCanExecuteAction::CreateLambda([=] { return ControlRig->CurrentControlSelection().Num() > 0; } )
+				)
 			);
 
 			MenuBuilder.AddMenuEntry(
