Sequencer: Control Rig:  Controls may not get selected when selected in Sequencer, merge issue.
#trivial

diff --git a/Engine/Plugins/Animation/ControlRig/Source/ControlRigEditor/Private/Sequencer/ControlRigParameterTrackEditor.cpp b/Engine/Plugins/Animation/ControlRig/Source/ControlRigEditor/Private/Sequencer/ControlRigParameterTrackEditor.cpp
--- a/Engine/Plugins/Animation/ControlRig/Source/ControlRigEditor/Private/Sequencer/ControlRigParameterTrackEditor.cpp
+++ b/Engine/Plugins/Animation/ControlRig/Source/ControlRigEditor/Private/Sequencer/ControlRigParameterTrackEditor.cpp
@@ -1939,6 +1939,7 @@ void FControlRigParameterTrackEditor::OnSelectionChanged(TArray<UMovieSceneTrack
 		if (ControlRigEditMode)
 		{
 			TMap<UControlRig*, TArray<FRigElementKey>> AllSelectedControls;
+			ControlRigEditMode->GetAllSelectedControls(AllSelectedControls);
 			for (TPair<UControlRig*, TArray<FRigElementKey>>& SelectedControl : AllSelectedControls)
 			{
 				ControlRig = SelectedControl.Key;
