Sequencer: Control Rig: Issue with moving CR components do construction scripts replacing CR rig multiple times. Previous CL mostly fixed this, just needed to clean up a few things.
#trivial

diff --git a/Engine/Plugins/Animation/ControlRig/Source/ControlRigEditor/Private/EditMode/ControlRigBaseDockableView.cpp b/Engine/Plugins/Animation/ControlRig/Source/ControlRigEditor/Private/EditMode/ControlRigBaseDockableView.cpp
--- a/Engine/Plugins/Animation/ControlRig/Source/ControlRigEditor/Private/EditMode/ControlRigBaseDockableView.cpp
+++ b/Engine/Plugins/Animation/ControlRig/Source/ControlRigEditor/Private/EditMode/ControlRigBaseDockableView.cpp
@@ -43,7 +43,9 @@ void FControlRigBaseDockableView::SetEditMode(FControlRigEditMode& InEditMode)
 	ModeTools = InEditMode.GetModeManager();
 	if (FControlRigEditMode* EditMode = static_cast<FControlRigEditMode*>(ModeTools->GetActiveMode(FControlRigEditMode::ModeName)))
 	{
+		EditMode->OnControlRigAddedOrRemoved().RemoveAll(this);
 		EditMode->OnControlRigAddedOrRemoved().AddRaw(this, &FControlRigBaseDockableView::HandleControlAdded);
+		EditMode->OnControlRigSelected().RemoveAll(this);
 		EditMode->OnControlRigSelected().AddRaw(this, &FControlRigBaseDockableView::HandlElementSelected);
 	}
 }
diff --git a/Engine/Plugins/Animation/ControlRig/Source/ControlRigEditor/Private/EditMode/ControlRigEditMode.cpp b/Engine/Plugins/Animation/ControlRig/Source/ControlRigEditor/Private/EditMode/ControlRigEditMode.cpp
--- a/Engine/Plugins/Animation/ControlRig/Source/ControlRigEditor/Private/EditMode/ControlRigEditMode.cpp
+++ b/Engine/Plugins/Animation/ControlRig/Source/ControlRigEditor/Private/EditMode/ControlRigEditMode.cpp
@@ -2873,7 +2873,7 @@ void FControlRigEditMode::RecreateControlShapeActors(const TArray<FRigElementKey
 			}
 		}
 	}
-	else
+	else if(ControlRigsToRecreate.Num() > 0)
 	{
 		for (UControlRig* ControlRig : ControlRigsToRecreate)
 		{
@@ -3791,7 +3791,7 @@ void FControlRigEditMode::DestroyShapesActors(UControlRig* ControlRig)
 			}
 		}
 		ControlRigShapeActors.Reset();
-
+		ControlRigsToRecreate.Reset();
 		if (OnWorldCleanupHandle.IsValid())
 		{
 			FWorldDelegates::OnWorldCleanup.Remove(OnWorldCleanupHandle);
@@ -3799,6 +3799,7 @@ void FControlRigEditMode::DestroyShapesActors(UControlRig* ControlRig)
 	}
 	else
 	{
+		ControlRigsToRecreate.Remove(ControlRig);
 		const TArray<AControlRigShapeActor*>* ShapeActors = ControlRigShapeActors.Find(ControlRig);
 		if (ShapeActors)
 		{
diff --git a/Engine/Plugins/Animation/ControlRig/Source/ControlRigEditor/Private/EditMode/SControlRigOutliner.cpp b/Engine/Plugins/Animation/ControlRig/Source/ControlRigEditor/Private/EditMode/SControlRigOutliner.cpp
--- a/Engine/Plugins/Animation/ControlRig/Source/ControlRigEditor/Private/EditMode/SControlRigOutliner.cpp
+++ b/Engine/Plugins/Animation/ControlRig/Source/ControlRigEditor/Private/EditMode/SControlRigOutliner.cpp
@@ -1142,7 +1142,6 @@ void SControlRigOutliner::HandleControlAdded(UControlRig* ControlRig, bool bIsAd
 	FControlRigBaseDockableView::HandleControlAdded(ControlRig, bIsAdded);
 	if (FControlRigEditMode* EditMode = static_cast<FControlRigEditMode*>(ModeTools->GetActiveMode(FControlRigEditMode::ModeName)))
 	{
-		EditMode->OnControlRigAddedOrRemoved().AddRaw(this, &SControlRigOutliner::HandleControlAdded);
 		TArrayView<TWeakObjectPtr<UControlRig>> ControlRigs = EditMode->GetControlRigs();
 		HierarchyTreeView->GetTreeView()->SetControlRigs(ControlRigs); //will refresh tree
 	}
