CIS FIX: Deferenced possible nullptr

diff --git a/Engine/Plugins/Animation/ControlRig/Source/ControlRigEditor/Private/EditMode/SControlRigSpacePicker.cpp b/Engine/Plugins/Animation/ControlRig/Source/ControlRigEditor/Private/EditMode/SControlRigSpacePicker.cpp
--- a/Engine/Plugins/Animation/ControlRig/Source/ControlRigEditor/Private/EditMode/SControlRigSpacePicker.cpp
+++ b/Engine/Plugins/Animation/ControlRig/Source/ControlRigEditor/Private/EditMode/SControlRigSpacePicker.cpp
@@ -134,7 +134,7 @@ void SControlRigSpacePicker::HandleControlSelected(UControlRig* Subject, FRigCon
 	else //set nothing
 	{
 		TArray<FRigElementKey> SelectedControls;
-		SpacePickerWidget->SetControls(ControlRig->GetHierarchy(), SelectedControls);
+		SpacePickerWidget->SetControls(nullptr, SelectedControls);
 	}
 }
 
