Control Rig: Fix undo/redo deleting function with local variables crashes

diff --git a/Engine/Plugins/Animation/ControlRig/Source/ControlRigDeveloper/Private/ControlRigBlueprint.cpp b/Engine/Plugins/Animation/ControlRig/Source/ControlRigDeveloper/Private/ControlRigBlueprint.cpp
--- a/Engine/Plugins/Animation/ControlRig/Source/ControlRigDeveloper/Private/ControlRigBlueprint.cpp
+++ b/Engine/Plugins/Animation/ControlRig/Source/ControlRigDeveloper/Private/ControlRigBlueprint.cpp
@@ -3181,11 +3181,20 @@ void UControlRigBlueprint::HandleModifiedEvent(ERigVMGraphNotifType InNotifType,
 				{
 					if (InNotifType == ERigVMGraphNotifType::NodeAdded)
 					{
+						// If the controller for this graph already exist, make sure it is referencing the correct graph
+						if (TObjectPtr<URigVMController>* Controller = Controllers.Find(CollapseNode->GetContainedGraph()))
+						{
+							(*Controller).Get()->SetGraph(CollapseNode->GetContainedGraph());
+						}
+						
 						CreateEdGraphForCollapseNodeIfNeeded(CollapseNode);
 					}
 					else
 					{
 						bNotifForOthersPending = !RemoveEdGraphForCollapseNode(CollapseNode, true);
+
+						// Cannot remove from the Controllers array because we would lose the action stack on that graph
+						// Controllers.Remove(CollapseNode->GetContainedGraph();
 					}
 
 					ClearTransientControls();
diff --git a/Engine/Plugins/Animation/ControlRig/Source/ControlRigDeveloper/Private/Graph/ControlRigGraphSchema.cpp b/Engine/Plugins/Animation/ControlRig/Source/ControlRigDeveloper/Private/Graph/ControlRigGraphSchema.cpp
--- a/Engine/Plugins/Animation/ControlRig/Source/ControlRigDeveloper/Private/Graph/ControlRigGraphSchema.cpp
+++ b/Engine/Plugins/Animation/ControlRig/Source/ControlRigDeveloper/Private/Graph/ControlRigGraphSchema.cpp
@@ -205,6 +205,15 @@ void FControlRigGraphSchemaAction_LocalVar::DeleteVariable()
 {
 	if (UControlRigGraph* Graph = Cast<UControlRigGraph>(GetVariableScope()))
 	{
+#if WITH_EDITOR
+
+		if (GEditor)
+		{
+			GEditor->CancelTransaction(0);
+		}
+	
+#endif
+		
 		Graph->GetController()->RemoveLocalVariable(GetVariableName(), true, true);
 	}
 }
@@ -268,6 +277,14 @@ UEdGraphNode* FControlRigGraphSchemaAction_PromoteToVariable::PerformAction(UEdG
 
 	if(bLocalVariable)
 	{
+#if WITH_EDITOR
+
+		if (GEditor)
+		{
+			GEditor->CancelTransaction(0);
+		}
+	
+#endif
 		const FRigVMGraphVariableDescription VariableDescription = Controller->AddLocalVariable(
 			*ModelPin->GetPinPath(),
 			ModelPin->GetCPPType(),
diff --git a/Engine/Plugins/Animation/ControlRig/Source/ControlRigDeveloper/Public/ControlRigBlueprint.h b/Engine/Plugins/Animation/ControlRig/Source/ControlRigDeveloper/Public/ControlRigBlueprint.h
--- a/Engine/Plugins/Animation/ControlRig/Source/ControlRigDeveloper/Public/ControlRigBlueprint.h
+++ b/Engine/Plugins/Animation/ControlRig/Source/ControlRigDeveloper/Public/ControlRigBlueprint.h
@@ -377,8 +377,9 @@ protected:
 	UPROPERTY(AssetRegistrySearchable)
 	TArray<FRigVMReferenceNodeData> FunctionReferenceNodeData;
 
+	// Related URigVMGraph path and controller
 	UPROPERTY(BlueprintReadOnly, transient, Category = "VM")
-	TMap<TObjectPtr<URigVMGraph>, TObjectPtr<URigVMController>> Controllers;
+	TMap<FSoftObjectPath, TObjectPtr<URigVMController>> Controllers;
 
 #if WITH_EDITORONLY_DATA
 
diff --git a/Engine/Source/Developer/RigVMDeveloper/Private/RigVMModel/RigVMController.cpp b/Engine/Source/Developer/RigVMDeveloper/Private/RigVMModel/RigVMController.cpp
--- a/Engine/Source/Developer/RigVMDeveloper/Private/RigVMModel/RigVMController.cpp
+++ b/Engine/Source/Developer/RigVMDeveloper/Private/RigVMModel/RigVMController.cpp
@@ -5974,6 +5974,11 @@ bool URigVMController::RemoveNode(URigVMNode* InNode, bool bSetupUndoRedo, bool
 		DestroyObject(InjectionInfo);
 	}
 
+	if (URigVMCollapseNode* CollapseNode = Cast<URigVMCollapseNode>(InNode))
+	{
+		DestroyObject(CollapseNode->GetContainedGraph());
+	}
+
 	DestroyObject(InNode);
 
 	if (bSetupUndoRedo)
@@ -12348,6 +12353,12 @@ bool URigVMController::IsValidGraph() const
 		ReportError(TEXT("Controller does not have a graph associated - use SetGraph / set_graph."));
 		return false;
 	}
+
+	if (!IsValid(Graph))
+	{
+		return false;
+	}
+
 	return true;
 }
 
@@ -15539,6 +15550,7 @@ void URigVMController::DestroyObject(UObject* InObjectToDestroy)
 {
 	RenameObject(InObjectToDestroy, nullptr, GetTransientPackage());
 	InObjectToDestroy->RemoveFromRoot();
+	InObjectToDestroy->MarkAsGarbage();
 }
 
 void URigVMController::AddNodePin(URigVMNode* InNode, URigVMPin* InPin)
