Sequencer: Control Rig: When a control is initalized it may clear out selection which can break tools that need to evaluate the rig at different frames.

diff --git a/Engine/Plugins/Animation/ControlRig/Source/ControlRig/Private/Sequencer/MovieSceneControlRigParameterTemplate.cpp b/Engine/Plugins/Animation/ControlRig/Source/ControlRig/Private/Sequencer/MovieSceneControlRigParameterTemplate.cpp
--- a/Engine/Plugins/Animation/ControlRig/Source/ControlRig/Private/Sequencer/MovieSceneControlRigParameterTemplate.cpp
+++ b/Engine/Plugins/Animation/ControlRig/Source/ControlRig/Private/Sequencer/MovieSceneControlRigParameterTemplate.cpp
@@ -562,6 +562,19 @@ namespace MovieScene
 } // namespace MovieScene
 } // namespace UE
 
+//since initialization can blow up selection, may need to just reselect, used in a few places
+static void SelectControls(UControlRig* ControlRig, TArray<FName>& SelectedNames)
+{
+	if (ControlRig)
+	{
+		ControlRig->ClearControlSelection();
+		for (const FName& Name : SelectedNames)
+		{
+			ControlRig->SelectControl(Name, true);
+		}
+	}
+}
+
 void FControlRigBindingHelper::BindToSequencerInstance(UControlRig* ControlRig)
 {
 	check(ControlRig);
@@ -576,10 +589,17 @@ void FControlRigBindingHelper::BindToSequencerInstance(UControlRig* ControlRig)
 				{
 					AnimInstance->RecalcRequiredBones();
 					AnimInstance->AddControlRigTrack(ControlRig->GetUniqueID(), ControlRig);
+					//initialization can blow up selection
+					TArray<FName> SelectedControls = ControlRig->CurrentControlSelection();
 					ControlRig->Initialize();
 					ControlRig->RequestInit();
 					ControlRig->SetBoneInitialTransformsFromSkeletalMeshComponent(SkeletalMeshComponent, true);
 					ControlRig->Evaluate_AnyThread();
+					TArray<FName> NewSelectedControls = ControlRig->CurrentControlSelection();
+					if (SelectedControls != NewSelectedControls)
+					{
+						SelectControls(ControlRig, SelectedControls);
+					}
 				}
 			}
 		}
@@ -939,7 +959,13 @@ struct FControlRigParameterPreAnimatedTokenProducer : IMovieScenePreAnimatedToke
 					}
 					else
 					{
+						TArray<FName> SelectedControls = ControlRig->CurrentControlSelection();
 						ControlRig->Initialize();
+						TArray<FName> NewSelectedControls = ControlRig->CurrentControlSelection();
+						if (SelectedControls != NewSelectedControls)
+						{
+							SelectControls(ControlRig,SelectedControls);
+						}
 					}
 				}
 				else if (USkeletalMeshComponent* SkeletalMeshComponent = Cast<USkeletalMeshComponent>(ControlRig->GetObjectBinding()->GetBoundObject()))
@@ -1084,6 +1110,7 @@ struct FControlRigParameterExecutionToken : IMovieSceneExecutionToken
 				if (!ControlRig->GetObjectBinding()->GetBoundObject())
 				{
 					ControlRig->GetObjectBinding()->BindToObject(BoundObjects[0].Get());
+					TArray<FName> SelectedControls = ControlRig->CurrentControlSelection();
 					ControlRig->Initialize();
 					if (ControlRig->IsA<UFKControlRig>())
 					{
@@ -1093,6 +1120,11 @@ struct FControlRigParameterExecutionToken : IMovieSceneExecutionToken
 							Track->ReplaceControlRig(ControlRig, true);
 						}
 					}
+					TArray<FName> NewSelectedControls = ControlRig->CurrentControlSelection();
+					if (SelectedControls != NewSelectedControls)
+					{
+						SelectControls(ControlRig, SelectedControls);
+					}
 				}
 				// ensure that pre animated state is saved, must be done before bind
 				Player.SavePreAnimatedState(*ControlRig, FMovieSceneControlRigParameterTemplate::GetAnimTypeID(), FControlRigParameterPreAnimatedTokenProducer(Operand.SequenceID));
