Sequencer: Anim tools, don't throlttle spinners so tween and blends are intereactive. Also moved a CL from main to make sure we don't auto-expand control selections

diff --git a/Engine/Plugins/Animation/ControlRig/Source/ControlRigEditor/Private/EditMode/SControlRigControlViews.cpp b/Engine/Plugins/Animation/ControlRig/Source/ControlRigEditor/Private/EditMode/SControlRigControlViews.cpp
--- a/Engine/Plugins/Animation/ControlRig/Source/ControlRigEditor/Private/EditMode/SControlRigControlViews.cpp
+++ b/Engine/Plugins/Animation/ControlRig/Source/ControlRigEditor/Private/EditMode/SControlRigControlViews.cpp
@@ -339,6 +339,7 @@ void SControlRigPoseView::Construct(const FArguments& InArgs)
 
 								SNew(SSpinBox<float>)
 								// Only allow spinning if we have a single value
+								.PreventThrottling(true)
 								.Value(this, &SControlRigPoseView::OnGetPoseBlendValue)
 								.ToolTipText(LOCTEXT("BlendTooltip", "Blend between current pose and pose asset. Use Ctrl drag for under and over shoot."))
 								.MinValue(0.0f)
diff --git a/Engine/Plugins/Animation/ControlRig/Source/ControlRigEditor/Private/EditMode/SControlRigTweenWidget.cpp b/Engine/Plugins/Animation/ControlRig/Source/ControlRigEditor/Private/EditMode/SControlRigTweenWidget.cpp
--- a/Engine/Plugins/Animation/ControlRig/Source/ControlRigEditor/Private/EditMode/SControlRigTweenWidget.cpp
+++ b/Engine/Plugins/Animation/ControlRig/Source/ControlRigEditor/Private/EditMode/SControlRigTweenWidget.cpp
@@ -51,6 +51,7 @@ void SControlRigTweenWidget::Construct(const FArguments& InArgs)
 			.HAlign(HAlign_Center)
 			[
 				SNew(SSpinBox<float>)
+				.PreventThrottling(true)
 				.Value(this, &SControlRigTweenWidget::OnGetPoseBlendValueFloat)
 				.ToolTipText(LOCTEXT("TweenTooltip", "Key at current frame between previous(-1.0) and next(1.0) poses. Use Ctrl drag for under and over shoot."))
 				.MinValue(-2.0f)
