Sequencer: Snapper: Just  select one control when snapping to parent.
#trival

diff --git a/Engine/Plugins/Animation/ControlRig/Source/ControlRigEditor/Private/EditMode/SControlRigSnapper.cpp b/Engine/Plugins/Animation/ControlRig/Source/ControlRigEditor/Private/EditMode/SControlRigSnapper.cpp
--- a/Engine/Plugins/Animation/ControlRig/Source/ControlRigEditor/Private/EditMode/SControlRigSnapper.cpp
+++ b/Engine/Plugins/Animation/ControlRig/Source/ControlRigEditor/Private/EditMode/SControlRigSnapper.cpp
@@ -487,11 +487,21 @@ FControlRigSnapperSelection SControlRigSnapper::GetSelection(bool bGetAll)
 				FControlRigForWorldTransforms ControlRigAndSelection;
 				ControlRigAndSelection.ControlRig = ControlRig;
 				ControlRigAndSelection.ControlNames = SelectedControls;
-				Selection.ControlRigs.Add(ControlRigAndSelection);
 				if (bGetAll == false)
 				{
+					//make sure to only use first control
+					if (SelectedControls.Num() > 1)
+					{
+						ControlRigAndSelection.ControlNames.SetNum(0);
+						ControlRigAndSelection.ControlNames.Add(SelectedControls[0]);
+					}
+					Selection.ControlRigs.Add(ControlRigAndSelection);
 					return Selection;
 				}
+				else
+				{
+					Selection.ControlRigs.Add(ControlRigAndSelection);
+				}
 			}
 		}
 	}
