Fix non-visible control actor selection

diff --git a/Engine/Plugins/Animation/ControlRig/Source/ControlRigEditor/Private/EditMode/ControlRigEditMode.cpp b/Engine/Plugins/Animation/ControlRig/Source/ControlRigEditor/Private/EditMode/ControlRigEditMode.cpp
--- a/Engine/Plugins/Animation/ControlRig/Source/ControlRigEditor/Private/EditMode/ControlRigEditMode.cpp
+++ b/Engine/Plugins/Animation/ControlRig/Source/ControlRigEditor/Private/EditMode/ControlRigEditMode.cpp
@@ -2934,7 +2934,7 @@ void FControlRigEditMode::CreateShapeActors(UControlRig* ControlRig)
 			Param.ShapeName = ControlElement->Settings.ShapeName;
 			Param.SpawnTransform = ControlRig->GetControlGlobalTransform(ControlElement->GetName());
 			Param.ShapeTransform = ControlRig->GetHierarchy()->GetControlShapeTransform(ControlElement, ERigTransformType::CurrentLocal);
-			Param.bSelectable = ControlElement->Settings.bAnimatable;
+			Param.bSelectable = ControlElement->Settings.bShapeVisible && ControlElement->Settings.bAnimatable;
 
 			if (const FControlRigShapeDefinition* ShapeDef = UControlRigShapeLibrary::GetShapeByName(ControlElement->Settings.ShapeName, ShapeLibraries))
 			{
