Control Rig - Python Script - Set Pin Expansion, Set Function Properties, Fixed AddExposedPin Output

diff --git a/Engine/Plugins/Animation/ControlRig/Source/ControlRigDeveloper/Private/ControlRigBlueprint.cpp b/Engine/Plugins/Animation/ControlRig/Source/ControlRigDeveloper/Private/ControlRigBlueprint.cpp
--- a/Engine/Plugins/Animation/ControlRig/Source/ControlRigDeveloper/Private/ControlRigBlueprint.cpp
+++ b/Engine/Plugins/Animation/ControlRig/Source/ControlRigDeveloper/Private/ControlRigBlueprint.cpp
@@ -1761,6 +1761,24 @@ TArray<FString> UControlRigBlueprint::GeneratePythonCommands(const FString InNew
 							*Graph->GetGraphName(),
 							Graph->GetEntryNode()->IsMutable() ? TEXT("True") : TEXT("False"),
 							*RigVMPythonUtils::NameToPep8(Graph->GetGraphName())));
+
+					URigVMLibraryNode* LibraryNode = Cast<URigVMLibraryNode>(Graph->GetOuter());
+
+					InternalCommands.Add(FString::Printf(TEXT("library_controller.set_node_category_by_name('%s', '%s')"),
+						*Graph->GetGraphName(),
+						*LibraryNode->GetNodeCategory()));
+
+					InternalCommands.Add(FString::Printf(TEXT("library_controller.set_node_keywords_by_name('%s', '%s')"),
+						*Graph->GetGraphName(),
+						*LibraryNode->GetNodeKeywords() ));
+
+					InternalCommands.Add(FString::Printf(TEXT("library_controller.set_node_description_by_name('%s', '%s')"),
+						*Graph->GetGraphName(),
+						*LibraryNode->GetNodeDescription()));
+
+					InternalCommands.Add(FString::Printf(TEXT("library_controller.set_node_color_by_name('%s', %s)"),
+						*Graph->GetGraphName(),
+						*RigVMPythonUtils::LinearColorToPythonString(LibraryNode->GetNodeColor()) ));
 					
 					URigVMFunctionEntryNode* EntryNode = Graph->GetEntryNode();
 					URigVMFunctionReturnNode* ReturnNode = Graph->GetReturnNode();
@@ -1809,7 +1827,7 @@ TArray<FString> UControlRigBlueprint::GeneratePythonCommands(const FString InNew
 									*Graph->GetGraphName(),
 									*Pin->GetName(), 
 									*Pin->GetCPPType(),
-									Pin->GetCPPTypeObject() ? *Pin->GetCPPTypeObject()->GetPathName() : TEXT("''"),
+									Pin->GetCPPTypeObject() ? *Pin->GetCPPTypeObject()->GetPathName() : TEXT(""),
 									*Pin->GetDefaultValue()));
 						}
 					}
diff --git a/Engine/Source/Developer/RigVMDeveloper/Private/RigVMModel/RigVMController.cpp b/Engine/Source/Developer/RigVMDeveloper/Private/RigVMModel/RigVMController.cpp
--- a/Engine/Source/Developer/RigVMDeveloper/Private/RigVMModel/RigVMController.cpp
+++ b/Engine/Source/Developer/RigVMDeveloper/Private/RigVMModel/RigVMController.cpp
@@ -630,6 +630,22 @@ TArray<FString> URigVMController::GetAddNodePythonCommands(URigVMNode* Node) con
 							*GraphName,
 							*PinPath,
 							*Pin->GetDefaultValue()));
+
+
+				TArray<const URigVMPin*> SubPins = { Pin };
+				for (int32 i = 0; i < SubPins.Num(); ++i)
+				{
+					if (SubPins[i]->IsStruct() || SubPins[i]->IsArray())
+					{
+						SubPins.Append(SubPins[i]->GetSubPins());
+						const FString SubPinPath = GetSanitizedPinPath(SubPins[i]->GetPinPath());
+
+						Commands.Add(FString::Printf(TEXT("blueprint.get_controller_by_name('%s').set_pin_expansion('%s', %s)"),
+							*GraphName,
+							*SubPinPath,
+							SubPins[i]->IsExpanded() ? TEXT("True") : TEXT("False")));
+					}
+				}
 			}
 
 			if (!Pin->GetBoundVariablePath().IsEmpty())
@@ -6464,11 +6480,11 @@ bool URigVMController::SetNodeColor(URigVMNode* InNode, const FLinearColor& InCo
 		const FString GraphName = GetSanitizedGraphName(GetGraph()->GetGraphName());
 		const FString NodePath = GetSanitizedPinPath(InNode->GetNodePath());
 
-		RigVMPythonUtils::Print(GetGraphOuterName(), 
-							FString::Printf(TEXT("blueprint.get_controller_by_name('%s').set_node_color_by_name('%s', %s)"),
-											*GraphName,
-											*NodePath,
-											*RigVMPythonUtils::LinearColorToPythonString(InColor)));
+		RigVMPythonUtils::Print(GetGraphOuterName(),
+			FString::Printf(TEXT("blueprint.get_controller_by_name('%s').set_node_color_by_name('%s', %s)"),
+				*GraphName,
+				*NodePath,
+				*RigVMPythonUtils::LinearColorToPythonString(InColor)));
 	}
 
 	return true;
