Control Rig - Rig Unit - Transform Array to SRT

diff --git a/Engine/Plugins/Animation/ControlRig/Source/ControlRig/Private/Units/Math/RigUnit_MathTransform.cpp b/Engine/Plugins/Animation/ControlRig/Source/ControlRig/Private/Units/Math/RigUnit_MathTransform.cpp
--- a/Engine/Plugins/Animation/ControlRig/Source/ControlRig/Private/Units/Math/RigUnit_MathTransform.cpp
+++ b/Engine/Plugins/Animation/ControlRig/Source/ControlRig/Private/Units/Math/RigUnit_MathTransform.cpp
@@ -163,6 +163,28 @@ FRigUnit_MathTransformFromSRT_Execute()
 	EulerTransform.FromFTransform(Transform);
 }
 
+
+FRigUnit_MathTransformArrayToSRT_Execute()
+{
+	DECLARE_SCOPE_HIERARCHICAL_COUNTER_RIGUNIT()
+	if (Transforms.Num() == 0)
+	{
+		return;
+	}
+
+	Translations.SetNumUninitialized(Transforms.Num());
+	Rotations.SetNumUninitialized(Transforms.Num());
+	Scales.SetNumUninitialized(Transforms.Num());
+
+	for (int32 Index = 0; Index < Transforms.Num(); Index++)
+	{
+		Translations[Index] = Transforms[Index].GetTranslation();
+		Rotations[Index] = Transforms[Index].GetRotation();
+		Scales[Index] = Transforms[Index].GetScale3D();
+	}
+
+}
+
 FRigUnit_MathTransformClampSpatially_Execute()
 {
 	DECLARE_SCOPE_HIERARCHICAL_COUNTER_RIGUNIT()
diff --git a/Engine/Plugins/Animation/ControlRig/Source/ControlRig/Private/Units/Math/RigUnit_MathTransform.h b/Engine/Plugins/Animation/ControlRig/Source/ControlRig/Private/Units/Math/RigUnit_MathTransform.h
--- a/Engine/Plugins/Animation/ControlRig/Source/ControlRig/Private/Units/Math/RigUnit_MathTransform.h
+++ b/Engine/Plugins/Animation/ControlRig/Source/ControlRig/Private/Units/Math/RigUnit_MathTransform.h
@@ -406,6 +406,35 @@ struct CONTROLRIG_API FRigUnit_MathTransformFromSRT : public FRigUnit_MathTransf
 	FEulerTransform EulerTransform;
 };
 
+
+/**
+ * Decomposes a Transform Array to its components.
+ */
+USTRUCT(meta = (DisplayName = "Transform Array to SRT", Keywords = "EulerTransform,Scale,Rotation,Orientation,Translation,Location"))
+struct CONTROLRIG_API FRigUnit_MathTransformArrayToSRT : public FRigUnit_MathTransformBase
+{
+	GENERATED_BODY()
+
+	FRigUnit_MathTransformArrayToSRT()
+	{
+	}
+
+	RIGVM_METHOD()
+		virtual void Execute(const FRigUnitContext& Context) override;
+
+	UPROPERTY(meta = (Input))
+	TArray<FTransform> Transforms;
+
+	UPROPERTY(meta = (Output))
+	TArray<FVector> Translations;
+
+	UPROPERTY(meta = (Output))
+	TArray<FQuat> Rotations;
+
+	UPROPERTY(meta = (Output))
+	TArray<FVector> Scales;
+};
+
 /**
  * Clamps a position using a plane collision, cylindric collision or spherical collision.
  */
diff --git a/Engine/Plugins/Animation/ControlRig/Source/ControlRig/Private/Units/Math/RigUnit_MathTransformTests.cpp b/Engine/Plugins/Animation/ControlRig/Source/ControlRig/Private/Units/Math/RigUnit_MathTransformTests.cpp
--- a/Engine/Plugins/Animation/ControlRig/Source/ControlRig/Private/Units/Math/RigUnit_MathTransformTests.cpp
+++ b/Engine/Plugins/Animation/ControlRig/Source/ControlRig/Private/Units/Math/RigUnit_MathTransformTests.cpp
@@ -127,4 +127,23 @@ IMPLEMENT_RIGUNIT_AUTOMATION_TEST(FRigUnit_MathTransformTransformVector)
 	return true;
 }
 
+IMPLEMENT_RIGUNIT_AUTOMATION_TEST(FRigUnit_MathTransformArrayToSRT)
+{
+	Unit.Transforms.SetNumUninitialized(5);
+
+	for (int32 Index = 0; Index < Unit.Transforms.Num(); Index++)
+	{
+		Unit.Transforms[Index] = FTransform(FQuat(FVector(1.f, 0.f, 0.f), HALF_PI * float(Index)), FVector(float(Index), 0.f, 0.f));
+	}
+
+	InitAndExecute();
+
+	for (int32 Index = 0; Index < Unit.Transforms.Num(); Index++)
+	{
+		AddErrorIfFalse(FRigUnit_MathTransformTest_Utils::IsNearlyEqual(Unit.Translations[Index], FVector(float(Index), 0.f, 0.f)), TEXT("unexpected result"));
+		AddErrorIfFalse(FRigUnit_MathTransformTest_Utils::IsNearlyEqual(Unit.Rotations[Index], FQuat(FVector(1.f, 0.f, 0.f), HALF_PI * float(Index))), TEXT("unexpected result"));
+	}
+	return true;
+}
+
 #endif
