Control Rig: Speculative fix for a crash when transient controls' scope closes.

#rnx

diff --git a/Engine/Plugins/Animation/ControlRig/Source/ControlRig/Private/ControlRig.cpp b/Engine/Plugins/Animation/ControlRig/Source/ControlRig/Private/ControlRig.cpp
--- a/Engine/Plugins/Animation/ControlRig/Source/ControlRig/Private/ControlRig.cpp
+++ b/Engine/Plugins/Animation/ControlRig/Source/ControlRig/Private/ControlRig.cpp
@@ -3085,19 +3085,21 @@ UControlRig::FTransientControlScope::FTransientControlScope(TObjectPtr<URigHiera
 
 UControlRig::FTransientControlScope::~FTransientControlScope()
 {
-	URigHierarchyController* Controller = Hierarchy->GetController();
-	for (const FTransientControlInfo& Info : SavedTransientControls)
-	{
-		Controller->AddControl(
-            Info.Name,
-            Info.Parent,
-            Info.Settings,
-            Info.Value,
-            Info.OffsetTransform,
-            Info.ShapeTransform,
-            false,
-            false
-        );
+	if (URigHierarchyController* Controller = Hierarchy->GetController())
+	{
+		for (const FTransientControlInfo& Info : SavedTransientControls)
+		{
+			Controller->AddControl(
+				Info.Name,
+				Info.Parent,
+				Info.Settings,
+				Info.Value,
+				Info.OffsetTransform,
+				Info.ShapeTransform,
+				false,
+				false
+			);
+		}
 	}
 }
 
