Sequencer: Control Rig: When duplicating control rig's do to construction scripts running make sure to put sequencer in  silent mode else we will evaluate the sequence and possibly erase any change to that actor. Did some refactoring also so we only do an evaluation at th end. Also minor fix for recent change since array size can change while running through the loop

diff --git a/Engine/Plugins/Animation/ControlRig/Source/ControlRigEditor/Private/EditMode/ControlRigEditMode.cpp b/Engine/Plugins/Animation/ControlRig/Source/ControlRigEditor/Private/EditMode/ControlRigEditMode.cpp
--- a/Engine/Plugins/Animation/ControlRig/Source/ControlRigEditor/Private/EditMode/ControlRigEditMode.cpp
+++ b/Engine/Plugins/Animation/ControlRig/Source/ControlRigEditor/Private/EditMode/ControlRigEditMode.cpp
@@ -402,7 +402,8 @@ static void ClearOutAnyActiveTools()
 void FControlRigEditMode::Exit()
 {
 	ClearOutAnyActiveTools();
-
+	OnControlRigAddedOrRemovedDelegate.Clear();
+	OnControlRigSelectedDelegate.Clear();
 	for (TWeakObjectPtr<UControlRig>& RuntimeRigPtr : RuntimeControlRigs)
 	{
 		if (UControlRig* ControlRig = RuntimeRigPtr.Get())
@@ -437,8 +438,7 @@ void FControlRigEditMode::Exit()
 	}
 
 	DestroyShapesActors(nullptr);
-	OnControlRigAddedOrRemovedDelegate.Clear();
-	OnControlRigSelectedDelegate.Clear();
+
 
 	TArray<TWeakObjectPtr<UControlRig>> PreviousRuntimeRigs = RuntimeControlRigs;
 	for (int32 PreviousRuntimeRigIndex = 0; PreviousRuntimeRigIndex < PreviousRuntimeRigs.Num(); PreviousRuntimeRigIndex++)
@@ -2875,11 +2875,13 @@ void FControlRigEditMode::RecreateControlShapeActors(const TArray<FRigElementKey
 	}
 	else if(ControlRigsToRecreate.Num() > 0)
 	{
-		for (UControlRig* ControlRig : ControlRigsToRecreate)
+		TArray < UControlRig*> ControlRigsCopy = ControlRigsToRecreate;
+		for (UControlRig* ControlRig : ControlRigsCopy)
 		{
 			DestroyShapesActors(ControlRig);
 			CreateShapeActors(ControlRig);
 		}
+		ControlRigsToRecreate.SetNum(0);
 	}
 }
 
@@ -3786,7 +3788,7 @@ void FControlRigEditMode::DestroyShapesActors(UControlRig* ControlRig)
 				UWorld* World = ShapeActor->GetWorld();
 				if (World)
 				{
-					World->DestroyActor(ShapeActor);
+					World->EditorDestroyActor(ShapeActor, true);
 				}
 			}
 		}
@@ -3808,7 +3810,7 @@ void FControlRigEditMode::DestroyShapesActors(UControlRig* ControlRig)
 				UWorld* World = ShapeActor->GetWorld();
 				if (World)
 				{
-					World->DestroyActor(ShapeActor);
+					World->EditorDestroyActor(ShapeActor,true);
 				}
 			}
 			ControlRigShapeActors.Remove(ControlRig);
diff --git a/Engine/Plugins/Animation/ControlRig/Source/ControlRigEditor/Private/Sequencer/ControlRigParameterTrackEditor.cpp b/Engine/Plugins/Animation/ControlRig/Source/ControlRigEditor/Private/Sequencer/ControlRigParameterTrackEditor.cpp
--- a/Engine/Plugins/Animation/ControlRig/Source/ControlRigEditor/Private/Sequencer/ControlRigParameterTrackEditor.cpp
+++ b/Engine/Plugins/Animation/ControlRig/Source/ControlRigEditor/Private/Sequencer/ControlRigParameterTrackEditor.cpp
@@ -222,53 +222,57 @@ FControlRigParameterTrackEditor::FControlRigParameterTrackEditor(TSharedRef<ISeq
 		//we check for two things, one if the control rig has been replaced if so we need to switch.
 		//the other is if bound object on the edit mode is null we request a re-evaluate which will reset it up.
 		FDelegateHandle OnObjectsReplacedHandle = FCoreUObjectDelegates::OnObjectsReplaced.AddLambda([&](const TMap<UObject*, UObject*>& ReplacementMap)
+		{
+			static bool bHasEnteredSilent = false;
+			if (GetSequencer().IsValid())
 			{
-				if (GetSequencer().IsValid())
+				bool bRequestEvaluate = false;
+				TMap<UControlRig*, UControlRig*> OldToNewControlRigs;
+				FControlRigEditMode* ControlRigEditMode = GetEditMode();
+				if (ControlRigEditMode)
 				{
-					TMap<UControlRig*, UControlRig*> OldToNewControlRigs;
-					FControlRigEditMode* ControlRigEditMode = GetEditMode();
-					if (ControlRigEditMode) 
+					TArrayView<TWeakObjectPtr<UControlRig>> ControlRigs = ControlRigEditMode->GetControlRigs();
+					for (TWeakObjectPtr<UControlRig>& ControlRigPtr : ControlRigs)
 					{
-						TArrayView<TWeakObjectPtr<UControlRig>> ControlRigs = ControlRigEditMode->GetControlRigs();
-						bool bRequestEvaluate = false;
-						for (TWeakObjectPtr<UControlRig>& ControlRigPtr : ControlRigs)
+						if (UControlRig* ControlRig = ControlRigPtr.Get())
 						{
-							if (UControlRig* ControlRig = ControlRigPtr.Get())
+							if (ControlRig->GetObjectBinding() && ControlRig->GetObjectBinding()->GetBoundObject() == nullptr)
 							{
-								if (ControlRig->GetObjectBinding() && ControlRig->GetObjectBinding()->GetBoundObject() == nullptr)
-								{
-									bRequestEvaluate = true;
-									break;
-								}
+								bRequestEvaluate = true;
+								break;
 							}
 						}
-
-						if (bRequestEvaluate)
-						{
-							GetSequencer()->RequestEvaluate();
-						}
 					}
-					//Reset Bindings for replaced objects.
-					for (TPair<UObject*, UObject*> ReplacedObject : ReplacementMap)
+				}
+				//Reset Bindings for replaced objects.
+				for (TPair<UObject*, UObject*> ReplacedObject : ReplacementMap)
+				{
+					if (UControlRigComponent* OldControlRigComponent = Cast<UControlRigComponent>(ReplacedObject.Key))
 					{
-						if (UControlRigComponent* OldControlRigComponent = Cast<UControlRigComponent>(ReplacedObject.Key))
+						UControlRigComponent* NewControlRigComponent = Cast<UControlRigComponent>(ReplacedObject.Value);
+						if (OldControlRigComponent->GetControlRig())
 						{
-							UControlRigComponent* NewControlRigComponent = Cast<UControlRigComponent>(ReplacedObject.Value);
-							if (OldControlRigComponent->GetControlRig())
+							UControlRig* NewControlRig = nullptr;
+							if (NewControlRigComponent)
 							{
-								UControlRig* NewControlRig = nullptr;
-								if (NewControlRigComponent)
-								{
-									NewControlRig = NewControlRigComponent->GetControlRig();
-								}
-								OldToNewControlRigs.Emplace(OldControlRigComponent->GetControlRig(), NewControlRig);
+								NewControlRig = NewControlRigComponent->GetControlRig();
 							}
+							OldToNewControlRigs.Emplace(OldControlRigComponent->GetControlRig(), NewControlRig);
 						}
-						else if (UControlRig* OldControlRig = Cast<UControlRig>(ReplacedObject.Key))
-						{
-							UControlRig* NewControlRig = Cast<UControlRig>(ReplacedObject.Value);
-							OldToNewControlRigs.Emplace(OldControlRig, NewControlRig);
-						}
+					}
+					else if (UControlRig* OldControlRig = Cast<UControlRig>(ReplacedObject.Key))
+					{
+						UControlRig* NewControlRig = Cast<UControlRig>(ReplacedObject.Value);
+						OldToNewControlRigs.Emplace(OldControlRig, NewControlRig);
+					}
+				}
+				if (OldToNewControlRigs.Num() > 0)
+				{
+					//need to avoid any evaluations when doing this replacement otherwise we will evaluate sequencer
+					if (bHasEnteredSilent == false)
+					{
+						GetSequencer()->EnterSilentMode();
+						bHasEnteredSilent = true;
 					}
 					UMovieScene* MovieScene = GetSequencer()->GetFocusedMovieSceneSequence()->GetMovieScene();
 					const TArray<FMovieSceneBinding>& Bindings = MovieScene->GetBindings();
@@ -289,8 +293,7 @@ FControlRigParameterTrackEditor::FControlRigParameterTrackEditor(TSharedRef<ISeq
 									Track->Modify();
 									Track->ReplaceControlRig(*NewControlRig, OldControlRig->GetClass() != (*NewControlRig)->GetClass());
 									BindControlRig(*NewControlRig);
-
-									GetSequencer()->NotifyMovieSceneDataChanged(EMovieSceneDataChangeType::MovieSceneStructureItemsChanged);
+									bRequestEvaluate = true;
 								}
 								else
 								{
@@ -299,37 +302,44 @@ FControlRigParameterTrackEditor::FControlRigParameterTrackEditor(TSharedRef<ISeq
 								if (ControlRigEditMode)
 								{
 									ControlRigEditMode->ReplaceControlRig(OldControlRig, *NewControlRig);
-									
+
 									UControlRig* PtrNewControlRig = *NewControlRig;
 
-									auto UpdateSelectionDelegate = [this, SelectedControls, PtrNewControlRig]()
+									auto UpdateSelectionDelegate = [this, SelectedControls, PtrNewControlRig, bRequestEvaluate]()
 									{
 
-										UE_LOG(LogControlRigEditor, Log, TEXT("UpdateSelectionTimer"));
 										if (!(FSlateApplication::Get().HasAnyMouseCaptor() || GUnrealEd->IsUserInteracting()))
 										{
-											UE_LOG(LogControlRigEditor, Log, TEXT("UpdateSelectionTimer - Sync"));
 											TGuardValue<bool> Guard(bIsDoingSelection, true);
 											GetSequencer()->ExternalSelectionHasChanged();
 											if (PtrNewControlRig)
 											{
-												GEditor->GetTimerManager()->SetTimerForNextTick([SelectedControls,PtrNewControlRig]()
-												{
-													PtrNewControlRig->ClearControlSelection();
-													for (const FName& ControlName : SelectedControls)
+												GEditor->GetTimerManager()->SetTimerForNextTick([this, SelectedControls, PtrNewControlRig, bRequestEvaluate]()
 													{
-														PtrNewControlRig->SelectControl(ControlName, true);
-													}
-												});
+														PtrNewControlRig->ClearControlSelection();
+														for (const FName& ControlName : SelectedControls)
+														{
+															PtrNewControlRig->SelectControl(ControlName, true);
+														}
+													});
+											}
+											
+											if (bHasEnteredSilent == true)
+											{
+												GetSequencer()->ExitSilentMode();
+												bHasEnteredSilent = false;
+											}
+											
+											if (bRequestEvaluate)
+											{
+												GetSequencer()->NotifyMovieSceneDataChanged(EMovieSceneDataChangeType::MovieSceneStructureItemsChanged);
 											}
 											if (UpdateSelectionTimerHandle.IsValid())
 											{
-												UE_LOG(LogControlRigEditor, Log, TEXT("UpdateSelectionTimer - Clear"));
 												GEditor->GetTimerManager()->ClearTimer(UpdateSelectionTimerHandle);
 											}
 										}
 
-
 									};
 
 									GEditor->GetTimerManager()->SetTimer(UpdateSelectionTimerHandle, UpdateSelectionDelegate, 0.01f, true);
@@ -337,9 +347,23 @@ FControlRigParameterTrackEditor::FControlRigParameterTrackEditor(TSharedRef<ISeq
 							}
 						}
 					}
+					if (!ControlRigEditMode)
+					{
+						if (bHasEnteredSilent == true)
+						{
+							GetSequencer()->ExitSilentMode();
+							bHasEnteredSilent = false;
+						}
+						if (bRequestEvaluate)
+						{
+							GetSequencer()->NotifyMovieSceneDataChanged(EMovieSceneDataChangeType::MovieSceneStructureItemsChanged);
+						}
+					}
 				}
+			}
+		});
+	
 
-			});
 		AcquiredResources.Add([=] { FCoreUObjectDelegates::OnObjectsReplaced.Remove(OnObjectsReplacedHandle); });
 	}
 	//register all modified/selections for control rigs
