Fixes 1-frame lag on meshes attached using "Copy Pose From Mesh" when interacting with a Control Rig.

diff --git a/Engine/Plugins/Animation/ControlRig/Source/ControlRigEditor/Private/EditMode/ControlRigEditMode.cpp b/Engine/Plugins/Animation/ControlRig/Source/ControlRigEditor/Private/EditMode/ControlRigEditMode.cpp
--- a/Engine/Plugins/Animation/ControlRig/Source/ControlRigEditor/Private/EditMode/ControlRigEditMode.cpp
+++ b/Engine/Plugins/Animation/ControlRig/Source/ControlRigEditor/Private/EditMode/ControlRigEditMode.cpp
@@ -3755,12 +3755,42 @@ void FControlRigEditMode::TickManipulatableObjects(float DeltaTime)
 			}
 			else if (USkeletalMeshComponent* MeshComponent = Cast<USkeletalMeshComponent>(SceneComponent))
 			{
-				MeshComponent->RefreshBoneTransforms();
-				MeshComponent->RefreshSlaveComponents();
-				MeshComponent->UpdateComponentToWorld();
-				MeshComponent->FinalizeBoneTransform();
-				MeshComponent->MarkRenderTransformDirty();
-				MeshComponent->MarkRenderDynamicDataDirty();
+				// NOTE: we have to update/tick ALL children skeletal mesh components here because user can attach
+				// additional skeletal meshes via the "Copy Pose from Mesh" node.
+				//
+				// If this is left up to the viewport tick(), the attached meshes will render before they get the latest
+				// parent bone transforms resulting in a visible lag on all attached components.
+
+				// get hierarchically ordered list of ALL child skeletal mesh components (recursive)
+				const AActor* ThisActor = MeshComponent->GetOwner();
+				TArray<USceneComponent*> ChildrenComponents;
+				MeshComponent->GetChildrenComponents(true, ChildrenComponents);
+				TInlineComponentArray<USkeletalMeshComponent*> SkeletalMeshesToUpdate;
+				SkeletalMeshesToUpdate.Add(MeshComponent);
+				for (USceneComponent* ChildComponent : ChildrenComponents)
+				{
+					if (USkeletalMeshComponent* ChildMeshComponent = Cast<USkeletalMeshComponent>(ChildComponent))
+					{
+						if (ThisActor == ChildMeshComponent->GetOwner())
+						{
+							SkeletalMeshesToUpdate.Add(ChildMeshComponent);
+						}
+					}
+				}
+
+				// update pose of all children skeletal meshes in this actor
+				for (USkeletalMeshComponent* SkeletalMeshToUpdate : SkeletalMeshesToUpdate)
+				{
+					// "Copy Pose from Mesh" requires AnimInstance::PreUpdate() to copy the parent bone transforms.
+					// have to TickAnimation() to ensure that PreUpdate() is called on all anim instances
+					SkeletalMeshToUpdate->TickAnimation(0.0f, false);
+					SkeletalMeshToUpdate->RefreshBoneTransforms();
+					SkeletalMeshToUpdate->RefreshSlaveComponents();
+					SkeletalMeshToUpdate->UpdateComponentToWorld();
+					SkeletalMeshToUpdate->FinalizeBoneTransform();
+					SkeletalMeshToUpdate->MarkRenderTransformDirty();
+					SkeletalMeshToUpdate->MarkRenderDynamicDataDirty();
+				}
 			}
 		}
 	}
