Sequencer Python CS Fix: Test failing since control rig could lose selection when selection in python due to control rig getting selected

diff --git a/Engine/Plugins/Animation/ControlRig/Source/ControlRigEditor/Private/Sequencer/ControlRigParameterTrackEditor.cpp b/Engine/Plugins/Animation/ControlRig/Source/ControlRigEditor/Private/Sequencer/ControlRigParameterTrackEditor.cpp
--- a/Engine/Plugins/Animation/ControlRig/Source/ControlRigEditor/Private/Sequencer/ControlRigParameterTrackEditor.cpp
+++ b/Engine/Plugins/Animation/ControlRig/Source/ControlRigEditor/Private/Sequencer/ControlRigParameterTrackEditor.cpp
@@ -2342,27 +2342,41 @@ void FControlRigParameterTrackEditor::HandleSpaceKeyMoved(UMovieSceneControlRigP
 
 void FControlRigParameterTrackEditor::SetUpEditModeIfNeeded(UControlRig* ControlRig)
 {
-
-	FControlRigEditMode* ControlRigEditMode = GetEditMode();
-	if (!ControlRigEditMode)
+	if (ControlRig)
 	{
-		ControlRigEditMode = GetEditMode(true);
-		if (TSharedPtr<IControlRigObjectBinding> ObjectBinding = ControlRig->GetObjectBinding())
+		//this could clear the selection so if it does reset it
+		TArray<FName> ControlRigSelection = ControlRig->CurrentControlSelection();
+
+		FControlRigEditMode* ControlRigEditMode = GetEditMode();
+		if (!ControlRigEditMode)
 		{
-			if (ControlRigEditMode)
+			ControlRigEditMode = GetEditMode(true);
+			if (TSharedPtr<IControlRigObjectBinding> ObjectBinding = ControlRig->GetObjectBinding())
 			{
-				ControlRigEditMode->AddControlRigObject(ControlRig, GetSequencer());
+				if (ControlRigEditMode)
+				{
+					ControlRigEditMode->AddControlRigObject(ControlRig, GetSequencer());
+				}
 			}
 		}
-	}
-	else
-	{
-		if (ControlRigEditMode->AddControlRigObject(ControlRig, GetSequencer()))
+		else
+		{
+			if (ControlRigEditMode->AddControlRigObject(ControlRig, GetSequencer()))
 			{
-			//force an evaluation, this will get the control rig setup so edit mode looks good.
-			if (GetSequencer().IsValid())
+				//force an evaluation, this will get the control rig setup so edit mode looks good.
+				if (GetSequencer().IsValid())
+				{
+					GetSequencer()->NotifyMovieSceneDataChanged(EMovieSceneDataChangeType::Unknown);
+				}
+			}
+		}
+		TArray<FName> NewControlRigSelection = ControlRig->CurrentControlSelection();
+		if (ControlRigSelection.Num() != NewControlRigSelection.Num())
+		{
+			ControlRig->ClearControlSelection();
+			for (const FName& Name : ControlRigSelection)
 			{
-				GetSequencer()->NotifyMovieSceneDataChanged(EMovieSceneDataChangeType::Unknown);
+				ControlRig->SelectControl(Name, true);
 			}
 		}
 	}
