ControlRig: Python Script: Fixed set_parent not logging and object created in wrong order

diff --git a/Engine/Plugins/Animation/ControlRig/Source/ControlRig/Private/Rigs/RigHierarchyController.cpp b/Engine/Plugins/Animation/ControlRig/Source/ControlRig/Private/Rigs/RigHierarchyController.cpp
--- a/Engine/Plugins/Animation/ControlRig/Source/ControlRig/Private/Rigs/RigHierarchyController.cpp
+++ b/Engine/Plugins/Animation/ControlRig/Source/ControlRig/Private/Rigs/RigHierarchyController.cpp
@@ -1362,11 +1362,12 @@ TArray<FRigElementKey> URigHierarchyController::ImportFromHierarchyContainer(con
 TArray<FString> URigHierarchyController::GeneratePythonCommands()
 {
 	TArray<FString> Commands;
-	Hierarchy->ForEach([&](FRigBaseElement* Element) -> bool
+	Hierarchy->Traverse([&](FRigBaseElement* Element, bool& bContinue)
 	{
 		Commands.Append(GetAddElementPythonCommands(Element));
 		
-		return true;
+		bContinue = true;
+		return;
 	});
 
 	return Commands;
diff --git a/Engine/Plugins/Animation/ControlRig/Source/ControlRigEditor/Private/Editor/SRigHierarchy.cpp b/Engine/Plugins/Animation/ControlRig/Source/ControlRigEditor/Private/Editor/SRigHierarchy.cpp
--- a/Engine/Plugins/Animation/ControlRig/Source/ControlRigEditor/Private/Editor/SRigHierarchy.cpp
+++ b/Engine/Plugins/Animation/ControlRig/Source/ControlRigEditor/Private/Editor/SRigHierarchy.cpp
@@ -2585,7 +2585,7 @@ FReply SRigHierarchy::ReparentOrMatchTransform(const TArray<FRigElementKey>& Dra
 
 			if(ParentKey.IsValid())
 			{
-				Controller->SetParent(DraggedKey, ParentKey, true, true);
+				Controller->SetParent(DraggedKey, ParentKey, true, true, true);
 			}
 			else
 			{
