Control Rig: Fix crash when changing control shape color in worker thread

diff --git a/Engine/Plugins/Animation/ControlRig/Source/ControlRigEditor/Private/EditMode/ControlRigControlsProxy.cpp b/Engine/Plugins/Animation/ControlRig/Source/ControlRigEditor/Private/EditMode/ControlRigControlsProxy.cpp
--- a/Engine/Plugins/Animation/ControlRig/Source/ControlRigEditor/Private/EditMode/ControlRigControlsProxy.cpp
+++ b/Engine/Plugins/Animation/ControlRig/Source/ControlRigEditor/Private/EditMode/ControlRigControlsProxy.cpp
@@ -902,11 +902,14 @@ void UControlRigDetailPanelControlProxies::RecreateAllProxies(UControlRig* Contr
 
 void UControlRigDetailPanelControlProxies::ProxyChanged(UControlRig* ControlRig, const FName& Name)
 {
-	UControlRigControlsProxy* Proxy = FindProxy(ControlRig,Name);
-	if (Proxy)
+	if (IsInGameThread())
 	{
-		Modify();
-		Proxy->ValueChanged();
+		UControlRigControlsProxy* Proxy = FindProxy(ControlRig,Name);
+		if (Proxy)
+		{
+			Modify();
+			Proxy->ValueChanged();
+		}
 	}
 }
 
diff --git a/Engine/Plugins/Animation/ControlRig/Source/ControlRigEditor/Private/EditMode/ControlRigEditMode.cpp b/Engine/Plugins/Animation/ControlRig/Source/ControlRigEditor/Private/EditMode/ControlRigEditMode.cpp
--- a/Engine/Plugins/Animation/ControlRig/Source/ControlRigEditor/Private/EditMode/ControlRigEditMode.cpp
+++ b/Engine/Plugins/Animation/ControlRig/Source/ControlRigEditor/Private/EditMode/ControlRigEditMode.cpp
@@ -297,7 +297,7 @@ void FControlRigEditMode::SetObjects_Internal()
 			RuntimeControlRig->GetHierarchy()->OnModified().RemoveAll(this);
 
 			RuntimeControlRig->ControlModified().AddSP(this, &FControlRigEditMode::OnControlModified);
-			RuntimeControlRig->GetHierarchy()->OnModified().AddSP(this, &FControlRigEditMode::OnHierarchyModified);
+			RuntimeControlRig->GetHierarchy()->OnModified().AddSP(this, &FControlRigEditMode::OnHierarchyModified_AnyThread);
 			if (USkeletalMeshComponent* MeshComponent = Cast<USkeletalMeshComponent>(GetHostingSceneComponent(RuntimeControlRig)))
 			{
 				TStrongObjectPtr<UControlRigEditModeDelegateHelper>* DelegateHelper = DelegateHelpers.Find(RuntimeControlRig);
@@ -3221,6 +3221,28 @@ void FControlRigEditMode::OnHierarchyModified(ERigHierarchyNotification InNotif,
 	}
 }
 
+void FControlRigEditMode::OnHierarchyModified_AnyThread(ERigHierarchyNotification InNotif, URigHierarchy* InHierarchy, const FRigBaseElement* InElement)
+{
+	FRigElementKey Key;
+	if(InElement)
+	{
+		Key = InElement->GetKey();
+	}
+
+	TWeakObjectPtr<URigHierarchy> WeakHierarchy = InHierarchy;
+	
+	FFunctionGraphTask::CreateAndDispatchWhenReady([this, InNotif, WeakHierarchy, Key]()
+	{
+		if(!WeakHierarchy.IsValid())
+		{
+			return;
+		}
+		const FRigBaseElement* Element = WeakHierarchy.Get()->Find(Key);
+		OnHierarchyModified(InNotif, WeakHierarchy.Get(), Element);
+		
+	}, TStatId(), NULL, ENamedThreads::GameThread);
+}
+
 void FControlRigEditMode::OnControlModified(UControlRig* Subject, FRigControlElement* InControlElement, const FRigControlModifiedContext& Context)
 {
 	//this makes sure the details panel ui get's updated, don't remove
@@ -3668,7 +3690,7 @@ void FControlRigEditMode::AddControlRigInternal(UControlRig* InControlRig)
 	InControlRig->SetControlsVisible(true);
 	InControlRig->PostInitInstanceIfRequired();
 	InControlRig->GetHierarchy()->OnModified().RemoveAll(this);
-	InControlRig->GetHierarchy()->OnModified().AddSP(this, &FControlRigEditMode::OnHierarchyModified);
+	InControlRig->GetHierarchy()->OnModified().AddSP(this, &FControlRigEditMode::OnHierarchyModified_AnyThread);
 
 	//needed for the control rig track editor delegates to get hooked up
 	if (WeakSequencer.IsValid())
diff --git a/Engine/Plugins/Animation/ControlRig/Source/ControlRigEditor/Private/EditMode/ControlRigEditMode.h b/Engine/Plugins/Animation/ControlRig/Source/ControlRigEditor/Private/EditMode/ControlRigEditMode.h
--- a/Engine/Plugins/Animation/ControlRig/Source/ControlRigEditor/Private/EditMode/ControlRigEditMode.h
+++ b/Engine/Plugins/Animation/ControlRig/Source/ControlRigEditor/Private/EditMode/ControlRigEditMode.h
@@ -160,6 +160,7 @@ public:
 
 	// callback that gets called when rig element is selected in other view
 	void OnHierarchyModified(ERigHierarchyNotification InNotif, URigHierarchy* InHierarchy, const FRigBaseElement* InElement);
+	void OnHierarchyModified_AnyThread(ERigHierarchyNotification InNotif, URigHierarchy* InHierarchy, const FRigBaseElement* InElement);
 	void OnControlModified(UControlRig* Subject, FRigControlElement* InControlElement, const FRigControlModifiedContext& Context);
 
 	/** return true if it can be removed from preview scene 
diff --git a/Engine/Plugins/Animation/ControlRig/Source/ControlRigEditor/Private/Editor/ControlRigEditor.cpp b/Engine/Plugins/Animation/ControlRig/Source/ControlRigEditor/Private/Editor/ControlRigEditor.cpp
--- a/Engine/Plugins/Animation/ControlRig/Source/ControlRigEditor/Private/Editor/ControlRigEditor.cpp
+++ b/Engine/Plugins/Animation/ControlRig/Source/ControlRigEditor/Private/Editor/ControlRigEditor.cpp
@@ -500,7 +500,7 @@ void FControlRigEditor::InitControlRigEditor(const EToolkitMode::Type Mode, cons
 
 		if (FControlRigEditMode* EditMode = GetEditMode())
 		{
-			InControlRigBlueprint->OnHierarchyModified().AddSP(EditMode, &FControlRigEditMode::OnHierarchyModified);
+			InControlRigBlueprint->OnHierarchyModified().AddSP(EditMode, &FControlRigEditMode::OnHierarchyModified_AnyThread);
 		}
 
 		InControlRigBlueprint->OnNodeDoubleClicked().AddSP(this, &FControlRigEditor::OnNodeDoubleClicked);
