Invalidate the viewport when ending control manipulation to update gizmo's hit proxy

diff --git a/Engine/Plugins/Animation/ControlRig/Source/ControlRigEditor/Private/EditMode/ControlRigEditMode.cpp b/Engine/Plugins/Animation/ControlRig/Source/ControlRigEditor/Private/EditMode/ControlRigEditMode.cpp
--- a/Engine/Plugins/Animation/ControlRig/Source/ControlRigEditor/Private/EditMode/ControlRigEditMode.cpp
+++ b/Engine/Plugins/Animation/ControlRig/Source/ControlRigEditor/Private/EditMode/ControlRigEditMode.cpp
@@ -930,8 +930,10 @@ bool FControlRigEditMode::InputKey(FEditorViewportClient* InViewportClient, FVie
 
 bool FControlRigEditMode::EndTracking(FEditorViewportClient* InViewportClient, FViewport* InViewport)
 {
+	const bool bWasInteracting = bManipulatorMadeChange && InteractionType != (uint8)EControlRigInteractionType::None;
+	
 	InteractionType = (uint8)EControlRigInteractionType::None;
-
+	
 	if (InteractionScopes.Num() > 0)
 	{
 		if (bManipulatorMadeChange)
@@ -940,7 +942,6 @@ bool FControlRigEditMode::EndTracking(FEditorViewportClient* InViewportClient, F
 			GEditor->EndTransaction();
 		}
 
-
 		for (TPair<UControlRig*, FControlRigInteractionScope*>& InteractionScope : InteractionScopes)
 		{
 			if (InteractionScope.Value)
@@ -950,11 +951,23 @@ bool FControlRigEditMode::EndTracking(FEditorViewportClient* InViewportClient, F
 		}
 		InteractionScopes.Reset();
 
+		if (bWasInteracting && IsInLevelEditor())
+		{
+			// We invalidate the hit proxies when in level editor to ensure that the gizmo's hit proxy is up to date.
+			// The invalidation is called here to avoid useless viewport update in the FControlRigEditMode::Tick
+			// function (that does an update when not in level editor)
+			TickManipulatableObjects(0.f);
+			
+			static constexpr bool bInvalidateChildViews = false;
+			static constexpr bool bInvalidateHitProxies = true;
+			InViewportClient->Invalidate(bInvalidateChildViews, bInvalidateHitProxies);
+		}
+		
 		return true;
 	}
 
 	bManipulatorMadeChange = false;
-
+	
 	return false;
 }
 
