Control Rig: Python Script: Fixed Float Control Limit not set properly

diff --git a/Engine/Plugins/Animation/ControlRig/Source/ControlRig/Private/Rigs/RigHierarchy.cpp b/Engine/Plugins/Animation/ControlRig/Source/ControlRig/Private/Rigs/RigHierarchy.cpp
--- a/Engine/Plugins/Animation/ControlRig/Source/ControlRig/Private/Rigs/RigHierarchy.cpp
+++ b/Engine/Plugins/Animation/ControlRig/Source/ControlRig/Private/Rigs/RigHierarchy.cpp
@@ -5349,6 +5349,12 @@ TArray<FString> URigHierarchy::ControlSettingsToPythonCommands(const FRigControl
 	Commands.Add(FString::Printf(TEXT("%s.limit_enabled = [%s]"),
 		*NameSettings,
 		*LimitEnabledStr));
+	Commands.Add(FString::Printf(TEXT("%s.minimum_value = %s"),
+		*NameSettings,
+		*Settings.MinimumValue.ToPythonString(Settings.ControlType)));
+	Commands.Add(FString::Printf(TEXT("%s.maximum_value = %s"),
+		*NameSettings,
+		*Settings.MaximumValue.ToPythonString(Settings.ControlType)));
 	Commands.Add(FString::Printf(TEXT("%s.primary_axis = unreal.RigControlAxis.%s"),
 		*NameSettings,
 		Settings.PrimaryAxis == ERigControlAxis::X ? TEXT("X") : Settings.PrimaryAxis == ERigControlAxis::Y ? TEXT("Y") : TEXT("Z")));
