Sequencer: Control Rig: Add INDEX_NONE check for generated keys
#trivial

diff --git a/Engine/Plugins/Animation/ControlRig/Source/ControlRigEditor/Private/Sequencer/ControlRigParameterTrackEditor.cpp b/Engine/Plugins/Animation/ControlRig/Source/ControlRigEditor/Private/Sequencer/ControlRigParameterTrackEditor.cpp
--- a/Engine/Plugins/Animation/ControlRig/Source/ControlRigEditor/Private/Sequencer/ControlRigParameterTrackEditor.cpp
+++ b/Engine/Plugins/Animation/ControlRig/Source/ControlRigEditor/Private/Sequencer/ControlRigParameterTrackEditor.cpp
@@ -2940,7 +2940,7 @@ bool FControlRigParameterTrackEditor::ModifyOurGeneratedKeysByCurrentAndWeight(U
 				if ((Val.ParameterName == ControlElement->GetName()))
 				{
 					pChannelIndex = Section->ControlChannelMap.Find(ControlElement->GetName());
-					if (pChannelIndex)
+					if (pChannelIndex && pChannelIndex->GeneratedKeyIndex != INDEX_NONE)
 					{
 						ChannelIndex = pChannelIndex->GeneratedKeyIndex;
 						float FVal = (float)Val.Val;
@@ -2965,7 +2965,7 @@ bool FControlRigParameterTrackEditor::ModifyOurGeneratedKeysByCurrentAndWeight(U
 				if ((Val.ParameterName == ControlElement->GetName()))
 				{
 					pChannelIndex = Section->ControlChannelMap.Find(ControlElement->GetName());
-					if (pChannelIndex)
+					if (pChannelIndex && pChannelIndex->GeneratedKeyIndex != INDEX_NONE)
 					{
 						ChannelIndex = pChannelIndex->GeneratedKeyIndex;
 						float FVal = (float)Val.Val.X;
@@ -2987,7 +2987,7 @@ bool FControlRigParameterTrackEditor::ModifyOurGeneratedKeysByCurrentAndWeight(U
 				if ((Val.ParameterName == ControlElement->GetName()))
 				{
 					pChannelIndex = Section->ControlChannelMap.Find(ControlElement->GetName());
-					if (pChannelIndex)
+					if (pChannelIndex && pChannelIndex->GeneratedKeyIndex != INDEX_NONE)
 					{
 						ChannelIndex = pChannelIndex->GeneratedKeyIndex;
 
@@ -3017,7 +3017,7 @@ bool FControlRigParameterTrackEditor::ModifyOurGeneratedKeysByCurrentAndWeight(U
 				if ((Val.ParameterName == ControlElement->GetName()))
 				{
 					pChannelIndex = Section->ControlChannelMap.Find(ControlElement->GetName());
-					if (pChannelIndex)
+					if (pChannelIndex && pChannelIndex->GeneratedKeyIndex != INDEX_NONE)
 					{
 						ChannelIndex = pChannelIndex->GeneratedKeyIndex;
 
