Control Rig: Fixed Execution Stack Index cropped in UI

diff --git a/Engine/Plugins/Animation/ControlRig/Source/ControlRigEditor/Private/Editor/SControlRigStackView.cpp b/Engine/Plugins/Animation/ControlRig/Source/ControlRigEditor/Private/Editor/SControlRigStackView.cpp
--- a/Engine/Plugins/Animation/ControlRig/Source/ControlRigEditor/Private/Editor/SControlRigStackView.cpp
+++ b/Engine/Plugins/Animation/ControlRig/Source/ControlRigEditor/Private/Editor/SControlRigStackView.cpp
@@ -93,7 +93,8 @@ void SRigStackItem::Construct(const FArguments& InArgs, const TSharedRef<STableV
 			[
 				SNew(SHorizontalBox)
 				+ SHorizontalBox::Slot()
-				.MaxWidth(25.f)
+				.MaxWidth(35.f)
+				.AutoWidth()
 				.VAlign(VAlign_Center)
 				.HAlign(HAlign_Right)
 				[
