Control Rig: Retain control values when controls are parented under bones

diff --git a/Engine/Plugins/Animation/ControlRig/Source/ControlRig/Private/Rigs/RigHierarchy.cpp b/Engine/Plugins/Animation/ControlRig/Source/ControlRig/Private/Rigs/RigHierarchy.cpp
--- a/Engine/Plugins/Animation/ControlRig/Source/ControlRig/Private/Rigs/RigHierarchy.cpp
+++ b/Engine/Plugins/Animation/ControlRig/Source/ControlRig/Private/Rigs/RigHierarchy.cpp
@@ -311,6 +311,7 @@ void URigHierarchy::Reset()
 
 	ResetPoseHash = INDEX_NONE;
 	ResetPoseIsFilteredOut.Reset();
+	ElementsToRetainLocalTransform.Reset();
 	DefaultParentPerElement.Reset();
 	OrderedSelection.Reset();
 
@@ -583,6 +584,7 @@ void URigHierarchy::ResetPoseToInitial(ERigElementType InTypeFilter)
 		if(Hash != ResetPoseHash)
 		{
 			ResetPoseIsFilteredOut.Reset();
+			ElementsToRetainLocalTransform.Reset();
 			ResetPoseHash = Hash;
 
 			// let's look at all elements and mark all parent of unaffected children
@@ -600,7 +602,14 @@ void URigHierarchy::ResetPoseToInitial(ERigElementType InTypeFilter)
 				{
 					if(!ResetPoseIsFilteredOut[Parent->GetIndex()])
 					{
-						ResetPoseIsFilteredOut[InElement->GetIndex()] = false;
+						if(InElement->IsA<FRigNullElement>() || InElement->IsA<FRigControlElement>())
+						{
+							ElementsToRetainLocalTransform.Add(InElement->GetIndex());
+						}
+						else
+						{
+							ResetPoseIsFilteredOut[InElement->GetIndex()] = false;
+						}
 					}
 				}
 			});
@@ -613,7 +622,38 @@ void URigHierarchy::ResetPoseToInitial(ERigElementType InTypeFilter)
 			bPerformFiltering = false;
 		}
 	}
-	
+
+	if(bPerformFiltering)
+	{
+		for(const int32 ElementIndex : ElementsToRetainLocalTransform)
+		{
+			if(FRigTransformElement* TransformElement = Get<FRigTransformElement>(ElementIndex))
+			{
+				// compute the local value if necessary
+				GetTransform(TransformElement, ERigTransformType::CurrentLocal);
+
+				if(FRigControlElement* ControlElement = Cast<FRigControlElement>(TransformElement))
+				{
+					// compute the local offset if necessary
+					GetControlOffsetTransform(ControlElement, ERigTransformType::CurrentLocal);
+				}
+			}
+		}
+		
+		for(const int32 ElementIndex : ElementsToRetainLocalTransform)
+		{
+			if(FRigTransformElement* TransformElement = Get<FRigTransformElement>(ElementIndex))
+			{
+				TransformElement->Pose.MarkDirty(ERigTransformType::CurrentGlobal);
+
+				if(FRigControlElement* ControlElement = Cast<FRigControlElement>(TransformElement))
+				{
+					ControlElement->Offset.MarkDirty(ERigTransformType::CurrentGlobal);
+				}
+			}
+		}
+	}
+
 	for(int32 ElementIndex=0; ElementIndex<Elements.Num(); ElementIndex++)
 	{
 		if(!ResetPoseIsFilteredOut.IsEmpty() && bPerformFiltering)
diff --git a/Engine/Plugins/Animation/ControlRig/Source/ControlRig/Public/Rigs/RigHierarchy.h b/Engine/Plugins/Animation/ControlRig/Source/ControlRig/Public/Rigs/RigHierarchy.h
--- a/Engine/Plugins/Animation/ControlRig/Source/ControlRig/Public/Rigs/RigHierarchy.h
+++ b/Engine/Plugins/Animation/ControlRig/Source/ControlRig/Public/Rigs/RigHierarchy.h
@@ -2831,6 +2831,7 @@ private:
 
 	int32 ResetPoseHash;
 	TArray<bool> ResetPoseIsFilteredOut;
+	TArray<int32> ElementsToRetainLocalTransform;
 
 #if WITH_EDITOR
 
