Control Rig: Fix crash when renaming a control

diff --git a/Engine/Plugins/Animation/ControlRig/Source/ControlRigEditor/Private/EditMode/ControlRigEditMode.cpp b/Engine/Plugins/Animation/ControlRig/Source/ControlRigEditor/Private/EditMode/ControlRigEditMode.cpp
--- a/Engine/Plugins/Animation/ControlRig/Source/ControlRigEditor/Private/EditMode/ControlRigEditMode.cpp
+++ b/Engine/Plugins/Animation/ControlRig/Source/ControlRigEditor/Private/EditMode/ControlRigEditMode.cpp
@@ -3236,6 +3236,18 @@ void FControlRigEditMode::OnHierarchyModified(ERigHierarchyNotification InNotif,
 
 void FControlRigEditMode::OnHierarchyModified_AnyThread(ERigHierarchyNotification InNotif, URigHierarchy* InHierarchy, const FRigBaseElement* InElement)
 {
+	if(bSuspendHierarchyNotifs)
+	{
+		return;
+	}
+
+	if (InNotif != ERigHierarchyNotification::ControlSettingChanged
+		&& InNotif != ERigHierarchyNotification::ControlShapeTransformChanged)
+	{
+		OnHierarchyModified(InNotif, InHierarchy, InElement);
+		return;
+	}
+	
 	FRigElementKey Key;
 	if(InElement)
 	{
