ControlRig: Added Unit Test for Attribute Nodes

diff --git a/Engine/Plugins/Animation/ControlRig/Source/ControlRig/Public/ControlRig.h b/Engine/Plugins/Animation/ControlRig/Source/ControlRig/Public/ControlRig.h
--- a/Engine/Plugins/Animation/ControlRig/Source/ControlRig/Public/ControlRig.h
+++ b/Engine/Plugins/Animation/ControlRig/Source/ControlRig/Public/ControlRig.h
@@ -518,7 +518,7 @@ private:
 		return nullptr;
 	}
 public:
-	class FAnimAttributeContainerPtrScope
+	class CONTROLRIG_API FAnimAttributeContainerPtrScope
 	{
 	public:
 		FAnimAttributeContainerPtrScope(UControlRig* InControlRig, UE::Anim::FStackAttributeContainer& InExternalContainer);
