Replacing legacy EditorStyle calls with AppStyle

diff --git a/Engine/Plugins/Animation/ControlRig/Source/ControlRigEditor/Private/Editor/SControlRigBreakLinksWidget.cpp b/Engine/Plugins/Animation/ControlRig/Source/ControlRigEditor/Private/Editor/SControlRigBreakLinksWidget.cpp
--- a/Engine/Plugins/Animation/ControlRig/Source/ControlRigEditor/Private/Editor/SControlRigBreakLinksWidget.cpp
+++ b/Engine/Plugins/Animation/ControlRig/Source/ControlRigEditor/Private/Editor/SControlRigBreakLinksWidget.cpp
@@ -8,7 +8,7 @@
 #include "Widgets/Layout/SUniformGridPanel.h"
 #include "DetailLayoutBuilder.h"
 #include "Editor.h"
-#include "EditorStyleSet.h"
+#include "Styling/AppStyle.h"
 #include "HAL/ConsoleManager.h"
 #include "RigVMModel/RigVMController.h"
 #include "ControlRigEditor.h"
@@ -84,7 +84,7 @@ void SControlRigBreakLinksWidget::Construct(const FArguments& InArgs, TArray<URi
 		[
 			SNew(SBorder)
 			.Visibility(EVisibility::Visible)
-			.BorderImage(FEditorStyle::GetBrush("Menu.Background"))
+			.BorderImage(FAppStyle::GetBrush("Menu.Background"))
 			[			
 			   SNew( SListView< URigVMLink* > )
 			     .ItemHeight(24)
@@ -167,14 +167,14 @@ void SControlRigBreakLinksDialog::Construct(const FArguments& InArgs)
                 .Padding(5)
                 [
                     SNew(SUniformGridPanel)
-                    .SlotPadding(FEditorStyle::GetMargin("StandardDialog.SlotPadding"))
-                    .MinDesiredSlotWidth(FEditorStyle::GetFloat("StandardDialog.MinDesiredSlotWidth"))
-                    .MinDesiredSlotHeight(FEditorStyle::GetFloat("StandardDialog.MinDesiredSlotHeight"))
+                    .SlotPadding(FAppStyle::GetMargin("StandardDialog.SlotPadding"))
+                    .MinDesiredSlotWidth(FAppStyle::GetFloat("StandardDialog.MinDesiredSlotWidth"))
+                    .MinDesiredSlotHeight(FAppStyle::GetFloat("StandardDialog.MinDesiredSlotHeight"))
                     +SUniformGridPanel::Slot(0, 0)
                     [
                         SNew(SButton)
                         .HAlign(HAlign_Center)
-                        .ContentPadding(FEditorStyle::GetMargin("StandardDialog.ContentPadding"))
+                        .ContentPadding(FAppStyle::GetMargin("StandardDialog.ContentPadding"))
                         .Text(LOCTEXT("OK", "OK"))
                         .OnClicked(this, &SControlRigBreakLinksDialog::OnButtonClick, EAppReturnType::Ok)
                     ]
@@ -182,7 +182,7 @@ void SControlRigBreakLinksDialog::Construct(const FArguments& InArgs)
                     [
                         SNew(SButton)
                         .HAlign(HAlign_Center)
-                        .ContentPadding(FEditorStyle::GetMargin("StandardDialog.ContentPadding"))
+                        .ContentPadding(FAppStyle::GetMargin("StandardDialog.ContentPadding"))
                         .Text(LOCTEXT("Cancel", "Cancel"))
                         .OnClicked(this, &SControlRigBreakLinksDialog::OnButtonClick, EAppReturnType::Cancel)
                     ]
