Control Rig: Fix function not editable after renaming

diff --git a/Engine/Plugins/Animation/ControlRig/Source/ControlRigDeveloper/Private/ControlRigBlueprint.cpp b/Engine/Plugins/Animation/ControlRig/Source/ControlRigDeveloper/Private/ControlRigBlueprint.cpp
--- a/Engine/Plugins/Animation/ControlRig/Source/ControlRigDeveloper/Private/ControlRigBlueprint.cpp
+++ b/Engine/Plugins/Animation/ControlRig/Source/ControlRigDeveloper/Private/ControlRigBlueprint.cpp
@@ -3383,6 +3383,20 @@ void UControlRigBlueprint::HandleModifiedEvent(ERigVMGraphNotifType InNotifType,
 			{
 				if (URigVMCollapseNode* CollapseNode = Cast<URigVMCollapseNode>(InSubject))
 				{
+					// Fix the path to the contained graph controller
+					{
+						FSoftObjectPath ContainedGraphPath = CollapseNode->GetContainedGraph();
+						FString SubPathString = ContainedGraphPath.GetSubPathString();
+						FString OldSubPathString = SubPathString.Replace(*CollapseNode->GetName(), *CollapseNode->GetPreviousFName().ToString());
+						FSoftObjectPath OldContainedGraphPath = ContainedGraphPath;
+						OldContainedGraphPath.SetSubPathString(OldSubPathString);
+						if (TObjectPtr<URigVMController>* Controller = Controllers.Find(OldContainedGraphPath))
+						{
+							Controllers.Add(ContainedGraphPath, *Controller);
+							Controllers.Remove(OldContainedGraphPath);
+						}
+					}
+					
 					FString NewNodePath = CollapseNode->GetNodePath(true /* recursive */);
 					FString Left, Right = NewNodePath;
 					URigVMNode::SplitNodePathAtEnd(NewNodePath, Left, Right);
diff --git a/Engine/Source/Developer/RigVMDeveloper/Private/RigVMModel/RigVMController.cpp b/Engine/Source/Developer/RigVMDeveloper/Private/RigVMModel/RigVMController.cpp
--- a/Engine/Source/Developer/RigVMDeveloper/Private/RigVMModel/RigVMController.cpp
+++ b/Engine/Source/Developer/RigVMDeveloper/Private/RigVMModel/RigVMController.cpp
@@ -15450,8 +15450,8 @@ bool URigVMController::PropagateTemplateFilteredTypes(URigVMTemplateNode* InNode
 						}
 						else
 						{
-							ensureMsgf(bIsTransacting, TEXT("Unexpected link broken in package %s"), *GetPackage()->GetPathName());
 							URigVMLink* Link = Pin->FindLinkForPin(OtherPin);
+							ensureMsgf(bIsTransacting, TEXT("Unexpected link broken %s in package %s"), *Link->GetPinPathRepresentation(), *GetPackage()->GetPathName());
 							BreakLink(Link->GetSourcePin(), Link->GetTargetPin(), bSetupUndoRedo);
 							return false;
 						}
