Control Rig: Fix crash when changing variable type to rotator

diff --git a/Engine/Plugins/Animation/ControlRig/Source/ControlRigDeveloper/Private/Graph/ControlRigGraphNode.cpp b/Engine/Plugins/Animation/ControlRig/Source/ControlRigDeveloper/Private/Graph/ControlRigGraphNode.cpp
--- a/Engine/Plugins/Animation/ControlRig/Source/ControlRigDeveloper/Private/Graph/ControlRigGraphNode.cpp
+++ b/Engine/Plugins/Animation/ControlRig/Source/ControlRigDeveloper/Private/Graph/ControlRigGraphNode.cpp
@@ -139,8 +139,11 @@ FText UControlRigGraphNode::GetNodeTitle(ENodeTitleType::Type TitleType) const
 									{
 										TArray<FString> Values;
 										DefaultValue.ParseIntoArray(Values, TEXT(","));
-										Values.Swap(0, 1);
-										Values.Swap(0, 2);
+										if (Values.Num() == 3)
+										{
+											Values.Swap(0, 1);
+											Values.Swap(0, 2);
+										}
 										DefaultValue = FString::Join(Values, TEXT(","));										
 									}
 									SubTitle = FString::Printf(TEXT("Default %s"), *DefaultValue);
