Control Rig: Proxy control support

diff --git a/Engine/Plugins/Animation/ControlRig/Source/ControlRig/Private/ControlRig.cpp b/Engine/Plugins/Animation/ControlRig/Source/ControlRig/Private/ControlRig.cpp
--- a/Engine/Plugins/Animation/ControlRig/Source/ControlRig/Private/ControlRig.cpp
+++ b/Engine/Plugins/Animation/ControlRig/Source/ControlRig/Private/ControlRig.cpp
@@ -13,6 +13,7 @@
 #include "Units/Execution/RigUnit_BeginExecution.h"
 #include "Units/Execution/RigUnit_PrepareForExecution.h"
 #include "Units/Execution/RigUnit_InverseExecution.h"
+#include "Units/Execution/RigUnit_InteractionExecution.h"
 #include "Components/SkeletalMeshComponent.h"
 #include "Animation/AnimInstance.h"
 #include "Animation/AnimationPoseData.h"
@@ -225,6 +226,7 @@ void UControlRig::Initialize(bool bInitRigUnits)
 	GetHierarchy()->OnModified().AddUObject(this, &UControlRig::HandleHierarchyModified);
 	GetHierarchy()->OnEventReceived().RemoveAll(this);
 	GetHierarchy()->OnEventReceived().AddUObject(this, &UControlRig::HandleHierarchyEvent);
+	GetHierarchy()->UpdateVisibilityOnProxyControls();
 }
 
 void UControlRig::InitializeFromCDO()
@@ -272,6 +274,21 @@ void UControlRig::Evaluate_AnyThread()
 	DECLARE_SCOPE_HIERARCHICAL_COUNTER_FUNC()
 	QUICK_SCOPE_CYCLE_COUNTER(STAT_ControlRig_Evaluate);
 
+	TGuardValue<TArray<FName>> EventQueueGuard(EventQueue, EventQueue);
+	if(InteractionType != (uint8)EControlRigInteractionType::None)
+	{
+		if(EventQueue.IsEmpty())
+		{
+			EventQueue.Add(FRigUnit_InteractionExecution::EventName);
+		}
+		else
+		{
+			// insert just before the last event so the interaction runs prior to
+			// forward solve or backwards solve.
+			EventQueue.Insert(FRigUnit_InteractionExecution::EventName, EventQueue.Num() - 1);
+		}
+	}
+
 	for (const FName& EventName : EventQueue)
 	{
 		Execute(EControlRigState::Update, EventName);
@@ -1005,8 +1022,7 @@ void UControlRig::Execute(const EControlRigState InState, const FName& InEventNa
 		{
 			const FRigControlSettings& Settings = ControlElement->Settings;
 
-			if (Settings.bShapeEnabled &&
-				Settings.bShapeVisible &&
+			if (Settings.IsVisible() &&
 				!Settings.bIsTransientControl &&
 				Settings.bDrawLimits &&
 				Settings.LimitEnabled.Contains(FRigControlLimitEnabled(true, true)))
@@ -1627,13 +1643,10 @@ void UControlRig::CreateRigControlsForCurveContainer()
 			if (Name.Contains(CtrlPrefix) && !DynamicHierarchy->Contains(FRigElementKey(*Name, ERigElementType::Curve))) //-V1051
 			{
 				FRigControlSettings Settings;
+				Settings.AnimationType = ERigControlAnimationType::AnimationChannel;
 				Settings.ControlType = ERigControlType::Float;
 				Settings.bIsCurve = true;
-				Settings.bAnimatable = true;
 				Settings.bDrawLimits = false;
-				Settings.bShapeEnabled = false;
-				Settings.bShapeVisible = false;
-				Settings.ShapeColor = FLinearColor::Red;
 
 				FRigControlValue Value;
 				Value.Set<float>(CurveElement->Value);
diff --git a/Engine/Plugins/Animation/ControlRig/Source/ControlRig/Private/ControlRigControlActor.cpp b/Engine/Plugins/Animation/ControlRig/Source/ControlRig/Private/ControlRigControlActor.cpp
--- a/Engine/Plugins/Animation/ControlRig/Source/ControlRig/Private/ControlRigControlActor.cpp
+++ b/Engine/Plugins/Animation/ControlRig/Source/ControlRig/Private/ControlRigControlActor.cpp
@@ -171,7 +171,7 @@ void AControlRigControlActor::Refresh()
 		URigHierarchy* Hierarchy = ControlRig->GetHierarchy();
 		Hierarchy->ForEach<FRigControlElement>([this, ShapeLibraries, Hierarchy](FRigControlElement* ControlElement) -> bool
         {
-			if (!ControlElement->Settings.bShapeEnabled)
+			if (!ControlElement->Settings.SupportsShape())
 			{
 				return true;
 			}
diff --git a/Engine/Plugins/Animation/ControlRig/Source/ControlRig/Private/ControlRigGizmoActor.cpp b/Engine/Plugins/Animation/ControlRig/Source/ControlRig/Private/ControlRigGizmoActor.cpp
--- a/Engine/Plugins/Animation/ControlRig/Source/ControlRig/Private/ControlRigGizmoActor.cpp
+++ b/Engine/Plugins/Animation/ControlRig/Source/ControlRig/Private/ControlRigGizmoActor.cpp
@@ -14,6 +14,7 @@ AControlRigShapeActor::AControlRigShapeActor(const FObjectInitializer& ObjectIni
 	, ControlRig(nullptr)
 	, ControlName(NAME_None)
 	, ShapeName(NAME_None)
+	, OverrideColor(0, 0, 0, 0)
 	, bEnabled(true)
 	, bSelected(false)
 	, bSelectable(true)
@@ -55,6 +56,10 @@ bool AControlRigShapeActor::IsEnabled() const
 
 void AControlRigShapeActor::SetSelected(bool bInSelected)
 {
+	if(!bSelectable && bInSelected)
+	{
+		return;
+	}
 	if(bSelected != bInSelected)
 	{
 		bSelected = bInSelected;
@@ -123,7 +128,7 @@ bool AControlRigShapeActor::UpdateControlSettings(
 	const FRigControlSettings& ControlSettings = InControlElement->Settings;
 	
 	// if this actor is not supposed to exist
-	if(!ControlSettings.bShapeEnabled)
+	if(!ControlSettings.SupportsShape())
 	{
 		return false;
 	}
@@ -179,13 +184,13 @@ bool AControlRigShapeActor::UpdateControlSettings(
 
 #if WITH_EDITOR
 	// update visibility
-	SetIsTemporarilyHiddenInEditor(!ControlSettings.bShapeVisible || bHideManipulators);
+	SetIsTemporarilyHiddenInEditor(!ControlSettings.IsVisible() || bHideManipulators);
 #endif
 	
 	// update selectability state
 	if (!bIsInLevelEditor) //don't change this in level editor otherwise we can never select it
 	{
-		SetSelectable(ControlSettings.bShapeVisible && !bHideManipulators && ControlSettings.bAnimatable);
+		SetSelectable(ControlSettings.IsSelectable() && !bHideManipulators);
 	}
 
 	return true;
diff --git a/Engine/Plugins/Animation/ControlRig/Source/ControlRig/Private/Rigs/RigHierarchy.cpp b/Engine/Plugins/Animation/ControlRig/Source/ControlRig/Private/Rigs/RigHierarchy.cpp
--- a/Engine/Plugins/Animation/ControlRig/Source/ControlRig/Private/Rigs/RigHierarchy.cpp
+++ b/Engine/Plugins/Animation/ControlRig/Source/ControlRig/Private/Rigs/RigHierarchy.cpp
@@ -2290,6 +2290,7 @@ void URigHierarchy::SendQueuedNotifications()
 			}
 			case ERigHierarchyNotification::ControlSettingChanged:
 			case ERigHierarchyNotification::ControlVisibilityChanged:
+			case ERigHierarchyNotification::ControlDrivenListChanged:
 			case ERigHierarchyNotification::ControlShapeTransformChanged:
 			case ERigHierarchyNotification::ParentChanged:
 			case ERigHierarchyNotification::ParentWeightsChanged:
@@ -3195,8 +3196,10 @@ void URigHierarchy::SetControlVisibility(FRigControlElement* InControlElement, b
 		return;
 	}
 
-	InControlElement->Settings.bShapeVisible = bVisibility;
-	Notify(ERigHierarchyNotification::ControlVisibilityChanged, InControlElement);
+	if(InControlElement->Settings.SetVisible(bVisibility))
+	{
+		Notify(ERigHierarchyNotification::ControlVisibilityChanged, InControlElement);
+	}
 
 #if WITH_EDITOR
 	if (ensure(!bPropagatingChange))
@@ -4441,6 +4444,64 @@ URigHierarchy::TElementDependencyMap URigHierarchy::GetDependenciesForVM(const U
 
 #endif
 
+void URigHierarchy::UpdateVisibilityOnProxyControls()
+{
+	URigHierarchy* HierarchyForSelection = HierarchyForSelectionPtr.Get();
+	if(HierarchyForSelection == nullptr)
+	{
+		HierarchyForSelection = this;
+	}
+
+	const UWorld* World = GetWorld();
+	if(World == nullptr)
+	{
+		return;
+	}
+	if(World->IsPreviewWorld())
+	{
+		return;
+	}
+
+	for(FRigBaseElement* Element : Elements)
+	{
+		if(FRigControlElement* ControlElement = Cast<FRigControlElement>(Element))
+		{
+			if(ControlElement->Settings.AnimationType == ERigControlAnimationType::ProxyControl)
+			{
+				if(ControlElement->Settings.ShapeVisibility == ERigControlVisibility::BasedOnSelection)
+				{
+					if(HierarchyForSelection->OrderedSelection.IsEmpty())
+					{
+						if(ControlElement->Settings.SetVisible(false, true))
+						{
+							Notify(ERigHierarchyNotification::ControlVisibilityChanged, ControlElement);
+						}
+					}
+					else
+					{
+						bool bVisible = HierarchyForSelection->IsSelected(ControlElement);
+						if(!bVisible)
+						{
+							if(ControlElement->Settings.DrivenControls.FindByPredicate([HierarchyForSelection](const FRigElementKey& AffectedControl) -> bool
+							{
+								return HierarchyForSelection->IsSelected(AffectedControl);
+							}) != nullptr)
+							{
+								bVisible = true;
+							}
+						}
+
+						if(ControlElement->Settings.SetVisible(bVisible, true))
+						{
+							Notify(ERigHierarchyNotification::ControlVisibilityChanged, ControlElement);
+						}
+					}
+				}
+			}
+		}
+	}
+}
+
 const TArray<FString>& URigHierarchy::GetTransformTypeStrings()
 {
 	static TArray<FString> TransformTypeStrings;
@@ -4673,6 +4734,38 @@ void URigHierarchy::ComputeAllTransforms()
 	}
 }
 
+bool URigHierarchy::IsAnimatable(const FRigElementKey& InKey) const
+{
+	if(const FRigControlElement* ControlElement = Find<FRigControlElement>(InKey))
+	{
+		return IsAnimatable(ControlElement);
+	}
+	return false;
+}
+
+bool URigHierarchy::IsAnimatable(const FRigControlElement* InControlElement) const
+{
+	if(InControlElement)
+	{
+		if(!InControlElement->Settings.IsAnimatable())
+		{
+			return false;
+		}
+
+		// animation channels are dependent on the control they are under.
+		if(InControlElement->Settings.AnimationType == ERigControlAnimationType::AnimationChannel)
+		{
+			if(const FRigControlElement* ParentControlElement = Cast<FRigControlElement>(GetFirstParent(InControlElement)))
+			{
+				return IsAnimatable(ParentControlElement);
+			}
+		}
+		
+		return true;
+	}
+	return false;
+}
+
 FTransform URigHierarchy::GetWorldTransformForReference(const FRigUnitContext* InContext, const FRigElementKey& InKey, bool bInitial)
 {
 	if(const USceneComponent* OuterSceneComponent = GetTypedOuter<USceneComponent>())
@@ -5330,21 +5423,9 @@ TArray<FString> URigHierarchy::ControlSettingsToPythonCommands(const FRigControl
 	TArray<FString> Commands;
 	Commands.Add(FString::Printf(TEXT("%s = unreal.RigControlSettings()"),
 			*NameSettings));
-	FString TypeStr;
-	switch (Settings.ControlType)
-	{
-		case ERigControlType::Bool: TypeStr = TEXT("BOOL"); break;							
-		case ERigControlType::Float: TypeStr = TEXT("FLOAT"); break;
-		case ERigControlType::Integer: TypeStr = TEXT("INTEGER"); break;
-		case ERigControlType::Position: TypeStr = TEXT("POSITION"); break;
-		case ERigControlType::Rotator: TypeStr = TEXT("ROTATOR"); break;
-		case ERigControlType::Scale: TypeStr = TEXT("SCALE"); break;
-		case ERigControlType::Transform: TypeStr = TEXT("EULER_TRANSFORM"); break;
-		case ERigControlType::EulerTransform: TypeStr = TEXT("EULER_TRANSFORM"); break;
-		case ERigControlType::Vector2D: TypeStr = TEXT("VECTOR2D"); break;
-		case ERigControlType::TransformNoScale: TypeStr = TEXT("EULER_TRANSFORM"); break;
-		default: ensure(false);
-	}
+	
+	const FString AnimationTypeStr = StaticEnum<ERigControlAnimationType>()->GetNameStringByValue((int64)Settings.AnimationType).ToUpper();
+	const FString ControlTypeStr = StaticEnum<ERigControlType>()->GetNameStringByValue((int64)Settings.ControlType).ToUpper();
 
 	static const TCHAR* TrueText = TEXT("True");
 	static const TCHAR* FalseText = TEXT("False");
@@ -5359,12 +5440,12 @@ TArray<FString> URigHierarchy::ControlSettingsToPythonCommands(const FRigControl
 	
 	const FString LimitEnabledStr = FString::Join(LimitEnabledParts, TEXT(", "));
 	
+	Commands.Add(FString::Printf(TEXT("%s.animation_type = unreal.RigControlAnimationType.%s"),
+									*NameSettings,
+									*AnimationTypeStr));
 	Commands.Add(FString::Printf(TEXT("%s.control_type = unreal.RigControlType.%s"),
 									*NameSettings,
-									*TypeStr));
-	Commands.Add(FString::Printf(TEXT("%s.animatable = %s"),
-		*NameSettings,
-		Settings.bAnimatable ? TrueText : FalseText));
+									*ControlTypeStr));
 	Commands.Add(FString::Printf(TEXT("%s.display_name = '%s'"),
 		*NameSettings,
 		*Settings.DisplayName.ToString()));
@@ -5374,9 +5455,6 @@ TArray<FString> URigHierarchy::ControlSettingsToPythonCommands(const FRigControl
 	Commands.Add(FString::Printf(TEXT("%s.shape_color = %s"),
 		*NameSettings,
 		*RigVMPythonUtils::LinearColorToPythonString(Settings.ShapeColor)));
-	Commands.Add(FString::Printf(TEXT("%s.shape_enabled = %s"),
-		*NameSettings,
-		Settings.bShapeEnabled ? TrueText : FalseText));
 	Commands.Add(FString::Printf(TEXT("%s.shape_name = '%s'"),
 		*NameSettings,
 		*Settings.ShapeName.ToString()));
diff --git a/Engine/Plugins/Animation/ControlRig/Source/ControlRig/Private/Rigs/RigHierarchyController.cpp b/Engine/Plugins/Animation/ControlRig/Source/ControlRig/Private/Rigs/RigHierarchyController.cpp
--- a/Engine/Plugins/Animation/ControlRig/Source/ControlRig/Private/Rigs/RigHierarchyController.cpp
+++ b/Engine/Plugins/Animation/ControlRig/Source/ControlRig/Private/Rigs/RigHierarchyController.cpp
@@ -104,6 +104,8 @@ bool URigHierarchyController::SelectElement(FRigElementKey InKey, bool bSelect,
 		Notify(ERigHierarchyNotification::ElementDeselected, Element);
 	}
 
+	Hierarchy->UpdateVisibilityOnProxyControls();
+
 	return true;
 }
 
@@ -151,7 +153,7 @@ bool URigHierarchyController::SetSelection(const TArray<FRigElementKey>& InKeys,
 			}
 		}
 	}
-	
+
 #if WITH_EDITOR
 	if (bPrintPythonCommand && !bSuspendPythonPrinting)
 	{
@@ -169,7 +171,7 @@ bool URigHierarchyController::SetSelection(const TArray<FRigElementKey>& InKeys,
 		}
 	}
 #endif
-	
+
 	return bResult;
 }
 
@@ -1264,12 +1266,12 @@ TArray<FRigElementKey> URigHierarchyController::ImportFromHierarchyContainer(con
 		Settings.DisplayName = Control.DisplayName;
 		Settings.PrimaryAxis = Control.PrimaryAxis;
 		Settings.bIsCurve = Control.bIsCurve;
-		Settings.bAnimatable = Control.bAnimatable;
+		Settings.SetAnimationTypeFromDeprecatedData(Control.bAnimatable, Control.bGizmoEnabled);
+		Settings.AnimationType = Control.bAnimatable ? ERigControlAnimationType::AnimationControl : ERigControlAnimationType::ProxyControl;
 		Settings.SetupLimitArrayForType(Control.bLimitTranslation, Control.bLimitRotation, Control.bLimitScale);
 		Settings.bDrawLimits = Control.bDrawLimits;
 		Settings.MinimumValue = Control.MinimumValue;
 		Settings.MaximumValue = Control.MaximumValue;
-		Settings.bShapeEnabled = Control.bGizmoEnabled;
 		Settings.bShapeVisible = Control.bGizmoVisible;
 		Settings.ShapeName = Control.GizmoName;
 		Settings.ShapeColor = Control.GizmoColor;
@@ -1767,6 +1769,7 @@ bool URigHierarchyController::RemoveElement(FRigBaseElement* InElement)
 		{
 			ControlElement->Settings.Customization.AvailableSpaces.Remove(InElement->GetKey());
 			ControlElement->Settings.Customization.RemovedSpaces.Remove(InElement->GetKey());
+			ControlElement->Settings.DrivenControls.Remove(InElement->GetKey());
 		}
 	}
 	
@@ -1898,6 +1901,14 @@ bool URigHierarchyController::RenameElement(FRigBaseElement* InElement, const FN
 					Favorite.Name = NewKey.Name;
 				}
 			}
+
+			for(FRigElementKey& DrivenControl : ControlElement->Settings.DrivenControls)
+			{
+				if(DrivenControl == OldKey)
+				{
+					DrivenControl.Name = NewKey.Name;
+				}
+			}
 		}
 	}
 	
@@ -1984,6 +1995,15 @@ bool URigHierarchyController::AddParent(FRigBaseElement* InChild, FRigBaseElemen
 		}
 	}
 
+	// we can only parent things to controls which are not animation channels (animation channels are not 3D things)
+	if(FRigControlElement* ParentControlElement = Cast<FRigControlElement>(InParent))
+	{
+		if(ParentControlElement->Settings.AnimationType == ERigControlAnimationType::AnimationChannel)
+		{
+			return false;
+		}
+	}
+
 	if(Hierarchy->IsParentedTo(InParent, InChild))
 	{
 		ReportErrorf(TEXT("Cannot parent '%s' to '%s' - would cause a cycle."), *InChild->Key.ToString(), *InParent->Key.ToString());
diff --git a/Engine/Plugins/Animation/ControlRig/Source/ControlRig/Private/Rigs/RigHierarchyElements.cpp b/Engine/Plugins/Animation/ControlRig/Source/ControlRig/Private/Rigs/RigHierarchyElements.cpp
--- a/Engine/Plugins/Animation/ControlRig/Source/ControlRig/Private/Rigs/RigHierarchyElements.cpp
+++ b/Engine/Plugins/Animation/ControlRig/Source/ControlRig/Private/Rigs/RigHierarchyElements.cpp
@@ -413,22 +413,23 @@ void FRigBoneElement::CopyFrom(URigHierarchy* InHierarchy, FRigBaseElement* InOt
 ////////////////////////////////////////////////////////////////////////////////
 
 FRigControlSettings::FRigControlSettings()
-: ControlType(ERigControlType::EulerTransform)
+: AnimationType(ERigControlAnimationType::AnimationControl)
+, ControlType(ERigControlType::EulerTransform)
 , DisplayName(NAME_None)
 , PrimaryAxis(ERigControlAxis::X)
 , bIsCurve(false)
-, bAnimatable(true)
 , LimitEnabled()
 , bDrawLimits(true)
 , MinimumValue()
 , MaximumValue()
-, bShapeEnabled(true)
 , bShapeVisible(true)
+, ShapeVisibility(ERigControlVisibility::UserDefined)
 , ShapeName(NAME_None)
 , ShapeColor(FLinearColor::Red)
 , bIsTransientControl(false)
 , ControlEnum(nullptr)
 , Customization()
+, bGroupWithParentControl(false)
 {
 	// rely on the default provided by the shape definition
 	ShapeName = FControlRigShapeDefinition().ShapeName; 
@@ -438,10 +439,14 @@ void FRigControlSettings::Save(FArchive& Ar)
 {
 	Ar.UsingCustomVersion(FControlRigObjectVersion::GUID);
 
+	static const UEnum* AnimationTypeEnum = StaticEnum<ERigControlAnimationType>();
 	static const UEnum* ControlTypeEnum = StaticEnum<ERigControlType>();
+	static const UEnum* ShapeVisibilityEnum = StaticEnum<ERigControlVisibility>();
 	static const UEnum* ControlAxisEnum = StaticEnum<ERigControlAxis>();
 
+	FName AnimationTypeName = AnimationTypeEnum->GetNameByValue((int64)AnimationType);
 	FName ControlTypeName = ControlTypeEnum->GetNameByValue((int64)ControlType);
+	FName ShapeVisibilityName = ShapeVisibilityEnum->GetNameByValue((int64)ShapeVisibility);
 	FName PrimaryAxisName = ControlAxisEnum->GetNameByValue((int64)PrimaryAxis);
 
 	FString ControlEnumPathName;
@@ -450,43 +455,56 @@ void FRigControlSettings::Save(FArchive& Ar)
 		ControlEnumPathName = ControlEnum->GetPathName();
 	}
 
+	Ar << AnimationTypeName;
 	Ar << ControlTypeName;
 	Ar << DisplayName;
 	Ar << PrimaryAxisName;
 	Ar << bIsCurve;
-	Ar << bAnimatable;
 	Ar << LimitEnabled;
 	Ar << bDrawLimits;
 	Ar << MinimumValue;
 	Ar << MaximumValue;
-	Ar << bShapeEnabled;
 	Ar << bShapeVisible;
+	Ar << ShapeVisibilityName;
 	Ar << ShapeName;
 	Ar << ShapeColor;
 	Ar << bIsTransientControl;
 	Ar << ControlEnumPathName;
 	Ar << Customization.AvailableSpaces;
+	Ar << DrivenControls;
+	Ar << bGroupWithParentControl;
 }
 
 void FRigControlSettings::Load(FArchive& Ar)
 {
 	Ar.UsingCustomVersion(FControlRigObjectVersion::GUID);
 
+	static const UEnum* AnimationTypeEnum = StaticEnum<ERigControlAnimationType>();
 	static const UEnum* ControlTypeEnum = StaticEnum<ERigControlType>();
+	static const UEnum* ShapeVisibilityEnum = StaticEnum<ERigControlVisibility>();
 	static const UEnum* ControlAxisEnum = StaticEnum<ERigControlAxis>();
 
-	FName ControlTypeName, PrimaryAxisName;
+	FName AnimationTypeName, ControlTypeName, ShapeVisibilityName, PrimaryAxisName;
 	FString ControlEnumPathName;
 
 	bool bLimitTranslation_DEPRECATED = false;
 	bool bLimitRotation_DEPRECATED = false;
 	bool bLimitScale_DEPRECATED = false;
+	bool bAnimatable_DEPRECATED = false;
+	bool bShapeEnabled_DEPRECATED = false;
 
+	if (Ar.CustomVer(FControlRigObjectVersion::GUID) >= FControlRigObjectVersion::ControlAnimationType)
+	{
+		Ar << AnimationTypeName;
+	}
 	Ar << ControlTypeName;
 	Ar << DisplayName;
 	Ar << PrimaryAxisName;
 	Ar << bIsCurve;
-	Ar << bAnimatable;
+	if (Ar.CustomVer(FControlRigObjectVersion::GUID) < FControlRigObjectVersion::ControlAnimationType)
+	{
+		Ar << bAnimatable_DEPRECATED;
+	}
 	if (Ar.CustomVer(FControlRigObjectVersion::GUID) < FControlRigObjectVersion::PerChannelLimits)
 	{
 		Ar << bLimitTranslation_DEPRECATED;
@@ -511,8 +529,25 @@ void FRigControlSettings::Load(FArchive& Ar)
 		Ar << MaximumTransform;
 	}
 
-	Ar << bShapeEnabled;
+	ControlType = (ERigControlType)ControlTypeEnum->GetValueByName(ControlTypeName);
+	
+	if (Ar.CustomVer(FControlRigObjectVersion::GUID) < FControlRigObjectVersion::ControlAnimationType)
+	{
+		Ar << bShapeEnabled_DEPRECATED;
+		SetAnimationTypeFromDeprecatedData(bAnimatable_DEPRECATED, bShapeEnabled_DEPRECATED);
+		AnimationTypeName = AnimationTypeEnum->GetNameByValue((int64)AnimationType);
+	}
+	
 	Ar << bShapeVisible;
+	
+	if (Ar.CustomVer(FControlRigObjectVersion::GUID) < FControlRigObjectVersion::ControlAnimationType)
+	{
+		ShapeVisibilityName = ShapeVisibilityEnum->GetNameByValue((int64)ERigControlVisibility::UserDefined);
+	}
+	else
+	{
+		Ar << ShapeVisibilityName;
+	}
 	Ar << ShapeName;
 
 	if(Ar.CustomVer(FControlRigObjectVersion::GUID) < FControlRigObjectVersion::RenameGizmoToShape)
@@ -527,8 +562,9 @@ void FRigControlSettings::Load(FArchive& Ar)
 	Ar << bIsTransientControl;
 	Ar << ControlEnumPathName;
 
-	ControlType = (ERigControlType)ControlTypeEnum->GetValueByName(ControlTypeName);
+	AnimationType = (ERigControlAnimationType)AnimationTypeEnum->GetValueByName(AnimationTypeName);
 	PrimaryAxis = (ERigControlAxis)ControlAxisEnum->GetValueByName(PrimaryAxisName);
+	ShapeVisibility = (ERigControlVisibility)ShapeVisibilityEnum->GetValueByName(ShapeVisibilityName);
 
 	if (Ar.CustomVer(FControlRigObjectVersion::GUID) < FControlRigObjectVersion::StorageMinMaxValuesAsFloatStorage)
 	{
@@ -558,10 +594,35 @@ void FRigControlSettings::Load(FArchive& Ar)
 		Customization.AvailableSpaces.Reset();
 	}
 
+	if (Ar.CustomVer(FControlRigObjectVersion::GUID) >= FControlRigObjectVersion::ControlAnimationType)
+	{
+		Ar << DrivenControls;
+	}
+	else
+	{
+		DrivenControls.Reset();
+	}
+
+	PreviouslyDrivenControls.Reset();
+
 	if (Ar.CustomVer(FControlRigObjectVersion::GUID) < FControlRigObjectVersion::PerChannelLimits)
 	{
 		SetupLimitArrayForType(bLimitTranslation_DEPRECATED, bLimitRotation_DEPRECATED, bLimitScale_DEPRECATED);
 	}
+
+	if (Ar.CustomVer(FControlRigObjectVersion::GUID) >= FControlRigObjectVersion::ControlAnimationType)
+	{
+		Ar << bGroupWithParentControl;
+	}
+	else
+	{
+		bGroupWithParentControl = IsAnimatable() && (
+			ControlType == ERigControlType::Bool ||
+			ControlType == ERigControlType::Float ||
+			ControlType == ERigControlType::Integer ||
+			ControlType == ERigControlType::Vector2D
+		);
+	}
 }
 
 ////////////////////////////////////////////////////////////////////////////////
@@ -570,6 +631,10 @@ void FRigControlSettings::Load(FArchive& Ar)
 
 bool FRigControlSettings::operator==(const FRigControlSettings& InOther) const
 {
+	if(AnimationType != InOther.AnimationType)
+	{
+		return false;
+	}
 	if(ControlType != InOther.ControlType)
 	{
 		return false;
@@ -586,10 +651,6 @@ bool FRigControlSettings::operator==(const FRigControlSettings& InOther) const
 	{
 		return false;
 	}
-	if(bAnimatable != InOther.bAnimatable)
-	{
-		return false;
-	}
 	if(LimitEnabled != InOther.LimitEnabled)
 	{
 		return false;
@@ -598,11 +659,11 @@ bool FRigControlSettings::operator==(const FRigControlSettings& InOther) const
 	{
 		return false;
 	}
-	if(bShapeEnabled != InOther.bShapeEnabled)
+	if(bShapeVisible != InOther.bShapeVisible)
 	{
 		return false;
 	}
-	if(bShapeVisible != InOther.bShapeVisible)
+	if(ShapeVisibility != InOther.ShapeVisibility)
 	{
 		return false;
 	}
@@ -626,7 +687,15 @@ bool FRigControlSettings::operator==(const FRigControlSettings& InOther) const
 	{
 		return false;
 	}
-
+	if(DrivenControls != InOther.DrivenControls)
+	{
+		return false;
+	}
+	if(bGroupWithParentControl != InOther.bGroupWithParentControl)
+	{
+		return false;
+	}
+	
 	const FTransform MinimumTransform = MinimumValue.GetAsTransform(ControlType, PrimaryAxis);
 	const FTransform OtherMinimumTransform = InOther.MinimumValue.GetAsTransform(ControlType, PrimaryAxis);
 	if(!MinimumTransform.Equals(OtherMinimumTransform, 0.001))
diff --git a/Engine/Plugins/Animation/ControlRig/Source/ControlRig/Private/Sequencer/MovieSceneControlRigParameterSection.cpp b/Engine/Plugins/Animation/ControlRig/Source/ControlRig/Private/Sequencer/MovieSceneControlRigParameterSection.cpp
--- a/Engine/Plugins/Animation/ControlRig/Source/ControlRig/Private/Sequencer/MovieSceneControlRigParameterSection.cpp
+++ b/Engine/Plugins/Animation/ControlRig/Source/ControlRig/Private/Sequencer/MovieSceneControlRigParameterSection.cpp
@@ -1349,7 +1349,7 @@ void UMovieSceneControlRigParameterSection::ReconstructChannelProxy()
 
 		for (FRigControlElement* ControlElement : SortedControls)
 		{
-			if (!ControlElement->Settings.bAnimatable)
+			if (!ControlRig->GetHierarchy()->IsAnimatable(ControlElement))
 			{
 				continue;
 			}
@@ -1361,6 +1361,16 @@ void UMovieSceneControlRigParameterSection::ReconstructChannelProxy()
 				ParentControlName = ParentControlElement->GetName();
 			}
 			bool bEnabled = ControlsMask[MaskIndex];
+
+			FText Group;
+			if(ParentControlElement && ControlElement->Settings.IsAnimatable() && ControlElement->Settings.bGroupWithParentControl)
+			{
+				if(ParentControlElement->Settings.AnimationType == ERigControlAnimationType::AnimationControl)
+				{
+					Group = FText::FromName(ParentControlElement->GetDisplayName());
+				}
+			}
+			
 #if WITH_EDITOR
 			switch (ControlElement->Settings.ControlType)
 			{
@@ -1370,25 +1380,6 @@ void UMovieSceneControlRigParameterSection::ReconstructChannelProxy()
 					{
 						if (ControlElement->GetName() == Scalar.ParameterName)
 						{
-
-							FText Group;
-							if(ParentControlElement)
-							{
-								switch(ParentControlElement->Settings.ControlType)
-								{
-									case ERigControlType::Position:
-									case ERigControlType::Scale:
-									case ERigControlType::Rotator:
-									case ERigControlType::Transform:
-									case ERigControlType::EulerTransform:
-									case ERigControlType::TransformNoScale:
-									{
-										Group = FText::FromName(ParentControlElement->GetDisplayName());
-										break;
-									}
-								}
-							}
-
 							if (Group.IsEmpty())
 							{
 								ControlChannelMap.Add(Scalar.ParameterName, FChannelMapInfo(ControlIndex, TotalIndex, FloatChannelIndex, INDEX_NONE, NAME_None, MaskIndex, CategoryIndex));
@@ -1401,16 +1392,9 @@ void UMovieSceneControlRigParameterSection::ReconstructChannelProxy()
 							}
 							else
 							{
-								FChannelMapInfo* pChannelIndex = ControlChannelMap.Find(ParentControlName);
-								if (pChannelIndex)
-								{
-									ControlChannelMap.Add(Scalar.ParameterName, FChannelMapInfo(ControlIndex, TotalIndex, FloatChannelIndex,pChannelIndex->ControlIndex, NAME_None, MaskIndex, CategoryIndex));
-								}
-								else
-								{
-									ControlChannelMap.Add(Scalar.ParameterName, FChannelMapInfo(ControlIndex, TotalIndex, FloatChannelIndex, INDEX_NONE, NAME_None, MaskIndex, CategoryIndex));
-
-								}
+								const FChannelMapInfo* pChannelIndex = ControlChannelMap.Find(ParentControlName);
+								const int32 ChannelControlIndex = pChannelIndex ? pChannelIndex->ControlIndex : INDEX_NONE; 
+								ControlChannelMap.Add(Scalar.ParameterName, FChannelMapInfo(ControlIndex, TotalIndex, FloatChannelIndex,ChannelControlIndex, NAME_None, MaskIndex, CategoryIndex));
 							}
 
 							FParameterFloatChannelEditorData EditorData(ControlRig, Scalar.ParameterName, bEnabled, Group, TotalIndex);
@@ -1429,24 +1413,6 @@ void UMovieSceneControlRigParameterSection::ReconstructChannelProxy()
 					{
 						if (ControlElement->GetName() == Bool.ParameterName)
 						{
-							FText Group;
-							if(ParentControlElement)
-							{
-								switch(ParentControlElement->Settings.ControlType)
-								{
-									case ERigControlType::Position:
-									case ERigControlType::Scale:
-									case ERigControlType::Rotator:
-									case ERigControlType::Transform:
-									case ERigControlType::EulerTransform:
-									case ERigControlType::TransformNoScale:
-									{
-										Group = FText::FromName(ParentControlElement->GetDisplayName());
-										break;
-									}
-								}
-							}
-
 							if (Group.IsEmpty())
 							{
 								ControlChannelMap.Add(Bool.ParameterName, FChannelMapInfo(ControlIndex, TotalIndex, BoolChannelIndex, INDEX_NONE,BoolChannelTypeName,MaskIndex, CategoryIndex));
@@ -1459,15 +1425,9 @@ void UMovieSceneControlRigParameterSection::ReconstructChannelProxy()
 							}
 							else
 							{
-								FChannelMapInfo* pChannelIndex = ControlChannelMap.Find(ParentControlName);
-								if (pChannelIndex)
-								{
-									ControlChannelMap.Add(Bool.ParameterName, FChannelMapInfo(ControlIndex, TotalIndex, BoolChannelIndex,pChannelIndex->ControlIndex,BoolChannelTypeName,MaskIndex, CategoryIndex));
-								}
-								else
-								{
-									ControlChannelMap.Add(Bool.ParameterName, FChannelMapInfo(ControlIndex, TotalIndex, BoolChannelIndex,INDEX_NONE, BoolChannelTypeName,MaskIndex, CategoryIndex));
-								}
+								const FChannelMapInfo* pChannelIndex = ControlChannelMap.Find(ParentControlName);
+								const int32 ChannelControlIndex = pChannelIndex ? pChannelIndex->ControlIndex : INDEX_NONE; 
+								ControlChannelMap.Add(Bool.ParameterName, FChannelMapInfo(ControlIndex, TotalIndex, BoolChannelIndex,ChannelControlIndex, BoolChannelTypeName,MaskIndex, CategoryIndex));
 							}
 
 							FMovieSceneChannelMetaData MetaData(Bool.ParameterName, Group, Group, bEnabled);
@@ -1490,24 +1450,6 @@ void UMovieSceneControlRigParameterSection::ReconstructChannelProxy()
 						{
 							if (ControlElement->GetName() == Enum.ParameterName)
 							{
-								FText Group;
-								if (ParentControlElement)
-								{
-									switch (ParentControlElement->Settings.ControlType)
-									{
-									case ERigControlType::Position:
-									case ERigControlType::Scale:
-									case ERigControlType::Rotator:
-									case ERigControlType::Transform:
-									case ERigControlType::EulerTransform:
-									case ERigControlType::TransformNoScale:
-									{
-										Group = FText::FromName(ParentControlElement->GetDisplayName());
-										break;
-									}
-									}
-								}
-
 								if (Group.IsEmpty())
 								{
 									ControlChannelMap.Add(Enum.ParameterName, FChannelMapInfo(ControlIndex, TotalIndex, EnumChannelIndex,INDEX_NONE, EnumChannelTypeName,MaskIndex, CategoryIndex));
@@ -1520,15 +1462,9 @@ void UMovieSceneControlRigParameterSection::ReconstructChannelProxy()
 								}
 								else
 								{
-									FChannelMapInfo* pChannelIndex = ControlChannelMap.Find(ParentControlName);
-									if (pChannelIndex)
-									{
-										ControlChannelMap.Add(Enum.ParameterName, FChannelMapInfo(ControlIndex, TotalIndex, EnumChannelIndex, pChannelIndex->ControlIndex, EnumChannelTypeName,MaskIndex, CategoryIndex));
-									}
-									else
-									{
-										ControlChannelMap.Add(Enum.ParameterName, FChannelMapInfo(ControlIndex, TotalIndex, EnumChannelIndex, INDEX_NONE, EnumChannelTypeName,MaskIndex, CategoryIndex));
-									}
+									const FChannelMapInfo* pChannelIndex = ControlChannelMap.Find(ParentControlName);
+									const int32 ChannelControlIndex = pChannelIndex ? pChannelIndex->ControlIndex : INDEX_NONE; 
+									ControlChannelMap.Add(Enum.ParameterName, FChannelMapInfo(ControlIndex, TotalIndex, EnumChannelIndex, ChannelControlIndex, EnumChannelTypeName,MaskIndex, CategoryIndex));
 								}
 
 								FMovieSceneChannelMetaData MetaData(Enum.ParameterName, Group, Group, bEnabled);
@@ -1548,23 +1484,6 @@ void UMovieSceneControlRigParameterSection::ReconstructChannelProxy()
 						{
 							if (ControlElement->GetName() == Integer.ParameterName)
 							{
-								FText Group;
-								if (ParentControlElement)
-								{
-									switch (ParentControlElement->Settings.ControlType)
-									{
-									case ERigControlType::Position:
-									case ERigControlType::Scale:
-									case ERigControlType::Rotator:
-									case ERigControlType::Transform:
-									case ERigControlType::EulerTransform:
-									case ERigControlType::TransformNoScale:
-									{
-										Group = FText::FromName(ParentControlElement->GetDisplayName());
-										break;
-									}
-									}
-								}
 								if (Group.IsEmpty())
 								{
 									ControlChannelMap.Add(Integer.ParameterName, FChannelMapInfo(ControlIndex, TotalIndex, IntegerChannelIndex,INDEX_NONE,IntegerChannelTypeName,MaskIndex, CategoryIndex));
@@ -1577,15 +1496,9 @@ void UMovieSceneControlRigParameterSection::ReconstructChannelProxy()
 								}
 								else
 								{
-									FChannelMapInfo* pChannelIndex = ControlChannelMap.Find(ParentControlName);
-									if (pChannelIndex)
-									{
-										ControlChannelMap.Add(Integer.ParameterName, FChannelMapInfo(ControlIndex, TotalIndex, IntegerChannelIndex, pChannelIndex->ControlIndex,IntegerChannelTypeName,MaskIndex, CategoryIndex));
-									}
-									else
-									{
-										ControlChannelMap.Add(Integer.ParameterName, FChannelMapInfo(ControlIndex, TotalIndex, IntegerChannelIndex, INDEX_NONE, IntegerChannelTypeName,MaskIndex, CategoryIndex));
-									}
+									const FChannelMapInfo* pChannelIndex = ControlChannelMap.Find(ParentControlName);
+									const int32 ChannelControlIndex = pChannelIndex ? pChannelIndex->ControlIndex : INDEX_NONE; 
+									ControlChannelMap.Add(Integer.ParameterName, FChannelMapInfo(ControlIndex, TotalIndex, IntegerChannelIndex, ChannelControlIndex, IntegerChannelTypeName,MaskIndex, CategoryIndex));
 								}
 
 								FMovieSceneChannelMetaData MetaData(Integer.ParameterName, Group, Group, bEnabled);
@@ -1608,13 +1521,22 @@ void UMovieSceneControlRigParameterSection::ReconstructChannelProxy()
 					{
 						if (ControlElement->GetName() == Vector2D.ParameterName)
 						{
-							ControlChannelMap.Add(Vector2D.ParameterName, FChannelMapInfo(ControlIndex, TotalIndex, FloatChannelIndex, INDEX_NONE, NAME_None, MaskIndex, CategoryIndex));
-							ControlIndex++;
-							if (bEnabled)
+							if(Group.IsEmpty())
 							{
-								++CategoryIndex;
+								ControlChannelMap.Add(Vector2D.ParameterName, FChannelMapInfo(ControlIndex, TotalIndex, FloatChannelIndex, INDEX_NONE, NAME_None, MaskIndex, CategoryIndex));
+								ControlIndex++;
+								if (bEnabled)
+								{
+									++CategoryIndex;
+								}
+								Group = FText::FromName(ControlElement->GetDisplayName());
+							}
+							else
+							{
+								const FChannelMapInfo* pChannelIndex = ControlChannelMap.Find(ParentControlName);
+								const int32 ChannelControlIndex = pChannelIndex ? pChannelIndex->ControlIndex : INDEX_NONE; 
+								ControlChannelMap.Add(Vector2D.ParameterName, FChannelMapInfo(ControlIndex, TotalIndex, FloatChannelIndex, ChannelControlIndex, NAME_None, MaskIndex, CategoryIndex));
 							}
-							FText Group = FText::FromName(ControlElement->GetDisplayName());
 							FParameterVectorChannelEditorData EditorData(ControlRig, Vector2D.ParameterName, bEnabled, Group, TotalIndex, 2);
 							Channels.Add(Vector2D.XCurve, EditorData.MetaData[0], EditorData.ExternalValues[0]);
 							Channels.Add(Vector2D.YCurve, EditorData.MetaData[1], EditorData.ExternalValues[1]);
@@ -1633,13 +1555,22 @@ void UMovieSceneControlRigParameterSection::ReconstructChannelProxy()
 					{
 						if (ControlElement->GetName() == Vector.ParameterName)
 						{
-							ControlChannelMap.Add(Vector.ParameterName, FChannelMapInfo(ControlIndex, TotalIndex, FloatChannelIndex, INDEX_NONE, NAME_None, MaskIndex, CategoryIndex));
-							ControlIndex++;
-							if (bEnabled)
+							if(Group.IsEmpty())
+							{
+								ControlChannelMap.Add(Vector.ParameterName, FChannelMapInfo(ControlIndex, TotalIndex, FloatChannelIndex, INDEX_NONE, NAME_None, MaskIndex, CategoryIndex));
+								ControlIndex++;
+								if (bEnabled)
+								{
+									++CategoryIndex;
+								}
+								Group = FText::FromName(ControlElement->GetDisplayName());
+							}
+							else
 							{
-								++CategoryIndex;
+								const FChannelMapInfo* pChannelIndex = ControlChannelMap.Find(ParentControlName);
+								const int32 ChannelControlIndex = pChannelIndex ? pChannelIndex->ControlIndex : INDEX_NONE; 
+								ControlChannelMap.Add(Vector.ParameterName, FChannelMapInfo(ControlIndex, TotalIndex, FloatChannelIndex, ChannelControlIndex, NAME_None, MaskIndex, CategoryIndex));
 							}
-							FText Group = FText::FromName(ControlElement->GetDisplayName());
 							if (FSpaceControlNameAndChannel* SpaceChannel = GetSpaceChannel(Vector.ParameterName))
 							{
 
@@ -1683,13 +1614,22 @@ void UMovieSceneControlRigParameterSection::ReconstructChannelProxy()
 					{
 						if (ControlElement->GetName() == Transform.ParameterName)
 						{
-							ControlChannelMap.Add(Transform.ParameterName, FChannelMapInfo(ControlIndex, TotalIndex, FloatChannelIndex, INDEX_NONE, NAME_None, MaskIndex, CategoryIndex));
-							ControlIndex++;
-							if (bEnabled)
+							if(Group.IsEmpty())
+							{
+								ControlChannelMap.Add(Transform.ParameterName, FChannelMapInfo(ControlIndex, TotalIndex, FloatChannelIndex, INDEX_NONE, NAME_None, MaskIndex, CategoryIndex));
+								ControlIndex++;
+								if (bEnabled)
+								{
+									++CategoryIndex;
+								}
+								Group = FText::FromName(ControlElement->GetDisplayName());
+							}
+							else
 							{
-								++CategoryIndex;
+								const FChannelMapInfo* pChannelIndex = ControlChannelMap.Find(ParentControlName);
+								const int32 ChannelControlIndex = pChannelIndex ? pChannelIndex->ControlIndex : INDEX_NONE; 
+								ControlChannelMap.Add(Transform.ParameterName, FChannelMapInfo(ControlIndex, TotalIndex, FloatChannelIndex, ChannelControlIndex, NAME_None, MaskIndex, CategoryIndex));
 							}
-							FText Group = FText::FromName(ControlElement->GetDisplayName());
 
 							if (FSpaceControlNameAndChannel* SpaceChannel = GetSpaceChannel(Transform.ParameterName))
 							{
@@ -2128,7 +2068,7 @@ void UMovieSceneControlRigParameterSection::RecreateWithThisControlRig(UControlR
 
 	for (FRigControlElement* ControlElement : SortedControls)
 	{
-		if (!ControlElement->Settings.bAnimatable)
+		if (!ControlRig->GetHierarchy()->IsAnimatable(ControlElement))
 		{
 			continue;
 		}
@@ -2355,7 +2295,7 @@ void UMovieSceneControlRigParameterSection::RecordControlRigKey(FFrameNumber Fra
 		
 		for (FRigControlElement* ControlElement : Controls)
 		{
-			if (!ControlElement->Settings.bAnimatable)
+			if (!ControlRig->GetHierarchy()->IsAnimatable(ControlElement))
 			{
 				continue;
 			}
diff --git a/Engine/Plugins/Animation/ControlRig/Source/ControlRig/Private/Tools/ControlRigPose.cpp b/Engine/Plugins/Animation/ControlRig/Source/ControlRig/Private/Tools/ControlRigPose.cpp
--- a/Engine/Plugins/Animation/ControlRig/Source/ControlRig/Private/Tools/ControlRigPose.cpp
+++ b/Engine/Plugins/Animation/ControlRig/Source/ControlRig/Private/Tools/ControlRigPose.cpp
@@ -16,7 +16,7 @@ void FControlRigControlPose::SavePose(UControlRig* ControlRig, bool bUseAll)
 	
 	for (FRigControlElement* ControlElement : CurrentControls)
 	{
-		if (ControlElement->Settings.bAnimatable && (bUseAll || ControlRig->IsControlSelected(ControlElement->GetName())))
+		if (ControlRig->GetHierarchy()->IsAnimatable(ControlElement) && (bUseAll || ControlRig->IsControlSelected(ControlElement->GetName())))
 		{
 			//we store poses in default parent space so if not in that space we need to compensate it
 			bool bHasNonDefaultParent = false;
diff --git a/Engine/Plugins/Animation/ControlRig/Source/ControlRig/Private/Units/Execution/RigUnit_InteractionExecution.cpp b/Engine/Plugins/Animation/ControlRig/Source/ControlRig/Private/Units/Execution/RigUnit_InteractionExecution.cpp
new file mode 100644
--- /dev/null
+++ b/Engine/Plugins/Animation/ControlRig/Source/ControlRig/Private/Units/Execution/RigUnit_InteractionExecution.cpp
@@ -0,0 +1,14 @@
+// Copyright Epic Games, Inc. All Rights Reserved.
+
+#include "RigUnit_InteractionExecution.h"
+#include "Units/RigUnitContext.h"
+
+FName FRigUnit_InteractionExecution::EventName = TEXT("Interaction");
+
+FRigUnit_InteractionExecution_Execute()
+{
+    DECLARE_SCOPE_HIERARCHICAL_COUNTER_RIGUNIT()
+	ExecuteContext.CopyFrom(RigVMExecuteContext);
+	ExecuteContext.Hierarchy = Context.Hierarchy;
+	ExecuteContext.EventName = FRigUnit_InteractionExecution::EventName;
+}
diff --git a/Engine/Plugins/Animation/ControlRig/Source/ControlRig/Private/Units/Execution/RigUnit_InteractionExecution.h b/Engine/Plugins/Animation/ControlRig/Source/ControlRig/Private/Units/Execution/RigUnit_InteractionExecution.h
new file mode 100644
--- /dev/null
+++ b/Engine/Plugins/Animation/ControlRig/Source/ControlRig/Private/Units/Execution/RigUnit_InteractionExecution.h
@@ -0,0 +1,27 @@
+// Copyright Epic Games, Inc. All Rights Reserved.
+
+#pragma once
+
+#include "Units/RigUnit.h"
+#include "ControlRigDefines.h"
+#include "RigUnit_InteractionExecution.generated.h"
+
+/**
+ * Event for executing logic during an interaction
+ */
+USTRUCT(meta=(DisplayName="Interaction", Category="Events", NodeColor="1, 0, 0", Keywords="Manipulation,Event,During,Interacting"))
+struct CONTROLRIG_API FRigUnit_InteractionExecution : public FRigUnit
+{
+	GENERATED_BODY()
+
+	RIGVM_METHOD()
+	virtual void Execute(const FRigUnitContext& Context) override;
+
+	virtual FName GetEventName() const override { return EventName; }
+
+	// The execution result
+	UPROPERTY(EditAnywhere, Transient, DisplayName = "Execute", Category = "BeginExecution", meta = (Output))
+	FControlRigExecuteContext ExecuteContext;
+
+	static FName EventName;
+};
diff --git a/Engine/Plugins/Animation/ControlRig/Source/ControlRig/Private/Units/Execution/RigUnit_IsInteracting.h b/Engine/Plugins/Animation/ControlRig/Source/ControlRig/Private/Units/Execution/RigUnit_IsInteracting.h
--- a/Engine/Plugins/Animation/ControlRig/Source/ControlRig/Private/Units/Execution/RigUnit_IsInteracting.h
+++ b/Engine/Plugins/Animation/ControlRig/Source/ControlRig/Private/Units/Execution/RigUnit_IsInteracting.h
@@ -9,7 +9,7 @@
 /**
  * Returns true if the Control Rig is being interacted
  */
-USTRUCT(meta=(DisplayName="Is Interacting", Category="Execution", TitleColor="1 0 0", NodeColor="1 1 1", Keywords="IsInteracting,Gizmo,Manipulation,Interaction,Tracking", Varying))
+USTRUCT(meta=(DisplayName="Get Interaction", Category="Execution", TitleColor="1 0 0", NodeColor="1 1 1", Keywords="IsInteracting,Gizmo,Manipulation,Interaction,Tracking", Varying))
 struct CONTROLRIG_API FRigUnit_IsInteracting : public FRigUnit
 {
 	GENERATED_BODY()
diff --git a/Engine/Plugins/Animation/ControlRig/Source/ControlRig/Private/Units/Hierarchy/RigUnit_SetControlColor.cpp b/Engine/Plugins/Animation/ControlRig/Source/ControlRig/Private/Units/Hierarchy/RigUnit_SetControlColor.cpp
--- a/Engine/Plugins/Animation/ControlRig/Source/ControlRig/Private/Units/Hierarchy/RigUnit_SetControlColor.cpp
+++ b/Engine/Plugins/Animation/ControlRig/Source/ControlRig/Private/Units/Hierarchy/RigUnit_SetControlColor.cpp
@@ -3,6 +3,39 @@
 #include "RigUnit_SetControlColor.h"
 #include "Units/RigUnitContext.h"
 
+FRigUnit_GetControlColor_Execute()
+{
+    DECLARE_SCOPE_HIERARCHICAL_COUNTER_RIGUNIT()
+
+    Color = FLinearColor::Black;
+
+	const URigHierarchy* Hierarchy = Context.Hierarchy;
+	if (Hierarchy)
+	{
+		switch (Context.State)
+		{
+			case EControlRigState::Init:
+			{
+				CachedControlIndex.Reset();
+				break;
+			}
+			case EControlRigState::Update:
+			{
+				if (CachedControlIndex.UpdateCache(FRigElementKey(Control, ERigElementType::Control), Hierarchy))
+				{
+					const FRigControlElement* ControlElement = Hierarchy->GetChecked<FRigControlElement>(CachedControlIndex);
+					Color = ControlElement->Settings.ShapeColor;
+				}
+				break;
+			}
+			default:
+			{
+				break;
+			}
+		}
+	}
+}
+
 FRigUnit_SetControlColor_Execute()
 {
     DECLARE_SCOPE_HIERARCHICAL_COUNTER_RIGUNIT()
diff --git a/Engine/Plugins/Animation/ControlRig/Source/ControlRig/Private/Units/Hierarchy/RigUnit_SetControlColor.h b/Engine/Plugins/Animation/ControlRig/Source/ControlRig/Private/Units/Hierarchy/RigUnit_SetControlColor.h
--- a/Engine/Plugins/Animation/ControlRig/Source/ControlRig/Private/Units/Hierarchy/RigUnit_SetControlColor.h
+++ b/Engine/Plugins/Animation/ControlRig/Source/ControlRig/Private/Units/Hierarchy/RigUnit_SetControlColor.h
@@ -6,7 +6,41 @@
 #include "RigUnit_SetControlColor.generated.h"
 
 /**
- * SetControlColor is used to change the gizmo's color on a control at runtime
+ * GetControlColor is used to retrieve the color of control
+ */
+USTRUCT(meta=(DisplayName="Set Gontrol Color", Category="Controls", DocumentationPolicy="Strict", Keywords = "GetControlColor,GetGizmoColor", TemplateName="GetControlColor", NodeColor = "0.0 0.36470600962638855 1.0"))
+struct CONTROLRIG_API FRigUnit_GetControlColor : public FRigUnit
+{
+	GENERATED_BODY()
+
+	FRigUnit_GetControlColor()
+		: Control(NAME_None)
+		, Color(FLinearColor::Black)
+		, CachedControlIndex(FCachedRigElement())
+	{}
+
+	RIGVM_METHOD()
+	virtual void Execute(const FRigUnitContext& Context) override;
+
+	/**
+	 * The name of the Control to get the color for.
+	 */
+	UPROPERTY(meta = (Input, CustomWidget = "ControlName" ))
+	FName Control;
+
+	/**
+	 * The color of the control
+	 */
+	UPROPERTY(meta = (Output))
+	FLinearColor Color;
+
+	// Used to cache the internally used bone index
+	UPROPERTY()
+	FCachedRigElement CachedControlIndex;
+};
+
+/**
+ * SetControlColor is used to change the control's color
  */
 USTRUCT(meta=(DisplayName="Set Control Color", Category="Controls", DocumentationPolicy="Strict", Keywords = "SetControlColor,SetGizmoColor", TemplateName="SetControlColor", NodeColor = "0.0 0.36470600962638855 1.0"))
 struct CONTROLRIG_API FRigUnit_SetControlColor : public FRigUnitMutable
diff --git a/Engine/Plugins/Animation/ControlRig/Source/ControlRig/Private/Units/Hierarchy/RigUnit_SetControlDrivenList.cpp b/Engine/Plugins/Animation/ControlRig/Source/ControlRig/Private/Units/Hierarchy/RigUnit_SetControlDrivenList.cpp
new file mode 100644
--- /dev/null
+++ b/Engine/Plugins/Animation/ControlRig/Source/ControlRig/Private/Units/Hierarchy/RigUnit_SetControlDrivenList.cpp
@@ -0,0 +1,74 @@
+// Copyright Epic Games, Inc. All Rights Reserved.
+
+#include "RigUnit_SetControlDrivenList.h"
+
+#include "Rigs/RigHierarchyController.h"
+#include "Units/RigUnitContext.h"
+
+FRigUnit_GetControlDrivenList_Execute()
+{
+    DECLARE_SCOPE_HIERARCHICAL_COUNTER_RIGUNIT()
+
+    Driven.Reset();
+
+	const URigHierarchy* Hierarchy = Context.Hierarchy;
+	if (Hierarchy)
+	{
+		switch (Context.State)
+		{
+			case EControlRigState::Init:
+			{
+				CachedControlIndex.Reset();
+				break;
+			}
+			case EControlRigState::Update:
+			{
+				if (CachedControlIndex.UpdateCache(FRigElementKey(Control, ERigElementType::Control), Hierarchy))
+				{
+					const FRigControlElement* ControlElement = Hierarchy->GetChecked<FRigControlElement>(CachedControlIndex);
+					Driven = ControlElement->Settings.DrivenControls;
+				}
+				break;
+			}
+			default:
+			{
+				break;
+			}
+		}
+	}
+}
+
+FRigUnit_SetControlDrivenList_Execute()
+{
+    DECLARE_SCOPE_HIERARCHICAL_COUNTER_RIGUNIT()
+	URigHierarchy* Hierarchy = ExecuteContext.Hierarchy;
+	if (Hierarchy)
+	{
+		switch (Context.State)
+		{
+			case EControlRigState::Init:
+			{
+				CachedControlIndex.Reset();
+				break;
+			}
+			case EControlRigState::Update:
+			{
+				if (CachedControlIndex.UpdateCache(FRigElementKey(Control, ERigElementType::Control), Hierarchy))
+				{
+					FRigControlElement* ControlElement = Hierarchy->GetChecked<FRigControlElement>(CachedControlIndex);
+					if(ControlElement->Settings.DrivenControls != Driven)
+					{
+						Swap(ControlElement->Settings.DrivenControls, ControlElement->Settings.PreviouslyDrivenControls);
+						ControlElement->Settings.DrivenControls = Driven;
+						Hierarchy->Notify(ERigHierarchyNotification::ControlDrivenListChanged, ControlElement);
+					}
+				}
+				break;
+			}
+			default:
+			{
+				break;
+			}
+		}
+	}
+}
diff --git a/Engine/Plugins/Animation/ControlRig/Source/ControlRig/Private/Units/Hierarchy/RigUnit_SetControlDrivenList.h b/Engine/Plugins/Animation/ControlRig/Source/ControlRig/Private/Units/Hierarchy/RigUnit_SetControlDrivenList.h
new file mode 100644
--- /dev/null
+++ b/Engine/Plugins/Animation/ControlRig/Source/ControlRig/Private/Units/Hierarchy/RigUnit_SetControlDrivenList.h
@@ -0,0 +1,74 @@
+// Copyright Epic Games, Inc. All Rights Reserved.
+
+#pragma once
+
+#include "Units/RigUnit.h"
+#include "RigUnit_SetControlDrivenList.generated.h"
+
+/**
+ * GetControlDrivenList is used to retrieve the list of affected controls of an indirect control
+ */
+USTRUCT(meta=(DisplayName="Get Driven Controls", Category="Controls", DocumentationPolicy="Strict", Keywords = "GetControlDrivenList,Interaction,Indirect", TemplateName="GetControlDrivenList", NodeColor = "0.0 0.36470600962638855 1.0"))
+struct CONTROLRIG_API FRigUnit_GetControlDrivenList : public FRigUnit
+{
+	GENERATED_BODY()
+
+	FRigUnit_GetControlDrivenList()
+		: Control(NAME_None)
+		, Driven()
+		, CachedControlIndex(FCachedRigElement())
+	{}
+
+	RIGVM_METHOD()
+	virtual void Execute(const FRigUnitContext& Context) override;
+
+	/**
+	 * The name of the Control to get the list for
+	 */
+	UPROPERTY(meta = (Input, CustomWidget = "ControlName" ))
+	FName Control;
+
+	/**
+	 * The list of affected controls
+	 */
+	UPROPERTY(meta = (Output))
+	TArray<FRigElementKey> Driven;
+
+	// Used to cache the internally used bone index
+	UPROPERTY()
+	FCachedRigElement CachedControlIndex;
+};
+
+/**
+ * SetControlDrivenList is used to change the list of affected controls of an indirect control
+ */
+USTRUCT(meta=(DisplayName="Set Driven Controls", Category="Controls", DocumentationPolicy="Strict", Keywords = "SetControlDrivenList,Interaction,Indirect", TemplateName="SetControlDrivenList", NodeColor = "0.0 0.36470600962638855 1.0"))
+struct CONTROLRIG_API FRigUnit_SetControlDrivenList : public FRigUnitMutable
+{
+	GENERATED_BODY()
+
+	FRigUnit_SetControlDrivenList()
+		: Control(NAME_None)
+		, Driven()
+		, CachedControlIndex(FCachedRigElement())
+	{}
+
+	RIGVM_METHOD()
+	virtual void Execute(const FRigUnitContext& Context) override;
+
+	/**
+	 * The name of the Control to set the list for
+	 */
+	UPROPERTY(meta = (Input, CustomWidget = "ControlName" ))
+	FName Control;
+
+	/**
+	 * The list of affected controls
+	 */
+	UPROPERTY(meta = (Input))
+	TArray<FRigElementKey> Driven;
+
+	// Used to cache the internally used bone index
+	UPROPERTY()
+	FCachedRigElement CachedControlIndex;
+};
diff --git a/Engine/Plugins/Animation/ControlRig/Source/ControlRig/Private/Units/Hierarchy/RigUnit_SetControlVisibility.cpp b/Engine/Plugins/Animation/ControlRig/Source/ControlRig/Private/Units/Hierarchy/RigUnit_SetControlVisibility.cpp
--- a/Engine/Plugins/Animation/ControlRig/Source/ControlRig/Private/Units/Hierarchy/RigUnit_SetControlVisibility.cpp
+++ b/Engine/Plugins/Animation/ControlRig/Source/ControlRig/Private/Units/Hierarchy/RigUnit_SetControlVisibility.cpp
@@ -3,6 +3,39 @@
 #include "RigUnit_SetControlVisibility.h"
 #include "Units/RigUnitContext.h"
 
+FRigUnit_GetControlVisibility_Execute()
+{
+    DECLARE_SCOPE_HIERARCHICAL_COUNTER_RIGUNIT()
+
+    bVisible = false;
+
+	const URigHierarchy* Hierarchy = Context.Hierarchy;
+	if (Hierarchy)
+	{
+		switch (Context.State)
+		{
+			case EControlRigState::Init:
+			{
+				CachedControlIndex.Reset();
+			}
+			case EControlRigState::Update:
+			{
+				if (CachedControlIndex.UpdateCache(Item, Hierarchy))
+				{
+					if(const FRigControlElement* ControlElement = Hierarchy->Get<FRigControlElement>(CachedControlIndex))
+					{
+						bVisible = ControlElement->Settings.bShapeVisible;
+					}
+				}
+				break;
+			}
+			default:
+			{
+				break;
+			}
+		}
+	}
+}
 
 FRigUnit_SetControlVisibility_Execute()
 {
@@ -49,7 +82,7 @@ FRigUnit_SetControlVisibility_Execute()
 
 				for (int32 Index = 0; Index < Keys.Num(); Index++)
 				{
-					CachedControlIndices[Index].UpdateCache(Keys[Index], ExecuteContext.Hierarchy);
+					CachedControlIndices[Index].UpdateCache(Keys[Index], Hierarchy);
 				}
 
 				for (const FCachedRigElement& CachedControlIndex : CachedControlIndices)
diff --git a/Engine/Plugins/Animation/ControlRig/Source/ControlRig/Private/Units/Hierarchy/RigUnit_SetControlVisibility.h b/Engine/Plugins/Animation/ControlRig/Source/ControlRig/Private/Units/Hierarchy/RigUnit_SetControlVisibility.h
--- a/Engine/Plugins/Animation/ControlRig/Source/ControlRig/Private/Units/Hierarchy/RigUnit_SetControlVisibility.h
+++ b/Engine/Plugins/Animation/ControlRig/Source/ControlRig/Private/Units/Hierarchy/RigUnit_SetControlVisibility.h
@@ -6,7 +6,40 @@
 #include "RigUnit_SetControlVisibility.generated.h"
 
 /**
- * SetControlVisibility is used to change the gizmo visibility on a control at runtime
+ * GetControlVisibility is used to retrieve the visibility of a control
+ */
+USTRUCT(meta=(DisplayName="Get Control Visibility", Category="Controls", DocumentationPolicy="Strict", Keywords = "GetControlVisibility,Visibility,Hide,Show,Hidden,Visible,SetGizmoVisibility", TemplateName="GetControlVisibility", NodeColor="0, 0.364706, 1.0"))
+struct CONTROLRIG_API FRigUnit_GetControlVisibility : public FRigUnit
+{
+	GENERATED_BODY()
+
+	FRigUnit_GetControlVisibility()
+		: Item(NAME_None, ERigElementType::Control)
+		, bVisible(true)
+	{}
+
+	RIGVM_METHOD()
+	virtual void Execute(const FRigUnitContext& Context) override;
+
+	/**
+	 * The name of the Control to set the visibility for.
+	 */
+	UPROPERTY(meta = (Input, ExpandByDefault))
+	FRigElementKey Item;
+
+	/**
+	 * The visibility of the control
+	 */
+	UPROPERTY(meta = (Output))
+	bool bVisible;
+
+	// Used to cache the internally used control index
+	UPROPERTY()
+	FCachedRigElement CachedControlIndex;
+};
+
+/**
+ * SetControlVisibility is used to change the visibility on a control at runtime
  */
 USTRUCT(meta=(DisplayName="Set Control Visibility", Category="Controls", DocumentationPolicy="Strict", Keywords = "SetControlVisibility,Visibility,Hide,Show,Hidden,Visible,SetGizmoVisibility", TemplateName="SetControlVisibility", NodeColor="0, 0.364706, 1.0"))
 struct CONTROLRIG_API FRigUnit_SetControlVisibility : public FRigUnitMutable
@@ -35,7 +68,7 @@ struct CONTROLRIG_API FRigUnit_SetControlVisibility : public FRigUnitMutable
 	FString Pattern;
 
 	/**
-	 * The color to set for the control
+	 * The visibility to set for the control
 	 */
 	UPROPERTY(meta = (Input))
 	bool bVisible;
diff --git a/Engine/Plugins/Animation/ControlRig/Source/ControlRig/Public/ControlRig.h b/Engine/Plugins/Animation/ControlRig/Source/ControlRig/Public/ControlRig.h
--- a/Engine/Plugins/Animation/ControlRig/Source/ControlRig/Public/ControlRig.h
+++ b/Engine/Plugins/Animation/ControlRig/Source/ControlRig/Public/ControlRig.h
@@ -941,6 +941,34 @@ public:
 		InControlRig->GetHierarchy()->StartInteraction();
 	}
 
+	FORCEINLINE_DEBUGGABLE FControlRigInteractionScope(
+		UControlRig* InControlRig,
+		const FRigElementKey& InKey,
+		EControlRigInteractionType InInteractionType = EControlRigInteractionType::All
+	)
+		: ControlRig(InControlRig)
+		, InteractionBracketScope(InControlRig->InteractionBracket)
+		, SyncBracketScope(InControlRig->InterRigSyncBracket)
+	{
+		InControlRig->ElementsBeingInteracted = { InKey };
+		InControlRig->InteractionType = (uint8)InInteractionType;
+		InControlRig->GetHierarchy()->StartInteraction();
+	}
+
+	FORCEINLINE_DEBUGGABLE FControlRigInteractionScope(
+		UControlRig* InControlRig,
+		const TArray<FRigElementKey>& InKeys,
+		EControlRigInteractionType InInteractionType = EControlRigInteractionType::All
+	)
+		: ControlRig(InControlRig)
+		, InteractionBracketScope(InControlRig->InteractionBracket)
+	, SyncBracketScope(InControlRig->InterRigSyncBracket)
+	{
+		InControlRig->ElementsBeingInteracted = InKeys;
+		InControlRig->InteractionType = (uint8)InInteractionType;
+		InControlRig->GetHierarchy()->StartInteraction();
+	}
+
 	FORCEINLINE_DEBUGGABLE ~FControlRigInteractionScope()
 	{
 		if(ensure(ControlRig.IsValid()))
diff --git a/Engine/Plugins/Animation/ControlRig/Source/ControlRig/Public/ControlRigGizmoActor.h b/Engine/Plugins/Animation/ControlRig/Source/ControlRig/Public/ControlRigGizmoActor.h
--- a/Engine/Plugins/Animation/ControlRig/Source/ControlRig/Public/ControlRigGizmoActor.h
+++ b/Engine/Plugins/Animation/ControlRig/Source/ControlRig/Public/ControlRigGizmoActor.h
@@ -86,6 +86,8 @@ public:
 	UPROPERTY()
 	FName ColorParameterName;
 
+	FLinearColor OverrideColor;
+
 	UFUNCTION(BlueprintSetter)
 	/** Set the control to be enabled/disabled */
 	virtual void SetEnabled(bool bInEnabled);
diff --git a/Engine/Plugins/Animation/ControlRig/Source/ControlRig/Public/ControlRigObjectVersion.h b/Engine/Plugins/Animation/ControlRig/Source/ControlRig/Public/ControlRigObjectVersion.h
--- a/Engine/Plugins/Animation/ControlRig/Source/ControlRig/Public/ControlRigObjectVersion.h
+++ b/Engine/Plugins/Animation/ControlRig/Source/ControlRig/Public/ControlRigObjectVersion.h
@@ -60,6 +60,9 @@ struct CONTROLRIG_API FControlRigObjectVersion
 		// Added rig curve element value state flag
 		CurveElementValueStateFlag,
 
+		// Added the notion of a per control animation type
+		ControlAnimationType,
+
 		// -----<new versions can be added above this line>-------------------------------------------------
 		VersionPlusOne,
 		LatestVersion = VersionPlusOne - 1
diff --git a/Engine/Plugins/Animation/ControlRig/Source/ControlRig/Public/Rigs/RigHierarchy.h b/Engine/Plugins/Animation/ControlRig/Source/ControlRig/Public/Rigs/RigHierarchy.h
--- a/Engine/Plugins/Animation/ControlRig/Source/ControlRig/Public/Rigs/RigHierarchy.h
+++ b/Engine/Plugins/Animation/ControlRig/Source/ControlRig/Public/Rigs/RigHierarchy.h
@@ -2586,6 +2586,16 @@ public:
 	 */
 	bool IsTracingChanges() const;
 
+	/**
+	 * Returns true if the control is animatable
+	 */
+	bool IsAnimatable(const FRigElementKey& InKey) const;
+
+	/**
+	 * Returns true if the control is animatable
+	 */
+	bool IsAnimatable(const FRigControlElement* InControlElement) const;
+
 #if WITH_EDITOR
 
 	/**
@@ -3099,6 +3109,8 @@ private:
 	
 #endif
 
+	void UpdateVisibilityOnProxyControls();
+
 	static const TArray<FString>& GetTransformTypeStrings();
 
 	struct FQueuedNotification
diff --git a/Engine/Plugins/Animation/ControlRig/Source/ControlRig/Public/Rigs/RigHierarchyDefines.h b/Engine/Plugins/Animation/ControlRig/Source/ControlRig/Public/Rigs/RigHierarchyDefines.h
--- a/Engine/Plugins/Animation/ControlRig/Source/ControlRig/Public/Rigs/RigHierarchyDefines.h
+++ b/Engine/Plugins/Animation/ControlRig/Source/ControlRig/Public/Rigs/RigHierarchyDefines.h
@@ -57,6 +57,7 @@ enum class ERigHierarchyNotification : uint8
 	HierarchyReset,
 	ControlSettingChanged,
 	ControlVisibilityChanged,
+	ControlDrivenListChanged,
 	ControlShapeTransformChanged,
 	ParentWeightsChanged,
 	InteractionBracketOpened,
@@ -103,6 +104,21 @@ enum class ERigControlType : uint8
     EulerTransform,
 };
 
+UENUM(BlueprintType)
+enum class ERigControlAnimationType : uint8
+{
+	// A visible, animatable control.
+	AnimationControl,
+	// An animation channel without a 3d shape
+	AnimationChannel,
+	// A control to drive other controls,
+	// not animatable in sequencer.
+	ProxyControl,
+	// Visual feedback only - the control is
+	// neither animatable nor selectable.
+	VisualCue
+};
+
 UENUM(BlueprintType)
 enum class ERigControlValueType : uint8
 {
@@ -112,6 +128,15 @@ enum class ERigControlValueType : uint8
     Maximum
 };
 
+UENUM(BlueprintType)
+enum class ERigControlVisibility : uint8
+{
+	// Visibility controlled by the graph
+	UserDefined, 
+	// Visibility Controlled by the selection of driven controls
+	BasedOnSelection 
+};
+
 UENUM(BlueprintType)
 enum class ERigControlAxis : uint8
 {
diff --git a/Engine/Plugins/Animation/ControlRig/Source/ControlRig/Public/Rigs/RigHierarchyElements.h b/Engine/Plugins/Animation/ControlRig/Source/ControlRig/Public/Rigs/RigHierarchyElements.h
--- a/Engine/Plugins/Animation/ControlRig/Source/ControlRig/Public/Rigs/RigHierarchyElements.h
+++ b/Engine/Plugins/Animation/ControlRig/Source/ControlRig/Public/Rigs/RigHierarchyElements.h
@@ -8,6 +8,7 @@
 
 struct FRigUnitContext;
 struct FRigBaseElement;
+struct FRigControlElement;
 class URigHierarchy;
 
 DECLARE_DELEGATE_RetVal_ThreeParams(FTransform, FRigReferenceGetWorldTransformDelegate, const FRigUnitContext*, const FRigElementKey& /* Key */, bool /* bInitial */);
@@ -876,6 +877,9 @@ struct CONTROLRIG_API FRigControlSettings
 	void Load(FArchive& Ar);
 
 	UPROPERTY(BlueprintReadWrite, EditAnywhere, Category = Control)
+	ERigControlAnimationType AnimationType;
+
+	UPROPERTY(BlueprintReadWrite, EditAnywhere, Category = Control, meta=(DisplayName="Value Type"))
 	ERigControlType ControlType;
 
 	UPROPERTY(BlueprintReadWrite, EditAnywhere, Category = Control)
@@ -889,10 +893,6 @@ struct CONTROLRIG_API FRigControlSettings
 	UPROPERTY(transient)
 	bool bIsCurve;
 
-	/** If the control is animatable in sequencer */
-	UPROPERTY(BlueprintReadWrite, EditAnywhere, Category = Control)
-	bool bAnimatable;
-
 	/** True if the control has limits. */
 	UPROPERTY(BlueprintReadWrite, EditAnywhere, Category = Limits)
 	TArray<FRigControlLimitEnabled> LimitEnabled;
@@ -912,19 +912,19 @@ struct CONTROLRIG_API FRigControlSettings
 	UPROPERTY(BlueprintReadWrite, Category = Limits)
 	FRigControlValue MaximumValue;
 
-	/** Set to true if the shape is enabled in 3d */
-	UPROPERTY(BlueprintReadWrite, EditAnywhere, Category = Gizmo)
-	bool bShapeEnabled;
-
 	/** Set to true if the shape is currently visible in 3d */
-	UPROPERTY(BlueprintReadWrite, EditAnywhere, Category = Shape, meta = (EditCondition = "bShapeEnabled"))
+	UPROPERTY(BlueprintReadWrite, EditAnywhere, Category = Shape)
 	bool bShapeVisible;
 
+	/** Defines how the shape visibility should be changed */
+	UPROPERTY(BlueprintReadWrite, EditAnywhere, Category = Shape)
+	ERigControlVisibility ShapeVisibility;
+
 	/* This is optional UI setting - this doesn't mean this is always used, but it is optional for manipulation layer to use this*/
-	UPROPERTY(BlueprintReadWrite, EditAnywhere, Category = Shape, meta = (EditCondition = "bShapeEnabled"))
+	UPROPERTY(BlueprintReadWrite, EditAnywhere, Category = Shape)
 	FName ShapeName;
 
-	UPROPERTY(BlueprintReadWrite, EditAnywhere, Category = Shape, meta = (EditCondition = "bShapeEnabled"))
+	UPROPERTY(BlueprintReadWrite, EditAnywhere, Category = Shape)
 	FLinearColor ShapeColor;
 
 	/** If the control is transient and only visible in the control rig editor */
@@ -942,6 +942,23 @@ struct CONTROLRIG_API FRigControlSettings
 	UPROPERTY(BlueprintReadOnly, EditAnywhere, Category = Animation, meta = (DisplayName = "Customization"))
 	FRigControlElementCustomization Customization;
 
+	/**
+	 * The list of driven controls for this proxy control.
+	 */
+	UPROPERTY(BlueprintReadOnly, EditAnywhere, Category = Animation)
+	TArray<FRigElementKey> DrivenControls;
+
+	/**
+	 * The list of previously driven controls - prior to a procedural change
+	 */
+	TArray<FRigElementKey> PreviouslyDrivenControls;
+
+	/**
+	 * If set to true the control will be grouped with the parent control in sequencer
+	 */
+	UPROPERTY(BlueprintReadOnly, EditAnywhere, Category = Animation)
+	bool bGroupWithParentControl;
+
 	/** Applies the limits expressed by these settings to a value */
 	FORCEINLINE void ApplyLimits(FRigControlValue& InOutValue) const
 	{
@@ -972,6 +989,73 @@ struct CONTROLRIG_API FRigControlSettings
 	}
 
 	void SetupLimitArrayForType(bool bLimitTranslation = false, bool bLimitRotation = false, bool bLimitScale = false);
+
+	FORCEINLINE bool IsAnimatable() const
+	{
+		return (AnimationType == ERigControlAnimationType::AnimationControl) ||
+			(AnimationType == ERigControlAnimationType::AnimationChannel);
+	}
+
+	FORCEINLINE bool SupportsShape() const
+	{
+		return (AnimationType != ERigControlAnimationType::AnimationChannel) &&
+			(ControlType != ERigControlType::Bool);
+	}
+
+	FORCEINLINE bool IsVisible() const
+	{
+		return SupportsShape() && bShapeVisible;
+	}
+	
+	FORCEINLINE bool SetVisible(bool bVisible, bool bForce = false)
+	{
+		if(!bForce)
+		{
+			if(AnimationType == ERigControlAnimationType::ProxyControl)
+			{
+				if(ShapeVisibility == ERigControlVisibility::BasedOnSelection)
+				{
+					return false;
+				}
+			}
+		}
+		
+		if(SupportsShape())
+		{
+			if(bShapeVisible == bVisible)
+			{
+				return false;
+			}
+			bShapeVisible = bVisible;
+		}
+		return SupportsShape();
+	}
+
+	FORCEINLINE bool IsSelectable(bool bRespectVisibility = true) const
+	{
+		return (AnimationType == ERigControlAnimationType::AnimationControl ||
+			AnimationType == ERigControlAnimationType::ProxyControl) &&
+			(IsVisible() || !bRespectVisibility);
+	}
+
+	FORCEINLINE void SetAnimationTypeFromDeprecatedData(bool bAnimatable, bool bShapeEnabled)
+	{
+		if(bAnimatable)
+		{
+			if(bShapeEnabled && (ControlType != ERigControlType::Bool))
+			{
+				AnimationType = ERigControlAnimationType::AnimationControl;
+			}
+			else
+			{
+				AnimationType = ERigControlAnimationType::AnimationChannel;
+			}
+		}
+		else
+		{
+			AnimationType = ERigControlAnimationType::ProxyControl;
+		}
+	}
 };
 
 USTRUCT(BlueprintType)
diff --git a/Engine/Plugins/Animation/ControlRig/Source/ControlRigDeveloper/Private/AnimGraphNode_ControlRig.cpp b/Engine/Plugins/Animation/ControlRig/Source/ControlRigDeveloper/Private/AnimGraphNode_ControlRig.cpp
--- a/Engine/Plugins/Animation/ControlRig/Source/ControlRigDeveloper/Private/AnimGraphNode_ControlRig.cpp
+++ b/Engine/Plugins/Animation/ControlRig/Source/ControlRigDeveloper/Private/AnimGraphNode_ControlRig.cpp
@@ -196,7 +196,7 @@ void UAnimGraphNode_ControlRig::CreateCustomPins(TArray<UEdGraphPin*>* OldPins)
 			{
 				Hierarchy->ForEach<FRigControlElement>([&](FRigControlElement* ControlElement) -> bool
 				{
-					if(ControlElement->Settings.bAnimatable)
+					if (Hierarchy->IsAnimatable(ControlElement))
 					{
 						const FName ControlName = ControlElement->GetName();
 						BeginExposableNames.Remove(ControlName);
@@ -366,7 +366,7 @@ void UAnimGraphNode_ControlRig::RebuildExposedProperties()
 			{
 				Hierarchy->ForEach<FRigControlElement>([&](const FRigControlElement* ControlElement) -> bool
 				{
-					if(ControlElement->Settings.bAnimatable)
+					if (Hierarchy->IsAnimatable(ControlElement))
 					{
 						const FName ControlName = ControlElement->GetName();
 						CustomPinProperties.Add(MakeOptionalPin(ControlName));
@@ -627,7 +627,7 @@ void UAnimGraphNode_ControlRig::GetAvailableMapping(const FName& PathName, TArra
 				{
 					Hierarchy->ForEach<FRigControlElement>([&](FRigControlElement* ControlElement) -> bool
 					{
-						if(ControlElement->Settings.bAnimatable)
+						if (Hierarchy->IsAnimatable(ControlElement))
 						{
 							OutArray.Add(ControlElement->GetName());
 						}
@@ -701,7 +701,7 @@ void UAnimGraphNode_ControlRig::CreateVariableMapping(const FString& FilteredTex
 				{
 					Hierarchy->ForEach<FRigControlElement>([&](FRigControlElement* ControlElement) -> bool
 					{
-						if(ControlElement->Settings.bAnimatable)
+						if (Hierarchy->IsAnimatable(ControlElement))
 						{
 							const FName ControlName = ControlElement->GetName();
 							const FString& DisplayName = ControlName.ToString();
diff --git a/Engine/Plugins/Animation/ControlRig/Source/ControlRigDeveloper/Private/ControlRigBlueprint.cpp b/Engine/Plugins/Animation/ControlRig/Source/ControlRigDeveloper/Private/ControlRigBlueprint.cpp
--- a/Engine/Plugins/Animation/ControlRig/Source/ControlRigDeveloper/Private/ControlRigBlueprint.cpp
+++ b/Engine/Plugins/Animation/ControlRig/Source/ControlRigDeveloper/Private/ControlRigBlueprint.cpp
@@ -359,7 +359,7 @@ void UControlRigBlueprint::PreSave(FObjectPreSaveContext ObjectSaveContext)
 	bExposesAnimatableControls = false;
 	Hierarchy->ForEach<FRigControlElement>([this](FRigControlElement* ControlElement) -> bool
     {
-		if (ControlElement->Settings.bAnimatable)
+		if (Hierarchy->IsAnimatable(ControlElement))
 		{
 			bExposesAnimatableControls = true;
 			return false;
diff --git a/Engine/Plugins/Animation/ControlRig/Source/ControlRigEditor/Private/ControlRigElementDetails.cpp b/Engine/Plugins/Animation/ControlRig/Source/ControlRigEditor/Private/ControlRigElementDetails.cpp
--- a/Engine/Plugins/Animation/ControlRig/Source/ControlRigEditor/Private/ControlRigElementDetails.cpp
+++ b/Engine/Plugins/Animation/ControlRig/Source/ControlRigEditor/Private/ControlRigElementDetails.cpp
@@ -737,7 +737,45 @@ bool FRigBaseElementDetails::IsAnyElementNotOfType(ERigElementType InType) const
 	return false;
 }
 
-bool FRigBaseElementDetails::IsAnyControlOfType(ERigControlType InType) const
+bool FRigBaseElementDetails::IsAnyControlOfAnimationType(ERigControlAnimationType InType) const
+{
+	for(TWeakObjectPtr<UDetailsViewWrapperObject> ObjectBeingCustomized : ObjectsBeingCustomized)
+	{
+		if(ObjectBeingCustomized.IsValid())
+		{
+			if(ObjectBeingCustomized->IsChildOf<FRigControlElement>())
+			{
+				const FRigControlElement ControlElement = ObjectBeingCustomized->GetContent<FRigControlElement>();
+				if(ControlElement.Settings.AnimationType == InType)
+				{
+					return true;
+				}
+			}
+		}
+	}
+	return false;
+}
+
+bool FRigBaseElementDetails::IsAnyControlNotOfAnimationType(ERigControlAnimationType InType) const
+{
+	for(TWeakObjectPtr<UDetailsViewWrapperObject> ObjectBeingCustomized : ObjectsBeingCustomized)
+	{
+		if(ObjectBeingCustomized.IsValid())
+		{
+			if(ObjectBeingCustomized->IsChildOf<FRigControlElement>())
+			{
+				const FRigControlElement ControlElement = ObjectBeingCustomized->GetContent<FRigControlElement>();
+				if(ControlElement.Settings.AnimationType != InType)
+				{
+					return true;
+				}
+			}
+		}
+	}
+	return false;
+}
+
+bool FRigBaseElementDetails::IsAnyControlOfValueType(ERigControlType InType) const
 {
 	for(TWeakObjectPtr<UDetailsViewWrapperObject> ObjectBeingCustomized : ObjectsBeingCustomized)
 	{
@@ -756,7 +794,7 @@ bool FRigBaseElementDetails::IsAnyControlOfType(ERigControlType InType) const
 	return false;
 }
 
-bool FRigBaseElementDetails::IsAnyControlNotOfType(ERigControlType InType) const
+bool FRigBaseElementDetails::IsAnyControlNotOfValueType(ERigControlType InType) const
 {
 	for(TWeakObjectPtr<UDetailsViewWrapperObject> ObjectBeingCustomized : ObjectsBeingCustomized)
 	{
@@ -816,7 +854,7 @@ void FRigTransformElementDetails::CustomizeTransform(IDetailLayoutBuilder& Detai
 
 	TArray<FRigElementKey> Keys = GetElementKeys();
 	Keys = HierarchyBeingDebugged->SortKeys(Keys);
-	const bool bAllControls = !IsAnyElementNotOfType(ERigElementType::Control) && !IsAnyControlOfType(ERigControlType::Bool);
+	const bool bAllControls = !IsAnyElementNotOfType(ERigElementType::Control) && !IsAnyControlOfValueType(ERigControlType::Bool);
 	bool bShowLimits = false;
 	TArray<ERigTransformElementDetailsTransform::Type> TransformTypes;
 	TArray<FText> ButtonLabels;
@@ -840,7 +878,7 @@ void FRigTransformElementDetails::CustomizeTransform(IDetailLayoutBuilder& Detai
 			LOCTEXT("OffsetTooltip", "Offset transform under the control")
 		};
 
-		bShowLimits = !IsAnyControlNotOfType(ERigControlType::EulerTransform);
+		bShowLimits = !IsAnyControlNotOfValueType(ERigControlType::EulerTransform);
 
 		if(bShowLimits)
 		{
@@ -884,8 +922,8 @@ void FRigTransformElementDetails::CustomizeTransform(IDetailLayoutBuilder& Detai
 		bool bIsTransformEnabled = true;
 		if (IsAnyElementOfType(ERigElementType::Control))
 		{
-			bIsTransformEnabled = IsAnyControlOfType(ERigControlType::EulerTransform) ||
-				IsAnyControlOfType(ERigControlType::Transform) ||
+			bIsTransformEnabled = IsAnyControlOfValueType(ERigControlType::EulerTransform) ||
+				IsAnyControlOfValueType(ERigControlType::Transform) ||
 				CurrentTransformType == ERigTransformElementDetailsTransform::Offset;
 
 			if(!bIsTransformEnabled)
@@ -2015,8 +2053,6 @@ void FRigBoneElementDetails::CustomizeDetails(IDetailLayoutBuilder& DetailBuilde
 	CustomizeTransform(DetailBuilder);
 }
 
-TArray<TSharedPtr<FString>> FRigControlElementDetails::ControlTypeList;
-
 void FRigControlElementDetails::CustomizeDetails(IDetailLayoutBuilder& DetailBuilder)
 {
 	FRigTransformElementDetails::CustomizeDetails(DetailBuilder);
@@ -2272,16 +2308,83 @@ void FRigControlElementDetails::CustomizeControl(IDetailLayoutBuilder& DetailBui
 		.IsEnabled(ObjectsBeingCustomized.Num() == 1)
 	];
 
-	if (ControlTypeList.Num() == 0)
-	{
-		UEnum* Enum = StaticEnum<ERigControlType>();
-		for (int64 Index = 0; Index < Enum->GetMaxEnumValue(); Index++)
+	const TSharedRef<IPropertyUtilities> PropertyUtilities = DetailBuilder.GetPropertyUtilities();
+
+	// when control type changes, we have to refresh detail panel
+	const TSharedPtr<IPropertyHandle> AnimationTypeHandle = SettingsHandle->GetChildHandle(TEXT("AnimationType"));
+	AnimationTypeHandle->SetOnPropertyValueChanged(FSimpleDelegate::CreateLambda(
+		[this, PropertyUtilities]()
 		{
-			ControlTypeList.Add(MakeShared<FString>(Enum->GetDisplayNameTextByValue(Index).ToString()));
+			TArray<FRigControlElement> ControlElementsInView = GetElementsInDetailsView<FRigControlElement>();
+			TArray<FRigControlElement*> ControlElementsInHierarchy = GetElementsInHierarchy<FRigControlElement>();
+			check(ControlElementsInView.Num() == ControlElementsInHierarchy.Num());
+
+			if (this->HierarchyBeingCustomized && ControlElementsInHierarchy.Num() > 0)
+			{
+				this->HierarchyBeingCustomized->Modify();
+				
+				for(int32 ControlIndex = 0; ControlIndex< ControlElementsInView.Num(); ControlIndex++)
+				{
+					const FRigControlElement& ViewElement = ControlElementsInView[ControlIndex];
+					FRigControlElement* ControlElement = ControlElementsInHierarchy[ControlIndex];
+					
+					ControlElement->Settings.AnimationType = ViewElement.Settings.AnimationType;
+
+					ControlElement->Settings.bGroupWithParentControl =
+						ControlElement->Settings.ControlType == ERigControlType::Bool ||
+						ControlElement->Settings.ControlType == ERigControlType::Float ||
+						ControlElement->Settings.ControlType == ERigControlType::Integer ||
+						ControlElement->Settings.ControlType == ERigControlType::Vector2D;
+
+					switch(ControlElement->Settings.AnimationType)
+					{
+						case ERigControlAnimationType::AnimationControl:
+						{
+							ControlElement->Settings.ShapeVisibility = ERigControlVisibility::UserDefined;
+							ControlElement->Settings.bShapeVisible = true;
+							break;
+						}
+						case ERigControlAnimationType::AnimationChannel:
+						{
+							ControlElement->Settings.ShapeVisibility = ERigControlVisibility::UserDefined;
+							ControlElement->Settings.bShapeVisible = false;
+							break;
+						}
+						case ERigControlAnimationType::ProxyControl:
+						{
+							ControlElement->Settings.ShapeVisibility = ERigControlVisibility::BasedOnSelection;
+							ControlElement->Settings.bShapeVisible = true;
+							ControlElement->Settings.bGroupWithParentControl = false;
+							break;
+						}
+						default:
+						{
+							ControlElement->Settings.ShapeVisibility = ERigControlVisibility::UserDefined;
+							ControlElement->Settings.bShapeVisible = true;
+							ControlElement->Settings.bGroupWithParentControl = false;
+							break;
+						}
+					}
+
+					this->HierarchyBeingCustomized->SetControlSettings(ControlElement, ControlElement->Settings, true, true, true);
+					ObjectsBeingCustomized[ControlIndex]->SetContent<FRigControlElement>(*ControlElement);
+
+					if (this->HierarchyBeingCustomized != this->BlueprintBeingCustomized->Hierarchy)
+					{
+						if(FRigControlElement* OtherControlElement = this->BlueprintBeingCustomized->Hierarchy->Find<FRigControlElement>(ControlElement->GetKey()))
+						{
+							OtherControlElement->Settings = ControlElement->Settings;
+							this->BlueprintBeingCustomized->Hierarchy->SetControlSettings(OtherControlElement, OtherControlElement->Settings, true, true, true);
+						}
+					}
+				}
+				
+				PropertyUtilities->ForceRefresh();
+			}
 		}
-	}
+	));
 
-	const TSharedRef<IPropertyUtilities> PropertyUtilities = DetailBuilder.GetPropertyUtilities();
+	ControlCategory.AddProperty(AnimationTypeHandle.ToSharedRef());
 
 	// when control type changes, we have to refresh detail panel
 	const TSharedPtr<IPropertyHandle> ControlTypeHandle = SettingsHandle->GetChildHandle(TEXT("ControlType"));
@@ -2305,12 +2408,15 @@ void FRigControlElementDetails::CustomizeControl(IDetailLayoutBuilder& DetailBui
 
 					ControlElement->Settings.ControlType = ViewElement.Settings.ControlType;
 					ControlElement->Settings.LimitEnabled.Reset();
+					ControlElement->Settings.bGroupWithParentControl = false;
 
 					switch (ControlElement->Settings.ControlType)
 					{
 						case ERigControlType::Bool:
 						{
+							ControlElement->Settings.AnimationType = ERigControlAnimationType::AnimationChannel;
 							ValueToSet = FRigControlValue::Make<bool>(false);
+							ControlElement->Settings.bGroupWithParentControl = ControlElement->Settings.IsAnimatable();
 							break;
 						}
 						case ERigControlType::Float:
@@ -2319,6 +2425,7 @@ void FRigControlElementDetails::CustomizeControl(IDetailLayoutBuilder& DetailBui
 							ControlElement->Settings.SetupLimitArrayForType(true);
 							ControlElement->Settings.MinimumValue = FRigControlValue::Make<float>(0.f);
 							ControlElement->Settings.MaximumValue = FRigControlValue::Make<float>(100.f);
+							ControlElement->Settings.bGroupWithParentControl = ControlElement->Settings.IsAnimatable();
 							break;
 						}
 						case ERigControlType::Integer:
@@ -2327,6 +2434,7 @@ void FRigControlElementDetails::CustomizeControl(IDetailLayoutBuilder& DetailBui
 							ControlElement->Settings.SetupLimitArrayForType(true);
 							ControlElement->Settings.MinimumValue = FRigControlValue::Make<int32>(0);
 							ControlElement->Settings.MaximumValue = FRigControlValue::Make<int32>(100);
+							ControlElement->Settings.bGroupWithParentControl = ControlElement->Settings.IsAnimatable();
 							break;
 						}
 						case ERigControlType::Vector2D:
@@ -2335,6 +2443,7 @@ void FRigControlElementDetails::CustomizeControl(IDetailLayoutBuilder& DetailBui
 							ControlElement->Settings.SetupLimitArrayForType(true);
 							ControlElement->Settings.MinimumValue = FRigControlValue::Make<FVector2D>(FVector2D::ZeroVector);
 							ControlElement->Settings.MaximumValue = FRigControlValue::Make<FVector2D>(FVector2D(100.f, 100.f));
+							ControlElement->Settings.bGroupWithParentControl = ControlElement->Settings.IsAnimatable();
 							break;
 						}
 						case ERigControlType::Position:
@@ -2419,15 +2528,26 @@ void FRigControlElementDetails::CustomizeControl(IDetailLayoutBuilder& DetailBui
 
 	ControlCategory.AddProperty(ControlTypeHandle.ToSharedRef());
 
-	if(!(IsAnyControlNotOfType(ERigControlType::Integer) &&
-		IsAnyControlNotOfType(ERigControlType::Float) &&
-		IsAnyControlNotOfType(ERigControlType::Vector2D)))
+	const bool bSupportsShape = !IsAnyControlOfAnimationType(ERigControlAnimationType::AnimationChannel) &&
+		!IsAnyControlOfAnimationType(ERigControlAnimationType::VisualCue);
+
+	if(IsAnyControlOfAnimationType(ERigControlAnimationType::AnimationControl) ||
+	IsAnyControlOfAnimationType(ERigControlAnimationType::AnimationChannel))
+	{
+		const TSharedPtr<IPropertyHandle> GroupWithParentControlHandle = SettingsHandle->GetChildHandle(TEXT("bGroupWithParentControl"));
+		ControlCategory.AddProperty(GroupWithParentControlHandle.ToSharedRef()).DisplayName(FText::FromString(TEXT("Group Channels")));
+	}
+	
+	if(bSupportsShape &&
+		!(IsAnyControlNotOfValueType(ERigControlType::Integer) &&
+		IsAnyControlNotOfValueType(ERigControlType::Float) &&
+		IsAnyControlNotOfValueType(ERigControlType::Vector2D)))
 	{
 		const TSharedPtr<IPropertyHandle> PrimaryAxisHandle = SettingsHandle->GetChildHandle(TEXT("PrimaryAxis"));
 		ControlCategory.AddProperty(PrimaryAxisHandle.ToSharedRef()).DisplayName(FText::FromString(TEXT("Primary Axis")));
 	}
 
-	if(IsAnyControlOfType(ERigControlType::Integer))
+	if(IsAnyControlOfValueType(ERigControlType::Integer))
 	{
 		const TSharedPtr<IPropertyHandle> ControlEnumHandle = SettingsHandle->GetChildHandle(TEXT("ControlEnum"));
 		ControlCategory.AddProperty(ControlEnumHandle.ToSharedRef()).DisplayName(FText::FromString(TEXT("Control Enum")));
@@ -2485,41 +2605,60 @@ void FRigControlElementDetails::CustomizeControl(IDetailLayoutBuilder& DetailBui
 		
 	}
 
-	ControlCategory.AddProperty(SettingsHandle->GetChildHandle(TEXT("bAnimatable")).ToSharedRef());
-
 	const TSharedPtr<IPropertyHandle> CustomizationHandle = SettingsHandle->GetChildHandle(TEXT("Customization"));
-	const TSharedPtr<IPropertyHandle> AvailableSpacesHandle = CustomizationHandle->GetChildHandle(TEXT("AvailableSpaces"));
-	ControlCategory.AddProperty(AvailableSpacesHandle.ToSharedRef());
+
+	if(bSupportsShape)
+	{
+		const TSharedPtr<IPropertyHandle> AvailableSpacesHandle = CustomizationHandle->GetChildHandle(TEXT("AvailableSpaces"));
+		ControlCategory.AddProperty(AvailableSpacesHandle.ToSharedRef());
+	}
 
 	TArray<FRigElementKey> Keys = GetElementKeys();
 	URigHierarchy* HierarchyBeingDebugged = GetHierarchyBeingDebugged();
 
-	const TSharedPtr<IPropertyHandle> DrawLimitsHandle = SettingsHandle->GetChildHandle(TEXT("bDrawLimits"));
-	
-	ControlCategory
-	.AddProperty(DrawLimitsHandle.ToSharedRef()).DisplayName(FText::FromString(TEXT("Draw Limits")))
-	.IsEnabled(TAttribute<bool>::CreateLambda([Keys, HierarchyBeingDebugged]() -> bool
+	if(bSupportsShape)
 	{
-		for(const FRigElementKey& Key : Keys)
+		const TSharedPtr<IPropertyHandle> DrawLimitsHandle = SettingsHandle->GetChildHandle(TEXT("bDrawLimits"));
+		
+		ControlCategory
+		.AddProperty(DrawLimitsHandle.ToSharedRef()).DisplayName(FText::FromString(TEXT("Draw Limits")))
+		.IsEnabled(TAttribute<bool>::CreateLambda([Keys, HierarchyBeingDebugged]() -> bool
 		{
-			if(const FRigControlElement* ControlElement = HierarchyBeingDebugged->Find<FRigControlElement>(Key))
+			for(const FRigElementKey& Key : Keys)
 			{
-				if(ControlElement->Settings.LimitEnabled.Contains(FRigControlLimitEnabled(true, true)))
+				if(const FRigControlElement* ControlElement = HierarchyBeingDebugged->Find<FRigControlElement>(Key))
 				{
-					return true;
+					if(ControlElement->Settings.LimitEnabled.Contains(FRigControlLimitEnabled(true, true)))
+					{
+						return true;
+					}
 				}
 			}
-		}
-		return false;
-	}));
+			return false;
+		}));
+	}
+	
+	if(!IsAnyControlNotOfAnimationType(ERigControlAnimationType::ProxyControl))
+	{
+		ControlCategory.AddProperty(SettingsHandle->GetChildHandle(TEXT("DrivenControls")).ToSharedRef());
+	}
 }
 
 void FRigControlElementDetails::CustomizeShape(IDetailLayoutBuilder& DetailBuilder)
 {
-	// bools don't have shapes
-	if(IsAnyControlOfType(ERigControlType::Bool))
+	for(TWeakObjectPtr<UDetailsViewWrapperObject> ObjectBeingCustomized : ObjectsBeingCustomized)
 	{
-		return;
+		if(ObjectBeingCustomized.IsValid())
+		{
+			if(ObjectBeingCustomized->IsChildOf<FRigControlElement>())
+			{
+				const FRigControlElement ControlElement = ObjectBeingCustomized->GetContent<FRigControlElement>();
+				if(!ControlElement.Settings.SupportsShape())
+				{
+					return;
+				}
+			}
+		}
 	}
 
 	TSharedPtr<IPropertyHandle> ShapeHandle = DetailBuilder.GetProperty(TEXT("Shape"));
@@ -2553,8 +2692,12 @@ void FRigControlElementDetails::CustomizeShape(IDetailLayoutBuilder& DetailBuild
 
 	const TSharedPtr<IPropertyHandle> SettingsHandle = DetailBuilder.GetProperty(TEXT("Settings"));
 
-	ShapeCategory.AddProperty(SettingsHandle->GetChildHandle(TEXT("bShapeEnabled")).ToSharedRef())
-	.DisplayName(FText::FromString(TEXT("Enabled")));
+	if(!IsAnyControlNotOfAnimationType(ERigControlAnimationType::ProxyControl))
+	{
+		ShapeCategory.AddProperty(SettingsHandle->GetChildHandle(TEXT("ShapeVisibility")).ToSharedRef())
+		.DisplayName(FText::FromString(TEXT("Visibility Mode")));
+	}
+
 	ShapeCategory.AddProperty(SettingsHandle->GetChildHandle(TEXT("bShapeVisible")).ToSharedRef())
 	.DisplayName(FText::FromString(TEXT("Visible")));
 
@@ -2981,10 +3124,7 @@ bool FRigControlElementDetails::IsShapeEnabled() const
 			if(ObjectBeingCustomized->IsChildOf<FRigControlElement>())
 			{
 				const FRigControlElement ControlElement = ObjectBeingCustomized->GetContent<FRigControlElement>();
-				if(ControlElement.Settings.bShapeEnabled)
-				{
-					return true;
-				}
+				return ControlElement.Settings.SupportsShape();
 			}
 		}
 	}
@@ -2996,11 +3136,6 @@ const TArray<TSharedPtr<FString>>& FRigControlElementDetails::GetShapeNameList()
 	return ShapeNameList;
 }
 
-const TArray<TSharedPtr<FString>>& FRigControlElementDetails::GetControlTypeList() const
-{
-	return ControlTypeList;
-}
-
 FText FRigControlElementDetails::GetDisplayName() const
 {
 	FName DisplayName(NAME_None);
diff --git a/Engine/Plugins/Animation/ControlRig/Source/ControlRigEditor/Private/ControlRigElementDetails.h b/Engine/Plugins/Animation/ControlRig/Source/ControlRigEditor/Private/ControlRigElementDetails.h
--- a/Engine/Plugins/Animation/ControlRig/Source/ControlRigEditor/Private/ControlRigElementDetails.h
+++ b/Engine/Plugins/Animation/ControlRig/Source/ControlRigEditor/Private/ControlRigElementDetails.h
@@ -154,8 +154,10 @@ public:
 
 	bool IsAnyElementOfType(ERigElementType InType) const;
 	bool IsAnyElementNotOfType(ERigElementType InType) const;
-	bool IsAnyControlOfType(ERigControlType InType) const;
-	bool IsAnyControlNotOfType(ERigControlType InType) const;
+	bool IsAnyControlOfAnimationType(ERigControlAnimationType InType) const;
+	bool IsAnyControlNotOfAnimationType(ERigControlAnimationType InType) const;
+	bool IsAnyControlOfValueType(ERigControlType InType) const;
+	bool IsAnyControlNotOfValueType(ERigControlType InType) const;
 
 	static void RegisterSectionMappings(FPropertyEditorModule& PropertyEditorModule);
 	virtual void RegisterSectionMappings(FPropertyEditorModule& PropertyEditorModule, UClass* InClass);
@@ -256,7 +258,6 @@ public:
 	bool IsShapeEnabled() const;
 
 	const TArray<TSharedPtr<FString>>& GetShapeNameList() const;
-	const TArray<TSharedPtr<FString>>& GetControlTypeList() const;
 
 	FText GetDisplayName() const;
 	void SetDisplayName(const FText& InNewText, ETextCommit::Type InCommitType);
@@ -550,7 +551,6 @@ private:
 	}
 
 	TArray<TSharedPtr<FString>> ShapeNameList;
-	static TArray<TSharedPtr<FString>> ControlTypeList;
 	TSharedPtr<FRigInfluenceEntryModifier> InfluenceModifier;
 	TSharedPtr<FStructOnScope> InfluenceModifierStruct;
 
diff --git a/Engine/Plugins/Animation/ControlRig/Source/ControlRigEditor/Private/EditMode/ControlRigControlsProxy.cpp b/Engine/Plugins/Animation/ControlRig/Source/ControlRigEditor/Private/EditMode/ControlRigControlsProxy.cpp
--- a/Engine/Plugins/Animation/ControlRig/Source/ControlRigEditor/Private/EditMode/ControlRigControlsProxy.cpp
+++ b/Engine/Plugins/Animation/ControlRig/Source/ControlRigEditor/Private/EditMode/ControlRigControlsProxy.cpp
@@ -108,7 +108,7 @@ void UControlRigControlsProxy::PostEditChangeProperty(struct FPropertyChangedEve
 		FRigControlElement* ControlElement = GetControlElement();
 		if (ControlElement && ControlRig.IsValid())
 		{
-			FControlRigInteractionScope InteractionScope(ControlRig.Get());
+			FControlRigInteractionScope InteractionScope(ControlRig.Get(), ControlElement->GetKey());
 			ControlRig->SelectControl(ControlName, bSelected);
 			ControlRig->Evaluate_AnyThread();
 		}
@@ -145,7 +145,7 @@ void UControlRigTransformControlProxy::PostEditChangeProperty(struct FPropertyCh
 		if (ControlElement && ControlRig.IsValid())
 		{
 			//MUST set through ControlRig
-			FControlRigInteractionScope InteractionScope(ControlRig.Get());
+			FControlRigInteractionScope InteractionScope(ControlRig.Get(), ControlElement->GetKey());
 			FTransform RealTransform = Transform; //Transform is FEulerTransform
 			ControlRig->SetControlValue<FRigControlValue::FTransform_Float>(ControlName, RealTransform, true, EControlRigSetKey::DoNotCare,false);
 			ControlRig->Evaluate_AnyThread();
@@ -215,7 +215,7 @@ void UControlRigTransformNoScaleControlProxy::PostEditChangeProperty(struct FPro
 		if (ControlElement && ControlRig.IsValid())
 		{
 			//MUST set through ControlRig
-			FControlRigInteractionScope InteractionScope(ControlRig.Get());
+			FControlRigInteractionScope InteractionScope(ControlRig.Get(), ControlElement->GetKey());
 			ControlRig->SetControlValue<FRigControlValue::FTransformNoScale_Float>(ControlName, Transform, true, EControlRigSetKey::DoNotCare,false);
 			ControlRig->Evaluate_AnyThread();
 
@@ -278,7 +278,7 @@ void UControlRigEulerTransformControlProxy::PostEditChangeProperty(struct FPrope
 		if (ControlElement && ControlRig.IsValid())
 		{
 			//MUST set through ControlRig
-			FControlRigInteractionScope InteractionScope(ControlRig.Get());
+			FControlRigInteractionScope InteractionScope(ControlRig.Get(), ControlElement->GetKey());
 			ControlRig->SetControlValue<FRigControlValue::FEulerTransform_Float>(ControlName, Transform, true, EControlRigSetKey::DoNotCare,false);
 			ControlRig->Evaluate_AnyThread();
 
@@ -345,7 +345,7 @@ void UControlRigFloatControlProxy::PostEditChangeProperty(struct FPropertyChange
 		if (ControlElement && ControlRig.IsValid())
 		{
 			//MUST set through ControlRig
-			FControlRigInteractionScope InteractionScope(ControlRig.Get());
+			FControlRigInteractionScope InteractionScope(ControlRig.Get(), ControlElement->GetKey());
 			ControlRig->SetControlValue<float>(ControlName, Float, true, EControlRigSetKey::DoNotCare,false);
 			ControlRig->Evaluate_AnyThread();
 
@@ -397,7 +397,7 @@ void UControlRigIntegerControlProxy::PostEditChangeProperty(struct FPropertyChan
 		if (ControlElement && ControlRig.IsValid())
 		{
 			//MUST set through ControlRig
-			FControlRigInteractionScope InteractionScope(ControlRig.Get());
+			FControlRigInteractionScope InteractionScope(ControlRig.Get(), ControlElement->GetKey());
 			ControlRig->SetControlValue<int32>(ControlName, Integer, true, EControlRigSetKey::DoNotCare,false);
 			ControlRig->Evaluate_AnyThread();
 		}
@@ -447,7 +447,7 @@ void UControlRigEnumControlProxy::PostEditChangeProperty(struct FPropertyChanged
 		if (ControlElement && ControlRig.IsValid())
 		{
 			//MUST set through ControlRig
-			FControlRigInteractionScope InteractionScope(ControlRig.Get());
+			FControlRigInteractionScope InteractionScope(ControlRig.Get(), ControlElement->GetKey());
 			ControlRig->SetControlValue<int32>(ControlName, Enum.EnumIndex, true, EControlRigSetKey::DoNotCare,false);
 			ControlRig->Evaluate_AnyThread();
 
@@ -503,7 +503,7 @@ void UControlRigVectorControlProxy::PostEditChangeProperty(struct FPropertyChang
 		if (ControlElement && ControlRig.IsValid())
 		{
 			//MUST set through ControlRig
-			FControlRigInteractionScope InteractionScope(ControlRig.Get());
+			FControlRigInteractionScope InteractionScope(ControlRig.Get(), ControlElement->GetKey());
 			ControlRig->SetControlValue<FVector3f>(ControlName, Vector, true, EControlRigSetKey::DoNotCare,false);
 			ControlRig->Evaluate_AnyThread();
 
@@ -603,7 +603,7 @@ void UControlRigVector2DControlProxy::PostEditChangeProperty(struct FPropertyCha
 		if (ControlElement && ControlRig.IsValid())
 		{
 			//MUST set through ControlRig
-			FControlRigInteractionScope InteractionScope(ControlRig.Get());
+			FControlRigInteractionScope InteractionScope(ControlRig.Get(), ControlElement->GetKey());
 			ControlRig->SetControlValue<FVector3f>(ControlName, FVector3f(Vector2D.X, Vector2D.Y, 0.f), true, EControlRigSetKey::DoNotCare,false);
 			ControlRig->Evaluate_AnyThread();
 
@@ -667,7 +667,7 @@ void UControlRigBoolControlProxy::PostEditChangeProperty(struct FPropertyChanged
 		if (ControlElement && ControlRig.IsValid())
 		{
 			//MUST set through ControlRig
-			FControlRigInteractionScope InteractionScope(ControlRig.Get());
+			FControlRigInteractionScope InteractionScope(ControlRig.Get(), ControlElement->GetKey());
 			ControlRig->SetControlValue<bool>(ControlName, Bool, true, EControlRigSetKey::DoNotCare,false);
 			ControlRig->Evaluate_AnyThread();
 		}
@@ -753,11 +753,6 @@ void UControlRigDetailPanelControlProxies::AddProxy(UControlRig* ControlRig, con
 			case ERigControlType::Float:
 			{
 				Proxy = NewObject<UControlRigFloatControlProxy>(GetTransientPackage(), NAME_None);
-				FRigControlElement* ParentControlElement = Cast<FRigControlElement>(ControlRig->GetHierarchy()->GetFirstParent(ControlElement));
-				if (ParentControlElement)
-				{
-					Proxy->bIsIndividual = true;
-				}
 				break;
 
 			}
@@ -773,11 +768,6 @@ void UControlRigDetailPanelControlProxies::AddProxy(UControlRig* ControlRig, con
 					EnumProxy->Enum.EnumType = ControlElement->Settings.ControlEnum;
 					Proxy = EnumProxy;
 				}
-				FRigControlElement* ParentControlElement = Cast<FRigControlElement>(ControlRig->GetHierarchy()->GetFirstParent(ControlElement));
-				if (ParentControlElement)
-				{
-					Proxy->bIsIndividual = true;
-				}
 				break;
 
 			}
@@ -798,11 +788,6 @@ void UControlRigDetailPanelControlProxies::AddProxy(UControlRig* ControlRig, con
 			case ERigControlType::Bool:
 			{
 				Proxy = NewObject<UControlRigBoolControlProxy>(GetTransientPackage(), NAME_None);
-				FRigControlElement* ParentControlElement = Cast<FRigControlElement>(ControlRig->GetHierarchy()->GetFirstParent(ControlElement));
-				if (ParentControlElement)
-				{
-					Proxy->bIsIndividual = true;
-				}
 				break;
 
 			}
@@ -811,6 +796,9 @@ void UControlRigDetailPanelControlProxies::AddProxy(UControlRig* ControlRig, con
 		}
 		if (Proxy)
 		{
+			Proxy->bIsIndividual =
+				(ControlElement->Settings.AnimationType == ERigControlAnimationType::AnimationChannel) ||
+				(ControlElement->Settings.AnimationType == ERigControlAnimationType::ProxyControl) ;
 			Proxy->SetFlags(RF_Transactional);
 			Proxy->SetName(Name);
 			Proxy->ControlRig = ControlRig;
@@ -893,7 +881,7 @@ void UControlRigDetailPanelControlProxies::RecreateAllProxies(UControlRig* Contr
 	TArray<FRigControlElement*> Controls = ControlRig->AvailableControls();
 	for (FRigControlElement* ControlElement : Controls)
 	{
-		if(ControlElement->Settings.bShapeEnabled && ControlElement->Settings.bAnimatable)
+		if(ControlElement->Settings.AnimationType != ERigControlAnimationType::VisualCue)
 		{
 			AddProxy(ControlRig,ControlElement->GetName(), ControlElement);
 		}
@@ -934,3 +922,12 @@ void UControlRigDetailPanelControlProxies::SelectProxy(UControlRig* ControlRig,c
 	}
 }
 
+bool UControlRigDetailPanelControlProxies::IsSelected(UControlRig* InControlRig, const FName& Name) const
+{
+	if (UControlRigControlsProxy* Proxy = FindProxy(InControlRig, Name))
+	{
+		return SelectedProxies.Contains(Proxy);
+	}
+	return false;
+}
+
diff --git a/Engine/Plugins/Animation/ControlRig/Source/ControlRigEditor/Private/EditMode/ControlRigControlsProxy.h b/Engine/Plugins/Animation/ControlRig/Source/ControlRigEditor/Private/EditMode/ControlRigControlsProxy.h
--- a/Engine/Plugins/Animation/ControlRig/Source/ControlRigEditor/Private/EditMode/ControlRigControlsProxy.h
+++ b/Engine/Plugins/Animation/ControlRig/Source/ControlRigEditor/Private/EditMode/ControlRigControlsProxy.h
@@ -330,5 +330,6 @@ public:
 	void RecreateAllProxies(UControlRig* InControlRig);
 	void SelectProxy(UControlRig* InControlRig, const FName& Name, bool bSelected);
 	const TArray<UControlRigControlsProxy*>& GetSelectedProxies() const { return SelectedProxies;}
+	bool IsSelected(UControlRig* InControlRig, const FName& Name) const;
 
 };
diff --git a/Engine/Plugins/Animation/ControlRig/Source/ControlRigEditor/Private/EditMode/ControlRigEditMode.cpp b/Engine/Plugins/Animation/ControlRig/Source/ControlRigEditor/Private/EditMode/ControlRigEditMode.cpp
--- a/Engine/Plugins/Animation/ControlRig/Source/ControlRigEditor/Private/EditMode/ControlRigEditMode.cpp
+++ b/Engine/Plugins/Animation/ControlRig/Source/ControlRigEditor/Private/EditMode/ControlRigEditMode.cpp
@@ -60,8 +60,6 @@
 #include "InteractiveToolManager.h"
 #include "EdModeInteractiveToolsContext.h"
 
-
-
 void UControlRigEditModeDelegateHelper::OnPoseInitialized()
 {
 	if (EditMode)
@@ -2584,8 +2582,11 @@ void FControlRigEditMode::ResetTransforms(bool bSelectionOnly)
 				ControlsToReset.SetNum(0);
 				for (const FRigControlElement* Control : Controls)
 				{
-					ControlsToReset.Add(Control->GetKey());
-
+					if(Control->Settings.AnimationType == ERigControlAnimationType::AnimationControl ||
+						Control->Settings.AnimationType == ERigControlAnimationType::AnimationChannel)
+					{
+						ControlsToReset.Add(Control->GetKey());
+					}
 				}
 			}
 			bool bHasNonDefaultParent = false;
@@ -2601,6 +2602,7 @@ void FControlRigEditMode::ResetTransforms(bool bSelectionOnly)
 			}
 
 			FScopedTransaction Transaction(LOCTEXT("HierarchyResetTransforms", "Reset Transforms"));
+			FControlRigInteractionScope InteractionScope(ControlRig, ControlsToReset);
 
 			for (const FRigElementKey& ControlToReset : ControlsToReset)
 			{
@@ -2684,6 +2686,12 @@ void FControlRigEditMode::ResetTransforms(bool bSelectionOnly)
 					}
 				}
 			}
+			else
+			{
+				// we at least have to run the interaction event
+				TGuardValue<TArray<FName>> EventGuard(ControlRig->EventQueue, {});
+				ControlRig->Evaluate_AnyThread();
+			}
 		}
 	}
 }
@@ -2935,7 +2943,7 @@ void FControlRigEditMode::CreateShapeActors(UControlRig* ControlRig)
 	int32 ControlRigIndex = RuntimeControlRigs.Find(ControlRig);
 	for (FRigControlElement* ControlElement : Controls)
 	{
-		if (!ControlElement->Settings.bShapeEnabled)
+		if (!ControlElement->Settings.SupportsShape())
 		{
 			continue;
 		}
@@ -2949,7 +2957,7 @@ void FControlRigEditMode::CreateShapeActors(UControlRig* ControlRig)
 			Param.ShapeName = ControlElement->Settings.ShapeName;
 			Param.SpawnTransform = ControlRig->GetControlGlobalTransform(ControlElement->GetName());
 			Param.ShapeTransform = ControlRig->GetHierarchy()->GetControlShapeTransform(ControlElement, ERigTransformType::CurrentLocal);
-			Param.bSelectable = ControlElement->Settings.bAnimatable;
+			Param.bSelectable = ControlElement->Settings.IsSelectable(false);
 
 			if (const FControlRigShapeDefinition* ShapeDef = UControlRigShapeLibrary::GetShapeByName(ControlElement->Settings.ShapeName, ShapeLibraries))
 			{
@@ -3130,6 +3138,7 @@ void FControlRigEditMode::OnHierarchyModified(ERigHierarchyNotification InNotif,
 			break;
 		}
 		case ERigHierarchyNotification::ControlSettingChanged:
+		case ERigHierarchyNotification::ControlVisibilityChanged:
 		case ERigHierarchyNotification::ControlShapeTransformChanged:
 		{
 			const FRigElementKey Key = InElement->GetKey();
@@ -3152,7 +3161,42 @@ void FControlRigEditMode::OnHierarchyModified(ERigHierarchyNotification InNotif,
 
 			// if we can't deal with this lazily, let's fall back to recreating all control shape actors
 			RequestToRecreateControlShapeActors(ControlRig);
+			break;
+		}
+		case ERigHierarchyNotification::ControlDrivenListChanged:
+		{
+			if (IsInLevelEditor())
+			{
+				// to synchronize the selection between the viewport / editmode and the details panel / sequencer
+				// we re-select the control. during deselection we recover the previously set driven list
+				// and then select the control again with the up2date list. this makes sure that the tracks
+				// are correctly selected in sequencer to match what the proxy control is driving.
+				if (FRigControlElement* ControlElement = InHierarchy->Find<FRigControlElement>(InElement->GetKey()))
+				{
+					UControlRig* ControlRig = InHierarchy->GetTypedOuter<UControlRig>();
+					if(ControlProxy->IsSelected(ControlRig, ControlElement->GetName()))
+					{
+						// reselect the control - to affect the details panel / sequencer
+						if(URigHierarchyController* Controller = InHierarchy->GetController())
+						{
+							const FRigElementKey Key = ControlElement->GetKey();
+							{
+								// Restore the previously selected driven elements
+								// so that we can deselect them accordingly.
+								TGuardValue<TArray<FRigElementKey>> DrivenGuard(
+									ControlElement->Settings.DrivenControls,
+									ControlElement->Settings.PreviouslyDrivenControls);
+								
+								Controller->DeselectElement(Key);
+							}
 
+							// now select the proxy control again given the new driven list
+							Controller->SelectElement(Key);
+						}
+					}
+				}
+			}
+			break;
 		}
 		case ERigHierarchyNotification::ElementSelected:
 		case ERigHierarchyNotification::ElementDeselected:
@@ -3200,17 +3244,55 @@ void FControlRigEditMode::OnHierarchyModified(ERigHierarchyNotification InNotif,
 						}
 						if (IsInLevelEditor())
 						{
-							if (bSelected)
+							if (const FRigControlElement* ControlElement = ControlRig->GetHierarchy()->Find<FRigControlElement>(Key))
 							{
-								if (ControlRig->GetHierarchy()->Find<FRigControlElement>(Key))
+								ControlProxy->SelectProxy(ControlRig,Key.Name, bSelected);
+
+								if(ControlElement->Settings.AnimationType == ERigControlAnimationType::ProxyControl)
 								{
-									ControlProxy->SelectProxy(ControlRig,Key.Name, true);
+									const UControlRigEditModeSettings* Settings = GetDefault<UControlRigEditModeSettings>();
+
+									const TArray<FRigElementKey>& DrivenKeys = ControlElement->Settings.DrivenControls;
+									for(const FRigElementKey& DrivenKey : DrivenKeys)
+									{
+										if (const FRigControlElement* DrivenControl = ControlRig->GetHierarchy()->Find<FRigControlElement>(DrivenKey))
+										{
+											ControlProxy->SelectProxy(ControlRig, DrivenControl->GetName(), bSelected);
+
+											if (AControlRigShapeActor* DrivenShapeActor = GetControlShapeFromControlName(ControlRig,DrivenControl->GetName()))
+											{
+												if(bSelected)
+												{
+													DrivenShapeActor->OverrideColor = Settings->DrivenControlColor;
+												}
+												else
+												{
+													DrivenShapeActor->OverrideColor = FLinearColor(0, 0, 0, 0);
+												}
+											}
+										}
+									}
+
+									ControlRig->GetHierarchy()->ForEach<FRigControlElement>(
+										[this, ControlRig, ControlElement, DrivenKeys, bSelected](FRigControlElement* AnimationChannelControl) -> bool
+										{
+											if(AnimationChannelControl->Settings.AnimationType == ERigControlAnimationType::AnimationChannel)
+											{
+												if(const FRigControlElement* ParentControlElement =
+													Cast<FRigControlElement>(ControlRig->GetHierarchy()->GetFirstParent(AnimationChannelControl)))
+												{
+													if(DrivenKeys.Contains(ParentControlElement->GetKey()) ||
+														ParentControlElement->GetKey() == ControlElement->GetKey())
+													{
+														ControlProxy->SelectProxy(ControlRig, AnimationChannelControl->GetName(), bSelected);
+													}
+												}
+											}
+											return true;
+										}
+									);
 								}
 							}
-							else
-							{
-								ControlProxy->SelectProxy(ControlRig,Key.Name, false);
-							}
 						}
 
 					}
@@ -3241,8 +3323,18 @@ void FControlRigEditMode::OnHierarchyModified_AnyThread(ERigHierarchyNotificatio
 		return;
 	}
 
-	if (InNotif != ERigHierarchyNotification::ControlSettingChanged
-		&& InNotif != ERigHierarchyNotification::ControlShapeTransformChanged)
+	if(IsInGameThread())
+	{
+		OnHierarchyModified(InNotif, InHierarchy, InElement);
+		return;
+	}
+
+	if (InNotif != ERigHierarchyNotification::ControlSettingChanged &&
+		InNotif != ERigHierarchyNotification::ControlVisibilityChanged &&
+		InNotif != ERigHierarchyNotification::ControlDrivenListChanged &&
+		InNotif != ERigHierarchyNotification::ControlShapeTransformChanged &&
+		InNotif != ERigHierarchyNotification::ElementSelected &&
+		InNotif != ERigHierarchyNotification::ElementDeselected)
 	{
 		OnHierarchyModified(InNotif, InHierarchy, InElement);
 		return;
@@ -3335,7 +3427,7 @@ bool FControlRigEditMode::CanChangeControlShapeTransform()
 						// only enable for a Control with Gizmo enabled and visible
 						if (FRigControlElement* ControlElement = ControlRig->GetHierarchy()->Find<FRigControlElement>(SelectedRigElements[0]))
 						{
-							if (ControlElement->Settings.bShapeEnabled && ControlElement->Settings.bShapeVisible)
+							if (ControlElement->Settings.IsVisible())
 							{
 								if (AControlRigShapeActor* ShapeActor = GetControlShapeFromControlName(ControlRig,SelectedRigElements[0].Name))
 								{
@@ -3680,12 +3772,33 @@ void FControlRigEditMode::TickControlShape(AControlRigShapeActor* ShapeActor, co
 
 			if (FRigControlElement* ControlElement = ControlRig->FindControl(ShapeActor->ControlName))
 			{
-				ShapeActor->SetShapeColor(ControlElement->Settings.ShapeColor);
-				ShapeActor->SetIsTemporarilyHiddenInEditor(!ControlElement->Settings.bShapeVisible || Settings->bHideControlShapes || !ControlRig->GetControlsVisible());
-				if (!IsInLevelEditor()) //don't change this in level editor otherwise we can never select it and render has to be based on viewport.
+				const bool bControlsHiddenInViewport = Settings->bHideControlShapes || !ControlRig->GetControlsVisible();
+				
+				bool bIsVisible = ControlElement->Settings.IsVisible();
+				bool bRespectVisibilityForSelection = true;
+
+				if(!bControlsHiddenInViewport)
 				{
-					ShapeActor->SetSelectable(ControlElement->Settings.bShapeVisible && !Settings->bHideControlShapes && ControlElement->Settings.bAnimatable && ControlRig->GetControlsVisible());
+					if(ControlElement->Settings.AnimationType == ERigControlAnimationType::ProxyControl)
+					{
+						bRespectVisibilityForSelection = false;
+						
+						if(Settings->bShowAllProxyControls)
+						{
+							bIsVisible = true;
+						}
+					}
 				}
+				
+				ShapeActor->SetShapeColor(ShapeActor->OverrideColor.A < SMALL_NUMBER ?
+					ControlElement->Settings.ShapeColor : ShapeActor->OverrideColor);
+				ShapeActor->SetIsTemporarilyHiddenInEditor(
+					(!ControlElement->Settings.IsVisible() && !Settings->bShowAllProxyControls) ||
+					bControlsHiddenInViewport);
+				
+				ShapeActor->SetSelectable(
+					ControlElement->Settings.IsSelectable(bRespectVisibilityForSelection) &&
+					!bControlsHiddenInViewport);
 			}
 		}
 	}
diff --git a/Engine/Plugins/Animation/ControlRig/Source/ControlRigEditor/Private/EditMode/ControlRigEditModeSettings.h b/Engine/Plugins/Animation/ControlRig/Source/ControlRigEditor/Private/EditMode/ControlRigEditModeSettings.h
--- a/Engine/Plugins/Animation/ControlRig/Source/ControlRigEditor/Private/EditMode/ControlRigEditModeSettings.h
+++ b/Engine/Plugins/Animation/ControlRig/Source/ControlRigEditor/Private/EditMode/ControlRigEditModeSettings.h
@@ -17,6 +17,7 @@ class UControlRigEditModeSettings : public UObject
 		: bDisplayHierarchy(false)
 		, bDisplayNulls(false)
 		, bHideControlShapes(false)
+		, bShowAllProxyControls(false)
 		, bDisplayAxesOnSelection(false)
 		, AxisScale(10.f)
 		, bCoordSystemPerWidgetMode(true)
@@ -25,6 +26,7 @@ class UControlRigEditModeSettings : public UObject
 		, GizmoScale(1.0f)
 	{
 		LastInViewportTweenWidgetLocation = FVector2D(EForceInit::ForceInitToZero);
+		DrivenControlColor = FLinearColor::White * FLinearColor(FVector::OneVector * 0.8f);
 	}
 
 	// UObject interface
@@ -47,6 +49,14 @@ public:
 	UPROPERTY(config, EditAnywhere, AdvancedDisplay, Category = "Animation Settings")
 	bool bHideControlShapes;
 
+	/** Should we always hide control shapes in viewport */
+	UPROPERTY(config, EditAnywhere, AdvancedDisplay, Category = "Animation Settings", meta = (EditCondition = "!bHideControlShapes"))
+	bool bShowAllProxyControls;
+
+	/** Indicates a control being driven by a proxy control */
+	UPROPERTY(config, EditAnywhere, AdvancedDisplay, Category = "Animation Settings", meta = (EditCondition = "!bHideControlShapes"))
+	FLinearColor DrivenControlColor;
+
 	/** Should we show axes for the selected elements */
 	UPROPERTY(config, EditAnywhere, AdvancedDisplay, Category = "Animation Settings")
 	bool bDisplayAxesOnSelection;
diff --git a/Engine/Plugins/Animation/ControlRig/Source/ControlRigEditor/Private/EditMode/SControlRigDetails.cpp b/Engine/Plugins/Animation/ControlRig/Source/ControlRigEditor/Private/EditMode/SControlRigDetails.cpp
--- a/Engine/Plugins/Animation/ControlRig/Source/ControlRigEditor/Private/EditMode/SControlRigDetails.cpp
+++ b/Engine/Plugins/Animation/ControlRig/Source/ControlRigEditor/Private/EditMode/SControlRigDetails.cpp
@@ -237,85 +237,34 @@ void SControlRigDetails::Construct(const FArguments& InArgs, FControlRigEditMode
 	FDetailsViewArgs IndividualDetailsViewArgs = DetailsViewArgs;
 	IndividualDetailsViewArgs.bAllowMultipleTopLevelObjects = true; //this is the secret sauce to show multiple objects in a details view
 
-	ControlEulerTransformDetailsView = FModuleManager::GetModuleChecked<FPropertyEditorModule>("PropertyEditor").CreateDetailView(DetailsViewArgs);
-	ControlEulerTransformDetailsView->SetKeyframeHandler(SharedThis(this));
-	ControlEulerTransformDetailsView->SetIsPropertyVisibleDelegate(FIsPropertyVisible::CreateSP(this, &SControlRigDetails::ShouldShowPropertyOnDetailCustomization));
-	ControlEulerTransformDetailsView->SetIsPropertyReadOnlyDelegate(FIsPropertyReadOnly::CreateSP(this, &SControlRigDetails::IsReadOnlyPropertyOnDetailCustomization));
-	ControlEulerTransformDetailsView->SetGenericLayoutDetailsDelegate(FOnGetDetailCustomizationInstance::CreateStatic(&FControlRigEditModeGenericDetails::MakeInstance, ModeTools));
-
-	ControlTransformDetailsView = FModuleManager::GetModuleChecked<FPropertyEditorModule>("PropertyEditor").CreateDetailView(DetailsViewArgs);
-	ControlTransformDetailsView->SetKeyframeHandler(SharedThis(this));
-	ControlTransformDetailsView->SetIsPropertyVisibleDelegate(FIsPropertyVisible::CreateSP(this, &SControlRigDetails::ShouldShowPropertyOnDetailCustomization));
-	ControlTransformDetailsView->SetIsPropertyReadOnlyDelegate(FIsPropertyReadOnly::CreateSP(this, &SControlRigDetails::IsReadOnlyPropertyOnDetailCustomization));
-	ControlTransformDetailsView->SetGenericLayoutDetailsDelegate(FOnGetDetailCustomizationInstance::CreateStatic(&FControlRigEditModeGenericDetails::MakeInstance, ModeTools));
-
-	ControlTransformNoScaleDetailsView = FModuleManager::GetModuleChecked<FPropertyEditorModule>("PropertyEditor").CreateDetailView(DetailsViewArgs);
-	ControlTransformNoScaleDetailsView->SetKeyframeHandler(SharedThis(this));
-	ControlTransformNoScaleDetailsView->SetIsPropertyVisibleDelegate(FIsPropertyVisible::CreateSP(this, &SControlRigDetails::ShouldShowPropertyOnDetailCustomization));
-	ControlTransformNoScaleDetailsView->SetIsPropertyReadOnlyDelegate(FIsPropertyReadOnly::CreateSP(this, &SControlRigDetails::IsReadOnlyPropertyOnDetailCustomization));
-	ControlTransformNoScaleDetailsView->SetGenericLayoutDetailsDelegate(FOnGetDetailCustomizationInstance::CreateStatic(&FControlRigEditModeGenericDetails::MakeInstance, ModeTools));
-
-	ControlFloatDetailsView = FModuleManager::GetModuleChecked<FPropertyEditorModule>("PropertyEditor").CreateDetailView(DetailsViewArgs);
-	ControlFloatDetailsView->SetKeyframeHandler(SharedThis(this));
-	ControlFloatDetailsView->SetIsPropertyVisibleDelegate(FIsPropertyVisible::CreateSP(this, &SControlRigDetails::ShouldShowPropertyOnDetailCustomization));
-	ControlFloatDetailsView->SetIsPropertyReadOnlyDelegate(FIsPropertyReadOnly::CreateSP(this, &SControlRigDetails::IsReadOnlyPropertyOnDetailCustomization));
-	ControlFloatDetailsView->SetGenericLayoutDetailsDelegate(FOnGetDetailCustomizationInstance::CreateStatic(&FControlRigEditModeGenericDetails::MakeInstance, ModeTools));
-
-	IndividualControlFloatDetailsView = FModuleManager::GetModuleChecked<FPropertyEditorModule>("PropertyEditor").CreateDetailView(IndividualDetailsViewArgs);
-	IndividualControlFloatDetailsView->SetKeyframeHandler(SharedThis(this));
-	IndividualControlFloatDetailsView->SetIsPropertyVisibleDelegate(FIsPropertyVisible::CreateSP(this, &SControlRigDetails::ShouldShowPropertyOnDetailCustomization));
-	IndividualControlFloatDetailsView->SetIsPropertyReadOnlyDelegate(FIsPropertyReadOnly::CreateSP(this, &SControlRigDetails::IsReadOnlyPropertyOnDetailCustomization));
-	IndividualControlFloatDetailsView->SetGenericLayoutDetailsDelegate(FOnGetDetailCustomizationInstance::CreateStatic(&FControlRigEditModeGenericDetails::MakeInstance, ModeTools));
-
-	ControlEnumDetailsView = FModuleManager::GetModuleChecked<FPropertyEditorModule>("PropertyEditor").CreateDetailView(DetailsViewArgs);
-	ControlEnumDetailsView->SetKeyframeHandler(SharedThis(this));
-	ControlEnumDetailsView->SetIsPropertyVisibleDelegate(FIsPropertyVisible::CreateSP(this, &SControlRigDetails::ShouldShowPropertyOnDetailCustomization));
-	ControlEnumDetailsView->SetIsPropertyReadOnlyDelegate(FIsPropertyReadOnly::CreateSP(this, &SControlRigDetails::IsReadOnlyPropertyOnDetailCustomization));
-	ControlEnumDetailsView->SetGenericLayoutDetailsDelegate(FOnGetDetailCustomizationInstance::CreateStatic(&FControlRigEditModeGenericDetails::MakeInstance, ModeTools));
-
-	IndividualControlEnumDetailsView = FModuleManager::GetModuleChecked<FPropertyEditorModule>("PropertyEditor").CreateDetailView(IndividualDetailsViewArgs);
-	IndividualControlEnumDetailsView->SetKeyframeHandler(SharedThis(this));
-	IndividualControlEnumDetailsView->SetIsPropertyVisibleDelegate(FIsPropertyVisible::CreateSP(this, &SControlRigDetails::ShouldShowPropertyOnDetailCustomization));
-	IndividualControlEnumDetailsView->SetIsPropertyReadOnlyDelegate(FIsPropertyReadOnly::CreateSP(this, &SControlRigDetails::IsReadOnlyPropertyOnDetailCustomization));
-	IndividualControlEnumDetailsView->SetGenericLayoutDetailsDelegate(FOnGetDetailCustomizationInstance::CreateStatic(&FControlRigEditModeGenericDetails::MakeInstance, ModeTools));
-
-	ControlIntegerDetailsView = FModuleManager::GetModuleChecked<FPropertyEditorModule>("PropertyEditor").CreateDetailView(DetailsViewArgs);
-	ControlIntegerDetailsView->SetKeyframeHandler(SharedThis(this));
-	ControlIntegerDetailsView->SetIsPropertyVisibleDelegate(FIsPropertyVisible::CreateSP(this, &SControlRigDetails::ShouldShowPropertyOnDetailCustomization));
-	ControlIntegerDetailsView->SetIsPropertyReadOnlyDelegate(FIsPropertyReadOnly::CreateSP(this, &SControlRigDetails::IsReadOnlyPropertyOnDetailCustomization));
-	ControlIntegerDetailsView->SetGenericLayoutDetailsDelegate(FOnGetDetailCustomizationInstance::CreateStatic(&FControlRigEditModeGenericDetails::MakeInstance, ModeTools));
-
-	IndividualControlIntegerDetailsView = FModuleManager::GetModuleChecked<FPropertyEditorModule>("PropertyEditor").CreateDetailView(IndividualDetailsViewArgs);
-	IndividualControlIntegerDetailsView->SetKeyframeHandler(SharedThis(this));
-	IndividualControlIntegerDetailsView->SetIsPropertyVisibleDelegate(FIsPropertyVisible::CreateSP(this, &SControlRigDetails::ShouldShowPropertyOnDetailCustomization));
-	IndividualControlIntegerDetailsView->SetIsPropertyReadOnlyDelegate(FIsPropertyReadOnly::CreateSP(this, &SControlRigDetails::IsReadOnlyPropertyOnDetailCustomization));
-	IndividualControlIntegerDetailsView->SetGenericLayoutDetailsDelegate(FOnGetDetailCustomizationInstance::CreateStatic(&FControlRigEditModeGenericDetails::MakeInstance, ModeTools));
-
-
-	ControlBoolDetailsView = FModuleManager::GetModuleChecked<FPropertyEditorModule>("PropertyEditor").CreateDetailView(DetailsViewArgs);
-	ControlBoolDetailsView->SetKeyframeHandler(SharedThis(this));
-	ControlBoolDetailsView->SetIsPropertyVisibleDelegate(FIsPropertyVisible::CreateSP(this, &SControlRigDetails::ShouldShowPropertyOnDetailCustomization));
-	ControlBoolDetailsView->SetIsPropertyReadOnlyDelegate(FIsPropertyReadOnly::CreateSP(this, &SControlRigDetails::IsReadOnlyPropertyOnDetailCustomization));
-	ControlBoolDetailsView->SetGenericLayoutDetailsDelegate(FOnGetDetailCustomizationInstance::CreateStatic(&FControlRigEditModeGenericDetails::MakeInstance, ModeTools));
-
-	IndividualControlBoolDetailsView = FModuleManager::GetModuleChecked<FPropertyEditorModule>("PropertyEditor").CreateDetailView(IndividualDetailsViewArgs);
-	IndividualControlBoolDetailsView->SetKeyframeHandler(SharedThis(this));
-	IndividualControlBoolDetailsView->SetIsPropertyVisibleDelegate(FIsPropertyVisible::CreateSP(this, &SControlRigDetails::ShouldShowPropertyOnDetailCustomization));
-	IndividualControlBoolDetailsView->SetIsPropertyReadOnlyDelegate(FIsPropertyReadOnly::CreateSP(this, &SControlRigDetails::IsReadOnlyPropertyOnDetailCustomization));
-	IndividualControlBoolDetailsView->SetGenericLayoutDetailsDelegate(FOnGetDetailCustomizationInstance::CreateStatic(&FControlRigEditModeGenericDetails::MakeInstance, ModeTools));
-
-
-	ControlVectorDetailsView = FModuleManager::GetModuleChecked<FPropertyEditorModule>("PropertyEditor").CreateDetailView(DetailsViewArgs);
-	ControlVectorDetailsView->SetKeyframeHandler(SharedThis(this));
-	ControlVectorDetailsView->SetIsPropertyVisibleDelegate(FIsPropertyVisible::CreateSP(this, &SControlRigDetails::ShouldShowPropertyOnDetailCustomization));
-	ControlVectorDetailsView->SetIsPropertyReadOnlyDelegate(FIsPropertyReadOnly::CreateSP(this, &SControlRigDetails::IsReadOnlyPropertyOnDetailCustomization));
-	ControlVectorDetailsView->SetGenericLayoutDetailsDelegate(FOnGetDetailCustomizationInstance::CreateStatic(&FControlRigEditModeGenericDetails::MakeInstance, ModeTools));
-
-	ControlVector2DDetailsView = FModuleManager::GetModuleChecked<FPropertyEditorModule>("PropertyEditor").CreateDetailView(DetailsViewArgs);
-	ControlVector2DDetailsView->SetKeyframeHandler(SharedThis(this));
-	ControlVector2DDetailsView->SetIsPropertyVisibleDelegate(FIsPropertyVisible::CreateSP(this, &SControlRigDetails::ShouldShowPropertyOnDetailCustomization));
-	ControlVector2DDetailsView->SetIsPropertyReadOnlyDelegate(FIsPropertyReadOnly::CreateSP(this, &SControlRigDetails::IsReadOnlyPropertyOnDetailCustomization));
-	ControlVector2DDetailsView->SetGenericLayoutDetailsDelegate(FOnGetDetailCustomizationInstance::CreateStatic(&FControlRigEditModeGenericDetails::MakeInstance, ModeTools));
+	auto CreateDetailsView = [this](FDetailsViewArgs InDetailsViewArgs) -> TSharedPtr<IDetailsView>
+	{
+		TSharedPtr<IDetailsView> DetailsView = FModuleManager::GetModuleChecked<FPropertyEditorModule>("PropertyEditor").CreateDetailView(InDetailsViewArgs);
+		DetailsView->SetKeyframeHandler(SharedThis(this));
+		DetailsView->SetIsPropertyVisibleDelegate(FIsPropertyVisible::CreateSP(this, &SControlRigDetails::ShouldShowPropertyOnDetailCustomization));
+		DetailsView->SetIsPropertyReadOnlyDelegate(FIsPropertyReadOnly::CreateSP(this, &SControlRigDetails::IsReadOnlyPropertyOnDetailCustomization));
+		DetailsView->SetGenericLayoutDetailsDelegate(FOnGetDetailCustomizationInstance::CreateStatic(&FControlRigEditModeGenericDetails::MakeInstance, ModeTools));
+		return DetailsView;
+	};
+
+	ControlEulerTransformDetailsView = CreateDetailsView(DetailsViewArgs);
+	IndividualControlEulerTransformDetailsView = CreateDetailsView(IndividualDetailsViewArgs);
+	ControlTransformDetailsView = CreateDetailsView(DetailsViewArgs);
+	IndividualControlTransformDetailsView = CreateDetailsView(IndividualDetailsViewArgs);
+	ControlTransformNoScaleDetailsView = CreateDetailsView(DetailsViewArgs);
+	IndividualControlTransformNoScaleDetailsView = CreateDetailsView(IndividualDetailsViewArgs);
+	ControlFloatDetailsView = CreateDetailsView(DetailsViewArgs);
+	IndividualControlFloatDetailsView = CreateDetailsView(IndividualDetailsViewArgs);
+	ControlEnumDetailsView = CreateDetailsView(DetailsViewArgs);
+	IndividualControlEnumDetailsView = CreateDetailsView(IndividualDetailsViewArgs);
+	ControlIntegerDetailsView = CreateDetailsView(DetailsViewArgs);
+	IndividualControlIntegerDetailsView = CreateDetailsView(IndividualDetailsViewArgs);
+	ControlBoolDetailsView = CreateDetailsView(DetailsViewArgs);
+	IndividualControlBoolDetailsView = CreateDetailsView(IndividualDetailsViewArgs);
+	ControlVectorDetailsView = CreateDetailsView(DetailsViewArgs);
+	IndividualControlVectorDetailsView = CreateDetailsView(IndividualDetailsViewArgs);
+	ControlVector2DDetailsView = CreateDetailsView(DetailsViewArgs);
+	IndividualControlVector2DDetailsView = CreateDetailsView(IndividualDetailsViewArgs);
 
 	ChildSlot
 		[
@@ -370,6 +319,22 @@ void SControlRigDetails::Construct(const FArguments& InArgs, FControlRigEditMode
 				[
 					ControlFloatDetailsView.ToSharedRef()
 				]
+
+			+ SVerticalBox::Slot()
+				.AutoHeight()
+				[
+					IndividualControlEulerTransformDetailsView.ToSharedRef()
+				]
+			+ SVerticalBox::Slot()
+				.AutoHeight()
+				[
+					IndividualControlTransformDetailsView.ToSharedRef()
+				]
+			+ SVerticalBox::Slot()
+				.AutoHeight()
+				[
+					IndividualControlTransformNoScaleDetailsView.ToSharedRef()
+				]
 			+ SVerticalBox::Slot()
 				.AutoHeight()
 				[
@@ -385,12 +350,21 @@ void SControlRigDetails::Construct(const FArguments& InArgs, FControlRigEditMode
 				[
 					IndividualControlEnumDetailsView.ToSharedRef()
 				]
+			+ SVerticalBox::Slot()
+				.AutoHeight()
+				[
+					IndividualControlVectorDetailsView.ToSharedRef()
+				]
+			+ SVerticalBox::Slot()
+				.AutoHeight()
+				[
+					IndividualControlVector2DDetailsView.ToSharedRef()
+				]
 			+ SVerticalBox::Slot()
 				.AutoHeight()
 				[
 					IndividualControlFloatDetailsView.ToSharedRef()
 				]
-
 			]
 		];
 
@@ -422,7 +396,12 @@ void SControlRigDetails::UpdateProxies()
 		TArray<TWeakObjectPtr<>> Bools;
 		TArray<TWeakObjectPtr<>> Integers;
 		TArray<TWeakObjectPtr<>> Enums;
+		TArray<TWeakObjectPtr<>> IndividualEulers;
+		TArray<TWeakObjectPtr<>> IndividualTransforms;
+		TArray<TWeakObjectPtr<>> IndividualTransformNoScales;
 		TArray<TWeakObjectPtr<>> IndividualFloats;
+		TArray<TWeakObjectPtr<>> IndividualVectors;
+		TArray<TWeakObjectPtr<>> IndividualVector2Ds;
 		TArray<TWeakObjectPtr<>> IndividualBools;
 		TArray<TWeakObjectPtr<>> IndividualIntegers;
 		TArray<TWeakObjectPtr<>> IndividualEnums;
@@ -437,15 +416,36 @@ void SControlRigDetails::UpdateProxies()
 				{
 					if (Proxy->GetClass() == UControlRigTransformControlProxy::StaticClass())
 					{
-						Transforms.Add(Proxy);
+						if (Proxy->bIsIndividual)
+						{
+							IndividualTransforms.Add(Proxy);
+						}
+						else
+						{
+							Transforms.Add(Proxy);
+						}
 					}
 					else if (Proxy->GetClass() == UControlRigTransformNoScaleControlProxy::StaticClass())
 					{
-						TransformNoScales.Add(Proxy);
+						if (Proxy->bIsIndividual)
+						{
+							IndividualTransformNoScales.Add(Proxy);
+						}
+						else
+						{
+							TransformNoScales.Add(Proxy);
+						}
 					}
 					else if (Proxy->GetClass() == UControlRigEulerTransformControlProxy::StaticClass())
 					{
-						Eulers.Add(Proxy);
+						if (Proxy->bIsIndividual)
+						{
+							IndividualEulers.Add(Proxy);
+						}
+						else
+						{
+							Eulers.Add(Proxy);
+						}
 					}
 					else if (Proxy->GetClass() == UControlRigFloatControlProxy::StaticClass())
 					{
@@ -460,11 +460,25 @@ void SControlRigDetails::UpdateProxies()
 					}
 					else if (Proxy->GetClass() == UControlRigVectorControlProxy::StaticClass())
 					{
-						Vectors.Add(Proxy);
+						if (Proxy->bIsIndividual)
+						{
+							IndividualVectors.Add(Proxy);
+						}
+						else
+						{
+							Vectors.Add(Proxy);
+						}
 					}
 					else if (Proxy->GetClass() == UControlRigVector2DControlProxy::StaticClass())
 					{
-						Vector2Ds.Add(Proxy);
+						if (Proxy->bIsIndividual)
+						{
+							IndividualVector2Ds.Add(Proxy);
+						}
+						else
+						{
+							Vector2Ds.Add(Proxy);
+						}
 					}
 					else if (Proxy->GetClass() == UControlRigBoolControlProxy::StaticClass())
 					{
@@ -519,128 +533,136 @@ void SControlRigDetails::UpdateProxies()
 			}
 		}
 
-		SetTransformDetailsObjects(Transforms);
-		SetTransformNoScaleDetailsObjects(TransformNoScales);
-		SetEulerTransformDetailsObjects(Eulers);
-		SetVectorDetailsObjects(Vectors);
-		SetVector2DDetailsObjects(Vector2Ds);
+		SetTransformDetailsObjects(Transforms, false);
+		SetTransformNoScaleDetailsObjects(TransformNoScales, false);
+		SetEulerTransformDetailsObjects(Eulers, false);
+		SetVectorDetailsObjects(Vectors, false);
+		SetVector2DDetailsObjects(Vector2Ds, false);
 		SetFloatDetailsObjects(Floats, false);
 		SetBoolDetailsObjects(Bools,false);
 		SetIntegerDetailsObjects(Integers,false);
 		SetEnumDetailsObjects(Enums,false);
+		SetTransformDetailsObjects(IndividualTransforms, true);
+		SetTransformNoScaleDetailsObjects(IndividualTransformNoScales, true);
+		SetEulerTransformDetailsObjects(IndividualEulers, true);
+		SetVectorDetailsObjects(IndividualVectors, true);
+		SetVector2DDetailsObjects(IndividualVector2Ds, true);
 		SetFloatDetailsObjects(IndividualFloats, true);
-		SetBoolDetailsObjects(IndividualBools, true);
-		SetIntegerDetailsObjects(IndividualIntegers, true);
-		SetEnumDetailsObjects(IndividualEnums, true);
+		SetBoolDetailsObjects(IndividualBools,true);
+		SetIntegerDetailsObjects(IndividualIntegers,true);
+		SetEnumDetailsObjects(IndividualEnums,true);
 	});
 }
 
-void SControlRigDetails::SetEulerTransformDetailsObjects(const TArray<TWeakObjectPtr<>>& InObjects)
+void SControlRigDetails::SetEulerTransformDetailsObjects(const TArray<TWeakObjectPtr<>>& InObjects, bool bIsIndividual)
 {
-	if (ControlEulerTransformDetailsView)
+	TSharedPtr<IDetailsView>& DetailsView = bIsIndividual ? IndividualControlEulerTransformDetailsView : ControlEulerTransformDetailsView; 
+	if (DetailsView)
 	{
-		ControlEulerTransformDetailsView->SetObjects(InObjects);
+		DetailsView->SetObjects(InObjects);
 	}
 };
 
-void SControlRigDetails::SetTransformDetailsObjects(const TArray<TWeakObjectPtr<>>& InObjects)
+void SControlRigDetails::SetTransformDetailsObjects(const TArray<TWeakObjectPtr<>>& InObjects, bool bIsIndividual)
 {
-	if (ControlTransformDetailsView)
+	TSharedPtr<IDetailsView>& DetailsView = bIsIndividual ? IndividualControlTransformDetailsView : ControlTransformDetailsView; 
+	if (DetailsView)
 	{
-		ControlTransformDetailsView->SetObjects(InObjects);
+		DetailsView->SetObjects(InObjects);
 	}
 }
 
-void SControlRigDetails::SetTransformNoScaleDetailsObjects(const TArray<TWeakObjectPtr<>>& InObjects)
+void SControlRigDetails::SetTransformNoScaleDetailsObjects(const TArray<TWeakObjectPtr<>>& InObjects, bool bIsIndividual)
 {
-	if (ControlTransformNoScaleDetailsView)
+	TSharedPtr<IDetailsView>& DetailsView = bIsIndividual ? IndividualControlTransformNoScaleDetailsView : ControlTransformNoScaleDetailsView; 
+	if (DetailsView)
 	{
-		ControlTransformNoScaleDetailsView->SetObjects(InObjects);
+		DetailsView->SetObjects(InObjects);
 	}
 }
 
 void SControlRigDetails::SetFloatDetailsObjects(const TArray<TWeakObjectPtr<>>& InObjects, bool bIsIndividual)
 {
-	if (bIsIndividual)
-	{
-		if (IndividualControlFloatDetailsView)
-		{
-			IndividualControlFloatDetailsView->SetObjects(InObjects);
-		}
-	}
-	else if (ControlFloatDetailsView)
+	TSharedPtr<IDetailsView>& DetailsView = bIsIndividual ? IndividualControlFloatDetailsView : ControlFloatDetailsView; 
+	if (DetailsView)
 	{
-		ControlFloatDetailsView->SetObjects(InObjects);
+		DetailsView->SetObjects(InObjects);
 	}
 }
 
 void SControlRigDetails::SetBoolDetailsObjects(const TArray<TWeakObjectPtr<>>& InObjects, bool bIsIndividual)
 {
-	if (bIsIndividual)
+	TSharedPtr<IDetailsView>& DetailsView = bIsIndividual ? IndividualControlBoolDetailsView : ControlBoolDetailsView; 
+	if (DetailsView)
 	{
-		if (IndividualControlBoolDetailsView)
-		{
-			IndividualControlBoolDetailsView->SetObjects(InObjects);
-		}
-	}
-	else if (ControlBoolDetailsView)
-	{
-		ControlBoolDetailsView->SetObjects(InObjects);
+		DetailsView->SetObjects(InObjects);
 	}
 }
 
 void SControlRigDetails::SetIntegerDetailsObjects(const TArray<TWeakObjectPtr<>>& InObjects, bool bIsIndividual)
 {
-
-	if (bIsIndividual)
+	TSharedPtr<IDetailsView>& DetailsView = bIsIndividual ? IndividualControlIntegerDetailsView : ControlIntegerDetailsView; 
+	if (DetailsView)
 	{
-		if (IndividualControlIntegerDetailsView)
-		{
-			IndividualControlIntegerDetailsView->SetObjects(InObjects);
-		}
-	}
-	else if (ControlIntegerDetailsView)
-	{
-		ControlIntegerDetailsView->SetObjects(InObjects);
+		DetailsView->SetObjects(InObjects);
 	}
 }
+
 void SControlRigDetails::SetEnumDetailsObjects(const TArray<TWeakObjectPtr<>>& InObjects, bool bIsIndividual)
 {
-	if (bIsIndividual)
+	TSharedPtr<IDetailsView>& DetailsView = bIsIndividual ? IndividualControlEnumDetailsView : ControlEnumDetailsView; 
+	if (DetailsView)
 	{
-		if (IndividualControlEnumDetailsView)
-		{
-			IndividualControlEnumDetailsView->SetObjects(InObjects);
-		}
-	}
-	else if (ControlEnumDetailsView)
-	{
-		ControlEnumDetailsView->SetObjects(InObjects);
+		DetailsView->SetObjects(InObjects);
 	}
 }
 
-void SControlRigDetails::SetVectorDetailsObjects(const TArray<TWeakObjectPtr<>>& InObjects)
+void SControlRigDetails::SetVectorDetailsObjects(const TArray<TWeakObjectPtr<>>& InObjects, bool bIsIndividual)
 {
-	if (ControlVectorDetailsView)
+	TSharedPtr<IDetailsView>& DetailsView = bIsIndividual ? IndividualControlVectorDetailsView : ControlVectorDetailsView; 
+	if (DetailsView)
 	{
-		ControlVectorDetailsView->SetObjects(InObjects);
+		DetailsView->SetObjects(InObjects);
 	}
 }
 
-void SControlRigDetails::SetVector2DDetailsObjects(const TArray<TWeakObjectPtr<>>& InObjects)
+void SControlRigDetails::SetVector2DDetailsObjects(const TArray<TWeakObjectPtr<>>& InObjects, bool bIsIndividual)
 {
-	if (ControlVector2DDetailsView)
+	TSharedPtr<IDetailsView>& DetailsView = bIsIndividual ? IndividualControlVector2DDetailsView : ControlVector2DDetailsView; 
+	if (DetailsView)
 	{
-		ControlVector2DDetailsView->SetObjects(InObjects);
+		DetailsView->SetObjects(InObjects);
 	}
 }
+
 bool SControlRigDetails::IsPropertyKeyable(const UClass* InObjectClass, const IPropertyHandle& InPropertyHandle) const
 {
+	TArray<UObject*> OuterObjects;
+	InPropertyHandle.GetOuterObjects(OuterObjects);
+	if(OuterObjects.Num() == 1)
+	{
+		if (UControlRigControlsProxy* Proxy = Cast< UControlRigControlsProxy>(OuterObjects[0]))
+		{
+			if(UControlRig* ControlRig = Proxy->ControlRig.Get())
+			{
+				const FRigElementKey Key = FRigElementKey(Proxy->ControlName, ERigElementType::Control);
+				if(const FRigControlElement* ControlElement = ControlRig->GetHierarchy()->Find<FRigControlElement>(Key))
+				{
+					if(!ControlRig->GetHierarchy()->IsAnimatable(ControlElement))
+					{
+						return false;
+					}
+				}
+			}
+		}
+	}
+	
 	if (InObjectClass && InObjectClass->IsChildOf(UControlRigTransformNoScaleControlProxy::StaticClass()) && InObjectClass->IsChildOf(UControlRigEulerTransformControlProxy::StaticClass()) && InPropertyHandle.GetProperty()
 		&& InPropertyHandle.GetProperty()->GetFName() == GET_MEMBER_NAME_CHECKED(UControlRigTransformControlProxy, Transform))
 	{
 		return true;
 	}
+
 	FCanKeyPropertyParams CanKeyPropertyParams(InObjectClass, InPropertyHandle);
 	ISequencer* Sequencer = GetSequencer();
 	if (Sequencer && Sequencer->CanKeyProperty(CanKeyPropertyParams))
diff --git a/Engine/Plugins/Animation/ControlRig/Source/ControlRigEditor/Private/EditMode/SControlRigDetails.h b/Engine/Plugins/Animation/ControlRig/Source/ControlRigEditor/Private/EditMode/SControlRigDetails.h
--- a/Engine/Plugins/Animation/ControlRig/Source/ControlRigEditor/Private/EditMode/SControlRigDetails.h
+++ b/Engine/Plugins/Animation/ControlRig/Source/ControlRigEditor/Private/EditMode/SControlRigDetails.h
@@ -64,11 +64,11 @@ private:
 
 	/** Set the objects to be displayed in the details panel */
 	void SetSettingsDetailsObject(const TWeakObjectPtr<>& InObject);
-	void SetEulerTransformDetailsObjects(const TArray<TWeakObjectPtr<>>& InObjects);
-	void SetTransformDetailsObjects(const TArray<TWeakObjectPtr<>>& InObjects);
-	void SetTransformNoScaleDetailsObjects(const TArray<TWeakObjectPtr<>>& InObjects);
-	void SetVectorDetailsObjects(const TArray<TWeakObjectPtr<>>& InObjects);
-	void SetVector2DDetailsObjects(const TArray<TWeakObjectPtr<>>& InObjects);
+	void SetEulerTransformDetailsObjects(const TArray<TWeakObjectPtr<>>& InObjects, bool bIsIndividual);
+	void SetTransformDetailsObjects(const TArray<TWeakObjectPtr<>>& InObjects, bool bIsIndividual);
+	void SetTransformNoScaleDetailsObjects(const TArray<TWeakObjectPtr<>>& InObjects, bool bIsIndividual);
+	void SetVectorDetailsObjects(const TArray<TWeakObjectPtr<>>& InObjects, bool bIsIndividual);
+	void SetVector2DDetailsObjects(const TArray<TWeakObjectPtr<>>& InObjects, bool bIsIndividual);
 	void SetFloatDetailsObjects(const TArray<TWeakObjectPtr<>>& InObjects, bool bIsIndividual);
 	void SetBoolDetailsObjects(const TArray<TWeakObjectPtr<>>& InObjects, bool bIsIndividual);
 	void SetIntegerDetailsObjects(const TArray<TWeakObjectPtr<>>& InObjects, bool bIsIndividual);
@@ -88,10 +88,15 @@ private:
 	TSharedPtr<IDetailsView> ControlVector2DDetailsView;
 	TSharedPtr<IDetailsView> ControlVectorDetailsView;
 	//these will show more than one of the same type, will happen for controls with parents.
+	TSharedPtr<IDetailsView> IndividualControlEulerTransformDetailsView;
+	TSharedPtr<IDetailsView> IndividualControlTransformDetailsView;
+	TSharedPtr<IDetailsView> IndividualControlTransformNoScaleDetailsView;
 	TSharedPtr<IDetailsView> IndividualControlFloatDetailsView;
 	TSharedPtr<IDetailsView> IndividualControlBoolDetailsView;
 	TSharedPtr<IDetailsView> IndividualControlIntegerDetailsView;
 	TSharedPtr<IDetailsView> IndividualControlEnumDetailsView;
+	TSharedPtr<IDetailsView> IndividualControlVector2DDetailsView;
+	TSharedPtr<IDetailsView> IndividualControlVectorDetailsView;
 
 };
 
diff --git a/Engine/Plugins/Animation/ControlRig/Source/ControlRigEditor/Private/EditMode/SControlRigOutliner.cpp b/Engine/Plugins/Animation/ControlRig/Source/ControlRigEditor/Private/EditMode/SControlRigOutliner.cpp
--- a/Engine/Plugins/Animation/ControlRig/Source/ControlRigEditor/Private/EditMode/SControlRigOutliner.cpp
+++ b/Engine/Plugins/Animation/ControlRig/Source/ControlRigEditor/Private/EditMode/SControlRigOutliner.cpp
@@ -286,71 +286,7 @@ TPair<const FSlateBrush*, FSlateColor> SMultiRigHierarchyItem::GetBrushForElemen
 	if (InData.Key.IsSet())
 	{
 		const FRigElementKey Key = InData.Key.GetValue();
-
-		switch (Key.Type)
-		{
-		case ERigElementType::Control:
-		{
-			Brush = FControlRigEditorStyle::Get().GetBrush("ControlRig.Tree.Control");
-			if (const FRigControlElement* Control = InHierarchy->Find<FRigControlElement>(InData.Key.GetValue()))
-			{
-				FLinearColor ShapeColor = Control->Settings.ShapeColor;
-				// ensure the alpha is always visible
-				ShapeColor.A = 1.f;
-				Color = FSlateColor(ShapeColor);
-			}
-			break;
-		}
-		case ERigElementType::Null:
-		{
-			Brush = FControlRigEditorStyle::Get().GetBrush("ControlRig.Tree.Null");
-			break;
-		}
-		case ERigElementType::Bone:
-		{
-			ERigBoneType BoneType = ERigBoneType::User;
-
-			if (InHierarchy)
-			{
-				const FRigBoneElement* BoneElement = InHierarchy->Find<FRigBoneElement>(InData.Key.GetValue());
-				if (BoneElement)
-				{
-					BoneType = BoneElement->BoneType;
-				}
-			}
-
-			switch (BoneType)
-			{
-			case ERigBoneType::Imported:
-			{
-				Brush = FControlRigEditorStyle::Get().GetBrush("ControlRig.Tree.BoneImported");
-				break;
-			}
-			case ERigBoneType::User:
-			default:
-			{
-				Brush = FControlRigEditorStyle::Get().GetBrush("ControlRig.Tree.BoneUser");
-				break;
-			}
-			}
-
-			break;
-		}
-		case ERigElementType::RigidBody:
-		{
-			Brush = FControlRigEditorStyle::Get().GetBrush("ControlRig.Tree.RigidBody");
-			break;
-		}
-		case ERigElementType::Reference:
-		{
-			Brush = FControlRigEditorStyle::Get().GetBrush("ControlRig.Tree.Socket");
-			break;
-		}
-		default:
-		{
-			break;
-		}
-		}
+		return SRigHierarchyItem::GetBrushForElementType(InHierarchy, Key);
 	}
 	else
 	{
@@ -547,7 +483,8 @@ bool SMultiRigHierarchyTreeView::AddElement(UControlRig* InControlRig, const FRi
 	}
 	case ERigElementType::Control:
 	{
-		if (!Settings.bShowControls)
+		const FRigControlElement* ControlElement = CastChecked<FRigControlElement>(InElement);
+		if (!Settings.bShowControls || ControlElement->Settings.AnimationType == ERigControlAnimationType::VisualCue)
 		{
 			return false;
 		}
diff --git a/Engine/Plugins/Animation/ControlRig/Source/ControlRigEditor/Private/Editor/SRigHierarchy.cpp b/Engine/Plugins/Animation/ControlRig/Source/ControlRigEditor/Private/Editor/SRigHierarchy.cpp
--- a/Engine/Plugins/Animation/ControlRig/Source/ControlRigEditor/Private/Editor/SRigHierarchy.cpp
+++ b/Engine/Plugins/Animation/ControlRig/Source/ControlRigEditor/Private/Editor/SRigHierarchy.cpp
@@ -2714,7 +2714,7 @@ void SRigHierarchy::HandleSetShapeTransformFromCurrent()
 					{
 						const FRigElementKey Key = ControlElement->GetKey();
 						
-						if (ControlElement->Settings.bShapeEnabled)
+						if (ControlElement->Settings.SupportsShape())
 						{
 							const FTransform OffsetGlobalTransform = DebuggedHierarchy->GetGlobalControlOffsetTransform(Key); 
 							const FTransform ShapeGlobalTransform = DebuggedHierarchy->GetGlobalControlShapeTransform(Key);
diff --git a/Engine/Plugins/Animation/ControlRig/Source/ControlRigEditor/Private/Editor/SRigHierarchyTreeView.cpp b/Engine/Plugins/Animation/ControlRig/Source/ControlRigEditor/Private/Editor/SRigHierarchyTreeView.cpp
--- a/Engine/Plugins/Animation/ControlRig/Source/ControlRigEditor/Private/Editor/SRigHierarchyTreeView.cpp
+++ b/Engine/Plugins/Animation/ControlRig/Source/ControlRigEditor/Private/Editor/SRigHierarchyTreeView.cpp
@@ -16,6 +16,8 @@
 #include "ControlRig.h"
 #include "ControlRigEditorStyle.h"
 #include "Settings/ControlRigSettings.h"
+#include "Graph/ControlRigGraphSchema.h"
+#include "Styling/AppStyle.h"
 
 #define LOCTEXT_NAMESPACE "SRigHierarchyTreeView"
 
@@ -691,7 +693,75 @@ TPair<const FSlateBrush*, FSlateColor> SRigHierarchyItem::GetBrushForElementType
 			Brush = FControlRigEditorStyle::Get().GetBrush("ControlRig.Tree.Control");
 			if(const FRigControlElement* Control = InHierarchy->Find<FRigControlElement>(InKey))
 			{
-				FLinearColor ShapeColor = Control->Settings.ShapeColor;
+				FLinearColor ShapeColor = FLinearColor::White;
+				
+				if(Control->Settings.SupportsShape())
+				{
+					ShapeColor = Control->Settings.ShapeColor;
+				}
+				else
+				{
+					static FName TypeIcon(TEXT("Kismet.VariableList.TypeIcon"));
+					Brush = FAppStyle::GetBrush(TypeIcon);
+
+					FEdGraphPinType PinType;
+					switch(Control->Settings.ControlType)
+					{
+						case ERigControlType::Bool:
+						{
+							PinType = RigVMTypeUtils::PinTypeFromCPPType(RigVMTypeUtils::BoolType, nullptr);
+							break;
+						}
+						case ERigControlType::Float:
+						{
+							PinType = RigVMTypeUtils::PinTypeFromCPPType(RigVMTypeUtils::FloatType, nullptr);
+							break;
+						}
+						case ERigControlType::Integer:
+						{
+							if(Control->Settings.ControlEnum)
+							{
+								PinType = RigVMTypeUtils::PinTypeFromCPPType(FString(), Control->Settings.ControlEnum);
+							}
+							else
+							{
+								PinType = RigVMTypeUtils::PinTypeFromCPPType(RigVMTypeUtils::Int32Type, nullptr);
+							}
+							break;
+						}
+						case ERigControlType::Vector2D:
+						{
+							UScriptStruct* Struct = TBaseStructure<FVector2D>::Get(); 
+							PinType = RigVMTypeUtils::PinTypeFromCPPType(Struct->GetStructCPPName(), Struct);
+							break;
+						}
+						case ERigControlType::Position:
+						case ERigControlType::Scale:
+						{
+							UScriptStruct* Struct = TBaseStructure<FVector>::Get(); 
+							PinType = RigVMTypeUtils::PinTypeFromCPPType(Struct->GetStructCPPName(), Struct);
+							break;
+						}
+						case ERigControlType::Rotator:
+						{
+							UScriptStruct* Struct = TBaseStructure<FRotator>::Get(); 
+							PinType = RigVMTypeUtils::PinTypeFromCPPType(Struct->GetStructCPPName(), Struct);
+							break;
+						}
+						case ERigControlType::Transform:
+						case ERigControlType::TransformNoScale:
+						case ERigControlType::EulerTransform:
+						default:
+						{
+							UScriptStruct* Struct = TBaseStructure<FTransform>::Get(); 
+							PinType = RigVMTypeUtils::PinTypeFromCPPType(Struct->GetStructCPPName(), Struct);
+							break;
+						}
+					}
+					const UControlRigGraphSchema* Schema = GetDefault<UControlRigGraphSchema>();
+					ShapeColor = Schema->GetPinTypeColor(PinType);
+				}
+				
 				// ensure the alpha is always visible
 				ShapeColor.A = 1.f;
 				Color = FSlateColor(ShapeColor);
diff --git a/Engine/Plugins/Animation/ControlRig/Source/ControlRigEditor/Private/Sequencer/ControlRigParameterTrackEditor.cpp b/Engine/Plugins/Animation/ControlRig/Source/ControlRigEditor/Private/Sequencer/ControlRigParameterTrackEditor.cpp
--- a/Engine/Plugins/Animation/ControlRig/Source/ControlRigEditor/Private/Sequencer/ControlRigParameterTrackEditor.cpp
+++ b/Engine/Plugins/Animation/ControlRig/Source/ControlRigEditor/Private/Sequencer/ControlRigParameterTrackEditor.cpp
@@ -198,6 +198,7 @@ FControlRigParameterTrackEditor::FControlRigParameterTrackEditor(TSharedRef<ISeq
 	: FKeyframeTrackEditor<UMovieSceneControlRigParameterTrack>(InSequencer)
 	, bCurveDisplayTickIsPending(false)
 	, bIsDoingSelection(false)
+	, bSkipNextSelectionFromTimer(false)
 	, bFilterAssetBySkeleton(true)
 	, bFilterAssetByAnimatableControls(true)
 {
@@ -1536,7 +1537,8 @@ static void EvaluateThisControl(UMovieSceneControlRigParameterSection* Section,
 	}
 	if(FRigControlElement* ControlElement = ControlRig->FindControl(ControlName))
 	{ 
-		FControlRigInteractionScope InteractionScope(ControlRig);
+		FControlRigInteractionScope InteractionScope(ControlRig, ControlElement->GetKey());
+		
 		//eval any space for this channel, if not additive section
 		if (Section->GetBlendType().Get() != EMovieSceneBlendType::Additive)
 		{
@@ -1950,6 +1952,12 @@ void FControlRigParameterTrackEditor::OnSelectionChanged(TArray<UMovieSceneTrack
 
 	TGuardValue<bool> Guard(bIsDoingSelection, true);
 
+	if(bSkipNextSelectionFromTimer)
+	{
+		bSkipNextSelectionFromTimer = false;
+		return;
+	}
+
 	FControlRigEditMode* ControlRigEditMode = GetEditMode();
 	UControlRig* ControlRig = nullptr;
 
@@ -2058,20 +2066,15 @@ void FControlRigParameterTrackEditor::SelectRigsAndControls(UControlRig* Control
 					{
 						if(const FRigControlElement* ControlElement = Hierarchy->Find<FRigControlElement>(FRigElementKey(ControlName, ERigElementType::Control)))
 						{
-							if (ControlElement->Settings.ControlType == ERigControlType::Bool ||
-								ControlElement->Settings.ControlType == ERigControlType::Float ||
-								ControlElement->Settings.ControlType == ERigControlType::Integer)
+							if(ControlElement->Settings.AnimationType == ERigControlAnimationType::AnimationChannel)
 							{
-								if(ControlElement->Settings.bShapeEnabled)
+								if(const FRigControlElement* ParentControlElement = Cast<FRigControlElement>(Hierarchy->GetFirstParent(ControlElement)))
 								{
-									if(const FRigControlElement* ParentControlElement = Cast<FRigControlElement>(Hierarchy->GetFirstParent(ControlElement)))
+									if(const TSet<FName>* Controls = RigsAndControls.Find(ControlRig))
 									{
-										if(const TSet<FName>* Controls = RigsAndControls.Find(ControlRig))
+										if(Controls->Contains(ParentControlElement->GetName()))
 										{
-											if(Controls->Contains(ParentControlElement->GetName()))
-											{
-												continue;
-											}
+											continue;
 										}
 									}
 								}
@@ -2384,20 +2387,37 @@ void FControlRigParameterTrackEditor::SetUpEditModeIfNeeded(UControlRig* Control
 
 void FControlRigParameterTrackEditor::HandleControlSelected(UControlRig* Subject, FRigControlElement* ControlElement, bool bSelected)
 {
+	URigHierarchy* Hierarchy = Subject->GetHierarchy();
+	static bool bIsSelectingIndirectControl = false;
+
+	if(ControlElement && ControlElement->Settings.AnimationType == ERigControlAnimationType::ProxyControl)
+	{
+		const TArray<FRigElementKey>& DrivenControls = ControlElement->Settings.DrivenControls;
+		for(const FRigElementKey& DrivenKey : DrivenControls)
+		{
+			if(FRigControlElement* DrivenControl = Hierarchy->Find<FRigControlElement>(DrivenKey))
+			{
+				TGuardValue<bool> SubControlGuard(bIsSelectingIndirectControl, true);
+				HandleControlSelected(Subject, DrivenControl, bSelected);
+			}
+		}
+		return;
+	}
+	
 	//if parent selected we select child here if it's a bool,integer or single float
 	TArray<FRigControl> Controls;
 
-	URigHierarchy* Hierarchy = Subject->GetHierarchy();
-
+	if(!bIsSelectingIndirectControl)
+	{
 		if (URigHierarchyController* Controller = Hierarchy->GetController())
 		{
-			Hierarchy->ForEach<FRigControlElement>([ControlElement, Controller, bSelected](FRigControlElement* OtherControlElement) -> bool
+			Hierarchy->ForEach<FRigControlElement>([Hierarchy, ControlElement, Controller, bSelected](FRigControlElement* OtherControlElement) -> bool
 				{
 					if (OtherControlElement->Settings.ControlType == ERigControlType::Bool ||
 						OtherControlElement->Settings.ControlType == ERigControlType::Float ||
 						OtherControlElement->Settings.ControlType == ERigControlType::Integer)
 					{
-						if(OtherControlElement->Settings.bShapeEnabled || !OtherControlElement->Settings.bAnimatable)
+						if(OtherControlElement->Settings.SupportsShape() || !Hierarchy->IsAnimatable(OtherControlElement))
 						{
 							return true;
 						}
@@ -2415,7 +2435,7 @@ void FControlRigParameterTrackEditor::HandleControlSelected(UControlRig* Subject
 					return true;
 				});
 		}
-
+	}
 	
 	if (bIsDoingSelection)
 	{
@@ -2458,7 +2478,8 @@ void FControlRigParameterTrackEditor::HandleControlSelected(UControlRig* Subject
 			SetUpEditModeIfNeeded(Subject);
 
 			//Force refresh later, not now
-			GetSequencer()->NotifyMovieSceneDataChanged(EMovieSceneDataChangeType::RefreshTree);				
+			bSkipNextSelectionFromTimer = true;
+			GetSequencer()->NotifyMovieSceneDataChanged(EMovieSceneDataChangeType::RefreshTree);
 		}
 	}
 }
@@ -2529,7 +2550,7 @@ void FControlRigParameterTrackEditor::GetControlRigKeys(UControlRig* InControlRi
 	{
 		FRigControlElement* ControlElement = Controls[ControlIndex];
 
-		if (!ControlElement->Settings.bAnimatable)
+		if (!InControlRig->GetHierarchy()->IsAnimatable(ControlElement))
 		{
 			continue;
 		}
@@ -2927,7 +2948,7 @@ bool FControlRigParameterTrackEditor::ModifyOurGeneratedKeysByCurrentAndWeight(U
 	FChannelMapInfo* pChannelIndex = nullptr;
 	for (FRigControlElement* ControlElement : Controls)
 	{
-		if (!ControlElement->Settings.bAnimatable)
+		if (!InControlRig->GetHierarchy()->IsAnimatable(ControlElement))
 		{
 			continue;
 		}
@@ -3536,7 +3557,7 @@ bool FControlRigParameterTrackEditor::CollapseAllLayers(TSharedPtr<ISequencer>&
 
 		for (FRigControlElement* ControlElement : Controls)
 		{
-			if (!ControlElement->Settings.bAnimatable)
+			if (!ControlRig->GetHierarchy()->IsAnimatable(ControlElement))
 			{
 				continue;
 			}
diff --git a/Engine/Plugins/Animation/ControlRig/Source/ControlRigEditor/Private/Sequencer/ControlRigParameterTrackEditor.h b/Engine/Plugins/Animation/ControlRig/Source/ControlRigEditor/Private/Sequencer/ControlRigParameterTrackEditor.h
--- a/Engine/Plugins/Animation/ControlRig/Source/ControlRigEditor/Private/Sequencer/ControlRigParameterTrackEditor.h
+++ b/Engine/Plugins/Animation/ControlRig/Source/ControlRigEditor/Private/Sequencer/ControlRigParameterTrackEditor.h
@@ -223,6 +223,9 @@ private:
 	/** Guard to stop infinite loops when handling control selections*/
 	bool bIsDoingSelection;
 
+	/** A flag to determine if the next update coming from the timer should be skipped */
+	bool bSkipNextSelectionFromTimer;
+
 	/** Whether or not we should check Skeleton when filtering*/
 	bool bFilterAssetBySkeleton;
 
diff --git a/Engine/Plugins/Animation/ControlRig/Source/ControlRigEditor/Private/Tools/ControlRigTweener.cpp b/Engine/Plugins/Animation/ControlRig/Source/ControlRigEditor/Private/Tools/ControlRigTweener.cpp
--- a/Engine/Plugins/Animation/ControlRig/Source/ControlRigEditor/Private/Tools/ControlRigTweener.cpp
+++ b/Engine/Plugins/Animation/ControlRig/Source/ControlRigEditor/Private/Tools/ControlRigTweener.cpp
@@ -53,7 +53,7 @@ void FControlsToTween::Setup(const TArray<UControlRig*>& SelectedControlRigs, TW
 					TArray<FKeyHandle> Handles;
 					for (FRigControlElement* ControlElement : CurrentControls)
 					{
-						if (ControlElement->Settings.bAnimatable &&  ControlRig->IsControlSelected(ControlElement->GetName()))
+						if (ControlRig->GetHierarchy()->IsAnimatable(ControlElement) &&  ControlRig->IsControlSelected(ControlElement->GetName()))
 						{
 							FControlRigChannels ControlRigChannels;
 							ControlRigChannels.NumChannels = 0;
diff --git a/Engine/Source/Runtime/RigVM/Public/RigVMCore/RigVMExecuteContext.h b/Engine/Source/Runtime/RigVM/Public/RigVMCore/RigVMExecuteContext.h
--- a/Engine/Source/Runtime/RigVM/Public/RigVMCore/RigVMExecuteContext.h
+++ b/Engine/Source/Runtime/RigVM/Public/RigVMCore/RigVMExecuteContext.h
@@ -274,6 +274,7 @@ private:
 	friend struct FRigUnit_BeginExecution;
 	friend struct FRigUnit_InverseExecution;
 	friend struct FRigUnit_PrepareForExecution;
+	friend struct FRigUnit_InteractionExecution;
 	friend struct FEngineTestRigVM_Begin;
 	friend struct FEngineTestRigVM_Setup; 
 };
