Add ability to tweak initial setup of FK-controlrig (optimization purposes):
- Whether or not to import all curves from the Skeleton
- Define specific set of curve names to import
- Whether or not to generate curve controls
- Whether or not to generate bone controls
- Define specific set of bonenames to create controls for
(to be used by AnimSequencer)

diff --git a/Engine/Plugins/Animation/ControlRig/Source/ControlRig/Private/Rigs/FKControlRig.cpp b/Engine/Plugins/Animation/ControlRig/Source/ControlRig/Private/Rigs/FKControlRig.cpp
--- a/Engine/Plugins/Animation/ControlRig/Source/ControlRig/Private/Rigs/FKControlRig.cpp
+++ b/Engine/Plugins/Animation/ControlRig/Source/ControlRig/Private/Rigs/FKControlRig.cpp
@@ -319,54 +319,87 @@ void UFKControlRig::CreateRigElements(const FReferenceSkeleton& InReferenceSkele
 	{
 		Controller->ImportBones(InReferenceSkeleton, NAME_None, false, false, true, false);
 
-		if (InSmartNameMapping)
+		if (InitializationOptions.bGenerateCurveControls)
 		{
-			TArray<FName> NameArray;
-			InSmartNameMapping->FillNameArray(NameArray);
-			for (int32 Index = 0; Index < NameArray.Num(); ++Index)
-			{
-				 Controller->AddCurve(NameArray[Index], 0.f, false);
-			}
-		}
+			// try to add curves for manually specified names
+			if (InitializationOptions.CurveNames.Num() > 0)
+            {
+            	for (int32 Index = 0; Index < InitializationOptions.CurveNames.Num(); ++Index)
+            	{
+            		if(InSmartNameMapping && InSmartNameMapping->FindUID(InitializationOptions.CurveNames[Index]) != SmartName::MaxUID)
+            		{
+            			Controller->AddCurve(InitializationOptions.CurveNames[Index], 0.f, false);
+            		}
+            	}
+            }
+			// add all curves found on the skeleton
+            else if (InSmartNameMapping && InitializationOptions.bImportCurves)
+            {
+            	TArray<FName> NameArray;
+            	InSmartNameMapping->FillNameArray(NameArray);
+            	for (int32 Index = 0; Index < NameArray.Num(); ++Index)
+            	{
+            		Controller->AddCurve(NameArray[Index], 0.f, false);
+            	}
+            }
+		}		
 
 		// add control for all bone hierarchy 
-		int32 ControlIndex = 0;
-
-		GetHierarchy()->ForEach<FRigBoneElement>([&](FRigBoneElement* BoneElement) -> bool
+		if (InitializationOptions.bGenerateBoneControls)
 		{
-			const FName BoneName = BoneElement->GetName();
-			const FName ControlName = GetControlName(BoneName, BoneElement->GetType());
-			const FRigElementKey ParentKey = GetHierarchy()->GetFirstParent(BoneElement->GetKey());
-
-			FTransform OffsetTransform = FTransform::Identity;
-
-			FRigControlSettings Settings;
-			Settings.ControlType = ERigControlType::EulerTransform;
-			Settings.DisplayName = BoneName;
-
-			OffsetTransform.NormalizeRotation();
+			auto CreateControlForBoneElement = [this, Controller](const FRigBoneElement* BoneElement)
+			{
+				const FName BoneName = BoneElement->GetName();
+				const FName ControlName = GetControlName(BoneName, BoneElement->GetType());
+				const FRigElementKey ParentKey = GetHierarchy()->GetFirstParent(BoneElement->GetKey());
+
+				FRigControlSettings Settings;
+				Settings.ControlType = ERigControlType::EulerTransform;
+				Settings.DisplayName = BoneName;
+					
+				Controller->AddControl(ControlName, ParentKey, Settings, FRigControlValue::Make(FEulerTransform::Identity), FTransform::Identity, FTransform::Identity, false);
+			};
 			
-			Controller->AddControl(ControlName, ParentKey, Settings, FRigControlValue::Make(FEulerTransform::Identity), OffsetTransform, FTransform::Identity, false);
-
-			return true;
-		});
+			if (InitializationOptions.BoneNames.Num())
+			{
+				for (const FName& BoneName : InitializationOptions.BoneNames)
+				{
+					FRigElementKey BoneElementKey(BoneName, ERigElementType::Bone);
+					if (const FRigBoneElement* BoneElement = GetHierarchy()->Find<FRigBoneElement>(BoneElementKey))
+					{
+						CreateControlForBoneElement(BoneElement);
+					}
+				}
+			}
+			else
+			{
+				GetHierarchy()->ForEach<FRigBoneElement>([&](const FRigBoneElement* BoneElement) -> bool
+				{
+					CreateControlForBoneElement(BoneElement);
+					return true;
+				});
+			}
 
-		SetControlOffsetsFromBoneInitials();
+			SetControlOffsetsFromBoneInitials();
+		}
 		
-		GetHierarchy()->ForEach<FRigCurveElement>([&](FRigCurveElement* CurveElement) -> bool
+		if (InitializationOptions.bGenerateCurveControls)
 		{
-			const FName ControlName = GetControlName(CurveElement->GetName(), CurveElement->GetType());
-
-			FRigControlSettings Settings;
-			Settings.ControlType = ERigControlType::Float;
-			Settings.DisplayName = CurveElement->GetName();
-
-			const FName DisplayCurveControlName(*(CurveElement->GetName().ToString() + TEXT(" Curve")));
-		    Settings.DisplayName = DisplayCurveControlName;
-			Controller->AddControl(ControlName, FRigElementKey(), Settings, FRigControlValue::Make(CurveElement->Value), FTransform::Identity, FTransform::Identity, false);
-			
-			return true;
-		});
+			GetHierarchy()->ForEach<FRigCurveElement>([&](const FRigCurveElement* CurveElement) -> bool
+			{
+				const FName ControlName = GetControlName(CurveElement->GetName(), CurveElement->GetType());
+
+				FRigControlSettings Settings;
+				Settings.ControlType = ERigControlType::Float;
+				Settings.DisplayName = CurveElement->GetName();
+
+				const FName DisplayCurveControlName(*(CurveElement->GetName().ToString() + TEXT(" Curve")));
+				Settings.DisplayName = DisplayCurveControlName;
+				Controller->AddControl(ControlName, FRigElementKey(), Settings, FRigControlValue::Make(CurveElement->Value), FTransform::Identity, FTransform::Identity, false);
+				
+				return true;
+			});
+		}
 
 		RefreshActiveControls();
 	}
diff --git a/Engine/Plugins/Animation/ControlRig/Source/ControlRig/Public/Rigs/FKControlRig.h b/Engine/Plugins/Animation/ControlRig/Source/ControlRig/Public/Rigs/FKControlRig.h
--- a/Engine/Plugins/Animation/ControlRig/Source/ControlRig/Public/Rigs/FKControlRig.h
+++ b/Engine/Plugins/Animation/ControlRig/Source/ControlRig/Public/Rigs/FKControlRig.h
@@ -65,6 +65,24 @@ public:
 
 	// Ensures that controls mask is updated according to contained ControlRig (control) elements
 	void RefreshActiveControls();
+
+	struct FRigElementInitializationOptions
+	{	
+		// Flag whether or not to generate a transform control for bones
+		bool bGenerateBoneControls = true;
+		// Flag whether or not to generate a float control for all curves in the hierarchy
+		bool bGenerateCurveControls = true;
+		
+		// Flag whether or not to import all curves from SmartNameMapping
+		bool bImportCurves = true;
+
+		// Set of bone names to generate a transform control for
+		TArray<FName> BoneNames;
+		// Set of curve names to generate a float control for (requires bImportCurves to be false)
+		TArray<FName> CurveNames;
+	};
+	void SetInitializationOptions(const FRigElementInitializationOptions& Options) { InitializationOptions = Options; }
+	
 private:
 
 	/** Create RigElements - bone hierarchy and curves - from incoming skeleton */
@@ -79,5 +97,6 @@ private:
 	EControlRigFKRigExecuteMode ApplyMode;
 	EControlRigFKRigExecuteMode CachedToggleApplyMode;
 
+	FRigElementInitializationOptions InitializationOptions;
 	friend class FControlRigInteractionTest;
 };
\ No newline at end of file
