Control Rig: Space Switching:  Baking borked from multiple control rigs, need to make sure to use the right control rig and hierarchy.
#trivial

diff --git a/Engine/Plugins/Animation/ControlRig/Source/ControlRigEditor/Private/EditMode/SControlRigEditModeTools.cpp b/Engine/Plugins/Animation/ControlRig/Source/ControlRigEditor/Private/EditMode/SControlRigEditModeTools.cpp
--- a/Engine/Plugins/Animation/ControlRig/Source/ControlRigEditor/Private/EditMode/SControlRigEditModeTools.cpp
+++ b/Engine/Plugins/Animation/ControlRig/Source/ControlRigEditor/Private/EditMode/SControlRigEditModeTools.cpp
@@ -675,13 +675,11 @@ void SControlRigEditModeTools::OnRigElementSelected(UControlRig* Subject, FRigCo
 		}
 	}
 #endif
-	//mz todo handle multiple rigs
-	UControlRig* Rig = ControlRigs.Num() > 0 ? ControlRigs[0].Get() : nullptr;
-	if (Rig)
+	if (Subject)
 	{
 		// get the selected controls
-		TArray<FRigElementKey> SelectedControls = Rig->GetHierarchy()->GetSelectedKeys(ERigElementType::Control);
-		SpacePickerWidget->SetControls(Rig->GetHierarchy(), SelectedControls);
+		TArray<FRigElementKey> SelectedControls = Subject->GetHierarchy()->GetSelectedKeys(ERigElementType::Control);
+		SpacePickerWidget->SetControls(Subject->GetHierarchy(), SelectedControls);
 	}
 }
 
@@ -805,7 +803,7 @@ FReply SControlRigEditModeTools::OnBakeControlsToNewSpaceButtonClicked()
 	bool bNoValidControlRig = true;
 	for (TWeakObjectPtr<UControlRig>& ControlRig : ControlRigs)
 	{
-		if (ControlRig.IsValid())
+		if (ControlRig.IsValid() && SpacePickerWidget->GetHierarchy() == ControlRig->GetHierarchy())
 		{
 			bNoValidControlRig = false;
 			break;
@@ -823,7 +821,7 @@ FReply SControlRigEditModeTools::OnBakeControlsToNewSpaceButtonClicked()
 	}
 	for (TWeakObjectPtr<UControlRig>& ControlRig : ControlRigs)
 	{
-		if (ControlRig.IsValid())
+		if (ControlRig.IsValid() && SpacePickerWidget->GetHierarchy() == ControlRig->GetHierarchy())
 		{
 
 			FRigSpacePickerBakeSettings Settings;
