ControlRig: Improve Control Rig outliner for FKControlRig

diff --git a/Engine/Plugins/Animation/ControlRig/Source/ControlRigEditor/Private/EditMode/SControlRigOutliner.cpp b/Engine/Plugins/Animation/ControlRig/Source/ControlRigEditor/Private/EditMode/SControlRigOutliner.cpp
--- a/Engine/Plugins/Animation/ControlRig/Source/ControlRigEditor/Private/EditMode/SControlRigOutliner.cpp
+++ b/Engine/Plugins/Animation/ControlRig/Source/ControlRigEditor/Private/EditMode/SControlRigOutliner.cpp
@@ -39,6 +39,7 @@
 #include "Editor/EditorEngine.h"
 #include "HelperUtil.h"
 #include "ScopedTransaction.h"
+#include "Rigs/FKControlRig.h"
 
 #define LOCTEXT_NAMESPACE "ControlRigOutliner"
 
@@ -82,6 +83,24 @@ FText FMultiRigData::GetName() const
 	return FText::FromName(None);
 }
 
+FText FMultiRigData::GetDisplayName() const
+{
+	if (Key.IsSet())
+	{
+		if(ControlRig.IsValid())
+		{
+			if(const FRigControlElement* ControlElement = ControlRig->GetHierarchy()->Find<FRigControlElement>(Key.GetValue()))
+			{
+				if(!ControlElement->Settings.DisplayName.IsNone())
+				{
+					return FText::FromName(ControlElement->Settings.DisplayName);
+				}
+			}
+		}
+	}
+	return GetName();
+}
+
 bool FMultiRigData::operator== (const FMultiRigData & Other) const
 {
 	if (ControlRig.IsValid() == false && Other.ControlRig.IsValid() == false)
@@ -241,7 +260,7 @@ void SMultiRigHierarchyItem::Construct(const FArguments& InArgs, const TSharedRe
 			.VAlign(VAlign_Center)
 			[
 				SAssignNew(InlineWidget, SInlineEditableTextBlock)
-				.Text(this, &SMultiRigHierarchyItem::GetName)
+				.Text(this, &SMultiRigHierarchyItem::GetDisplayName)
 			.MultiLine(false)
 			.ColorAndOpacity_Lambda([this]()
 				{
@@ -279,6 +298,11 @@ FText SMultiRigHierarchyItem::GetName() const
 	return (WeakRigTreeElement.Pin()->Data.GetName());
 }
 
+FText SMultiRigHierarchyItem::GetDisplayName() const
+{
+	return (WeakRigTreeElement.Pin()->Data.GetDisplayName());
+}
+
 TPair<const FSlateBrush*, FSlateColor> SMultiRigHierarchyItem::GetBrushForElementType(const URigHierarchy* InHierarchy, const FMultiRigData& InData)
 {
 	const FSlateBrush* Brush = nullptr;
@@ -457,64 +481,73 @@ bool SMultiRigHierarchyTreeView::AddElement(UControlRig* InControlRig, const FRi
 
 	const FRigTreeDisplaySettings& Settings = Delegates.GetDisplaySettings();
 
-	switch (InElement->GetType())
-	{
-	case ERigElementType::Bone:
+	auto IsElementShown = [Settings](const FRigBaseElement* InElement) -> bool
 	{
-		if (!Settings.bShowBones)
+		switch (InElement->GetType())
 		{
-			return false;
-		}
+			case ERigElementType::Bone:
+			{
+				if (!Settings.bShowBones)
+				{
+					return false;
+				}
 
-		const FRigBoneElement* BoneElement = CastChecked<FRigBoneElement>(InElement);
-		if (!Settings.bShowImportedBones && BoneElement->BoneType == ERigBoneType::Imported)
-		{
-			return false;
-		}
-		break;
-	}
-	case ERigElementType::Null:
-	{
-		if (!Settings.bShowNulls)
-		{
-			return false;
-		}
-		break;
-	}
-	case ERigElementType::Control:
-	{
-		const FRigControlElement* ControlElement = CastChecked<FRigControlElement>(InElement);
-		if (!Settings.bShowControls || ControlElement->Settings.AnimationType == ERigControlAnimationType::VisualCue)
-		{
-			return false;
-		}
-		break;
-	}
-	case ERigElementType::RigidBody:
-	{
-		if (!Settings.bShowRigidBodies)
-		{
-			return false;
-		}
-		break;
-	}
-	case ERigElementType::Reference:
-	{
-		if (!Settings.bShowReferences)
-		{
-			return false;
+				const FRigBoneElement* BoneElement = CastChecked<FRigBoneElement>(InElement);
+				if (!Settings.bShowImportedBones && BoneElement->BoneType == ERigBoneType::Imported)
+				{
+					return false;
+				}
+				break;
+			}
+			case ERigElementType::Null:
+			{
+				if (!Settings.bShowNulls)
+				{
+					return false;
+				}
+				break;
+			}
+			case ERigElementType::Control:
+			{
+				const FRigControlElement* ControlElement = CastChecked<FRigControlElement>(InElement);
+				if (!Settings.bShowControls || ControlElement->Settings.AnimationType == ERigControlAnimationType::VisualCue)
+				{
+					return false;
+				}
+				break;
+			}
+			case ERigElementType::RigidBody:
+			{
+				if (!Settings.bShowRigidBodies)
+				{
+					return false;
+				}
+				break;
+			}
+			case ERigElementType::Reference:
+			{
+				if (!Settings.bShowReferences)
+				{
+					return false;
+				}
+				break;
+			}
+			case ERigElementType::Curve:
+			{
+				return false;
+			}
+			default:
+			{
+				break;
+			}
 		}
-		break;
-	}
-	case ERigElementType::Curve:
+		return true;
+	};
+
+	if(!IsElementShown(InElement))
 	{
 		return false;
 	}
-	default:
-	{
-		break;
-	}
-	}
 
 	FMultiRigData ParentData;
 	ParentData.ControlRig = InControlRig;
@@ -524,6 +557,8 @@ bool SMultiRigHierarchyTreeView::AddElement(UControlRig* InControlRig, const FRi
 		return false;
 	}
 
+	UFKControlRig* FKControlRig = Cast<UFKControlRig>(InControlRig);
+
 	if (ElementMap.Contains(Data))
 	{
 		if (const URigHierarchy* Hierarchy = InControlRig->GetHierarchy())
@@ -549,15 +584,35 @@ bool SMultiRigHierarchyTreeView::AddElement(UControlRig* InControlRig, const FRi
 			}
 			if (ParentKey.IsValid())
 			{
+				if(FKControlRig)
+				{
+					if(const FRigControlElement* ControlElement = Cast<FRigControlElement>(InElement))
+					{
+						if(ControlElement->Settings.AnimationType == ERigControlAnimationType::AnimationControl)
+						{
+							const FRigElementKey& ElementKey = InElement->GetKey();
+							const FName BoneName = FKControlRig->GetControlTargetName(ElementKey.Name, ElementKey.Type);
+							const FRigElementKey ParentBoneKey = Hierarchy->GetFirstParent(FRigElementKey(BoneName, ERigElementType::Bone));
+							if(ParentBoneKey.IsValid())
+							{
+								ParentKey = FRigElementKey(FKControlRig->GetControlName(ParentBoneKey.Name, ElementKey.Type), ElementKey.Type);
+							}
+						}
+					}
+				}
+
 				if (const FRigBaseElement* ParentElement = Hierarchy->Find(ParentKey))
 				{
-					AddElement(InControlRig,ParentElement);
+					if(ParentElement != nullptr)
+					{
+						AddElement(InControlRig,ParentElement);
 
-					FMultiRigData NewParentData(InControlRig, ParentKey);
+						FMultiRigData NewParentData(InControlRig, ParentKey);
 
-					if (ElementMap.Contains(NewParentData))
-					{
-						ReparentElement(Data, NewParentData);
+						if (ElementMap.Contains(NewParentData))
+						{
+							ReparentElement(Data, NewParentData);
+						}
 					}
 				}
 			}
diff --git a/Engine/Plugins/Animation/ControlRig/Source/ControlRigEditor/Private/EditMode/SControlRigOutliner.h b/Engine/Plugins/Animation/ControlRig/Source/ControlRigEditor/Private/EditMode/SControlRigOutliner.h
--- a/Engine/Plugins/Animation/ControlRig/Source/ControlRigEditor/Private/EditMode/SControlRigOutliner.h
+++ b/Engine/Plugins/Animation/ControlRig/Source/ControlRigEditor/Private/EditMode/SControlRigOutliner.h
@@ -80,6 +80,7 @@ public:
 	FMultiRigData() {};
 	FMultiRigData(UControlRig* InControlRig, FRigElementKey InKey) : ControlRig(InControlRig), Key(InKey) {};
 	FText GetName() const;
+	FText GetDisplayName() const;
 	bool operator == (const FMultiRigData & Other) const;
 	bool IsValid() const;
 	URigHierarchy* GetHierarchy() const;
@@ -127,6 +128,7 @@ private:
 	FMultiRigTreeDelegates Delegates;
 
 	FText GetName() const;
+	FText GetDisplayName() const;
 	FReply OnGetSelectedClicked();
 
 };
