ControlRig: Only remove rig_unit variables for very-old assets

diff --git a/Engine/Plugins/Animation/ControlRig/Source/ControlRigDeveloper/Private/ControlRigBlueprint.cpp b/Engine/Plugins/Animation/ControlRig/Source/ControlRigDeveloper/Private/ControlRigBlueprint.cpp
--- a/Engine/Plugins/Animation/ControlRig/Source/ControlRigDeveloper/Private/ControlRigBlueprint.cpp
+++ b/Engine/Plugins/Animation/ControlRig/Source/ControlRigDeveloper/Private/ControlRigBlueprint.cpp
@@ -2936,7 +2936,11 @@ void UControlRigBlueprint::PopulateModelFromGraphForBackwardsCompatibility(UCont
 				const int32 VarIndex = FBlueprintEditorUtils::FindNewVariableIndex(this, PropertyName);
 				if (VarIndex != INDEX_NONE)
 				{
-					NewVariables.RemoveAt(VarIndex);
+					// only remove the variable if it was backing up a rig unit
+					if(UnitStruct != nullptr)
+					{
+						NewVariables.RemoveAt(VarIndex);
+					}
 					FBlueprintEditorUtils::RemoveVariableNodes(this, PropertyName);
 				}
 			}
