Control Rig: Fix filtering of available nodes in the context menu when dragging a pin

diff --git a/Engine/Plugins/Animation/ControlRig/Source/ControlRigEditor/Private/Graph/NodeSpawners/ControlRigTemplateNodeSpawner.cpp b/Engine/Plugins/Animation/ControlRig/Source/ControlRigEditor/Private/Graph/NodeSpawners/ControlRigTemplateNodeSpawner.cpp
--- a/Engine/Plugins/Animation/ControlRig/Source/ControlRigEditor/Private/Graph/NodeSpawners/ControlRigTemplateNodeSpawner.cpp
+++ b/Engine/Plugins/Animation/ControlRig/Source/ControlRigEditor/Private/Graph/NodeSpawners/ControlRigTemplateNodeSpawner.cpp
@@ -71,6 +71,57 @@ FBlueprintNodeSignature UControlRigTemplateNodeSpawner::GetSpawnerSignature() co
 	return FBlueprintNodeSignature(FString("RigVMTemplate_") + FString::FromInt(NotationHash));
 }
 
+bool UControlRigTemplateNodeSpawner::IsTemplateNodeFilteredOut(FBlueprintActionFilter const& Filter) const
+{
+	bool bFilteredOut = true;
+	for (UEdGraphPin* Pin : Filter.Context.Pins)
+	{
+		for (UEdGraph* Graph : Filter.Context.Graphs)
+		{
+			if (!Graph->Nodes.Contains(Pin->GetOwningNode()))
+			{
+				continue;
+			}
+
+			if (UControlRigGraph* ControlRigGraph = Cast<UControlRigGraph>(Graph))
+			{
+				if (URigVMGraph* RigGraph = ControlRigGraph->GetModel())
+				{
+					if (URigVMPin* ModelPin = RigGraph->FindPin(Pin->GetName()))
+					{
+						if (URigVMTemplateNode* TemplateNode = Cast<URigVMTemplateNode>(ModelPin->GetNode()))
+						{
+							// Only filter when the sourece pin is not a wildcard
+							if (ModelPin->GetCPPType() != RigVMTypeUtils::GetWildCardCPPType())
+							{
+								if (const FRigVMTemplate* Template = FRigVMRegistry::Get().FindTemplate(TemplateNotation))
+								{
+									FRigVMTemplateArgument::FType Type(ModelPin->GetCPPType(), ModelPin->GetCPPTypeObject());
+									for (int32 i=0; i<Template->NumArguments(); ++i)
+									{
+										const FRigVMTemplateArgument* Argument = Template->GetArgument(i);
+										if (Template->ArgumentSupportsType(Argument->GetName(), Type.CPPType))
+										{
+											return false;
+										}
+									}
+									return true;
+								}									
+							}
+							else
+							{
+								// TODO: given the filtered types of the source pin, provide only the nodes that will not break any connections
+							}
+						}
+					}					
+				}
+			}
+		}
+	}
+	
+	return Super::IsTemplateNodeFilteredOut(Filter);
+}
+
 FBlueprintActionUiSpec UControlRigTemplateNodeSpawner::GetUiSpec(FBlueprintActionContext const& Context, FBindingSet const& Bindings) const
 {
 	UEdGraph* TargetGraph = (Context.Graphs.Num() > 0) ? Context.Graphs[0] : nullptr;
diff --git a/Engine/Plugins/Animation/ControlRig/Source/ControlRigEditor/Private/Graph/NodeSpawners/ControlRigTemplateNodeSpawner.h b/Engine/Plugins/Animation/ControlRig/Source/ControlRigEditor/Private/Graph/NodeSpawners/ControlRigTemplateNodeSpawner.h
--- a/Engine/Plugins/Animation/ControlRig/Source/ControlRigEditor/Private/Graph/NodeSpawners/ControlRigTemplateNodeSpawner.h
+++ b/Engine/Plugins/Animation/ControlRig/Source/ControlRigEditor/Private/Graph/NodeSpawners/ControlRigTemplateNodeSpawner.h
@@ -37,6 +37,7 @@ public:
 	// UBlueprintNodeSpawner interface
 	virtual void Prime() override;
 	virtual FBlueprintNodeSignature GetSpawnerSignature() const override;
+	bool IsTemplateNodeFilteredOut(FBlueprintActionFilter const& Filter) const override;
 	virtual FBlueprintActionUiSpec GetUiSpec(FBlueprintActionContext const& Context, FBindingSet const& Bindings) const override;
 	virtual UEdGraphNode* Invoke(UEdGraph* ParentGraph, FBindingSet const& Bindings, FVector2D const Location) const override;
 	// End UBlueprintNodeSpawner interface
