Control Rig Animation Outliner - Catch nullptr

diff --git a/Engine/Plugins/Animation/ControlRig/Source/ControlRigEditor/Private/EditMode/SControlRigOutliner.cpp b/Engine/Plugins/Animation/ControlRig/Source/ControlRigEditor/Private/EditMode/SControlRigOutliner.cpp
--- a/Engine/Plugins/Animation/ControlRig/Source/ControlRigEditor/Private/EditMode/SControlRigOutliner.cpp
+++ b/Engine/Plugins/Animation/ControlRig/Source/ControlRigEditor/Private/EditMode/SControlRigOutliner.cpp
@@ -1014,10 +1014,13 @@ void SControlRigOutliner::OnObjectsReplaced(const TMap<UObject*, UObject*>& OldT
 	bool bNewControlRig = false;
 	for (const TPair<UObject*, UObject*>& Pair : OldToNewInstanceMap)
 	{
-		if (Pair.Key->IsA<UControlRig>() && Pair.Value->IsA<UControlRig>())
+		if(Pair.Key && Pair.Value)
 		{
-			bNewControlRig = false;
-			break;
+			if (Pair.Key->IsA<UControlRig>() && Pair.Value->IsA<UControlRig>())
+			{
+				bNewControlRig = false;
+				break;
+			}
 		}
 	}
 	if (bNewControlRig)
