Sequencer: Control Rig: Curve Editor selection no longer selects Controls based on animators feedback. This was added way back when we had limited ways to select controls, but with then new details panel, pickers etc.. it is no longer needed.

diff --git a/Engine/Plugins/Animation/ControlRig/Source/ControlRigEditor/Private/Sequencer/ControlRigParameterTrackEditor.cpp b/Engine/Plugins/Animation/ControlRig/Source/ControlRigEditor/Private/Sequencer/ControlRigParameterTrackEditor.cpp
--- a/Engine/Plugins/Animation/ControlRig/Source/ControlRigEditor/Private/Sequencer/ControlRigParameterTrackEditor.cpp
+++ b/Engine/Plugins/Animation/ControlRig/Source/ControlRigEditor/Private/Sequencer/ControlRigParameterTrackEditor.cpp
@@ -48,8 +48,6 @@
 #include "LevelEditorViewport.h"
 #include "IKeyArea.h"
 #include "ISequencer.h"
-#include "CurveModel.h"
-#include "CurveEditor.h"
 #include "ControlRigEditorModule.h"
 #include "SequencerSettings.h"
 #include "Framework/Application/SlateApplication.h"
@@ -214,7 +212,6 @@ FControlRigParameterTrackEditor::FControlRigParameterTrackEditor(TSharedRef<ISeq
 	SelectionChangedHandle = InSequencer->GetSelectionChangedTracks().AddRaw(this, &FControlRigParameterTrackEditor::OnSelectionChanged);
 	SequencerChangedHandle = InSequencer->OnMovieSceneDataChanged().AddRaw(this, &FControlRigParameterTrackEditor::OnSequencerDataChanged);
 	OnActivateSequenceChangedHandle = InSequencer->OnActivateSequence().AddRaw(this, &FControlRigParameterTrackEditor::OnActivateSequenceChanged);
-	CurveChangedHandle = InSequencer->GetCurveDisplayChanged().AddRaw(this, &FControlRigParameterTrackEditor::OnCurveDisplayChanged);
 	OnChannelChangedHandle = InSequencer->OnChannelChanged().AddRaw(this, &FControlRigParameterTrackEditor::OnChannelChanged);
 	OnMovieSceneChannelChangedHandle = MovieScene->OnChannelChanged().AddRaw(this, &FControlRigParameterTrackEditor::OnChannelChanged);
 	OnActorAddedToSequencerHandle = InSequencer->OnActorAddedToSequencer().AddRaw(this, &FControlRigParameterTrackEditor::HandleActorAdded);
@@ -1809,115 +1806,6 @@ void FControlRigParameterTrackEditor::OnSequencerDataChanged(EMovieSceneDataChan
 	}
 }
 
-void FControlRigParameterTrackEditor::OnCurveDisplayChanged(FCurveModel* CurveModel, bool bDisplayed, const FCurveEditor* InCurveEditor)
-{
-	//if already doing a selection or the curve editor isn't doing a direct selection, for example sequencer filtering removed the curve, we dont' update control selection
-	if (bIsDoingSelection || (InCurveEditor && InCurveEditor->IsDoingDirectSelection() == false))
-	{
-		return;
-	}
-	
-	TGuardValue<bool> Guard(bIsDoingSelection, true);
-	FScopedTransaction ScopedTransaction(LOCTEXT("SelectControlTransaction", "Select Control"), !GIsTransacting);
-
-	TArray<FString> StringArray;
-	FControlRigEditMode* ControlRigEditMode = GetEditMode();
-	UControlRig* ControlRig = nullptr;
-
-	if (CurveModel)
-	{
-		UMovieSceneControlRigParameterSection* MovieSection = Cast<UMovieSceneControlRigParameterSection>(CurveModel->GetOwningObject());
-		if (MovieSection)
-		{
-			ControlRig = MovieSection->GetControlRig();
-			//Only create the edit mode if we have a  curve selected and it's not set and we have some boundobjects.
-			if (!ControlRigEditMode)
-			{
-				ControlRigEditMode = GetEditMode(true);
-				if (TSharedPtr<IControlRigObjectBinding> ObjectBinding = ControlRig->GetObjectBinding())
-				{
-					if (ControlRigEditMode)
-					{
-						ControlRigEditMode->AddControlRigObject(ControlRig, GetSequencer());
-					}
-				}
-			}
-			else
-			{
-				TArray<UControlRig*> ControlRigs = ControlRigEditMode->GetControlRigsArray(false /*bIsVisible*/);
-				if (ControlRigs.Contains(ControlRig)== false)
-				{
-					ControlRigEditMode->AddControlRigObject(ControlRig, GetSequencer());
-				}
-			}
-			//Not 100% safe but for now it is since that's all we show in the curve editor
-			//We need the Float Curve Model so we can get the ChannelHandle so we can also select the keyarea in the sequencer window if needed.
-			FFloatChannelCurveModel* FCurveModel = static_cast<FFloatChannelCurveModel*>(CurveModel);
-			FString String = CurveModel->GetLongDisplayName().ToString();
-			StringArray.SetNum(0);
-			String.ParseIntoArray(StringArray, TEXT("."));
-			if (StringArray.Num() > 2)
-			{
-				//Not great but it should always be the third name
-				FName ControlName(*StringArray[2]);
-				ControlRig->SelectControl(ControlName, bDisplayed); //mz need to check this after merge
-				if (bDisplayed)
-				{
-					DisplayedControls.Add(ControlName);
-				}
-				else
-				{
-					UnDisplayedControls.Add(ControlName);
-				}
-			}
-			else
-			{
-				UE_LOG(LogControlRigEditor, Display, TEXT("Could not find Rig Control From FCurveModel::LongName"));
-			}
-
-			if (bCurveDisplayTickIsPending == false)
-			{
-				bCurveDisplayTickIsPending = true;
-				GEditor->GetTimerManager()->SetTimerForNextTick([MovieSection, bDisplayed, this]()
-					{
-
-						if (DisplayedControls.Num() > 0 || UnDisplayedControls.Num() > 0)
-						{
-							TGuardValue<bool> Guard(bIsDoingSelection, true);
-							UMovieSceneControlRigParameterSection* ParamSection = Cast<UMovieSceneControlRigParameterSection>(MovieSection);
-							bool bSync = GetSequencer()->GetSequencerSettings()->ShouldSyncCurveEditorSelection();
-							GetSequencer()->SuspendSelectionBroadcast();
-							GetSequencer()->GetSequencerSettings()->SyncCurveEditorSelection(false);
-							if (UnDisplayedControls.Num() > 0)
-							{
-								for (const FName& ControlName : UnDisplayedControls)
-								{
-									SelectSequencerNodeInSection(ParamSection, ControlName, false);
-								}
-								UnDisplayedControls.Empty();
-							}
-							if (DisplayedControls.Num() > 0)
-							{
-								for (const FName& ControlName : DisplayedControls)
-								{
-									SelectSequencerNodeInSection(ParamSection, ControlName, true);
-								}
-								DisplayedControls.Empty();
-							}
-							GetSequencer()->ResumeSelectionBroadcast(); //need to resume first so when we refreh the tree we do the Selection.Tick, which since syncing is off won't 
-																		//mess up the curve editor.
-							GetSequencer()->RefreshTree();
-							GetSequencer()->GetSequencerSettings()->SyncCurveEditorSelection(bSync);
-						};
-						bCurveDisplayTickIsPending = false;
-
-					});
-
-			}
-		}
-
-	}
-}
 
 void FControlRigParameterTrackEditor::PostEvaluation(UMovieScene* MovieScene, FFrameNumber Frame)
 {
diff --git a/Engine/Plugins/Animation/ControlRig/Source/ControlRigEditor/Private/Sequencer/ControlRigParameterTrackEditor.h b/Engine/Plugins/Animation/ControlRig/Source/ControlRigEditor/Private/Sequencer/ControlRigParameterTrackEditor.h
--- a/Engine/Plugins/Animation/ControlRig/Source/ControlRigEditor/Private/Sequencer/ControlRigParameterTrackEditor.h
+++ b/Engine/Plugins/Animation/ControlRig/Source/ControlRigEditor/Private/Sequencer/ControlRigParameterTrackEditor.h
@@ -113,9 +113,6 @@ private:
 	/** Delegate for when track immediately changes, so we need to do an interaction edit for things like fk/ik*/
 	void OnChannelChanged(const FMovieSceneChannelMetaData* MetaData, UMovieSceneSection* InSection);
 
-	/** Delegate for Curve Selection Changed Event */
-	void OnCurveDisplayChanged(FCurveModel* InCurveModel, bool bDisplayed, const FCurveEditor* InCurveEditor);
-
 	/** Delegate for difference focused movie scene sequence*/
 	void OnActivateSequenceChanged(FMovieSceneSequenceIDRef ID);
 
