Control Rig:  Fix animation channel parenting + cross selection

diff --git a/Engine/Plugins/Animation/ControlRig/Source/ControlRig/Private/Rigs/RigHierarchyController.cpp b/Engine/Plugins/Animation/ControlRig/Source/ControlRig/Private/Rigs/RigHierarchyController.cpp
--- a/Engine/Plugins/Animation/ControlRig/Source/ControlRig/Private/Rigs/RigHierarchyController.cpp
+++ b/Engine/Plugins/Animation/ControlRig/Source/ControlRig/Private/Rigs/RigHierarchyController.cpp
@@ -1267,7 +1267,6 @@ TArray<FRigElementKey> URigHierarchyController::ImportFromHierarchyContainer(con
 		Settings.PrimaryAxis = Control.PrimaryAxis;
 		Settings.bIsCurve = Control.bIsCurve;
 		Settings.SetAnimationTypeFromDeprecatedData(Control.bAnimatable, Control.bGizmoEnabled);
-		Settings.AnimationType = Control.bAnimatable ? ERigControlAnimationType::AnimationControl : ERigControlAnimationType::ProxyControl;
 		Settings.SetupLimitArrayForType(Control.bLimitTranslation, Control.bLimitRotation, Control.bLimitScale);
 		Settings.bDrawLimits = Control.bDrawLimits;
 		Settings.MinimumValue = Control.MinimumValue;
@@ -1276,6 +1275,12 @@ TArray<FRigElementKey> URigHierarchyController::ImportFromHierarchyContainer(con
 		Settings.ShapeName = Control.GizmoName;
 		Settings.ShapeColor = Control.GizmoColor;
 		Settings.ControlEnum = Control.ControlEnum;
+		Settings.bGroupWithParentControl = Settings.IsAnimatable() && (
+			Settings.ControlType == ERigControlType::Bool ||
+			Settings.ControlType == ERigControlType::Float ||
+			Settings.ControlType == ERigControlType::Integer ||
+			Settings.ControlType == ERigControlType::Vector2D
+		);
 
 		if(Settings.ShapeName == FRigControl().GizmoName)
 		{
diff --git a/Engine/Plugins/Animation/ControlRig/Source/ControlRig/Private/Sequencer/MovieSceneControlRigParameterSection.cpp b/Engine/Plugins/Animation/ControlRig/Source/ControlRig/Private/Sequencer/MovieSceneControlRigParameterSection.cpp
--- a/Engine/Plugins/Animation/ControlRig/Source/ControlRig/Private/Sequencer/MovieSceneControlRigParameterSection.cpp
+++ b/Engine/Plugins/Animation/ControlRig/Source/ControlRig/Private/Sequencer/MovieSceneControlRigParameterSection.cpp
@@ -1384,7 +1384,6 @@ void UMovieSceneControlRigParameterSection::ReconstructChannelProxy()
 							{
 								ControlChannelMap.Add(Scalar.ParameterName, FChannelMapInfo(ControlIndex, TotalIndex, FloatChannelIndex, INDEX_NONE, NAME_None, MaskIndex, CategoryIndex));
 								Group = FText::FromName(ControlElement->GetDisplayName());
-								ControlIndex++;  //up the index only if no parent
 								if (bEnabled)
 								{
 									++CategoryIndex;
@@ -1393,8 +1392,8 @@ void UMovieSceneControlRigParameterSection::ReconstructChannelProxy()
 							else
 							{
 								const FChannelMapInfo* pChannelIndex = ControlChannelMap.Find(ParentControlName);
-								const int32 ChannelControlIndex = pChannelIndex ? pChannelIndex->ControlIndex : INDEX_NONE; 
-								ControlChannelMap.Add(Scalar.ParameterName, FChannelMapInfo(ControlIndex, TotalIndex, FloatChannelIndex,ChannelControlIndex, NAME_None, MaskIndex, CategoryIndex));
+								const int32 ParentControlIndex = pChannelIndex ? pChannelIndex->ControlIndex : INDEX_NONE; 
+								ControlChannelMap.Add(Scalar.ParameterName, FChannelMapInfo(ControlIndex, TotalIndex, FloatChannelIndex,ParentControlIndex, NAME_None, MaskIndex, CategoryIndex));
 							}
 
 							FParameterFloatChannelEditorData EditorData(ControlRig, Scalar.ParameterName, bEnabled, Group, TotalIndex);
@@ -1402,6 +1401,7 @@ void UMovieSceneControlRigParameterSection::ReconstructChannelProxy()
 							Channels.Add(Scalar.ParameterCurve, EditorData.MetaData, EditorData.ExternalValues);
 							FloatChannelIndex += 1;
 							TotalIndex += 1;
+							ControlIndex += 1;
 							break;
 						}
 					}
@@ -1417,7 +1417,6 @@ void UMovieSceneControlRigParameterSection::ReconstructChannelProxy()
 							{
 								ControlChannelMap.Add(Bool.ParameterName, FChannelMapInfo(ControlIndex, TotalIndex, BoolChannelIndex, INDEX_NONE,BoolChannelTypeName,MaskIndex, CategoryIndex));
 								Group = FText::FromName(ControlElement->GetDisplayName());
-								ControlIndex++;  //up the index only if no parent
 								if (bEnabled)
 								{
 									++CategoryIndex;
@@ -1426,14 +1425,15 @@ void UMovieSceneControlRigParameterSection::ReconstructChannelProxy()
 							else
 							{
 								const FChannelMapInfo* pChannelIndex = ControlChannelMap.Find(ParentControlName);
-								const int32 ChannelControlIndex = pChannelIndex ? pChannelIndex->ControlIndex : INDEX_NONE; 
-								ControlChannelMap.Add(Bool.ParameterName, FChannelMapInfo(ControlIndex, TotalIndex, BoolChannelIndex,ChannelControlIndex, BoolChannelTypeName,MaskIndex, CategoryIndex));
+								const int32 ParentControlIndex = pChannelIndex ? pChannelIndex->ControlIndex : INDEX_NONE; 
+								ControlChannelMap.Add(Bool.ParameterName, FChannelMapInfo(ControlIndex, TotalIndex, BoolChannelIndex,ParentControlIndex, BoolChannelTypeName,MaskIndex, CategoryIndex));
 							}
 
 							FMovieSceneChannelMetaData MetaData(Bool.ParameterName, Group, Group, bEnabled);
 							MetaData.DisplayText = FText::FromName(ControlElement->GetDisplayName());
 							MetaData.SortOrder = TotalIndex++;
 							BoolChannelIndex += 1;
+							ControlIndex += 1;
 							// Prevent single channels from collapsing to the track node
 							MetaData.bCanCollapseToTrack = false;
 							Channels.Add(Bool.ParameterCurve, MetaData, TMovieSceneExternalValue<bool>());
@@ -1454,7 +1454,6 @@ void UMovieSceneControlRigParameterSection::ReconstructChannelProxy()
 								{
 									ControlChannelMap.Add(Enum.ParameterName, FChannelMapInfo(ControlIndex, TotalIndex, EnumChannelIndex,INDEX_NONE, EnumChannelTypeName,MaskIndex, CategoryIndex));
 									Group = FText::FromName(ControlElement->GetDisplayName());
-									ControlIndex++;  //up the index only if no parent
 									if (bEnabled)
 									{
 										++CategoryIndex;
@@ -1463,13 +1462,14 @@ void UMovieSceneControlRigParameterSection::ReconstructChannelProxy()
 								else
 								{
 									const FChannelMapInfo* pChannelIndex = ControlChannelMap.Find(ParentControlName);
-									const int32 ChannelControlIndex = pChannelIndex ? pChannelIndex->ControlIndex : INDEX_NONE; 
-									ControlChannelMap.Add(Enum.ParameterName, FChannelMapInfo(ControlIndex, TotalIndex, EnumChannelIndex, ChannelControlIndex, EnumChannelTypeName,MaskIndex, CategoryIndex));
+									const int32 ParentControlIndex = pChannelIndex ? pChannelIndex->ControlIndex : INDEX_NONE; 
+									ControlChannelMap.Add(Enum.ParameterName, FChannelMapInfo(ControlIndex, TotalIndex, EnumChannelIndex, ParentControlIndex, EnumChannelTypeName,MaskIndex, CategoryIndex));
 								}
 
 								FMovieSceneChannelMetaData MetaData(Enum.ParameterName, Group, Group, bEnabled);
 								MetaData.DisplayText = FText::FromName(ControlElement->GetDisplayName());
 								EnumChannelIndex += 1;
+								ControlIndex += 1;
 								MetaData.SortOrder = TotalIndex++;
 								// Prevent single channels from collapsing to the track node
 								MetaData.bCanCollapseToTrack = false;
@@ -1488,7 +1488,6 @@ void UMovieSceneControlRigParameterSection::ReconstructChannelProxy()
 								{
 									ControlChannelMap.Add(Integer.ParameterName, FChannelMapInfo(ControlIndex, TotalIndex, IntegerChannelIndex,INDEX_NONE,IntegerChannelTypeName,MaskIndex, CategoryIndex));
 									Group = FText::FromName(ControlElement->GetDisplayName());
-									ControlIndex++;  //up the index only if no parent
 									if (bEnabled)
 									{
 										++CategoryIndex;
@@ -1497,13 +1496,14 @@ void UMovieSceneControlRigParameterSection::ReconstructChannelProxy()
 								else
 								{
 									const FChannelMapInfo* pChannelIndex = ControlChannelMap.Find(ParentControlName);
-									const int32 ChannelControlIndex = pChannelIndex ? pChannelIndex->ControlIndex : INDEX_NONE; 
-									ControlChannelMap.Add(Integer.ParameterName, FChannelMapInfo(ControlIndex, TotalIndex, IntegerChannelIndex, ChannelControlIndex, IntegerChannelTypeName,MaskIndex, CategoryIndex));
+									const int32 ParentControlIndex = pChannelIndex ? pChannelIndex->ControlIndex : INDEX_NONE; 
+									ControlChannelMap.Add(Integer.ParameterName, FChannelMapInfo(ControlIndex, TotalIndex, IntegerChannelIndex, ParentControlIndex, IntegerChannelTypeName,MaskIndex, CategoryIndex));
 								}
 
 								FMovieSceneChannelMetaData MetaData(Integer.ParameterName, Group, Group, bEnabled);
 								MetaData.DisplayText = FText::FromName(ControlElement->GetDisplayName());
 								IntegerChannelIndex += 1;
+								ControlIndex += 1;
 								MetaData.SortOrder = TotalIndex++;
 								// Prevent single channels from collapsing to the track node
 								MetaData.bCanCollapseToTrack = false;
@@ -1524,7 +1524,6 @@ void UMovieSceneControlRigParameterSection::ReconstructChannelProxy()
 							if(Group.IsEmpty())
 							{
 								ControlChannelMap.Add(Vector2D.ParameterName, FChannelMapInfo(ControlIndex, TotalIndex, FloatChannelIndex, INDEX_NONE, NAME_None, MaskIndex, CategoryIndex));
-								ControlIndex++;
 								if (bEnabled)
 								{
 									++CategoryIndex;
@@ -1534,14 +1533,15 @@ void UMovieSceneControlRigParameterSection::ReconstructChannelProxy()
 							else
 							{
 								const FChannelMapInfo* pChannelIndex = ControlChannelMap.Find(ParentControlName);
-								const int32 ChannelControlIndex = pChannelIndex ? pChannelIndex->ControlIndex : INDEX_NONE; 
-								ControlChannelMap.Add(Vector2D.ParameterName, FChannelMapInfo(ControlIndex, TotalIndex, FloatChannelIndex, ChannelControlIndex, NAME_None, MaskIndex, CategoryIndex));
+								const int32 ParentControlIndex = pChannelIndex ? pChannelIndex->ControlIndex : INDEX_NONE; 
+								ControlChannelMap.Add(Vector2D.ParameterName, FChannelMapInfo(ControlIndex, TotalIndex, FloatChannelIndex, ParentControlIndex, NAME_None, MaskIndex, CategoryIndex));
 							}
 							FParameterVectorChannelEditorData EditorData(ControlRig, Vector2D.ParameterName, bEnabled, Group, TotalIndex, 2);
 							Channels.Add(Vector2D.XCurve, EditorData.MetaData[0], EditorData.ExternalValues[0]);
 							Channels.Add(Vector2D.YCurve, EditorData.MetaData[1], EditorData.ExternalValues[1]);
 							FloatChannelIndex += 2;
 							TotalIndex += 2;
+							ControlIndex += 1;
 							break;
 						}
 					}
@@ -1558,7 +1558,6 @@ void UMovieSceneControlRigParameterSection::ReconstructChannelProxy()
 							if(Group.IsEmpty())
 							{
 								ControlChannelMap.Add(Vector.ParameterName, FChannelMapInfo(ControlIndex, TotalIndex, FloatChannelIndex, INDEX_NONE, NAME_None, MaskIndex, CategoryIndex));
-								ControlIndex++;
 								if (bEnabled)
 								{
 									++CategoryIndex;
@@ -1568,8 +1567,8 @@ void UMovieSceneControlRigParameterSection::ReconstructChannelProxy()
 							else
 							{
 								const FChannelMapInfo* pChannelIndex = ControlChannelMap.Find(ParentControlName);
-								const int32 ChannelControlIndex = pChannelIndex ? pChannelIndex->ControlIndex : INDEX_NONE; 
-								ControlChannelMap.Add(Vector.ParameterName, FChannelMapInfo(ControlIndex, TotalIndex, FloatChannelIndex, ChannelControlIndex, NAME_None, MaskIndex, CategoryIndex));
+								const int32 ParentControlIndex = pChannelIndex ? pChannelIndex->ControlIndex : INDEX_NONE; 
+								ControlChannelMap.Add(Vector.ParameterName, FChannelMapInfo(ControlIndex, TotalIndex, FloatChannelIndex, ParentControlIndex, NAME_None, MaskIndex, CategoryIndex));
 							}
 							if (FSpaceControlNameAndChannel* SpaceChannel = GetSpaceChannel(Vector.ParameterName))
 							{
@@ -1600,6 +1599,7 @@ void UMovieSceneControlRigParameterSection::ReconstructChannelProxy()
 							Channels.Add(Vector.ZCurve, EditorData.MetaData[2], EditorData.ExternalValues[2]);
 							FloatChannelIndex += 3;
 							TotalIndex += 3;
+							ControlIndex += 1;
 							break;
 						}
 					}
@@ -1617,7 +1617,6 @@ void UMovieSceneControlRigParameterSection::ReconstructChannelProxy()
 							if(Group.IsEmpty())
 							{
 								ControlChannelMap.Add(Transform.ParameterName, FChannelMapInfo(ControlIndex, TotalIndex, FloatChannelIndex, INDEX_NONE, NAME_None, MaskIndex, CategoryIndex));
-								ControlIndex++;
 								if (bEnabled)
 								{
 									++CategoryIndex;
@@ -1627,8 +1626,8 @@ void UMovieSceneControlRigParameterSection::ReconstructChannelProxy()
 							else
 							{
 								const FChannelMapInfo* pChannelIndex = ControlChannelMap.Find(ParentControlName);
-								const int32 ChannelControlIndex = pChannelIndex ? pChannelIndex->ControlIndex : INDEX_NONE; 
-								ControlChannelMap.Add(Transform.ParameterName, FChannelMapInfo(ControlIndex, TotalIndex, FloatChannelIndex, ChannelControlIndex, NAME_None, MaskIndex, CategoryIndex));
+								const int32 ParentControlIndex = pChannelIndex ? pChannelIndex->ControlIndex : INDEX_NONE; 
+								ControlChannelMap.Add(Transform.ParameterName, FChannelMapInfo(ControlIndex, TotalIndex, FloatChannelIndex, ParentControlIndex, NAME_None, MaskIndex, CategoryIndex));
 							}
 
 							if (FSpaceControlNameAndChannel* SpaceChannel = GetSpaceChannel(Transform.ParameterName))
@@ -1683,6 +1682,7 @@ void UMovieSceneControlRigParameterSection::ReconstructChannelProxy()
 								TotalIndex += 6;
 
 							}
+							ControlIndex += 1;
 							break;
 						}
 					}
@@ -1703,7 +1703,7 @@ void UMovieSceneControlRigParameterSection::ReconstructChannelProxy()
 							Channels.Add(Scalar.ParameterCurve);
 							FloatChannelIndex += 1;
 							TotalIndex += 1;
-							ControlIndex++;
+							ControlIndex += 1;
 							break;
 						}
 					}
@@ -1719,7 +1719,7 @@ void UMovieSceneControlRigParameterSection::ReconstructChannelProxy()
 							Channels.Add(Bool.ParameterCurve);
 							BoolChannelIndex += 1;
 							TotalIndex += 1;
-							ControlIndex++;
+							ControlIndex += 1;
 							break;
 						}
 					}
@@ -1737,7 +1737,7 @@ void UMovieSceneControlRigParameterSection::ReconstructChannelProxy()
 								Channels.Add(Enum.ParameterCurve);
 								EnumChannelIndex += 1;
 								TotalIndex += 1;
-								ControlIndex++;
+								ControlIndex += 1;
 								break;
 							}
 						}
@@ -1752,7 +1752,7 @@ void UMovieSceneControlRigParameterSection::ReconstructChannelProxy()
 								Channels.Add(Integer.ParameterCurve);
 								IntegerChannelIndex += 1;
 								TotalIndex += 1;
-								ControlIndex++;
+								ControlIndex += 1;
 								break;
 							}
 						}
@@ -1770,7 +1770,7 @@ void UMovieSceneControlRigParameterSection::ReconstructChannelProxy()
 							Channels.Add(Vector2D.YCurve);
 							FloatChannelIndex += 2;
 							TotalIndex += 2;
-							ControlIndex++;
+							ControlIndex += 1;
 							break;
 						}
 					}
@@ -1806,7 +1806,7 @@ void UMovieSceneControlRigParameterSection::ReconstructChannelProxy()
 							Channels.Add(Vector.ZCurve);
 							FloatChannelIndex += 3;
 							TotalIndex += 3;
-							ControlIndex++;
+							ControlIndex += 1;
 							break;
 						}
 					}
@@ -1871,7 +1871,7 @@ void UMovieSceneControlRigParameterSection::ReconstructChannelProxy()
 								TotalIndex += 6;
 							}
 
-							ControlIndex++;
+							ControlIndex += 1;
 							break;
 						}
 					}
diff --git a/Engine/Plugins/Animation/ControlRig/Source/ControlRigEditor/Private/Sequencer/ControlRigParameterTrackEditor.cpp b/Engine/Plugins/Animation/ControlRig/Source/ControlRigEditor/Private/Sequencer/ControlRigParameterTrackEditor.cpp
--- a/Engine/Plugins/Animation/ControlRig/Source/ControlRigEditor/Private/Sequencer/ControlRigParameterTrackEditor.cpp
+++ b/Engine/Plugins/Animation/ControlRig/Source/ControlRigEditor/Private/Sequencer/ControlRigParameterTrackEditor.cpp
@@ -2366,7 +2366,7 @@ void FControlRigParameterTrackEditor::HandleControlSelected(UControlRig* Subject
 			SetUpEditModeIfNeeded(Subject);
 
 			//Force refresh later, not now
-			bSkipNextSelectionFromTimer = true;
+			bSkipNextSelectionFromTimer = bIsSelectingIndirectControl && ControlElement->Settings.AnimationType == ERigControlAnimationType::AnimationControl;
 			GetSequencer()->NotifyMovieSceneDataChanged(EMovieSceneDataChangeType::RefreshTree);
 		}
 	}
diff --git a/Engine/Source/Editor/Sequencer/Private/Sequencer.cpp b/Engine/Source/Editor/Sequencer/Private/Sequencer.cpp
--- a/Engine/Source/Editor/Sequencer/Private/Sequencer.cpp
+++ b/Engine/Source/Editor/Sequencer/Private/Sequencer.cpp
@@ -7172,11 +7172,25 @@ void FSequencer::SelectByChannels(UMovieSceneSection* Section, TArrayView<const
 				{
 					if (bSelectParentInstead || bSelect == false)
 					{
-						Nodes.Add(KeyAreaNode->GetParent()->AsShared());
+						if(bSelect)
+						{
+							NodesToSelect.Add(KeyAreaNode->GetParent()->AsShared());
+						}
+						else
+						{
+							Nodes.Add(KeyAreaNode->GetParent()->AsShared());
+						}
 					}
 					if (!bSelectParentInstead || bSelect == false)
 					{
-						Nodes.Add(KeyAreaNode);
+						if(bSelect)
+						{
+							NodesToSelect.Add(KeyAreaNode);
+						}
+						else
+						{
+							Nodes.Add(KeyAreaNode);
+						}
 					}
 				}
 			}
