Control Rig: Fix static analysis issues with ControlRigParameterTrackEditor

diff --git a/Engine/Plugins/Animation/ControlRig/Source/ControlRigEditor/Private/Sequencer/ControlRigParameterTrackEditor.cpp b/Engine/Plugins/Animation/ControlRig/Source/ControlRigEditor/Private/Sequencer/ControlRigParameterTrackEditor.cpp
--- a/Engine/Plugins/Animation/ControlRig/Source/ControlRigEditor/Private/Sequencer/ControlRigParameterTrackEditor.cpp
+++ b/Engine/Plugins/Animation/ControlRig/Source/ControlRigEditor/Private/Sequencer/ControlRigParameterTrackEditor.cpp
@@ -2388,6 +2388,10 @@ void FControlRigParameterTrackEditor::HandleOnInitialized(UControlRig* ControlRi
 
 void FControlRigParameterTrackEditor::HandleControlModified(UControlRig* ControlRig, FRigControlElement* ControlElement, const FRigControlModifiedContext& Context)
 {
+	if(ControlElement == nullptr)
+	{
+		return;
+	}
 	if (!GetSequencer().IsValid() || !GetSequencer()->IsAllowedToChange() || Context.SetKey == EControlRigSetKey::Never)
 	{
 		return;
@@ -2437,6 +2441,7 @@ void FControlRigParameterTrackEditor::GetControlRigKeys(UControlRig* InControlRi
 	for (int32 ControlIndex = 0; ControlIndex < Controls.Num(); ++ControlIndex)
 	{
 		FRigControlElement* ControlElement = Controls[ControlIndex];
+		check(ControlElement);
 
 		if (!InControlRig->GetHierarchy()->IsAnimatable(ControlElement))
 		{
