Control Rig: Interaction Event improvements

Enable debug drawing during interaction
Reset proxy controls - but don't interact on them during Ctrl-Shift-G

diff --git a/Engine/Plugins/Animation/ControlRig/Source/ControlRig/Private/ControlRig.cpp b/Engine/Plugins/Animation/ControlRig/Source/ControlRig/Private/ControlRig.cpp
--- a/Engine/Plugins/Animation/ControlRig/Source/ControlRig/Private/ControlRig.cpp
+++ b/Engine/Plugins/Animation/ControlRig/Source/ControlRig/Private/ControlRig.cpp
@@ -701,7 +701,11 @@ void UControlRig::Execute(const EControlRigState InState, const FName& InEventNa
 	}
 
 	FRigUnitContext Context;
-	DrawInterface.Reset();
+
+	if(!GetEventQueue().Contains(InEventName) || (GetEventQueue()[0] == InEventName))
+	{
+		DrawInterface.Reset();
+	}
 	Context.DrawInterface = &DrawInterface;
 
 	Context.AnimAttributeContainer = ExternalAnimAttributeContainer;
@@ -1012,7 +1016,7 @@ void UControlRig::Execute(const EControlRigState InState, const FName& InEventNa
 		}
 	}
 
-	if (Context.DrawInterface && Context.DrawContainer)
+	if (Context.DrawInterface && Context.DrawContainer && GetEventQueue().Last() == InEventName)
 	{
 		Context.DrawInterface->Instructions.Append(Context.DrawContainer->Instructions);
 
@@ -1773,20 +1777,26 @@ FTransform UControlRig::GetControlLocalTransform(const FName& InControlName)
 
 const TArray<TSoftObjectPtr<UControlRigShapeLibrary>>& UControlRig::GetShapeLibraries() const
 {
-	if (UControlRig* CDO = Cast<UControlRig>(GetClass()->GetDefaultObject()))
+	const TArray<TSoftObjectPtr<UControlRigShapeLibrary>>* LibrariesPtr = &ShapeLibraries;
+
+	if(!GetClass()->IsNative())
 	{
-		for(TSoftObjectPtr<UControlRigShapeLibrary>& ShapeLibrary : CDO->ShapeLibraries)
+		if (UControlRig* CDO = Cast<UControlRig>(GetClass()->GetDefaultObject()))
 		{
-			if (!ShapeLibrary.IsValid())
-			{
-				ShapeLibrary.LoadSynchronous();
-			}
+			LibrariesPtr = &CDO->ShapeLibraries;
+		}
+	}
+
+	const TArray<TSoftObjectPtr<UControlRigShapeLibrary>>& Libraries = *LibrariesPtr;
+	for(const TSoftObjectPtr<UControlRigShapeLibrary>& ShapeLibrary : Libraries)
+	{
+		if (!ShapeLibrary.IsValid())
+		{
+			ShapeLibrary.LoadSynchronous();
 		}
-		return CDO->ShapeLibraries;
 	}
 
-	static TArray<TSoftObjectPtr<UControlRigShapeLibrary>> EmptyShapeLibraries;
-	return EmptyShapeLibraries;
+	return Libraries;
 }
 
 void UControlRig::SelectControl(const FName& InControlName, bool bSelect)
diff --git a/Engine/Plugins/Animation/ControlRig/Source/ControlRig/Private/ControlRigGizmoActor.cpp b/Engine/Plugins/Animation/ControlRig/Source/ControlRig/Private/ControlRigGizmoActor.cpp
--- a/Engine/Plugins/Animation/ControlRig/Source/ControlRig/Private/ControlRigGizmoActor.cpp
+++ b/Engine/Plugins/Animation/ControlRig/Source/ControlRig/Private/ControlRigGizmoActor.cpp
@@ -15,9 +15,7 @@ AControlRigShapeActor::AControlRigShapeActor(const FObjectInitializer& ObjectIni
 	, ControlName(NAME_None)
 	, ShapeName(NAME_None)
 	, OverrideColor(0, 0, 0, 0)
-	, bEnabled(true)
 	, bSelected(false)
-	, bSelectable(true)
 	, bHovered(false)
 {
 
@@ -35,28 +33,12 @@ AControlRigShapeActor::AControlRigShapeActor(const FObjectInitializer& ObjectIni
 	StaticMeshComponent->SetupAttachment(RootComponent);
 	StaticMeshComponent->bCastStaticShadow = false;
 	StaticMeshComponent->bCastDynamicShadow = false;
-	StaticMeshComponent->bSelectable = bSelectable && bEnabled;
-}
-
-void AControlRigShapeActor::SetEnabled(bool bInEnabled)
-{
-	if(bEnabled != bInEnabled)
-	{
-		bEnabled = bInEnabled;
-		StaticMeshComponent->bSelectable = bSelectable && bEnabled;
-		FEditorScriptExecutionGuard Guard;
-		OnEnabledChanged(bEnabled);
-	}
-}
-
-bool AControlRigShapeActor::IsEnabled() const
-{
-	return bEnabled;
+	StaticMeshComponent->bSelectable = true;
 }
 
 void AControlRigShapeActor::SetSelected(bool bInSelected)
 {
-	if(!bSelectable && bInSelected)
+	if(!IsSelectable() && bInSelected)
 	{
 		return;
 	}
@@ -73,16 +55,23 @@ bool AControlRigShapeActor::IsSelectedInEditor() const
 	return bSelected;
 }
 
+bool AControlRigShapeActor::IsSelectable() const
+{
+	return StaticMeshComponent->bSelectable;
+}
+
 void AControlRigShapeActor::SetSelectable(bool bInSelectable)
 {
-	if (bSelectable != bInSelectable)
+	if (StaticMeshComponent->bSelectable != bInSelectable)
 	{
-		bSelectable = bInSelectable;
-		StaticMeshComponent->bSelectable = bSelectable && bEnabled;
-		if (!bSelectable)
+		StaticMeshComponent->bSelectable = bInSelectable;
+		if (!StaticMeshComponent->bSelectable)
 		{
 			SetSelected(false);
 		}
+
+		FEditorScriptExecutionGuard Guard;
+		OnEnabledChanged(bInSelectable);
 	}
 }
 
@@ -182,17 +171,6 @@ bool AControlRigShapeActor::UpdateControlSettings(
 	// update the shape color
 	SetShapeColor(ControlSettings.ShapeColor);
 
-#if WITH_EDITOR
-	// update visibility
-	SetIsTemporarilyHiddenInEditor(!ControlSettings.IsVisible() || bHideManipulators);
-#endif
-	
-	// update selectability state
-	if (!bIsInLevelEditor) //don't change this in level editor otherwise we can never select it
-	{
-		SetSelectable(ControlSettings.IsSelectable() && !bHideManipulators);
-	}
-
 	return true;
 }
 
diff --git a/Engine/Plugins/Animation/ControlRig/Source/ControlRig/Public/ControlRig.h b/Engine/Plugins/Animation/ControlRig/Source/ControlRig/Public/ControlRig.h
--- a/Engine/Plugins/Animation/ControlRig/Source/ControlRig/Public/ControlRig.h
+++ b/Engine/Plugins/Animation/ControlRig/Source/ControlRig/Public/ControlRig.h
@@ -450,7 +450,7 @@ public:
 	void SwapVMToNativizedIfRequired(UClass* InNativizedClass = nullptr);
 	static bool AreNativizedVMsDisabled() ;
 
-private:
+protected:
 
 	void PostInitInstance(UControlRig* InCDO);
 
diff --git a/Engine/Plugins/Animation/ControlRig/Source/ControlRig/Public/ControlRigGizmoActor.h b/Engine/Plugins/Animation/ControlRig/Source/ControlRig/Public/ControlRigGizmoActor.h
--- a/Engine/Plugins/Animation/ControlRig/Source/ControlRig/Public/ControlRigGizmoActor.h
+++ b/Engine/Plugins/Animation/ControlRig/Source/ControlRig/Public/ControlRigGizmoActor.h
@@ -90,11 +90,11 @@ public:
 
 	UFUNCTION(BlueprintSetter)
 	/** Set the control to be enabled/disabled */
-	virtual void SetEnabled(bool bInEnabled);
+	virtual void SetEnabled(bool bInEnabled) { SetSelectable(bInEnabled); }
 
 	UFUNCTION(BlueprintGetter)
 	/** Get whether the control is enabled/disabled */
-	virtual bool IsEnabled() const;
+	virtual bool IsEnabled() const { return IsSelectable(); }
 
 	UFUNCTION(BlueprintSetter)
 	/** Set the control to be selected/unselected */
@@ -105,7 +105,7 @@ public:
 	virtual bool IsSelectedInEditor() const;
 
 	/** Get wether the control is selectable/unselectable */
-	virtual bool IsSelectable() const { return bSelectable; }
+	virtual bool IsSelectable() const;
 
 	UFUNCTION(BlueprintSetter)
 	/** Set the control to be selected/unselected */
@@ -165,18 +165,11 @@ public:
 		bool bIsInLevelEditor);
 	
 private:
-	/** Whether this control is enabled */
-	UPROPERTY(BlueprintGetter = IsEnabled, BlueprintSetter= SetEnabled, Category = "ControlRig|Shape")
-	uint8 bEnabled : 1;
 
 	/** Whether this control is selected */
 	UPROPERTY(BlueprintGetter = IsSelectedInEditor, BlueprintSetter = SetSelected, Category = "ControlRig|Shape")
 	uint8 bSelected : 1;
 
-	/** Whether this control can be selected */
-	UPROPERTY(BlueprintGetter = IsSelectable, BlueprintSetter = SetSelectable, Category = "ControlRig|Shape")
-	uint8 bSelectable : 1;
-
 	/** Whether this control is hovered */
 	UPROPERTY(BlueprintGetter = IsHovered, BlueprintSetter = SetHovered, Category = "ControlRig|Shape")
 	uint8 bHovered : 1;
diff --git a/Engine/Plugins/Animation/ControlRig/Source/ControlRig/Public/Rigs/RigHierarchy.h b/Engine/Plugins/Animation/ControlRig/Source/ControlRig/Public/Rigs/RigHierarchy.h
--- a/Engine/Plugins/Animation/ControlRig/Source/ControlRig/Public/Rigs/RigHierarchy.h
+++ b/Engine/Plugins/Animation/ControlRig/Source/ControlRig/Public/Rigs/RigHierarchy.h
@@ -285,6 +285,21 @@ public:
 		return INDEX_NONE;
 	}
 
+	/**
+	 * Returns the indices of an array of keys
+	 * @param InKeys The keys of the elements to retrieve the indices for
+	 * @return The indices of the elements or INDEX_NONE
+	 */
+	FORCEINLINE TArray<int32> GetIndices(const TArray<FRigElementKey>& InKeys) const
+	{
+		TArray<int32> Indices;
+		for(const FRigElementKey& Key : InKeys)
+		{
+			Indices.Add(GetIndex(Key));
+		}
+		return Indices;
+	}
+
 	/**
 	 * Returns the key of an element given its index
 	 * @param InElementIndex The index of the element to retrieve the key for
@@ -300,6 +315,22 @@ public:
 		return FRigElementKey();
 	}
 
+	/**
+	 * Returns the keys of an array of indices
+	 * @param InElementIndices The indices to retrieve the keys for
+	 * @return The keys of the elements given the indices
+	 */
+	UFUNCTION(BlueprintCallable, Category = URigHierarchy)
+	FORCEINLINE TArray<FRigElementKey> GetKeys(const TArray<int32> InElementIndices) const
+	{
+		TArray<FRigElementKey> Keys;
+		for(int32 Index : InElementIndices)
+		{
+			Keys.Add(GetKey(Index));
+		}
+		return Keys;
+	}
+
 	/**
 	 * Returns an element at a given index or nullptr.
 	 * @param InIndex The index of the element to retrieve.
diff --git a/Engine/Plugins/Animation/ControlRig/Source/ControlRigEditor/Private/EditMode/ControlRigEditMode.cpp b/Engine/Plugins/Animation/ControlRig/Source/ControlRigEditor/Private/EditMode/ControlRigEditMode.cpp
--- a/Engine/Plugins/Animation/ControlRig/Source/ControlRigEditor/Private/EditMode/ControlRigEditMode.cpp
+++ b/Engine/Plugins/Animation/ControlRig/Source/ControlRigEditor/Private/EditMode/ControlRigEditMode.cpp
@@ -706,6 +706,7 @@ void FControlRigEditMode::Render(const FSceneView* View, FViewport* Viewport, FP
 			{
 				ComponentTransform = GetHostingSceneComponentTransform(ControlRig);
 			}
+			/*
 			const TArray<AControlRigShapeActor*>* ShapeActors = ControlRigShapeActors.Find(ControlRig);
 			if (ShapeActors)
 			{
@@ -717,6 +718,8 @@ void FControlRigEditMode::Render(const FSceneView* View, FViewport* Viewport, FP
 					}
 				}
 			}
+			*/
+			
 			//only draw stuff if not in game view
 			if (!bIsInGameView)
 			{
@@ -732,6 +735,14 @@ void FControlRigEditMode::Render(const FSceneView* View, FViewport* Viewport, FP
 					}
 					Hierarchy->ForEach<FRigTransformElement>([PDI, Hierarchy, ComponentTransform, ControlRig, bHasFKRig, bBoolSetHitProxies, ActiveControlName](FRigTransformElement* TransformElement) -> bool
 						{
+							if(const FRigControlElement* ControlElement = Cast<FRigControlElement>(TransformElement))
+							{
+								if(ControlElement->Settings.AnimationType != ERigControlAnimationType::AnimationControl)
+								{
+									return true;
+								}
+							}
+						
 							const FTransform Transform = Hierarchy->GetTransform(TransformElement, ERigTransformType::CurrentGlobal);
 
 							FRigBaseElementParentArray Parents = Hierarchy->GetParents(TransformElement);
@@ -893,20 +904,6 @@ void FControlRigEditMode::Render(const FSceneView* View, FViewport* Viewport, FP
 				}
 			}
 		}
-		else
-		{
-			const TArray<AControlRigShapeActor*>* ShapeActors = ControlRigShapeActors.Find(ControlRig);
-			if (ShapeActors)
-			{
-				for (AControlRigShapeActor* Actor : *ShapeActors)
-				{
-					if (GIsEditor && Actor->GetWorld() != nullptr && !Actor->GetWorld()->IsPlayInEditor())
-					{
-						Actor->SetIsTemporarilyHiddenInEditor(true);
-					}
-				}
-			}
-		}
 	}
 }
 
@@ -2575,17 +2572,20 @@ void FControlRigEditMode::ResetTransforms(bool bSelectionOnly)
 
 			TArray<FRigElementKey> SelectedRigElements = GetSelectedRigElements(ControlRig);
 			TArray<FRigElementKey> ControlsToReset = SelectedRigElements;
+			TArray<FRigElementKey> ControlsInteracting = SelectedRigElements;
 			if (!bSelectionOnly)
 			{
 				TArray<FRigControlElement*> Controls;
 				ControlRig->GetControlsInOrder(Controls);
 				ControlsToReset.SetNum(0);
+				ControlsInteracting.SetNum(0);
 				for (const FRigControlElement* Control : Controls)
 				{
+					ControlsToReset.Add(Control->GetKey());
 					if(Control->Settings.AnimationType == ERigControlAnimationType::AnimationControl ||
 						Control->Settings.AnimationType == ERigControlAnimationType::AnimationChannel)
 					{
-						ControlsToReset.Add(Control->GetKey());
+						ControlsInteracting.Add(Control->GetKey());
 					}
 				}
 			}
@@ -2602,7 +2602,7 @@ void FControlRigEditMode::ResetTransforms(bool bSelectionOnly)
 			}
 
 			FScopedTransaction Transaction(LOCTEXT("HierarchyResetTransforms", "Reset Transforms"));
-			FControlRigInteractionScope InteractionScope(ControlRig, ControlsToReset);
+			FControlRigInteractionScope InteractionScope(ControlRig, ControlsInteracting);
 
 			for (const FRigElementKey& ControlToReset : ControlsToReset)
 			{
@@ -3793,7 +3793,7 @@ void FControlRigEditMode::TickControlShape(AControlRigShapeActor* ShapeActor, co
 				ShapeActor->SetShapeColor(ShapeActor->OverrideColor.A < SMALL_NUMBER ?
 					ControlElement->Settings.ShapeColor : ShapeActor->OverrideColor);
 				ShapeActor->SetIsTemporarilyHiddenInEditor(
-					(!ControlElement->Settings.IsVisible() && !Settings->bShowAllProxyControls) ||
+					!bIsVisible ||
 					bControlsHiddenInViewport);
 				
 				ShapeActor->SetSelectable(
