Control Rig: Added missing get attribute nodes for vector/quat types

diff --git a/Engine/Plugins/Animation/ControlRig/Source/ControlRig/Private/Units/Animation/RigUnit_AnimAttribute.cpp b/Engine/Plugins/Animation/ControlRig/Source/ControlRig/Private/Units/Animation/RigUnit_AnimAttribute.cpp
--- a/Engine/Plugins/Animation/ControlRig/Source/ControlRig/Private/Units/Animation/RigUnit_AnimAttribute.cpp
+++ b/Engine/Plugins/Animation/ControlRig/Source/ControlRig/Private/Units/Animation/RigUnit_AnimAttribute.cpp
@@ -186,6 +186,21 @@ FRigUnit_GetAnimAttribute_Transform_Execute()
 	Value = ValuePtr ? *ValuePtr : FallbackValue;
 }
 
+FRigUnit_GetAnimAttribute_Vector_Execute()
+{
+	DECLARE_SCOPE_HIERARCHICAL_COUNTER_RIGUNIT()
 
+	const FVector* ValuePtr = GetAnimAttributeValue<FVector>(false, Context, Name, BoneName, CachedBoneName, CachedBoneIndex);
+	bWasFound = ValuePtr ? true : false;
+	Value = ValuePtr ? *ValuePtr : FallbackValue;	
+}
 
+FRigUnit_GetAnimAttribute_Quaternion_Execute()
+{
+	DECLARE_SCOPE_HIERARCHICAL_COUNTER_RIGUNIT()
+
+	const FQuat* ValuePtr = GetAnimAttributeValue<FQuat>(false, Context, Name, BoneName, CachedBoneName, CachedBoneIndex);
+	bWasFound = ValuePtr ? true : false;
+	Value = ValuePtr ? *ValuePtr : FallbackValue;	
+}
 
diff --git a/Engine/Plugins/Animation/ControlRig/Source/ControlRig/Private/Units/Animation/RigUnit_AnimAttribute.h b/Engine/Plugins/Animation/ControlRig/Source/ControlRig/Private/Units/Animation/RigUnit_AnimAttribute.h
--- a/Engine/Plugins/Animation/ControlRig/Source/ControlRig/Private/Units/Animation/RigUnit_AnimAttribute.h
+++ b/Engine/Plugins/Animation/ControlRig/Source/ControlRig/Private/Units/Animation/RigUnit_AnimAttribute.h
@@ -320,6 +320,90 @@ struct CONTROLRIG_API FRigUnit_GetAnimAttribute_Transform: public FRigUnit_GetAn
 	UPROPERTY()
 	FName CachedBoneName;
 	
+	UPROPERTY()
+	int32 CachedBoneIndex;
+};
+
+
+USTRUCT(meta=(DisplayName="Get Attribute - Vector"))
+struct CONTROLRIG_API FRigUnit_GetAnimAttribute_Vector: public FRigUnit_GetAnimAttributeBase
+{
+	GENERATED_BODY()
+	
+	FRigUnit_GetAnimAttribute_Vector()
+	{
+		Name = TEXT("AttributeName");
+		BoneName = NAME_None;
+		FallbackValue = FVector::ZeroVector;
+		Value = FVector::ZeroVector;
+		bWasFound = false;
+		CachedBoneName = NAME_None;
+		CachedBoneIndex = 0;
+	}
+
+	RIGVM_METHOD()
+	virtual void Execute(const FRigUnitContext& Context) override;
+
+	UPROPERTY(meta=(Input))
+	FName Name;
+	
+	UPROPERTY(meta=(Input, CustomWidget="BoneName"))
+	FName BoneName;
+
+	UPROPERTY(meta=(Input))
+	FVector FallbackValue;
+	
+	UPROPERTY(meta=(Output))
+	FVector Value;
+
+	UPROPERTY(meta=(Output))
+	bool bWasFound;
+
+	UPROPERTY()
+	FName CachedBoneName;
+	
+	UPROPERTY()
+	int32 CachedBoneIndex;
+};
+
+
+USTRUCT(meta=(DisplayName="Get Attribute - Quaternion"))
+struct CONTROLRIG_API FRigUnit_GetAnimAttribute_Quaternion: public FRigUnit_GetAnimAttributeBase
+{
+	GENERATED_BODY()
+	
+	FRigUnit_GetAnimAttribute_Quaternion()
+	{
+		Name = TEXT("AttributeName");
+		BoneName = NAME_None;
+		FallbackValue = FQuat::Identity;
+		Value = FQuat::Identity;
+		bWasFound = false;
+		CachedBoneName = NAME_None;
+		CachedBoneIndex = 0;
+	}
+
+	RIGVM_METHOD()
+	virtual void Execute(const FRigUnitContext& Context) override;
+
+	UPROPERTY(meta=(Input))
+	FName Name;
+	
+	UPROPERTY(meta=(Input, CustomWidget="BoneName"))
+	FName BoneName;
+
+	UPROPERTY(meta=(Input))
+	FQuat FallbackValue;
+	
+	UPROPERTY(meta=(Output))
+	FQuat Value;
+
+	UPROPERTY(meta=(Output))
+	bool bWasFound;
+
+	UPROPERTY()
+	FName CachedBoneName;
+	
 	UPROPERTY()
 	int32 CachedBoneIndex;
 };
\ No newline at end of file
