Constraints in Sequencer experiments

This is a first prototype of a constraints manager system.

note: have a look at the jira description for the features leading this work

The current design is based on the idea of having a unique constraint manager that stores several constraints (UTickableConstraint).
Each constraint holds a tick function (FConstraintTickFunction) that can register one or several functions (represented by a basic TFunction) that are evaluated when calling
the FConstraintTickFunction::ExecuteTick function. This FConstraintTickFunction can basically represent anything.

diff --git a/Engine/Plugins/Animation/ControlRig/Source/ControlRig/ControlRig.Build.cs b/Engine/Plugins/Animation/ControlRig/Source/ControlRig/ControlRig.Build.cs
--- a/Engine/Plugins/Animation/ControlRig/Source/ControlRig/ControlRig.Build.cs
+++ b/Engine/Plugins/Animation/ControlRig/Source/ControlRig/ControlRig.Build.cs
@@ -33,6 +33,7 @@ namespace UnrealBuildTool.Rules
                     "LevelSequence",
                     "RigVM",
                     "RHI",
+                    "Constraints"
                 }
             );
 
diff --git a/Engine/Plugins/Animation/ControlRig/Source/ControlRig/Private/Constraints/ControlRigTransformableHandle.cpp b/Engine/Plugins/Animation/ControlRig/Source/ControlRig/Private/Constraints/ControlRigTransformableHandle.cpp
new file mode 100644
--- /dev/null
+++ b/Engine/Plugins/Animation/ControlRig/Source/ControlRig/Private/Constraints/ControlRigTransformableHandle.cpp
@@ -0,0 +1,135 @@
+﻿// Copyright Epic Games, Inc. All Rights Reserved.
+
+
+#include "Constraints/ControlRigTransformableHandle.h"
+
+#include "ControlRig.h"
+#include "IControlRigObjectBinding.h"
+#include "Rigs/RigHierarchyElements.h"
+
+/**
+ * UTransformableControlHandle
+ */
+
+UTransformableControlHandle::~UTransformableControlHandle()
+{}
+
+bool UTransformableControlHandle::IsValid() const
+{
+	if (!ControlRig.IsValid() || ControlName == NAME_None)
+	{
+		return false;
+	}
+
+	const USkeletalMeshComponent* SkeletalMeshComponent = GetSkeletalMesh();
+	if (!SkeletalMeshComponent)
+	{
+		return false;
+	}
+
+	const FRigControlElement* ControlElement = ControlRig->FindControl(ControlName);
+	if (!ControlElement)
+	{
+		return false;
+	}
+	
+	return true;
+}
+
+// NOTE should we cache the skeletal mesh and the CtrlIndex to avoid looking for if every time
+// probably not for handling runtime changes
+void UTransformableControlHandle::SetTransform(const FTransform& InGlobal) const
+{
+	if (!ControlRig.IsValid() || ControlName == NAME_None)
+	{
+		return;
+	}
+
+	const USkeletalMeshComponent* SkeletalMeshComponent = GetSkeletalMesh();
+	if (!SkeletalMeshComponent)
+	{
+		return;
+	}
+
+	const FRigControlElement* ControlElement = ControlRig->FindControl(ControlName);
+	if (!ControlElement)
+	{
+		return;
+	}
+	
+	URigHierarchy* Hierarchy = ControlRig->GetHierarchy();
+	const FRigElementKey& ControlKey = ControlElement->GetKey();
+	const int32 CtrlIndex = Hierarchy->GetIndex(ControlKey);
+	
+	const FTransform& ComponentTransform = SkeletalMeshComponent->GetComponentTransform();
+	Hierarchy->SetGlobalTransform(CtrlIndex, InGlobal.GetRelativeTransform(ComponentTransform));
+}
+
+// NOTE should we cache the skeletal mesh and the CtrlIndex to avoid looking for if every time
+// probably not for handling runtime changes
+FTransform UTransformableControlHandle::GetTransform() const
+{
+	if (!ControlRig.IsValid() || ControlName == NAME_None)
+	{
+		return FTransform::Identity;
+	}
+
+	const USkeletalMeshComponent* SkeletalMeshComponent = GetSkeletalMesh();
+	if (!SkeletalMeshComponent)
+	{
+		return FTransform::Identity;
+	}
+
+	const FRigControlElement* ControlElement = ControlRig->FindControl(ControlName);
+	if (!ControlElement)
+	{
+		return FTransform::Identity;
+	}
+	
+	const FRigElementKey& ControlKey = ControlElement->GetKey();
+	const URigHierarchy* Hierarchy = ControlRig->GetHierarchy();
+	const int32 CtrlIndex = Hierarchy->GetIndex(ControlKey);
+
+	const FTransform& ComponentTransform = SkeletalMeshComponent->GetComponentTransform();
+	return Hierarchy->GetGlobalTransform(CtrlIndex) * ComponentTransform;
+}
+
+UObject* UTransformableControlHandle::GetPrerequisiteObject() const
+{
+	return GetSkeletalMesh(); 
+}
+
+FTickFunction* UTransformableControlHandle::GetTickFunction() const
+{
+	USkeletalMeshComponent* SkelMeshComponent = GetSkeletalMesh();
+	return SkelMeshComponent ? &SkelMeshComponent->PrimaryComponentTick : nullptr;
+}
+
+uint32 UTransformableControlHandle::GetHash() const
+{
+	if (ControlRig.IsValid() && ControlName != NAME_None)
+	{
+		return HashCombine(GetTypeHash(ControlRig.Get()), GetTypeHash(ControlName));
+	}
+	return 0;
+}
+
+USkeletalMeshComponent* UTransformableControlHandle::GetSkeletalMesh() const
+{
+	const TSharedPtr<IControlRigObjectBinding> ObjectBinding = ControlRig.IsValid() ? ControlRig->GetObjectBinding() : nullptr;
+	return ObjectBinding ? Cast<USkeletalMeshComponent>(ObjectBinding->GetBoundObject()) : nullptr;
+}
+
+#if WITH_EDITOR
+FName UTransformableControlHandle::GetName() const
+{
+	const USkeletalMeshComponent* SkeletalMesh = GetSkeletalMesh();
+	
+	const AActor* Actor = SkeletalMesh ? SkeletalMesh->GetOwner() : nullptr;
+	const FName ControlRigName = Actor ? FName(*Actor->GetActorLabel()) : SkeletalMesh ? SkeletalMesh->GetFName() : NAME_None; 
+
+	const FString FullName = FString::Printf(TEXT("%s/%s"), *ControlRigName.ToString(), *ControlName.ToString() );
+
+	return FName(*FullName);
+}
+#endif
\ No newline at end of file
diff --git a/Engine/Plugins/Animation/ControlRig/Source/ControlRig/Private/ControlRigModule.cpp b/Engine/Plugins/Animation/ControlRig/Source/ControlRig/Private/ControlRigModule.cpp
--- a/Engine/Plugins/Animation/ControlRig/Source/ControlRig/Private/ControlRigModule.cpp
+++ b/Engine/Plugins/Animation/ControlRig/Source/ControlRig/Private/ControlRigModule.cpp
@@ -1,10 +1,14 @@
 // Copyright Epic Games, Inc. All Rights Reserved.
 
 #include "ControlRigModule.h"
+
+#include "ControlRigGizmoActor.h"
+#include "TransformableRegistry.h"
 #include "Modules/ModuleManager.h"
 #include "Sequencer/ControlRigObjectSpawner.h"
 #include "ILevelSequenceModule.h"
 #include "ControlRigObjectVersion.h"
+#include "Constraints/ControlRigTransformableHandle.h"
 #include "UObject/DevObjectVersion.h"
 
 // Unique Control Rig Object version id
@@ -18,6 +22,8 @@ void FControlRigModule::StartupModule()
 	OnCreateMovieSceneObjectSpawnerHandle = LevelSequenceModule.RegisterObjectSpawner(FOnCreateMovieSceneObjectSpawner::CreateStatic(&FControlRigObjectSpawner::CreateObjectSpawner));
 
 	ManipulatorMaterial = LoadObject<UMaterial>(nullptr, TEXT("/ControlRig/M_Manip.M_Manip"));
+
+	RegisterTransformableCustomization();
 }
 
 void FControlRigModule::ShutdownModule()
@@ -29,4 +35,36 @@ void FControlRigModule::ShutdownModule()
 	}
 }
 
+void FControlRigModule::RegisterTransformableCustomization() const
+{
+	// Register UTransformableControlHandle and has function to handle control thru the transform constraints system
+	// as AControlRigShapeActor is only available if the ControlRig plugin is loaded.
+	auto CreateControlHandleFunc = [](const UObject* InObject, UObject* Outer)->UTransformableHandle*
+	{
+		if (const AControlRigShapeActor* ControlActor = Cast<AControlRigShapeActor>(InObject))
+		{
+			UTransformableControlHandle* CtrlHandle = NewObject<UTransformableControlHandle>(Outer);
+			CtrlHandle->ControlRig = ControlActor->ControlRig;
+			CtrlHandle->ControlName = ControlActor->ControlName;
+			return CtrlHandle;
+		}
+		return nullptr;
+	};
+
+	auto GetControlHashFunc = [](const UObject* InObject)->uint32
+	{
+		if (const AControlRigShapeActor* ControlActor = Cast<AControlRigShapeActor>(InObject))
+		{
+			const uint32 ControlHash = HashCombine(
+				GetTypeHash(ControlActor->ControlRig.Get()),
+				GetTypeHash(ControlActor->ControlName));
+			return ControlHash;
+		}
+		return 0;
+	};
+
+	FTransformableRegistry& Registry = FTransformableRegistry::Get();
+	Registry.Register( AControlRigShapeActor::StaticClass(), CreateControlHandleFunc, GetControlHashFunc);
+}
+
 IMPLEMENT_MODULE(FControlRigModule, ControlRig)
diff --git a/Engine/Plugins/Animation/ControlRig/Source/ControlRig/Private/ControlRigModule.h b/Engine/Plugins/Animation/ControlRig/Source/ControlRig/Private/ControlRigModule.h
--- a/Engine/Plugins/Animation/ControlRig/Source/ControlRig/Private/ControlRigModule.h
+++ b/Engine/Plugins/Animation/ControlRig/Source/ControlRig/Private/ControlRigModule.h
@@ -18,4 +18,6 @@ public:
 	
 private:
 	FDelegateHandle OnCreateMovieSceneObjectSpawnerHandle;
+
+	void RegisterTransformableCustomization() const;
 };
diff --git a/Engine/Plugins/Animation/ControlRig/Source/ControlRig/Public/Constraints/ControlRigTransformableHandle.h b/Engine/Plugins/Animation/ControlRig/Source/ControlRig/Public/Constraints/ControlRigTransformableHandle.h
new file mode 100644
--- /dev/null
+++ b/Engine/Plugins/Animation/ControlRig/Source/ControlRig/Public/Constraints/ControlRigTransformableHandle.h
@@ -0,0 +1,58 @@
+﻿// Copyright Epic Games, Inc. All Rights Reserved.
+
+#pragma once
+
+#include "TransformableHandle.h"
+
+#include "ControlRigTransformableHandle.generated.h"
+
+class UControlRig;
+class USkeletalMeshComponent;
+
+/**
+ * UTransformableControlHandle
+ */
+
+UCLASS()
+class UTransformableControlHandle : public UTransformableHandle 
+{
+	GENERATED_BODY()
+	
+public:
+	
+	virtual ~UTransformableControlHandle();
+
+	/** @todo document */
+	virtual bool IsValid() const override;
+
+	/** @todo document */
+	virtual void SetTransform(const FTransform& InGlobal) const override;
+	/** @todo document */
+	virtual FTransform GetTransform() const  override;
+
+	/** @todo document */
+	virtual UObject* GetPrerequisiteObject() const override;
+	/** @todo document */
+	virtual FTickFunction* GetTickFunction() const override;
+
+	/** @todo document */
+	virtual uint32 GetHash() const override;
+
+#if WITH_EDITOR
+	/** @todo document */
+	virtual FName GetName() const override;
+#endif
+
+	/** @todo document */
+	UPROPERTY()
+	TWeakObjectPtr<UControlRig> ControlRig;
+	
+	/** @todo document */
+	UPROPERTY()
+	FName ControlName;
+	
+private:
+
+	/** @todo document */
+	USkeletalMeshComponent* GetSkeletalMesh() const;
+};
\ No newline at end of file
diff --git a/Engine/Plugins/Animation/ControlRig/Source/ControlRigEditor/ControlRigEditor.Build.cs b/Engine/Plugins/Animation/ControlRig/Source/ControlRigEditor/ControlRigEditor.Build.cs
--- a/Engine/Plugins/Animation/ControlRig/Source/ControlRigEditor/ControlRigEditor.Build.cs
+++ b/Engine/Plugins/Animation/ControlRig/Source/ControlRigEditor/ControlRigEditor.Build.cs
@@ -85,7 +85,8 @@ namespace UnrealBuildTool.Rules
 					"AdvancedPreviewScene",
 					"ToolWidgets",
                     "AnimationWidgets",
-                    "ToolWidgets"
+                    "ActorPickerMode",
+                    "Constraints"
 				}
             );
 
diff --git a/Engine/Plugins/Animation/ControlRig/Source/ControlRigEditor/Private/EditMode/SControlRigEditModeTools.cpp b/Engine/Plugins/Animation/ControlRig/Source/ControlRigEditor/Private/EditMode/SControlRigEditModeTools.cpp
--- a/Engine/Plugins/Animation/ControlRig/Source/ControlRigEditor/Private/EditMode/SControlRigEditModeTools.cpp
+++ b/Engine/Plugins/Animation/ControlRig/Source/ControlRigEditor/Private/EditMode/SControlRigEditModeTools.cpp
@@ -38,6 +38,8 @@
 #include "ScopedTransaction.h"
 #include "ControlRigEditModeToolkit.h"
 #include "SControlRigDetails.h"
+#include "Constraints/SConstraintsWidget.h"
+
 #define LOCTEXT_NAMESPACE "ControlRigEditModeTools"
 
 
@@ -322,6 +324,79 @@ void SControlRigEditModeTools::Construct(const FArguments& InArgs, TSharedPtr<FC
 				]
 			]
 
+			+SVerticalBox::Slot()
+			.AutoHeight()
+			[
+				SAssignNew(PickerExpander, SExpandableArea)
+				.InitiallyCollapsed(true)
+				.AreaTitle(LOCTEXT("ConstraintsWidget", "Constraints"))
+				.AreaTitleFont(FAppStyle::GetFontStyle("DetailsView.CategoryFontStyle"))
+				.BorderBackgroundColor(FLinearColor(.6f, .6f, .6f))
+				.Padding(FMargin(8.f))
+				.HeaderContent()
+				[
+					SNew(SHorizontalBox)
+
+					// "Constraints" label
+					+SHorizontalBox::Slot()
+					.AutoWidth()
+					.HAlign(HAlign_Left)
+					.VAlign(VAlign_Center)
+					.Padding(0.f, 0.f, 0.f, 0.f)
+					[
+						SNew(STextBlock)
+						.Text(LOCTEXT("ConstraintsWidget", "Constraints"))
+						.Font(FCoreStyle::Get().GetFontStyle("ExpandableArea.TitleFont"))
+					]
+
+					// Spacer
+					+SHorizontalBox::Slot()
+					.FillWidth(1.f)
+					[
+						SNew(SSpacer)
+					]
+
+					// "Plus" icon
+					+SHorizontalBox::Slot()
+					.AutoWidth()
+					.HAlign(HAlign_Right)
+					.VAlign(VAlign_Center)
+					.Padding(0.f, 2.f, 8.f, 2.f)
+					[
+						SNew(SButton)
+						.ContentPadding(0.0f)
+						.ButtonStyle(FAppStyle::Get(), "NoBorder")
+						.IsEnabled_Lambda([InWorld]()
+						{
+							const ULevel* CurrentLevel = InWorld->GetCurrentLevel();
+							const TArray<AActor*> SelectedActors = CurrentLevel->Actors.FilterByPredicate( [](const AActor* Actor)
+							{
+								return Actor && Actor->IsSelected();
+							});
+							return !SelectedActors.IsEmpty();
+						})
+						.OnClicked(this, &SControlRigEditModeTools::HandleAddConstraintClicked)
+						.Cursor(EMouseCursor::Default)
+						.ToolTipText(LOCTEXT("AddConstraint", "Add Constraint"))
+						[
+							SNew(SImage)
+							.Image(FAppStyle::GetBrush(TEXT("Icons.PlusCircle")))
+						]
+					]
+				]
+				.BodyContent()
+				[
+					SAssignNew(ConstraintsEditionWidget, SConstraintsEditionWidget)
+				]
+				.OnAreaExpansionChanged_Lambda( [this](bool bIsExpanded)
+				{
+					if (ConstraintsEditionWidget)
+					{
+						ConstraintsEditionWidget->RefreshConstraintList();
+					}
+				})
+			]
+
 			+ SVerticalBox::Slot()
 			.AutoHeight()
 			[
@@ -675,11 +750,18 @@ void SControlRigEditModeTools::OnRigElementSelected(UControlRig* Subject, FRigCo
 		}
 	}
 #endif
-	if (Subject)
+	//mz todo handle multiple rigs
+	UControlRig* Rig = ControlRigs.Num() > 0 ? ControlRigs[0].Get() : nullptr;
+	if (Rig)
 	{
 		// get the selected controls
-		TArray<FRigElementKey> SelectedControls = Subject->GetHierarchy()->GetSelectedKeys(ERigElementType::Control);
-		SpacePickerWidget->SetControls(Subject->GetHierarchy(), SelectedControls);
+		TArray<FRigElementKey> SelectedControls = Rig->GetHierarchy()->GetSelectedKeys(ERigElementType::Control);
+		SpacePickerWidget->SetControls(Rig->GetHierarchy(), SelectedControls);
+
+		if (ConstraintsEditionWidget)
+		{
+			ConstraintsEditionWidget->InvalidateConstraintList();
+		}
 	}
 }
 
@@ -803,7 +885,7 @@ FReply SControlRigEditModeTools::OnBakeControlsToNewSpaceButtonClicked()
 	bool bNoValidControlRig = true;
 	for (TWeakObjectPtr<UControlRig>& ControlRig : ControlRigs)
 	{
-		if (ControlRig.IsValid() && SpacePickerWidget->GetHierarchy() == ControlRig->GetHierarchy())
+		if (ControlRig.IsValid())
 		{
 			bNoValidControlRig = false;
 			break;
@@ -821,7 +903,7 @@ FReply SControlRigEditModeTools::OnBakeControlsToNewSpaceButtonClicked()
 	}
 	for (TWeakObjectPtr<UControlRig>& ControlRig : ControlRigs)
 	{
-		if (ControlRig.IsValid() && SpacePickerWidget->GetHierarchy() == ControlRig->GetHierarchy())
+		if (ControlRig.IsValid())
 		{
 
 			FRigSpacePickerBakeSettings Settings;
@@ -901,6 +983,36 @@ FReply SControlRigEditModeTools::OnBakeControlsToNewSpaceButtonClicked()
 	return FReply::Unhandled();
 }
 
+FReply SControlRigEditModeTools::HandleAddConstraintClicked()
+{
+	const TSharedPtr<SConstraintsCreationWidget> Widget =
+		SNew(SConstraintsCreationWidget)
+		.OnConstraintCreated_Lambda( [this]()
+		{
+			if (ConstraintsEditionWidget)
+			{
+				ConstraintsEditionWidget->RefreshConstraintList();
+			}		
+		});
+	
+	FMenuBuilder MenuBuilder(true, nullptr);
+	MenuBuilder.BeginSection("CreateConstraint", LOCTEXT("CreateConstraintHeader", "Create New..."));
+	{
+		MenuBuilder.AddWidget(Widget.ToSharedRef(), FText::GetEmpty(), true);
+	}
+	MenuBuilder.EndSection();
+	
+	
+	FSlateApplication::Get().PushMenu(
+		AsShared(),
+		FWidgetPath(),
+		MenuBuilder.MakeWidget(),
+		FSlateApplication::Get().GetCursorPos(),
+		FPopupTransitionEffect(FPopupTransitionEffect::ContextMenu));
+	
+	return FReply::Handled();
+}
+
 EVisibility SControlRigEditModeTools::GetRigOptionExpanderVisibility() const
 {
 	for (const TWeakObjectPtr<UControlRig>& ControlRig : ControlRigs)
diff --git a/Engine/Plugins/Animation/ControlRig/Source/ControlRigEditor/Private/EditMode/SControlRigEditModeTools.h b/Engine/Plugins/Animation/ControlRig/Source/ControlRigEditor/Private/EditMode/SControlRigEditModeTools.h
--- a/Engine/Plugins/Animation/ControlRig/Source/ControlRigEditor/Private/EditMode/SControlRigEditModeTools.h
+++ b/Engine/Plugins/Animation/ControlRig/Source/ControlRigEditor/Private/EditMode/SControlRigEditModeTools.h
@@ -10,6 +10,7 @@
 #include "SRigSpacePickerWidget.h"
 #include "Widgets/DeclarativeSyntaxSupport.h"
 
+class SConstraintsEditionWidget;
 class FControlRigEditMode;
 class IDetailsView;
 class ISequencer;
@@ -91,6 +92,9 @@ private:
 	/** Storage for control rigs */
 	TArray<TWeakObjectPtr<UControlRig>> ControlRigs;
 
+	/** Constraint edition widget. */
+	TSharedPtr<SConstraintsEditionWidget> ConstraintsEditionWidget = nullptr; 
+
 	/** Display or edit set up for property */
 	bool ShouldShowPropertyOnDetailCustomization(const struct FPropertyAndParent& InPropertyAndParent) const;
 	bool IsReadOnlyPropertyOnDetailCustomization(const struct FPropertyAndParent& InPropertyAndParent) const;
@@ -108,6 +112,8 @@ private:
 	FReply HandleAddSpaceClicked();
 	FReply OnBakeControlsToNewSpaceButtonClicked();
 
+	FReply HandleAddConstraintClicked();
+
 	EVisibility GetRigOptionExpanderVisibility() const;
 
 	void OnRigOptionFinishedChange(const FPropertyChangedEvent& PropertyChangedEvent);
diff --git a/Engine/Plugins/Animation/ControlRig/Source/ControlRigEditor/Private/Editor/Constraints/SConstraintsWidget.cpp b/Engine/Plugins/Animation/ControlRig/Source/ControlRigEditor/Private/Editor/Constraints/SConstraintsWidget.cpp
new file mode 100644
--- /dev/null
+++ b/Engine/Plugins/Animation/ControlRig/Source/ControlRigEditor/Private/Editor/Constraints/SConstraintsWidget.cpp
@@ -0,0 +1,668 @@
+﻿// Copyright Epic Games, Inc. All Rights Reserved.
+
+#include "SConstraintsWidget.h"
+
+#include "ActorPickerMode.h"
+#include "ConstraintsActor.h"
+#include "ControlRigEditorStyle.h"
+#include "DetailLayoutBuilder.h"
+#include "LevelEditorActions.h"
+#include "LevelEditorViewport.h"
+#include "PropertyCustomizationHelpers.h"
+#include "Selection.h"
+#include "SlateOptMacros.h"
+#include "LevelEditor.h"
+#include "TransformConstraint.h"
+
+
+#define LOCTEXT_NAMESPACE "SConstraintsWidget"
+
+BEGIN_SLATE_FUNCTION_BUILD_OPTIMIZATION
+
+TArray< SConstraintsCreationWidget::ItemSharedPtr > SConstraintsCreationWidget::ListItems({
+	FDroppableConstraintItem::Make(ETransformConstraintType::Translation),
+	FDroppableConstraintItem::Make(ETransformConstraintType::Rotation),
+	FDroppableConstraintItem::Make(ETransformConstraintType::Scale),
+	FDroppableConstraintItem::Make(ETransformConstraintType::Parent),
+	FDroppableConstraintItem::Make(ETransformConstraintType::LookAt)
+});
+
+TArray< const FSlateBrush* > FConstraintInfo::Brushes({
+	FAppStyle::Get().GetBrush("EditorViewport.TranslateMode"),
+	FAppStyle::Get().GetBrush("EditorViewport.RotateMode"),
+	FAppStyle::Get().GetBrush("EditorViewport.ScaleMode"),
+	FAppStyle::Get().GetBrush("Icons.Transform"),
+	FAppStyle::Get().GetBrush("Level.VisibleHighlightIcon16x")
+});
+
+TMap< UClass*, ETransformConstraintType > FConstraintInfo::ConstraintToType({
+	{UTickableTranslationConstraint::StaticClass(), ETransformConstraintType::Translation},
+	{UTickableRotationConstraint::StaticClass(), ETransformConstraintType::Rotation},
+	{UTickableScaleConstraint::StaticClass(), ETransformConstraintType::Scale},
+	{UTickableParentConstraint::StaticClass(), ETransformConstraintType::Parent},
+	{UTickableLookAtConstraint::StaticClass(), ETransformConstraintType::LookAt}
+});
+
+const FSlateBrush* FConstraintInfo::GetBrush(uint8 InType)
+{
+	static const UEnum* ETransformConstraintTypeEnum = StaticEnum<ETransformConstraintType>();
+	if (ETransformConstraintTypeEnum->IsValidEnumValue(InType))
+	{
+		return Brushes[InType];
+	}
+
+	return FAppStyle::Get().GetDefaultBrush();
+}
+
+int8 FConstraintInfo::GetType(UClass* InClass)
+{
+	if (ETransformConstraintType* TransformConstraint = ConstraintToType.Find(InClass))
+	{
+		return static_cast<int8>(*TransformConstraint); 
+	}
+	return -1;
+}
+
+namespace
+{
+
+// NOTE we use this function to get the current selection as control actors are Temporary Editor Actors so won't be
+// pushed added to the selection list.
+TArray<AActor*> GetCurrentSelection()
+{
+	const UWorld* World = GCurrentLevelEditingViewportClient->GetWorld();
+	const ULevel* CurrentLevel = World->GetCurrentLevel();
+
+	return CurrentLevel->Actors.FilterByPredicate( [](const AActor* Actor)
+	{
+		return Actor && Actor->IsSelected();
+	});	
+}
+	
+}
+
+/**
+ * SConstraintItem
+ */
+
+void SDroppableConstraintItem::Construct(
+		const FArguments& InArgs,
+		const TSharedPtr<const FDroppableConstraintItem>& InItem,
+		TSharedPtr<SConstraintsCreationWidget> InConstraintsWidget)
+{
+	ConstraintItem = InItem;
+	ConstraintType = InItem->Type;
+	ConstraintsWidget = InConstraintsWidget;
+	
+	const FButtonStyle& ButtonStyle = FAppStyle::GetWidgetStyle<FButtonStyle>( "PlacementBrowser.Asset" );
+
+	// enum to string
+	const UEnum* ETransformConstraintTypeEnum = StaticEnum<ETransformConstraintType>();
+	// const uint8 ConstraintType = static_cast<uint8>(InItem->Type);
+	const FString TypeStr = ETransformConstraintTypeEnum->GetNameStringByValue((uint8)ConstraintType);
+
+	// tooltip
+	const FString ToolTipStr = FString::Printf(TEXT("Create new %s constraint."), *TypeStr);
+	const TSharedPtr<IToolTip> ToolTip = FSlateApplicationBase::Get().MakeToolTip(FText::FromString(ToolTipStr));
+	
+	ChildSlot
+	.Padding(FMargin(8.f, 2.f, 12.f, 2.f))
+	[
+		SNew(SOverlay)
+
+		+SOverlay::Slot()
+		[
+			SNew(SBorder)
+			.BorderImage( FAppStyle::Get().GetBrush("PlacementBrowser.Asset.Background"))
+			.Cursor( EMouseCursor::GrabHand )
+			.ToolTip(ToolTip)
+			.Padding(0)
+			[
+				SNew( SHorizontalBox )
+
+				+ SHorizontalBox::Slot()
+				.Padding(8.0f, 4.f)
+				.AutoWidth()
+				.HAlign(HAlign_Center)
+				.VAlign(VAlign_Center)
+				[
+					SNew( SBox )
+					.WidthOverride(40)
+					.HeightOverride(40)
+					[
+					 	SNew(SImage)
+					 	.DesiredSizeOverride(FVector2D(16, 16))
+						.Image(FConstraintInfo::GetBrush((uint8)ConstraintType))
+						.ColorAndOpacity(FSlateColor::UseForeground())
+					]
+				]
+
+				+ SHorizontalBox::Slot()
+				.VAlign(VAlign_Fill)
+				.Padding(0)
+				[
+					SNew(SBorder)
+					.BorderImage(FAppStyle::Get().GetBrush("PlacementBrowser.Asset.LabelBack"))
+					[
+						SNew(SHorizontalBox)
+						+SHorizontalBox::Slot()
+						.Padding(9, 0, 0, 1)
+						.VAlign(VAlign_Center)
+						[
+							SNew( STextBlock )
+							.TextStyle( FAppStyle::Get(), "PlacementBrowser.Asset.Name" )
+							.Text_Lambda( [TypeStr]
+							{
+								return FText::FromString(TypeStr);
+							} )
+						]
+					]
+				]
+			]
+		]
+
+		+SOverlay::Slot()
+		[
+			SNew(SBorder)
+			.BorderImage_Lambda( [this, &ButtonStyle]
+			{
+				if (bIsPressed)
+				{
+					return &ButtonStyle.Pressed;
+				}
+				
+				if (IsHovered())
+				{
+					return &ButtonStyle.Hovered;
+				}
+
+				return &ButtonStyle.Normal;
+			})
+			.Cursor( EMouseCursor::GrabHand )
+			.ToolTip( ToolTip )
+		]
+	];
+}
+
+FReply SDroppableConstraintItem::OnMouseButtonDown(const FGeometry& MyGeometry, const FPointerEvent& MouseEvent)
+{
+	if (MouseEvent.GetEffectingButton() == EKeys::LeftMouseButton)
+	{
+		bIsPressed = true;
+		return FReply::Handled().DetectDrag( SharedThis( this ), MouseEvent.GetEffectingButton() );
+	}
+
+	return FReply::Unhandled();
+}
+
+FReply SDroppableConstraintItem::OnMouseButtonUp(const FGeometry& MyGeometry, const FPointerEvent& MouseEvent)
+{
+	if ( MouseEvent.GetEffectingButton() == EKeys::LeftMouseButton )
+	{
+		bIsPressed = false;
+	}
+
+	return FReply::Unhandled();
+}
+
+FReply SDroppableConstraintItem::OnDragDetected(const FGeometry& MyGeometry, const FPointerEvent& MouseEvent)
+{
+	bIsPressed = false;
+
+	// FIXME temp approach for selecting the parent
+	auto OnConstraintDragged = [this]()
+	{
+		static const FActorPickerModeModule& ActorPickerMode = FModuleManager::Get().GetModuleChecked<FActorPickerModeModule>("ActorPickerMode");
+
+		TSharedPtr<SConstraintsCreationWidget> ConstraintsCreationWidget = this->ConstraintsWidget.Pin();
+		ETransformConstraintType ConstraintTypeCopy = this->ConstraintType;
+		ActorPickerMode.BeginActorPickingMode(
+			FOnGetAllowedClasses(), 
+			FOnShouldFilterActor(), 
+			FOnActorSelected::CreateLambda([ConstraintsCreationWidget, ConstraintTypeCopy](AActor* InActor)
+			{
+				SDroppableConstraintItem::CreateConstraint(InActor, ConstraintsCreationWidget, ConstraintTypeCopy);
+			}) );
+
+		
+		return FReply::Handled();
+	};
+	
+	if (MouseEvent.IsMouseButtonDown( EKeys::LeftMouseButton ))
+	{
+		return OnConstraintDragged();
+	}
+	
+	return FReply::Handled();
+}
+
+void SDroppableConstraintItem::CreateConstraint(
+	AActor* InParent,
+	TSharedPtr<SConstraintsCreationWidget> InCreationWidget,
+	const ETransformConstraintType InConstraintType)
+{
+	if (!InParent)
+	{
+		return;	
+	}
+
+	const UEnum* ETransformConstraintTypeEnum = StaticEnum<ETransformConstraintType>();
+	if (!ETransformConstraintTypeEnum->IsValidEnumValue(static_cast<int64>(InConstraintType)))
+	{
+		return;
+	}
+	
+	// get selected actors
+	const TArray<AActor*> Selection = GetCurrentSelection();
+	if (Selection.IsEmpty())
+	{
+		return;	
+	}
+
+	// create constraints
+	UWorld* World = GCurrentLevelEditingViewportClient->GetWorld();
+	bool bCreated = false;
+	for (const AActor* Child: Selection)
+	{
+		if (Child != InParent)
+		{
+			bCreated |= TransformConstraintUtils::Create(World, InParent, Child, InConstraintType);
+		}
+	}
+
+	// update list
+	const bool bCreationBound = InCreationWidget.IsValid() ? InCreationWidget->OnConstraintCreated.IsBound() : false;
+	if (bCreated && bCreationBound)
+	{
+		InCreationWidget->OnConstraintCreated.Execute();
+	}
+}
+
+/**
+ * SConstraintCreationWidget
+ */
+
+void SConstraintsCreationWidget::Construct(const FArguments& InArgs)
+{
+	OnConstraintCreated = InArgs._OnConstraintCreated;
+	
+	ChildSlot
+	[
+		SNew( SVerticalBox )
+
+		+ SVerticalBox::Slot()
+		.Padding(FMargin(0.0f, 3.f))
+		[
+			SNew(SOverlay)
+
+			+ SOverlay::Slot()
+			[
+				SAssignNew(ListView, ConstraintItemListView)
+				.SelectionMode(ESelectionMode::None)
+				.ListItemsSource( &ListItems )
+				.OnGenerateRow(this, &SConstraintsCreationWidget::OnGenerateWidgetForItem)
+			]
+		]
+	];
+}
+
+TSharedRef<ITableRow> SConstraintsCreationWidget::OnGenerateWidgetForItem(
+	ItemSharedPtr InItem,
+	const TSharedRef<STableViewBase>& OwnerTable)
+{
+	return SNew(STableRow<ItemSharedPtr>, OwnerTable)
+	.Style(&FAppStyle::Get(), "PlacementBrowser.PlaceableItemRow")
+	[
+		SNew(SDroppableConstraintItem, InItem.ToSharedRef(), SharedThis(this))
+	];
+}
+
+/**
+ * SEditableConstraintItem
+ */
+
+void SEditableConstraintItem::Construct(
+	const FArguments& InArgs,
+	const TSharedPtr<FEditableConstraintItem>& InItem,
+	TSharedPtr<SConstraintsEditionWidget> InConstraintsWidget)
+{
+	ConstraintItem = InItem;
+	ConstraintsWidget = InConstraintsWidget;
+
+	static const FSlateBrush* RoundedBoxBrush = FControlRigEditorStyle::Get().GetBrush(TEXT("ControlRig.SpacePicker.RoundedRect"));
+	
+	// enum to string
+	const uint8 ConstraintType = static_cast<uint8>(InItem->Type);
+
+	const FSimpleDelegate OnConstraintRemoved = FSimpleDelegate::CreateLambda([this]()
+	{
+		if(ConstraintsWidget.IsValid())
+		{
+			ConstraintsWidget.Pin()->RemoveItem(ConstraintItem);
+		}
+	});
+	
+	ChildSlot
+	.Padding(FMargin(8.f, 2.f, 12.f, 2.f))
+	[
+		SNew(SHorizontalBox)
+		
+		+ SHorizontalBox::Slot()
+		.AutoWidth()
+		.VAlign(VAlign_Center)
+		.HAlign(HAlign_Fill)
+		.Padding(0)
+		[
+			SNew(SBorder)
+			.Padding(FMargin(5.0, 2.0, 5.0, 2.0))
+			.BorderImage(RoundedBoxBrush)
+			.BorderBackgroundColor(FStyleColors::Transparent)
+			.Content()
+			[
+				SNew(SHorizontalBox)
+
+				// constraint icon
+				+ SHorizontalBox::Slot()
+				.AutoWidth()
+				.VAlign(VAlign_Center)
+				.HAlign(HAlign_Left)
+				.Padding(FMargin(0.f, 0.f, 3.f, 0.f))
+				[
+					SNew(SImage)
+					.Image(FConstraintInfo::GetBrush(ConstraintType))
+					.ColorAndOpacity(FSlateColor::UseForeground())
+				]
+
+				// constraint name
+				+ SHorizontalBox::Slot()
+				.AutoWidth()
+				.VAlign(VAlign_Center)
+				.HAlign(HAlign_Left)
+				.Padding(0)
+				[
+					SNew( STextBlock )
+					.Text_Lambda( [InItem]()
+					{
+						if(InItem->ParentName == NAME_None)
+						{
+							return FText::FromName(InItem->Name);
+						}
+
+						const FString TextStr = FString::Printf(TEXT("%s (%s)"), *InItem->Name.ToString(), *InItem->ParentName.ToString());
+						return FText::FromString(TextStr);
+					})
+					.Font( IDetailLayoutBuilder::GetDetailFont() )
+				]
+			]
+		]
+
+		+ SHorizontalBox::Slot()
+		.FillWidth(1.f)
+		[
+			SNew(SSpacer)
+		]
+
+		// move up
+		+ SHorizontalBox::Slot()
+		.AutoWidth()
+		.VAlign(VAlign_Center)
+		.HAlign(HAlign_Left)
+		.Padding(0)
+		[
+			SNew(SButton)
+			.ButtonStyle(FAppStyle::Get(), TEXT("SimpleButton"))
+			.ContentPadding(0)
+			.OnClicked_Lambda(	[this, InItem]()
+			{
+				if (ConstraintsWidget.IsValid())
+				{
+					ConstraintsWidget.Pin()->MoveItemUp(InItem);
+				}
+				return FReply::Handled();
+			})
+			.IsEnabled_Lambda([this]()
+			{
+				return ConstraintsWidget.IsValid() ? !ConstraintsWidget.Pin()->IsFirstItem(ConstraintItem) : false;
+			})
+			.ToolTipText(LOCTEXT("MoveConstraintUp", "Move this constraint up in the list."))
+			[
+				SNew(SImage)
+				.Image(FAppStyle::GetBrush("Icons.ChevronUp"))
+				.ColorAndOpacity(FSlateColor::UseForeground())
+			]
+		]
+
+		// move down
+		+ SHorizontalBox::Slot()
+		.AutoWidth()
+		.VAlign(VAlign_Center)
+		.HAlign(HAlign_Left)
+		.Padding(0)
+		[
+			SNew(SButton)
+			.ButtonStyle(FAppStyle::Get(), TEXT("SimpleButton"))
+			.ContentPadding(0)
+			.OnClicked_Lambda([this, InItem]()
+			{
+				if (ConstraintsWidget.IsValid())
+				{
+					ConstraintsWidget.Pin()->MoveItemDown(InItem);
+				}	
+				return FReply::Handled();
+			})
+			.IsEnabled_Lambda([this]()
+			{
+				return ConstraintsWidget.IsValid() ? !ConstraintsWidget.Pin()->IsLastItem(ConstraintItem) : false;
+			})
+			.ToolTipText(LOCTEXT("MoveConstraintDown", "Move this constraint down in the list."))
+			[
+				SNew(SImage)
+				.Image(FAppStyle::GetBrush("Icons.ChevronDown"))
+				.ColorAndOpacity(FSlateColor::UseForeground())
+			]
+		]
+
+		// deletion
+		+ SHorizontalBox::Slot()
+		.AutoWidth()
+		.VAlign(VAlign_Center)
+		.HAlign(HAlign_Left)
+		.Padding(0)
+		[
+			 PropertyCustomizationHelpers::MakeClearButton(OnConstraintRemoved, LOCTEXT("DeleteConstraint", "Remove this constraint."), true)
+		]
+	];
+}
+
+/**
+ * SConstraintsEditionWidget
+ */
+
+void SConstraintsEditionWidget::Construct(const FArguments& InArgs)
+{
+	ChildSlot
+	[
+		SNew( SVerticalBox )
+
+		+ SVerticalBox::Slot()
+		.Padding(FMargin(0.0f, 3.f))
+		[
+			SNew(SOverlay)
+
+			+ SOverlay::Slot()
+			[
+				SAssignNew(ListView, ConstraintItemListView)
+				.SelectionMode(ESelectionMode::Single)
+				.ListItemsSource( &ListItems )
+				.OnGenerateRow(this, &SConstraintsEditionWidget::OnGenerateWidgetForItem)
+			]
+		]
+	];
+
+	RefreshConstraintList();
+	RegisterSelectionChanged();
+}
+
+SConstraintsEditionWidget::~SConstraintsEditionWidget()
+{
+	UnregisterSelectionChanged();
+}
+
+void SConstraintsEditionWidget::RegisterSelectionChanged()
+{
+	FLevelEditorModule& LevelEditor = FModuleManager::GetModuleChecked<FLevelEditorModule>("LevelEditor");
+	FLevelEditorModule::FActorSelectionChangedEvent& ActorSelectionChangedEvent = LevelEditor.OnActorSelectionChanged();
+
+	// unregister previous one
+	if (OnSelectionChangedHandle.IsValid())
+	{
+		ActorSelectionChangedEvent.Remove(OnSelectionChangedHandle);
+		OnSelectionChangedHandle.Reset();
+	}
+
+	// register
+	ActorSelectionChangedEvent.AddRaw( this, &SConstraintsEditionWidget::OnActorSelectionChanged );
+}
+
+void SConstraintsEditionWidget::UnregisterSelectionChanged()
+{
+	if (OnSelectionChangedHandle.IsValid())
+	{
+		FLevelEditorModule& LevelEditor = FModuleManager::GetModuleChecked<FLevelEditorModule>("LevelEditor");
+		LevelEditor.OnActorSelectionChanged().Remove(OnSelectionChangedHandle);
+		OnSelectionChangedHandle.Reset();
+	}
+}
+
+void SConstraintsEditionWidget::Tick(const FGeometry& AllottedGeometry, const double InCurrentTime, const float InDeltaTime)
+{
+	SCompoundWidget::Tick(AllottedGeometry, InCurrentTime, InDeltaTime);
+
+	if(bNeedsRefresh)
+	{
+		RefreshConstraintList();
+		bNeedsRefresh = false;
+	}
+}
+
+TSharedRef<ITableRow> SConstraintsEditionWidget::OnGenerateWidgetForItem(
+	ItemSharedPtr InItem,
+	const TSharedRef<STableViewBase>& OwnerTable)
+{
+	return SNew(STableRow<ItemSharedPtr>, OwnerTable)
+	.Style(&FAppStyle::Get(), "PlacementBrowser.PlaceableItemRow")
+	[
+		SNew(SEditableConstraintItem, InItem.ToSharedRef(), SharedThis(this))
+	];
+}
+
+bool SConstraintsEditionWidget::IsFirstItem(const TSharedPtr<FEditableConstraintItem>& Item) const
+{
+	return ListItems.Num() > 1 && ListItems[0] == Item; 
+}
+
+bool SConstraintsEditionWidget::IsLastItem(const TSharedPtr<FEditableConstraintItem>& Item) const
+{
+	return ListItems.Num() > 1 && ListItems.Last() == Item;
+}
+
+void SConstraintsEditionWidget::MoveItemUp(const TSharedPtr<FEditableConstraintItem>& Item)
+{
+	const int32 Index = ListItems.IndexOfByKey(Item);
+	if (Index > 0 && ListItems.IsValidIndex(Index))
+	{
+		UWorld* World = GCurrentLevelEditingViewportClient->GetWorld();
+		const FConstraintsManagerController& Controller = FConstraintsManagerController::Get(World);
+
+		// the current item needs to tick before the previous item 
+		const FName BeforeName(ListItems[Index]->Name.ToString());
+		const FName AfterName(ListItems[Index-1]->Name.ToString());
+		Controller.SetConstraintsDependencies(BeforeName, AfterName);
+		
+		RefreshConstraintList();
+	}
+}
+
+void SConstraintsEditionWidget::MoveItemDown(const TSharedPtr<FEditableConstraintItem>& Item)
+{
+	const int32 Index = ListItems.IndexOfByKey(Item);
+	if (ListItems.IsValidIndex(Index))
+	{
+		UWorld* World = GCurrentLevelEditingViewportClient->GetWorld();
+		const FConstraintsManagerController& Controller = FConstraintsManagerController::Get(World);
+
+		// the current item needs to tick after the next item 
+		const FName BeforeName(ListItems[Index+1]->Name.ToString());
+		const FName AfterName(ListItems[Index]->Name.ToString());
+		Controller.SetConstraintsDependencies(BeforeName, AfterName);
+		
+		RefreshConstraintList();
+	}
+}
+
+void SConstraintsEditionWidget::RemoveItem(const TSharedPtr<FEditableConstraintItem>& Item)
+{
+	UWorld* World = GCurrentLevelEditingViewportClient->GetWorld();
+	const FConstraintsManagerController& Controller = FConstraintsManagerController::Get(World);
+		
+	const FName ConstraintName(Item->Name.ToString());
+	Controller.RemoveConstraint(ConstraintName);
+
+	RefreshConstraintList();
+}
+
+void SConstraintsEditionWidget::InvalidateConstraintList()
+{
+	bNeedsRefresh = true;
+}
+
+void SConstraintsEditionWidget::RefreshConstraintList()
+{
+	// get constraints
+	UWorld* World = GCurrentLevelEditingViewportClient->GetWorld();
+	const TArray<AActor*> Selection = GetCurrentSelection();
+	
+	const bool bIsConstraintsActor = Selection.Num() == 1 && Selection[0]->IsA<AConstraintsActor>();
+	
+	TArray< TObjectPtr<UTickableConstraint> > Constraints;
+	if (bIsConstraintsActor)
+	{
+		const FConstraintsManagerController& Controller = FConstraintsManagerController::Get(World);
+		Constraints = Controller.GetConstrainsArray();
+	}
+	else
+	{
+		for (const AActor* Actor : Selection)
+		{
+			TransformConstraintUtils::GetParentConstraints(World, Actor, Constraints);
+		}
+	}
+	
+	// rebuild item list
+	ListItems.Empty();
+	
+	const UEnum* ETransformConstraintTypeEnum = StaticEnum<ETransformConstraintType>();
+	for (const TObjectPtr<UTickableConstraint>& Constraint: Constraints)
+	{
+		const int8 Type = FConstraintInfo::GetType(Constraint->GetClass()); 
+		if (ETransformConstraintTypeEnum->IsValidEnumValue(Type))
+		{
+			const ETransformConstraintType ConstraintType = static_cast<ETransformConstraintType>(Type);
+			ListItems.Emplace( FEditableConstraintItem::Make(Constraint->GetFName(), ConstraintType, Constraint->GetLabel()) );
+		}
+	}
+
+	// refresh tree
+	ListView->RequestListRefresh();
+}
+
+void SConstraintsEditionWidget::OnActorSelectionChanged(const TArray<UObject*>& NewSelection, bool bForceRefresh)
+{
+	// NOTE we use this delegate to trigger an tree update, however, control actors are not selected as they are 
+	// Temporary Editor Actors so NewSelection won't contain the controls
+	InvalidateConstraintList();
+}
+
+END_SLATE_FUNCTION_BUILD_OPTIMIZATION
+
+#undef LOCTEXT_NAMESPACE
diff --git a/Engine/Plugins/Animation/ControlRig/Source/ControlRigEditor/Private/Editor/Constraints/SConstraintsWidget.h b/Engine/Plugins/Animation/ControlRig/Source/ControlRigEditor/Private/Editor/Constraints/SConstraintsWidget.h
new file mode 100644
--- /dev/null
+++ b/Engine/Plugins/Animation/ControlRig/Source/ControlRigEditor/Private/Editor/Constraints/SConstraintsWidget.h
@@ -0,0 +1,244 @@
+﻿// Copyright Epic Games, Inc. All Rights Reserved.
+
+#pragma once
+
+#include "Widgets/SCompoundWidget.h"
+#include "Constraint.h"
+
+class AActor;
+class SConstraintsCreationWidget;
+class SConstraintsEditionWidget;
+class UTickableConstraint;
+
+/**
+ * FConstraintInfo
+ */
+
+class FConstraintInfo
+{
+public:
+	static const FSlateBrush* GetBrush(uint8 InType);
+	static int8 GetType(UClass* InClass);
+private:
+	static TArray< const FSlateBrush* > Brushes;
+	static TMap< UClass*, ETransformConstraintType > ConstraintToType;
+};
+
+/**
+ * The classes below are used to implement a constraint creation widget.
+ *
+ * It represents the different constraints type that can be created between two objects. At this stage, it only
+ * transform constraints (as described in EConstType but can basically represents any type  of constraint.
+ * The resulting widget is a tree filled with drag & droppable items that can be dropped on actors.
+ * The selection represents the child and the picked actor will be the parent of the constraint. 
+ */
+
+/**
+ * FDroppableConstraintItem
+ */
+
+struct FDroppableConstraintItem
+{
+	static TSharedRef<FDroppableConstraintItem> Make(ETransformConstraintType InType)
+	{
+		return MakeShareable(new FDroppableConstraintItem(InType));
+	}
+	
+	ETransformConstraintType Type = ETransformConstraintType::Parent;
+
+private:
+	FDroppableConstraintItem(ETransformConstraintType InType)
+		: Type(InType)
+	{}
+	FDroppableConstraintItem() = default;
+};
+
+/**
+ * SDroppableConstraintItem
+ */
+
+class SDroppableConstraintItem : public SCompoundWidget
+{
+public:
+	SLATE_BEGIN_ARGS(SDroppableConstraintItem){}
+	SLATE_END_ARGS()
+
+	/** Constructs this widget with InArgs and the actual tree item. */
+	void Construct(
+		const FArguments& InArgs,
+		const TSharedPtr<const FDroppableConstraintItem>& InItem,
+		TSharedPtr<SConstraintsCreationWidget> InConstraintsWidget);
+
+	// SWidget interface
+	virtual FReply OnMouseButtonDown(const FGeometry& MyGeometry, const FPointerEvent& MouseEvent) override;
+	virtual FReply OnMouseButtonUp(const FGeometry& MyGeometry, const FPointerEvent& MouseEvent) override;
+	virtual FReply OnDragDetected(const FGeometry& MyGeometry, const FPointerEvent& MouseEvent) override;
+	// End of SWidget interface
+	
+private:
+	/** Creates the constraint between the current selection and the picked actor. */
+	static void CreateConstraint(
+		AActor* InParent,
+		TSharedPtr<SConstraintsCreationWidget> InConstraints,
+		const ETransformConstraintType InConstraintType);
+
+	/** TSharedPtr to the tree item. */
+	TSharedPtr<const FDroppableConstraintItem> ConstraintItem;
+
+	ETransformConstraintType ConstraintType = ETransformConstraintType::Parent;
+
+	/** State on the item. */
+	bool bIsPressed = false;
+
+	TWeakPtr<SConstraintsCreationWidget> ConstraintsWidget;
+};
+
+/**
+ * SConstraintsCreationWidget
+ */
+
+class CONTROLRIGEDITOR_API SConstraintsCreationWidget : public SCompoundWidget
+{
+public:
+
+	DECLARE_DELEGATE(FOnConstraintCreated);
+	
+	SLATE_BEGIN_ARGS(SConstraintsCreationWidget)	{}
+
+		SLATE_EVENT(FOnConstraintCreated, OnConstraintCreated)
+	
+	SLATE_END_ARGS()
+
+	FOnConstraintCreated OnConstraintCreated;
+	
+	/** Constructs this widget with InArgs */
+	void Construct(const FArguments& InArgs);
+
+private:
+	/** Types */
+	using ItemSharedPtr = TSharedPtr<FDroppableConstraintItem>;
+	using ConstraintItemListView = SListView< ItemSharedPtr >;
+
+	/** Generates a widget for the specified item */
+	TSharedRef<ITableRow> OnGenerateWidgetForItem(ItemSharedPtr InItem, const TSharedRef<STableViewBase>& OwnerTable);
+	
+	/** List view that shows constraint types*/
+	TSharedPtr< ConstraintItemListView > ListView;
+
+	/** Static list of constraint types */
+	static TArray< ItemSharedPtr > ListItems;
+};
+
+/**
+ * The classes below are used to display a list of constraints. 
+ */
+
+/**
+ * FEditableConstraintItem
+ */
+
+class FEditableConstraintItem
+{
+public:
+	static TSharedRef<FEditableConstraintItem> Make(
+		const FName& InName,
+		ETransformConstraintType InType,
+		const FName& InParentName)
+	{
+		return MakeShareable(new FEditableConstraintItem(InName, InType, InParentName));
+	}
+	
+	FName Name = NAME_None;
+	ETransformConstraintType Type = ETransformConstraintType::Parent;
+	FName ParentName = NAME_None;
+
+private:
+	FEditableConstraintItem(const FName& InName, ETransformConstraintType InType, const FName& InParentName)
+		: Name(InName)
+		, Type(InType)
+		, ParentName(InParentName)
+	{}
+	FEditableConstraintItem() {}
+};
+
+/**
+ * SEditableConstraintItem
+ */
+
+class SEditableConstraintItem : public SCompoundWidget
+{
+public:
+	SLATE_BEGIN_ARGS(SEditableConstraintItem){}
+	SLATE_END_ARGS()
+
+	/** Constructs this widget with InArgs and the actual tree item. */
+	void Construct(
+		const FArguments& InArgs,
+		const TSharedPtr<FEditableConstraintItem>& InItem,
+		TSharedPtr<SConstraintsEditionWidget> InConstraintsWidget);
+	
+private:
+	/** TSharedPtr to the tree item. */
+	TSharedPtr<FEditableConstraintItem> ConstraintItem;
+	TWeakPtr<SConstraintsEditionWidget> ConstraintsWidget;
+};
+
+/**
+ * SConstraintsEditionWidget
+ */
+
+class CONTROLRIGEDITOR_API SConstraintsEditionWidget : public SCompoundWidget
+{
+public:
+	SLATE_BEGIN_ARGS(SConstraintsEditionWidget)	{}
+	SLATE_END_ARGS()
+
+	/** Constructs this widget with InArgs */
+	void Construct(const FArguments& InArgs);
+	virtual ~SConstraintsEditionWidget() override;
+
+	/** Override for SWidget::Tick */
+	virtual void Tick(const FGeometry& AllottedGeometry, const double InCurrentTime, const float InDeltaTime) override;
+
+	/** Invalidates the constraint list for further rebuild. */
+	void InvalidateConstraintList();
+	
+	/** Rebuild the constraint list based on the current selection. */
+	void RefreshConstraintList();
+
+	/** Triggers a constraint list invalidation when selection in the level viewport. */
+	void OnActorSelectionChanged(const TArray<UObject*>& NewSelection, bool bForceRefresh);
+
+	/**  */
+	bool IsFirstItem(const TSharedPtr<FEditableConstraintItem>& Item) const;
+	bool IsLastItem(const TSharedPtr<FEditableConstraintItem>& Item) const;
+
+	/**  */
+	void MoveItemUp(const TSharedPtr<FEditableConstraintItem>& Item);
+	void MoveItemDown(const TSharedPtr<FEditableConstraintItem>& Item);
+
+	/**  */
+	void RemoveItem(const TSharedPtr<FEditableConstraintItem>& Item);
+
+private:
+	/** Types */
+	using ItemSharedPtr = TSharedPtr<FEditableConstraintItem>;
+	using ConstraintItemListView = SListView< ItemSharedPtr >;
+
+	/** Generates a widget for the specified item */
+	TSharedRef<ITableRow> OnGenerateWidgetForItem(ItemSharedPtr InItem, const TSharedRef<STableViewBase>& OwnerTable);
+
+	void RegisterSelectionChanged();
+	void UnregisterSelectionChanged();
+	
+	/** List view that shows constraint types*/
+	TSharedPtr< ConstraintItemListView > ListView;
+
+	/** List of constraint types */
+	TArray< ItemSharedPtr > ListItems;
+
+	/** Boolean used to handle the items list refresh. */
+	bool bNeedsRefresh = false;
+
+	FDelegateHandle OnSelectionChangedHandle;
+};
diff --git a/Engine/Source/Runtime/AnimationCore/Public/Constraint.h b/Engine/Source/Runtime/AnimationCore/Public/Constraint.h
--- a/Engine/Source/Runtime/AnimationCore/Public/Constraint.h
+++ b/Engine/Source/Runtime/AnimationCore/Public/Constraint.h
@@ -390,7 +390,8 @@ enum class ETransformConstraintType : uint8
 	Translation,
 	Rotation,
 	Scale, 
-	Parent
+	Parent,
+	LookAt
 };
 
 /** A description of how to apply a simple transform constraint */
diff --git a/Engine/Source/Runtime/Experimental/Animation/Constraints/Constraints.Build.cs b/Engine/Source/Runtime/Experimental/Animation/Constraints/Constraints.Build.cs
new file mode 100644
--- /dev/null
+++ b/Engine/Source/Runtime/Experimental/Animation/Constraints/Constraints.Build.cs
@@ -0,0 +1,23 @@
+// Copyright Epic Games, Inc. All Rights Reserved.
+
+using UnrealBuildTool;
+
+public class Constraints : ModuleRules
+{
+	public Constraints(ReadOnlyTargetRules Target) : base(Target)
+	{
+		PCHUsage = ModuleRules.PCHUsageMode.UseExplicitOrSharedPCHs;
+		PrivateIncludePaths.Add("Runtime/Experimental/Animation/Constraints/Private");
+		PublicIncludePaths.Add(ModuleDirectory + "/Public");
+
+		PublicDependencyModuleNames.AddRange(
+			new string[]
+			{
+				"Core",
+				"CoreUObject",
+				"Engine",
+				"AnimationCore"
+			}
+		);
+	}
+}
\ No newline at end of file
diff --git a/Engine/Source/Runtime/Experimental/Animation/Constraints/Private/ConstraintsActor.cpp b/Engine/Source/Runtime/Experimental/Animation/Constraints/Private/ConstraintsActor.cpp
new file mode 100644
--- /dev/null
+++ b/Engine/Source/Runtime/Experimental/Animation/Constraints/Private/ConstraintsActor.cpp
@@ -0,0 +1,75 @@
+﻿// Copyright Epic Games, Inc. All Rights Reserved.
+
+
+#include "ConstraintsActor.h"
+
+#include "ConstraintsManager.h"
+
+// Sets default values
+AConstraintsActor::AConstraintsActor(const FObjectInitializer& ObjectInitializer)
+	: Super(ObjectInitializer)
+{
+	// Set this actor to call Tick() every frame.  You can turn this off to improve performance if you don't need it.
+	PrimaryActorTick.bCanEverTick = true;
+	
+#if WITH_EDITORONLY_DATA
+	bLockLocation = true;
+	bHiddenEd = true;
+#endif // WITH_EDITORONLY_DATA
+}
+
+AConstraintsActor::~AConstraintsActor()
+{}
+
+void AConstraintsActor::BeginDestroy()
+{
+	Super::BeginDestroy();
+}
+
+void AConstraintsActor::Destroyed()
+{
+	if(ConstraintsManager)
+	{
+		ConstraintsManager->Clear();
+	}
+	
+	Super::Destroyed();
+}
+
+void AConstraintsActor::BeginPlay()
+{
+	Super::BeginPlay();
+}
+
+void AConstraintsActor::Tick(float DeltaTime)
+{	
+	Super::Tick(DeltaTime);
+}
+
+void AConstraintsActor::PostRegisterAllComponents()
+{
+	Super::PostRegisterAllComponents();	
+	RegisterConstraintsTickFunctions();
+}
+
+void AConstraintsActor::RegisterConstraintsTickFunctions() const
+{
+	// note ensure that this is not done when useless ( use Level->bIsAssociatingLevel ?)
+	if (ULevel* Level = GetLevel())
+	{
+		if (ConstraintsManager)
+		{
+			// remove invalid pointers
+			ConstraintsManager->Constraints.RemoveAll( [](const TObjectPtr<UTickableConstraint>& InConstraint)
+			{
+				return !IsValid(InConstraint);
+			});
+			
+			// ensure registration
+			for (TObjectPtr<UTickableConstraint> ConstPtr: ConstraintsManager->Constraints)
+			{
+				ConstPtr->ConstraintTick.RegisterTickFunction(Level);
+			}
+		}
+	}
+}
\ No newline at end of file
diff --git a/Engine/Source/Runtime/Experimental/Animation/Constraints/Private/ConstraintsManager.cpp b/Engine/Source/Runtime/Experimental/Animation/Constraints/Private/ConstraintsManager.cpp
new file mode 100644
--- /dev/null
+++ b/Engine/Source/Runtime/Experimental/Animation/Constraints/Private/ConstraintsManager.cpp
@@ -0,0 +1,417 @@
+﻿// Copyright Epic Games, Inc. All Rights Reserved.
+
+#include "ConstraintsManager.h"
+
+#include "ConstraintsActor.h"
+
+#include "Algo/Copy.h"
+#include "Algo/StableSort.h"
+
+
+/** 
+ * FConstraintTickFunction
+ **/
+
+FConstraintTickFunction::FConstraintTickFunction()
+{
+	TickGroup = TG_PrePhysics;
+	bCanEverTick = true;
+	bStartWithTickEnabled = true;
+	bHighPriority = true;
+}
+
+FConstraintTickFunction::~FConstraintTickFunction()
+{}
+
+void FConstraintTickFunction::ExecuteTick(
+	float DeltaTime,
+	ELevelTick TickType,
+	ENamedThreads::Type CurrentThread,
+	const FGraphEventRef& MyCompletionGraphEvent)
+{
+	EvaluateFunctions();
+}
+
+void FConstraintTickFunction::RegisterFunction(ConstraintFunction InConstraint)
+{
+	ConstraintFunctions.Add(InConstraint);
+}
+
+void FConstraintTickFunction::EvaluateFunctions() const
+{
+	for (const ConstraintFunction& Function: ConstraintFunctions)
+	{
+		Function();
+	}	
+}
+
+FString FConstraintTickFunction::DiagnosticMessage()
+{
+	if(!Constraint.IsValid())
+	{
+		return FString::Printf(TEXT("FConstraintTickFunction::Tick[%p]"), this);	
+	}
+
+#if WITH_EDITOR
+	return FString::Printf(TEXT("FConstraintTickFunction::Tick[%p] (%s)"), this, *Constraint->GetLabel().ToString());
+#else
+	return FString::Printf(TEXT("FConstraintTickFunction::Tick[%p] (%s)"), this, *Constraint->GetName());
+#endif
+}
+
+/** 
+ * UTickableConstraint
+ **/
+
+#if WITH_EDITOR
+FName UTickableConstraint::GetLabel() const
+{
+	return GetFName();
+}
+#endif
+
+/** 
+ * UConstraintsManager
+ **/
+
+UConstraintsManager::UConstraintsManager()
+	: Level(nullptr)
+{}
+
+UConstraintsManager::~UConstraintsManager()
+{}
+
+void UConstraintsManager::Init(ULevel* InLevel)
+{
+	Level = InLevel;
+}
+
+UConstraintsManager* UConstraintsManager::Get(UWorld* InWorld)
+{
+	// look for ConstraintsActor and return its manager
+	if (UConstraintsManager* Manager = Find(InWorld))
+	{
+		return Manager;
+	}
+
+	// create a new ConstraintsActor
+	AConstraintsActor* ConstraintsActor = InWorld->SpawnActor<AConstraintsActor>();
+#if WITH_EDITOR
+	ConstraintsActor->SetActorLabel(TEXT("Constraints Manager"));
+#endif // WITH_EDITOR
+	ConstraintsActor->ConstraintsManager = NewObject<UConstraintsManager>(ConstraintsActor);
+	
+	ULevel* Level = InWorld->GetCurrentLevel();
+	ConstraintsActor->ConstraintsManager->Init(Level);
+	ConstraintsActor->ConstraintsManager->Level = Level;
+
+	return ConstraintsActor->ConstraintsManager;
+}
+
+UConstraintsManager* UConstraintsManager::Find(const UWorld* InWorld)
+{
+	// should we work with the persistent level?
+	const ULevel* Level = InWorld->GetCurrentLevel();
+
+	// look for first ConstraintsActor
+	auto FindFirstConstraintsActor = [](const ULevel* Level)
+	{
+		const int32 Index = Level->Actors.IndexOfByPredicate([](const AActor* Actor)
+		{
+			return IsValid(Actor) && Actor->IsA<AConstraintsActor>();
+		} );
+
+		return Index != INDEX_NONE ? Cast<AConstraintsActor>(Level->Actors[Index]) : nullptr;
+	};
+
+	const AConstraintsActor* ConstraintsActor = FindFirstConstraintsActor(Level);
+	return ConstraintsActor ? ConstraintsActor->ConstraintsManager : nullptr;
+}
+
+// we want InFunctionToTickBefore to tick first = InFunctionToTickBefore is a prerex of InFunctionToTickAfter
+void UConstraintsManager::SetConstraintDependencies(
+	FConstraintTickFunction* InFunctionToTickBefore,
+	FConstraintTickFunction* InFunctionToTickAfter)
+{
+	// look for child tick function in in parent's prerequisites. 
+	const TArray<FTickPrerequisite>& ParentPrerequisites = InFunctionToTickAfter->GetPrerequisites();
+	const bool bIsChildAPrerexOfParent = ParentPrerequisites.ContainsByPredicate([InFunctionToTickBefore](const FTickPrerequisite& Prerex)
+	{
+		return Prerex.PrerequisiteTickFunction == InFunctionToTickBefore;
+	});
+	
+	// child tick function is already a prerex -> parent already ticks after child
+	if (bIsChildAPrerexOfParent)
+	{
+		return;
+	}
+
+	// look for parent tick function in in child's prerequisites
+	const TArray<FTickPrerequisite>& ChildPrerequisites = InFunctionToTickBefore->GetPrerequisites();
+	const bool bIsParentAPrerexOfChild = ChildPrerequisites.ContainsByPredicate([InFunctionToTickAfter](const FTickPrerequisite& Prerex)
+	{
+		return Prerex.PrerequisiteTickFunction == InFunctionToTickAfter;
+	});
+	
+	// parent tick function is a prerex of the child tick function (child ticks after parent)
+	// so remove it before setting new dependencies.
+	if (bIsParentAPrerexOfChild)
+	{
+		InFunctionToTickBefore->RemovePrerequisite(this, *InFunctionToTickAfter);
+	}
+	
+	// set dependency
+	InFunctionToTickAfter->AddPrerequisite(this, *InFunctionToTickBefore);
+}
+
+void UConstraintsManager::Clear()
+{
+	Constraints.Empty();
+}
+
+void UConstraintsManager::Dump() const
+{
+	UE_LOG(LogTemp, Error, TEXT("nb consts = %d"), Constraints.Num());
+	for (const TObjectPtr<UTickableConstraint>& Constraint: Constraints)
+	{
+		if (IsValid(Constraint))
+		{
+			UE_LOG(LogTemp, Warning, TEXT("\t%s (target hash = %u)"),
+				*Constraint->GetName(), Constraint->GetTargetHash());
+		}
+	}
+}
+
+FConstraintsManagerController& FConstraintsManagerController::Get(UWorld* InWorld)
+{
+	static FConstraintsManagerController Singleton;
+	Singleton.World = InWorld;
+	return Singleton;
+}
+
+UConstraintsManager* FConstraintsManagerController::GetManager() const
+{
+	if (!World)
+	{
+		return nullptr;
+	}
+	
+	// look for ConstraintsActor and return its manager
+	if (UConstraintsManager* Manager = FindManager())
+	{
+		return Manager;
+	}
+
+	// create a new ConstraintsActor
+	AConstraintsActor* ConstraintsActor = World->SpawnActor<AConstraintsActor>();
+#if WITH_EDITOR
+	ConstraintsActor->SetActorLabel(TEXT("Constraints Manager"));
+#endif // WITH_EDITOR
+	ConstraintsActor->ConstraintsManager = NewObject<UConstraintsManager>(ConstraintsActor);
+	
+	ULevel* Level = World->GetCurrentLevel();
+	ConstraintsActor->ConstraintsManager->Level = Level;
+
+	return ConstraintsActor->ConstraintsManager;
+}
+
+UConstraintsManager* FConstraintsManagerController::FindManager() const
+{
+	if (!World)
+	{
+		return nullptr;
+	}
+	
+	// should we work with the persistent level?
+	const ULevel* Level = World->GetCurrentLevel();
+
+	// look for first ConstraintsActor
+	auto FindFirstConstraintsActor = [](const ULevel* Level)
+	{
+		const int32 Index = Level->Actors.IndexOfByPredicate([](const AActor* Actor)
+		{
+			return IsValid(Actor) && Actor->IsA<AConstraintsActor>();
+		} );
+
+		return Index != INDEX_NONE ? Cast<AConstraintsActor>(Level->Actors[Index]) : nullptr;
+	};
+
+	const AConstraintsActor* ConstraintsActor = FindFirstConstraintsActor(Level);
+	return ConstraintsActor ? ConstraintsActor->ConstraintsManager : nullptr;
+}
+
+void FConstraintsManagerController::DestroyManager() const
+{
+	if (!World)
+	{
+		return;
+	}
+	
+	// should we work with the persistent level?
+	const ULevel* Level = World->GetCurrentLevel();
+
+	// note there should be only one...
+	TArray<AActor*> ConstraintsActorsToRemove;
+	Algo::CopyIf(Level->Actors, ConstraintsActorsToRemove, [](const AActor* Actor)
+	{
+		return IsValid(Actor) && Actor->IsA<AConstraintsActor>();
+	} );
+
+	for (AActor* ConstraintsActor: ConstraintsActorsToRemove)
+	{
+		World->EditorDestroyActor(ConstraintsActor, true);
+	}
+}
+
+void FConstraintsManagerController::AddConstraint(UTickableConstraint* InConstraint) const
+{
+	if (!InConstraint)
+	{
+		return;
+	}
+	
+	UConstraintsManager* Manager = GetManager();
+	if (!Manager)
+	{
+		return;
+	}
+
+	// TODO handle transaction
+	Manager->Constraints.Emplace(InConstraint);
+
+	InConstraint->ConstraintTick.Constraint = InConstraint;
+	InConstraint->ConstraintTick.RegisterFunction(InConstraint->GetFunction());
+	InConstraint->ConstraintTick.RegisterTickFunction(World->GetCurrentLevel());
+}
+
+int32 FConstraintsManagerController::GetConstraintIndex(const FName& InConstraintName) const
+{
+	const UConstraintsManager* Manager = FindManager();
+	if (!Manager)
+	{
+		return INDEX_NONE;
+	}
+	
+	return Manager->Constraints.IndexOfByPredicate([InConstraintName](const TObjectPtr<UTickableConstraint>& Constraint)
+	{
+		return 	Constraint->GetFName() == InConstraintName;
+	} );
+}
+	
+void FConstraintsManagerController::RemoveConstraint(const FName& InConstraintName) const
+{
+	const int32 Index = GetConstraintIndex(InConstraintName);
+	if (Index == INDEX_NONE)
+	{
+		return;	
+	}
+	
+	RemoveConstraint(Index);
+}
+
+void FConstraintsManagerController::RemoveConstraint(const int32 InConstraintIndex) const
+{
+	UConstraintsManager* Manager = FindManager();
+	if (!Manager)
+	{
+		return;
+	}
+	
+	if (!Manager->Constraints.IsValidIndex(InConstraintIndex))
+	{
+		return;
+	}
+
+	// TODO handle transaction
+	
+	Manager->Constraints[InConstraintIndex]->ConstraintTick.UnRegisterTickFunction();
+	Manager->Constraints.RemoveAt(InConstraintIndex);
+
+	// destroy constraints actor if no constraints left
+	if (Manager->Constraints.IsEmpty())
+	{
+		DestroyManager();
+	}
+}
+
+TArray< TObjectPtr<UTickableConstraint> > FConstraintsManagerController::GetParentConstraints(
+	const uint32 InTargetHash,
+	const bool bSorted) const
+{
+		static const TArray< TObjectPtr<UTickableConstraint> > DummyArray;
+
+	const UConstraintsManager* Manager = FindManager();
+	if (!Manager)
+	{
+		return DummyArray;
+	}
+	
+	if (InTargetHash == 0)
+	{
+		return DummyArray;
+	}
+
+	using ConstraintPtr = TObjectPtr<UTickableConstraint>;
+	TArray< TObjectPtr<UTickableConstraint> > FilteredConstraints =
+	Manager->Constraints.FilterByPredicate( [InTargetHash](const ConstraintPtr& Constraint)
+	{
+		return Constraint->GetTargetHash() == InTargetHash;
+	} );
+	
+	if (bSorted)
+	{
+		// LHS ticks before RHS = LHS is a prerex of RHS 
+		auto TicksBefore = [](const UTickableConstraint& LHS, const UTickableConstraint& RHS)
+		{
+			const TArray<FTickPrerequisite>& RHSPrerex = RHS.ConstraintTick.GetPrerequisites();
+			const FConstraintTickFunction& LHSTickFunction = LHS.ConstraintTick;
+			const bool bIsLHSAPrerexOfRHS = RHSPrerex.ContainsByPredicate([&LHSTickFunction](const FTickPrerequisite& Prerex)
+			{
+				return Prerex.PrerequisiteTickFunction == &LHSTickFunction;
+			});
+			return bIsLHSAPrerexOfRHS;
+		};
+		
+		Algo::StableSort(FilteredConstraints, [TicksBefore](const ConstraintPtr& LHS, const ConstraintPtr& RHS)
+		{
+			return TicksBefore(*LHS, *RHS);
+		});
+	}
+	
+	return FilteredConstraints;
+}
+
+void FConstraintsManagerController::SetConstraintsDependencies(
+	const FName& InNameToTickBefore,
+	const FName& InNameToTickAfter) const
+{
+	UConstraintsManager* Manager = FindManager();
+	if (!Manager)
+	{
+		return;
+	}
+
+	const int32 IndexBefore = GetConstraintIndex(InNameToTickBefore);
+	const int32 IndexAfter = GetConstraintIndex(InNameToTickAfter);
+	if (IndexBefore == INDEX_NONE || IndexAfter == INDEX_NONE || IndexAfter == IndexBefore)
+	{
+		return;
+	}
+
+	FConstraintTickFunction& FunctionToTickBefore = Manager->Constraints[IndexBefore]->ConstraintTick;
+	FConstraintTickFunction& FunctionToTickAfter = Manager->Constraints[IndexAfter]->ConstraintTick;
+
+	Manager->SetConstraintDependencies( &FunctionToTickBefore, &FunctionToTickAfter);
+}
+
+const TArray< TObjectPtr<UTickableConstraint> >& FConstraintsManagerController::GetConstrainsArray() const
+{
+	UConstraintsManager* Manager = FindManager();
+	if (!Manager)
+	{
+		static const TArray< TObjectPtr<UTickableConstraint> > Empty;
+		return Empty;
+	}
+
+	return Manager->Constraints;
+}
\ No newline at end of file
diff --git a/Engine/Source/Runtime/Experimental/Animation/Constraints/Private/TransformConstraint.cpp b/Engine/Source/Runtime/Experimental/Animation/Constraints/Private/TransformConstraint.cpp
new file mode 100644
--- /dev/null
+++ b/Engine/Source/Runtime/Experimental/Animation/Constraints/Private/TransformConstraint.cpp
@@ -0,0 +1,473 @@
+﻿// Copyright Epic Games, Inc. All Rights Reserved.
+
+#include "TransformConstraint.h"
+
+#include "Animation/Constraints/Public/TransformableHandle.h"
+#include "ConstraintsManager.h"
+#include "TransformableRegistry.h"
+
+/** 
+ * UTickableTransformConstraint
+ **/
+
+/** @todo remove to use something else. */
+int64 UTickableTransformConstraint::GetType() const
+{
+	return static_cast<int64>(Type);
+}
+
+#if WITH_EDITOR
+FName UTickableTransformConstraint::GetLabel() const
+{
+	const FName ParentName = ParentTRSHandle->IsValid() ? ParentTRSHandle->GetName() : NAME_None;
+	const FName ChildName = ChildTRSHandle->IsValid() ? ChildTRSHandle->GetName() : NAME_None;
+	const FString Label = FString::Printf(TEXT("%s -> %s"), *ParentName.ToString(), *ChildName.ToString() );
+	return FName(*Label);
+}
+#endif
+
+void UTickableTransformConstraint::Setup()
+{
+	if (!ParentTRSHandle->IsValid() || !ChildTRSHandle->IsValid())
+	{
+		// handle error
+		return;
+	}
+	
+	ComputeOffset();
+	SetupDependencies();
+}
+
+void UTickableTransformConstraint::SetupDependencies()
+{
+	FTickFunction* ParentTickFunction = ParentTRSHandle->IsValid() ? ParentTRSHandle->GetTickFunction() : nullptr;
+	if (ParentTickFunction)
+	{
+		// manage dependencies
+		// force ConstraintTickFunction to tick after InParent does.
+		// Note that this might not register anything if the parent can't tick (static meshes for instance)
+		ConstraintTick.AddPrerequisite(ParentTRSHandle->GetPrerequisiteObject(), *ParentTickFunction);
+	}
+	
+	// TODO also check for cycle dependencies here
+	FTickFunction* ChildTickFunction = ChildTRSHandle ? ChildTRSHandle->GetTickFunction() : nullptr;
+	if (ChildTickFunction && ChildTickFunction != ParentTickFunction)
+	{
+		// force InParent to tick after ConstraintTickFunction does.
+		// Note that this might not register anything if the child can't tick (static meshes for instance)
+		ChildTickFunction->AddPrerequisite(this, ConstraintTick);
+	}
+}
+
+void UTickableTransformConstraint::PostLoad()
+{
+	Super::PostLoad();
+
+	ConstraintTick.RegisterFunction(GetFunction() );
+}
+
+void UTickableTransformConstraint::PostDuplicate(bool bDuplicateForPIE)
+{
+	Super::PostDuplicate(bDuplicateForPIE);
+}
+
+uint32 UTickableTransformConstraint::GetTargetHash() const
+{
+	return ChildTRSHandle->IsValid() ? ChildTRSHandle->GetHash() : 0;
+}
+
+void UTickableTransformConstraint::SetChildGlobalTransform(const FTransform& InGlobal) const
+{
+	if(ChildTRSHandle->IsValid())
+	{
+		ChildTRSHandle->SetTransform(InGlobal);
+	}
+}
+
+FTransform UTickableTransformConstraint::GetChildGlobalTransform() const
+{
+	return ChildTRSHandle->IsValid() ? ChildTRSHandle->GetTransform() : FTransform::Identity;
+}
+
+FTransform UTickableTransformConstraint::GetParentGlobalTransform() const
+{
+	return ParentTRSHandle->IsValid() ? ParentTRSHandle->GetTransform() : FTransform::Identity;
+}
+
+/** 
+ * UTickableTranslationConstraint
+ **/
+
+UTickableTranslationConstraint::UTickableTranslationConstraint()
+{
+	Type = ETransformConstraintType::Translation;
+}
+
+void UTickableTranslationConstraint::ComputeOffset()
+{
+	UClass* StaticClass = UTickableTranslationConstraint::StaticClass();
+	
+	const FTransform InitParentTransform = GetParentGlobalTransform();
+	const FTransform InitChildTransform = GetChildGlobalTransform();
+	
+	OffsetTranslation = FVector::ZeroVector;
+	if (bMaintainOffset)
+	{
+		OffsetTranslation = InitChildTransform.GetLocation() - InitParentTransform.GetLocation();
+	}	
+}
+
+FConstraintTickFunction::ConstraintFunction UTickableTranslationConstraint::GetFunction() const
+{
+	return [this]()
+	{
+		const float ClampedWeight = FMath::Clamp<float>(Weight, 0.f, 1.f);
+		if (ClampedWeight < KINDA_SMALL_NUMBER)
+		{
+			return;
+		}
+
+		const FTransform ParentTransform = GetParentGlobalTransform();
+		FTransform Transform = GetChildGlobalTransform();
+		FVector NewTranslation = ParentTransform.GetLocation() + OffsetTranslation;
+		if (ClampedWeight < 1.0f - KINDA_SMALL_NUMBER)
+		{
+			NewTranslation = FMath::Lerp<FVector>(Transform.GetLocation(), NewTranslation, ClampedWeight);
+		}
+		Transform.SetLocation(NewTranslation);
+			
+		SetChildGlobalTransform(Transform);
+	};
+}
+
+/** 
+ * UTickableRotationConstraint
+ **/
+
+UTickableRotationConstraint::UTickableRotationConstraint()
+{
+	Type = ETransformConstraintType::Rotation;
+}
+
+void UTickableRotationConstraint::ComputeOffset()
+{
+	const FTransform InitParentTransform = GetParentGlobalTransform();
+	const FTransform InitChildTransform = GetChildGlobalTransform();
+	
+	OffsetRotation = FQuat::Identity;
+	if (bMaintainOffset)
+	{
+		OffsetRotation = InitParentTransform.GetRotation().Inverse() * InitChildTransform.GetRotation();
+		OffsetRotation.Normalize();
+	}
+}
+
+FConstraintTickFunction::ConstraintFunction UTickableRotationConstraint::GetFunction() const
+{
+	return [this]()
+	{
+		const float ClampedWeight = FMath::Clamp<float>(Weight, 0.f, 1.f);
+		if (ClampedWeight < KINDA_SMALL_NUMBER)
+		{
+			return;
+		}
+		
+		const FTransform ParentTransform = GetParentGlobalTransform();
+		FTransform Transform = GetChildGlobalTransform();
+
+		FQuat NewRotation = ParentTransform.GetRotation() * OffsetRotation;
+		if (ClampedWeight < 1.0f - KINDA_SMALL_NUMBER)
+		{
+			NewRotation = FQuat::Slerp(Transform.GetRotation(), NewRotation, ClampedWeight);
+		}
+		Transform.SetRotation(NewRotation);
+		
+		SetChildGlobalTransform(Transform);
+	};
+}
+
+/** 
+ * UTickableScaleConstraint
+ **/
+
+UTickableScaleConstraint::UTickableScaleConstraint()
+{
+	Type = ETransformConstraintType::Scale;
+}
+
+void UTickableScaleConstraint::ComputeOffset()
+{
+	const FTransform InitParentTransform = GetParentGlobalTransform();
+	const FTransform InitChildTransform = GetChildGlobalTransform();
+	
+	OffsetScale = FVector::OneVector;
+	if (bMaintainOffset)
+	{
+		const FVector InitParentScale = InitParentTransform.GetScale3D();
+		OffsetScale = InitChildTransform.GetScale3D();
+		OffsetScale[0] = FMath::Abs(InitParentScale[0]) > KINDA_SMALL_NUMBER ? OffsetScale[0] / InitParentScale[0] : 0.f;
+		OffsetScale[1] = FMath::Abs(InitParentScale[1]) > KINDA_SMALL_NUMBER ? OffsetScale[1] / InitParentScale[1] : 0.f;
+		OffsetScale[2] = FMath::Abs(InitParentScale[2]) > KINDA_SMALL_NUMBER ? OffsetScale[2] / InitParentScale[2] : 0.f;
+	}
+}
+
+FConstraintTickFunction::ConstraintFunction UTickableScaleConstraint::GetFunction() const
+{
+	return [this]()
+	{
+		const float ClampedWeight = FMath::Clamp<float>(Weight, 0.f, 1.f);
+		if (ClampedWeight < KINDA_SMALL_NUMBER)
+		{
+			return;
+		}
+		
+		const FTransform ParentTransform = GetParentGlobalTransform();
+		FTransform Transform = GetChildGlobalTransform();
+		FVector NewScale = ParentTransform.GetScale3D() * OffsetScale;
+		if (ClampedWeight < 1.0f - KINDA_SMALL_NUMBER)
+		{
+			NewScale = FMath::Lerp<FVector>(Transform.GetScale3D(), NewScale, ClampedWeight);
+		}
+		Transform.SetScale3D(NewScale);
+			
+		SetChildGlobalTransform(Transform);
+	};
+}
+
+/** 
+ * UTickableParentConstraint
+ **/
+
+UTickableParentConstraint::UTickableParentConstraint()
+{
+	Type = ETransformConstraintType::Parent;
+}
+
+void UTickableParentConstraint::ComputeOffset()
+{
+	const FTransform InitParentTransform = GetParentGlobalTransform();
+	const FTransform InitChildTransform = GetChildGlobalTransform();
+	
+	OffsetTransform = FTransform::Identity;
+	if (bMaintainOffset)
+	{
+		OffsetTransform = InitChildTransform.GetRelativeTransform(InitParentTransform);
+	}
+}
+
+FConstraintTickFunction::ConstraintFunction UTickableParentConstraint::GetFunction() const
+{
+	return [this]()
+	{
+		const float ClampedWeight = FMath::Clamp<float>(Weight, 0.f, 1.f);
+		if (ClampedWeight < KINDA_SMALL_NUMBER)
+		{
+			return;
+		}
+		
+		const FTransform ParentTransform = ParentTRSHandle->GetTransform();
+
+		FTransform TargetTransform = OffsetTransform * ParentTransform;
+		if (ClampedWeight < 1.0f - KINDA_SMALL_NUMBER)
+		{
+			const  FTransform Transform = GetChildGlobalTransform();
+			TargetTransform.SetLocation(
+				FMath::Lerp<FVector>(Transform.GetLocation(), TargetTransform.GetLocation(), ClampedWeight));
+			TargetTransform.SetRotation(
+				FQuat::Slerp(Transform.GetRotation(), TargetTransform.GetRotation(), ClampedWeight));
+			TargetTransform.SetScale3D(
+				FMath::Lerp<FVector>(Transform.GetScale3D(), TargetTransform.GetScale3D(), ClampedWeight));
+		}
+		
+		SetChildGlobalTransform(TargetTransform);
+	};
+}
+
+/** 
+ * UTickableLookAtConstraint
+ **/
+
+UTickableLookAtConstraint::UTickableLookAtConstraint()
+{
+	Type = ETransformConstraintType::LookAt;
+}
+
+void UTickableLookAtConstraint::ComputeOffset()
+{
+	// TODO compute offset
+}
+
+FConstraintTickFunction::ConstraintFunction UTickableLookAtConstraint::GetFunction() const
+{
+	// @todo handle weight here
+	return [this]()
+	{
+		const FTransform ParentTransform = GetParentGlobalTransform();
+		const FTransform ChildTransform = GetChildGlobalTransform();
+		
+		const FVector LookAtDir = ParentTransform.GetLocation() - ChildTransform.GetLocation();
+		const FRotator LookAtRotation = LookAtDir.Rotation();
+
+		FTransform Transform = GetChildGlobalTransform();
+		Transform.SetRotation(LookAtRotation.Quaternion());
+		
+		SetChildGlobalTransform(Transform);
+	};
+}
+
+/** 
+ * TransformConstraintUtils
+ **/
+
+namespace
+{
+
+UTransformableHandle* GetHandle(const AActor* InActor, UObject* Outer)
+{
+	// look for customized transform handle
+	const FTransformableRegistry& Registry = FTransformableRegistry::Get();
+	if (const FTransformableRegistry::CreateHandleFuncT CreateFunction = Registry.GetCreateFunction(InActor->GetClass()))
+	{
+		return CreateFunction(InActor, Outer);
+	}
+
+	// scene component handle
+	UTransformableComponentHandle* ComponentHandle = NewObject<UTransformableComponentHandle>(Outer);
+	ComponentHandle->Component = InActor->GetRootComponent();
+	
+	return ComponentHandle;
+}
+	
+uint32 GetConstrainableHash(const AActor* InActor)
+{
+	// look for customized hash function
+	const FTransformableRegistry& Registry = FTransformableRegistry::Get();
+	if (const FTransformableRegistry::GetHashFuncT HashFunction = Registry.GetHashFunction(InActor->GetClass()))
+	{
+		return HashFunction(InActor);
+	}
+
+	// scene component hash
+	const uint32 ComponentHash = GetTypeHash(InActor->GetRootComponent());
+	return ComponentHash;
+}
+	
+}
+
+void TransformConstraintUtils::GetParentConstraints(
+	UWorld* InWorld,
+	const AActor* InChild,
+	TArray< TObjectPtr<UTickableConstraint> >& OutConstraints)
+{
+	if (!InWorld || !InChild)
+	{
+		return;
+	}
+
+	const uint32 ChildHash = GetConstrainableHash(InChild);
+	if (ChildHash == 0)
+	{
+		return;
+	}
+	
+	static constexpr bool bSorted = true;
+	const FConstraintsManagerController& Controller = FConstraintsManagerController::Get(InWorld);
+	OutConstraints.Append(Controller.GetParentConstraints(ChildHash, bSorted));
+}
+
+bool TransformConstraintUtils::Create(
+	UWorld* InWorld,
+	const AActor* InParent,
+	const AActor* InChild,
+	const ETransformConstraintType InType,
+	const bool bMaintainOffset)
+{
+	// SANITY CHECK
+	if (!InWorld || !InParent || !InChild)
+	{
+		UE_LOG(LogTemp, Error, TEXT("TransformConstraintUtils::Create sanity check failed."));
+		return false;
+	}
+	
+	UConstraintsManager* ConstraintsManager = UConstraintsManager::Get(InWorld);
+	if (!ConstraintsManager)
+	{
+		UE_LOG(LogTemp, Error, TEXT("TransformConstraintUtils::Create constraint manager is null."));
+		return false;
+	}
+
+	UTransformableHandle* ParentHandle = GetHandle(InParent, ConstraintsManager);
+	UTransformableHandle* ChildHandle = GetHandle(InChild, ConstraintsManager);
+	
+	// DO IT
+	if (ParentHandle->IsValid() && ChildHandle->IsValid())
+	{
+		return AddConst(InWorld, ParentHandle, ChildHandle, InType, bMaintainOffset);
+	}
+	
+	return false;
+}
+
+bool TransformConstraintUtils::AddConst(
+	UWorld* InWorld,
+	UTransformableHandle* InParentHandle,
+	UTransformableHandle* InChildHandle,
+	const ETransformConstraintType InType,
+	const bool bMaintainOffset)
+{
+	const bool bIsValidParent = InParentHandle && InParentHandle->IsValid();
+	const bool bIsValidChild = InChildHandle && InChildHandle->IsValid();
+	if (!bIsValidParent || !bIsValidChild)
+	{
+		UE_LOG(LogTemp, Error, TEXT("TransformConstraintUtils::AddConst error creating constraint"));
+		return false;
+	}
+
+	const UEnum* ETransformConstraintTypeEnum = StaticEnum<ETransformConstraintType>();
+	if (!ETransformConstraintTypeEnum->IsValidEnumValue(static_cast<int64>(InType)))
+	{
+		UE_LOG(LogTemp, Error, TEXT("Constraint Type %d not recognized"), InType);
+		return false;
+	}
+	
+	// unique name (we may want to use another approach here to manage uniqueness)
+	const FString ConstraintTypeStr = ETransformConstraintTypeEnum->GetNameStringByValue((uint8)InType);
+	const FName BaseName(*FString::Printf(TEXT("%sConstraint"), *ConstraintTypeStr));
+
+	auto SetupConstraint = [InParentHandle, InChildHandle, bMaintainOffset](UTickableTransformConstraint* InConstraint)
+	{
+		InConstraint->ParentTRSHandle = InParentHandle;
+		InConstraint->ChildTRSHandle = InChildHandle;
+		InConstraint->bMaintainOffset = bMaintainOffset;
+		InConstraint->Setup();
+	};
+	
+	const FConstraintsManagerController& Controller = FConstraintsManagerController::Get(InWorld);
+	
+	UTickableTransformConstraint* Constraint = nullptr;
+	switch (InType)
+	{
+	case ETransformConstraintType::Translation:
+		Constraint = Controller.AllocateConstraintT<UTickableTranslationConstraint>(BaseName);
+		break;
+	case ETransformConstraintType::Rotation:
+		Constraint = Controller.AllocateConstraintT<UTickableRotationConstraint>(BaseName);
+		break;
+	case ETransformConstraintType::Scale:
+		Constraint = Controller.AllocateConstraintT<UTickableScaleConstraint>(BaseName);
+		break;
+	case ETransformConstraintType::Parent:
+		Constraint = Controller.AllocateConstraintT<UTickableParentConstraint>(BaseName);
+		break;
+	case ETransformConstraintType::LookAt:
+		Constraint = Controller.AllocateConstraintT<UTickableLookAtConstraint>(BaseName);
+		break;
+	default:
+		ensure(false);
+		return false;
+	}
+
+	SetupConstraint(Constraint);
+
+	return true;
+}
diff --git a/Engine/Source/Runtime/Experimental/Animation/Constraints/Private/TransformableHandle.cpp b/Engine/Source/Runtime/Experimental/Animation/Constraints/Private/TransformableHandle.cpp
new file mode 100644
--- /dev/null
+++ b/Engine/Source/Runtime/Experimental/Animation/Constraints/Private/TransformableHandle.cpp
@@ -0,0 +1,74 @@
+﻿// Copyright Epic Games, Inc. All Rights Reserved.
+
+
+#include "TransformableHandle.h"
+
+/**
+ * UTransformableHandle
+ */
+
+UTransformableHandle::~UTransformableHandle()
+{}
+
+void UTransformableHandle::PostLoad()
+{
+	Super::PostLoad();
+}
+
+/**
+ * UTransformableComponentHandle
+ */
+
+UTransformableComponentHandle::~UTransformableComponentHandle()
+{}
+
+void UTransformableComponentHandle::PostLoad()
+{
+	Super::PostLoad();
+}
+
+bool UTransformableComponentHandle::IsValid() const
+{
+	return Component.IsValid();
+}
+
+void UTransformableComponentHandle::SetTransform(const FTransform& InGlobal) const
+{
+	if(Component.IsValid())
+	{
+		Component->SetWorldTransform(InGlobal);
+	}
+}
+
+FTransform UTransformableComponentHandle::GetTransform() const
+{
+	return Component.IsValid() ? Component->GetComponentTransform() : FTransform::Identity;
+}
+
+UObject* UTransformableComponentHandle::GetPrerequisiteObject() const
+{
+	return Component.Get(); 
+}
+
+FTickFunction* UTransformableComponentHandle::GetTickFunction() const
+{
+	return Component.IsValid() ? &Component->PrimaryComponentTick : nullptr;
+}
+
+uint32 UTransformableComponentHandle::GetHash() const
+{
+	return Component.IsValid() ? GetTypeHash(Component.Get()) : 0;
+}
+
+#if WITH_EDITOR
+FName UTransformableComponentHandle::GetName() const
+{
+	if (!Component.IsValid())
+	{
+		return NAME_None;
+	}
+
+	const AActor* Actor = Component->GetOwner();
+	return Actor ? FName(*Actor->GetActorLabel()) : Component->GetFName();
+}
+#endif
\ No newline at end of file
diff --git a/Engine/Source/Runtime/Experimental/Animation/Constraints/Private/TransformableRegistry.cpp b/Engine/Source/Runtime/Experimental/Animation/Constraints/Private/TransformableRegistry.cpp
new file mode 100644
--- /dev/null
+++ b/Engine/Source/Runtime/Experimental/Animation/Constraints/Private/TransformableRegistry.cpp
@@ -0,0 +1,54 @@
+﻿// Copyright Epic Games, Inc. All Rights Reserved.
+
+#include "TransformableRegistry.h"
+
+FTransformableRegistry::~FTransformableRegistry() = default;
+
+FTransformableRegistry& FTransformableRegistry::Get()
+{
+	static FTransformableRegistry Singleton;
+	return Singleton;
+}
+
+void FTransformableRegistry::Register(
+	UClass* InClass,
+	CreateHandleFuncT InHandleFunc,
+	GetHashFuncT InHashFunc)
+{
+	const FTransformableInfo Info { InHandleFunc, InHashFunc};
+	Transformables.Emplace(InClass, Info);	
+}
+
+FTransformableRegistry::GetHashFuncT FTransformableRegistry::GetHashFunction(const UClass* InClass) const
+{
+	if (InClass)
+	{
+		if (const FTransformableInfo* Info = FindInfo(InClass))
+		{
+			return Info->GetHashFunc; 
+		}
+	}
+	return {};
+}
+
+FTransformableRegistry::CreateHandleFuncT FTransformableRegistry::GetCreateFunction(const UClass* InClass) const
+{
+	if (InClass)
+	{
+		if(const FTransformableInfo* Info = FindInfo(InClass))
+		{
+			return Info->CreateHandleFunc; 
+		}
+	}
+	return {};
+}
+
+const FTransformableRegistry::FTransformableInfo* FTransformableRegistry::FindInfo(const UClass* InClass) const
+{
+	if (const FTransformableInfo* Info = Transformables.Find(InClass))
+	{
+		return Info;
+	}
+
+	return nullptr;
+}
\ No newline at end of file
diff --git a/Engine/Source/Runtime/Experimental/Animation/Constraints/Public/ConstraintsActor.h b/Engine/Source/Runtime/Experimental/Animation/Constraints/Public/ConstraintsActor.h
new file mode 100644
--- /dev/null
+++ b/Engine/Source/Runtime/Experimental/Animation/Constraints/Public/ConstraintsActor.h
@@ -0,0 +1,39 @@
+﻿// Copyright Epic Games, Inc. All Rights Reserved.
+
+#pragma once
+
+#include "GameFramework/Actor.h"
+
+#include "ConstraintsActor.generated.h"
+
+class UConstraintsManager;
+
+UCLASS(NotPlaceable)
+class CONSTRAINTS_API AConstraintsActor : public AActor
+{
+	GENERATED_BODY()
+
+public:
+	// Sets default values for this actor's properties
+	AConstraintsActor(const FObjectInitializer& ObjectInitializer);
+	virtual ~AConstraintsActor();
+
+	virtual void BeginDestroy() override;
+	virtual void Destroyed() override;
+
+	virtual void PostRegisterAllComponents() override;
+	
+protected:
+	// Called when the game starts or when spawned
+	virtual void BeginPlay() override;
+
+public:
+	// Called every frame
+	virtual void Tick(float DeltaTime) override;
+
+	UPROPERTY()
+	TObjectPtr<UConstraintsManager> ConstraintsManager = nullptr;
+	
+private:
+	void RegisterConstraintsTickFunctions() const;
+};
diff --git a/Engine/Source/Runtime/Experimental/Animation/Constraints/Public/ConstraintsManager.h b/Engine/Source/Runtime/Experimental/Animation/Constraints/Public/ConstraintsManager.h
new file mode 100644
--- /dev/null
+++ b/Engine/Source/Runtime/Experimental/Animation/Constraints/Public/ConstraintsManager.h
@@ -0,0 +1,209 @@
+﻿// Copyright Epic Games, Inc. All Rights Reserved.
+
+#pragma once
+
+#include "UObject/Object.h"
+
+#include "ConstraintsManager.generated.h"
+
+class UTickableConstraint;
+
+/** 
+ * FConstraintTickFunction
+ * Represents the interface of constraint as a tick function. This allows both to evaluate a constraint in the
+ * UE ticking system but also to handle dependencies between parents/children and constraints between themselves
+ * using the tick prerequisites system.
+ **/
+
+USTRUCT()
+struct CONSTRAINTS_API FConstraintTickFunction : public FTickFunction
+{
+	GENERATED_USTRUCT_BODY()
+public:
+	FConstraintTickFunction();
+	~FConstraintTickFunction();
+
+	/* Begin FTickFunction Interface */
+	virtual void ExecuteTick(
+		float DeltaTime,
+		ELevelTick TickType,
+		ENamedThreads::Type CurrentThread,
+		const FGraphEventRef& MyCompletionGraphEvent) override;
+	
+	virtual FString DiagnosticMessage() override;
+	/* End FTickFunction Interface */
+
+	/** Callable function that represents the actual constraint. **/
+	using ConstraintFunction = TFunction<void()>;
+	
+	/** Register a callable function. **/
+	void RegisterFunction(ConstraintFunction InConstraint);
+	
+	/** Register a callable function. **/
+	void EvaluateFunctions() const;
+
+	/** Weak ptr to the Constraint holding this tick function. **/
+	TWeakObjectPtr<UTickableConstraint> Constraint;
+	
+	/** The list of the constraint functions that will be called within the tick function. **/
+	TArray<ConstraintFunction> ConstraintFunctions;
+};
+
+template<>
+struct TStructOpsTypeTraits<FConstraintTickFunction> : public TStructOpsTypeTraitsBase2<FConstraintTickFunction>
+{
+	enum
+	{
+		WithCopy = false
+	};
+};
+
+/** 
+ * UTickableConstraint
+ * Represents the basic interface of constraint within the constraints manager.
+ **/
+
+UCLASS(Abstract)
+class CONSTRAINTS_API UTickableConstraint : public UObject
+{
+	GENERATED_BODY()
+	
+public:
+	UTickableConstraint() {}
+	virtual ~UTickableConstraint() {}
+
+	/** Returns the actual function that the tick function needs to evaluate. */
+	virtual FConstraintTickFunction::ConstraintFunction GetFunction() const PURE_VIRTUAL(GetFunction, return {};);
+
+	/** Tick function that will be registered and evaluated. */
+	UPROPERTY()
+	FConstraintTickFunction ConstraintTick;
+
+	/** @todo document */
+	virtual uint32 GetTargetHash() const PURE_VIRTUAL(GetTargetHash, return 0;);
+	
+#if WITH_EDITOR
+	/** @todo document */
+	virtual FName GetLabel() const;
+#endif
+};
+
+/** 
+ * UConstraintsManager
+ * This object gathers the different constraints of the level and is held by the ConstraintsActor (unique in the level)
+ **/
+
+UCLASS()
+class CONSTRAINTS_API UConstraintsManager : public UObject
+{
+	GENERATED_BODY()
+public:
+
+	using ConstraintPtr = TObjectPtr<UTickableConstraint>;
+	
+	UConstraintsManager();
+	virtual ~UConstraintsManager();
+	
+	/** Get the existing Constraints Manager if existing or create a new one. */
+	static UConstraintsManager* Get(UWorld* InWorld);
+
+	/** Find the existing Constraints Manager. */
+	static UConstraintsManager* Find(const UWorld* InWorld);
+
+	/** @todo document */
+	void Init(ULevel* InLevel);
+	
+	/* Set tick dependencies between two constraints. */
+	void SetConstraintDependencies(
+		FConstraintTickFunction* InFunctionToTickBefore,
+		FConstraintTickFunction* InFunctionToTickAfter);
+
+	/** @todo document */
+	void Clear();
+	
+private:
+
+	/** @todo document */
+	void Dump() const;
+
+	/** @todo document */
+	UPROPERTY()
+	TObjectPtr<ULevel> Level = nullptr;
+
+	/** @todo document */
+	UPROPERTY()
+	TArray< TObjectPtr<UTickableConstraint> > Constraints;
+	
+	friend class FConstraintsManagerController;
+	friend class AConstraintsActor;
+};
+
+/** 
+ * FConstraintsManagerController
+ * Basic controller to add / remove / get constraints
+ **/
+
+class CONSTRAINTS_API FConstraintsManagerController
+{
+public:
+	/** @todo document */
+	static FConstraintsManagerController& Get(UWorld* InWorld);
+
+	/** @todo document */
+	template< typename TConstraint >
+	TConstraint* AllocateConstraintT(const FName& InBaseName) const
+	{
+		UConstraintsManager* Manager = GetManager();
+		if (!Manager)
+		{
+			return nullptr;
+		}
+
+		// unique name (we may want to use another approach here to manage uniqueness)
+		const FName Name = MakeUniqueObjectName(Manager, TConstraint::StaticClass(), InBaseName);
+
+		TConstraint* NewConstraint = NewObject<TConstraint>(Manager, Name);
+		AddConstraint(NewConstraint);
+		
+		return NewConstraint;
+	}
+
+	/** @todo document */
+	void AddConstraint(UTickableConstraint* InConstraint) const;
+	
+	/** Get the index of the given constraint's name. */
+	int32 GetConstraintIndex(const FName& InConstraintName) const;
+	
+	/** Remove the constraint by name. */
+	void RemoveConstraint(const FName& InConstraintName) const;
+
+	/** Remove the constraint at the given index. */
+	void RemoveConstraint(const int32 InConstraintIndex) const;
+
+	/** Get read-only access to the array of constraints. */
+	const TArray< TObjectPtr<UTickableConstraint> >& GetConstrainsArray() const;
+	
+	/** Get parent constraints of the specified child. If bSorted is true, then the constraints will be sorted by dependency. */
+	TArray< TObjectPtr<UTickableConstraint> > GetParentConstraints(
+		const uint32 InTargetHash,
+		const bool bSorted = false) const;
+
+	/** Set dependencies between two constraints. */
+	void SetConstraintsDependencies(
+		const FName& InNameToTickBefore,
+		const FName& InNameToTickAfter) const;
+	
+private:
+	
+	/** Find the existing Constraints Manager in World or create a new one. */
+	UConstraintsManager* GetManager() const;
+	
+	/** Find the existing Constraints Manager in World. */
+	UConstraintsManager* FindManager() const;
+
+	/** Destroy the ConstraintsManager from the World. */
+	void DestroyManager() const;
+
+	/** The World that holds the ConstraintsManagerActor. */
+	UWorld* World = nullptr;
+};
\ No newline at end of file
diff --git a/Engine/Source/Runtime/Experimental/Animation/Constraints/Public/TransformConstraint.h b/Engine/Source/Runtime/Experimental/Animation/Constraints/Public/TransformConstraint.h
new file mode 100644
--- /dev/null
+++ b/Engine/Source/Runtime/Experimental/Animation/Constraints/Public/TransformConstraint.h
@@ -0,0 +1,230 @@
+﻿// Copyright Epic Games, Inc. All Rights Reserved.
+
+#pragma once
+
+#include "Constraint.h"
+#include "ConstraintsManager.h"
+
+#include "TransformConstraint.generated.h"
+
+class UTransformableHandle;
+
+using SetTransformFunc = TFunction<void(const FTransform&)>;
+using GetTransformFunc = TFunction<FTransform()>;
+
+/** 
+ * UTickableTransformConstraint
+ **/
+
+UCLASS(Abstract)
+class CONSTRAINTS_API UTickableTransformConstraint : public UTickableConstraint
+{
+	GENERATED_BODY()
+
+public:
+
+	/** @todo document */
+	void Setup();
+	/** @todo document */
+	virtual void PostLoad() override;
+	/** @todo document */
+	virtual void PostDuplicate(bool bDuplicateForPIE) override;
+
+	/** @todo document */
+	virtual uint32 GetTargetHash() const override;
+	
+	/** @todo document */
+	UPROPERTY()
+	TObjectPtr<UTransformableHandle> ParentTRSHandle;
+	/** @todo document */
+	UPROPERTY()
+	TObjectPtr<UTransformableHandle> ChildTRSHandle;
+
+	/** @todo document */
+	UPROPERTY()
+	bool bMaintainOffset = true;
+
+	/** @todo document */
+	UPROPERTY(meta = (Input, ClampMin = "0.0", ClampMax = "1.0", UIMin = "0.0", UIMax = "1.0"))
+	float Weight = 1.f;
+	
+	/** @todo document. */
+	int64 GetType() const;
+		
+protected:
+	/** @todo document */
+	virtual void ComputeOffset() PURE_VIRTUAL(ComputeOffset, return;);
+	
+	/** @todo document */
+	void SetupDependencies();
+	
+	/** @todo document */
+	void SetChildGlobalTransform(const FTransform& InGlobal) const;
+
+	/** @todo document */
+	FTransform GetChildGlobalTransform() const;
+
+	/** @todo document */
+	FTransform GetParentGlobalTransform() const;
+
+	/** @todo document */
+	UPROPERTY()
+	ETransformConstraintType Type = ETransformConstraintType::Parent;
+
+#if WITH_EDITOR
+public:
+	/** @todo document */
+	virtual FName GetLabel() const override;
+#endif
+};
+
+/** 
+ * UTickableTranslationConstraint
+ **/
+
+UCLASS()
+class CONSTRAINTS_API UTickableTranslationConstraint : public UTickableTransformConstraint
+{
+	GENERATED_BODY()
+
+public:
+	UTickableTranslationConstraint();
+
+	/** @todo document */
+	virtual FConstraintTickFunction::ConstraintFunction GetFunction() const override;
+
+protected:
+	/** @todo document */
+	virtual void ComputeOffset() override;
+
+	/** @todo document */
+	UPROPERTY()
+	FVector OffsetTranslation = FVector::ZeroVector;
+};
+
+/** 
+ * UTickableRotationConstraint
+ **/
+
+UCLASS()
+class CONSTRAINTS_API UTickableRotationConstraint : public UTickableTransformConstraint
+{
+	GENERATED_BODY()
+
+public:
+	UTickableRotationConstraint();
+
+	/** @todo document */
+	virtual FConstraintTickFunction::ConstraintFunction GetFunction() const override;
+
+protected:
+	/** @todo document */
+	virtual void ComputeOffset() override;
+	
+	/** @todo document */
+	UPROPERTY()
+	FQuat OffsetRotation = FQuat::Identity;
+};
+
+/** 
+ * UTickableScaleConstraint
+ **/
+
+UCLASS()
+class CONSTRAINTS_API UTickableScaleConstraint : public UTickableTransformConstraint
+{
+	GENERATED_BODY()
+
+public:
+	UTickableScaleConstraint();
+	
+	/** @todo document */
+	virtual FConstraintTickFunction::ConstraintFunction GetFunction() const override;
+
+protected:
+	/** @todo document */
+	virtual void ComputeOffset() override;
+
+	/** @todo document */
+	UPROPERTY()
+	FVector OffsetScale = FVector::OneVector;
+};
+
+/** 
+ * UTickableParentConstraint
+ **/
+
+UCLASS()
+class CONSTRAINTS_API UTickableParentConstraint : public UTickableTransformConstraint
+{
+	GENERATED_BODY()
+
+public:
+	UTickableParentConstraint();
+	
+	/** @todo document */
+	virtual FConstraintTickFunction::ConstraintFunction GetFunction() const override;
+
+protected:
+	/** @todo document */
+	virtual void ComputeOffset() override;
+
+	/** @todo document */
+	UPROPERTY()
+	FTransform OffsetTransform = FTransform::Identity;
+};
+
+/** 
+ * UTickableLookAtConstraint
+ **/
+
+UCLASS()
+class CONSTRAINTS_API UTickableLookAtConstraint : public UTickableTransformConstraint
+{
+	GENERATED_BODY()
+
+public:
+	UTickableLookAtConstraint();
+	
+	/** @todo document */
+	virtual FConstraintTickFunction::ConstraintFunction GetFunction() const override;
+
+protected:
+	/** @todo document */
+	virtual void ComputeOffset() override;
+
+	/** @todo document */
+	UPROPERTY()
+	FTransform OffsetTransform = FTransform::Identity;
+};
+
+/** 
+ * TransformConstraintUtils
+ **/
+
+struct CONSTRAINTS_API TransformConstraintUtils
+{
+	/** @todo document */
+	static bool Create(
+		UWorld* World,
+		const AActor* InParent,
+		const AActor* InChild,
+		const ETransformConstraintType InType = ETransformConstraintType::Parent,
+		const bool bMaintainOffset = true);
+
+	/** @todo document */
+	static void GetParentConstraints(
+		UWorld* World,
+		const AActor* InChild,
+		TArray< TObjectPtr<UTickableConstraint> >& OutConstraints);
+	
+private:
+	
+	/** Registers a new transform constraint within the constraints manager. */	
+	static bool AddConst(
+		UWorld* InWorld,
+		UTransformableHandle* InParentHandle,
+		UTransformableHandle* InChildHandle,
+		const ETransformConstraintType InType,
+		const bool bMaintainOffset = true);
+};
diff --git a/Engine/Source/Runtime/Experimental/Animation/Constraints/Public/TransformableHandle.h b/Engine/Source/Runtime/Experimental/Animation/Constraints/Public/TransformableHandle.h
new file mode 100644
--- /dev/null
+++ b/Engine/Source/Runtime/Experimental/Animation/Constraints/Public/TransformableHandle.h
@@ -0,0 +1,94 @@
+﻿// Copyright Epic Games, Inc. All Rights Reserved.
+
+#pragma once
+
+#include "UObject/Object.h"
+
+#include "TransformableHandle.generated.h"
+
+struct FTickFunction;
+class USceneComponent;
+
+/**
+ * UTransformableHandle
+ */
+
+UCLASS(Abstract)
+class CONSTRAINTS_API UTransformableHandle : public UObject 
+{
+	GENERATED_BODY()
+	
+public:
+	virtual ~UTransformableHandle();
+
+	virtual void PostLoad() override;
+	
+	/** Sanity check to ensure the handle is safe to use. */
+	virtual bool IsValid() const PURE_VIRTUAL(IsValid, return false;);
+	
+	/** Sets the transform of the underlying transformable object. */
+	virtual void SetTransform(const FTransform& InGlobal) const PURE_VIRTUAL(SetTransform, );
+	/** Gets the transform of the underlying transformable object. */
+	virtual FTransform GetTransform() const PURE_VIRTUAL(GetTransform, return FTransform::Identity;);
+
+	/**
+	 * Returns the target object containing the tick function (returned in GetTickFunction).
+	 * See FTickFunction::AddPrerequisite for details.
+	 **/
+	virtual UObject* GetPrerequisiteObject() const PURE_VIRTUAL(GetPrerequisiteObject, return nullptr;);
+	/**
+	 * Returns the tick function of the underlying transformable object.
+	 * This is used to set dependencies with the constraint.
+	**/
+	virtual FTickFunction* GetTickFunction() const PURE_VIRTUAL(GetTickFunction, return nullptr;);
+
+	/**
+	 * Generates a hash value of the underlying transformable object.
+	**/
+	virtual uint32 GetHash() const PURE_VIRTUAL(GetHash, return 0;);
+	
+#if WITH_EDITOR
+	virtual FName GetName() const PURE_VIRTUAL(GetName, return NAME_None;);
+#endif
+};
+
+/**
+ * UTransformableComponentHandle
+ */
+
+UCLASS()
+class CONSTRAINTS_API UTransformableComponentHandle : public UTransformableHandle 
+{
+	GENERATED_BODY()
+	
+public:
+	
+	virtual ~UTransformableComponentHandle();
+	
+	/** @todo document */
+	virtual bool IsValid() const override;
+
+	virtual void PostLoad() override;
+	
+	/** @todo document */
+	virtual void SetTransform(const FTransform& InGlobal) const override;
+	/** @todo document */
+	virtual FTransform GetTransform() const override;
+
+	/** @todo document */
+	virtual UObject* GetPrerequisiteObject() const override;
+	/** @todo document */
+	virtual FTickFunction* GetTickFunction() const override;
+
+	/** @todo document */
+	virtual uint32 GetHash() const override;
+
+#if WITH_EDITOR
+	/** @todo document */
+	virtual FName GetName() const override;
+#endif
+	
+	/** @todo document */
+	UPROPERTY()
+	TWeakObjectPtr<USceneComponent> Component;
+};
\ No newline at end of file
diff --git a/Engine/Source/Runtime/Experimental/Animation/Constraints/Public/TransformableRegistry.h b/Engine/Source/Runtime/Experimental/Animation/Constraints/Public/TransformableRegistry.h
new file mode 100644
--- /dev/null
+++ b/Engine/Source/Runtime/Experimental/Animation/Constraints/Public/TransformableRegistry.h
@@ -0,0 +1,48 @@
+﻿// Copyright Epic Games, Inc. All Rights Reserved.
+
+#pragma once
+
+class UTransformableHandle;
+class UClass;
+
+/**
+ * TransformableRegistry
+ * Can be used to register customization transformable handle factory for specific objects. 
+ */
+
+class CONSTRAINTS_API FTransformableRegistry
+{
+public:
+	/** @todo document */
+	using CreateHandleFuncT = TFunction<UTransformableHandle*(const UObject*, UObject*)>;
+	using GetHashFuncT = TFunction<uint32(const UObject*)>;
+	
+	~FTransformableRegistry();
+
+	/** @todo document */
+	static FTransformableRegistry& Get();
+
+	/** @todo document */
+	void Register(UClass* InClass, CreateHandleFuncT InHandleFunc, GetHashFuncT InHashFunc);
+
+	/** @todo document */
+	GetHashFuncT GetHashFunction(const UClass* InClass) const;
+	/** @todo document */
+	CreateHandleFuncT GetCreateFunction(const UClass* InClass) const;
+
+private:
+	FTransformableRegistry() = default;
+
+	/** @todo document */
+	struct FTransformableInfo
+	{
+		CreateHandleFuncT CreateHandleFunc;
+		GetHashFuncT GetHashFunc;
+	};
+
+	/** @todo document */
+	const FTransformableInfo* FindInfo(const UClass* InClass) const;
+
+	/** @todo document */
+	TMap<UClass*, FTransformableInfo> Transformables;
+};
