Fix monolithic editor static initialization order
(also fix missing includes)
#rnx

diff --git a/Engine/Plugins/Animation/ControlRig/Source/ControlRigEditor/Private/Editor/Constraints/SConstraintsWidget.cpp b/Engine/Plugins/Animation/ControlRig/Source/ControlRigEditor/Private/Editor/Constraints/SConstraintsWidget.cpp
--- a/Engine/Plugins/Animation/ControlRig/Source/ControlRigEditor/Private/Editor/Constraints/SConstraintsWidget.cpp
+++ b/Engine/Plugins/Animation/ControlRig/Source/ControlRigEditor/Private/Editor/Constraints/SConstraintsWidget.cpp
@@ -13,6 +13,8 @@
 #include "SlateOptMacros.h"
 #include "LevelEditor.h"
 #include "TransformConstraint.h"
+#include "Modules/ModuleManager.h"
+#include "Widgets/Input/SButton.h"
 
 
 #define LOCTEXT_NAMESPACE "SConstraintsWidget"
@@ -27,13 +29,17 @@ TArray< SConstraintsCreationWidget::ItemSharedPtr > SConstraintsCreationWidget::
 	FDroppableConstraintItem::Make(ETransformConstraintType::LookAt)
 });
 
-TArray< const FSlateBrush* > FConstraintInfo::Brushes({
-	FAppStyle::Get().GetBrush("EditorViewport.TranslateMode"),
-	FAppStyle::Get().GetBrush("EditorViewport.RotateMode"),
-	FAppStyle::Get().GetBrush("EditorViewport.ScaleMode"),
-	FAppStyle::Get().GetBrush("Icons.Transform"),
-	FAppStyle::Get().GetBrush("Level.VisibleHighlightIcon16x")
-});
+const TArray< const FSlateBrush* >& FConstraintInfo::GetBrushes()
+{
+	static const TArray< const FSlateBrush* > Brushes({
+		FAppStyle::Get().GetBrush("EditorViewport.TranslateMode"),
+		FAppStyle::Get().GetBrush("EditorViewport.RotateMode"),
+		FAppStyle::Get().GetBrush("EditorViewport.ScaleMode"),
+		FAppStyle::Get().GetBrush("Icons.Transform"),
+		FAppStyle::Get().GetBrush("Level.VisibleHighlightIcon16x")
+		});
+	return Brushes;
+}
 
 TMap< UClass*, ETransformConstraintType > FConstraintInfo::ConstraintToType({
 	{UTickableTranslationConstraint::StaticClass(), ETransformConstraintType::Translation},
@@ -48,7 +54,7 @@ const FSlateBrush* FConstraintInfo::GetBrush(uint8 InType)
 	static const UEnum* ETransformConstraintTypeEnum = StaticEnum<ETransformConstraintType>();
 	if (ETransformConstraintTypeEnum->IsValidEnumValue(InType))
 	{
-		return Brushes[InType];
+		return GetBrushes()[InType];
 	}
 
 	return FAppStyle::Get().GetDefaultBrush();
diff --git a/Engine/Plugins/Animation/ControlRig/Source/ControlRigEditor/Private/Editor/Constraints/SConstraintsWidget.h b/Engine/Plugins/Animation/ControlRig/Source/ControlRigEditor/Private/Editor/Constraints/SConstraintsWidget.h
--- a/Engine/Plugins/Animation/ControlRig/Source/ControlRigEditor/Private/Editor/Constraints/SConstraintsWidget.h
+++ b/Engine/Plugins/Animation/ControlRig/Source/ControlRigEditor/Private/Editor/Constraints/SConstraintsWidget.h
@@ -3,6 +3,7 @@
 #pragma once
 
 #include "Widgets/SCompoundWidget.h"
+#include "Widgets/Views/SListView.h"
 #include "Constraint.h"
 
 class AActor;
@@ -20,7 +21,7 @@ public:
 	static const FSlateBrush* GetBrush(uint8 InType);
 	static int8 GetType(UClass* InClass);
 private:
-	static TArray< const FSlateBrush* > Brushes;
+	static const TArray< const FSlateBrush* >& GetBrushes();
 	static TMap< UClass*, ETransformConstraintType > ConstraintToType;
 };
 
