Fix additional monolithic editor static initialization order issue
#rnx

diff --git a/Engine/Plugins/Animation/ControlRig/Source/ControlRigEditor/Private/Editor/Constraints/SConstraintsWidget.cpp b/Engine/Plugins/Animation/ControlRig/Source/ControlRigEditor/Private/Editor/Constraints/SConstraintsWidget.cpp
--- a/Engine/Plugins/Animation/ControlRig/Source/ControlRigEditor/Private/Editor/Constraints/SConstraintsWidget.cpp
+++ b/Engine/Plugins/Animation/ControlRig/Source/ControlRigEditor/Private/Editor/Constraints/SConstraintsWidget.cpp
@@ -41,13 +41,17 @@ const TArray< const FSlateBrush* >& FConstraintInfo::GetBrushes()
 	return Brushes;
 }
 
-TMap< UClass*, ETransformConstraintType > FConstraintInfo::ConstraintToType({
-	{UTickableTranslationConstraint::StaticClass(), ETransformConstraintType::Translation},
-	{UTickableRotationConstraint::StaticClass(), ETransformConstraintType::Rotation},
-	{UTickableScaleConstraint::StaticClass(), ETransformConstraintType::Scale},
-	{UTickableParentConstraint::StaticClass(), ETransformConstraintType::Parent},
-	{UTickableLookAtConstraint::StaticClass(), ETransformConstraintType::LookAt}
-});
+const TMap< UClass*, ETransformConstraintType >& FConstraintInfo::GetConstraintToType()
+{
+	static const TMap< UClass*, ETransformConstraintType > ConstraintToType({
+		{UTickableTranslationConstraint::StaticClass(), ETransformConstraintType::Translation},
+		{UTickableRotationConstraint::StaticClass(), ETransformConstraintType::Rotation},
+		{UTickableScaleConstraint::StaticClass(), ETransformConstraintType::Scale},
+		{UTickableParentConstraint::StaticClass(), ETransformConstraintType::Parent},
+		{UTickableLookAtConstraint::StaticClass(), ETransformConstraintType::LookAt}
+		});
+	return ConstraintToType;
+}
 
 const FSlateBrush* FConstraintInfo::GetBrush(uint8 InType)
 {
@@ -62,7 +66,7 @@ const FSlateBrush* FConstraintInfo::GetBrush(uint8 InType)
 
 int8 FConstraintInfo::GetType(UClass* InClass)
 {
-	if (ETransformConstraintType* TransformConstraint = ConstraintToType.Find(InClass))
+	if (const ETransformConstraintType* TransformConstraint = GetConstraintToType().Find(InClass))
 	{
 		return static_cast<int8>(*TransformConstraint); 
 	}
diff --git a/Engine/Plugins/Animation/ControlRig/Source/ControlRigEditor/Private/Editor/Constraints/SConstraintsWidget.h b/Engine/Plugins/Animation/ControlRig/Source/ControlRigEditor/Private/Editor/Constraints/SConstraintsWidget.h
--- a/Engine/Plugins/Animation/ControlRig/Source/ControlRigEditor/Private/Editor/Constraints/SConstraintsWidget.h
+++ b/Engine/Plugins/Animation/ControlRig/Source/ControlRigEditor/Private/Editor/Constraints/SConstraintsWidget.h
@@ -22,7 +22,7 @@ public:
 	static int8 GetType(UClass* InClass);
 private:
 	static const TArray< const FSlateBrush* >& GetBrushes();
-	static TMap< UClass*, ETransformConstraintType > ConstraintToType;
+	static const TMap< UClass*, ETransformConstraintType >& GetConstraintToType();
 };
 
 /**
