Sequencer: Control Rig: Added the missing orbit around selection behavior for control rig objects.

diff --git a/Engine/Plugins/Animation/ControlRig/Source/ControlRigEditor/Private/EditMode/ControlRigEditMode.cpp b/Engine/Plugins/Animation/ControlRig/Source/ControlRigEditor/Private/EditMode/ControlRigEditMode.cpp
--- a/Engine/Plugins/Animation/ControlRig/Source/ControlRigEditor/Private/EditMode/ControlRigEditMode.cpp
+++ b/Engine/Plugins/Animation/ControlRig/Source/ControlRigEditor/Private/EditMode/ControlRigEditMode.cpp
@@ -1120,6 +1120,16 @@ FVector FControlRigEditMode::GetWidgetLocation() const
 	return FEdMode::GetWidgetLocation();
 }
 
+bool FControlRigEditMode::GetPivotForOrbit(FVector& OutPivot) const
+{
+	if (UsesTransformWidget())
+	{
+		OutPivot = GetWidgetLocation();
+		return true;
+	}
+	return FEdMode::GetPivotForOrbit(OutPivot);
+}
+
 bool FControlRigEditMode::GetCustomDrawingCoordinateSystem(FMatrix& OutMatrix, void* InData)
 {
 	//since we strip translation just want the first one
diff --git a/Engine/Plugins/Animation/ControlRig/Source/ControlRigEditor/Private/EditMode/ControlRigEditMode.h b/Engine/Plugins/Animation/ControlRig/Source/ControlRigEditor/Private/EditMode/ControlRigEditMode.h
--- a/Engine/Plugins/Animation/ControlRig/Source/ControlRigEditor/Private/EditMode/ControlRigEditMode.h
+++ b/Engine/Plugins/Animation/ControlRig/Source/ControlRigEditor/Private/EditMode/ControlRigEditMode.h
@@ -119,6 +119,7 @@ public:
 	virtual void SelectNone() override;
 	virtual bool InputDelta(FEditorViewportClient* InViewportClient, FViewport* InViewport, FVector& InDrag, FRotator& InRot, FVector& InScale) override;
 	virtual bool UsesTransformWidget() const override;
+	virtual bool GetPivotForOrbit(FVector& OutPivot) const override;
 	virtual bool UsesTransformWidget(UE::Widget::EWidgetMode CheckMode) const;
 	virtual FVector GetWidgetLocation() const override;
 	virtual bool GetCustomDrawingCoordinateSystem(FMatrix& OutMatrix, void* InData) override;
