Control Rig: Respect channel grouping settings

diff --git a/Engine/Plugins/Animation/ControlRig/Source/ControlRig/Private/Rigs/RigHierarchy.cpp b/Engine/Plugins/Animation/ControlRig/Source/ControlRig/Private/Rigs/RigHierarchy.cpp
--- a/Engine/Plugins/Animation/ControlRig/Source/ControlRig/Private/Rigs/RigHierarchy.cpp
+++ b/Engine/Plugins/Animation/ControlRig/Source/ControlRig/Private/Rigs/RigHierarchy.cpp
@@ -4759,6 +4759,37 @@ bool URigHierarchy::IsAnimatable(const FRigControlElement* InControlElement) con
 	return false;
 }
 
+bool URigHierarchy::ShouldBeGrouped(const FRigElementKey& InKey) const
+{
+	if(const FRigControlElement* ControlElement = Find<FRigControlElement>(InKey))
+	{
+		return ShouldBeGrouped(ControlElement);
+	}
+	return false;
+}
+
+bool URigHierarchy::ShouldBeGrouped(const FRigControlElement* InControlElement) const
+{
+	if(InControlElement)
+	{
+		if(!InControlElement->Settings.ShouldBeGrouped())
+		{
+			return false;
+		}
+
+		if(!GetChildren(InControlElement).IsEmpty())
+		{
+			return false;
+		}
+
+		if(const FRigControlElement* ParentControlElement = Cast<FRigControlElement>(GetFirstParent(InControlElement)))
+		{
+			return ParentControlElement->Settings.AnimationType == ERigControlAnimationType::AnimationControl;
+		}
+	}
+	return false;
+}
+
 FTransform URigHierarchy::GetWorldTransformForReference(const FRigUnitContext* InContext, const FRigElementKey& InKey, bool bInitial)
 {
 	if(const USceneComponent* OuterSceneComponent = GetTypedOuter<USceneComponent>())
diff --git a/Engine/Plugins/Animation/ControlRig/Source/ControlRig/Private/Sequencer/MovieSceneControlRigParameterSection.cpp b/Engine/Plugins/Animation/ControlRig/Source/ControlRig/Private/Sequencer/MovieSceneControlRigParameterSection.cpp
--- a/Engine/Plugins/Animation/ControlRig/Source/ControlRig/Private/Sequencer/MovieSceneControlRigParameterSection.cpp
+++ b/Engine/Plugins/Animation/ControlRig/Source/ControlRig/Private/Sequencer/MovieSceneControlRigParameterSection.cpp
@@ -1354,23 +1354,20 @@ void UMovieSceneControlRigParameterSection::ReconstructChannelProxy()
 				continue;
 			}
 
-			FRigControlElement* ParentControlElement = Cast<FRigControlElement>(ControlRig->GetHierarchy()->GetFirstParent(ControlElement));
 			FName ParentControlName = NAME_None;
-			if(ParentControlElement)
-			{
-				ParentControlName = ParentControlElement->GetName();
-			}
-			bool bEnabled = ControlsMask[MaskIndex];
-
 			FText Group;
-			if(ParentControlElement && ControlElement->Settings.IsAnimatable() && ControlElement->Settings.bGroupWithParentControl)
+
+			if(ControlRig->GetHierarchy()->ShouldBeGrouped(ControlElement))
 			{
-				if(ParentControlElement->Settings.AnimationType == ERigControlAnimationType::AnimationControl)
+				if(const FRigControlElement* ParentControlElement = Cast<FRigControlElement>(ControlRig->GetHierarchy()->GetFirstParent(ControlElement)))
 				{
+					ParentControlName = ParentControlElement->GetName();
 					Group = FText::FromName(ParentControlElement->GetDisplayName());
 				}
 			}
-			
+
+			bool bEnabled = ControlsMask[MaskIndex];
+
 #if WITH_EDITOR
 			switch (ControlElement->Settings.ControlType)
 			{
diff --git a/Engine/Plugins/Animation/ControlRig/Source/ControlRig/Public/Rigs/RigHierarchy.h b/Engine/Plugins/Animation/ControlRig/Source/ControlRig/Public/Rigs/RigHierarchy.h
--- a/Engine/Plugins/Animation/ControlRig/Source/ControlRig/Public/Rigs/RigHierarchy.h
+++ b/Engine/Plugins/Animation/ControlRig/Source/ControlRig/Public/Rigs/RigHierarchy.h
@@ -2627,6 +2627,16 @@ public:
 	 */
 	bool IsAnimatable(const FRigControlElement* InControlElement) const;
 
+	/**
+	 * Returns true if the control should be grouped in editor
+	 */
+	bool ShouldBeGrouped(const FRigElementKey& InKey) const;
+
+	/**
+	 * Returns true if the control should be grouped in editor
+	 */
+	bool ShouldBeGrouped(const FRigControlElement* InControlElement) const;
+
 #if WITH_EDITOR
 
 	/**
diff --git a/Engine/Plugins/Animation/ControlRig/Source/ControlRig/Public/Rigs/RigHierarchyElements.h b/Engine/Plugins/Animation/ControlRig/Source/ControlRig/Public/Rigs/RigHierarchyElements.h
--- a/Engine/Plugins/Animation/ControlRig/Source/ControlRig/Public/Rigs/RigHierarchyElements.h
+++ b/Engine/Plugins/Animation/ControlRig/Source/ControlRig/Public/Rigs/RigHierarchyElements.h
@@ -996,6 +996,11 @@ struct CONTROLRIG_API FRigControlSettings
 			(AnimationType == ERigControlAnimationType::AnimationChannel);
 	}
 
+	FORCEINLINE bool ShouldBeGrouped() const
+	{
+		return IsAnimatable() && bGroupWithParentControl;
+	}
+
 	FORCEINLINE bool SupportsShape() const
 	{
 		return (AnimationType != ERigControlAnimationType::AnimationChannel) &&
diff --git a/Engine/Plugins/Animation/ControlRig/Source/ControlRigEditor/Private/ControlRigElementDetails.cpp b/Engine/Plugins/Animation/ControlRig/Source/ControlRigEditor/Private/ControlRigElementDetails.cpp
--- a/Engine/Plugins/Animation/ControlRig/Source/ControlRigEditor/Private/ControlRigElementDetails.cpp
+++ b/Engine/Plugins/Animation/ControlRig/Source/ControlRigEditor/Private/ControlRigElementDetails.cpp
@@ -2531,11 +2531,43 @@ void FRigControlElementDetails::CustomizeControl(IDetailLayoutBuilder& DetailBui
 	const bool bSupportsShape = !IsAnyControlOfAnimationType(ERigControlAnimationType::AnimationChannel) &&
 		!IsAnyControlOfAnimationType(ERigControlAnimationType::VisualCue);
 
-	if(IsAnyControlOfAnimationType(ERigControlAnimationType::AnimationControl) ||
-	IsAnyControlOfAnimationType(ERigControlAnimationType::AnimationChannel))
+	if (HierarchyBeingCustomized != nullptr)
 	{
-		const TSharedPtr<IPropertyHandle> GroupWithParentControlHandle = SettingsHandle->GetChildHandle(TEXT("bGroupWithParentControl"));
-		ControlCategory.AddProperty(GroupWithParentControlHandle.ToSharedRef()).DisplayName(FText::FromString(TEXT("Group Channels")));
+		bool bEnableGroupWithParentControl = true;
+		for(TWeakObjectPtr<UDetailsViewWrapperObject> ObjectBeingCustomized : ObjectsBeingCustomized)
+		{
+			if(ObjectBeingCustomized.IsValid())
+			{
+				if(ObjectBeingCustomized->IsChildOf<FRigControlElement>())
+				{
+					bool bSingleEnableGroupWithParentControl = false;
+					const FRigElementKey Key = ObjectBeingCustomized->GetContent<FRigControlElement>().GetKey();
+					if(const FRigControlElement* ControlElement = HierarchyBeingCustomized->Find<FRigControlElement>(Key))
+					{
+						if(const FRigControlElement* ParentElement =
+							Cast<FRigControlElement>(HierarchyBeingCustomized->GetFirstParent(ControlElement)))
+						{
+							if(ControlElement->Settings.IsAnimatable() &&
+								HierarchyBeingCustomized->GetChildren(ControlElement).IsEmpty())
+							{
+								bSingleEnableGroupWithParentControl = true;
+							}
+						}
+					}
+
+					if(!bSingleEnableGroupWithParentControl)
+					{
+						bEnableGroupWithParentControl = false;
+						break;
+					}
+				}
+			}
+		}
+		if(bEnableGroupWithParentControl)
+		{
+			const TSharedPtr<IPropertyHandle> GroupWithParentControlHandle = SettingsHandle->GetChildHandle(TEXT("bGroupWithParentControl"));
+			ControlCategory.AddProperty(GroupWithParentControlHandle.ToSharedRef()).DisplayName(FText::FromString(TEXT("Group Channels")));
+		}
 	}
 	
 	if(bSupportsShape &&
