Control Rig: Enable auto key for proxy controls

diff --git a/Engine/Plugins/Animation/ControlRig/Source/ControlRigEditor/Private/EditMode/ControlRigEditMode.cpp b/Engine/Plugins/Animation/ControlRig/Source/ControlRigEditor/Private/EditMode/ControlRigEditMode.cpp
--- a/Engine/Plugins/Animation/ControlRig/Source/ControlRigEditor/Private/EditMode/ControlRigEditMode.cpp
+++ b/Engine/Plugins/Animation/ControlRig/Source/ControlRigEditor/Private/EditMode/ControlRigEditMode.cpp
@@ -3579,7 +3579,23 @@ void FControlRigEditMode::MoveControlShape(AControlRigShapeActor* ShapeActor, co
 
 				}
 
-				ControlRig->Evaluate_AnyThread();		
+				ControlRig->Evaluate_AnyThread();
+
+				// if we are changing a proxy control - we also need to notify the change for the driven controls
+				if (FRigControlElement* ControlElement = ControlRig->GetHierarchy()->Find<FRigControlElement>(ShapeActor->GetElementKey()))
+				{
+					if(ControlElement->Settings.AnimationType == ERigControlAnimationType::ProxyControl)
+					{
+						for(const FRigElementKey& DrivenKey : ControlElement->Settings.DrivenControls)
+						{
+							if(DrivenKey.Type == ERigElementType::Control)
+							{
+								const FTransform DrivenTransform = ControlRig->GetHierarchy()->GetLocalTransform(DrivenKey);
+								ControlRig->SetControlLocalTransform(DrivenKey.Name, DrivenTransform, true, Context, false);
+							}
+						}
+					}
+				}
 			}
 		}
 	}
