Control Rig: Fix static analysis issue in ControlRigParameterTrackEditor.cpp

diff --git a/Engine/Plugins/Animation/ControlRig/Source/ControlRigEditor/Private/Sequencer/ControlRigParameterTrackEditor.cpp b/Engine/Plugins/Animation/ControlRig/Source/ControlRigEditor/Private/Sequencer/ControlRigParameterTrackEditor.cpp
--- a/Engine/Plugins/Animation/ControlRig/Source/ControlRigEditor/Private/Sequencer/ControlRigParameterTrackEditor.cpp
+++ b/Engine/Plugins/Animation/ControlRig/Source/ControlRigEditor/Private/Sequencer/ControlRigParameterTrackEditor.cpp
@@ -2275,10 +2275,15 @@ void FControlRigParameterTrackEditor::SetUpEditModeIfNeeded(UControlRig* Control
 
 void FControlRigParameterTrackEditor::HandleControlSelected(UControlRig* Subject, FRigControlElement* ControlElement, bool bSelected)
 {
+	if(ControlElement == nullptr)
+	{
+		return;
+	}
+	
 	URigHierarchy* Hierarchy = Subject->GetHierarchy();
 	static bool bIsSelectingIndirectControl = false;
 
-	if(ControlElement && ControlElement->Settings.AnimationType == ERigControlAnimationType::ProxyControl)
+	if(ControlElement->Settings.AnimationType == ERigControlAnimationType::ProxyControl)
 	{
 		const TArray<FRigElementKey>& DrivenControls = ControlElement->Settings.DrivenControls;
 		for(const FRigElementKey& DrivenKey : DrivenControls)
