Sequencer: Follow up to keying only umuted and active sections - also need to use only keyable sections when calculating weight for blending and interrogation

diff --git a/Engine/Plugins/Animation/ControlRig/Source/ControlRig/Private/Sequencer/MovieSceneControlRigParameterTemplate.cpp b/Engine/Plugins/Animation/ControlRig/Source/ControlRig/Private/Sequencer/MovieSceneControlRigParameterTemplate.cpp
--- a/Engine/Plugins/Animation/ControlRig/Source/ControlRig/Private/Sequencer/MovieSceneControlRigParameterTemplate.cpp
+++ b/Engine/Plugins/Animation/ControlRig/Source/ControlRig/Private/Sequencer/MovieSceneControlRigParameterTemplate.cpp
@@ -3,9 +3,11 @@
 #include "Sequencer/MovieSceneControlRigParameterTemplate.h"
 #include "Evaluation/Blending/MovieSceneMultiChannelBlending.h"
 #include "Sequencer/MovieSceneControlRigParameterTrack.h"
+#include "Sequencer/MovieSceneControlRigParameterSection.h"
 #include "ControlRig.h"
 #include "Evaluation/MovieSceneEvaluation.h"
 #include "IMovieScenePlayer.h"
+#include "MovieSceneCommonHelpers.h"
 #include "Components/SkeletalMeshComponent.h"
 #include "AnimCustomInstanceHelper.h"
 #include "Sequencer/ControlRigLayerInstance.h"
@@ -2056,7 +2058,7 @@ void FMovieSceneControlRigParameterTemplate::Interrogate(const FMovieSceneContex
 	const FFrameTime Time = Context.GetTime();
 
 	const UMovieSceneControlRigParameterSection* Section = Cast<UMovieSceneControlRigParameterSection>(GetSourceSection());
-	if (Section && Section->GetControlRig())
+	if (Section && Section->GetControlRig() && MovieSceneHelpers::IsSectionKeyable(Section))
 	{
 		FEvaluatedControlRigParameterSectionChannelMasks ChannelMasks;
 		ChannelMasks.Initialize(Section, Scalars, Bools, Integers, Enums, Vector2Ds, Vectors, Colors, Transforms);
diff --git a/Engine/Source/Runtime/MovieScene/Private/MovieSceneCommonHelpers.cpp b/Engine/Source/Runtime/MovieScene/Private/MovieSceneCommonHelpers.cpp
--- a/Engine/Source/Runtime/MovieScene/Private/MovieSceneCommonHelpers.cpp
+++ b/Engine/Source/Runtime/MovieScene/Private/MovieSceneCommonHelpers.cpp
@@ -17,7 +17,7 @@
 #include "MovieSceneTrack.h"
 #include "Engine/Engine.h"
 
-bool MovieSceneHelpers::IsSectionKeyable(UMovieSceneSection* Section)
+bool MovieSceneHelpers::IsSectionKeyable(const UMovieSceneSection* Section)
 {
 	if (!Section)
 	{
@@ -477,7 +477,7 @@ float MovieSceneHelpers::CalculateWeightForBlending(UMovieSceneSection* SectionT
 			TArray<UMovieSceneSection*, TInlineAllocator<4>> OverlappingSections;
 			for (UMovieSceneSection* Section : Sections)
 			{
-				if (Section->GetRange().Contains(Time))
+				if (MovieSceneHelpers::IsSectionKeyable(Section) && Section->GetRange().Contains(Time))
 				{
 					OverlappingSections.Add(Section);
 				}
diff --git a/Engine/Source/Runtime/MovieScene/Public/MovieSceneCommonHelpers.h b/Engine/Source/Runtime/MovieScene/Public/MovieSceneCommonHelpers.h
--- a/Engine/Source/Runtime/MovieScene/Public/MovieSceneCommonHelpers.h
+++ b/Engine/Source/Runtime/MovieScene/Public/MovieSceneCommonHelpers.h
@@ -28,7 +28,7 @@ public:
 	/** 
 	 * @return Whether the section is keyable (active, on a track that is not muted, etc 
 	 */
-	static bool IsSectionKeyable(UMovieSceneSection*);
+	static bool IsSectionKeyable(const UMovieSceneSection*);
 
 	/**
 	 * Finds a section that exists at a given time
diff --git a/Engine/Source/Runtime/MovieSceneTracks/Private/EntitySystem/Interrogation/MovieSceneInterrogationLinker.cpp b/Engine/Source/Runtime/MovieSceneTracks/Private/EntitySystem/Interrogation/MovieSceneInterrogationLinker.cpp
--- a/Engine/Source/Runtime/MovieSceneTracks/Private/EntitySystem/Interrogation/MovieSceneInterrogationLinker.cpp
+++ b/Engine/Source/Runtime/MovieSceneTracks/Private/EntitySystem/Interrogation/MovieSceneInterrogationLinker.cpp
@@ -664,6 +664,11 @@ void FSystemInterrogator::ImportTrack(UMovieSceneTrack* Track, FInterrogationCha
 
 void FSystemInterrogator::ImportTrack(UMovieSceneTrack* Track, const FGuid& ObjectBindingID, FInterrogationChannel InChannel)
 {
+	if (Track->IsEvalDisabled())
+	{
+		return;
+	}
+
 	TGuardValue<FEntityManager*> DebugVizGuard(GEntityManagerForDebuggingVisualizers, &Linker->EntityManager);
 
 	FFrameRate TickResolution = Track->GetTypedOuter<UMovieScene>()->GetTickResolution();
@@ -682,6 +687,11 @@ void FSystemInterrogator::ImportTrack(UMovieSceneTrack* Track, const FGuid& Obje
 			continue;
 		}
 
+		if (Track->IsRowEvalDisabled(Entry.Section->GetRowIndex()))
+		{
+			continue;
+		}
+
 		FMovieSceneEvaluationFieldEntityMetaData MetaData;
 
 		MetaData.ForcedTime = Entry.ForcedTime;
diff --git a/Engine/Source/Runtime/MovieSceneTracks/Private/Systems/MovieScenePredictionSystem.cpp b/Engine/Source/Runtime/MovieSceneTracks/Private/Systems/MovieScenePredictionSystem.cpp
--- a/Engine/Source/Runtime/MovieSceneTracks/Private/Systems/MovieScenePredictionSystem.cpp
+++ b/Engine/Source/Runtime/MovieSceneTracks/Private/Systems/MovieScenePredictionSystem.cpp
@@ -10,6 +10,7 @@
 #include "EntitySystem/Interrogation/MovieSceneInterrogatedPropertyInstantiator.h"
 
 #include "Sections/MovieScene3DTransformSection.h"
+#include "MovieSceneCommonHelpers.h"
 #include "MovieSceneTracksComponentTypes.h"
 #include "MovieSceneSequence.h"
 #include "MovieSceneSequencePlayer.h"
@@ -184,6 +185,11 @@ int32 UMovieSceneAsyncAction_SequencePrediction::ImportTransformEntities(UMovieS
 			return true;
 		}
 
+		if (!MovieSceneHelpers::IsSectionKeyable(TransformSection))
+		{
+			return true;
+		}
+
 		const FMovieSceneEvaluationFieldSharedEntityMetaData* SharedMetaData = ComponentField->FindSharedMetaData(InQuery);
 		if (!SharedMetaData || SharedMetaData->ObjectBindingID != ObjectGuid)
 		{
diff --git a/Engine/Source/Runtime/MovieSceneTracks/Private/Tracks/MovieScene3DTransformTrack.cpp b/Engine/Source/Runtime/MovieSceneTracks/Private/Tracks/MovieScene3DTransformTrack.cpp
--- a/Engine/Source/Runtime/MovieSceneTracks/Private/Tracks/MovieScene3DTransformTrack.cpp
+++ b/Engine/Source/Runtime/MovieSceneTracks/Private/Tracks/MovieScene3DTransformTrack.cpp
@@ -6,6 +6,7 @@
 #include "Compilation/MovieSceneSegmentCompiler.h"
 #include "Evaluation/MovieSceneEvaluationTrack.h"
 #include "Channels/MovieSceneChannelProxy.h"
+#include "MovieSceneCommonHelpers.h"
 #include "Algo/BinarySearch.h"
 
 UMovieScene3DTransformTrack::UMovieScene3DTransformTrack( const FObjectInitializer& ObjectInitializer )
@@ -175,7 +176,7 @@ TArray<FTrajectoryKey> UMovieScene3DTransformTrack::GetTrajectoryData(FFrameNumb
 	for (UMovieSceneSection* Section : Sections)
 	{
 		UMovieScene3DTransformSection* TransformSection = Cast<UMovieScene3DTransformSection>(Section);
-		if (TransformSection)
+		if (TransformSection && MovieSceneHelpers::IsSectionKeyable(TransformSection))
 		{
 			TArrayView<FMovieSceneDoubleChannel*>        DoubleChannels = TransformSection->GetChannelProxy().GetChannels<FMovieSceneDoubleChannel>();
 			TArrayView<const FMovieSceneChannelMetaData> MetaData      = TransformSection->GetChannelProxy().GetMetaData<FMovieSceneDoubleChannel>();
