Control Rig: Fix hierarchy invalidity bugs

diff --git a/Engine/Plugins/Animation/ControlRig/Source/ControlRig/Private/Rigs/RigHierarchy.cpp b/Engine/Plugins/Animation/ControlRig/Source/ControlRig/Private/Rigs/RigHierarchy.cpp
--- a/Engine/Plugins/Animation/ControlRig/Source/ControlRig/Private/Rigs/RigHierarchy.cpp
+++ b/Engine/Plugins/Animation/ControlRig/Source/ControlRig/Private/Rigs/RigHierarchy.cpp
@@ -636,6 +636,7 @@ void URigHierarchy::ResetPoseToInitial(ERigElementType InTypeFilter)
 				{
 					// compute the local offset if necessary
 					GetControlOffsetTransform(ControlElement, ERigTransformType::CurrentLocal);
+					GetControlShapeTransform(ControlElement, ERigTransformType::CurrentLocal);
 				}
 			}
 		}
@@ -649,6 +650,7 @@ void URigHierarchy::ResetPoseToInitial(ERigElementType InTypeFilter)
 				if(FRigControlElement* ControlElement = Cast<FRigControlElement>(TransformElement))
 				{
 					ControlElement->Offset.MarkDirty(ERigTransformType::CurrentGlobal);
+					ControlElement->Shape.MarkDirty(ERigTransformType::CurrentGlobal);
 				}
 			}
 		}
diff --git a/Engine/Plugins/Animation/ControlRig/Source/ControlRig/Public/ControlRig.h b/Engine/Plugins/Animation/ControlRig/Source/ControlRig/Public/ControlRig.h
--- a/Engine/Plugins/Animation/ControlRig/Source/ControlRig/Public/ControlRig.h
+++ b/Engine/Plugins/Animation/ControlRig/Source/ControlRig/Public/ControlRig.h
@@ -857,7 +857,7 @@ public:
 	class FTransientControlPoseScope
 	{
 	public:
-		FTransientControlPoseScope(TObjectPtr<UControlRig> InControlRig)
+		FORCEINLINE_DEBUGGABLE FTransientControlPoseScope(TObjectPtr<UControlRig> InControlRig)
 		{
 			ControlRig = InControlRig;
 
@@ -867,14 +867,20 @@ public:
 			{
 				Keys.Add(TransientControl->GetKey());
 			}
-	
-			CachedPose = ControlRig->GetHierarchy()->GetPose(false, ERigElementType::Control, TArrayView<FRigElementKey>(Keys));
+
+			if(Keys.Num() > 0)
+			{
+				CachedPose = ControlRig->GetHierarchy()->GetPose(false, ERigElementType::Control, TArrayView<FRigElementKey>(Keys));
+			}
 		}
-		~FTransientControlPoseScope()
+		FORCEINLINE_DEBUGGABLE ~FTransientControlPoseScope()
 		{
 			check(ControlRig);
 
-			ControlRig->GetHierarchy()->SetPose(CachedPose);
+			if(CachedPose.Num() > 0)
+			{
+				ControlRig->GetHierarchy()->SetPose(CachedPose);
+			}
 		}
 	
 	private:
