Control Rig: Fix crash when upgradring nodes

diff --git a/Engine/Plugins/Animation/ControlRig/Source/ControlRigEditor/Private/Graph/SControlRigGraphNode.cpp b/Engine/Plugins/Animation/ControlRig/Source/ControlRigEditor/Private/Graph/SControlRigGraphNode.cpp
--- a/Engine/Plugins/Animation/ControlRig/Source/ControlRigEditor/Private/Graph/SControlRigGraphNode.cpp
+++ b/Engine/Plugins/Animation/ControlRig/Source/ControlRigEditor/Private/Graph/SControlRigGraphNode.cpp
@@ -1215,6 +1215,11 @@ void SControlRigGraphNode::RefreshErrorInfo()
 
 void SControlRigGraphNode::Tick(const FGeometry& AllottedGeometry, const double InCurrentTime, const float InDeltaTime)
 {
+	if (!ModelNode.IsValid())
+	{
+		return;
+	}
+	
 	SGraphNode::Tick(AllottedGeometry, InCurrentTime, InDeltaTime);
 
 	if (GraphNode)
