Control Rig: XRay mode

diff --git a/Engine/Plugins/Animation/ControlRig/Source/ControlRig/Private/ControlRig.cpp b/Engine/Plugins/Animation/ControlRig/Source/ControlRig/Private/ControlRig.cpp
--- a/Engine/Plugins/Animation/ControlRig/Source/ControlRig/Private/ControlRig.cpp
+++ b/Engine/Plugins/Animation/ControlRig/Source/ControlRig/Private/ControlRig.cpp
@@ -85,6 +85,7 @@ UControlRig::UControlRig(const FObjectInitializer& ObjectInitializer)
 	, InteractionBracket(0)
 	, InterRigSyncBracket(0)
 	, InteractionType((uint8)EControlRigInteractionType::None)
+	, bInteractionJustBegan(false)
 #if WITH_EDITORONLY_DATA
 	, VMSnapshotBeforeExecution(nullptr)
 #endif
@@ -1244,6 +1245,11 @@ void UControlRig::Execute(const EControlRigState InState, const FName& InEventNa
 			return true;
 		});
 	}
+
+	if(InEventName == FRigUnit_InteractionExecution::EventName)
+	{
+		bInteractionJustBegan = false;
+	}
 }
 
 void UControlRig::ExecuteUnits(FRigUnitContext& InOutContext, const FName& InEventName)
diff --git a/Engine/Plugins/Animation/ControlRig/Source/ControlRig/Private/ControlRigGizmoLibrary.cpp b/Engine/Plugins/Animation/ControlRig/Source/ControlRig/Private/ControlRigGizmoLibrary.cpp
--- a/Engine/Plugins/Animation/ControlRig/Source/ControlRig/Private/ControlRigGizmoLibrary.cpp
+++ b/Engine/Plugins/Animation/ControlRig/Source/ControlRig/Private/ControlRigGizmoLibrary.cpp
@@ -6,6 +6,9 @@
 
 UControlRigShapeLibrary::UControlRigShapeLibrary()
 {
+#if WITH_EDITOR
+	XRayMaterial = LoadObject<UMaterial>(nullptr, TEXT("/ControlRig/Controls/ControlRigXRayMaterial.ControlRigXRayMaterial"));
+#endif
 }
 
 #if WITH_EDITOR
diff --git a/Engine/Plugins/Animation/ControlRig/Source/ControlRig/Public/ControlRig.h b/Engine/Plugins/Animation/ControlRig/Source/ControlRig/Public/ControlRig.h
--- a/Engine/Plugins/Animation/ControlRig/Source/ControlRig/Public/ControlRig.h
+++ b/Engine/Plugins/Animation/ControlRig/Source/ControlRig/Public/ControlRig.h
@@ -700,6 +700,7 @@ protected:
 	int32 InterRigSyncBracket;
 	uint8 InteractionType;
 	TArray<FRigElementKey> ElementsBeingInteracted;
+	bool bInteractionJustBegan;
 
 	TWeakObjectPtr<USceneComponent> OuterSceneComponent;
 
@@ -913,6 +914,7 @@ private:
 	friend class UControlRigValidator;
 	friend struct FAnimNode_ControlRig;
 	friend class URigHierarchy;
+	friend class UFKControlRig;
 };
 
 class CONTROLRIG_API FControlRigBracketScope
@@ -958,6 +960,7 @@ public:
 	{
 		InControlRig->ElementsBeingInteracted = { InKey };
 		InControlRig->InteractionType = (uint8)InInteractionType;
+		InControlRig->bInteractionJustBegan = true;
 		InControlRig->GetHierarchy()->StartInteraction();
 	}
 
@@ -972,6 +975,7 @@ public:
 	{
 		InControlRig->ElementsBeingInteracted = InKeys;
 		InControlRig->InteractionType = (uint8)InInteractionType;
+		InControlRig->bInteractionJustBegan = true;
 		InControlRig->GetHierarchy()->StartInteraction();
 	}
 
@@ -981,6 +985,7 @@ public:
 		{
 			ControlRig->GetHierarchy()->EndInteraction();
 			ControlRig->InteractionType = (uint8)EControlRigInteractionType::None;
+			ControlRig->bInteractionJustBegan = false;
 			ControlRig->ElementsBeingInteracted.Reset();
 		}
 	}
diff --git a/Engine/Plugins/Animation/ControlRig/Source/ControlRig/Public/ControlRigGizmoLibrary.h b/Engine/Plugins/Animation/ControlRig/Source/ControlRig/Public/ControlRigGizmoLibrary.h
--- a/Engine/Plugins/Animation/ControlRig/Source/ControlRig/Public/ControlRigGizmoLibrary.h
+++ b/Engine/Plugins/Animation/ControlRig/Source/ControlRig/Public/ControlRigGizmoLibrary.h
@@ -51,6 +51,9 @@ public:
 	UPROPERTY(EditAnywhere, Category = "ShapeLibrary")
 	TSoftObjectPtr<UMaterial> DefaultMaterial;
 
+	UPROPERTY(EditAnywhere, Category = "ShapeLibrary")
+	TSoftObjectPtr<UMaterial> XRayMaterial;
+
 	UPROPERTY(EditAnywhere, Category = "ShapeLibrary")
 	FName MaterialColorParameter;
 
diff --git a/Engine/Plugins/Animation/ControlRig/Source/ControlRigEditor/Private/ControlRigGizmoLibraryFactory.cpp b/Engine/Plugins/Animation/ControlRig/Source/ControlRigEditor/Private/ControlRigGizmoLibraryFactory.cpp
--- a/Engine/Plugins/Animation/ControlRig/Source/ControlRigEditor/Private/ControlRigGizmoLibraryFactory.cpp
+++ b/Engine/Plugins/Animation/ControlRig/Source/ControlRigEditor/Private/ControlRigGizmoLibraryFactory.cpp
@@ -19,6 +19,7 @@ UObject* UControlRigShapeLibraryFactory::FactoryCreateNew(UClass* Class, UObject
 
 	ShapeLibrary->DefaultShape.StaticMesh = LoadObject<UStaticMesh>(nullptr, TEXT("/ControlRig/Controls/ControlRig_Sphere_solid.ControlRig_Sphere_solid"));
 	ShapeLibrary->DefaultMaterial = LoadObject<UMaterial>(nullptr, TEXT("/ControlRig/Controls/ControlRigGizmoMaterial.ControlRigGizmoMaterial"));
+	ShapeLibrary->XRayMaterial = LoadObject<UMaterial>(nullptr, TEXT("/ControlRig/Controls/ControlRigXRayMaterial.ControlRigXRayMaterial"));
 	ShapeLibrary->MaterialColorParameter = TEXT("Color");
 
 	return ShapeLibrary;
diff --git a/Engine/Plugins/Animation/ControlRig/Source/ControlRigEditor/Private/EditMode/ControlRigEditMode.cpp b/Engine/Plugins/Animation/ControlRig/Source/ControlRigEditor/Private/EditMode/ControlRigEditMode.cpp
--- a/Engine/Plugins/Animation/ControlRig/Source/ControlRigEditor/Private/EditMode/ControlRigEditMode.cpp
+++ b/Engine/Plugins/Animation/ControlRig/Source/ControlRigEditor/Private/EditMode/ControlRigEditMode.cpp
@@ -141,10 +141,14 @@ FControlRigEditMode::FControlRigEditMode()
 	, CurrentViewportClient(nullptr)
 	, bIsChangingCoordSystem(false)
 	, InteractionType((uint8)EControlRigInteractionType::None)
+	, bShowControlsAsOverlay(false)
 {
 	ControlProxy = NewObject<UControlRigDetailPanelControlProxies>(GetTransientPackage(), NAME_None);
 	ControlProxy->SetFlags(RF_Transactional);
 
+	const UControlRigEditModeSettings* Settings = GetDefault<UControlRigEditModeSettings>();
+	bShowControlsAsOverlay = Settings->bShowControlsAsOverlay;
+
 	CommandBindings = MakeShareable(new FUICommandList);
 	BindCommands();
 
@@ -497,6 +501,21 @@ void FControlRigEditMode::Tick(FEditorViewportClient* ViewportClient, float Delt
 	}
 	RecalcPivotTransform();
 
+	// check if the settings for xray rendering are different for any of the control shape actors
+	const UControlRigEditModeSettings* Settings = GetDefault<UControlRigEditModeSettings>();
+	if(bShowControlsAsOverlay != Settings->bShowControlsAsOverlay)
+	{
+		bShowControlsAsOverlay = Settings->bShowControlsAsOverlay;
+		for (TWeakObjectPtr<UControlRig>& RuntimeRigPtr : RuntimeControlRigs)
+		{
+			if (UControlRig* RuntimeControlRig = RuntimeRigPtr.Get())
+			{
+				UpdateSelectabilityOnSkeletalMeshes(RuntimeControlRig, !bShowControlsAsOverlay);
+			}
+		}
+		RequestToRecreateControlShapeActors();
+	}
+
 	// Defer creation of shapes if manipulating the viewport
 	if (RecreateControlShapesRequired != ERecreateControlRigShape::RecreateNone && !(FSlateApplication::Get().HasAnyMouseCaptor() || GUnrealEd->IsUserInteracting()))
 	{
@@ -1032,6 +1051,7 @@ bool FControlRigEditMode::StartTracking(FEditorViewportClient* InViewportClient,
 					//todo need to add multiple
 					FControlRigInteractionScope* InteractionScope = new FControlRigInteractionScope(ControlRig);
 					InteractionScopes.Add(ControlRig,InteractionScope);
+					ControlRig->bInteractionJustBegan = true;
 				}
 				else
 				{
@@ -2626,7 +2646,8 @@ void FControlRigEditMode::ResetTransforms(bool bSelectionOnly)
 						{
 							ControlRig->GetHierarchy()->SwitchToDefaultParent(ControlElement->GetKey());
 						}
-						ControlRig->GetHierarchy()->SetLocalTransform(ControlToReset, InitialLocalTransform);
+						ControlRig->SetControlLocalTransform(ControlToReset.Name, InitialLocalTransform, true);
+						NotifyDrivenControls(ControlRig, ControlToReset);
 						if (bHasNonDefaultParent == false)
 						{
 							ControlRig->ControlModified().Broadcast(ControlRig, ControlElement, EControlRigSetKey::DoNotCare);
@@ -2655,6 +2676,7 @@ void FControlRigEditMode::ResetTransforms(bool bSelectionOnly)
 						FTransform GlobalTransform = ControlRig->GetHierarchy()->GetGlobalTransform(ControlToReset);
 						GlobalTransforms.Add(GlobalTransform);
 					}
+					NotifyDrivenControls(ControlRig, ControlToReset);
 				}
 				//switch back to original parent space
 				int32 Index = 0;
@@ -2678,8 +2700,9 @@ void FControlRigEditMode::ResetTransforms(bool bSelectionOnly)
 						FRigControlElement* ControlElement = ControlRig->FindControl(ControlToReset.Name);
 						if (ControlElement && !ControlElement->Settings.bIsTransientControl)
 						{
-							ControlRig->GetHierarchy()->SetGlobalTransform(ControlToReset, GlobalTransforms[Index]);
+							ControlRig->SetControlGlobalTransform(ControlToReset.Name, GlobalTransforms[Index], true);
 							ControlRig->Evaluate_AnyThread();
+							NotifyDrivenControls(ControlRig, ControlToReset);
 							++Index;
 						}
 					}
@@ -2700,6 +2723,10 @@ void FControlRigEditMode::ResetTransforms(bool bSelectionOnly)
 				// we at least have to run the interaction event
 				TGuardValue<TArray<FName>> EventGuard(ControlRig->EventQueue, {});
 				ControlRig->Evaluate_AnyThread();
+				for (const FRigElementKey& ControlToReset : ControlsToReset)
+				{
+					NotifyDrivenControls(ControlRig, ControlToReset);
+				}
 			}
 		}
 	}
@@ -2972,7 +2999,8 @@ void FControlRigEditMode::CreateShapeActors(UControlRig* ControlRig)
 			{
 				Param.MeshTransform = ShapeDef->Transform;
 				Param.StaticMesh = ShapeDef->StaticMesh;
-				Param.Material = ShapeDef->Library->DefaultMaterial;
+				Param.Material = (bShowControlsAsOverlay && ShapeDef->Library->XRayMaterial.IsValid()) ?
+					ShapeDef->Library->XRayMaterial : ShapeDef->Library->DefaultMaterial;
 				Param.ColorParameterName = ShapeDef->Library->MaterialColorParameter;
 			}
 
@@ -3580,21 +3608,7 @@ void FControlRigEditMode::MoveControlShape(AControlRigShapeActor* ShapeActor, co
 
 				ControlRig->Evaluate_AnyThread();
 
-				// if we are changing a proxy control - we also need to notify the change for the driven controls
-				if (FRigControlElement* ControlElement = ControlRig->GetHierarchy()->Find<FRigControlElement>(ShapeActor->GetElementKey()))
-				{
-					if(ControlElement->Settings.AnimationType == ERigControlAnimationType::ProxyControl)
-					{
-						for(const FRigElementKey& DrivenKey : ControlElement->Settings.DrivenControls)
-						{
-							if(DrivenKey.Type == ERigElementType::Control)
-							{
-								const FTransform DrivenTransform = ControlRig->GetHierarchy()->GetLocalTransform(DrivenKey);
-								ControlRig->SetControlLocalTransform(DrivenKey.Name, DrivenTransform, true, Context, false);
-							}
-						}
-					}
-				}
+				NotifyDrivenControls(ControlRig, ShapeActor->GetElementKey());
 			}
 		}
 	}
@@ -3862,6 +3876,8 @@ void FControlRigEditMode::AddControlRigInternal(UControlRig* InControlRig)
 		Sequencer->ObjectImplicitlyAdded(InControlRig);
 	}
 	OnControlRigAddedOrRemovedDelegate.Broadcast(InControlRig, true);
+
+	UpdateSelectabilityOnSkeletalMeshes(InControlRig, !bShowControlsAsOverlay);
 }
 
 TArrayView<const TWeakObjectPtr<UControlRig>> FControlRigEditMode::GetControlRigs() const
@@ -3924,6 +3940,8 @@ void FControlRigEditMode::RemoveControlRig(UControlRig* InControlRig)
 		Sequencer->ObjectImplicitlyRemoved(InControlRig);
 	}
 	OnControlRigAddedOrRemovedDelegate.Broadcast(InControlRig, false);
+	
+	UpdateSelectabilityOnSkeletalMeshes(InControlRig, true);
 }
 
 void FControlRigEditMode::TickManipulatableObjects(float DeltaTime)
@@ -4083,6 +4101,52 @@ void FControlRigEditMode::PostPoseUpdate()
 	}
 
 }
+
+void FControlRigEditMode::NotifyDrivenControls(UControlRig* InControlRig, const FRigElementKey& InKey)
+{
+	// if we are changing a proxy control - we also need to notify the change for the driven controls
+	if (FRigControlElement* ControlElement = InControlRig->GetHierarchy()->Find<FRigControlElement>(InKey))
+	{
+		if(ControlElement->Settings.AnimationType == ERigControlAnimationType::ProxyControl)
+		{
+			FRigControlModifiedContext Context;
+			Context.EventName = FRigUnit_BeginExecution::EventName;
+
+			for(const FRigElementKey& DrivenKey : ControlElement->Settings.DrivenControls)
+			{
+				if(DrivenKey.Type == ERigElementType::Control)
+				{
+					const FTransform DrivenTransform = InControlRig->GetHierarchy()->GetLocalTransform(DrivenKey);
+					InControlRig->SetControlLocalTransform(DrivenKey.Name, DrivenTransform, true, Context, false);
+				}
+			}
+		}
+	}
+}
+
+void FControlRigEditMode::UpdateSelectabilityOnSkeletalMeshes(UControlRig* InControlRig, bool bEnabled)
+{
+	if(const USceneComponent* HostingComponent = GetHostingSceneComponent(InControlRig))
+	{
+		if(const AActor* HostingOwner = HostingComponent->GetOwner())
+		{
+			for(UActorComponent* ActorComponent : HostingOwner->GetComponents())
+			{
+				if (USkeletalMeshComponent* SkeletalMeshComponent = Cast<USkeletalMeshComponent>(ActorComponent))
+				{
+					SkeletalMeshComponent->bSelectable = bEnabled;
+					SkeletalMeshComponent->MarkRenderStateDirty();
+				}
+				else if(UStaticMeshComponent* StaticMeshComponent = Cast<UStaticMeshComponent>(ActorComponent))
+				{
+					StaticMeshComponent->bSelectable = bEnabled;
+					SkeletalMeshComponent->MarkRenderStateDirty();
+				}
+			}
+		}
+	}
+}
+
 void FControlRigEditMode::SetOnlySelectRigControls(bool Val)
 {
 	UControlRigEditModeSettings* Settings = GetMutableDefault<UControlRigEditModeSettings>();
diff --git a/Engine/Plugins/Animation/ControlRig/Source/ControlRigEditor/Private/EditMode/ControlRigEditMode.h b/Engine/Plugins/Animation/ControlRig/Source/ControlRigEditor/Private/EditMode/ControlRigEditMode.h
--- a/Engine/Plugins/Animation/ControlRig/Source/ControlRigEditor/Private/EditMode/ControlRigEditMode.h
+++ b/Engine/Plugins/Animation/ControlRig/Source/ControlRigEditor/Private/EditMode/ControlRigEditMode.h
@@ -409,6 +409,8 @@ private:
 	// Post pose update handler
 	void OnPoseInitialized();
 	void PostPoseUpdate();
+	void NotifyDrivenControls(UControlRig* InControlRig, const FRigElementKey& InKey);
+	void UpdateSelectabilityOnSkeletalMeshes(UControlRig* InControlRig, bool bEnabled);
 
 	// world clean up handlers
 	FDelegateHandle OnWorldCleanupHandle;
@@ -427,6 +429,7 @@ private:
 	/** Computes the current interaction types based on the widget mode */
 	static uint8 GetInteractionType(FEditorViewportClient* InViewportClient);
 	uint8 InteractionType;
+	bool bShowControlsAsOverlay;
 
 	friend class FControlRigEditorModule;
 	friend class FControlRigEditor;
diff --git a/Engine/Plugins/Animation/ControlRig/Source/ControlRigEditor/Private/EditMode/ControlRigEditModeSettings.h b/Engine/Plugins/Animation/ControlRig/Source/ControlRigEditor/Private/EditMode/ControlRigEditModeSettings.h
--- a/Engine/Plugins/Animation/ControlRig/Source/ControlRigEditor/Private/EditMode/ControlRigEditModeSettings.h
+++ b/Engine/Plugins/Animation/ControlRig/Source/ControlRigEditor/Private/EditMode/ControlRigEditModeSettings.h
@@ -18,6 +18,7 @@ class UControlRigEditModeSettings : public UObject
 		, bDisplayNulls(false)
 		, bHideControlShapes(false)
 		, bShowAllProxyControls(false)
+		, bShowControlsAsOverlay(false)
 		, bDisplayAxesOnSelection(false)
 		, AxisScale(10.f)
 		, bCoordSystemPerWidgetMode(true)
@@ -53,6 +54,10 @@ public:
 	UPROPERTY(config, EditAnywhere, AdvancedDisplay, Category = "Animation Settings", meta = (EditCondition = "!bHideControlShapes"))
 	bool bShowAllProxyControls;
 
+	/** Determins if controls should be rendered on top of other controls */
+	UPROPERTY(config, EditAnywhere, AdvancedDisplay, Category = "Animation Settings", meta = (EditCondition = "!bHideControlShapes"))
+	bool bShowControlsAsOverlay;
+
 	/** Indicates a control being driven by a proxy control */
 	UPROPERTY(config, EditAnywhere, AdvancedDisplay, Category = "Animation Settings", meta = (EditCondition = "!bHideControlShapes"))
 	FLinearColor DrivenControlColor;
