Control Rig: Fix some pins not resolving on load

diff --git a/Engine/Plugins/Animation/ControlRig/Source/ControlRigDeveloper/Private/ControlRigBlueprint.cpp b/Engine/Plugins/Animation/ControlRig/Source/ControlRigDeveloper/Private/ControlRigBlueprint.cpp
--- a/Engine/Plugins/Animation/ControlRig/Source/ControlRigDeveloper/Private/ControlRigBlueprint.cpp
+++ b/Engine/Plugins/Animation/ControlRig/Source/ControlRigDeveloper/Private/ControlRigBlueprint.cpp
@@ -937,7 +937,7 @@ void UControlRigBlueprint::RefreshAllModels()
 				TemplateNode->InvalidateCache();
 			}
 		}
-		Controller->RecomputeAllTemplateFilteredTypes(false);
+		//Controller->RecomputeAllTemplateFilteredTypes(false);
 	}
 }
 
