Sequencer: Control Rig:  Couldn't update a pose, multiple control rig fallout yet again.
#rivial

diff --git a/Engine/Plugins/Animation/ControlRig/Source/ControlRigEditor/Private/EditMode/SControlRigUpdatePose.cpp b/Engine/Plugins/Animation/ControlRig/Source/ControlRigEditor/Private/EditMode/SControlRigUpdatePose.cpp
--- a/Engine/Plugins/Animation/ControlRig/Source/ControlRigEditor/Private/EditMode/SControlRigUpdatePose.cpp
+++ b/Engine/Plugins/Animation/ControlRig/Source/ControlRigEditor/Private/EditMode/SControlRigUpdatePose.cpp
@@ -82,6 +82,7 @@ private:
 		{
 			TMap<UControlRig*, TArray<FRigElementKey>> AllSelectedControls;
 			ControlRigEditMode->GetAllSelectedControls(AllSelectedControls);
+			return (AllSelectedControls.Num() > 0);
 		}
 		return false;
 
