ControlRig: Fix copy & paste issue in FControlRigEditMode::UpdateSelectabilityOnSkeletalMeshes

diff --git a/Engine/Plugins/Animation/ControlRig/Source/ControlRigEditor/Private/EditMode/ControlRigEditMode.cpp b/Engine/Plugins/Animation/ControlRig/Source/ControlRigEditor/Private/EditMode/ControlRigEditMode.cpp
--- a/Engine/Plugins/Animation/ControlRig/Source/ControlRigEditor/Private/EditMode/ControlRigEditMode.cpp
+++ b/Engine/Plugins/Animation/ControlRig/Source/ControlRigEditor/Private/EditMode/ControlRigEditMode.cpp
@@ -4140,7 +4140,7 @@ void FControlRigEditMode::UpdateSelectabilityOnSkeletalMeshes(UControlRig* InCon
 				else if(UStaticMeshComponent* StaticMeshComponent = Cast<UStaticMeshComponent>(ActorComponent))
 				{
 					StaticMeshComponent->bSelectable = bEnabled;
-					SkeletalMeshComponent->MarkRenderStateDirty();
+					StaticMeshComponent->MarkRenderStateDirty();
 				}
 			}
 		}
