Control Rig: Disable errors during backwards compat code for variables

diff --git a/Engine/Plugins/Animation/ControlRig/Source/ControlRigDeveloper/Private/ControlRigBlueprint.cpp b/Engine/Plugins/Animation/ControlRig/Source/ControlRigDeveloper/Private/ControlRigBlueprint.cpp
--- a/Engine/Plugins/Animation/ControlRig/Source/ControlRigDeveloper/Private/ControlRigBlueprint.cpp
+++ b/Engine/Plugins/Animation/ControlRig/Source/ControlRigDeveloper/Private/ControlRigBlueprint.cpp
@@ -3519,6 +3519,14 @@ void UControlRigBlueprint::CreateMemberVariablesOnLoad()
 	int32 LinkerVersion = GetLinkerCustomVersion(FControlRigObjectVersion::GUID);
 	if (LinkerVersion < FControlRigObjectVersion::SwitchedToRigVM)
 	{
+		// ignore errors during the first potential compile of the VM
+		// since that this point variable nodes may still be ill-formed.
+		TGuardValue<FRigVMReportDelegate> SuspendReportDelegate(VMCompileSettings.ASTSettings.ReportDelegate,
+			FRigVMReportDelegate::CreateLambda([](EMessageSeverity::Type,  UObject*, const FString&)
+			{
+				// do nothing
+			})
+		);
 		InitializeModelIfRequired();
 	}
 
