Control Rig: Guard timed TFunction with weakptr

diff --git a/Engine/Plugins/Animation/ControlRig/Source/ControlRigEditor/Private/EditMode/SControlRigDetails.cpp b/Engine/Plugins/Animation/ControlRig/Source/ControlRigEditor/Private/EditMode/SControlRigDetails.cpp
--- a/Engine/Plugins/Animation/ControlRig/Source/ControlRigEditor/Private/EditMode/SControlRigDetails.cpp
+++ b/Engine/Plugins/Animation/ControlRig/Source/ControlRigEditor/Private/EditMode/SControlRigDetails.cpp
@@ -384,9 +384,21 @@ void SControlRigDetails::HandleControlSelected(UControlRig* Subject, FRigControl
 
 void SControlRigDetails::UpdateProxies()
 {
+	TWeakPtr<SWidget> WeakPtr = AsWeak();
+	
 	//proxies that are in edit mode are also listening to the same messages so they may not be set up yet so need to wait
-	GEditor->GetTimerManager()->SetTimerForNextTick([this]()
+	GEditor->GetTimerManager()->SetTimerForNextTick([WeakPtr]()
 	{
+		if(!WeakPtr.IsValid())
+		{
+			return;
+		}
+		TSharedPtr<SControlRigDetails> StrongThis = StaticCastSharedPtr<SControlRigDetails>(WeakPtr.Pin());
+		if(!StrongThis.IsValid())
+		{
+			return;
+		}
+		
 		TArray<TWeakObjectPtr<>> Eulers;
 		TArray<TWeakObjectPtr<>> Transforms;
 		TArray<TWeakObjectPtr<>> TransformNoScales;
@@ -405,9 +417,9 @@ void SControlRigDetails::UpdateProxies()
 		TArray<TWeakObjectPtr<>> IndividualBools;
 		TArray<TWeakObjectPtr<>> IndividualIntegers;
 		TArray<TWeakObjectPtr<>> IndividualEnums;
-		TArray<UControlRig*> ControlRigs = GetControlRigs();
+		TArray<UControlRig*> ControlRigs = StrongThis->GetControlRigs();
 	
-		if (FControlRigEditMode* EditMode = static_cast<FControlRigEditMode*>(ModeTools->GetActiveMode(FControlRigEditMode::ModeName)))
+		if (FControlRigEditMode* EditMode = static_cast<FControlRigEditMode*>(StrongThis->ModeTools->GetActiveMode(FControlRigEditMode::ModeName)))
 		{
 			if (UControlRigDetailPanelControlProxies* ControlProxy = EditMode->GetDetailProxies())
 			{
@@ -533,24 +545,24 @@ void SControlRigDetails::UpdateProxies()
 			}
 		}
 
-		SetTransformDetailsObjects(Transforms, false);
-		SetTransformNoScaleDetailsObjects(TransformNoScales, false);
-		SetEulerTransformDetailsObjects(Eulers, false);
-		SetVectorDetailsObjects(Vectors, false);
-		SetVector2DDetailsObjects(Vector2Ds, false);
-		SetFloatDetailsObjects(Floats, false);
-		SetBoolDetailsObjects(Bools,false);
-		SetIntegerDetailsObjects(Integers,false);
-		SetEnumDetailsObjects(Enums,false);
-		SetTransformDetailsObjects(IndividualTransforms, true);
-		SetTransformNoScaleDetailsObjects(IndividualTransformNoScales, true);
-		SetEulerTransformDetailsObjects(IndividualEulers, true);
-		SetVectorDetailsObjects(IndividualVectors, true);
-		SetVector2DDetailsObjects(IndividualVector2Ds, true);
-		SetFloatDetailsObjects(IndividualFloats, true);
-		SetBoolDetailsObjects(IndividualBools,true);
-		SetIntegerDetailsObjects(IndividualIntegers,true);
-		SetEnumDetailsObjects(IndividualEnums,true);
+		StrongThis->SetTransformDetailsObjects(Transforms, false);
+		StrongThis->SetTransformNoScaleDetailsObjects(TransformNoScales, false);
+		StrongThis->SetEulerTransformDetailsObjects(Eulers, false);
+		StrongThis->SetVectorDetailsObjects(Vectors, false);
+		StrongThis->SetVector2DDetailsObjects(Vector2Ds, false);
+		StrongThis->SetFloatDetailsObjects(Floats, false);
+		StrongThis->SetBoolDetailsObjects(Bools,false);
+		StrongThis->SetIntegerDetailsObjects(Integers,false);
+		StrongThis->SetEnumDetailsObjects(Enums,false);
+		StrongThis->SetTransformDetailsObjects(IndividualTransforms, true);
+		StrongThis->SetTransformNoScaleDetailsObjects(IndividualTransformNoScales, true);
+		StrongThis->SetEulerTransformDetailsObjects(IndividualEulers, true);
+		StrongThis->SetVectorDetailsObjects(IndividualVectors, true);
+		StrongThis->SetVector2DDetailsObjects(IndividualVector2Ds, true);
+		StrongThis->SetFloatDetailsObjects(IndividualFloats, true);
+		StrongThis->SetBoolDetailsObjects(IndividualBools,true);
+		StrongThis->SetIntegerDetailsObjects(IndividualIntegers,true);
+		StrongThis->SetEnumDetailsObjects(IndividualEnums,true);
 	});
 }
 
