Control Rig Details: Fixed copy & paste bug

diff --git a/Engine/Plugins/Animation/ControlRig/Source/ControlRigEditor/Private/EditMode/SControlRigDetails.cpp b/Engine/Plugins/Animation/ControlRig/Source/ControlRigEditor/Private/EditMode/SControlRigDetails.cpp
--- a/Engine/Plugins/Animation/ControlRig/Source/ControlRigEditor/Private/EditMode/SControlRigDetails.cpp
+++ b/Engine/Plugins/Animation/ControlRig/Source/ControlRigEditor/Private/EditMode/SControlRigDetails.cpp
@@ -540,7 +540,7 @@ void SControlRigDetails::UpdateProxies()
 				UControlRigControlsProxy* Proxy = Cast<UControlRigControlsProxy>(Object.Get());
 				if (Proxy)
 				{
-					Proxy->SetIsMultiple(Transforms.Num() > 1);
+					Proxy->SetIsMultiple(Floats.Num() > 1);
 				}
 			}
 		}
