Control Rig: Blueprint handlepackage load to engage in -game even if the package isn't known yet

diff --git a/Engine/Plugins/Animation/ControlRig/Source/ControlRigDeveloper/Private/ControlRigBlueprint.cpp b/Engine/Plugins/Animation/ControlRig/Source/ControlRigDeveloper/Private/ControlRigBlueprint.cpp
--- a/Engine/Plugins/Animation/ControlRig/Source/ControlRigDeveloper/Private/ControlRigBlueprint.cpp
+++ b/Engine/Plugins/Animation/ControlRig/Source/ControlRigDeveloper/Private/ControlRigBlueprint.cpp
@@ -618,11 +618,7 @@ void UControlRigBlueprint::PostLoad()
 	// but GIsEditor is turned off
 	if(!GIsEditor)
 	{
-		FEndLoadPackageContext Context;
-		Context.bSynchronous = true;
-		Context.LoadedPackages = {GetPackage()};
-		Context.RecursiveDepth = 0;
-		HandlePackageDone(Context);
+		HandlePackageDone();
 	}
 #endif
 }
@@ -634,7 +630,11 @@ void UControlRigBlueprint::HandlePackageDone(const FEndLoadPackageContext& Conte
 	{
 		return;
 	}
+	HandlePackageDone();
+}
 
+void UControlRigBlueprint::HandlePackageDone()
+{
 	FCoreUObjectDelegates::OnEndLoadPackage.RemoveAll(this);
 
 	if (ShapeLibrariesToLoadOnPackageLoaded.Num() > 0)
diff --git a/Engine/Plugins/Animation/ControlRig/Source/ControlRigDeveloper/Public/ControlRigBlueprint.h b/Engine/Plugins/Animation/ControlRig/Source/ControlRigDeveloper/Public/ControlRigBlueprint.h
--- a/Engine/Plugins/Animation/ControlRig/Source/ControlRigDeveloper/Public/ControlRigBlueprint.h
+++ b/Engine/Plugins/Animation/ControlRig/Source/ControlRigDeveloper/Public/ControlRigBlueprint.h
@@ -605,6 +605,7 @@ private:
 
 #if WITH_EDITOR
 	void HandlePackageDone(const FEndLoadPackageContext& Context);
+	void HandlePackageDone();
 	// ControlRigBP, once end-loaded, will inform other ControlRig-Dependent systems that ControlRig instances are ready.
 	void BroadcastControlRigPackageDone();
 
