Sequencer: Control Rig: Very small drags will cause all controls to get selected to do frustum hit testing, so we check for this and then bail if that's the case. Also suport Only Select Control Rig Controls with marquee  selection.
#trivial

diff --git a/Engine/Plugins/Animation/ControlRig/Source/ControlRigEditor/Private/EditMode/ControlRigEditMode.cpp b/Engine/Plugins/Animation/ControlRig/Source/ControlRigEditor/Private/EditMode/ControlRigEditMode.cpp
--- a/Engine/Plugins/Animation/ControlRig/Source/ControlRigEditor/Private/EditMode/ControlRigEditMode.cpp
+++ b/Engine/Plugins/Animation/ControlRig/Source/ControlRigEditor/Private/EditMode/ControlRigEditMode.cpp
@@ -1593,17 +1593,32 @@ bool FControlRigEditMode::BoxSelect(FBox& InBox, bool InSelect)
 	if (bSomethingSelected == true)
 	{
 		return true;
-	}
-	
+	}	
 	ScopedTransaction.Cancel();
+	//if only selecting controls return true to stop any more selections
+	if (Settings && Settings->bOnlySelectRigControls)
+	{
+		return true;
+	}
 	return FEdMode::BoxSelect(InBox, InSelect);
 }
 
 bool FControlRigEditMode::FrustumSelect(const FConvexVolume& InFrustum, FEditorViewportClient* InViewportClient, bool InSelect)
 {
 	const UControlRigEditModeSettings* Settings = GetDefault<UControlRigEditModeSettings>();
-	if (InViewportClient->IsInGameView() == true || Settings->bHideControlShapes)
+	//need to check for a zero frustum since ComponentIsTouchingSelectionFrustum will return true, selecting everything, when this is the case
+	const bool bMalformedFrustum = (InFrustum.Planes[0].IsNearlyZero() && InFrustum.Planes[1].IsNearlyZero() && InFrustum.Planes[2].IsNearlyZero() &&
+		InFrustum.Planes[3].IsNearlyZero());
+	if (bMalformedFrustum || InViewportClient->IsInGameView() == true || Settings->bHideControlShapes)
 	{
+		if (Settings->bOnlySelectRigControls)
+		{
+			return true;
+		}
+		else
+		{
+			return false;
+		}
 		return FEdMode::FrustumSelect(InFrustum, InViewportClient, InSelect);
 	}
 
@@ -1642,8 +1657,14 @@ bool FControlRigEditMode::FrustumSelect(const FConvexVolume& InFrustum, FEditorV
 		return true;
 	}
 	ScopedTransaction.Cancel();
+	//if only selecting controls return true to stop any more selections
+	if (Settings && Settings->bOnlySelectRigControls)
+	{
+		return true;
+	}
 	return FEdMode::FrustumSelect(InFrustum, InViewportClient, InSelect);
 }
+
 void FControlRigEditMode::SelectNone()
 {
 	ClearRigElementSelection(FRigElementTypeHelper::ToMask(ERigElementType::All));
