Control Rig: Fix crash when upgrading sequence node

diff --git a/Engine/Plugins/Animation/ControlRig/Source/ControlRigDeveloper/Private/Graph/ControlRigGraph.cpp b/Engine/Plugins/Animation/ControlRig/Source/ControlRigDeveloper/Private/Graph/ControlRigGraph.cpp
--- a/Engine/Plugins/Animation/ControlRig/Source/ControlRigDeveloper/Private/Graph/ControlRigGraph.cpp
+++ b/Engine/Plugins/Animation/ControlRig/Source/ControlRigDeveloper/Private/Graph/ControlRigGraph.cpp
@@ -421,7 +421,7 @@ void UControlRigGraph::HandleModifiedEvent(ERigVMGraphNotifType InNotifType, URi
 				UEdGraphNode* EdNode = FindNodeForModelNodeName(ModelNode->GetFName(), false);
 				if (EdNode)
 				{
-					RemoveAndDeleteNode(EdNode);
+					RemoveNode(EdNode);
 				}
 			}
 			break;
@@ -927,21 +927,26 @@ URigVMController* UControlRigGraph::GetController() const
 	return nullptr;
 }
 
-void UControlRigGraph::RemoveAndDeleteNode(UEdGraphNode* InNode)
+void UControlRigGraph::RemoveNode(UEdGraphNode* InNode)
 {
 	// Make sure EdGraph is not part of the transaction
 	TGuardValue<ITransaction*> TransactionGuard(GUndo, nullptr);
 					
-	static UObject* NewOuter = GetTransientPackage();
-	InNode->Rename(nullptr, NewOuter, REN_ForceNoResetLoaders | REN_DontCreateRedirectors);
-					
-	// Remove the node and notify of the removal
+	// Rename the soon to be deleted object to a unique name, so that other objects can use
+	// the old name
+	FString DeletedName;
 	{
-		Nodes.Remove(InNode);
-		InNode->BreakAllNodeLinks();
-						
-		NotifyGraphChanged();
+		UObject* ExistingObject = nullptr;
+		static int32 DeletedIndex = FMath::Rand();
+		do
+		{
+			DeletedName = FString::Printf(TEXT("%s_Deleted_%d"), *InNode->GetName(), DeletedIndex++); 
+			ExistingObject = StaticFindObject(/*Class=*/ NULL, this, *DeletedName, true);						
+		}
+		while (ExistingObject);
 	}
+	InNode->Rename(*DeletedName, nullptr, REN_ForceNoResetLoaders | REN_DontCreateRedirectors);	
+	Super::RemoveNode(InNode);
 }
 
 URigVMController* UControlRigGraph::GetTemplateController()
diff --git a/Engine/Plugins/Animation/ControlRig/Source/ControlRigDeveloper/Public/Graph/ControlRigGraph.h b/Engine/Plugins/Animation/ControlRig/Source/ControlRigDeveloper/Public/Graph/ControlRigGraph.h
--- a/Engine/Plugins/Animation/ControlRig/Source/ControlRigDeveloper/Public/Graph/ControlRigGraph.h
+++ b/Engine/Plugins/Animation/ControlRig/Source/ControlRigDeveloper/Public/Graph/ControlRigGraph.h
@@ -135,7 +135,7 @@ private:
 
 	TMap<URigVMNode*, TPair<int32, int32>> CachedInstructionIndices;
 
-	void RemoveAndDeleteNode(UEdGraphNode* InNode);
+	void RemoveNode(UEdGraphNode* InNode);
 
 #endif
 #if WITH_EDITORONLY_DATA
diff --git a/Engine/Plugins/Animation/ControlRig/Source/ControlRigEditor/Private/ControlRigGraphDetails.cpp b/Engine/Plugins/Animation/ControlRig/Source/ControlRigEditor/Private/ControlRigGraphDetails.cpp
--- a/Engine/Plugins/Animation/ControlRig/Source/ControlRigEditor/Private/ControlRigGraphDetails.cpp
+++ b/Engine/Plugins/Animation/ControlRig/Source/ControlRigEditor/Private/ControlRigGraphDetails.cpp
@@ -563,6 +563,7 @@ FControlRigArgumentDefaultNode::FControlRigArgumentDefaultNode(
 		{
 			if (UControlRigGraph* RigGraph = Cast<UControlRigGraph>(ControlRigBlueprintPtr->GetEdGraph(LibraryNode->GetGraph())))
 			{
+				EdGraphOuterPtr = RigGraph;
 				GraphChangedDelegateHandle = RigGraph->AddOnGraphChangedHandler(
 					FOnGraphChanged::FDelegate::CreateRaw(this, &FControlRigArgumentDefaultNode::OnGraphChanged)
 				);
@@ -574,20 +575,17 @@ FControlRigArgumentDefaultNode::FControlRigArgumentDefaultNode(
 
 FControlRigArgumentDefaultNode::~FControlRigArgumentDefaultNode()
 {
-	if (GraphPtr.IsValid() && ControlRigBlueprintPtr.IsValid())
+	if (ControlRigBlueprintPtr.IsValid())
 	{
 		ControlRigBlueprintPtr.Get()->OnModified().RemoveAll(this);
-
-		if (URigVMLibraryNode* LibraryNode = Cast<URigVMLibraryNode>(GraphPtr->GetOuter()))
+	}
+	
+	if (EdGraphOuterPtr.IsValid())
+	{
+		if (GraphChangedDelegateHandle.IsValid())
 		{
-			if (UControlRigGraph* RigGraph = Cast<UControlRigGraph>(ControlRigBlueprintPtr->GetEdGraph(LibraryNode->GetGraph())))
-			{
-				if (GraphChangedDelegateHandle.IsValid())
-				{
-					RigGraph->RemoveOnGraphChangedHandler(GraphChangedDelegateHandle);
-				}
-			}
-		}
+			EdGraphOuterPtr->RemoveOnGraphChangedHandler(GraphChangedDelegateHandle);
+		}		
 	}
 }
 
diff --git a/Engine/Plugins/Animation/ControlRig/Source/ControlRigEditor/Private/ControlRigGraphDetails.h b/Engine/Plugins/Animation/ControlRig/Source/ControlRigEditor/Private/ControlRigGraphDetails.h
--- a/Engine/Plugins/Animation/ControlRig/Source/ControlRigEditor/Private/ControlRigGraphDetails.h
+++ b/Engine/Plugins/Animation/ControlRig/Source/ControlRigEditor/Private/ControlRigGraphDetails.h
@@ -147,6 +147,7 @@ private:
 	void HandleModifiedEvent(ERigVMGraphNotifType InNotifType, URigVMGraph* InGraph, UObject* InSubject);
 
 	TWeakObjectPtr<URigVMGraph> GraphPtr;
+	TWeakObjectPtr<UControlRigGraph> EdGraphOuterPtr;
 	TWeakObjectPtr<UControlRigBlueprint> ControlRigBlueprintPtr;
 	FSimpleDelegate OnRebuildChildren;
 	TSharedPtr<SControlRigGraphNode> OwnedNodeWidget;
