Control Rig: Support for preferred euler angles

- Removed implicit casts between FEulerTransform and FTransform
- moved interrogator and tokens inside of sequencer to use FEulerTransform
- added the notion of a preferred rotator for Controls. transform takes precedence for runtime solves, but UI / sequencer relies on the preferred rotator

diff --git a/Engine/Plugins/Animation/ControlRig/Source/ControlRig/Private/AnimNode_ControlRig.cpp b/Engine/Plugins/Animation/ControlRig/Source/ControlRig/Private/AnimNode_ControlRig.cpp
--- a/Engine/Plugins/Animation/ControlRig/Source/ControlRig/Private/AnimNode_ControlRig.cpp
+++ b/Engine/Plugins/Animation/ControlRig/Source/ControlRig/Private/AnimNode_ControlRig.cpp
@@ -506,7 +506,7 @@ void FAnimNode_ControlRig::PropagateInputProperties(const UObject* InSourceInsta
 							if(ensure(StructProperty->Struct == TBaseStructure<FTransform>::Get()))
 							{
 								const FTransform& SrcTransform = *(FTransform*)SrcPtr;  
-								Value = FRigControlValue::Make<FEulerTransform>(SrcTransform);
+								Value = FRigControlValue::Make<FEulerTransform>(FEulerTransform(SrcTransform));
 								bIsValid = true;
 							}
 						}
diff --git a/Engine/Plugins/Animation/ControlRig/Source/ControlRig/Private/ControlRig.cpp b/Engine/Plugins/Animation/ControlRig/Source/ControlRig/Private/ControlRig.cpp
--- a/Engine/Plugins/Animation/ControlRig/Source/ControlRig/Private/ControlRig.cpp
+++ b/Engine/Plugins/Animation/ControlRig/Source/ControlRig/Private/ControlRig.cpp
@@ -1713,7 +1713,7 @@ FTransform UControlRig::GetControlGlobalTransform(const FName& InControlName) co
 	return DynamicHierarchy->GetGlobalTransform(FRigElementKey(InControlName, ERigElementType::Control), false);
 }
 
-bool UControlRig::SetControlGlobalTransform(const FName& InControlName, const FTransform& InGlobalTransform, bool bNotify, const FRigControlModifiedContext& Context, bool bSetupUndo, bool bPrintPythonCommands)
+bool UControlRig::SetControlGlobalTransform(const FName& InControlName, const FTransform& InGlobalTransform, bool bNotify, const FRigControlModifiedContext& Context, bool bSetupUndo, bool bPrintPythonCommands, bool bFixEulerFlips)
 {
 	FTransform GlobalTransform = InGlobalTransform;
 	if (IsSetupModeEnabled())
@@ -1731,7 +1731,7 @@ bool UControlRig::SetControlGlobalTransform(const FName& InControlName, const FT
 		}
 	}
 
-	SetControlValue(InControlName, Value, bNotify, Context, bSetupUndo, bPrintPythonCommands);
+	SetControlValue(InControlName, Value, bNotify, Context, bSetupUndo, bPrintPythonCommands, bFixEulerFlips);
 	return true;
 }
 
@@ -1756,7 +1756,7 @@ FRigControlValue UControlRig::GetControlValueFromGlobalTransform(const FName& In
 	return Value;
 }
 
-void UControlRig::SetControlLocalTransform(const FName& InControlName, const FTransform& InLocalTransform, bool bNotify, const FRigControlModifiedContext& Context, bool bSetupUndo)
+void UControlRig::SetControlLocalTransform(const FName& InControlName, const FTransform& InLocalTransform, bool bNotify, const FRigControlModifiedContext& Context, bool bSetupUndo, bool bFixEulerFlips)
 {
 	if (FRigControlElement* ControlElement = FindControl(InControlName))
 	{
@@ -1768,7 +1768,7 @@ void UControlRig::SetControlLocalTransform(const FName& InControlName, const FTr
 			OnFilterControl.Broadcast(this, ControlElement, Value);
 			
 		}
-		SetControlValue(InControlName, Value, bNotify, Context, bSetupUndo);
+		SetControlValue(InControlName, Value, bNotify, Context, bSetupUndo, bFixEulerFlips);
 	}
 }
 
diff --git a/Engine/Plugins/Animation/ControlRig/Source/ControlRig/Private/Rigs/RigHierarchy.cpp b/Engine/Plugins/Animation/ControlRig/Source/ControlRig/Private/Rigs/RigHierarchy.cpp
--- a/Engine/Plugins/Animation/ControlRig/Source/ControlRig/Private/Rigs/RigHierarchy.cpp
+++ b/Engine/Plugins/Animation/ControlRig/Source/ControlRig/Private/Rigs/RigHierarchy.cpp
@@ -112,6 +112,7 @@ URigHierarchy::URigHierarchy()
 #endif
 , bEnableCacheValidityCheck(bEnableValidityCheckbyDefault)
 , HierarchyForCacheValidation()
+, bUpdatePreferedEulerAngleWhenSettingTransform(true)
 , ExecuteContext(nullptr)
 #if WITH_EDITOR
 , bRecordTransformsAtRuntime(true)
@@ -2402,7 +2403,22 @@ FTransform URigHierarchy::GetTransform(FRigTransformElement* InTransformElement,
 			if(FRigControlElement* ControlElement = Cast<FRigControlElement>(InTransformElement))
 			{
 				const FTransform NewTransform = ComputeLocalControlValue(ControlElement, ControlElement->Pose.Get(OpposedType), GlobalType);
-				InTransformElement->Pose.Set(InTransformType, NewTransform);
+				ControlElement->Pose.Set(InTransformType, NewTransform);
+				switch(ControlElement->Settings.ControlType)
+				{
+					case ERigControlType::Rotator:
+					case ERigControlType::EulerTransform:
+					case ERigControlType::Transform:
+					case ERigControlType::TransformNoScale:
+					{
+						ControlElement->PreferredEulerAngles.SetRotator(NewTransform.Rotator(), IsInitial(InTransformType), true);
+						break;
+					}
+					default:
+					{
+						break;
+					}
+				}
 			}
 			else if(FRigMultiParentElement* MultiParentElement = Cast<FRigMultiParentElement>(InTransformElement))
 			{
@@ -2514,6 +2530,12 @@ void URigHierarchy::SetTransform(FRigTransformElement* InTransformElement, const
 	if(FRigControlElement* ControlElement = Cast<FRigControlElement>(InTransformElement))
 	{
 		ControlElement->Shape.MarkDirty(MakeGlobal(InTransformType));
+
+		if(bUpdatePreferedEulerAngleWhenSettingTransform && ERigTransformType::IsLocal(InTransformType))
+		{
+			const bool bInitial = ERigTransformType::IsInitial(InTransformType);
+			ControlElement->PreferredEulerAngles.SetRotator(InTransform.Rotator(), bInitial, true);
+		}
 	}
 
 	EnsureCacheValidity();
@@ -3035,10 +3057,40 @@ FRigControlValue URigHierarchy::GetControlValue(FRigControlElement* InControlEle
 
 	if(InControlElement != nullptr)
 	{
+		auto GetValueFromPreferredEulerAngles = [this, InControlElement, &Value, InValueType]() -> bool
+		{
+			const bool bInitial = InValueType == ERigControlValueType::Initial;
+			switch(InControlElement->Settings.ControlType)
+			{
+				case ERigControlType::Rotator:
+				{
+					Value = MakeControlValueFromRotator(InControlElement->PreferredEulerAngles.GetRotator(bInitial)); 
+					return true;
+				}
+				case ERigControlType::EulerTransform:
+				{
+					FEulerTransform EulerTransform(GetTransform(InControlElement, CurrentLocal));
+					EulerTransform.Rotation = InControlElement->PreferredEulerAngles.GetRotator(bInitial);
+					Value = MakeControlValueFromEulerTransform(EulerTransform); 
+					return true;
+				}
+				default:
+				{
+					break;
+				}
+			}
+			return false;
+		};
+		
 		switch(InValueType)
 		{
 			case ERigControlValueType::Current:
 			{
+				if(GetValueFromPreferredEulerAngles())
+				{
+					break;
+				}
+					
 				Value.SetFromTransform(
                     GetTransform(InControlElement, CurrentLocal),
                     InControlElement->Settings.ControlType,
@@ -3048,6 +3100,11 @@ FRigControlValue URigHierarchy::GetControlValue(FRigControlElement* InControlEle
 			}
 			case ERigControlValueType::Initial:
 			{
+				if(GetValueFromPreferredEulerAngles())
+				{
+					break;
+				}
+
 				Value.SetFromTransform(
                     GetTransform(InControlElement, InitialLocal),
                     InControlElement->Settings.ControlType,
@@ -3068,7 +3125,7 @@ FRigControlValue URigHierarchy::GetControlValue(FRigControlElement* InControlEle
 	return Value;
 }
 
-void URigHierarchy::SetControlValue(FRigControlElement* InControlElement, const FRigControlValue& InValue, ERigControlValueType InValueType, bool bSetupUndo, bool bForce, bool bPrintPythonCommands)
+void URigHierarchy::SetControlValue(FRigControlElement* InControlElement, const FRigControlValue& InValue, ERigControlValueType InValueType, bool bSetupUndo, bool bForce, bool bPrintPythonCommands, bool bFixEulerFlips)
 {
 	LLM_SCOPE_BYNAME(TEXT("Animation/ControlRig"));
 	
@@ -3076,13 +3133,46 @@ void URigHierarchy::SetControlValue(FRigControlElement* InControlElement, const
 
 	if(InControlElement != nullptr)
 	{
+		auto SetPreferredEulerAnglesFromValue = [this, InControlElement, &InValue, InValueType, bFixEulerFlips]()
+		{
+			const bool bInitial = InValueType == ERigControlValueType::Initial;
+			switch(InControlElement->Settings.ControlType)
+			{
+				case ERigControlType::Rotator:
+				{
+					InControlElement->PreferredEulerAngles.SetRotator(GetRotatorFromControlValue(InValue), bInitial, bFixEulerFlips);
+					break;
+				}
+				case ERigControlType::EulerTransform:
+				{
+					InControlElement->PreferredEulerAngles.SetRotator(GetEulerTransformFromControlValue(InValue).Rotation, bInitial, bFixEulerFlips);
+					break;
+				}
+				case ERigControlType::Transform:
+				{
+					InControlElement->PreferredEulerAngles.SetRotator(GetTransformFromControlValue(InValue).Rotator(), bInitial, bFixEulerFlips);
+					break;
+				}
+				case ERigControlType::TransformNoScale:
+				{
+					InControlElement->PreferredEulerAngles.SetRotator(GetTransformNoScaleFromControlValue(InValue).Rotation.Rotator(), bInitial, bFixEulerFlips);
+					break;
+				}
+				default:
+				{
+					break;
+				}
+			}
+		};
+
 		switch(InValueType)
 		{
 			case ERigControlValueType::Current:
 			{
 				FRigControlValue Value = InValue;
 				InControlElement->Settings.ApplyLimits(Value);
-					
+
+				TGuardValue<bool> DontSetPreferredEulerAngle(bUpdatePreferedEulerAngleWhenSettingTransform, false);
 				SetTransform(
 					InControlElement,
 					Value.GetAsTransform(
@@ -3094,7 +3184,9 @@ void URigHierarchy::SetControlValue(FRigControlElement* InControlElement, const
 					bSetupUndo,
 					bForce,
 					bPrintPythonCommands
-				); 
+				);
+
+				SetPreferredEulerAnglesFromValue();
 				break;
 			}
 			case ERigControlValueType::Initial:
@@ -3102,6 +3194,7 @@ void URigHierarchy::SetControlValue(FRigControlElement* InControlElement, const
 				FRigControlValue Value = InValue;
 				InControlElement->Settings.ApplyLimits(Value);
 
+				TGuardValue<bool> DontSetPreferredEulerAngle(bUpdatePreferedEulerAngleWhenSettingTransform, false);
 				SetTransform(
 					InControlElement,
 					Value.GetAsTransform(
@@ -3113,7 +3206,9 @@ void URigHierarchy::SetControlValue(FRigControlElement* InControlElement, const
 					bSetupUndo,
 					bForce,
 					bPrintPythonCommands
-				); 
+				);
+
+				SetPreferredEulerAnglesFromValue();
 				break;
 			}
 			case ERigControlValueType::Minimum:
diff --git a/Engine/Plugins/Animation/ControlRig/Source/ControlRig/Private/Rigs/RigHierarchyElements.cpp b/Engine/Plugins/Animation/ControlRig/Source/ControlRig/Private/Rigs/RigHierarchyElements.cpp
--- a/Engine/Plugins/Animation/ControlRig/Source/ControlRig/Private/Rigs/RigHierarchyElements.cpp
+++ b/Engine/Plugins/Animation/ControlRig/Source/ControlRig/Private/Rigs/RigHierarchyElements.cpp
@@ -6,6 +6,7 @@
 #include "Units/RigUnitContext.h"
 #include "ControlRigObjectVersion.h"
 #include "ControlRigGizmoLibrary.h"
+#include "AnimationCoreLibrary.h"
 
 ////////////////////////////////////////////////////////////////////////////////
 // FRigBaseElement
@@ -136,6 +137,84 @@ void FRigCurrentAndInitialTransform::Load(FArchive& Ar)
 	Initial.Load(Ar);
 }
 
+////////////////////////////////////////////////////////////////////////////////
+// FRigPreferredEulerAngles
+////////////////////////////////////////////////////////////////////////////////
+
+void FRigPreferredEulerAngles::Save(FArchive& Ar)
+{
+	static const UEnum* RotationOrderEnum = StaticEnum<EEulerRotationOrder>();
+	FName RotationOrderName = RotationOrderEnum->GetNameByValue((int64)RotationOrder);
+	Ar << RotationOrderName;
+	Ar << Current;
+	Ar << Initial;
+}
+
+void FRigPreferredEulerAngles::Load(FArchive& Ar)
+{
+	static const UEnum* RotationOrderEnum = StaticEnum<EEulerRotationOrder>();
+	FName RotationOrderName;
+	Ar << RotationOrderName;
+	RotationOrder = (EEulerRotationOrder)RotationOrderEnum->GetValueByName(RotationOrderName);
+	Ar << Current;
+	Ar << Initial;
+}
+
+void FRigPreferredEulerAngles::Reset()
+{
+	RotationOrder = DefaultRotationOrder;
+	Initial = Current = FVector::ZeroVector;
+}
+
+FRotator FRigPreferredEulerAngles::GetRotator(bool bInitial) const
+{
+	return FRotator::MakeFromEuler(GetAngles(bInitial, DefaultRotationOrder));
+}
+
+FRotator FRigPreferredEulerAngles::SetRotator(const FRotator& InValue, bool bInitial, bool bFixEulerFlips)
+{
+	if(RotationOrder == DefaultRotationOrder)
+	{
+		if(bFixEulerFlips)
+		{
+			const FRotator CurrentValue = GetRotator(bInitial);
+			
+			//Find Diff of the rotation from current and just add that instead of setting so we can go over/under -180
+			FRotator CurrentWinding;
+			FRotator CurrentRotRemainder;
+			CurrentValue.GetWindingAndRemainder(CurrentWinding, CurrentRotRemainder);
+
+			FRotator DeltaRot = InValue - CurrentRotRemainder;
+			DeltaRot.Normalize();
+			const FRotator FixedValue = CurrentValue + DeltaRot;
+
+			SetAngles(FixedValue.Euler(), bInitial, DefaultRotationOrder);
+			return FixedValue;
+		}
+	}
+	SetAngles(InValue.Euler(), bInitial, DefaultRotationOrder);
+	return InValue;
+}
+
+FVector FRigPreferredEulerAngles::GetAngles(bool bInitial, EEulerRotationOrder InRotationOrder) const
+{
+	if(RotationOrder == InRotationOrder)
+	{
+		return Get(bInitial);
+	}
+	return AnimationCore::ChangeEulerRotationOrder(Get(bInitial), RotationOrder, InRotationOrder);
+}
+
+void FRigPreferredEulerAngles::SetAngles(const FVector& InValue, bool bInitial, EEulerRotationOrder InRotationOrder)
+{
+	FVector Value = InValue;
+	if(RotationOrder != InRotationOrder)
+	{
+		Value = AnimationCore::ChangeEulerRotationOrder(Value, InRotationOrder, RotationOrder);
+	}
+	Get(bInitial) = Value;
+}
+
 ////////////////////////////////////////////////////////////////////////////////
 // FRigTransformElement
 ////////////////////////////////////////////////////////////////////////////////
@@ -782,6 +861,7 @@ void FRigControlElement::Save(FArchive& Ar, URigHierarchy* Hierarchy, ESerializa
 		Settings.Save(Ar);
 		Offset.Save(Ar);
 		Shape.Save(Ar);
+		PreferredEulerAngles.Save(Ar);
 	}
 }
 
@@ -794,6 +874,15 @@ void FRigControlElement::Load(FArchive& Ar, URigHierarchy* Hierarchy, ESerializa
 		Settings.Load(Ar);
 		Offset.Load(Ar);
 		Shape.Load(Ar);
+
+		if (Ar.CustomVer(FControlRigObjectVersion::GUID) >= FControlRigObjectVersion::PreferredEulerAnglesForControls)
+		{
+			PreferredEulerAngles.Load(Ar);
+		}
+		else
+		{
+			PreferredEulerAngles.Reset();
+		}
 	}
 }
 
@@ -805,6 +894,7 @@ void FRigControlElement::CopyFrom(URigHierarchy* InHierarchy, FRigBaseElement* I
 	Settings = Source->Settings;
 	Offset = Source->Offset;
 	Shape = Source->Shape;
+	PreferredEulerAngles = Source->PreferredEulerAngles;
 }
 
 void FRigControlElement::CopyPose(FRigBaseElement* InOther, bool bCurrent, bool bInitial, bool bWeights)
@@ -817,11 +907,13 @@ void FRigControlElement::CopyPose(FRigBaseElement* InOther, bool bCurrent, bool
 		{
 			Offset.Current = Other->Offset.Current;
 			Shape.Current = Other->Shape.Current;
+			PreferredEulerAngles.SetAngles(Other->PreferredEulerAngles.GetAngles(false), false);
 		}
 		if(bInitial)
 		{
 			Offset.Initial = Other->Offset.Initial;
 			Shape.Initial = Other->Shape.Initial;
+			PreferredEulerAngles.SetAngles(Other->PreferredEulerAngles.GetAngles(true), true);
 		}
 	}
 }
diff --git a/Engine/Plugins/Animation/ControlRig/Source/ControlRig/Private/Sequencer/MovieSceneControlRigParameterSection.cpp b/Engine/Plugins/Animation/ControlRig/Source/ControlRig/Private/Sequencer/MovieSceneControlRigParameterSection.cpp
--- a/Engine/Plugins/Animation/ControlRig/Source/ControlRig/Private/Sequencer/MovieSceneControlRigParameterSection.cpp
+++ b/Engine/Plugins/Animation/ControlRig/Source/ControlRig/Private/Sequencer/MovieSceneControlRigParameterSection.cpp
@@ -584,29 +584,9 @@ struct FParameterTransformChannelEditorData
 			FRigControlElement* ControlElement = ControlRig->FindControl(ParameterName);
 			if (ControlElement)
 			{
-				if (ControlElement->Settings.ControlType == ERigControlType::Transform)
-				{
-					const FRigControlValue::FTransform_Float Transform = 
-						ControlRig->GetHierarchy()
-						->GetControlValue(ControlElement, ERigControlValueType::Current).Get<FRigControlValue::FTransform_Float>();
-					return Transform.GetRotation().Rotator();
-				}
-				else if (ControlElement->Settings.ControlType == ERigControlType::TransformNoScale)
-				{
-					const FRigControlValue::FTransformNoScale_Float Transform = 
-						ControlRig->GetHierarchy()
-						->GetControlValue(ControlElement, ERigControlValueType::Current).Get<FRigControlValue::FTransformNoScale_Float>();
-					return Transform.GetRotation().Rotator();
-				}
-				else if (ControlElement->Settings.ControlType == ERigControlType::EulerTransform)
-				{
-					const FRigControlValue::FEulerTransform_Float Transform = 
-						ControlRig->GetHierarchy()
-						->GetControlValue(ControlElement, ERigControlValueType::Current).Get<FRigControlValue::FEulerTransform_Float>();
-					return Transform.GetRotator();
+				return ControlRig->GetHierarchy()->GetControlPreferredRotator(ControlElement);
 			}
 		}
-		}
 		return TOptional<FRotator>();
 	}
 
@@ -650,12 +630,12 @@ struct FParameterTransformChannelEditorData
 		FVector CurrentPos; FRotator CurrentRot;
 		FVector CurrentScale;
 
-		for (const FTransformInterrogationData& Transform : InterrogationData.Iterate<FTransformInterrogationData>(UMovieSceneControlRigParameterSection::GetTransformInterrogationKey()))
+		for (const FEulerTransformInterrogationData& Transform : InterrogationData.Iterate<FEulerTransformInterrogationData>(UMovieSceneControlRigParameterSection::GetTransformInterrogationKey()))
 		{
 			if (Transform.ParameterName == ParameterName)
 			{
-				CurrentPos = Transform.Val.GetTranslation();
-				CurrentRot = Transform.Val.GetRotation().Rotator();
+				CurrentPos = Transform.Val.GetLocation();
+				CurrentRot = Transform.Val.Rotator();
 				CurrentScale = Transform.Val.GetScale3D();
 				break;
 			}
@@ -1207,7 +1187,7 @@ void UMovieSceneControlRigParameterSection::AddColorParameter(FName InParameterN
 	}
 }
 
-void UMovieSceneControlRigParameterSection::AddTransformParameter(FName InParameterName, TOptional<FTransform> DefaultValue, bool bReconstructChannel)
+void UMovieSceneControlRigParameterSection::AddTransformParameter(FName InParameterName, TOptional<FEulerTransform> DefaultValue, bool bReconstructChannel)
 {
 	FTransformParameterNameAndCurves* ExistingCurves = nullptr;
 
@@ -1217,10 +1197,10 @@ void UMovieSceneControlRigParameterSection::AddTransformParameter(FName InParame
 		ExistingCurves = &TransformParameterNamesAndCurves[NewIndex];
 		if (DefaultValue.IsSet())
 		{
-			FTransform& InValue = DefaultValue.GetValue();
-			FVector Translation = InValue.GetTranslation();
-			FRotator Rotator = InValue.GetRotation().Rotator();
-			FVector Scale = InValue.GetScale3D();
+			FEulerTransform& InValue = DefaultValue.GetValue();
+			const FVector& Translation = InValue.GetLocation();
+			const FRotator& Rotator = InValue.Rotator();
+			const FVector& Scale = InValue.GetScale3D();
 			ExistingCurves->Translation[0].SetDefault(Translation[0]);
 			ExistingCurves->Translation[1].SetDefault(Translation[1]);
 			ExistingCurves->Translation[2].SetDefault(Translation[2]);
@@ -2161,14 +2141,14 @@ void UMovieSceneControlRigParameterSection::RecreateWithThisControlRig(UControlR
 		case ERigControlType::TransformNoScale:
 		case ERigControlType::Transform:
 		{
-			TOptional<FTransform> DefaultValue;
+			TOptional<FEulerTransform> DefaultValue;
 			if (bSetDefault)
 			{
 				if (ControlElement->Settings.ControlType == ERigControlType::Transform)
 				{
-					DefaultValue = 
+					DefaultValue = FEulerTransform(
 						ControlRig->GetHierarchy()->
-						GetControlValue(ControlElement, ERigControlValueType::Current).Get<FRigControlValue::FTransform_Float>().ToTransform();
+						GetControlValue(ControlElement, ERigControlValueType::Current).Get<FRigControlValue::FTransform_Float>().ToTransform());
 				}
 				else if (ControlElement->Settings.ControlType == ERigControlType::EulerTransform)
 
@@ -2183,7 +2163,7 @@ void UMovieSceneControlRigParameterSection::RecreateWithThisControlRig(UControlR
 					FTransformNoScale NoScale = 
 						ControlRig->GetHierarchy()->
 						GetControlValue(ControlElement, ERigControlValueType::Current).Get<FRigControlValue::FTransformNoScale_Float>().ToTransform();
-					DefaultValue = NoScale;
+					DefaultValue = FEulerTransform(NoScale.Rotation.Rotator(), NoScale.Location, FVector::OneVector);
 				}
 			}
 			AddTransformParameter(ControlElement->GetName(), DefaultValue, false);
@@ -2979,13 +2959,13 @@ TOptional<FLinearColor>UMovieSceneControlRigParameterSection:: EvaluateColorPara
 	return OptValue;
 }
 
-TOptional<FTransform> UMovieSceneControlRigParameterSection::EvaluateTransformParameter(const  FFrameTime& InTime, FName InParameterName)
+TOptional<FEulerTransform> UMovieSceneControlRigParameterSection::EvaluateTransformParameter(const  FFrameTime& InTime, FName InParameterName)
 {
-	TOptional<FTransform> OptValue;
+	TOptional<FEulerTransform> OptValue;
 	if (const FChannelMapInfo* ChannelInfo = ControlChannelMap.Find(InParameterName))
 	{
 		TArrayView<FMovieSceneFloatChannel*> FloatChannels = ChannelProxy->GetChannels<FMovieSceneFloatChannel>();
-		FTransform Value = FTransform::Identity;
+		FEulerTransform Value = FEulerTransform::Identity;
 		FVector3f Translation(ForceInitToZero), Scale(FVector3f::OneVector);
 		FRotator3f Rotator(0.0f, 0.0f, 0.0f);
 
@@ -3010,7 +2990,7 @@ TOptional<FTransform> UMovieSceneControlRigParameterSection::EvaluateTransformPa
 
 			}
 		}
-		Value = FTransform(FRotator(Rotator), (FVector)Translation, (FVector)Scale);
+		Value = FEulerTransform(FRotator(Rotator), (FVector)Translation, (FVector)Scale);
 		OptValue = Value;
 	}
 	return OptValue;
diff --git a/Engine/Plugins/Animation/ControlRig/Source/ControlRig/Private/Sequencer/MovieSceneControlRigParameterTemplate.cpp b/Engine/Plugins/Animation/ControlRig/Source/ControlRig/Private/Sequencer/MovieSceneControlRigParameterTemplate.cpp
--- a/Engine/Plugins/Animation/ControlRig/Source/ControlRig/Private/Sequencer/MovieSceneControlRigParameterTemplate.cpp
+++ b/Engine/Plugins/Animation/ControlRig/Source/ControlRig/Private/Sequencer/MovieSceneControlRigParameterTemplate.cpp
@@ -155,21 +155,16 @@ struct FColorParameterStringAndValue
 	FLinearColor Value;
 };
 
-struct FTransformParameterStringAndValue
+struct FEulerTransformParameterStringAndValue
 {
 
 	/** The name of the transform  parameter. */
 	FName ParameterName;
-	/** Translation component */
-	FVector Translation;
-	/** Rotation component */
-	FRotator Rotation;
-	/** Scale component */
-	FVector Scale;
-
-	FTransformParameterStringAndValue(FName InParameterName, const FVector& InTranslation,
-		const FRotator& InRotation, const FVector& InScale) : Translation(InTranslation),
-		Rotation(InRotation), Scale(InScale)
+	/** Transform component */
+	FEulerTransform Transform;
+
+	FEulerTransformParameterStringAndValue(FName InParameterName, const FEulerTransform& InTransform)
+		: Transform(InTransform)
 	{
 		ParameterName = InParameterName;
 	}
@@ -377,7 +372,7 @@ struct FEvaluatedControlRigParameterSectionValues
 	/** Array of evaluated color values */
 	TArray<FColorParameterStringAndValue, TInlineAllocator<2>> ColorValues;
 	/** Array of evaluated transform values */
-	TArray<FTransformParameterStringAndValue, TInlineAllocator<2>> TransformValues;
+	TArray<FEulerTransformParameterStringAndValue, TInlineAllocator<2>> TransformValues;
 };
 
 /** Token for control rig control parameters */
@@ -427,10 +422,10 @@ struct FControlRigTrackTokenVector
 struct FControlRigTrackTokenTransform
 {
 	FControlRigTrackTokenTransform() {}
-	FControlRigTrackTokenTransform(FTransform InValue)
+	FControlRigTrackTokenTransform(FEulerTransform InValue)
 		: Value(InValue)
 	{}
-	FTransform Value;
+	FEulerTransform Value;
 
 };
 
@@ -545,8 +540,8 @@ namespace MovieScene
 
 	void MultiChannelFromData(const FControlRigTrackTokenTransform& In, TMultiChannelValue<float, 9>& Out)
 	{
-		FVector Translation = In.Value.GetTranslation();
-		FVector Rotation = In.Value.GetRotation().Rotator().Euler();
+		FVector Translation = In.Value.GetLocation();
+		FVector Rotation = In.Value.Rotator().Euler();
 		FVector Scale = In.Value.GetScale3D();
 		Out = { Translation.X, Translation.Y, Translation.Z, Rotation.X, Rotation.Y, Rotation.Z, Scale.X, Scale.Y, Scale.Z };
 
@@ -554,7 +549,7 @@ namespace MovieScene
 
 	void ResolveChannelsToData(const TMultiChannelValue<float, 9>& In, FControlRigTrackTokenTransform& Out)
 	{
-		Out.Value = FTransform(
+		Out.Value = FEulerTransform(
 			FRotator::MakeFromEuler(FVector(In[3], In[4], In[5])),
 			FVector(In[0], In[1], In[2]),
 			FVector(In[6], In[7], In[8])
@@ -777,7 +772,7 @@ struct FControlRigParameterPreAnimatedTokenProducer : IMovieScenePreAnimatedToke
 									}
 								}
 
-								for (TNameAndValue<FTransform>& Value : TransformValues)
+								for (TNameAndValue<FEulerTransform>& Value : TransformValues)
 								{
 									if (FRigControlElement* ControlElement = ControlRig->FindControl(Value.Name))
 									{
@@ -785,13 +780,15 @@ struct FControlRigParameterPreAnimatedTokenProducer : IMovieScenePreAnimatedToke
 										{
 										case ERigControlType::Transform:
 										{
-											ControlRig->SetControlValue<FRigControlValue::FTransform_Float>(Value.Name, Value.Value, true, FRigControlModifiedContext(EControlRigSetKey::Never), bSetupUndo);
+											ControlRig->SetControlValue<FRigControlValue::FTransform_Float>(Value.Name, Value.Value.ToFTransform(), true, FRigControlModifiedContext(EControlRigSetKey::Never), bSetupUndo);
+											ControlRig->GetHierarchy()->SetControlPreferredRotator(ControlElement, Value.Value.Rotation);
 											break;
 										}
 										case ERigControlType::TransformNoScale:
 										{
-											FTransformNoScale NoScale = Value.Value;
+											FTransformNoScale NoScale = Value.Value.ToFTransform();
 											ControlRig->SetControlValue<FRigControlValue::FTransformNoScale_Float>(Value.Name, NoScale, true, FRigControlModifiedContext(EControlRigSetKey::Never), bSetupUndo);
+											ControlRig->GetHierarchy()->SetControlPreferredRotator(ControlElement, Value.Value.Rotation);
 											break;
 										}
 										case ERigControlType::EulerTransform:
@@ -800,7 +797,10 @@ struct FControlRigParameterPreAnimatedTokenProducer : IMovieScenePreAnimatedToke
 											ControlRig->SetControlValue<FRigControlValue::FEulerTransform_Float>(Value.Name, EulerTransform, true, FRigControlModifiedContext(EControlRigSetKey::Never), bSetupUndo);
 											break;
 										}
-
+										default:
+										{
+											break;
+										}
 										}
 									}
 								}
@@ -845,7 +845,7 @@ struct FControlRigParameterPreAnimatedTokenProducer : IMovieScenePreAnimatedToke
 			TArray< TNameAndValue<int32> > IntegerValues;
 			TArray< TNameAndValue<FVector> > VectorValues;
 			TArray< TNameAndValue<FVector2D> > Vector2DValues;
-			TArray< TNameAndValue<FTransform> > TransformValues;
+			TArray< TNameAndValue<FEulerTransform> > TransformValues;
 			FSkeletalMeshRestoreState SkeletalMeshRestoreState;
 
 		};
@@ -913,7 +913,9 @@ struct FControlRigParameterPreAnimatedTokenProducer : IMovieScenePreAnimatedToke
 						SpaceValue.SpaceType = EMovieSceneControlRigSpaceType::Parent;
 						Token.SpaceValues.Add(FControlSpaceAndValue(ControlElement->GetName(), SpaceValue));
 						const FTransform Val = ControlRig->GetHierarchy()->GetControlValue(ControlElement, ERigControlValueType::Current).Get<FRigControlValue::FTransform_Float>().ToTransform();
-						Token.TransformValues.Add(TNameAndValue<FTransform>{ ControlElement->GetName(), Val });
+						FEulerTransform EulerTransform(Val);
+						EulerTransform.Rotation = ControlRig->GetHierarchy()->GetControlPreferredRotator(ControlElement);
+						Token.TransformValues.Add(TNameAndValue<FEulerTransform>{ ControlElement->GetName(), EulerTransform });
 						break;
 					}
 					case ERigControlType::TransformNoScale:
@@ -921,17 +923,17 @@ struct FControlRigParameterPreAnimatedTokenProducer : IMovieScenePreAnimatedToke
 						const FTransformNoScale NoScale = 
 							ControlRig->GetHierarchy()
 							->GetControlValue(ControlElement, ERigControlValueType::Current).Get<FRigControlValue::FTransformNoScale_Float>().ToTransform();
-						FTransform Val = NoScale;
-						Token.TransformValues.Add(TNameAndValue<FTransform>{ ControlElement->GetName(), Val });
+						FEulerTransform EulerTransform(NoScale.ToFTransform());
+						EulerTransform.Rotation = ControlRig->GetHierarchy()->GetControlPreferredRotator(ControlElement);
+						Token.TransformValues.Add(TNameAndValue<FEulerTransform>{ ControlElement->GetName(), EulerTransform });
 						break;
 					}
 					case ERigControlType::EulerTransform:
 					{
-						const FEulerTransform Euler = 
+						const FEulerTransform EulerTransform = 
 							ControlRig->GetHierarchy()
 							->GetControlValue(ControlElement, ERigControlValueType::Current).Get<FRigControlValue::FEulerTransform_Float>().ToTransform();
-						FTransform Val = Euler.ToFTransform();
-						Token.TransformValues.Add(TNameAndValue<FTransform>{ ControlElement->GetName(), Val });
+						Token.TransformValues.Add(TNameAndValue<FEulerTransform>{ ControlElement->GetName(), EulerTransform });
 						break;
 					}
 				}
@@ -1482,20 +1484,23 @@ struct TControlRigParameterActuatorTransform : TMovieSceneBlendingActuator<FCont
 			FRigControlElement* ControlElement = ControlRig->FindControl(ParameterName);
 			if (ControlElement && ControlElement->Settings.ControlType == ERigControlType::Transform)
 			{
-				FTransform Val = ControlRig->GetHierarchy()->GetControlValue(ControlElement, ERigControlValueType::Current).Get<FRigControlValue::FTransform_Float>().ToTransform();
-				return FControlRigTrackTokenTransform(Val);
+				const FTransform Val = ControlRig->GetHierarchy()->GetControlValue(ControlElement, ERigControlValueType::Current).Get<FRigControlValue::FTransform_Float>().ToTransform();
+				FEulerTransform EulerTransform(Val);
+				EulerTransform.Rotation = ControlRig->GetHierarchy()->GetControlPreferredRotator(ControlElement);
+				return FControlRigTrackTokenTransform(EulerTransform);
 			}
 			else if(ControlElement && ControlElement->Settings.ControlType == ERigControlType::TransformNoScale)
 			{
 				FTransformNoScale ValNoScale = ControlRig->GetHierarchy()->GetControlValue(ControlElement, ERigControlValueType::Current).Get<FRigControlValue::FTransformNoScale_Float>().ToTransform();
 				FTransform Val = ValNoScale;
-				return FControlRigTrackTokenTransform(Val);
+				FEulerTransform EulerTransform(Val);
+				EulerTransform.Rotation = ControlRig->GetHierarchy()->GetControlPreferredRotator(ControlElement);
+				return FControlRigTrackTokenTransform(EulerTransform);
 			}
 			else if (ControlElement && ControlElement->Settings.ControlType == ERigControlType::EulerTransform)
 			{
-				FEulerTransform Euler = ControlRig->GetHierarchy()->GetControlValue(ControlElement, ERigControlValueType::Current).Get<FRigControlValue::FEulerTransform_Float>().ToTransform();
-				FTransform Val = Euler.ToFTransform();
-				return FControlRigTrackTokenTransform(Val);
+				FEulerTransform EulerTransform = ControlRig->GetHierarchy()->GetControlValue(ControlElement, ERigControlValueType::Current).Get<FRigControlValue::FEulerTransform_Float>().ToTransform();
+				return FControlRigTrackTokenTransform(EulerTransform);
 		}
 		}
 		return FControlRigTrackTokenTransform();
@@ -1526,17 +1531,18 @@ struct TControlRigParameterActuatorTransform : TMovieSceneBlendingActuator<FCont
 				FRigControlElement* ControlElement = ControlRig->FindControl(ParameterName);
 				if (ControlElement && ControlElement->Settings.ControlType == ERigControlType::Transform)
 				{
-						ControlRig->SetControlValue<FRigControlValue::FTransform_Float>(ParameterName, InFinalValue.Value,true, EControlRigSetKey::Never,bSetupUndo);
+					ControlRig->SetControlValue<FRigControlValue::FTransform_Float>(ParameterName, InFinalValue.Value.ToFTransform(),true, EControlRigSetKey::Never,bSetupUndo);
+					ControlRig->GetHierarchy()->SetControlPreferredRotator(ControlElement, InFinalValue.Value.Rotator());
 				}
 				else if (ControlElement && ControlElement->Settings.ControlType == ERigControlType::TransformNoScale)
 				{
-					FTransformNoScale NoScale = InFinalValue.Value;
-						ControlRig->SetControlValue<FRigControlValue::FTransformNoScale_Float>(ParameterName, NoScale, true, EControlRigSetKey::Never,bSetupUndo);
+					const FTransformNoScale NoScale = InFinalValue.Value.ToFTransform();
+					ControlRig->SetControlValue<FRigControlValue::FTransformNoScale_Float>(ParameterName, NoScale, true, EControlRigSetKey::Never,bSetupUndo);
+					ControlRig->GetHierarchy()->SetControlPreferredRotator(ControlElement, InFinalValue.Value.Rotator());
 				}
 				else if (ControlElement && ControlElement->Settings.ControlType == ERigControlType::EulerTransform)
 				{
-					FEulerTransform Euler = InFinalValue.Value;
-					ControlRig->SetControlValue<FRigControlValue::FEulerTransform_Float>(ParameterName, Euler, true, EControlRigSetKey::Never,bSetupUndo);
+					ControlRig->SetControlValue<FRigControlValue::FEulerTransform_Float>(ParameterName, InFinalValue.Value, true, EControlRigSetKey::Never,bSetupUndo);
 				}
 			}
 
@@ -1545,10 +1551,10 @@ struct TControlRigParameterActuatorTransform : TMovieSceneBlendingActuator<FCont
 	}
 	virtual void Actuate(FMovieSceneInterrogationData& InterrogationData, const FControlRigTrackTokenTransform& InValue, const TBlendableTokenStack<FControlRigTrackTokenTransform>& OriginalStack, const FMovieSceneContext& Context) const override
 	{
-		FTransformInterrogationData Data;
+		FEulerTransformInterrogationData Data;
 		Data.Val = InValue.Value;
 		Data.ParameterName = ParameterName;
-		InterrogationData.Add(FTransformInterrogationData(Data), UMovieSceneControlRigParameterSection::GetTransformInterrogationKey());
+		InterrogationData.Add(FEulerTransformInterrogationData(Data), UMovieSceneControlRigParameterSection::GetTransformInterrogationKey());
 	}
 	FName ParameterName;
 	TWeakObjectPtr<const UMovieSceneControlRigParameterSection> SectionData;
@@ -1646,7 +1652,7 @@ void FMovieSceneControlRigParameterTemplate::Evaluate(const FMovieSceneEvaluatio
 		}
 
 		UE::MovieScene::TMultiChannelValue<float, 9> TransformData;
-		for (const FTransformParameterStringAndValue& TransformNameAndValue : Values.TransformValues)
+		for (const FEulerTransformParameterStringAndValue& TransformNameAndValue : Values.TransformValues)
 		{
 			FMovieSceneAnimTypeID AnimTypeID = TypeIDs->FindTransform(TransformNameAndValue.ParameterName);
 			FMovieSceneBlendingActuatorID ActuatorTypeID(AnimTypeID);
@@ -1656,19 +1662,19 @@ void FMovieSceneControlRigParameterTemplate::Evaluate(const FMovieSceneEvaluatio
 				ExecutionTokens.GetBlendingAccumulator().DefineActuator(ActuatorTypeID, MakeShared <TControlRigParameterActuatorTransform>(AnimTypeID,  TransformNameAndValue.ParameterName, Section));
 			}
 
-			FTransform Transform(TransformNameAndValue.Rotation, TransformNameAndValue.Translation, TransformNameAndValue.Scale);
+			const FEulerTransform& Transform = TransformNameAndValue.Transform;
 
-			TransformData.Set(0, TransformNameAndValue.Translation.X);
-			TransformData.Set(1, TransformNameAndValue.Translation.Y);
-			TransformData.Set(2, TransformNameAndValue.Translation.Z);
+			TransformData.Set(0, Transform.Location.X);
+			TransformData.Set(1, Transform.Location.Y);
+			TransformData.Set(2, Transform.Location.Z);
 
-			TransformData.Set(3, TransformNameAndValue.Rotation.Roll);
-			TransformData.Set(4, TransformNameAndValue.Rotation.Pitch);
-			TransformData.Set(5, TransformNameAndValue.Rotation.Yaw);
+			TransformData.Set(3, Transform.Rotation.Roll);
+			TransformData.Set(4, Transform.Rotation.Pitch);
+			TransformData.Set(5, Transform.Rotation.Yaw);
 
-			TransformData.Set(6, TransformNameAndValue.Scale.X);
-			TransformData.Set(7, TransformNameAndValue.Scale.Y);
-			TransformData.Set(8, TransformNameAndValue.Scale.Z);
+			TransformData.Set(6, Transform.Scale.X);
+			TransformData.Set(7, Transform.Scale.Y);
+			TransformData.Set(8, Transform.Scale.Z);
 			ExecutionTokens.BlendToken(ActuatorTypeID, TBlendableToken<FControlRigTrackTokenTransform>(TransformData, Section->GetBlendType().Get(), Weight));
 		}
 
@@ -2038,7 +2044,7 @@ void FMovieSceneControlRigParameterTemplate::EvaluateCurvesWithMasks(const FMovi
 					}
 				}
 			}
-			FTransformParameterStringAndValue NameAndValue(Transform.ParameterName, (FVector)Translation, FRotator(Rotator), (FVector)Scale);
+			FEulerTransformParameterStringAndValue NameAndValue(Transform.ParameterName, FEulerTransform(FRotator(Rotator), (FVector)Translation, (FVector)Scale));
 			Values.TransformValues.Emplace(NameAndValue);
 		}
 	}
@@ -2129,7 +2135,7 @@ void FMovieSceneControlRigParameterTemplate::Interrogate(const FMovieSceneContex
 		}
 
 		UE::MovieScene::TMultiChannelValue<float, 9> TransformData;
-		for (const FTransformParameterStringAndValue& TransformNameAndValue : Values.TransformValues)
+		for (const FEulerTransformParameterStringAndValue& TransformNameAndValue : Values.TransformValues)
 		{
 			FMovieSceneAnimTypeID AnimTypeID = TypeIDs->FindTransform(TransformNameAndValue.ParameterName);
 			FMovieSceneBlendingActuatorID ActuatorTypeID(AnimTypeID);
@@ -2139,19 +2145,19 @@ void FMovieSceneControlRigParameterTemplate::Interrogate(const FMovieSceneContex
 				Container.GetAccumulator().DefineActuator(ActuatorTypeID, MakeShared <TControlRigParameterActuatorTransform>(AnimTypeID, TransformNameAndValue.ParameterName, Section));
 			}
 
-			FTransform Transform(TransformNameAndValue.Rotation, TransformNameAndValue.Translation, TransformNameAndValue.Scale);
+			const FEulerTransform& Transform = TransformNameAndValue.Transform;
 
-			TransformData.Set(0, TransformNameAndValue.Translation.X);
-			TransformData.Set(1, TransformNameAndValue.Translation.Y);
-			TransformData.Set(2, TransformNameAndValue.Translation.Z);
+			TransformData.Set(0, Transform.Location.X);
+			TransformData.Set(1, Transform.Location.Y);
+			TransformData.Set(2, Transform.Location.Z);
 
-			TransformData.Set(3, TransformNameAndValue.Rotation.Roll);
-			TransformData.Set(4, TransformNameAndValue.Rotation.Pitch);
-			TransformData.Set(5, TransformNameAndValue.Rotation.Yaw);
+			TransformData.Set(3, Transform.Rotation.Roll);
+			TransformData.Set(4, Transform.Rotation.Pitch);
+			TransformData.Set(5, Transform.Rotation.Yaw);
 
-			TransformData.Set(6, TransformNameAndValue.Scale.X);
-			TransformData.Set(7, TransformNameAndValue.Scale.Y);
-			TransformData.Set(8, TransformNameAndValue.Scale.Z);
+			TransformData.Set(6, Transform.Scale.X);
+			TransformData.Set(7, Transform.Scale.Y);
+			TransformData.Set(8, Transform.Scale.Z);
 			Container.GetAccumulator().BlendToken(FMovieSceneEvaluationOperand(),ActuatorTypeID, FMovieSceneEvaluationScope(), Context, TBlendableToken<FControlRigTrackTokenTransform>(TransformData, Section->GetBlendType().Get(), Weight));
 		}
 
diff --git a/Engine/Plugins/Animation/ControlRig/Source/ControlRig/Public/ControlRig.h b/Engine/Plugins/Animation/ControlRig/Source/ControlRig/Public/ControlRig.h
--- a/Engine/Plugins/Animation/ControlRig/Source/ControlRig/Public/ControlRig.h
+++ b/Engine/Plugins/Animation/ControlRig/Source/ControlRig/Public/ControlRig.h
@@ -318,9 +318,9 @@ public:
 	// Sets the relative value of a Control
 	template<class T>
 	FORCEINLINE_DEBUGGABLE void SetControlValue(const FName& InControlName, T InValue, bool bNotify = true,
-		const FRigControlModifiedContext& Context = FRigControlModifiedContext(), bool bSetupUndo = true, bool bPrintPythonCommnds = false)
+		const FRigControlModifiedContext& Context = FRigControlModifiedContext(), bool bSetupUndo = true, bool bPrintPythonCommnds = false, bool bFixEulerFlips = false)
 	{
-		SetControlValueImpl(InControlName, FRigControlValue::Make<T>(InValue), bNotify, Context, bSetupUndo, bPrintPythonCommnds);
+		SetControlValueImpl(InControlName, FRigControlValue::Make<T>(InValue), bNotify, Context, bSetupUndo, bPrintPythonCommnds, bFixEulerFlips);
 	}
 
 	// Returns the value of a Control
@@ -332,7 +332,7 @@ public:
 
 	// Sets the relative value of a Control
 	FORCEINLINE_DEBUGGABLE virtual void SetControlValueImpl(const FName& InControlName, const FRigControlValue& InValue, bool bNotify = true,
-		const FRigControlModifiedContext& Context = FRigControlModifiedContext(), bool bSetupUndo = true, bool bPrintPythonCommnds = false)
+		const FRigControlModifiedContext& Context = FRigControlModifiedContext(), bool bSetupUndo = true, bool bPrintPythonCommnds = false, bool bFixEulerFlips = false)
 	{
 		const FRigElementKey Key(InControlName, ERigElementType::Control);
 
@@ -342,7 +342,7 @@ public:
 			return;
 		}
 
-		DynamicHierarchy->SetControlValue(ControlElement, InValue, ERigControlValueType::Current, bSetupUndo, false, bPrintPythonCommnds);
+		DynamicHierarchy->SetControlValue(ControlElement, InValue, ERigControlValueType::Current, bSetupUndo, false, bPrintPythonCommnds, bFixEulerFlips);
 
 		if (bNotify && OnControlModified.IsBound())
 		{
@@ -350,11 +350,11 @@ public:
 		}
 	}
 
-	bool SetControlGlobalTransform(const FName& InControlName, const FTransform& InGlobalTransform, bool bNotify = true, const FRigControlModifiedContext& Context = FRigControlModifiedContext(), bool bSetupUndo = true, bool bPrintPythonCommands = false);
+	bool SetControlGlobalTransform(const FName& InControlName, const FTransform& InGlobalTransform, bool bNotify = true, const FRigControlModifiedContext& Context = FRigControlModifiedContext(), bool bSetupUndo = true, bool bPrintPythonCommands = false, bool bFixEulerFlips = false);
 
 	virtual FRigControlValue GetControlValueFromGlobalTransform(const FName& InControlName, const FTransform& InGlobalTransform);
 
-	virtual void SetControlLocalTransform(const FName& InControlName, const FTransform& InLocalTransform, bool bNotify = true, const FRigControlModifiedContext& Context = FRigControlModifiedContext(), bool bSetupUndo = true);
+	virtual void SetControlLocalTransform(const FName& InControlName, const FTransform& InLocalTransform, bool bNotify = true, const FRigControlModifiedContext& Context = FRigControlModifiedContext(), bool bSetupUndo = true, bool bFixEulerFlips = false);
 	virtual FTransform GetControlLocalTransform(const FName& InControlName) ;
 
 	virtual const TArray<TSoftObjectPtr<UControlRigShapeLibrary>>& GetShapeLibraries() const;
diff --git a/Engine/Plugins/Animation/ControlRig/Source/ControlRig/Public/ControlRigObjectVersion.h b/Engine/Plugins/Animation/ControlRig/Source/ControlRig/Public/ControlRigObjectVersion.h
--- a/Engine/Plugins/Animation/ControlRig/Source/ControlRig/Public/ControlRigObjectVersion.h
+++ b/Engine/Plugins/Animation/ControlRig/Source/ControlRig/Public/ControlRigObjectVersion.h
@@ -66,6 +66,9 @@ struct CONTROLRIG_API FControlRigObjectVersion
 		// Added preferred permutation for templates
 		TemplatesPreferredPermutatation,
 
+		// Added preferred euler angles to controls
+		PreferredEulerAnglesForControls,
+
 		// -----<new versions can be added above this line>-------------------------------------------------
 		VersionPlusOne,
 		LatestVersion = VersionPlusOne - 1
diff --git a/Engine/Plugins/Animation/ControlRig/Source/ControlRig/Public/Rigs/RigHierarchy.h b/Engine/Plugins/Animation/ControlRig/Source/ControlRig/Public/Rigs/RigHierarchy.h
--- a/Engine/Plugins/Animation/ControlRig/Source/ControlRig/Public/Rigs/RigHierarchy.h
+++ b/Engine/Plugins/Animation/ControlRig/Source/ControlRig/Public/Rigs/RigHierarchy.h
@@ -1256,6 +1256,100 @@ public:
 		return GetInitialControlValue(InElementIndex).Get<T>();
 	}
 
+	/**
+	 * Returns a control's preferred rotator (local transform rotation)
+	 * @param InKey The key of the element to retrieve the current value for
+	 * @param bInitial If true we'll return the preferred rotator for the initial - otherwise current transform
+	 * @return Returns the current preferred rotator
+	 */
+	UFUNCTION(BlueprintCallable, Category = URigHierarchy)
+	FORCEINLINE_DEBUGGABLE FRotator GetControlPreferredRotator(FRigElementKey InKey, bool bInitial = false) const
+	{
+		return GetControlPreferredRotatorByIndex(GetIndex(InKey), bInitial);
+	}
+
+	/**
+	 * Returns a control's preferred rotator (local transform rotation)
+	 * @param InElementIndex The element index to look up
+	 * @param bInitial If true we'll return the preferred rotator for the initial - otherwise current transform
+	 * @return Returns the current preferred rotator
+	 */
+	UFUNCTION(BlueprintCallable, Category = URigHierarchy)
+	FORCEINLINE_DEBUGGABLE FRotator GetControlPreferredRotatorByIndex(int32 InElementIndex, bool bInitial = false) const
+	{
+		if(Elements.IsValidIndex(InElementIndex))
+		{
+			if(FRigControlElement* ControlElement = Cast<FRigControlElement>(Elements[InElementIndex]))
+			{
+				return GetControlPreferredRotator(ControlElement, bInitial);
+			}
+		}
+		return FRotator::ZeroRotator;
+	}
+
+	/**
+	 * Returns a control's preferred rotator (local transform rotation)
+	 * @param InControlElement The element to look up
+	 * @param bInitial If true we'll return the preferred rotator for the initial - otherwise current transform
+	 * @return Returns the current preferred rotator
+	 */
+	FORCEINLINE_DEBUGGABLE FRotator GetControlPreferredRotator(FRigControlElement* InControlElement, bool bInitial = false) const
+	{
+		if(InControlElement)
+		{
+			return InControlElement->PreferredEulerAngles.GetRotator(bInitial);
+		}
+		return FRotator::ZeroRotator;
+	}
+
+	/**
+	 * Sets a control's preferred rotator (local transform rotation)
+	 * @param InKey The key of the element to retrieve the current value for
+	 * @param InValue The new preferred rotator to set
+	 * @param bInitial If true we'll return the preferred rotator for the initial - otherwise current transform
+	 * @param bFixEulerFlips If true the new rotator value will use the shortest path
+	 */
+	UFUNCTION(BlueprintCallable, Category = URigHierarchy)
+	FORCEINLINE_DEBUGGABLE void SetControlPreferredRotator(FRigElementKey InKey, const FRotator& InValue, bool bInitial = false, bool bFixEulerFlips = false)
+	{
+		SetControlPreferredRotatorByIndex(GetIndex(InKey), InValue, bInitial, bFixEulerFlips);
+	}
+	
+
+	/**
+	 * Sets a control's preferred rotator (local transform rotation)
+	 * @param InElementIndex The element index to look up
+	 * @param InValue The new preferred rotator to set
+	 * @param bInitial If true we'll return the preferred rotator for the initial - otherwise current transform
+	 * @param bFixEulerFlips If true the new rotator value will use the shortest path
+	 */
+	UFUNCTION(BlueprintCallable, Category = URigHierarchy)
+	FORCEINLINE_DEBUGGABLE void SetControlPreferredRotatorByIndex(int32 InElementIndex, const FRotator& InValue, bool bInitial = false, bool bFixEulerFlips = false)
+	{
+		if(Elements.IsValidIndex(InElementIndex))
+		{
+			if(FRigControlElement* ControlElement = Cast<FRigControlElement>(Elements[InElementIndex]))
+			{
+				SetControlPreferredRotator(ControlElement, InValue, bInitial, bFixEulerFlips);
+			}
+		}
+	}
+
+	/**
+	 * Sets a control's preferred rotator (local transform rotation)
+	 * @param InControlElement The element to look up
+	 * @param InValue The new preferred rotator to set
+	 * @param bInitial If true we'll return the preferred rotator for the initial - otherwise current transform
+	 * @param bFixEulerFlips If true the new rotator value will use the shortest path
+	 */
+	FORCEINLINE_DEBUGGABLE void SetControlPreferredRotator(FRigControlElement* InControlElement, const FRotator& InValue, bool bInitial = false, bool bFixEulerFlips = false)
+	{
+		if(InControlElement)
+		{
+			InControlElement->PreferredEulerAngles.SetRotator(InValue, bInitial, bFixEulerFlips);
+		}
+	}
+	
 	/**
 	 * Returns the pin type to use for a control
 	 * @param InControlElement The control to return the pin type for
@@ -2527,7 +2621,7 @@ public:
 	 * @param bSetupUndo If true the transform stack will be setup for undo / redo
 	 * @param bForce Set the transform even if it is the same as the previously set one
 	 */
-	void SetControlValue(FRigControlElement* InControlElement, const FRigControlValue& InValue, ERigControlValueType InValueType, bool bSetupUndo = false, bool bForce = false, bool bPrintPythonCommands = false);
+	void SetControlValue(FRigControlElement* InControlElement, const FRigControlValue& InValue, ERigControlValueType InValueType, bool bSetupUndo = false, bool bForce = false, bool bPrintPythonCommands = false, bool bFixEulerFlips = false);
 
 	template<typename T>
 	FORCEINLINE_DEBUGGABLE void SetControlValue(FRigControlElement* InControlElement, const T& InValue, ERigControlValueType InValueType, bool bSetupUndo = false, bool bForce = false) const
@@ -3130,7 +3224,9 @@ protected:
 	UPROPERTY(transient)
 	TObjectPtr<URigHierarchy> HierarchyForCacheValidation;
 
-	mutable TMap<FRigElementKey, FRigElementKey> DefaultParentPerElement;  
+	mutable TMap<FRigElementKey, FRigElementKey> DefaultParentPerElement;
+
+	bool bUpdatePreferedEulerAngleWhenSettingTransform;
 	
 private:
 	
diff --git a/Engine/Plugins/Animation/ControlRig/Source/ControlRig/Public/Rigs/RigHierarchyElements.h b/Engine/Plugins/Animation/ControlRig/Source/ControlRig/Public/Rigs/RigHierarchyElements.h
--- a/Engine/Plugins/Animation/ControlRig/Source/ControlRig/Public/Rigs/RigHierarchyElements.h
+++ b/Engine/Plugins/Animation/ControlRig/Source/ControlRig/Public/Rigs/RigHierarchyElements.h
@@ -397,6 +397,48 @@ struct CONTROLRIG_API FRigCurrentAndInitialTransform
 	FRigLocalAndGlobalTransform Initial;
 };
 
+USTRUCT(BlueprintType)
+struct CONTROLRIG_API FRigPreferredEulerAngles
+{
+	GENERATED_BODY()
+
+	static constexpr EEulerRotationOrder DefaultRotationOrder = EEulerRotationOrder::YZX;
+
+	FRigPreferredEulerAngles()
+	: RotationOrder(DefaultRotationOrder) // default for rotator
+	, Current(FVector::ZeroVector)
+	, Initial(FVector::ZeroVector)
+	{}
+
+	void Save(FArchive& Ar);
+	void Load(FArchive& Ar);
+
+	FORCEINLINE bool operator == (const FRigPreferredEulerAngles& Other) const
+	{
+		return RotationOrder == Other.RotationOrder &&
+			Current == Other.Current &&
+			Initial == Other.Initial;
+	}
+
+	void Reset();
+	FORCEINLINE FVector& Get(bool bInitial = false) { return bInitial ? Initial : Current; }
+	FORCEINLINE const FVector& Get(bool bInitial = false) const { return bInitial ? Initial : Current; }
+	FRotator GetRotator(bool bInitial = false) const;
+	FRotator SetRotator(const FRotator& InValue, bool bInitial = false, bool bFixEulerFlips = false);
+	FVector GetAngles(bool bInitial = false, EEulerRotationOrder InRotationOrder = DefaultRotationOrder) const;
+	void SetAngles(const FVector& InValue, bool bInitial = false, EEulerRotationOrder InRotationOrder = DefaultRotationOrder);
+
+	UPROPERTY(BlueprintReadOnly, EditAnywhere, Category = "Pose")
+	EEulerRotationOrder RotationOrder;
+
+	UPROPERTY(BlueprintReadOnly, EditAnywhere, Category = "Pose")
+	FVector Current;
+
+	UPROPERTY(BlueprintReadOnly, EditAnywhere, Category = "Pose")
+	FVector Initial;
+};
+
+
 struct FRigBaseElement;
 //typedef TArray<FRigBaseElement*> FRigBaseElementChildrenArray;
 typedef TArray<FRigBaseElement*, TInlineAllocator<3>> FRigBaseElementChildrenArray;
@@ -1107,6 +1149,9 @@ public:
 	FRigCurrentAndInitialTransform Shape;
 
 protected:
+
+	UPROPERTY(BlueprintReadOnly, EditAnywhere, Category = RigElement)
+	FRigPreferredEulerAngles PreferredEulerAngles;
 	
 	FORCEINLINE static bool IsClassOf(const FRigBaseElement* InElement)
 	{
@@ -1120,6 +1165,7 @@ public:
 protected:
 
 	friend struct FRigBaseElement;
+	friend class URigHierarchy;
 };
 
 USTRUCT(BlueprintType)
diff --git a/Engine/Plugins/Animation/ControlRig/Source/ControlRig/Public/Sequencer/MovieSceneControlRigParameterSection.h b/Engine/Plugins/Animation/ControlRig/Source/ControlRig/Public/Sequencer/MovieSceneControlRigParameterSection.h
--- a/Engine/Plugins/Animation/ControlRig/Source/ControlRig/Public/Sequencer/MovieSceneControlRigParameterSection.h
+++ b/Engine/Plugins/Animation/ControlRig/Source/ControlRig/Public/Sequencer/MovieSceneControlRigParameterSection.h
@@ -124,9 +124,9 @@ struct FVectorInterrogationData
 	FName ParameterName;
 };
 
-struct FTransformInterrogationData
+struct FEulerTransformInterrogationData
 {
-	FTransform Val;
+	FEulerTransform Val;
 	FName ParameterName;
 };
 
@@ -383,7 +383,7 @@ public:
 	void AddColorParameter(FName InParameterName, TOptional<FLinearColor> DefaultValue, bool bReconstructChannel);
 
 	/** Adds a a key for a specific transform parameter*/
-	void AddTransformParameter(FName InParameterName, TOptional<FTransform> DefaultValue, bool bReconstructChannel);
+	void AddTransformParameter(FName InParameterName, TOptional<FEulerTransform> DefaultValue, bool bReconstructChannel);
 
 	/** Add Space Parameter for a specified Control, no Default since that is Parent space*/
 	void AddSpaceChannel(FName InControlName, bool bReconstructChannel);
@@ -413,7 +413,8 @@ public:
 	TOptional<FLinearColor> EvaluateColorParameter(const  FFrameTime& InTime, FName InParameterName);
 
 	/** Evaluates a a key for a specific transform parameter. Will not get set if not found */
-	TOptional<FTransform> EvaluateTransformParameter(const  FFrameTime& InTime, FName InParameterName);
+	TOptional<FEulerTransform> EvaluateTransformParameter(const  FFrameTime& InTime, FName InParameterName);
+	
 
 	/** Evaluates a a key for a specific space parameter. Will not get set if not found */
 	TOptional<FMovieSceneControlRigSpaceBaseKey> EvaluateSpaceChannel(const  FFrameTime& InTime, FName InParameterName);
diff --git a/Engine/Plugins/Animation/ControlRig/Source/ControlRigEditor/Private/ControlRigElementDetails.cpp b/Engine/Plugins/Animation/ControlRig/Source/ControlRigEditor/Private/ControlRigElementDetails.cpp
--- a/Engine/Plugins/Animation/ControlRig/Source/ControlRigEditor/Private/ControlRigElementDetails.cpp
+++ b/Engine/Plugins/Animation/ControlRig/Source/ControlRigEditor/Private/ControlRigElementDetails.cpp
@@ -1430,8 +1430,27 @@ void FRigTransformElementDetails::CreateEulerTransformValueWidgetRow(
 		const bool bInitial = InTransformType == ERigTransformElementDetailsTransform::Initial; 
 		if(bInitial || InTransformType == ERigTransformElementDetailsTransform::Current)
 		{
-			RelativeTransform = HierarchyBeingDebugged->GetLocalTransform(Key, bInitial);
-			AbsoluteTransform = HierarchyBeingDebugged->GetGlobalTransform(Key, bInitial);
+			RelativeTransform.FromFTransform(HierarchyBeingDebugged->GetLocalTransform(Key, bInitial));
+			AbsoluteTransform.FromFTransform(HierarchyBeingDebugged->GetGlobalTransform(Key, bInitial));
+
+			if(FRigControlElement* ControlElement = HierarchyBeingDebugged->Find<FRigControlElement>(Key))
+			{
+				switch(ControlElement->Settings.ControlType)
+				{
+					case ERigControlType::Rotator:
+					case ERigControlType::EulerTransform:
+					case ERigControlType::Transform:
+					case ERigControlType::TransformNoScale:
+					{
+						RelativeTransform.Rotation = HierarchyBeingDebugged->GetControlPreferredRotator(ControlElement, bInitial);
+						break;
+					}
+					default:
+					{
+						break;
+					}
+				}
+			}
 		}
 		else
 		{
@@ -1441,8 +1460,8 @@ void FRigTransformElementDetails::CreateEulerTransformValueWidgetRow(
 
 				if(InTransformType == ERigTransformElementDetailsTransform::Offset)
 				{
-					RelativeTransform = HierarchyBeingDebugged->GetControlOffsetTransform(ControlElement, ERigTransformType::InitialLocal);
-					AbsoluteTransform = HierarchyBeingDebugged->GetControlOffsetTransform(ControlElement, ERigTransformType::InitialGlobal);
+					RelativeTransform.FromFTransform(HierarchyBeingDebugged->GetControlOffsetTransform(ControlElement, ERigTransformType::InitialLocal));
+					AbsoluteTransform.FromFTransform(HierarchyBeingDebugged->GetControlOffsetTransform(ControlElement, ERigTransformType::InitialGlobal));
 				}
 				else if(InTransformType == ERigTransformElementDetailsTransform::Minimum)
 				{
@@ -1525,7 +1544,7 @@ void FRigTransformElementDetails::CreateEulerTransformValueWidgetRow(
 			}
 		}
 
-		return TPair<FTransform, FTransform>(RelativeTransform, AbsoluteTransform);
+		return TPair<FEulerTransform, FEulerTransform>(RelativeTransform, AbsoluteTransform);
 	};
 
 	
@@ -1540,7 +1559,7 @@ void FRigTransformElementDetails::CreateEulerTransformValueWidgetRow(
 
 		FEulerTransform Xfo;
 		Xfo.SetLocation((IsComponentRelative(0)) ? RelativeTransform.GetLocation() : AbsoluteTransform.GetLocation());
-		Xfo.SetRotation((IsComponentRelative(1)) ? RelativeTransform.GetRotation() : AbsoluteTransform.GetRotation());
+		Xfo.SetRotator((IsComponentRelative(1)) ? RelativeTransform.Rotator() : AbsoluteTransform.Rotator());
 		Xfo.SetScale3D((IsComponentRelative(2)) ? RelativeTransform.GetScale3D() : AbsoluteTransform.GetScale3D());
 
 		return Xfo;
@@ -1586,11 +1605,30 @@ void FRigTransformElementDetails::CreateEulerTransformValueWidgetRow(
 			{
 				if(bIsRelative)
 				{
-					HierarchyToUpdate->SetLocalTransform(Key, InTransform, bInitial, true, bSetupUndoRedo);
+					HierarchyToUpdate->SetLocalTransform(Key, InTransform.ToFTransform(), bInitial, true, bSetupUndoRedo);
+
+					if(FRigControlElement* ControlElement = HierarchyToUpdate->Find<FRigControlElement>(Key))
+					{
+						switch(ControlElement->Settings.ControlType)
+						{
+							case ERigControlType::Rotator:
+							case ERigControlType::EulerTransform:
+							case ERigControlType::Transform:
+							case ERigControlType::TransformNoScale:
+							{
+								HierarchyToUpdate->SetControlPreferredRotator(ControlElement, InTransform.Rotator(), bInitial);
+								break;
+							}
+							default:
+							{
+								break;
+							}
+						}
+					}
 				}
 				else
 				{
-					HierarchyToUpdate->SetGlobalTransform(Key, InTransform, bInitial, true, bSetupUndoRedo);
+					HierarchyToUpdate->SetGlobalTransform(Key, InTransform.ToFTransform(), bInitial, true, bSetupUndoRedo);
 				}
 			}
 			else
@@ -1604,9 +1642,9 @@ void FRigTransformElementDetails::CreateEulerTransformValueWidgetRow(
 						if(!bIsRelative)
 						{
 							const FTransform ParentTransform = HierarchyToUpdate->GetParentTransform(Key, bInitial);
-							InTransform = FTransform(InTransform).GetRelativeTransform(ParentTransform);
+							InTransform.FromFTransform(InTransform.ToFTransform().GetRelativeTransform(ParentTransform));
 						}
-						HierarchyToUpdate->SetControlOffsetTransform(Key, InTransform, true, true, bSetupUndoRedo);
+						HierarchyToUpdate->SetControlOffsetTransform(Key, InTransform.ToFTransform(), true, true, bSetupUndoRedo);
 					}
 					else if(CurrentTransformType == ERigTransformElementDetailsTransform::Minimum)
 					{
@@ -2762,7 +2800,7 @@ void FRigControlElementDetails::CustomizeShape(IDetailLayoutBuilder& DetailBuild
 	{
 		if(FRigControlElement* ControlElement = HierarchyBeingDebugged->Find<FRigControlElement>(Key))
 		{
-			return HierarchyBeingDebugged->GetControlShapeTransform(ControlElement, ERigTransformType::InitialLocal);
+			return FEulerTransform(HierarchyBeingDebugged->GetControlShapeTransform(ControlElement, ERigTransformType::InitialLocal));
 		}
 		return FEulerTransform::Identity;
 	};
@@ -2775,7 +2813,7 @@ void FRigControlElementDetails::CustomizeShape(IDetailLayoutBuilder& DetailBuild
 	{
 		if(FRigControlElement* ControlElement = HierarchyBeingCustomized->Find<FRigControlElement>(Key))
 		{
-			HierarchyBeingCustomized->SetControlShapeTransform(ControlElement, InTransform, ERigTransformType::InitialLocal, bSetupUndo, true, bSetupUndo);
+			HierarchyBeingCustomized->SetControlShapeTransform(ControlElement, InTransform.ToFTransform(), ERigTransformType::InitialLocal, bSetupUndo, true, bSetupUndo);
 		}
 	};
 
diff --git a/Engine/Plugins/Animation/ControlRig/Source/ControlRigEditor/Private/EditMode/ControlRigControlsProxy.cpp b/Engine/Plugins/Animation/ControlRig/Source/ControlRigEditor/Private/EditMode/ControlRigControlsProxy.cpp
--- a/Engine/Plugins/Animation/ControlRig/Source/ControlRigEditor/Private/EditMode/ControlRigControlsProxy.cpp
+++ b/Engine/Plugins/Animation/ControlRig/Source/ControlRigEditor/Private/EditMode/ControlRigControlsProxy.cpp
@@ -146,8 +146,9 @@ void UControlRigTransformControlProxy::PostEditChangeProperty(struct FPropertyCh
 		{
 			//MUST set through ControlRig
 			FControlRigInteractionScope InteractionScope(ControlRig.Get(), ControlElement->GetKey());
-			FTransform RealTransform = Transform; //Transform is FEulerTransform
+			const FTransform RealTransform = Transform.ToFTransform(); //Transform is FEulerTransform
 			ControlRig->SetControlValue<FRigControlValue::FTransform_Float>(ControlName, RealTransform, true, EControlRigSetKey::DoNotCare,false);
+			ControlRig->GetHierarchy()->SetControlPreferredRotator(ControlElement, Transform.Rotation);
 			ControlRig->Evaluate_AnyThread();
 
 		}
@@ -163,7 +164,8 @@ void UControlRigTransformControlProxy::ValueChanged()
 		const FName PropertyName("Transform");
 		FTrackInstancePropertyBindings Binding(PropertyName, PropertyName.ToString());
 		const FTransform NewTransform = ControlRig.Get()->GetHierarchy()->GetControlValue(ControlElement, ERigControlValueType::Current).Get<FRigControlValue::FTransform_Float>().ToTransform();
-		const FEulerTransform EulerTransform = NewTransform;
+		FEulerTransform EulerTransform(NewTransform);
+		EulerTransform.Rotation = ControlRig->GetHierarchy()->GetControlPreferredRotator(ControlElement);
 		Binding.CallFunction<FEulerTransform>(*this, EulerTransform);
 	}
 }
@@ -175,8 +177,9 @@ void UControlRigTransformControlProxy::PostEditUndo()
 	if (ControlElement && ControlRig.IsValid() && ControlRig->GetHierarchy()->Contains(FRigElementKey(ControlName, ERigElementType::Control)))
 	{
 		ControlRig->SelectControl(ControlName, bSelected);
-		FTransform RealTransform = Transform; //Transform is FEulerTransform
+		const FTransform RealTransform = Transform.ToFTransform(); //Transform is FEulerTransform
 		ControlRig->SetControlValue<FRigControlValue::FTransform_Float>(ControlName, RealTransform, true, EControlRigSetKey::Never,false);
+		ControlRig->GetHierarchy()->SetControlPreferredRotator(ControlElement, Transform.Rotation);
 	}
 }
 #endif
@@ -200,8 +203,9 @@ void UControlRigTransformControlProxy::SetKey(const IPropertyHandle& KeyedProper
 		{
 			Context.KeyMask = (uint32)EControlRigContextChannelToKey::Scale;
 		}
-		FTransform RealTransform = Transform; //Transform is FEulerTransform
+		const FTransform RealTransform = Transform.ToFTransform(); //Transform is FEulerTransform
 		ControlRig->SetControlValue<FRigControlValue::FTransform_Float>(ControlName, RealTransform, true, Context, false);
+		ControlRig->GetHierarchy()->SetControlPreferredRotator(ControlElement, Transform.Rotation);
 	}
 }
 
diff --git a/Engine/Plugins/Animation/ControlRig/Source/ControlRigEditor/Private/EditMode/ControlRigEditMode.cpp b/Engine/Plugins/Animation/ControlRig/Source/ControlRigEditor/Private/EditMode/ControlRigEditMode.cpp
--- a/Engine/Plugins/Animation/ControlRig/Source/ControlRigEditor/Private/EditMode/ControlRigEditMode.cpp
+++ b/Engine/Plugins/Animation/ControlRig/Source/ControlRigEditor/Private/EditMode/ControlRigEditMode.cpp
@@ -3618,7 +3618,7 @@ void FControlRigEditMode::MoveControlShape(AControlRigShapeActor* ShapeActor, co
 				{
 					bPrintPythonCommands = World->IsPreviewWorld();
 				}
-				ControlRig->SetControlGlobalTransform(ShapeActor->ControlName, NewTransform, true, Context, true, bPrintPythonCommands);			// assumes it's attached to actor
+				ControlRig->SetControlGlobalTransform(ShapeActor->ControlName, NewTransform, true, Context, true, bPrintPythonCommands, true);			// assumes it's attached to actor
 				ShapeActor->SetGlobalTransform(CurrentTransform);
 				if (bCalcLocal)
 				{
diff --git a/Engine/Plugins/Animation/ControlRig/Source/ControlRigEditor/Private/Sequencer/ControlRigParameterTrackEditor.cpp b/Engine/Plugins/Animation/ControlRig/Source/ControlRigEditor/Private/Sequencer/ControlRigParameterTrackEditor.cpp
--- a/Engine/Plugins/Animation/ControlRig/Source/ControlRigEditor/Private/Sequencer/ControlRigParameterTrackEditor.cpp
+++ b/Engine/Plugins/Animation/ControlRig/Source/ControlRigEditor/Private/Sequencer/ControlRigParameterTrackEditor.cpp
@@ -1633,21 +1633,23 @@ static void EvaluateThisControl(UMovieSceneControlRigParameterSection* Section,
 			case ERigControlType::TransformNoScale:
 			case ERigControlType::EulerTransform:
 			{
-				TOptional <FTransform> Value = Section->EvaluateTransformParameter(FrameTime, ControlName);
+				TOptional <FEulerTransform> Value = Section->EvaluateTransformParameter(FrameTime, ControlName);
 				if (Value.IsSet())
 				{
 					if (ControlElement->Settings.ControlType == ERigControlType::Transform)
 					{
-						ControlRig->SetControlValue<FRigControlValue::FTransform_Float>(ControlName, Value.GetValue(), true, EControlRigSetKey::Never, bSetupUndo);
+						ControlRig->SetControlValue<FRigControlValue::FTransform_Float>(ControlName, Value.GetValue().ToFTransform(), true, EControlRigSetKey::Never, bSetupUndo);
+						ControlRig->GetHierarchy()->SetControlPreferredRotator(ControlElement, Value.GetValue().Rotation);
 					}
 					else if (ControlElement->Settings.ControlType == ERigControlType::TransformNoScale)
 					{
-						FTransformNoScale NoScale = Value.GetValue();
+						FTransformNoScale NoScale = Value.GetValue().ToFTransform();
 						ControlRig->SetControlValue<FRigControlValue::FTransformNoScale_Float>(ControlName, NoScale, true, EControlRigSetKey::Never, bSetupUndo);
+						ControlRig->GetHierarchy()->SetControlPreferredRotator(ControlElement, Value.GetValue().Rotation);
 					}
 					else if (ControlElement->Settings.ControlType == ERigControlType::EulerTransform)
 					{
-						FEulerTransform Euler = Value.GetValue();
+						const FEulerTransform& Euler = Value.GetValue();
 						ControlRig->SetControlValue<FRigControlValue::FEulerTransform_Float>(ControlName, Euler, true, EControlRigSetKey::Never, bSetupUndo);
 					}
 				}
@@ -2545,26 +2547,23 @@ void FControlRigParameterTrackEditor::GetControlRigKeys(UControlRig* InControlRi
 			case ERigControlType::EulerTransform:
 			{
 				FVector Translation, Scale(1.0f, 1.0f, 1.0f);
-				FRotator Rotation;
+				FRotator Rotation = InControlRig->GetHierarchy()->GetControlPreferredRotator(ControlElement);
 
 				if (ControlElement->Settings.ControlType == ERigControlType::TransformNoScale)
 				{
 					FTransformNoScale NoScale = ControlValue.Get<FRigControlValue::FTransformNoScale_Float>().ToTransform();
 					Translation = NoScale.Location;
-					Rotation = NoScale.Rotation.Rotator();
 				}
 				else if (ControlElement->Settings.ControlType == ERigControlType::EulerTransform)
 				{
 					FEulerTransform Euler = ControlValue.Get<FRigControlValue::FEulerTransform_Float>().ToTransform();
 					Translation = Euler.Location;
-					Rotation = Euler.Rotation;
 					Scale = Euler.Scale;
 				}
 				else
 				{
 					FTransform Val = ControlValue.Get<FRigControlValue::FTransform_Float>().ToTransform();
 					Translation = Val.GetTranslation();
-					Rotation = Val.GetRotation().Rotator();
 					Scale = Val.GetScale3D();
 				}
 				FVector3f CurrentVector = (FVector3f)Translation;
@@ -2625,19 +2624,6 @@ void FControlRigParameterTrackEditor::GetControlRigKeys(UControlRig* InControlRi
 					bKeyZ = false;
 				}
 
-				/* @Mike.Zyracki this is my gut feeling - we should run SetClosestToMe on the rotator SOMEWHERE....
-				FMovieSceneInterrogationData InterrogationData;
-				GetSequencer()->GetEvaluationTemplate().CopyActuators(InterrogationData.GetAccumulator());
-				for (const FTransformInterrogationData& PreviousVal : InterrogationData.Iterate<FTransformInterrogationData>(UMovieSceneControlRigParameterSection::GetTransformInterrogationKey()))
-				{
-				if ((PreviousVal.ParameterName == RigControl.Name))
-				{
-				FRotator PreviousRot = PreviousVal.Val.GetRotation().Rotator();
-				PreviousRot.SetClosestToMe(CurrentRotator);
-				}
-				}
-				*/
-
 				OutGeneratedKeys.Add(FMovieSceneChannelValueSetter::Create<FMovieSceneFloatChannel>(ChannelIndex++, CurrentRotator.Roll, bKeyX));
 				OutGeneratedKeys.Add(FMovieSceneChannelValueSetter::Create<FMovieSceneFloatChannel>(ChannelIndex++, CurrentRotator.Pitch, bKeyY));
 				OutGeneratedKeys.Add(FMovieSceneChannelValueSetter::Create<FMovieSceneFloatChannel>(ChannelIndex++, CurrentRotator.Yaw, bKeyZ));
@@ -2930,7 +2916,7 @@ bool FControlRigParameterTrackEditor::ModifyOurGeneratedKeysByCurrentAndWeight(U
 		case ERigControlType::TransformNoScale:
 		case ERigControlType::EulerTransform:
 		{
-			for (const FTransformInterrogationData& Val : InterrogationData.Iterate<FTransformInterrogationData>(UMovieSceneControlRigParameterSection::GetTransformInterrogationKey()))
+			for (const FEulerTransformInterrogationData& Val : InterrogationData.Iterate<FEulerTransformInterrogationData>(UMovieSceneControlRigParameterSection::GetTransformInterrogationKey()))
 			{
 
 				if ((Val.ParameterName == ControlElement->GetName()))
@@ -2945,8 +2931,8 @@ bool FControlRigParameterTrackEditor::ModifyOurGeneratedKeysByCurrentAndWeight(U
 							++ChannelIndex;
 						}
 						
-						FVector3f CurrentPos = (FVector3f)Val.Val.GetTranslation();
-						FRotator3f CurrentRot = FRotator3f(Val.Val.GetRotation().Rotator());
+						FVector3f CurrentPos = (FVector3f)Val.Val.GetLocation();
+						FRotator3f CurrentRot = FRotator3f(Val.Val.Rotator());
 						GeneratedTotalKeys[ChannelIndex]->ModifyByCurrentAndWeight(Proxy, KeyTime, (void*)&CurrentPos.X, Weight);
 						GeneratedTotalKeys[ChannelIndex + 1]->ModifyByCurrentAndWeight(Proxy, KeyTime, (void*)&CurrentPos.Y, Weight);
 						GeneratedTotalKeys[ChannelIndex + 2]->ModifyByCurrentAndWeight(Proxy, KeyTime, (void*)&CurrentPos.Z, Weight);
diff --git a/Engine/Source/Runtime/AnimationCore/Private/AnimationCoreLibrary.cpp b/Engine/Source/Runtime/AnimationCore/Private/AnimationCoreLibrary.cpp
--- a/Engine/Source/Runtime/AnimationCore/Private/AnimationCoreLibrary.cpp
+++ b/Engine/Source/Runtime/AnimationCore/Private/AnimationCoreLibrary.cpp
@@ -386,4 +386,15 @@ FVector EulerFromQuat(const FQuat& Rotation, EEulerRotationOrder RotationOrder)
 	return Result * 180.f / PI;
 }
 
+FVector ChangeEulerRotationOrder(const FVector& XYZAnglesInDegrees, EEulerRotationOrder SourceRotationOrder, EEulerRotationOrder TargetRotationOrder)
+{
+	if(SourceRotationOrder == TargetRotationOrder)
+	{
+		return XYZAnglesInDegrees;
+	}
+	
+	const FQuat Quaternion = QuatFromEuler(XYZAnglesInDegrees, SourceRotationOrder);
+	return EulerFromQuat(Quaternion, TargetRotationOrder); 
+}
+
 }
\ No newline at end of file
diff --git a/Engine/Source/Runtime/AnimationCore/Public/AnimationCoreLibrary.h b/Engine/Source/Runtime/AnimationCore/Public/AnimationCoreLibrary.h
--- a/Engine/Source/Runtime/AnimationCore/Public/AnimationCoreLibrary.h
+++ b/Engine/Source/Runtime/AnimationCore/Public/AnimationCoreLibrary.h
@@ -63,4 +63,8 @@ namespace AnimationCore
 	 */
 	ANIMATIONCORE_API FVector EulerFromQuat(const FQuat& Rotation, EEulerRotationOrder RotationOrder = EEulerRotationOrder::ZYX);
 
+	/**
+	 * Converts a euler rotation from one rotation order to another
+	 */
+	ANIMATIONCORE_API FVector ChangeEulerRotationOrder(const FVector& XYZAnglesInDegrees, EEulerRotationOrder SourceRotationOrder, EEulerRotationOrder TargetRotationOrder);
 }
diff --git a/Engine/Source/Runtime/AnimationCore/Public/EulerTransform.h b/Engine/Source/Runtime/AnimationCore/Public/EulerTransform.h
--- a/Engine/Source/Runtime/AnimationCore/Public/EulerTransform.h
+++ b/Engine/Source/Runtime/AnimationCore/Public/EulerTransform.h
@@ -49,7 +49,7 @@ struct ANIMATIONCORE_API FEulerTransform
 	{
 	}
 
-	FORCEINLINE FEulerTransform(const FTransform& InTransform)
+	FORCEINLINE explicit FEulerTransform(const FTransform& InTransform)
 		: Location(InTransform.GetLocation())
 		, Rotation(InTransform.GetRotation().Rotator())
 		, Scale(InTransform.GetScale3D())
@@ -57,17 +57,6 @@ struct ANIMATIONCORE_API FEulerTransform
 
 	}
 
-	FORCEINLINE FEulerTransform& operator =(const FTransform& InTransform)
-	{
-		FromFTransform(InTransform);
-		return *this;
-	}
-
-	FORCEINLINE operator FTransform() const
-	{
-		return ToFTransform();
-	}
-
 	/** The translation of this transform */
 	UPROPERTY(EditAnywhere, BlueprintReadWrite, Category="Transform")
 	FVector Location;
