Control Rig: Remove notion of recursive propagation

* Fix missing propagation during RetainLocalTransform

diff --git a/Engine/Plugins/Animation/ControlRig/Source/ControlRig/Private/Rigs/RigHierarchy.cpp b/Engine/Plugins/Animation/ControlRig/Source/ControlRig/Private/Rigs/RigHierarchy.cpp
--- a/Engine/Plugins/Animation/ControlRig/Source/ControlRig/Private/Rigs/RigHierarchy.cpp
+++ b/Engine/Plugins/Animation/ControlRig/Source/ControlRig/Private/Rigs/RigHierarchy.cpp
@@ -249,11 +249,7 @@ void URigHierarchy::Load(FArchive& Ar)
 	{
 		if(FRigTransformElement* TransformElement = Cast<FRigTransformElement>(Elements[ElementIndex]))
 		{
-#if URIGHIERARCHY_RECURSIVE_DIRTY_PROPAGATION
 			FRigBaseElementParentArray CurrentParents = GetParents(TransformElement, false);
-#else
-			FRigBaseElementParentArray CurrentParents = GetParents(TransformElement, true);
-#endif
 			for (FRigBaseElement* CurrentParent : CurrentParents)
 			{
 				if(FRigTransformElement* TransformParent = Cast<FRigTransformElement>(CurrentParent))
@@ -639,6 +635,8 @@ void URigHierarchy::ResetPoseToInitial(ERigElementType InTypeFilter)
 					GetControlOffsetTransform(ControlElement, ERigTransformType::CurrentLocal);
 					GetControlShapeTransform(ControlElement, ERigTransformType::CurrentLocal);
 				}
+
+				PropagateDirtyFlags(TransformElement, false, true, true, false);
 			}
 		}
 		
@@ -646,6 +644,11 @@ void URigHierarchy::ResetPoseToInitial(ERigElementType InTypeFilter)
 		{
 			if(FRigTransformElement* TransformElement = Get<FRigTransformElement>(ElementIndex))
 			{
+				if(TransformElement->Pose.IsDirty(ERigTransformType::CurrentGlobal))
+				{
+					continue;
+				}
+				
 				TransformElement->Pose.MarkDirty(ERigTransformType::CurrentGlobal);
 
 				if(FRigControlElement* ControlElement = Cast<FRigControlElement>(TransformElement))
@@ -653,6 +656,8 @@ void URigHierarchy::ResetPoseToInitial(ERigElementType InTypeFilter)
 					ControlElement->Offset.MarkDirty(ERigTransformType::CurrentGlobal);
 					ControlElement->Shape.MarkDirty(ERigTransformType::CurrentGlobal);
 				}
+
+				PropagateDirtyFlags(TransformElement, false, true, false, true);
 			}
 		}
 	}
@@ -3933,11 +3938,7 @@ void URigHierarchy::DestroyElement(FRigBaseElement*& InElement)
 	InElement = nullptr;
 }
 
-#if URIGHIERARCHY_RECURSIVE_DIRTY_PROPAGATION
 void URigHierarchy::PropagateDirtyFlags(FRigTransformElement* InTransformElement, bool bInitial, bool bAffectChildren, bool bComputeOpposed, bool bMarkDirty) const
-#else
-void URigHierarchy::PropagateDirtyFlags(FRigTransformElement* InTransformElement, bool bInitial, bool bAffectChildren) const
-#endif
 {
 	if(!bEnableDirtyPropagation)
 	{
@@ -3951,13 +3952,10 @@ void URigHierarchy::PropagateDirtyFlags(FRigTransformElement* InTransformElement
 	const ERigTransformType::Type TypeToCompute = bAffectChildren ? LocalType : GlobalType;
 	const ERigTransformType::Type TypeToDirty = SwapLocalAndGlobal(TypeToCompute);
 
-#if URIGHIERARCHY_RECURSIVE_DIRTY_PROPAGATION
 	if(bComputeOpposed)
-#endif
 	{
 		for(const FRigTransformElement::FElementToDirty& ElementToDirty : InTransformElement->ElementsToDirty)
 		{
-#if URIGHIERARCHY_RECURSIVE_DIRTY_PROPAGATION
 			if(FRigControlElement* ControlElement = Cast<FRigControlElement>(ElementToDirty.Element))
 			{
 				if(ERigTransformType::IsGlobal(TypeToDirty))
@@ -3987,14 +3985,6 @@ void URigHierarchy::PropagateDirtyFlags(FRigTransformElement* InTransformElement
 					continue;
 				}
 			}
-#else
-
-			if(!bAffectChildren && ElementToDirty.HierarchyDistance > 1)
-			{
-				continue;
-			}
-
-#endif
 
 			if(FRigControlElement* ControlElement = Cast<FRigControlElement>(ElementToDirty.Element))
 			{
@@ -4002,20 +3992,14 @@ void URigHierarchy::PropagateDirtyFlags(FRigTransformElement* InTransformElement
 			}
 			GetTransform(ElementToDirty.Element, TypeToCompute); // make sure the local / global transform is up 2 date
 
-#if URIGHIERARCHY_RECURSIVE_DIRTY_PROPAGATION
 			PropagateDirtyFlags(ElementToDirty.Element, bInitial, bAffectChildren, true, false);
-#endif
 		}
 	}
-#if URIGHIERARCHY_RECURSIVE_DIRTY_PROPAGATION
 
 	if(bMarkDirty)
-#endif
 	{
 		for(const FRigTransformElement::FElementToDirty& ElementToDirty : InTransformElement->ElementsToDirty)
 		{
-#if URIGHIERARCHY_RECURSIVE_DIRTY_PROPAGATION
-
 			if(FRigControlElement* ControlElement = Cast<FRigControlElement>(ElementToDirty.Element))
 			{
 				if(ERigTransformType::IsGlobal(TypeToDirty))
@@ -4046,15 +4030,6 @@ void URigHierarchy::PropagateDirtyFlags(FRigTransformElement* InTransformElement
 				}
 			}
 						
-#else
-
-			if(!bAffectChildren && ElementToDirty.HierarchyDistance > 1)
-			{
-				continue;
-			}
-			
-#endif
-
 			ElementToDirty.Element->Pose.MarkDirty(TypeToDirty);
 		
 			if(FRigControlElement* ControlElement = Cast<FRigControlElement>(ElementToDirty.Element))
@@ -4063,16 +4038,10 @@ void URigHierarchy::PropagateDirtyFlags(FRigTransformElement* InTransformElement
 				ControlElement->Shape.MarkDirty(GlobalType);
 			}
 
-#if URIGHIERARCHY_RECURSIVE_DIRTY_PROPAGATION
-
 			if(bAffectChildren)
 			{
-#if URIGHIERARCHY_RECURSIVE_DIRTY_PROPAGATION
 				PropagateDirtyFlags(ElementToDirty.Element, bInitial, bAffectChildren, false, true);
-#endif
 			}
-			
-#endif
 		}
 	}
 }
diff --git a/Engine/Plugins/Animation/ControlRig/Source/ControlRig/Private/Rigs/RigHierarchyController.cpp b/Engine/Plugins/Animation/ControlRig/Source/ControlRig/Private/Rigs/RigHierarchyController.cpp
--- a/Engine/Plugins/Animation/ControlRig/Source/ControlRig/Private/Rigs/RigHierarchyController.cpp
+++ b/Engine/Plugins/Animation/ControlRig/Source/ControlRig/Private/Rigs/RigHierarchyController.cpp
@@ -2608,22 +2608,6 @@ void URigHierarchyController::AddElementToDirty(FRigBaseElement* InParent, FRigB
 		const FRigTransformElement::FElementToDirty ElementToDirty(ElementToAdd, InHierarchyDistance);
 		TransformParent->ElementsToDirty.AddUnique(ElementToDirty);
 	}
-
-#if URIGHIERARCHY_RECURSIVE_DIRTY_PROPAGATION
-	// nothing to do 
-#else
-	if(FRigSingleParentElement* SingleParentElement = Cast<FRigSingleParentElement>(InParent))
-	{
-		AddElementToDirty(SingleParentElement->ParentElement, InElementToAdd, InHierarchyDistance + 1);
-	}
-	else if(FRigMultiParentElement* MultiParentElement = Cast<FRigMultiParentElement>(InParent))
-	{
-		for(FRigTransformElement* ParentElement : MultiParentElement->ParentElements)
-		{
-			AddElementToDirty(ParentElement, InElementToAdd, InHierarchyDistance + 1);
-		}
-	}
-#endif
 }
 
 void URigHierarchyController::RemoveElementToDirty(FRigBaseElement* InParent, FRigBaseElement* InElementToRemove) const
@@ -2643,22 +2627,6 @@ void URigHierarchyController::RemoveElementToDirty(FRigBaseElement* InParent, FR
 	{
 		TransformParent->ElementsToDirty.Remove(ElementToRemove);
 	}
-
-#if URIGHIERARCHY_RECURSIVE_DIRTY_PROPAGATION
-	// nothing to do 
-#else
-	if(FRigSingleParentElement* SingleParentElement = Cast<FRigSingleParentElement>(InParent))
-	{
-		RemoveElementToDirty(SingleParentElement->ParentElement, InElementToRemove);
-	}
-	else if(FRigMultiParentElement* MultiParentElement = Cast<FRigMultiParentElement>(InParent))
-	{
-		for(FRigTransformElement* ParentElement : MultiParentElement->ParentElements)
-		{
-			RemoveElementToDirty(ParentElement, InElementToRemove);
-		}
-	}
-#endif
 }
 
 void URigHierarchyController::ReportWarning(const FString& InMessage) const
diff --git a/Engine/Plugins/Animation/ControlRig/Source/ControlRig/Public/Rigs/RigHierarchy.h b/Engine/Plugins/Animation/ControlRig/Source/ControlRig/Public/Rigs/RigHierarchy.h
--- a/Engine/Plugins/Animation/ControlRig/Source/ControlRig/Public/Rigs/RigHierarchy.h
+++ b/Engine/Plugins/Animation/ControlRig/Source/ControlRig/Public/Rigs/RigHierarchy.h
@@ -2822,11 +2822,7 @@ private:
 	 * @param bInitial If true the initial transform will be dirtied
 	 * @param bAffectChildren If set to false children will not move (maintain global).
 	 */
-#if URIGHIERARCHY_RECURSIVE_DIRTY_PROPAGATION
 	void PropagateDirtyFlags(FRigTransformElement* InTransformElement, bool bInitial, bool bAffectChildren, bool bComputeOpposed = true, bool bMarkDirty = true) const;
-#else
-	void PropagateDirtyFlags(FRigTransformElement* InTransformElement, bool bInitial, bool bAffectChildren) const;
-#endif
 
 public:
 
@@ -2977,7 +2973,7 @@ private:
 	int32 ResetPoseHash;
 	TArray<bool> ResetPoseIsFilteredOut;
 	TArray<int32> ElementsToRetainLocalTransform;
-
+	
 #if WITH_EDITOR
 
 	// this is mainly used for propagating changes between hierarchies in the direction of blueprint -> CDO -> other instances
diff --git a/Engine/Plugins/Animation/ControlRig/Source/ControlRig/Public/Rigs/RigHierarchyDefines.h b/Engine/Plugins/Animation/ControlRig/Source/ControlRig/Public/Rigs/RigHierarchyDefines.h
--- a/Engine/Plugins/Animation/ControlRig/Source/ControlRig/Public/Rigs/RigHierarchyDefines.h
+++ b/Engine/Plugins/Animation/ControlRig/Source/ControlRig/Public/Rigs/RigHierarchyDefines.h
@@ -9,11 +9,6 @@
 
 class URigHierarchy;
 
-// Define to switch between recursive calls for dirty propagation or flat iteration.
-// Based on this each element will either contains only the next tier of elements to dirty (recursive)
-// or a full flattened list of elements to dirty (no need to recurse)
-#define URIGHIERARCHY_RECURSIVE_DIRTY_PROPAGATION 1
-
 // Debug define which performs a full check on the cache validity for all elements of the hierarchy.
 // This can be useful for debugging cache validity bugs.
 #define URIGHIERARCHY_ENSURE_CACHE_VALIDITY 0
