Sequencer: Add commands for control rig controls and control rig selected controls filters. Add command for reset filters

diff --git a/Engine/Plugins/Animation/ControlRig/Source/ControlRigEditor/Private/Sequencer/ControlRigSequencerFilter.cpp b/Engine/Plugins/Animation/ControlRig/Source/ControlRigEditor/Private/Sequencer/ControlRigSequencerFilter.cpp
--- a/Engine/Plugins/Animation/ControlRig/Source/ControlRigEditor/Private/Sequencer/ControlRigSequencerFilter.cpp
+++ b/Engine/Plugins/Animation/ControlRig/Source/ControlRigEditor/Private/Sequencer/ControlRigSequencerFilter.cpp
@@ -5,23 +5,59 @@
 #include "ControlRigSkeletalMeshComponent.h"
 #include "Sequencer/MovieSceneControlRigParameterTrack.h"
 #include "Styling/SlateIconFinder.h"
+#include "Framework/Commands/Commands.h"
+#include "ISequencer.h"
 
 #define LOCTEXT_NAMESPACE "ControlRigSequencerTrackFilters"
 
 //////////////////////////////////////////////////////////////////////////
 //
 
-class FSequencerTrackFilter_ControlRigControls : public  FSequencerTrackFilter
+class FSequencerTrackFilter_ControlRigControlsCommands
+	: public TCommands<FSequencerTrackFilter_ControlRigControlsCommands>
 {
+public:
+
+	FSequencerTrackFilter_ControlRigControlsCommands()
+		: TCommands<FSequencerTrackFilter_ControlRigControlsCommands>
+		(
+			"FSequencerTrackFilter_ControlRigControls",
+			NSLOCTEXT("Contexts", "FSequencerTrackFilter_ControlRigControls", "FSequencerTrackFilter_ControlRigControls"),
+			NAME_None,
+			FAppStyle::GetAppStyleSetName() // Icon Style Set
+			)
+	{ }
+
+	/** Toggle the control rig controls filter */
+	TSharedPtr< FUICommandInfo > ToggleControlRigControls;
+
+	/** Initialize commands */
+	virtual void RegisterCommands() override
+	{
+		UI_COMMAND(ToggleControlRigControls, "Control Rig Controls", "Toggle the filter for Control Rig Controls.", EUserInterfaceActionType::ToggleButton, FInputChord(EKeys::F9));
+	}
+};
+
+class FSequencerTrackFilter_ControlRigControls : public FSequencerTrackFilter
+{
+public:
+
+	FSequencerTrackFilter_ControlRigControls()
+	{
+		FSequencerTrackFilter_ControlRigControlsCommands::Register();
+	}
+
+	~FSequencerTrackFilter_ControlRigControls()
+	{
+		FSequencerTrackFilter_ControlRigControlsCommands::Unregister();
+	}
+
 	virtual FString GetName() const override { return TEXT("ControlRigControlsFilter"); }
 	virtual FText GetDisplayName() const override { return LOCTEXT("SequenceTrackFilter_ControlRigControls", "Control Rig Controls"); }
-	virtual FText GetToolTipText() const override { return LOCTEXT("SequencerTrackFilter_ControlRigControlsToolTip", "Show Only Control Rig Controls."); }
 	virtual FSlateIcon GetIcon() const { return FSlateIconFinder::FindIconForClass(UControlRig::StaticClass()); }
 
-
 	virtual bool PassesFilterWithDisplayName(FTrackFilterType InItem, const FText& InText) const
 	{
-
 		const UMovieSceneControlRigParameterTrack* Track = Cast<UMovieSceneControlRigParameterTrack>(InItem);
 		return (Track != nullptr);
 	}
@@ -30,13 +66,89 @@ class FSequencerTrackFilter_ControlRigControls : public  FSequencerTrackFilter
 	{
 		return true;
 	}
+
+	virtual FText GetToolTipText() const override
+	{
+		if (!FSequencerTrackFilter_ControlRigControlsCommands::IsRegistered())
+		{
+			FSequencerTrackFilter_ControlRigControlsCommands::Register();
+		}
+
+		const FSequencerTrackFilter_ControlRigControlsCommands& Commands = FSequencerTrackFilter_ControlRigControlsCommands::Get();
+
+		const TSharedRef<const FInputChord> FirstActiveChord = Commands.ToggleControlRigControls->GetFirstValidChord();
+
+		FText Tooltip = LOCTEXT("SequencerTrackFilter_ControlRigControlsTip", "Show Only Control Rig Controls.");
+
+		if (FirstActiveChord->IsValidChord())
+		{
+			return FText::Join(FText::FromString(TEXT(" ")), Tooltip, FirstActiveChord->GetInputText());
+		}
+		return Tooltip;
+	}
+
+	virtual void BindCommands(TSharedRef<FUICommandList> CommandBindings, TWeakPtr<ISequencer> Sequencer) override
+	{
+		// See comment above
+		if (!FSequencerTrackFilter_ControlRigControlsCommands::IsRegistered())
+		{
+			FSequencerTrackFilter_ControlRigControlsCommands::Register();
+		}
+
+		const FSequencerTrackFilter_ControlRigControlsCommands& Commands = FSequencerTrackFilter_ControlRigControlsCommands::Get();
+
+		CommandBindings->MapAction(
+			Commands.ToggleControlRigControls,
+			FExecuteAction::CreateLambda([this, Sequencer] { Sequencer.Pin()->SetTrackFilterEnabled(GetDisplayName(), !Sequencer.Pin()->IsTrackFilterEnabled(GetDisplayName())); }),
+			FCanExecuteAction::CreateLambda([this, Sequencer] { return true; }),
+			FIsActionChecked::CreateLambda([this, Sequencer] { return Sequencer.Pin()->IsTrackFilterEnabled(GetDisplayName()); }));
+	}
+};
+
+//////////////////////////////////////////////////////////////////////////
+//
+
+class FSequencerTrackFilter_ControlRigSelectedControlsCommands
+	: public TCommands<FSequencerTrackFilter_ControlRigSelectedControlsCommands>
+{
+public:
+
+	FSequencerTrackFilter_ControlRigSelectedControlsCommands()
+		: TCommands<FSequencerTrackFilter_ControlRigSelectedControlsCommands>
+		(
+			"FSequencerTrackFilter_ControlRigSelectedControls",
+			NSLOCTEXT("Contexts", "FSequencerTrackFilter_ControlRigSelectedControls", "FSequencerTrackFilter_ControlRigSelectedControls"),
+			NAME_None,
+			FAppStyle::GetAppStyleSetName() // Icon Style Set
+			)
+	{ }
+
+	/** Toggle the control rig selected controls filter */
+	TSharedPtr< FUICommandInfo > ToggleControlRigSelectedControls;
+
+	/** Initialize commands */
+	virtual void RegisterCommands() override
+	{
+		UI_COMMAND(ToggleControlRigSelectedControls, "Control Rig Selected Controls", "Toggle the filter for Control Rig Selected Controls.", EUserInterfaceActionType::ToggleButton, FInputChord(EKeys::F10));
+	}
 };
 
 class FSequencerTrackFilter_ControlRigSelectedControls : public FSequencerTrackFilter
 {
+public:
+
+	FSequencerTrackFilter_ControlRigSelectedControls()
+	{
+		FSequencerTrackFilter_ControlRigSelectedControlsCommands::Register();
+	}
+
+	~FSequencerTrackFilter_ControlRigSelectedControls()
+	{
+		FSequencerTrackFilter_ControlRigSelectedControlsCommands::Unregister();
+	}
+
 	virtual FString GetName() const override { return TEXT("ControlRigControlsSelectedFilter"); }
 	virtual FText GetDisplayName() const override { return LOCTEXT("SequenceTrackFilter_ControlRigSelectedControls", "Selected Control Rig Controls"); }
-	virtual FText GetToolTipText() const override { return LOCTEXT("SequencerTrackFilter_ControlRigSelectedControls", "Show Only Selected Control Rig Controls."); }
 	virtual FSlateIcon GetIcon() const { return FSlateIconFinder::FindIconForClass(UControlRig::StaticClass()); }
 
 	virtual bool PassesFilter(FTrackFilterType InItem) const override
@@ -67,6 +179,42 @@ class FSequencerTrackFilter_ControlRigSelectedControls : public FSequencerTrackF
 		return false;
 	}
 
+	virtual FText GetToolTipText() const override
+	{
+		if (!FSequencerTrackFilter_ControlRigSelectedControlsCommands::IsRegistered())
+		{
+			FSequencerTrackFilter_ControlRigSelectedControlsCommands::Register();
+		}
+
+		const FSequencerTrackFilter_ControlRigSelectedControlsCommands& Commands = FSequencerTrackFilter_ControlRigSelectedControlsCommands::Get();
+
+		const TSharedRef<const FInputChord> FirstActiveChord = Commands.ToggleControlRigSelectedControls->GetFirstValidChord();
+
+		FText Tooltip = LOCTEXT("SequencerTrackFilter_ControlRigSelectedControlsTip", "Show Only Selected Control Rig Controls.");
+
+		if (FirstActiveChord->IsValidChord())
+		{
+			return FText::Join(FText::FromString(TEXT(" ")), Tooltip, FirstActiveChord->GetInputText());
+		}
+		return Tooltip;
+	}
+
+	virtual void BindCommands(TSharedRef<FUICommandList> CommandBindings, TWeakPtr<ISequencer> Sequencer) override
+	{
+		// See comment above
+		if (!FSequencerTrackFilter_ControlRigSelectedControlsCommands::IsRegistered())
+		{
+			FSequencerTrackFilter_ControlRigSelectedControlsCommands::Register();
+		}
+
+		const FSequencerTrackFilter_ControlRigSelectedControlsCommands& Commands = FSequencerTrackFilter_ControlRigSelectedControlsCommands::Get();
+
+		CommandBindings->MapAction(
+			Commands.ToggleControlRigSelectedControls,
+			FExecuteAction::CreateLambda([this, Sequencer] { Sequencer.Pin()->SetTrackFilterEnabled(GetDisplayName(), !Sequencer.Pin()->IsTrackFilterEnabled(GetDisplayName())); }),
+			FCanExecuteAction::CreateLambda([this, Sequencer] { return true; }),
+			FIsActionChecked::CreateLambda([this, Sequencer] { return Sequencer.Pin()->IsTrackFilterEnabled(GetDisplayName()); }));
+	}
 };
 
 /*
@@ -87,7 +235,6 @@ void UControlRigTrackFilter::AddTrackFilterExtensions(TArray< TSharedRef<class F
 {
 	InOutFilterList.Add(MakeShared<FSequencerTrackFilter_ControlRigControls>());
 	InOutFilterList.Add(MakeShared<FSequencerTrackFilter_ControlRigSelectedControls>());
-
 }
 
 #undef LOCTEXT_NAMESPACE
diff --git a/Engine/Source/Editor/Sequencer/Private/SSequencer.cpp b/Engine/Source/Editor/Sequencer/Private/SSequencer.cpp
--- a/Engine/Source/Editor/Sequencer/Private/SSequencer.cpp
+++ b/Engine/Source/Editor/Sequencer/Private/SSequencer.cpp
@@ -1796,12 +1796,7 @@ TSharedRef<SWidget> SSequencer::MakeFilterMenu()
 			FSlateIcon(),
 			FUIAction(FExecuteAction::CreateSP(this, &SSequencer::OnEnableAllFilters)));
 
-		MenuBuilder.AddMenuEntry(
-			LOCTEXT("FilterListResetFilters", "Reset Filters"),
-			LOCTEXT("FilterListResetToolTip", "Resets current filter selection"),
-			FSlateIcon(),
-			FUIAction(FExecuteAction::CreateSP(this, &SSequencer::OnResetFilters))
-		);
+		MenuBuilder.AddMenuEntry(FSequencerCommands::Get().ResetFilters);
 	}
 	MenuBuilder.EndSection();
 
@@ -2049,7 +2044,7 @@ TArray<FText> SSequencer::GetTrackFilterNames() const
 	return Filters;
 }
 
-void SSequencer::OnResetFilters()
+void SSequencer::ResetFilters()
 {
 	TSharedPtr<FSequencer> Sequencer = SequencerPtr.Pin();
 	Sequencer->GetNodeTree()->RemoveAllFilters();
diff --git a/Engine/Source/Editor/Sequencer/Private/SSequencer.h b/Engine/Source/Editor/Sequencer/Private/SSequencer.h
--- a/Engine/Source/Editor/Sequencer/Private/SSequencer.h
+++ b/Engine/Source/Editor/Sequencer/Private/SSequencer.h
@@ -358,6 +358,9 @@ public:
 	/** Gets whether the specified track filter is on/off */
 	bool IsTrackFilterEnabled(const FText& InTrackFilterName) const;
 
+	/** Reset all enabled filters */
+	void ResetFilters();
+
 	/** Gets all the available track filter names */
 	TArray<FText> GetTrackFilterNames() const;
 
@@ -474,7 +477,6 @@ public:
 
 private:
 
-	void OnResetFilters();
 	void OnEnableAllFilters();
 	void OnTrackFilterClicked(TSharedRef<FSequencerTrackFilter> TrackFilter);
 	bool IsTrackFilterActive(TSharedRef<FSequencerTrackFilter> TrackFilter) const;
diff --git a/Engine/Source/Editor/Sequencer/Private/Sequencer.cpp b/Engine/Source/Editor/Sequencer/Private/Sequencer.cpp
--- a/Engine/Source/Editor/Sequencer/Private/Sequencer.cpp
+++ b/Engine/Source/Editor/Sequencer/Private/Sequencer.cpp
@@ -11109,6 +11109,11 @@ void FSequencer::CollapseAllNodes()
 	SequencerWidget->GetTreeView()->ToggleExpandCollapseNodes(ETreeRecursion::Recursive, bExpandAll, bCollapseAll);
 }
 
+void FSequencer::ResetFilters()
+{
+	SequencerWidget->ResetFilters();
+}
+
 void FSequencer::AddSelectedActors()
 {
 	USelection* ActorSelection = GEditor->GetSelectedActors();
@@ -12404,6 +12409,10 @@ void FSequencer::BindCommands()
 		Commands.CollapseAllNodes,
 		FExecuteAction::CreateSP(this, &FSequencer::CollapseAllNodes));
 
+	SequencerCommandBindings->MapAction(
+		Commands.ResetFilters,
+		FExecuteAction::CreateSP(this, &FSequencer::ResetFilters));
+
 	SequencerCommandBindings->MapAction(
 		Commands.AddActorsToSequencer,
 		FExecuteAction::CreateSP( this, &FSequencer::AddSelectedActors));
diff --git a/Engine/Source/Editor/Sequencer/Private/Sequencer.h b/Engine/Source/Editor/Sequencer/Private/Sequencer.h
--- a/Engine/Source/Editor/Sequencer/Private/Sequencer.h
+++ b/Engine/Source/Editor/Sequencer/Private/Sequencer.h
@@ -1019,6 +1019,9 @@ protected:
 	void ExpandAllNodes();
 	void CollapseAllNodes();
 
+	/** Reset all enabled filters */
+	void ResetFilters();
+
 	/** Sort all nodes and their descendants by category then alphabetically */
 	void SortAllNodesAndDescendants();
 
diff --git a/Engine/Source/Editor/Sequencer/Private/SequencerCommands.cpp b/Engine/Source/Editor/Sequencer/Private/SequencerCommands.cpp
--- a/Engine/Source/Editor/Sequencer/Private/SequencerCommands.cpp
+++ b/Engine/Source/Editor/Sequencer/Private/SequencerCommands.cpp
@@ -54,6 +54,7 @@ void FSequencerCommands::RegisterCommands()
 	UI_COMMAND( ExpandAllNodes, "Expand All Nodes", "Expand all nodes and descendants", EUserInterfaceActionType::Button, FInputChord() );
 	UI_COMMAND( CollapseAllNodes, "Collapse All Nodes", "Collapse all nodes and descendants", EUserInterfaceActionType::Button, FInputChord() );
 	UI_COMMAND( SortAllNodesAndDescendants, "Sort All Nodes", "Sorts all nodes by type and then alphabetically.", EUserInterfaceActionType::Button, FInputChord());
+	UI_COMMAND( ResetFilters, "Reset Filters", "Reset all enabled filters.", EUserInterfaceActionType::Button, FInputChord());
 
 	UI_COMMAND( SetSelectionRangeEnd, "Set Selection End", "Sets the end of the selection range", EUserInterfaceActionType::Button, FInputChord(EKeys::O) );
 	UI_COMMAND( SetSelectionRangeStart, "Set Selection Start", "Sets the start of the selection range", EUserInterfaceActionType::Button, FInputChord(EKeys::I) );
diff --git a/Engine/Source/Editor/Sequencer/Public/SequencerCommands.h b/Engine/Source/Editor/Sequencer/Public/SequencerCommands.h
--- a/Engine/Source/Editor/Sequencer/Public/SequencerCommands.h
+++ b/Engine/Source/Editor/Sequencer/Public/SequencerCommands.h
@@ -151,6 +151,9 @@ public:
 	/** Sort all nodes and descendants */
 	TSharedPtr< FUICommandInfo > SortAllNodesAndDescendants;
 
+	/** Reset all enabled filters */
+	TSharedPtr< FUICommandInfo > ResetFilters;
+
 	/** Sets the upper bound of the selection range */
 	TSharedPtr< FUICommandInfo > SetSelectionRangeEnd;
 
