Control Rig: Sequencer: For Proxy controls make sure we don't key or evaluate them.

diff --git a/Engine/Plugins/Animation/ControlRig/Source/ControlRig/Private/Sequencer/MovieSceneControlRigParameterTemplate.cpp b/Engine/Plugins/Animation/ControlRig/Source/ControlRig/Private/Sequencer/MovieSceneControlRigParameterTemplate.cpp
--- a/Engine/Plugins/Animation/ControlRig/Source/ControlRig/Private/Sequencer/MovieSceneControlRigParameterTemplate.cpp
+++ b/Engine/Plugins/Animation/ControlRig/Source/ControlRig/Private/Sequencer/MovieSceneControlRigParameterTemplate.cpp
@@ -1215,7 +1215,9 @@
 					if (Section->ControlsToSet.Num() == 0 || Section->ControlsToSet.Contains(BoolNameAndValue.ParameterName))
 					{
 						FRigControlElement* RigControl = ControlRig->FindControl(BoolNameAndValue.ParameterName);
-						if (RigControl && RigControl->Settings.ControlType == ERigControlType::Bool)
+						if (RigControl && RigControl->Settings.AnimationType != ERigControlAnimationType::ProxyControl &&
+							RigControl->Settings.AnimationType != ERigControlAnimationType::VisualCue
+							&& RigControl->Settings.ControlType == ERigControlType::Bool)
 						{
 							ControlRig->SetControlValue<bool>(BoolNameAndValue.ParameterName, BoolNameAndValue.Value, true, EControlRigSetKey::Never,bSetupUndo);
 						}
@@ -1227,7 +1229,9 @@
 					if (Section->ControlsToSet.Num() == 0 || Section->ControlsToSet.Contains(IntegerNameAndValue.ParameterName))
 					{
 						FRigControlElement* RigControl = ControlRig->FindControl(IntegerNameAndValue.ParameterName);
-						if (RigControl && RigControl->Settings.ControlType == ERigControlType::Integer)
+						if (RigControl && RigControl->Settings.AnimationType != ERigControlAnimationType::ProxyControl &&
+							RigControl->Settings.AnimationType != ERigControlAnimationType::VisualCue
+							&& RigControl->Settings.ControlType == ERigControlType::Integer)
 						{
 							ControlRig->SetControlValue<int32>(IntegerNameAndValue.ParameterName, IntegerNameAndValue.Value, true, EControlRigSetKey::Never,bSetupUndo);
 						}
@@ -1277,7 +1281,9 @@
 		if (ControlRig)
 		{
 			FRigControlElement* ControlElement = ControlRig->FindControl(ParameterName);
-			if (ControlElement && ControlElement->Settings.ControlType == ERigControlType::Float)
+			if (ControlElement && ControlElement->Settings.AnimationType != ERigControlAnimationType::ProxyControl &&
+				ControlElement->Settings.AnimationType != ERigControlAnimationType::VisualCue
+				&& ControlElement->Settings.ControlType == ERigControlType::Float)
 			{
 				const float Val = ControlRig->GetHierarchy()->GetControlValue(ControlElement, ERigControlValueType::Current).Get<float>();
 				return FControlRigTrackTokenFloat(Val);
@@ -1303,7 +1309,8 @@
 			if (ControlRig && (Section->ControlsToSet.Num() == 0 || Section->ControlsToSet.Contains(ParameterName)))
 			{
 				FRigControlElement* ControlElement = ControlRig->FindControl(ParameterName);
-				if (ControlElement && ControlElement->Settings.ControlType == ERigControlType::Float)
+				if (ControlElement && ControlElement->Settings.AnimationType != ERigControlAnimationType::ProxyControl &&
+					ControlElement->Settings.AnimationType != ERigControlAnimationType::VisualCue && ControlElement->Settings.ControlType == ERigControlType::Float)
 				{
 					ControlRig->SetControlValue<float>(ParameterName, InFinalValue.Value, true, EControlRigSetKey::Never,bSetupUndo);
 				}
@@ -1347,7 +1354,9 @@
 		if (ControlRig)
 		{
 			FRigControlElement* ControlElement = ControlRig->FindControl(ParameterName);
-			if (ControlElement && (ControlElement->Settings.ControlType == ERigControlType::Vector2D))
+			if (ControlElement && ControlElement->Settings.AnimationType != ERigControlAnimationType::ProxyControl &&
+				ControlElement->Settings.AnimationType != ERigControlAnimationType::VisualCue
+				&& (ControlElement->Settings.ControlType == ERigControlType::Vector2D))
 			{
 				FVector3f Val = ControlRig->GetHierarchy()->GetControlValue(ControlElement, ERigControlValueType::Current).Get<FVector3f>();
 				return FControlRigTrackTokenVector2D(FVector2D(Val.X, Val.Y));
@@ -1373,7 +1382,9 @@
 			if (ControlRig && (Section->ControlsToSet.Num() == 0 || Section->ControlsToSet.Contains(ParameterName)))
 			{
 				FRigControlElement* ControlElement = ControlRig->FindControl(ParameterName);
-				if (ControlElement && (ControlElement->Settings.ControlType == ERigControlType::Vector2D))
+				if (ControlElement && ControlElement->Settings.AnimationType != ERigControlAnimationType::ProxyControl &&
+					ControlElement->Settings.AnimationType != ERigControlAnimationType::VisualCue
+					&& (ControlElement->Settings.ControlType == ERigControlType::Vector2D))
 				{
 					const FVector3f Value(InFinalValue.Value.X, InFinalValue.Value.Y, 0.f);
 					ControlRig->SetControlValue<FVector3f>(ParameterName, Value, true, EControlRigSetKey::Never,bSetupUndo);
@@ -1416,8 +1427,11 @@
 		if (ControlRig)
 		{
 			FRigControlElement* ControlElement = ControlRig->FindControl(ParameterName);
-			if (ControlElement && (ControlElement->Settings.ControlType == ERigControlType::Position || ControlElement->Settings.ControlType == ERigControlType::Scale || ControlElement->Settings.ControlType == ERigControlType::Rotator))
-			{;
+			
+			if (ControlElement && ControlElement->Settings.AnimationType != ERigControlAnimationType::ProxyControl &&
+				ControlElement->Settings.AnimationType != ERigControlAnimationType::VisualCue
+				&& (ControlElement->Settings.ControlType == ERigControlType::Position || ControlElement->Settings.ControlType == ERigControlType::Scale || ControlElement->Settings.ControlType == ERigControlType::Rotator))
+			{
 				FVector3f Val = ControlRig->GetHierarchy()->GetControlValue(ControlElement, ERigControlValueType::Current).Get<FVector3f>();
 				return FControlRigTrackTokenVector((FVector)Val);
 			}
@@ -1440,7 +1454,9 @@
 			if (ControlRig && (Section->ControlsToSet.Num() == 0 || Section->ControlsToSet.Contains(ParameterName)))
 			{
 				FRigControlElement* ControlElement = ControlRig->FindControl(ParameterName);
-				if (ControlElement && (ControlElement->Settings.ControlType == ERigControlType::Position || ControlElement->Settings.ControlType == ERigControlType::Scale || ControlElement->Settings.ControlType == ERigControlType::Rotator))
+				if (ControlElement && ControlElement->Settings.AnimationType != ERigControlAnimationType::ProxyControl &&
+					ControlElement->Settings.AnimationType != ERigControlAnimationType::VisualCue
+					&& (ControlElement->Settings.ControlType == ERigControlType::Position || ControlElement->Settings.ControlType == ERigControlType::Scale || ControlElement->Settings.ControlType == ERigControlType::Rotator))
 				{
 						ControlRig->SetControlValue<FVector3f>(ParameterName, (FVector3f)InFinalValue.Value, true, EControlRigSetKey::Never,bSetupUndo);
 				}
@@ -1482,27 +1498,31 @@
 		if (ControlRig && Section && (Section->ControlsToSet.Num() == 0 || Section->ControlsToSet.Contains(ParameterName)))
 		{
 			FRigControlElement* ControlElement = ControlRig->FindControl(ParameterName);
-			if (ControlElement && ControlElement->Settings.ControlType == ERigControlType::Transform)
+			if (ControlElement && ControlElement->Settings.AnimationType != ERigControlAnimationType::ProxyControl &&
+				ControlElement->Settings.AnimationType != ERigControlAnimationType::VisualCue)
 			{
-				const FTransform Val = ControlRig->GetHierarchy()->GetControlValue(ControlElement, ERigControlValueType::Current).Get<FRigControlValue::FTransform_Float>().ToTransform();
-				FEulerTransform EulerTransform(Val);
-				EulerTransform.Rotation = ControlRig->GetHierarchy()->GetControlPreferredRotator(ControlElement);
-				return FControlRigTrackTokenTransform(EulerTransform);
+				if (ControlElement && ControlElement->Settings.ControlType == ERigControlType::Transform)
+				{
+					const FTransform Val = ControlRig->GetHierarchy()->GetControlValue(ControlElement, ERigControlValueType::Current).Get<FRigControlValue::FTransform_Float>().ToTransform();
+					FEulerTransform EulerTransform(Val);
+					EulerTransform.Rotation = ControlRig->GetHierarchy()->GetControlPreferredRotator(ControlElement);
+					return FControlRigTrackTokenTransform(EulerTransform);
+				}
+				else if (ControlElement && ControlElement->Settings.ControlType == ERigControlType::TransformNoScale)
+				{
+					FTransformNoScale ValNoScale = ControlRig->GetHierarchy()->GetControlValue(ControlElement, ERigControlValueType::Current).Get<FRigControlValue::FTransformNoScale_Float>().ToTransform();
+					FTransform Val = ValNoScale;
+					FEulerTransform EulerTransform(Val);
+					EulerTransform.Rotation = ControlRig->GetHierarchy()->GetControlPreferredRotator(ControlElement);
+					return FControlRigTrackTokenTransform(EulerTransform);
+				}
+				else if (ControlElement && ControlElement->Settings.ControlType == ERigControlType::EulerTransform)
+				{
+					FEulerTransform EulerTransform = ControlRig->GetHierarchy()->GetControlValue(ControlElement, ERigControlValueType::Current).Get<FRigControlValue::FEulerTransform_Float>().ToTransform();
+					return FControlRigTrackTokenTransform(EulerTransform);
+				}
 			}
-			else if(ControlElement && ControlElement->Settings.ControlType == ERigControlType::TransformNoScale)
-			{
-				FTransformNoScale ValNoScale = ControlRig->GetHierarchy()->GetControlValue(ControlElement, ERigControlValueType::Current).Get<FRigControlValue::FTransformNoScale_Float>().ToTransform();
-				FTransform Val = ValNoScale;
-				FEulerTransform EulerTransform(Val);
-				EulerTransform.Rotation = ControlRig->GetHierarchy()->GetControlPreferredRotator(ControlElement);
-				return FControlRigTrackTokenTransform(EulerTransform);
-			}
-			else if (ControlElement && ControlElement->Settings.ControlType == ERigControlType::EulerTransform)
-			{
-				FEulerTransform EulerTransform = ControlRig->GetHierarchy()->GetControlValue(ControlElement, ERigControlValueType::Current).Get<FRigControlValue::FEulerTransform_Float>().ToTransform();
-				return FControlRigTrackTokenTransform(EulerTransform);
 		}
-		}
 		return FControlRigTrackTokenTransform();
 	}
 
@@ -1529,20 +1549,24 @@
 			if (ControlRig && (Section->ControlsToSet.Num() == 0 || Section->ControlsToSet.Contains(ParameterName)))
 			{
 				FRigControlElement* ControlElement = ControlRig->FindControl(ParameterName);
-				if (ControlElement && ControlElement->Settings.ControlType == ERigControlType::Transform)
+				if (ControlElement && ControlElement->Settings.AnimationType != ERigControlAnimationType::ProxyControl &&
+					ControlElement->Settings.AnimationType != ERigControlAnimationType::VisualCue)
 				{
-					ControlRig->SetControlValue<FRigControlValue::FTransform_Float>(ParameterName, InFinalValue.Value.ToFTransform(),true, EControlRigSetKey::Never,bSetupUndo);
-					ControlRig->GetHierarchy()->SetControlPreferredRotator(ControlElement, InFinalValue.Value.Rotator());
-				}
-				else if (ControlElement && ControlElement->Settings.ControlType == ERigControlType::TransformNoScale)
-				{
-					const FTransformNoScale NoScale = InFinalValue.Value.ToFTransform();
-					ControlRig->SetControlValue<FRigControlValue::FTransformNoScale_Float>(ParameterName, NoScale, true, EControlRigSetKey::Never,bSetupUndo);
-					ControlRig->GetHierarchy()->SetControlPreferredRotator(ControlElement, InFinalValue.Value.Rotator());
-				}
-				else if (ControlElement && ControlElement->Settings.ControlType == ERigControlType::EulerTransform)
-				{
-					ControlRig->SetControlValue<FRigControlValue::FEulerTransform_Float>(ParameterName, InFinalValue.Value, true, EControlRigSetKey::Never,bSetupUndo);
+					if (ControlElement && ControlElement->Settings.ControlType == ERigControlType::Transform)
+					{
+						ControlRig->SetControlValue<FRigControlValue::FTransform_Float>(ParameterName, InFinalValue.Value.ToFTransform(), true, EControlRigSetKey::Never, bSetupUndo);
+						ControlRig->GetHierarchy()->SetControlPreferredRotator(ControlElement, InFinalValue.Value.Rotator());
+					}
+					else if (ControlElement && ControlElement->Settings.ControlType == ERigControlType::TransformNoScale)
+					{
+						const FTransformNoScale NoScale = InFinalValue.Value.ToFTransform();
+						ControlRig->SetControlValue<FRigControlValue::FTransformNoScale_Float>(ParameterName, NoScale, true, EControlRigSetKey::Never, bSetupUndo);
+						ControlRig->GetHierarchy()->SetControlPreferredRotator(ControlElement, InFinalValue.Value.Rotator());
+					}
+					else if (ControlElement && ControlElement->Settings.ControlType == ERigControlType::EulerTransform)
+					{
+						ControlRig->SetControlValue<FRigControlValue::FEulerTransform_Float>(ParameterName, InFinalValue.Value, true, EControlRigSetKey::Never, bSetupUndo);
+					}
 				}
 			}
 
diff --git a/Engine/Plugins/Animation/ControlRig/Source/ControlRigEditor/Private/Sequencer/ControlRigParameterTrackEditor.cpp b/Engine/Plugins/Animation/ControlRig/Source/ControlRigEditor/Private/Sequencer/ControlRigParameterTrackEditor.cpp
--- a/Engine/Plugins/Animation/ControlRig/Source/ControlRigEditor/Private/Sequencer/ControlRigParameterTrackEditor.cpp
+++ b/Engine/Plugins/Animation/ControlRig/Source/ControlRigEditor/Private/Sequencer/ControlRigParameterTrackEditor.cpp
@@ -2399,7 +2399,9 @@
 	{
 		return;
 	}
-	if (!GetSequencer().IsValid() || !GetSequencer()->IsAllowedToChange() || Context.SetKey == EControlRigSetKey::Never)
+	if (!GetSequencer().IsValid() || !GetSequencer()->IsAllowedToChange() || Context.SetKey == EControlRigSetKey::Never 
+		|| ControlElement->Settings.AnimationType == ERigControlAnimationType::ProxyControl
+		|| ControlElement->Settings.AnimationType == ERigControlAnimationType::VisualCue)
 	{
 		return;
 	}
