Sequencer/Valykrie: More asset filter permissions

diff --git a/Engine/Plugins/Animation/ControlRig/Source/ControlRigEditor/Private/Sequencer/ControlRigParameterTrackEditor.cpp b/Engine/Plugins/Animation/ControlRig/Source/ControlRigEditor/Private/Sequencer/ControlRigParameterTrackEditor.cpp
--- a/Engine/Plugins/Animation/ControlRig/Source/ControlRigEditor/Private/Sequencer/ControlRigParameterTrackEditor.cpp
+++ b/Engine/Plugins/Animation/ControlRig/Source/ControlRigEditor/Private/Sequencer/ControlRigParameterTrackEditor.cpp
@@ -1275,11 +1275,14 @@ void FControlRigParameterTrackEditor::HandleAddControlRigSubMenu(FMenuBuilder& M
 		FClassViewerInitializationOptions Options;
 		Options.bShowUnloadedBlueprints = true;
 		Options.NameTypeToDisplay = EClassViewerNameTypeToDisplay::DisplayName;
-
+		
 		TSharedPtr<FControlRigClassFilter> ClassFilter = MakeShareable(new FControlRigClassFilter(bFilterAssetBySkeleton, bFilterAssetByAnimatableControls, false, Skeleton));
 		Options.ClassFilters.Add(ClassFilter.ToSharedRef());
 		Options.bShowNoneOption = false;
 
+		UMovieSceneSequence* Sequence = GetSequencer() ? GetSequencer()->GetFocusedMovieSceneSequence() : nullptr;
+		Options.AdditionalReferencingAssets.Add(FAssetData(Sequence));
+
 		FClassViewerModule& ClassViewerModule = FModuleManager::LoadModuleChecked<FClassViewerModule>("ClassViewer");
 		//	FNewMenuDelegate::CreateRaw(this, &FControlRigParameterTrackEditor::AddAnimationSubMenu, BoundObject, ObjectBindings[0], Skeleton)
 
@@ -3932,6 +3935,9 @@ bool FControlRigParameterSection::RequestDeleteKeyArea(const TArray<FName>& KeyA
 
 void FControlRigParameterSection::AddAnimationSubMenuForFK(FMenuBuilder& MenuBuilder, FGuid ObjectBinding, USkeleton* Skeleton, UMovieSceneControlRigParameterSection* Section)
 {
+	TSharedPtr<ISequencer> SequencerPtr = WeakSequencer.Pin();
+	UMovieSceneSequence* Sequence = SequencerPtr.IsValid() ? SequencerPtr->GetFocusedMovieSceneSequence() : nullptr;
+
 	FAssetPickerConfig AssetPickerConfig;
 	{
 		AssetPickerConfig.OnAssetSelected = FOnAssetSelected::CreateRaw(this, &FControlRigParameterSection::OnAnimationAssetSelectedForFK, ObjectBinding, Section);
@@ -3943,6 +3949,7 @@ void FControlRigParameterSection::AddAnimationSubMenuForFK(FMenuBuilder& MenuBui
 		AssetPickerConfig.Filter.ClassNames.Add(UAnimSequenceBase::StaticClass()->GetFName());
 		AssetPickerConfig.Filter.TagsAndValues.Add(TEXT("Skeleton"), FAssetData(Skeleton).GetExportTextName());
 		AssetPickerConfig.SaveSettingsName = TEXT("SequencerAssetPicker");
+		AssetPickerConfig.AdditionalReferencingAssets.Add(FAssetData(Sequence));
 	}
 
 	FContentBrowserModule& ContentBrowserModule = FModuleManager::Get().LoadModuleChecked<FContentBrowserModule>(TEXT("ContentBrowser"));
diff --git a/Engine/Source/Editor/MovieSceneTools/Private/Channels/BuiltInChannelEditors.cpp b/Engine/Source/Editor/MovieSceneTools/Private/Channels/BuiltInChannelEditors.cpp
--- a/Engine/Source/Editor/MovieSceneTools/Private/Channels/BuiltInChannelEditors.cpp
+++ b/Engine/Source/Editor/MovieSceneTools/Private/Channels/BuiltInChannelEditors.cpp
@@ -414,12 +414,20 @@
 			return Obj ? Obj->GetPathName() : FString();
 		};
 
+		TArray<FAssetData> AssetDataArray;
+		if (InSequencer.IsValid())
+		{
+			UMovieSceneSequence* Sequence = InSequencer.Pin()->GetFocusedMovieSceneSequence();
+			AssetDataArray.Add((FAssetData)Sequence);
+		}
+
 		return SNew(SObjectPropertyEntryBox)
 		.DisplayBrowse(true)
 		.DisplayUseSelected(false)
 		.ObjectPath_Lambda(GetObjectPathLambda)
 		.AllowedClass(RawChannel->GetPropertyClass())
-		.OnObjectChanged_Lambda(OnSetObjectLambda);
+		.OnObjectChanged_Lambda(OnSetObjectLambda)
+		.OwnerAssetDataArray(AssetDataArray);
 	}
 
 	return SNullWidget::NullWidget;
