Constraints in Sequencer:

Exposed basic constraints properties (via the contextual menu)
Added a basic baking process (via the contextual menu) factorizing code with the snapper when possible.

todo:
    - factorize the controls' snapping process to use the BakingHelper interface
    - create a better UI for the baker to stick to the design (currently via a contextual menu) and handle both frame range and constraint's active state animation

diff --git a/Engine/Plugins/Animation/ControlRig/Source/ControlRig/Private/Constraints/ControlRigTransformableHandle.cpp b/Engine/Plugins/Animation/ControlRig/Source/ControlRig/Private/Constraints/ControlRigTransformableHandle.cpp
--- a/Engine/Plugins/Animation/ControlRig/Source/ControlRig/Private/Constraints/ControlRigTransformableHandle.cpp
+++ b/Engine/Plugins/Animation/ControlRig/Source/ControlRig/Private/Constraints/ControlRigTransformableHandle.cpp
@@ -38,9 +38,10 @@
 
 // NOTE should we cache the skeletal mesh and the CtrlIndex to avoid looking for if every time
 // probably not for handling runtime changes
-void UTransformableControlHandle::SetTransform(const FTransform& InGlobal) const
+void UTransformableControlHandle::SetGlobalTransform(const FTransform& InGlobal) const
 {
-	if (!ControlRig.IsValid() || ControlName == NAME_None)
+	const FRigControlElement* ControlElement = GetControlElement();
+	if (!ControlElement)
 	{
 		return;
 	}
@@ -50,12 +51,6 @@
 	{
 		return;
 	}
-
-	const FRigControlElement* ControlElement = ControlRig->FindControl(ControlName);
-	if (!ControlElement)
-	{
-		return;
-	}
 	
 	URigHierarchy* Hierarchy = ControlRig->GetHierarchy();
 	const FRigElementKey& ControlKey = ControlElement->GetKey();
@@ -65,27 +60,37 @@
 	Hierarchy->SetGlobalTransform(CtrlIndex, InGlobal.GetRelativeTransform(ComponentTransform));
 }
 
+void UTransformableControlHandle::SetLocalTransform(const FTransform& InLocal) const
+{
+	const FRigControlElement* ControlElement = GetControlElement();
+	if (!ControlElement)
+	{
+		return;
+	}
+	
+	URigHierarchy* Hierarchy = ControlRig->GetHierarchy();
+	const FRigElementKey& ControlKey = ControlElement->GetKey();
+	const int32 CtrlIndex = Hierarchy->GetIndex(ControlKey);
+	
+	Hierarchy->SetLocalTransform(CtrlIndex, InLocal);
+}
+
 // NOTE should we cache the skeletal mesh and the CtrlIndex to avoid looking for if every time
 // probably not for handling runtime changes
-FTransform UTransformableControlHandle::GetTransform() const
+FTransform UTransformableControlHandle::GetGlobalTransform() const
 {
-	if (!ControlRig.IsValid() || ControlName == NAME_None)
+	const FRigControlElement* ControlElement = GetControlElement();
+	if (!ControlElement)
 	{
 		return FTransform::Identity;
 	}
-
+	
 	const USkeletalMeshComponent* SkeletalMeshComponent = GetSkeletalMesh();
 	if (!SkeletalMeshComponent)
 	{
 		return FTransform::Identity;
 	}
 
-	const FRigControlElement* ControlElement = ControlRig->FindControl(ControlName);
-	if (!ControlElement)
-	{
-		return FTransform::Identity;
-	}
-	
 	const FRigElementKey& ControlKey = ControlElement->GetKey();
 	const URigHierarchy* Hierarchy = ControlRig->GetHierarchy();
 	const int32 CtrlIndex = Hierarchy->GetIndex(ControlKey);
@@ -94,6 +99,21 @@
 	return Hierarchy->GetGlobalTransform(CtrlIndex) * ComponentTransform;
 }
 
+FTransform UTransformableControlHandle::GetLocalTransform() const
+{
+	const FRigControlElement* ControlElement = GetControlElement();
+	if (!ControlElement)
+	{
+		return FTransform::Identity;
+	}
+	
+	const FRigElementKey& ControlKey = ControlElement->GetKey();
+	const URigHierarchy* Hierarchy = ControlRig->GetHierarchy();
+	const int32 CtrlIndex = Hierarchy->GetIndex(ControlKey);
+
+	return Hierarchy->GetLocalTransform(CtrlIndex);
+}
+
 UObject* UTransformableControlHandle::GetPrerequisiteObject() const
 {
 	return GetSkeletalMesh(); 
@@ -120,6 +140,16 @@
 	return ObjectBinding ? Cast<USkeletalMeshComponent>(ObjectBinding->GetBoundObject()) : nullptr;
 }
 
+FRigControlElement* UTransformableControlHandle::GetControlElement() const
+{
+	if (!ControlRig.IsValid() || ControlName == NAME_None)
+	{
+		return nullptr;
+	}
+
+	return ControlRig->FindControl(ControlName);
+}
+
 #if WITH_EDITOR
 FName UTransformableControlHandle::GetName() const
 {
diff --git a/Engine/Plugins/Animation/ControlRig/Source/ControlRig/Public/Constraints/ControlRigTransformableHandle.h b/Engine/Plugins/Animation/ControlRig/Source/ControlRig/Public/Constraints/ControlRigTransformableHandle.h
--- a/Engine/Plugins/Animation/ControlRig/Source/ControlRig/Public/Constraints/ControlRigTransformableHandle.h
+++ b/Engine/Plugins/Animation/ControlRig/Source/ControlRig/Public/Constraints/ControlRigTransformableHandle.h
@@ -8,13 +8,14 @@
 
 class UControlRig;
 class USkeletalMeshComponent;
+struct FRigControlElement;
 
 /**
  * UTransformableControlHandle
  */
 
 UCLASS()
-class UTransformableControlHandle : public UTransformableHandle 
+class CONTROLRIG_API UTransformableControlHandle : public UTransformableHandle 
 {
 	GENERATED_BODY()
 	
@@ -26,9 +27,13 @@
 	virtual bool IsValid() const override;
 
 	/** @todo document */
-	virtual void SetTransform(const FTransform& InGlobal) const override;
+	virtual void SetGlobalTransform(const FTransform& InGlobal) const override;
 	/** @todo document */
-	virtual FTransform GetTransform() const  override;
+	virtual void SetLocalTransform(const FTransform& InLocal) const override;
+	/** @todo document */
+	virtual FTransform GetGlobalTransform() const  override;
+	/** @todo document */
+	virtual FTransform GetLocalTransform() const  override;
 
 	/** @todo document */
 	virtual UObject* GetPrerequisiteObject() const override;
@@ -55,4 +60,7 @@
 
 	/** @todo document */
 	USkeletalMeshComponent* GetSkeletalMesh() const;
+
+	/** @todo document */
+	FRigControlElement* GetControlElement() const;
 };
diff --git a/Engine/Plugins/Animation/ControlRig/Source/ControlRigEditor/Private/EditMode/ControlRigEditMode.cpp b/Engine/Plugins/Animation/ControlRig/Source/ControlRigEditor/Private/EditMode/ControlRigEditMode.cpp
--- a/Engine/Plugins/Animation/ControlRig/Source/ControlRigEditor/Private/EditMode/ControlRigEditMode.cpp
+++ b/Engine/Plugins/Animation/ControlRig/Source/ControlRigEditor/Private/EditMode/ControlRigEditMode.cpp
@@ -2411,13 +2411,13 @@
 	}
 }
 
-void FControlRigEditMode::GetAllSelectedControls(TMap<UControlRig*, TArray<FRigElementKey>>& OutSelectedControls)
+void FControlRigEditMode::GetAllSelectedControls(TMap<UControlRig*, TArray<FRigElementKey>>& OutSelectedControls) const
 {
-	for (TWeakObjectPtr<UControlRig>& RuntimeRigPtr : RuntimeControlRigs)
+	for (const TWeakObjectPtr<UControlRig>& RuntimeRigPtr : RuntimeControlRigs)
 	{
 		if (UControlRig* ControlRig = RuntimeRigPtr.Get())
 		{
-			if (URigHierarchy* Hierarchy = ControlRig->GetHierarchy())
+			if (const URigHierarchy* Hierarchy = ControlRig->GetHierarchy())
 			{
 				TArray<FRigElementKey> SelectedControls = Hierarchy->GetSelectedKeys(ERigElementType::Control);
 				if (SelectedControls.Num() > 0)
diff --git a/Engine/Plugins/Animation/ControlRig/Source/ControlRigEditor/Private/EditMode/ControlRigEditMode.h b/Engine/Plugins/Animation/ControlRig/Source/ControlRigEditor/Private/EditMode/ControlRigEditMode.h
--- a/Engine/Plugins/Animation/ControlRig/Source/ControlRigEditor/Private/EditMode/ControlRigEditMode.h
+++ b/Engine/Plugins/Animation/ControlRig/Source/ControlRigEditor/Private/EditMode/ControlRigEditMode.h
@@ -333,7 +333,7 @@
 	int32 GetNumSelectedRigElements(uint32 InTypes, UControlRig* InControlRig) const;
 
 	/** Get all of the selected Controls*/
-	void GetAllSelectedControls(TMap<UControlRig*, TArray<FRigElementKey>>& OutSelectedControls);
+	void GetAllSelectedControls(TMap<UControlRig*, TArray<FRigElementKey>>& OutSelectedControls) const;
 
 	/** Get all of the ControlRigs, maybe not valid anymore */
 	TArrayView<const TWeakObjectPtr<UControlRig>> GetControlRigs() const;
diff --git a/Engine/Plugins/Animation/ControlRig/Source/ControlRigEditor/Private/Editor/Constraints/SConstraintsWidget.cpp b/Engine/Plugins/Animation/ControlRig/Source/ControlRigEditor/Private/Editor/Constraints/SConstraintsWidget.cpp
--- a/Engine/Plugins/Animation/ControlRig/Source/ControlRigEditor/Private/Editor/Constraints/SConstraintsWidget.cpp
+++ b/Engine/Plugins/Animation/ControlRig/Source/ControlRigEditor/Private/Editor/Constraints/SConstraintsWidget.cpp
@@ -15,8 +15,11 @@
 #include "TransformConstraint.h"
 #include "Modules/ModuleManager.h"
 #include "Widgets/Input/SButton.h"
+#include "Framework/MultiBox/MultiBoxBuilder.h"
+#include "PropertyEditorModule.h"
+#include "DetailsViewArgs.h"
+#include "Tools/ConstraintBaker.h"
 
-
 #define LOCTEXT_NAMESPACE "SConstraintsWidget"
 
 BEGIN_SLATE_FUNCTION_BUILD_OPTIMIZATION
@@ -351,6 +354,19 @@
 			ConstraintsWidget.Pin()->RemoveItem(ConstraintItem);
 		}
 	});
+
+	auto GetFont = [this]()
+	{
+		UWorld* World = GCurrentLevelEditingViewportClient->GetWorld();
+		const FConstraintsManagerController& Controller = FConstraintsManagerController::Get(World);
+		UTickableConstraint* Constraint = Controller.GetConstraint(ConstraintItem->Name);
+		if (!Constraint)
+		{
+			return IDetailLayoutBuilder::GetDetailFont();
+		}
+
+		return Constraint->Active ? IDetailLayoutBuilder::GetDetailFont() : IDetailLayoutBuilder::GetDetailFontItalic();
+	};
 	
 	ChildSlot
 	.Padding(FMargin(8.f, 2.f, 12.f, 2.f))
@@ -401,7 +417,7 @@
 						const FString TextStr = FString::Printf(TEXT("%s (%s)"), *InItem->Name.ToString(), *InItem->ParentName.ToString());
 						return FText::FromString(TextStr);
 					})
-					.Font( IDetailLayoutBuilder::GetDetailFont() )
+					.Font_Lambda(GetFont)
 				]
 			]
 		]
@@ -432,7 +448,7 @@
 			})
 			.IsEnabled_Lambda([this]()
 			{
-				return ConstraintsWidget.IsValid() ? !ConstraintsWidget.Pin()->IsFirstItem(ConstraintItem) : false;
+				return ConstraintsWidget.IsValid() ? ConstraintsWidget.Pin()->CanMoveUp(ConstraintItem) : false;
 			})
 			.ToolTipText(LOCTEXT("MoveConstraintUp", "Move this constraint up in the list."))
 			[
@@ -462,7 +478,7 @@
 			})
 			.IsEnabled_Lambda([this]()
 			{
-				return ConstraintsWidget.IsValid() ? !ConstraintsWidget.Pin()->IsLastItem(ConstraintItem) : false;
+				return ConstraintsWidget.IsValid() ? ConstraintsWidget.Pin()->CanMoveDown(ConstraintItem) : false;
 			})
 			.ToolTipText(LOCTEXT("MoveConstraintDown", "Move this constraint down in the list."))
 			[
@@ -505,6 +521,7 @@
 				.SelectionMode(ESelectionMode::Single)
 				.ListItemsSource( &ListItems )
 				.OnGenerateRow(this, &SConstraintsEditionWidget::OnGenerateWidgetForItem)
+				.OnContextMenuOpening(this, &SConstraintsEditionWidget::CreateContextMenu)
 			]
 		]
 	];
@@ -566,18 +583,35 @@
 	];
 }
 
-bool SConstraintsEditionWidget::IsFirstItem(const TSharedPtr<FEditableConstraintItem>& Item) const
+bool SConstraintsEditionWidget::CanMoveUp(const TSharedPtr<FEditableConstraintItem>& Item) const
 {
-	return ListItems.Num() > 1 && ListItems[0] == Item; 
+	const int32 NumItems = ListItems.Num();
+	if (NumItems < 2)
+	{
+		return false;
+	}
+	
+	return ListItems[0] != Item; 
 }
 
-bool SConstraintsEditionWidget::IsLastItem(const TSharedPtr<FEditableConstraintItem>& Item) const
+bool SConstraintsEditionWidget::CanMoveDown(const TSharedPtr<FEditableConstraintItem>& Item) const
 {
-	return ListItems.Num() > 1 && ListItems.Last() == Item;
+	const int32 NumItems = ListItems.Num();
+	if (NumItems < 2)
+	{
+		return false;
+	}
+	
+	return ListItems.Last() != Item;
 }
 
 void SConstraintsEditionWidget::MoveItemUp(const TSharedPtr<FEditableConstraintItem>& Item)
 {
+	if (!CanMoveUp(Item))
+	{
+		return;
+	}
+	
 	const int32 Index = ListItems.IndexOfByKey(Item);
 	if (Index > 0 && ListItems.IsValidIndex(Index))
 	{
@@ -595,6 +629,11 @@
 
 void SConstraintsEditionWidget::MoveItemDown(const TSharedPtr<FEditableConstraintItem>& Item)
 {
+	if (!CanMoveDown(Item))
+	{
+		return;
+	}
+	
 	const int32 Index = ListItems.IndexOfByKey(Item);
 	if (ListItems.IsValidIndex(Index))
 	{
@@ -673,6 +712,81 @@
 	InvalidateConstraintList();
 }
 
+TSharedPtr<SWidget> SConstraintsEditionWidget::CreateContextMenu()
+{
+	const TArray<TSharedPtr<FEditableConstraintItem>> Selection = ListView->GetSelectedItems();
+	if (Selection.IsEmpty())
+	{
+		return SNullWidget::NullWidget;
+	}
+
+	const int32 Index = ListItems.IndexOfByKey(Selection[0]);
+	if (!ListItems.IsValidIndex(Index))
+	{
+		return SNullWidget::NullWidget;
+	}
+
+	UWorld* World = GCurrentLevelEditingViewportClient->GetWorld();
+	const FConstraintsManagerController& Controller = FConstraintsManagerController::Get(World);
+	UTickableConstraint* Constraint = Controller.GetConstraint(ListItems[Index]->Name);
+	if (!Constraint)
+	{
+		return SNullWidget::NullWidget;
+	}
+	
+	static constexpr bool CloseAfterSelection = true;
+	FMenuBuilder MenuBuilder(CloseAfterSelection, nullptr);
+
+	FDetailsViewArgs DetailsViewArgs;
+	{
+		DetailsViewArgs.bAllowSearch = false;
+		DetailsViewArgs.bCustomFilterAreaLocation = true;
+		DetailsViewArgs.bCustomNameAreaLocation = true;
+		DetailsViewArgs.bHideSelectionTip = true;
+		DetailsViewArgs.bLockable = false;
+		DetailsViewArgs.bSearchInitialKeyFocus = true;
+		DetailsViewArgs.bUpdatesFromSelection = false;
+		DetailsViewArgs.bShowOptions = false;
+		DetailsViewArgs.bShowModifiedPropertiesOption = false;
+		DetailsViewArgs.ColumnWidth = 0.45f;
+	}
+	
+	TSharedRef<IDetailsView> DetailsView = FModuleManager::GetModuleChecked<FPropertyEditorModule>("PropertyEditor").CreateDetailView(DetailsViewArgs);
+
+	TArray<TWeakObjectPtr<UObject>> ConstrainsToEdit;
+	ConstrainsToEdit.Add(Constraint);
+	
+	DetailsView->SetObjects(ConstrainsToEdit);
+
+	// constraint details
+	MenuBuilder.BeginSection("EditConstraint", LOCTEXT("EditConstraintHeader", "Edit Constraint"));
+	{
+		MenuBuilder.AddWidget(DetailsView, FText::GetEmpty(), true);
+	}
+	MenuBuilder.EndSection();
+
+	// baking (note that this will probably be moved)
+	MenuBuilder.BeginSection("BakeConstraint", LOCTEXT("BakeConstraintHeader", "Bake Constraint"));
+	{
+		MenuBuilder.AddMenuEntry(
+		LOCTEXT("BakeConstraintDoIt", "Do It"),
+		FText::Format(LOCTEXT("BakeConstraintDoItTooltip", "Bake {0} transforms."), FText::FromName(ListItems[Index]->Name)),
+		FSlateIcon(),
+		FUIAction(FExecuteAction::CreateLambda([Constraint]()
+		{
+			if (UTickableTransformConstraint* TransformConstraint = Cast<UTickableTransformConstraint>(Constraint))
+			{
+				FConstraintBaker::DoIt(TransformConstraint);
+			}
+		})),
+		NAME_None,
+		EUserInterfaceActionType::Button);
+	}
+	MenuBuilder.EndSection();
+
+	return MenuBuilder.MakeWidget();
+}
+
 END_SLATE_FUNCTION_BUILD_OPTIMIZATION
 
 #undef LOCTEXT_NAMESPACE
diff --git a/Engine/Plugins/Animation/ControlRig/Source/ControlRigEditor/Private/Editor/Constraints/SConstraintsWidget.h b/Engine/Plugins/Animation/ControlRig/Source/ControlRigEditor/Private/Editor/Constraints/SConstraintsWidget.h
--- a/Engine/Plugins/Animation/ControlRig/Source/ControlRigEditor/Private/Editor/Constraints/SConstraintsWidget.h
+++ b/Engine/Plugins/Animation/ControlRig/Source/ControlRigEditor/Private/Editor/Constraints/SConstraintsWidget.h
@@ -211,8 +211,8 @@
 	void OnActorSelectionChanged(const TArray<UObject*>& NewSelection, bool bForceRefresh);
 
 	/**  */
-	bool IsFirstItem(const TSharedPtr<FEditableConstraintItem>& Item) const;
-	bool IsLastItem(const TSharedPtr<FEditableConstraintItem>& Item) const;
+	bool CanMoveUp(const TSharedPtr<FEditableConstraintItem>& Item) const;
+	bool CanMoveDown(const TSharedPtr<FEditableConstraintItem>& Item) const;
 
 	/**  */
 	void MoveItemUp(const TSharedPtr<FEditableConstraintItem>& Item);
@@ -229,6 +229,9 @@
 	/** Generates a widget for the specified item */
 	TSharedRef<ITableRow> OnGenerateWidgetForItem(ItemSharedPtr InItem, const TSharedRef<STableViewBase>& OwnerTable);
 
+	/** @todo documentation. */
+	TSharedPtr< SWidget > CreateContextMenu();
+	
 	void RegisterSelectionChanged();
 	void UnregisterSelectionChanged();
 	
diff --git a/Engine/Plugins/Animation/ControlRig/Source/ControlRigEditor/Private/Sequencer/ControlRigParameterTrackEditor.cpp b/Engine/Plugins/Animation/ControlRig/Source/ControlRigEditor/Private/Sequencer/ControlRigParameterTrackEditor.cpp
--- a/Engine/Plugins/Animation/ControlRig/Source/ControlRigEditor/Private/Sequencer/ControlRigParameterTrackEditor.cpp
+++ b/Engine/Plugins/Animation/ControlRig/Source/ControlRigEditor/Private/Sequencer/ControlRigParameterTrackEditor.cpp
@@ -1495,33 +1495,36 @@ void FControlRigParameterTrackEditor::OnAddTransformKeysForSelectedObjects(EMovi
 	{
 		return;
 	}
-
-	FControlRigEditMode* ControlRigEditMode = GetEditMode();
-	if (ControlRigEditMode)
+	
+	const FControlRigEditMode* ControlRigEditMode = GetEditMode();
+	if (!ControlRigEditMode)
 	{
-		TMap<UControlRig*, TArray<FRigElementKey>> SelectedControls;
-		ControlRigEditMode->GetAllSelectedControls(SelectedControls);
+		return;
+	}
 
-		for (TPair<UControlRig*, TArray<FRigElementKey>>& Selection : SelectedControls)
+	TMap<UControlRig*, TArray<FRigElementKey>> SelectedControls;
+	ControlRigEditMode->GetAllSelectedControls(SelectedControls);
+
+	const EControlRigContextChannelToKey ChannelsToKey = static_cast<EControlRigContextChannelToKey>(Channel); 
+	
+	for (const TPair<UControlRig*, TArray<FRigElementKey>>& Selection : SelectedControls)
 	{
-			UControlRig* ControlRig = Selection.Key;
-		FString OurName = ControlRig->GetName();
-		FName Name(*OurName);
-		if (TSharedPtr<IControlRigObjectBinding> ObjectBinding = ControlRig->GetObjectBinding())
+		UControlRig* ControlRig = Selection.Key;
+		if (const TSharedPtr<IControlRigObjectBinding> ObjectBinding = ControlRig->GetObjectBinding())
 		{
-			TArray<FName> ControlNames = ControlRig->CurrentControlSelection();
-			for (const FName& ControlName : ControlNames)
+			if (USceneComponent* Component = Cast<USceneComponent>(ObjectBinding->GetBoundObject()))
 			{
-				USceneComponent* Component = Cast<USceneComponent>(ObjectBinding->GetBoundObject());
-				if (Component)
+				const FName Name(*ControlRig->GetName());
+			
+				const TArray<FName> ControlNames = ControlRig->CurrentControlSelection();
+				for (const FName& ControlName : ControlNames)
 				{
-					AddControlKeys(Component, ControlRig, Name, ControlName, (EControlRigContextChannelToKey)Channel, ESequencerKeyMode::ManualKeyForced, FLT_MAX);
+					AddControlKeys(Component, ControlRig, Name, ControlName, ChannelsToKey, ESequencerKeyMode::ManualKeyForced, FLT_MAX);
 				}
 			}
 		}
 	}
 }
-}
 
 //function to evaluate a Control and Set it on the ControlRig
 static void EvaluateThisControl(UMovieSceneControlRigParameterSection* Section, const FName& ControlName, const FFrameTime& FrameTime)
@@ -2765,7 +2768,14 @@ FKeyPropertyResult FControlRigParameterTrackEditor::AddKeysToControlRig(
 	return KeyPropertyResult;
 }
 
-void FControlRigParameterTrackEditor::AddControlKeys(USceneComponent* InSceneComp, UControlRig* InControlRig, FName ControlRigName, FName RigControlName, EControlRigContextChannelToKey ChannelsToKey, ESequencerKeyMode KeyMode, float InLocalTime)
+void FControlRigParameterTrackEditor::AddControlKeys(
+	USceneComponent* InSceneComp,
+	UControlRig* InControlRig,
+	FName ControlRigName,
+	FName RigControlName,
+	EControlRigContextChannelToKey ChannelsToKey,
+	ESequencerKeyMode KeyMode,
+	float InLocalTime)
 {
 	if (KeyMode == ESequencerKeyMode::ManualKey || (GetSequencer().IsValid() && !GetSequencer()->IsAllowedToChange()))
 	{
diff --git a/Engine/Plugins/Animation/ControlRig/Source/ControlRigEditor/Private/Tools/BakingHelper.cpp b/Engine/Plugins/Animation/ControlRig/Source/ControlRigEditor/Private/Tools/BakingHelper.cpp
new file mode 100644
--- /dev/null
+++ b/Engine/Plugins/Animation/ControlRig/Source/ControlRigEditor/Private/Tools/BakingHelper.cpp
@@ -0,0 +1,230 @@
+﻿// Copyright Epic Games, Inc. All Rights Reserved.
+
+
+#include "Tools/BakingHelper.h"
+
+#include "ControlRig.h"
+#include "ILevelSequenceEditorToolkit.h"
+#include "Tracks/MovieScene3DTransformTrack.h"
+#include "Sections/MovieScene3DTransformSection.h"
+#include "ISequencer.h"
+#include "LevelSequence.h"
+#include "LevelSequenceEditorBlueprintLibrary.h"
+#include "MovieSceneSequence.h"
+#include "MovieSceneToolHelpers.h"
+#include "Channels/MovieSceneChannelProxy.h"
+
+TWeakPtr<ISequencer> FBakingHelper::GetSequencer()
+{
+	// if getting sequencer from level sequence need to use the current(master), not the focused
+	if (ULevelSequence* LevelSequence = ULevelSequenceEditorBlueprintLibrary::GetCurrentLevelSequence())
+	{
+		if (UAssetEditorSubsystem* AssetEditorSubsystem = GEditor->GetEditorSubsystem<UAssetEditorSubsystem>())
+		{
+			static constexpr bool bFocusIfOpen = false;
+			IAssetEditorInstance* AssetEditor = AssetEditorSubsystem->FindEditorForAsset(LevelSequence, bFocusIfOpen);
+			const ILevelSequenceEditorToolkit* LevelSequenceEditor = static_cast<ILevelSequenceEditorToolkit*>(AssetEditor);
+			return LevelSequenceEditor ? LevelSequenceEditor->GetSequencer() : nullptr;
+		}
+	}
+	return nullptr;
+}
+
+
+void FBakingHelper::CalculateFramesBetween(
+	const UMovieScene* MovieScene,
+	FFrameNumber StartFrame,
+	FFrameNumber EndFrame,
+	TArray<FFrameNumber>& OutFrames)
+{
+	if(StartFrame > EndFrame)
+	{
+		Swap(StartFrame, EndFrame);
+	}
+	
+	const FFrameRate TickResolution = MovieScene->GetTickResolution();
+	const FFrameRate DisplayResolution = MovieScene->GetDisplayRate();
+
+	const FFrameNumber StartTimeInDisplay = FFrameRate::TransformTime(FFrameTime(StartFrame), TickResolution, DisplayResolution).FloorToFrame();
+	const FFrameNumber EndTimeInDisplay = FFrameRate::TransformTime(FFrameTime(EndFrame), TickResolution, DisplayResolution).CeilToFrame();
+
+	OutFrames.Reserve(EndTimeInDisplay.Value - StartTimeInDisplay.Value + 1);
+	for (FFrameNumber DisplayFrameNumber = StartTimeInDisplay; DisplayFrameNumber <= EndTimeInDisplay; ++DisplayFrameNumber)
+	{
+		FFrameNumber TickFrameNumber = FFrameRate::TransformTime(FFrameTime(DisplayFrameNumber), DisplayResolution, TickResolution).FrameNumber;
+		OutFrames.Add(TickFrameNumber);
+	}
+}
+
+UMovieScene3DTransformSection* FBakingHelper::GetTransformSection(const ISequencer* InSequencer, const FGuid& InGuid)
+{
+	if (!InSequencer || !InSequencer->GetFocusedMovieSceneSequence())
+	{
+		return nullptr;
+	}
+	
+	UMovieScene* MovieScene = InSequencer->GetFocusedMovieSceneSequence()->GetMovieScene();
+	if (!MovieScene)
+	{
+		return nullptr;
+	}
+
+	UMovieScene3DTransformTrack* TransformTrack = MovieScene->FindTrack<UMovieScene3DTransformTrack>(InGuid);
+	if (!TransformTrack)
+	{
+		MovieScene->Modify();
+		TransformTrack = MovieScene->AddTrack<UMovieScene3DTransformTrack>(InGuid);
+	}
+	TransformTrack->Modify();
+
+	static constexpr FFrameNumber Frame0;
+	bool bSectionAdded = false;
+	UMovieScene3DTransformSection* TransformSection =
+				Cast<UMovieScene3DTransformSection>(TransformTrack->FindOrAddSection(Frame0, bSectionAdded));
+	if (!TransformSection)
+	{
+		return nullptr;
+	}
+
+	TransformSection->Modify();
+	if (bSectionAdded)
+	{
+		TransformSection->SetRange(TRange<FFrameNumber>::All());
+	}
+	
+	return TransformSection;
+}
+
+bool FBakingHelper::AddTransformKeys(
+	const UMovieScene3DTransformSection* InTransformSection,
+	const TArray<FFrameNumber>& Frames,
+	const TArray<FTransform>& InLocalTransforms,
+	const EMovieSceneTransformChannel& InChannels)
+{
+	if (!InTransformSection)
+	{
+		return false;	
+	}
+
+	if (Frames.IsEmpty() || Frames.Num() != InLocalTransforms.Num())
+	{
+		return false;
+	}
+
+	auto GetValue = [](const uint32 Index,const FVector& InLocation, const FRotator& InRotation, const FVector& InScale)
+	{
+		switch (Index)
+		{
+		case 0:
+			return InLocation.X;
+		case 1:
+			return InLocation.Y;
+		case 2:
+			return InLocation.Z;
+		case 3:
+			return InRotation.Roll;
+		case 4:
+			return InRotation.Pitch;
+		case 5:
+			return InRotation.Yaw;
+		case 6:
+			return InScale.X;
+		case 7:
+			return InScale.Y;
+		case 8:
+			return InScale.Z;
+		default:
+			ensure(false);
+			break;
+		}
+		return 0.0;
+	};
+
+	const bool bKeyTranslation = EnumHasAllFlags(InChannels, EMovieSceneTransformChannel::Translation);
+	const bool bKeyRotation = EnumHasAllFlags(InChannels, EMovieSceneTransformChannel::Rotation);
+	const bool bKeyScale = EnumHasAllFlags(InChannels, EMovieSceneTransformChannel::Scale);
+
+	TArray<uint32> ChannelsIndexToKey;
+	if (bKeyTranslation)
+	{
+		ChannelsIndexToKey.Append({0,1,2});
+	}
+	if (bKeyRotation)
+	{
+		ChannelsIndexToKey.Append({3,4,5});
+	}
+	if (bKeyScale)
+	{
+		ChannelsIndexToKey.Append({6,7,8});
+	}
+
+	const TArrayView<FMovieSceneDoubleChannel*> Channels =
+		InTransformSection->GetChannelProxy().GetChannels<FMovieSceneDoubleChannel>();
+
+	for (int32 Index = 0; Index < Frames.Num(); ++Index)
+	{
+		const FFrameNumber& Frame = Frames[Index];
+		const FTransform& LocalTransform = InLocalTransforms[Index];
+		
+		const FVector Location = LocalTransform.GetLocation();
+		const FRotator Rotation = LocalTransform.GetRotation().Rotator();
+		const FVector Scale3D = LocalTransform.GetScale3D();
+
+		for (const int32 ChannelIndex: ChannelsIndexToKey)
+		{
+			// get value
+			const double Value = GetValue(ChannelIndex, Location, Rotation, Scale3D);
+
+			// set default
+			if (Index == 0 && !Channels[ChannelIndex]->GetDefault().IsSet())
+			{
+				Channels[ChannelIndex]->SetDefault(Value);
+			}
+			
+			// add key
+			TMovieSceneChannelData<FMovieSceneDoubleValue> ChannelData = Channels[ChannelIndex]->GetData();
+			MovieSceneToolHelpers::SetOrAddKey(ChannelData, Frame, Value);
+		}
+	}
+
+	//now we need to set auto tangents
+	for (const int32 ChannelIndex: ChannelsIndexToKey)
+	{
+		Channels[ChannelIndex]->AutoSetTangents();
+	}
+
+	return true;
+}
+
+bool FBakingHelper::AddTransformKeys(
+	UControlRig* InControlRig, const FName& InControlName,
+	const TArray<FFrameNumber>& InFrames,
+	const TArray<FTransform>& InLocalTransforms,
+	const EMovieSceneTransformChannel& InChannels,
+	const FFrameRate& InTickResolution)
+{
+	if (InFrames.IsEmpty() || InFrames.Num() != InLocalTransforms.Num())
+	{
+		return false;
+	}
+
+	static constexpr bool bNotify = true;
+	static constexpr bool bUndo = false;
+	static constexpr bool bFixEuler = true;
+	
+	FRigControlModifiedContext KeyframeContext;
+	KeyframeContext.SetKey = EControlRigSetKey::Always;
+	KeyframeContext.KeyMask = static_cast<uint32>(InChannels);
+	
+	for (int32 Index = 0; Index < InFrames.Num(); ++Index)
+	{
+		const FTransform& LocalTransform = InLocalTransforms[Index];
+
+		const FFrameNumber& Frame = InFrames[Index];
+		KeyframeContext.LocalTime = InTickResolution.AsSeconds(FFrameTime(Frame));
+		
+		InControlRig->SetControlLocalTransform(InControlName, LocalTransform, bNotify, KeyframeContext, bUndo, bFixEuler);
+	}
+
+	return true;
+}
\ No newline at end of file
diff --git a/Engine/Plugins/Animation/ControlRig/Source/ControlRigEditor/Private/Tools/ConstraintBaker.cpp b/Engine/Plugins/Animation/ControlRig/Source/ControlRigEditor/Private/Tools/ConstraintBaker.cpp
new file mode 100644
--- /dev/null
+++ b/Engine/Plugins/Animation/ControlRig/Source/ControlRigEditor/Private/Tools/ConstraintBaker.cpp
@@ -0,0 +1,198 @@
+﻿// Copyright Epic Games, Inc. All Rights Reserved.
+
+#include "Tools/ConstraintBaker.h"
+
+#include "IControlRigObjectBinding.h"
+#include "TransformConstraint.h"
+
+#include "ISequencer.h"
+
+#include "IKeyArea.h"
+#include "LevelSequence.h"
+#include "MovieSceneToolHelpers.h"
+#include "TransformableHandle.h"
+#include "Channels/MovieSceneChannelProxy.h"
+#include "Constraints/ControlRigTransformableHandle.h"
+#include "Tools/BakingHelper.h"
+#include "Sections/MovieScene3DTransformSection.h"
+
+#define LOCTEXT_NAMESPACE "ConstraintBaker"
+
+namespace
+{
+	
+UMovieScene3DTransformSection* GetTransformSection(const TSharedPtr<ISequencer>& InSequencer, AActor* InActor)
+{
+	if (!InSequencer || !InSequencer->GetFocusedMovieSceneSequence())
+	{
+		return nullptr;
+	}
+	
+	const FGuid Guid = InSequencer->GetHandleToObject(InActor,true);
+	if (!Guid.IsValid())
+	{
+		return nullptr;
+	}
+	
+	return FBakingHelper::GetTransformSection(InSequencer.Get(), Guid);
+}
+
+void BakeComponent(
+	const TSharedPtr<ISequencer>& InSequencer,
+	const UTransformableComponentHandle* InComponentHandle,
+	const TArray<FFrameNumber>& InFrames,
+	const TArray<FTransform>& InTransforms,
+	const EMovieSceneTransformChannel& InChannels)
+{
+	if (!InComponentHandle->IsValid())
+	{
+		return;
+	}
+	AActor* Actor = InComponentHandle->Component->GetOwner();
+	if (!Actor)
+	{
+		return;
+	}
+	const UMovieScene3DTransformSection* TransformSection = GetTransformSection(InSequencer, Actor);
+	if (!TransformSection)
+	{
+		return;
+	}
+
+	FBakingHelper::AddTransformKeys(TransformSection, InFrames, InTransforms, InChannels);
+}
+
+void BakeControl(
+	const TSharedPtr<ISequencer>& InSequencer,
+	const UTransformableControlHandle* InControlHandle,
+	const TArray<FFrameNumber>& InFrames,
+	const TArray<FTransform>& InLocalTransforms,
+	const EMovieSceneTransformChannel& InChannels)
+{
+	if (!InControlHandle->IsValid())
+	{
+		return;
+	}
+
+	const UMovieScene* MovieScene = InSequencer->GetFocusedMovieSceneSequence()->GetMovieScene();
+
+	FBakingHelper::AddTransformKeys(
+		InControlHandle->ControlRig.Get(),
+		InControlHandle->ControlName,
+		InFrames,
+		InLocalTransforms,
+		InChannels,
+		MovieScene->GetTickResolution());
+}
+
+EMovieSceneTransformChannel GetChannelsToKey(const UTickableTransformConstraint* InConstraint)
+{
+	static const TMap< ETransformConstraintType, EMovieSceneTransformChannel > ConstraintToChannels({
+	{ETransformConstraintType::Translation, EMovieSceneTransformChannel::Translation},
+	{ETransformConstraintType::Rotation, EMovieSceneTransformChannel::Rotation},
+	{ETransformConstraintType::Scale, EMovieSceneTransformChannel::Scale},
+	{ETransformConstraintType::Parent, EMovieSceneTransformChannel::AllTransform},
+	{ETransformConstraintType::LookAt, EMovieSceneTransformChannel::Rotation}
+	});
+
+	const ETransformConstraintType ConstType = static_cast<ETransformConstraintType>(InConstraint->GetType());
+	if (const EMovieSceneTransformChannel* Channel = ConstraintToChannels.Find(ConstType))
+	{
+		return *Channel;; 
+	}
+	
+	return EMovieSceneTransformChannel::AllTransform;
+}
+	
+}
+
+void FConstraintBaker::DoIt(UTickableTransformConstraint* InConstraint)
+{
+	const TWeakPtr<ISequencer> WeakSequencer = FBakingHelper::GetSequencer();
+	if (!WeakSequencer.IsValid() || !WeakSequencer.Pin()->GetFocusedMovieSceneSequence())
+	{
+		return;
+	}
+	const TSharedPtr<ISequencer> Sequencer = WeakSequencer.Pin();
+	const UMovieScene* MovieScene = Sequencer->GetFocusedMovieSceneSequence()->GetMovieScene();
+	if (!MovieScene)
+	{
+		return;
+	}
+
+	// compute frames
+	const FFrameNumber StartFrame = MovieScene->GetPlaybackRange().GetLowerBoundValue();
+	const FFrameNumber EndFrame = MovieScene->GetPlaybackRange().GetUpperBoundValue();
+	TArray<FFrameNumber> Frames;
+	FBakingHelper::CalculateFramesBetween(MovieScene, StartFrame, EndFrame, Frames);
+
+	// compute transforms
+	TArray<FTransform> Transforms;
+	GetChildTransforms(Sequencer, *InConstraint, Frames, true, Transforms);
+
+	// bake to channel curves
+	const EMovieSceneTransformChannel Channels = GetChannelsToKey(InConstraint);
+	BakeChild(Sequencer, InConstraint->ChildTRSHandle, Frames, Transforms, Channels);
+
+	// disable constraint
+	InConstraint->Active = false;
+	InConstraint->ConstraintTick.SetTickFunctionEnable(false);
+	
+	// notify
+	Sequencer->NotifyMovieSceneDataChanged(EMovieSceneDataChangeType::MovieSceneStructureItemsChanged);
+}
+
+void FConstraintBaker::GetChildTransforms(
+	const TSharedPtr<ISequencer>& InSequencer,
+	const UTickableTransformConstraint& InConstraint,
+	const TArray<FFrameNumber>& InFrames,
+	const bool bLocal,
+	TArray<FTransform>& OutTransforms)
+{
+	const UMovieScene* MovieScene = InSequencer->GetFocusedMovieSceneSequence()->GetMovieScene();
+	
+	const FFrameRate TickResolution = MovieScene->GetTickResolution();
+	const EMovieScenePlayerStatus::Type PlaybackStatus = InSequencer->GetPlaybackStatus();
+
+	OutTransforms.SetNum(InFrames.Num());
+	for (int32 Index = 0; Index < InFrames.Num(); ++Index)
+	{
+		const FFrameNumber& FrameNumber = InFrames[Index];
+	
+		// evaluate animation
+		const FMovieSceneEvaluationRange EvaluationRange = FMovieSceneEvaluationRange(FFrameTime (FrameNumber), TickResolution);
+		const FMovieSceneContext Context = FMovieSceneContext(EvaluationRange, PlaybackStatus).SetHasJumped(true);
+
+		InSequencer->GetEvaluationTemplate().Evaluate(Context, *InSequencer);
+	
+		// evaluate constraints
+		InConstraint.Evaluate();
+	
+		// evaluate ControlRig?
+		// ControlRig->Evaluate_AnyThread();
+		
+		OutTransforms[Index] = bLocal ? InConstraint.GetChildLocalTransform() : InConstraint.GetChildGlobalTransform();
+	}
+}
+
+void FConstraintBaker::BakeChild(
+	const TSharedPtr<ISequencer>& InSequencer,
+	UTransformableHandle* InHandle,
+	const TArray<FFrameNumber>& InFrames,
+	const TArray<FTransform>& InTransforms,
+	const EMovieSceneTransformChannel& InChannels)
+{
+	if (const UTransformableComponentHandle* ComponentHandle = Cast<UTransformableComponentHandle>(InHandle))
+	{
+		return BakeComponent(InSequencer, ComponentHandle, InFrames, InTransforms, InChannels); 
+	}
+	
+	if (const UTransformableControlHandle* ControlHandle = Cast<UTransformableControlHandle>(InHandle))
+	{
+		return BakeControl(InSequencer, ControlHandle, InFrames, InTransforms, InChannels); 
+	}
+}
+
+#undef LOCTEXT_NAMESPACE
+
+
diff --git a/Engine/Plugins/Animation/ControlRig/Source/ControlRigEditor/Private/Tools/ControlRigSnapper.cpp b/Engine/Plugins/Animation/ControlRig/Source/ControlRigEditor/Private/Tools/ControlRigSnapper.cpp
--- a/Engine/Plugins/Animation/ControlRig/Source/ControlRigEditor/Private/Tools/ControlRigSnapper.cpp
+++ b/Engine/Plugins/Animation/ControlRig/Source/ControlRigEditor/Private/Tools/ControlRigSnapper.cpp
@@ -29,6 +29,7 @@
 #include "MovieSceneToolHelpers.h"
 #include "LevelSequencePlayer.h"
 #include "LevelSequenceActor.h"
+#include "Tools/BakingHelper.h"
 
 #define LOCTEXT_NAMESPACE "ControlRigSnapper"
 
@@ -100,16 +101,7 @@
 
 TWeakPtr<ISequencer> FControlRigSnapper::GetSequencer()
 {
-	TWeakPtr<ISequencer> WeakSequencer = nullptr;
-	//if getting sequencer from level sequence need to use the current(master), not the focused
-	ULevelSequence* LevelSequence = ULevelSequenceEditorBlueprintLibrary::GetCurrentLevelSequence();
-	if (LevelSequence)
-	{
-		IAssetEditorInstance* AssetEditor = GEditor->GetEditorSubsystem<UAssetEditorSubsystem>()->FindEditorForAsset(LevelSequence, false);
-		ILevelSequenceEditorToolkit* LevelSequenceEditor = static_cast<ILevelSequenceEditorToolkit*>(AssetEditor);
-		WeakSequencer = LevelSequenceEditor ? LevelSequenceEditor->GetSequencer() : nullptr;
-	}
-	return WeakSequencer;
+	return FBakingHelper::GetSequencer();
 }
 
 static bool LocalGetControlRigControlTransforms(IMovieScenePlayer* Player, const TOptional<FFrameNumber>& CurrentFrame, UMovieSceneSequence* MovieSceneSequence, FMovieSceneSequenceIDRef Template, FMovieSceneSequenceTransform& RootToLocalTransform,
@@ -192,37 +184,18 @@
 	FGuid Guid;
 	AActor* Actor;
 
-	bool SetTransform(ISequencer* Sequencer, const TArray<FFrameNumber>& Frames,
-		const TArray<FTransform>& WorldTransformsToSnapTo, const UControlRigSnapSettings* SnapSettings)
+	bool SetTransform(
+		ISequencer* Sequencer, const TArray<FFrameNumber>& Frames,
+		const TArray<FTransform>& WorldTransformsToSnapTo, const UControlRigSnapSettings* SnapSettings) const
 	{
-		if (!Sequencer || !Sequencer->GetFocusedMovieSceneSequence())
-		{
-			return false;
-		}
-		UMovieScene* MovieScene = Sequencer->GetFocusedMovieSceneSequence()->GetMovieScene();
-
-		UMovieScene3DTransformTrack* TransformTrack = MovieScene->FindTrack<UMovieScene3DTransformTrack>(Guid);
-		if (!TransformTrack)
-		{
-			MovieScene->Modify();
-			TransformTrack = MovieScene->AddTrack<UMovieScene3DTransformTrack>(Guid);
-		}
-		TransformTrack->Modify();
-
-		bool bSectionAdded = false;
-		UMovieScene3DTransformSection* TransformSection = Cast<UMovieScene3DTransformSection>(TransformTrack->FindOrAddSection(0, bSectionAdded));
+		// get section
+		const UMovieScene3DTransformSection* TransformSection = FBakingHelper::GetTransformSection(Sequencer, Guid);
 		if (!TransformSection)
 		{
 			return false;
 		}
 
-		TransformSection->Modify();
-
-		if (bSectionAdded)
-		{
-			TransformSection->SetRange(TRange<FFrameNumber>::All());
-		}
-
+		// fill parent transforms
 		TArray<FTransform> ParentWorldTransforms; 
 		AActor* ParentActor = Actor->GetAttachParentActor();
 		if (ParentActor)
@@ -240,138 +213,41 @@
 			}
 		}
 
-		TArrayView<FMovieSceneDoubleChannel*> Channels = TransformSection->GetChannelProxy().GetChannels<FMovieSceneDoubleChannel>();
-		for (int32 Index = 0; Index < Frames.Num(); ++Index)
-		{
-			const FFrameNumber& Frame = Frames[Index];
-			FTransform ParentTransform = ParentWorldTransforms[Index];
-			FTransform WorldTransform = WorldTransformsToSnapTo[Index];
-			FTransform LocalTransform = WorldTransform.GetRelativeTransform(ParentTransform);
-			FVector Location = LocalTransform.GetLocation();
-			FRotator Rotation = LocalTransform.GetRotation().Rotator();
-			FVector Scale3D = LocalTransform.GetScale3D();
-			if (Index == 0) 
-			{
-				if (SnapSettings->bSnapPosition)
-				{
-					if (!Channels[0]->GetDefault().IsSet())
-					{
-						Channels[0]->SetDefault(Location.X);
-					}
-					if (!Channels[1]->GetDefault().IsSet())
-					{
-						Channels[1]->SetDefault(Location.Y);
-					}
-					if (!Channels[2]->GetDefault().IsSet())
-					{
-						Channels[2]->SetDefault(Location.Z);
-					}
-				}
-				if (SnapSettings->bSnapRotation)
-				{
-					if (!Channels[3]->GetDefault().IsSet())
-					{
-						Channels[3]->SetDefault(Rotation.Roll);
-					}
-					if (!Channels[4]->GetDefault().IsSet())
-					{
-						Channels[4]->SetDefault(Rotation.Pitch);
-					}
-					if (!Channels[5]->GetDefault().IsSet())
-					{
-						Channels[5]->SetDefault(Rotation.Yaw);
-					}
-				}
-				if (SnapSettings->bSnapScale)
-				{
-					if (!Channels[6]->GetDefault().IsSet())
-					{
-						Channels[6]->SetDefault(Scale3D.X);
-					}
-					if (!Channels[7]->GetDefault().IsSet())
-					{
-						Channels[7]->SetDefault(Scale3D.Y);
-					}
-					if (!Channels[8]->GetDefault().IsSet())
-					{
-						Channels[8]->SetDefault(Scale3D.Z);
-					}
-				}
-			}
-
-			if (SnapSettings->bSnapPosition)
-			{
-				TMovieSceneChannelData<FMovieSceneDoubleValue> ChannelData = Channels[0]->GetData();
-				MovieSceneToolHelpers::SetOrAddKey(ChannelData, Frame, Location.X);
-				ChannelData = Channels[1]->GetData();
-				MovieSceneToolHelpers::SetOrAddKey(ChannelData, Frame, Location.Y);
-				ChannelData = Channels[2]->GetData();
-				MovieSceneToolHelpers::SetOrAddKey(ChannelData, Frame, Location.Z);
-			}
-			if (SnapSettings->bSnapRotation)
-			{
-				//todo winding
-				TMovieSceneChannelData<FMovieSceneDoubleValue> ChannelData = Channels[3]->GetData();
-				MovieSceneToolHelpers::SetOrAddKey(ChannelData, Frame, Rotation.Roll);
-				ChannelData = Channels[4]->GetData();
-				MovieSceneToolHelpers::SetOrAddKey(ChannelData, Frame, Rotation.Pitch);
-				ChannelData = Channels[5]->GetData();
-				MovieSceneToolHelpers::SetOrAddKey(ChannelData, Frame, Rotation.Yaw);
-			}
-			if (SnapSettings->bSnapScale)
-			{
-				TMovieSceneChannelData<FMovieSceneDoubleValue> ChannelData = Channels[6]->GetData();
-				MovieSceneToolHelpers::SetOrAddKey(ChannelData, Frame, Scale3D.X);
-				ChannelData = Channels[7]->GetData();
-				MovieSceneToolHelpers::SetOrAddKey(ChannelData, Frame, Scale3D.X);
-				ChannelData = Channels[8]->GetData();
-				MovieSceneToolHelpers::SetOrAddKey(ChannelData, Frame, Scale3D.X);
-			}		
-		}
-		//now we need to set auto tangents
+		// fill channels to key
+		EMovieSceneTransformChannel Channels;
 		if (SnapSettings->bSnapPosition)
 		{
-			Channels[0]->AutoSetTangents();
-			Channels[1]->AutoSetTangents();
-			Channels[2]->AutoSetTangents();		
+			EnumAddFlags(Channels, EMovieSceneTransformChannel::Translation);
 		}
 		if (SnapSettings->bSnapRotation)
 		{
-			Channels[3]->AutoSetTangents();
-			Channels[4]->AutoSetTangents();
-			Channels[5]->AutoSetTangents();
+			EnumAddFlags(Channels, EMovieSceneTransformChannel::Rotation);
 		}
 		if (SnapSettings->bSnapScale)
 		{
-			Channels[6]->AutoSetTangents();
-			Channels[7]->AutoSetTangents();
-			Channels[8]->AutoSetTangents();
+			EnumAddFlags(Channels, EMovieSceneTransformChannel::Scale);
 		}
+
+		// compute local transforms
+		TArray<FTransform> LocalTransforms;
+		LocalTransforms.Reserve(Frames.Num());
+		for (int32 Index = 0; Index < Frames.Num(); ++Index)
+		{
+			const FTransform& ParentTransform = ParentWorldTransforms[Index];
+			const FTransform& WorldTransform = WorldTransformsToSnapTo[Index];
+			LocalTransforms.Add(WorldTransform.GetRelativeTransform(ParentTransform));
+		}
+
+		// add keys
+		FBakingHelper::AddTransformKeys(TransformSection, Frames, LocalTransforms, Channels);
+
+		// notify
 		Sequencer->NotifyMovieSceneDataChanged(EMovieSceneDataChangeType::MovieSceneStructureItemsChanged);
 
 		return true;
 	}
 };
 
-static void CalculateFramesToSnap(ISequencer* InSequencer, UMovieScene* MovieScene, FFrameNumber StartFrame, FFrameNumber EndFrame, TArray<FFrameNumber>& OutFrames)
-{
-	if(StartFrame > EndFrame)
-	{
-		Swap(StartFrame, EndFrame);
-	}
-	
-	FFrameRate TickResolution = MovieScene->GetTickResolution();
-	FFrameRate DisplayResolution = MovieScene->GetDisplayRate();
-
-	FFrameNumber StartTimeInDisplay = FFrameRate::TransformTime(FFrameTime(StartFrame), TickResolution, DisplayResolution).FloorToFrame();
-	FFrameNumber EndTimeInDisplay = FFrameRate::TransformTime(FFrameTime(EndFrame), TickResolution, DisplayResolution).CeilToFrame();
-	for (FFrameNumber DisplayFrameNumber = StartTimeInDisplay; DisplayFrameNumber <= EndTimeInDisplay; ++DisplayFrameNumber)
-	{
-		FFrameNumber TickFrameNumber = FFrameRate::TransformTime(FFrameTime(DisplayFrameNumber), DisplayResolution, TickResolution).FrameNumber;
-		OutFrames.Add(TickFrameNumber);
-	}
-}
-
 //returns true if world is calculated, false if there are no parents
 static bool CalculateWorldTransformsFromParents(ISequencer* Sequencer, const FControlRigSnapperSelection& ParentToSnap,
 	const TArray<FFrameNumber>& Frames, TArray<FTransform>& OutParentWorldTransforms)
@@ -443,7 +319,7 @@
 		MovieScene->Modify();
 		
 		TArray<FFrameNumber> Frames;
-		CalculateFramesToSnap(Sequencer,MovieScene, StartFrame, EndFrame, Frames);
+		FBakingHelper::CalculateFramesBetween(MovieScene, StartFrame, EndFrame, Frames);
 
 		TArray<FTransform> WorldTransformToSnap;
 		bool bSnapToFirstFrameNotParents = !CalculateWorldTransformsFromParents(Sequencer, ParentToSnap, Frames, WorldTransformToSnap);
diff --git a/Engine/Plugins/Animation/ControlRig/Source/ControlRigEditor/Public/Tools/BakingHelper.h b/Engine/Plugins/Animation/ControlRig/Source/ControlRigEditor/Public/Tools/BakingHelper.h
new file mode 100644
--- /dev/null
+++ b/Engine/Plugins/Animation/ControlRig/Source/ControlRigEditor/Public/Tools/BakingHelper.h
@@ -0,0 +1,53 @@
+﻿// Copyright Epic Games, Inc. All Rights Reserved.
+
+#pragma once
+
+#include "Misc/FrameNumber.h"
+#include "Templates/SharedPointer.h"
+
+class UMovieScene3DTransformSection;
+class AActor;
+class ISequencer;
+struct FGuid;
+class UMovieScene;
+class UControlRig;
+struct FFrameRate;
+enum class EMovieSceneTransformChannel : uint32;
+
+/**
+ * FBakingHelper
+ */
+
+struct FBakingHelper
+{
+	/** Returns the current sequencer. */
+	static TWeakPtr<ISequencer> GetSequencer();
+	
+	/** Returns the frame numbers between start and end. */
+	static void CalculateFramesBetween(
+		const UMovieScene* MovieScene,
+		FFrameNumber StartFrame,
+		FFrameNumber EndFrame,
+		TArray<FFrameNumber>& OutFrames);
+	
+	/** Returns the transform section for that guid. */
+	static UMovieScene3DTransformSection* GetTransformSection(
+		const ISequencer* InSequencer,
+		const FGuid& InGuid);
+
+	/** Adds transform keys to the section based on the channels filters. */
+	static bool AddTransformKeys(
+		const UMovieScene3DTransformSection* InTransformSection,
+		const TArray<FFrameNumber>& Frames,
+		const TArray<FTransform>& InLocalTransforms,
+		const EMovieSceneTransformChannel& InChannels);
+
+	/** Adds transform keys to control based on the channels filters. */
+	static bool AddTransformKeys(
+		UControlRig* InControlRig,
+		const FName& InControlName,
+		const TArray<FFrameNumber>& InFrames,
+		const TArray<FTransform>& InLocalTransforms,
+		const EMovieSceneTransformChannel& InChannels,
+		const FFrameRate& InTickResolution);
+};
diff --git a/Engine/Plugins/Animation/ControlRig/Source/ControlRigEditor/Public/Tools/ConstraintBaker.h b/Engine/Plugins/Animation/ControlRig/Source/ControlRigEditor/Public/Tools/ConstraintBaker.h
new file mode 100644
--- /dev/null
+++ b/Engine/Plugins/Animation/ControlRig/Source/ControlRigEditor/Public/Tools/ConstraintBaker.h
@@ -0,0 +1,37 @@
+﻿// Copyright Epic Games, Inc. All Rights Reserved.
+
+#pragma once
+
+#include "Misc/FrameNumber.h"
+#include "Templates/SharedPointer.h"
+
+class UTransformableHandle;
+class UTickableTransformConstraint;
+class ISequencer;
+enum class EMovieSceneTransformChannel : uint32;
+
+struct FConstraintBaker
+{
+public:
+	/** @todo documentation. (subject to changes to handle start / stop as parameters) */
+	static void DoIt(UTickableTransformConstraint* InConstraint);
+	
+private:
+
+	/** Evaluates the constraint at each frames and stores the resulting child transforms. */
+	static void GetChildTransforms(
+		const TSharedPtr<ISequencer>& InSequencer,
+		const UTickableTransformConstraint& InConstraint,
+		const TArray<FFrameNumber>& InFrames,
+		const bool bLocal,
+		TArray<FTransform>& OutTransforms);
+
+	/** Bake the resulting transforms into the child's transform animation channels. */
+	static void BakeChild(
+		const TSharedPtr<ISequencer>& InSequencer,
+		UTransformableHandle* InHandle,
+		const TArray<FFrameNumber>& InFrames,
+		const TArray<FTransform>& InTransforms,
+		const EMovieSceneTransformChannel& InChannels);
+};
+
diff --git a/Engine/Source/Runtime/Experimental/Animation/Constraints/Private/ConstraintsManager.cpp b/Engine/Source/Runtime/Experimental/Animation/Constraints/Private/ConstraintsManager.cpp
--- a/Engine/Source/Runtime/Experimental/Animation/Constraints/Private/ConstraintsManager.cpp
+++ b/Engine/Source/Runtime/Experimental/Animation/Constraints/Private/ConstraintsManager.cpp
@@ -65,11 +65,34 @@
  * UTickableConstraint
  **/
 
+void UTickableConstraint::Evaluate() const
+{
+	ConstraintTick.EvaluateFunctions();
+}
+
 #if WITH_EDITOR
+
 FName UTickableConstraint::GetLabel() const
 {
 	return GetFName();
 }
+
+void UTickableConstraint::PostEditChangeProperty(FPropertyChangedEvent& PropertyChangedEvent)
+{
+	Super::PostEditChangeProperty(PropertyChangedEvent);
+
+	const FName PropertyName = PropertyChangedEvent.GetPropertyName();
+	if (PropertyName == GET_MEMBER_NAME_CHECKED(UTickableConstraint, Active))
+	{
+		ConstraintTick.SetTickFunctionEnable(Active);
+		if (Active)
+		{
+			Evaluate();
+		}
+	}
+}
+
+
 #endif
 
 /** 
@@ -336,6 +359,33 @@
 	}
 }
 
+UTickableConstraint* FConstraintsManagerController::GetConstraint(const FName& InConstraintName) const
+{
+	const int32 Index = GetConstraintIndex(InConstraintName);
+	if (Index == INDEX_NONE)
+	{
+		return nullptr;	
+	}
+	
+	return GetConstraint(Index);
+}
+
+UTickableConstraint* FConstraintsManagerController::GetConstraint(const int32 InConstraintIndex) const
+{
+	UConstraintsManager* Manager = FindManager();
+	if (!Manager)
+	{
+		return nullptr;	
+	}
+	
+	if (!Manager->Constraints.IsValidIndex(InConstraintIndex))
+	{
+		return nullptr;	
+	}
+
+	return Manager->Constraints[InConstraintIndex];
+}
+
 TArray< TObjectPtr<UTickableConstraint> > FConstraintsManagerController::GetParentConstraints(
 	const uint32 InTargetHash,
 	const bool bSorted) const
diff --git a/Engine/Source/Runtime/Experimental/Animation/Constraints/Private/TransformConstraint.cpp b/Engine/Source/Runtime/Experimental/Animation/Constraints/Private/TransformConstraint.cpp
--- a/Engine/Source/Runtime/Experimental/Animation/Constraints/Private/TransformConstraint.cpp
+++ b/Engine/Source/Runtime/Experimental/Animation/Constraints/Private/TransformConstraint.cpp
@@ -78,18 +78,36 @@
 {
 	if(ChildTRSHandle->IsValid())
 	{
-		ChildTRSHandle->SetTransform(InGlobal);
+		ChildTRSHandle->SetGlobalTransform(InGlobal);
 	}
 }
 
+void UTickableTransformConstraint::SetChildLocalTransform(const FTransform& InLocal) const
+{
+	if(ChildTRSHandle->IsValid())
+	{
+		ChildTRSHandle->SetLocalTransform(InLocal);
+	}
+}
+
 FTransform UTickableTransformConstraint::GetChildGlobalTransform() const
 {
-	return ChildTRSHandle->IsValid() ? ChildTRSHandle->GetTransform() : FTransform::Identity;
+	return ChildTRSHandle->IsValid() ? ChildTRSHandle->GetGlobalTransform() : FTransform::Identity;
 }
 
+FTransform UTickableTransformConstraint::GetChildLocalTransform() const
+{
+	return ChildTRSHandle->IsValid() ? ChildTRSHandle->GetLocalTransform() : FTransform::Identity;
+}
+
 FTransform UTickableTransformConstraint::GetParentGlobalTransform() const
 {
-	return ParentTRSHandle->IsValid() ? ParentTRSHandle->GetTransform() : FTransform::Identity;
+	return ParentTRSHandle->IsValid() ? ParentTRSHandle->GetGlobalTransform() : FTransform::Identity;
+}
+
+FTransform UTickableTransformConstraint::GetParentLocalTransform() const
+{
+	return ParentTRSHandle->IsValid() ? ParentTRSHandle->GetLocalTransform() : FTransform::Identity;
 }
 
 /** 
@@ -263,7 +281,7 @@
 			return;
 		}
 		
-		const FTransform ParentTransform = ParentTRSHandle->GetTransform();
+		const FTransform ParentTransform = ParentTRSHandle->GetGlobalTransform();
 
 		FTransform TargetTransform = OffsetTransform * ParentTransform;
 		if (ClampedWeight < 1.0f - KINDA_SMALL_NUMBER)
diff --git a/Engine/Source/Runtime/Experimental/Animation/Constraints/Private/TransformableHandle.cpp b/Engine/Source/Runtime/Experimental/Animation/Constraints/Private/TransformableHandle.cpp
--- a/Engine/Source/Runtime/Experimental/Animation/Constraints/Private/TransformableHandle.cpp
+++ b/Engine/Source/Runtime/Experimental/Animation/Constraints/Private/TransformableHandle.cpp
@@ -25,17 +25,12 @@
 UTransformableComponentHandle::~UTransformableComponentHandle()
 {}
 
-void UTransformableComponentHandle::PostLoad()
-{
-	Super::PostLoad();
-}
-
 bool UTransformableComponentHandle::IsValid() const
 {
 	return Component.IsValid();
 }
 
-void UTransformableComponentHandle::SetTransform(const FTransform& InGlobal) const
+void UTransformableComponentHandle::SetGlobalTransform(const FTransform& InGlobal) const
 {
 	if(Component.IsValid())
 	{
@@ -43,8 +38,21 @@
 	}
 }
 
-FTransform UTransformableComponentHandle::GetTransform() const
+void UTransformableComponentHandle::SetLocalTransform(const FTransform& InLocal) const
 {
+	if(Component.IsValid())
+	{
+		Component->SetRelativeTransform(InLocal);
+	}
+}
+
+FTransform UTransformableComponentHandle::GetLocalTransform() const
+{
+	return Component.IsValid() ? Component->GetRelativeTransform() : FTransform::Identity;
+}
+
+FTransform UTransformableComponentHandle::GetGlobalTransform() const
+{
 	return Component.IsValid() ? Component->GetComponentTransform() : FTransform::Identity;
 }
 
diff --git a/Engine/Source/Runtime/Experimental/Animation/Constraints/Public/ConstraintsManager.h b/Engine/Source/Runtime/Experimental/Animation/Constraints/Public/ConstraintsManager.h
--- a/Engine/Source/Runtime/Experimental/Animation/Constraints/Public/ConstraintsManager.h
+++ b/Engine/Source/Runtime/Experimental/Animation/Constraints/Public/ConstraintsManager.h
@@ -81,13 +81,24 @@
 	UPROPERTY()
 	FConstraintTickFunction ConstraintTick;
 
+	/** Evaluates the constraint in a context where it's mot done thru the ConstraintTick's tick function. */
+	void Evaluate() const;
+	
 	/** @todo document */
 	virtual uint32 GetTargetHash() const PURE_VIRTUAL(GetTargetHash, return 0;);
 	
 #if WITH_EDITOR
 	/** @todo document */
 	virtual FName GetLabel() const;
-#endif
+
+	// UObject interface
+	virtual void PostEditChangeProperty(FPropertyChangedEvent& PropertyChangedEvent) override;
+	// End of UObject interface
+	#endif
+
+	/** @todo documentation. */
+	UPROPERTY(EditAnywhere, DisplayName="Active State", Category="Constraint")
+	bool Active = true;
 };
 
 /** 
@@ -182,6 +193,12 @@
 	/** Remove the constraint at the given index. */
 	void RemoveConstraint(const int32 InConstraintIndex) const;
 
+	/** @todo document */
+	UTickableConstraint* GetConstraint(const FName& InConstraintName) const;
+
+	/** @todo document */
+	UTickableConstraint* GetConstraint(const int32 InConstraintIndex) const;
+	
 	/** Get read-only access to the array of constraints. */
 	const TArray< TObjectPtr<UTickableConstraint> >& GetConstrainsArray() const;
 	
diff --git a/Engine/Source/Runtime/Experimental/Animation/Constraints/Public/TransformConstraint.h b/Engine/Source/Runtime/Experimental/Animation/Constraints/Public/TransformConstraint.h
--- a/Engine/Source/Runtime/Experimental/Animation/Constraints/Public/TransformConstraint.h
+++ b/Engine/Source/Runtime/Experimental/Animation/Constraints/Public/TransformConstraint.h
@@ -41,15 +41,21 @@
 	TObjectPtr<UTransformableHandle> ChildTRSHandle;
 
 	/** @todo document */
-	UPROPERTY()
+	UPROPERTY(EditAnywhere, Category="Offset")
 	bool bMaintainOffset = true;
 
 	/** @todo document */
+	// UPROPERTY(EditAnywhere, Category="Weight", meta = (Input, ClampMin = "0.0", ClampMax = "1.0", UIMin = "0.0", UIMax = "1.0"))
 	UPROPERTY(meta = (Input, ClampMin = "0.0", ClampMax = "1.0", UIMin = "0.0", UIMax = "1.0"))
 	float Weight = 1.f;
 	
 	/** @todo document. */
 	int64 GetType() const;
+
+	/** @todo document */
+	FTransform GetChildGlobalTransform() const;
+	/** @todo document */
+	FTransform GetChildLocalTransform() const;
 		
 protected:
 	/** @todo document */
@@ -60,12 +66,13 @@
 	
 	/** @todo document */
 	void SetChildGlobalTransform(const FTransform& InGlobal) const;
+	/** @todo document */
+	void SetChildLocalTransform(const FTransform& InLocal) const;
 
 	/** @todo document */
-	FTransform GetChildGlobalTransform() const;
-
-	/** @todo document */
 	FTransform GetParentGlobalTransform() const;
+	/** @todo document */
+	FTransform GetParentLocalTransform() const;
 
 	/** @todo document */
 	UPROPERTY()
@@ -98,7 +105,7 @@
 	virtual void ComputeOffset() override;
 
 	/** @todo document */
-	UPROPERTY()
+	UPROPERTY(EditAnywhere, Category="Offset")
 	FVector OffsetTranslation = FVector::ZeroVector;
 };
 
@@ -122,7 +129,7 @@
 	virtual void ComputeOffset() override;
 	
 	/** @todo document */
-	UPROPERTY()
+	UPROPERTY(EditAnywhere, Category="Offset")
 	FQuat OffsetRotation = FQuat::Identity;
 };
 
@@ -146,7 +153,7 @@
 	virtual void ComputeOffset() override;
 
 	/** @todo document */
-	UPROPERTY()
+	UPROPERTY(EditAnywhere, Category="Offset")
 	FVector OffsetScale = FVector::OneVector;
 };
 
@@ -170,7 +177,7 @@
 	virtual void ComputeOffset() override;
 
 	/** @todo document */
-	UPROPERTY()
+	UPROPERTY(EditAnywhere, Category="Offset")
 	FTransform OffsetTransform = FTransform::Identity;
 };
 
@@ -194,7 +201,7 @@
 	virtual void ComputeOffset() override;
 
 	/** @todo document */
-	UPROPERTY()
+	UPROPERTY(EditAnywhere, Category="Offset")
 	FTransform OffsetTransform = FTransform::Identity;
 };
 
diff --git a/Engine/Source/Runtime/Experimental/Animation/Constraints/Public/TransformableHandle.h b/Engine/Source/Runtime/Experimental/Animation/Constraints/Public/TransformableHandle.h
--- a/Engine/Source/Runtime/Experimental/Animation/Constraints/Public/TransformableHandle.h
+++ b/Engine/Source/Runtime/Experimental/Animation/Constraints/Public/TransformableHandle.h
@@ -26,10 +26,14 @@
 	/** Sanity check to ensure the handle is safe to use. */
 	virtual bool IsValid() const PURE_VIRTUAL(IsValid, return false;);
 	
-	/** Sets the transform of the underlying transformable object. */
-	virtual void SetTransform(const FTransform& InGlobal) const PURE_VIRTUAL(SetTransform, );
+	/** Sets the global transform of the underlying transformable object. */
+	virtual void SetGlobalTransform(const FTransform& InGlobal) const PURE_VIRTUAL(SetGlobalTransform, );
+	/** Sets the local transform of the underlying transformable object in it's parent space. */
+	virtual void SetLocalTransform(const FTransform& InLocal) const PURE_VIRTUAL(SetLocalTransform, );
 	/** Gets the transform of the underlying transformable object. */
-	virtual FTransform GetTransform() const PURE_VIRTUAL(GetTransform, return FTransform::Identity;);
+	virtual FTransform GetGlobalTransform() const PURE_VIRTUAL(GetGlobalTransform, return FTransform::Identity;);
+	/** Gets the local transform of the underlying transformable object in it's parent space. */
+	virtual FTransform GetLocalTransform() const PURE_VIRTUAL(GetLocalTransform, return FTransform::Identity;);
 
 	/**
 	 * Returns the target object containing the tick function (returned in GetTickFunction).
@@ -67,13 +71,15 @@
 	
 	/** @todo document */
 	virtual bool IsValid() const override;
-
-	virtual void PostLoad() override;
 	
 	/** @todo document */
-	virtual void SetTransform(const FTransform& InGlobal) const override;
+	virtual void SetGlobalTransform(const FTransform& InGlobal) const override;
+	/** Sets the local transform of the underlying transformable object in it's parent space. */
+	virtual void SetLocalTransform(const FTransform& InLocal) const override;
 	/** @todo document */
-	virtual FTransform GetTransform() const override;
+	virtual FTransform GetGlobalTransform() const override;
+	/** @todo document */
+	virtual FTransform GetLocalTransform() const override;
 
 	/** @todo document */
 	virtual UObject* GetPrerequisiteObject() const override;
