Reset bake settings

diff --git a/Engine/Plugins/Animation/ControlRig/Source/ControlRigEditor/Private/BakeToControlRigSettings.cpp b/Engine/Plugins/Animation/ControlRig/Source/ControlRigEditor/Private/BakeToControlRigSettings.cpp
--- a/Engine/Plugins/Animation/ControlRig/Source/ControlRigEditor/Private/BakeToControlRigSettings.cpp
+++ b/Engine/Plugins/Animation/ControlRig/Source/ControlRigEditor/Private/BakeToControlRigSettings.cpp
@@ -1,10 +1,13 @@
-// Copyright Epic Games, Inc. All Rights Reserved.
-
-#include "BakeToControlRigSettings.h"
-
-UBakeToControlRigSettings::UBakeToControlRigSettings(const FObjectInitializer& Initializer)
-	: Super(Initializer)
-{
-	bReduceKeys = false;
-	Tolerance = 0.001f;
-}
+// Copyright Epic Games, Inc. All Rights Reserved.
+
+#include "BakeToControlRigSettings.h"
+
+UBakeToControlRigSettings::UBakeToControlRigSettings(const FObjectInitializer& Initializer)
+	: Super(Initializer)
+{}
+
+void UBakeToControlRigSettings::Reset()
+{
+	bReduceKeys = false;
+	Tolerance = 0.001f;
+}
diff --git a/Engine/Plugins/Animation/ControlRig/Source/ControlRigEditor/Private/BakeToControlRigSettings.h b/Engine/Plugins/Animation/ControlRig/Source/ControlRigEditor/Private/BakeToControlRigSettings.h
--- a/Engine/Plugins/Animation/ControlRig/Source/ControlRigEditor/Private/BakeToControlRigSettings.h
+++ b/Engine/Plugins/Animation/ControlRig/Source/ControlRigEditor/Private/BakeToControlRigSettings.h
@@ -1,25 +1,28 @@
-// Copyright Epic Games, Inc. All Rights Reserved.
-#pragma once
-
-#include "CoreMinimal.h"
-#include "UObject/ObjectMacros.h"
-#include "UObject/Object.h"
-#include "BakeToControlRigSettings.generated.h"
-
-UCLASS(BlueprintType, config = EditorSettings)
-class  UBakeToControlRigSettings : public UObject
-{
-public:
-	UBakeToControlRigSettings(const FObjectInitializer& Initializer);
-
-	GENERATED_BODY()
-	/** Reduce Keys */
-	UPROPERTY(EditAnywhere, Category = "Reduce Keys")
-	bool bReduceKeys;
-
-	/** Reduce Keys Tolerance*/
-	UPROPERTY(EditAnywhere, Category = "Reduce Keys")
-	float Tolerance;
-};
-
-
+// Copyright Epic Games, Inc. All Rights Reserved.
+#pragma once
+
+#include "CoreMinimal.h"
+#include "UObject/ObjectMacros.h"
+#include "UObject/Object.h"
+#include "BakeToControlRigSettings.generated.h"
+
+UCLASS(BlueprintType, config = EditorSettings)
+class  UBakeToControlRigSettings : public UObject
+{
+public:
+	UBakeToControlRigSettings(const FObjectInitializer& Initializer);
+
+	GENERATED_BODY()
+	/** Reduce Keys */
+	UPROPERTY(EditAnywhere, Category = "Reduce Keys")
+	bool bReduceKeys = false;
+
+	/** Reduce Keys Tolerance*/
+	UPROPERTY(EditAnywhere, Category = "Reduce Keys")
+	float Tolerance = 0.001f;
+
+	/** Resets the default properties. */
+	void Reset();
+};
+
+
diff --git a/Engine/Plugins/Animation/ControlRig/Source/ControlRigEditor/Private/Sequencer/ControlRigParameterTrackEditor.cpp b/Engine/Plugins/Animation/ControlRig/Source/ControlRigEditor/Private/Sequencer/ControlRigParameterTrackEditor.cpp
--- a/Engine/Plugins/Animation/ControlRig/Source/ControlRigEditor/Private/Sequencer/ControlRigParameterTrackEditor.cpp
+++ b/Engine/Plugins/Animation/ControlRig/Source/ControlRigEditor/Private/Sequencer/ControlRigParameterTrackEditor.cpp
@@ -926,9 +926,19 @@ void SBakeToAnimAndControlRigOptionsWindow::Construct(const FArguments& InArgs)
 
 FReply SBakeToAnimAndControlRigOptionsWindow::OnResetToDefaultClick() const
 {
-	ExportOptions->ResetToDefault();
-	//Refresh the view to make sure the custom UI are updating correctly
-	DetailsView->SetObject(ExportOptions, true);
+	if (ExportOptions)
+	{
+		ExportOptions->ResetToDefault();
+		//Refresh the view to make sure the custom UI are updating correctly
+		DetailsView->SetObject(ExportOptions, true);
+	}
+
+	if (BakeSettings)
+	{
+		BakeSettings->Reset();
+		DetailsView2->SetObject(BakeSettings, true);
+	}
+	
 	return FReply::Handled();
 }
 
