Sequencer: Follow up to expose bind filter commands to curve editor bindings

diff --git a/Engine/Plugins/Animation/ControlRig/Source/ControlRigEditor/Private/Sequencer/ControlRigSequencerFilter.cpp b/Engine/Plugins/Animation/ControlRig/Source/ControlRigEditor/Private/Sequencer/ControlRigSequencerFilter.cpp
--- a/Engine/Plugins/Animation/ControlRig/Source/ControlRigEditor/Private/Sequencer/ControlRigSequencerFilter.cpp
+++ b/Engine/Plugins/Animation/ControlRig/Source/ControlRigEditor/Private/Sequencer/ControlRigSequencerFilter.cpp
@@ -87,7 +87,7 @@
 		return Tooltip;
 	}
 
-	virtual void BindCommands(TSharedRef<FUICommandList> CommandBindings, TWeakPtr<ISequencer> Sequencer) override
+	virtual void BindCommands(TSharedRef<FUICommandList> SequencerBindings, TSharedRef<FUICommandList> CurveEditorBindings, TWeakPtr<ISequencer> Sequencer) override
 	{
 		// See comment above
 		if (!FSequencerTrackFilter_ControlRigControlsCommands::IsRegistered())
@@ -97,11 +97,13 @@
 
 		const FSequencerTrackFilter_ControlRigControlsCommands& Commands = FSequencerTrackFilter_ControlRigControlsCommands::Get();
 
-		CommandBindings->MapAction(
+		SequencerBindings->MapAction(
 			Commands.ToggleControlRigControls,
 			FExecuteAction::CreateLambda([this, Sequencer] { Sequencer.Pin()->SetTrackFilterEnabled(GetDisplayName(), !Sequencer.Pin()->IsTrackFilterEnabled(GetDisplayName())); }),
 			FCanExecuteAction::CreateLambda([this, Sequencer] { return true; }),
 			FIsActionChecked::CreateLambda([this, Sequencer] { return Sequencer.Pin()->IsTrackFilterEnabled(GetDisplayName()); }));
+
+		CurveEditorBindings->MapAction(Commands.ToggleControlRigControls, *SequencerBindings->GetActionForCommand(Commands.ToggleControlRigControls));
 	}
 };
 
@@ -199,7 +201,7 @@
 		return Tooltip;
 	}
 
-	virtual void BindCommands(TSharedRef<FUICommandList> CommandBindings, TWeakPtr<ISequencer> Sequencer) override
+	virtual void BindCommands(TSharedRef<FUICommandList> SequencerBindings, TSharedRef<FUICommandList> CurveEditorBindings, TWeakPtr<ISequencer> Sequencer) override
 	{
 		// See comment above
 		if (!FSequencerTrackFilter_ControlRigSelectedControlsCommands::IsRegistered())
@@ -209,11 +211,13 @@
 
 		const FSequencerTrackFilter_ControlRigSelectedControlsCommands& Commands = FSequencerTrackFilter_ControlRigSelectedControlsCommands::Get();
 
-		CommandBindings->MapAction(
+		SequencerBindings->MapAction(
 			Commands.ToggleControlRigSelectedControls,
 			FExecuteAction::CreateLambda([this, Sequencer] { Sequencer.Pin()->SetTrackFilterEnabled(GetDisplayName(), !Sequencer.Pin()->IsTrackFilterEnabled(GetDisplayName())); }),
 			FCanExecuteAction::CreateLambda([this, Sequencer] { return true; }),
 			FIsActionChecked::CreateLambda([this, Sequencer] { return Sequencer.Pin()->IsTrackFilterEnabled(GetDisplayName()); }));
+
+		CurveEditorBindings->MapAction(Commands.ToggleControlRigSelectedControls, *SequencerBindings->GetActionForCommand(Commands.ToggleControlRigSelectedControls));
 	}
 };
 
diff --git a/Engine/Source/Editor/Sequencer/Private/SSequencer.cpp b/Engine/Source/Editor/Sequencer/Private/SSequencer.cpp
--- a/Engine/Source/Editor/Sequencer/Private/SSequencer.cpp
+++ b/Engine/Source/Editor/Sequencer/Private/SSequencer.cpp
@@ -1106,7 +1106,7 @@ void SSequencer::Construct(const FArguments& InArgs, TSharedRef<FSequencer> InSe
 }
 PRAGMA_ENABLE_OPTIMIZATION
 
-void SSequencer::BindCommands(TSharedRef<FUICommandList> SequencerCommandBindings)
+void SSequencer::BindCommands(TSharedRef<FUICommandList> SequencerCommandBindings, TSharedRef<FUICommandList> CurveEditorSharedBindings)
 {
 	auto CanPasteFromHistory = [this]{
 		if (!HasFocusedDescendants() && !HasKeyboardFocus())
@@ -1201,7 +1201,7 @@ void SSequencer::BindCommands(TSharedRef<FUICommandList> SequencerCommandBinding
 	{
 		if (TrackFilter->SupportsSequence(Sequencer->GetFocusedMovieSceneSequence()))
 		{
-			TrackFilter->BindCommands(SequencerCommandBindings, SequencerPtr.Pin());
+			TrackFilter->BindCommands(SequencerCommandBindings, CurveEditorSharedBindings, SequencerPtr.Pin());
 		}
 	}
 }
@@ -1748,7 +1748,7 @@ TSharedRef<SWidget> SSequencer::MakeAddMenu()
 TSharedRef<SWidget> SSequencer::MakeFilterMenu()
 {
 	TSharedPtr<FExtender> Extender = FExtender::Combine(AddMenuExtenders);
-	FMenuBuilder MenuBuilder(false, nullptr, Extender);
+	FMenuBuilder MenuBuilder(false, SequencerPtr.Pin()->GetCommandBindings(), Extender);
 
 	// let track editors & object bindings populate the menu
 	TSharedPtr<FSequencer> Sequencer = SequencerPtr.Pin();
diff --git a/Engine/Source/Editor/Sequencer/Private/SSequencer.h b/Engine/Source/Editor/Sequencer/Private/SSequencer.h
--- a/Engine/Source/Editor/Sequencer/Private/SSequencer.h
+++ b/Engine/Source/Editor/Sequencer/Private/SSequencer.h
@@ -288,7 +288,7 @@ public:
 
 	void Construct(const FArguments& InArgs, TSharedRef<FSequencer> InSequencer);
 
-	void BindCommands(TSharedRef<FUICommandList> SequencerCommandBindings);
+	void BindCommands(TSharedRef<FUICommandList> SequencerCommandBindings, TSharedRef<FUICommandList> CurveEditorSharedBindings);
 	
 	~SSequencer();
 	
diff --git a/Engine/Source/Editor/Sequencer/Private/Sequencer.cpp b/Engine/Source/Editor/Sequencer/Private/Sequencer.cpp
--- a/Engine/Source/Editor/Sequencer/Private/Sequencer.cpp
+++ b/Engine/Source/Editor/Sequencer/Private/Sequencer.cpp
@@ -13378,10 +13378,12 @@ void FSequencer::BindCommands()
 	CurveEditorSharedBindings->MapAction(Commands.AddRotationKey, *SequencerCommandBindings->GetActionForCommand(Commands.AddRotationKey));
 	CurveEditorSharedBindings->MapAction(Commands.AddScaleKey, *SequencerCommandBindings->GetActionForCommand(Commands.AddScaleKey));
 
+	CurveEditorSharedBindings->MapAction(Commands.ResetFilters, *SequencerCommandBindings->GetActionForCommand(Commands.ResetFilters));
+
 	GetCurveEditor()->GetCommands()->Append(CurveEditorSharedBindings);
 
 	// bind widget specific commands
-	SequencerWidget->BindCommands(SequencerCommandBindings);
+	SequencerWidget->BindCommands(SequencerCommandBindings, CurveEditorSharedBindings);
 }
 
 void FSequencer::BuildAddTrackMenu(class FMenuBuilder& MenuBuilder)
diff --git a/Engine/Source/Editor/Sequencer/Private/SequencerTrackFilters.cpp b/Engine/Source/Editor/Sequencer/Private/SequencerTrackFilters.cpp
--- a/Engine/Source/Editor/Sequencer/Private/SequencerTrackFilters.cpp
+++ b/Engine/Source/Editor/Sequencer/Private/SequencerTrackFilters.cpp
@@ -180,7 +180,7 @@
 	return Tooltip;
 }
 
-void FSequencerTrackFilter_Animated::BindCommands(TSharedRef<FUICommandList> CommandBindings, TWeakPtr<ISequencer> Sequencer)
+void FSequencerTrackFilter_Animated::BindCommands(TSharedRef<FUICommandList> SequencerBindings, TSharedRef<FUICommandList> CurveEditorBindings, TWeakPtr<ISequencer> Sequencer)
 {
 	// See comment above
 	if (!FSequencerTrackFilter_AnimatedCommands::IsRegistered())
@@ -190,11 +190,13 @@
 
 	const FSequencerTrackFilter_AnimatedCommands& Commands = FSequencerTrackFilter_AnimatedCommands::Get();
 
-	CommandBindings->MapAction(
+	SequencerBindings->MapAction(
 		Commands.ToggleAnimatedTracks,
 		FExecuteAction::CreateLambda( [this, Sequencer]{ Sequencer.Pin()->SetTrackFilterEnabled(GetDisplayName(), !Sequencer.Pin()->IsTrackFilterEnabled(GetDisplayName())); } ),
 		FCanExecuteAction::CreateLambda( [this, Sequencer]{ return true; } ),
 		FIsActionChecked::CreateLambda( [this, Sequencer]{ return Sequencer.Pin()->IsTrackFilterEnabled(GetDisplayName()); } ) );
+
+	CurveEditorBindings->MapAction(Commands.ToggleAnimatedTracks, *SequencerBindings->GetActionForCommand(Commands.ToggleAnimatedTracks));
 }
 
 #undef LOCTEXT_NAMESPACE
diff --git a/Engine/Source/Editor/Sequencer/Private/SequencerTrackFilters.h b/Engine/Source/Editor/Sequencer/Private/SequencerTrackFilters.h
--- a/Engine/Source/Editor/Sequencer/Private/SequencerTrackFilters.h
+++ b/Engine/Source/Editor/Sequencer/Private/SequencerTrackFilters.h
@@ -362,6 +362,6 @@ public:
 			(WidgetAnimationClass != nullptr && InSequence->GetClass()->IsChildOf(WidgetAnimationClass)));
 	}
 	
-	virtual void BindCommands(TSharedRef<FUICommandList> CommandBindings, TWeakPtr<ISequencer> Sequencer) override;
+	virtual void BindCommands(TSharedRef<FUICommandList> SequencerBindings, TSharedRef<FUICommandList> CurveEditorBindings, TWeakPtr<ISequencer> Sequencer) override;
 };
 #undef LOCTEXT_NAMESPACE
\ No newline at end of file
diff --git a/Engine/Source/Editor/Sequencer/Public/SequencerTrackFilterBase.h b/Engine/Source/Editor/Sequencer/Public/SequencerTrackFilterBase.h
--- a/Engine/Source/Editor/Sequencer/Public/SequencerTrackFilterBase.h
+++ b/Engine/Source/Editor/Sequencer/Public/SequencerTrackFilterBase.h
@@ -53,7 +53,7 @@
 
 	virtual bool PassesFilterChannel(FMovieSceneChannel* InMovieSceneChannel) const { return false; }
 
-	virtual void BindCommands(TSharedRef<FUICommandList> CommandBindings, TWeakPtr<ISequencer> Sequencer) {}
+	virtual void BindCommands(TSharedRef<FUICommandList> SequencerBindings, TSharedRef<FUICommandList> CurveEditorBindings, TWeakPtr<ISequencer> Sequencer) {}
 
 protected:
 	void BroadcastChangedEvent() const { ChangedEvent.Broadcast(); }
