ControlRig Component: Use tick prerequistie to ensure mapped mesh components tick after the control rig component to avoid double ticking the SKM componet in the CR component. It was causing motion vectors to be zeroed and rendering artifacts with TSR/MotionBlur

diff --git a/Engine/Plugins/Animation/ControlRig/Source/ControlRig/Private/ControlRigComponent.cpp b/Engine/Plugins/Animation/ControlRig/Source/ControlRig/Private/ControlRigComponent.cpp
--- a/Engine/Plugins/Animation/ControlRig/Source/ControlRig/Private/ControlRigComponent.cpp
+++ b/Engine/Plugins/Animation/ControlRig/Source/ControlRig/Private/ControlRigComponent.cpp
@@ -338,6 +338,20 @@
 			CR->RequestInit();
 		}
 	}
+
+	// we want to make sure all components driven by control rig component tick
+	// after the control rig component such that by the time they tick they can
+	// send the latest data for rendering
+	USceneComponent* LastComponent = nullptr;
+	for (FControlRigComponentMappedElement& MappedElement : MappedElements)
+	{
+		if (LastComponent != MappedElement.SceneComponent && MappedElement.SceneComponent != nullptr)
+		{
+			
+			MappedElement.SceneComponent->AddTickPrerequisiteComponent(this);
+			LastComponent = MappedElement.SceneComponent;
+		}
+	}
 }
 
 void UControlRigComponent::Update(float DeltaTime)
@@ -1513,26 +1527,6 @@
 			}
 		}
 
-		for (USkeletalMeshComponent* SkeletalMeshComponent : ComponentsToTick)
-		{
-			if (SkeletalMeshComponent)
-			{
-				if (SkeletalMeshComponent->IsValidLowLevel() &&
-					!SkeletalMeshComponent->HasAnyFlags(RF_BeginDestroyed) &&
-					IsValid(SkeletalMeshComponent))
-				{
-					SkeletalMeshComponent->TickAnimation(0.f, false);
-					SkeletalMeshComponent->UpdateLODStatus();
-					SkeletalMeshComponent->RefreshBoneTransforms();
-					SkeletalMeshComponent->RefreshSlaveComponents();
-					SkeletalMeshComponent->UpdateComponentToWorld();
-					SkeletalMeshComponent->FinalizeBoneTransform();
-					SkeletalMeshComponent->MarkRenderTransformDirty();
-					SkeletalMeshComponent->MarkRenderDynamicDataDirty();
-				}
-			}
-		}
-
 #if WITH_EDITOR
 		if(URigHierarchy* Hierarchy = ControlRig->GetHierarchy())
 		{
