Sequencer: Follow up to filter commands - Add binding counts to track filter commands since they could be destroyed after a new track filter is generated

diff --git a/Engine/Plugins/Animation/ControlRig/Source/ControlRigEditor/Private/Sequencer/ControlRigSequencerFilter.cpp b/Engine/Plugins/Animation/ControlRig/Source/ControlRigEditor/Private/Sequencer/ControlRigSequencerFilter.cpp
--- a/Engine/Plugins/Animation/ControlRig/Source/ControlRigEditor/Private/Sequencer/ControlRigSequencerFilter.cpp
+++ b/Engine/Plugins/Animation/ControlRig/Source/ControlRigEditor/Private/Sequencer/ControlRigSequencerFilter.cpp
@@ -43,13 +43,19 @@
 public:
 
 	FSequencerTrackFilter_ControlRigControls()
+		: BindingCount(0)
 	{
 		FSequencerTrackFilter_ControlRigControlsCommands::Register();
 	}
 
 	~FSequencerTrackFilter_ControlRigControls()
 	{
-		FSequencerTrackFilter_ControlRigControlsCommands::Unregister();
+		BindingCount--;
+
+		if (BindingCount < 1)
+		{
+			FSequencerTrackFilter_ControlRigControlsCommands::Unregister();
+		}
 	}
 
 	virtual FString GetName() const override { return TEXT("ControlRigControlsFilter"); }
@@ -69,11 +75,6 @@
 
 	virtual FText GetToolTipText() const override
 	{
-		if (!FSequencerTrackFilter_ControlRigControlsCommands::IsRegistered())
-		{
-			FSequencerTrackFilter_ControlRigControlsCommands::Register();
-		}
-
 		const FSequencerTrackFilter_ControlRigControlsCommands& Commands = FSequencerTrackFilter_ControlRigControlsCommands::Get();
 
 		const TSharedRef<const FInputChord> FirstActiveChord = Commands.ToggleControlRigControls->GetFirstValidChord();
@@ -89,12 +90,6 @@
 
 	virtual void BindCommands(TSharedRef<FUICommandList> SequencerBindings, TSharedRef<FUICommandList> CurveEditorBindings, TWeakPtr<ISequencer> Sequencer) override
 	{
-		// See comment above
-		if (!FSequencerTrackFilter_ControlRigControlsCommands::IsRegistered())
-		{
-			FSequencerTrackFilter_ControlRigControlsCommands::Register();
-		}
-
 		const FSequencerTrackFilter_ControlRigControlsCommands& Commands = FSequencerTrackFilter_ControlRigControlsCommands::Get();
 
 		SequencerBindings->MapAction(
@@ -105,6 +100,9 @@
 
 		CurveEditorBindings->MapAction(Commands.ToggleControlRigControls, *SequencerBindings->GetActionForCommand(Commands.ToggleControlRigControls));
 	}
+
+private:
+	mutable uint32 BindingCount;
 };
 
 //////////////////////////////////////////////////////////////////////////
@@ -140,13 +138,19 @@
 public:
 
 	FSequencerTrackFilter_ControlRigSelectedControls()
+		: BindingCount(0)
 	{
 		FSequencerTrackFilter_ControlRigSelectedControlsCommands::Register();
 	}
 
 	~FSequencerTrackFilter_ControlRigSelectedControls()
 	{
-		FSequencerTrackFilter_ControlRigSelectedControlsCommands::Unregister();
+		BindingCount--;
+
+		if (BindingCount < 1)
+		{
+			FSequencerTrackFilter_ControlRigSelectedControlsCommands::Unregister();
+		}
 	}
 
 	virtual FString GetName() const override { return TEXT("ControlRigControlsSelectedFilter"); }
@@ -183,11 +187,6 @@
 
 	virtual FText GetToolTipText() const override
 	{
-		if (!FSequencerTrackFilter_ControlRigSelectedControlsCommands::IsRegistered())
-		{
-			FSequencerTrackFilter_ControlRigSelectedControlsCommands::Register();
-		}
-
 		const FSequencerTrackFilter_ControlRigSelectedControlsCommands& Commands = FSequencerTrackFilter_ControlRigSelectedControlsCommands::Get();
 
 		const TSharedRef<const FInputChord> FirstActiveChord = Commands.ToggleControlRigSelectedControls->GetFirstValidChord();
@@ -203,12 +202,6 @@
 
 	virtual void BindCommands(TSharedRef<FUICommandList> SequencerBindings, TSharedRef<FUICommandList> CurveEditorBindings, TWeakPtr<ISequencer> Sequencer) override
 	{
-		// See comment above
-		if (!FSequencerTrackFilter_ControlRigSelectedControlsCommands::IsRegistered())
-		{
-			FSequencerTrackFilter_ControlRigSelectedControlsCommands::Register();
-		}
-
 		const FSequencerTrackFilter_ControlRigSelectedControlsCommands& Commands = FSequencerTrackFilter_ControlRigSelectedControlsCommands::Get();
 
 		SequencerBindings->MapAction(
@@ -219,6 +212,9 @@
 
 		CurveEditorBindings->MapAction(Commands.ToggleControlRigSelectedControls, *SequencerBindings->GetActionForCommand(Commands.ToggleControlRigSelectedControls));
 	}
+
+private:
+	mutable uint32 BindingCount;
 };
 
 /*
diff --git a/Engine/Source/Editor/Sequencer/Private/SequencerTrackFilters.cpp b/Engine/Source/Editor/Sequencer/Private/SequencerTrackFilters.cpp
--- a/Engine/Source/Editor/Sequencer/Private/SequencerTrackFilters.cpp
+++ b/Engine/Source/Editor/Sequencer/Private/SequencerTrackFilters.cpp
@@ -148,25 +148,23 @@
 };
 
 FSequencerTrackFilter_Animated::FSequencerTrackFilter_Animated()
+	: BindingCount(0)
 {
 	FSequencerTrackFilter_AnimatedCommands::Register();	
 }
 
 FSequencerTrackFilter_Animated::~FSequencerTrackFilter_Animated()
 {
-	FSequencerTrackFilter_AnimatedCommands::Unregister();	
+	BindingCount--;
+
+	if (BindingCount < 1)
+	{
+		FSequencerTrackFilter_AnimatedCommands::Unregister();
+	}
 }
 
 FText FSequencerTrackFilter_Animated::GetToolTipText() const 
 { 
-	// When opening another sequence, FSequencer initializes the first sequence and then closes the previous sequence. 
-	// This causes the track filter commands to be initialized for the first sequence and then destroyed when the subsequence 
-	// sequence is opened. For now, register the commands before calling Get()
-	if (!FSequencerTrackFilter_AnimatedCommands::IsRegistered())
-	{
-		FSequencerTrackFilter_AnimatedCommands::Register();	
-	}
-
 	const FSequencerTrackFilter_AnimatedCommands& Commands = FSequencerTrackFilter_AnimatedCommands::Get();
 
 	const TSharedRef<const FInputChord> FirstActiveChord = Commands.ToggleAnimatedTracks->GetFirstValidChord();
@@ -182,12 +180,6 @@
 
 void FSequencerTrackFilter_Animated::BindCommands(TSharedRef<FUICommandList> SequencerBindings, TSharedRef<FUICommandList> CurveEditorBindings, TWeakPtr<ISequencer> Sequencer)
 {
-	// See comment above
-	if (!FSequencerTrackFilter_AnimatedCommands::IsRegistered())
-	{
-		FSequencerTrackFilter_AnimatedCommands::Register();	
-	}
-
 	const FSequencerTrackFilter_AnimatedCommands& Commands = FSequencerTrackFilter_AnimatedCommands::Get();
 
 	SequencerBindings->MapAction(
diff --git a/Engine/Source/Editor/Sequencer/Private/SequencerTrackFilters.h b/Engine/Source/Editor/Sequencer/Private/SequencerTrackFilters.h
--- a/Engine/Source/Editor/Sequencer/Private/SequencerTrackFilters.h
+++ b/Engine/Source/Editor/Sequencer/Private/SequencerTrackFilters.h
@@ -363,5 +363,8 @@
 	}
 	
 	virtual void BindCommands(TSharedRef<FUICommandList> SequencerBindings, TSharedRef<FUICommandList> CurveEditorBindings, TWeakPtr<ISequencer> Sequencer) override;
+
+private:
+	mutable uint32 BindingCount;
 };
 #undef LOCTEXT_NAMESPACE
