Sequencer: Anim Tools: Most new anim tools, pivot, motion trails, etc.. don't work in  sub sequences, times could be wrong and even worse bindings not found since not using the root to local transform to calculate time correctly.

diff --git a/Engine/Plugins/Animation/ControlRig/Source/ControlRigEditor/Private/ControlRigSequencerEditorLibrary.cpp b/Engine/Plugins/Animation/ControlRig/Source/ControlRigEditor/Private/ControlRigSequencerEditorLibrary.cpp
--- a/Engine/Plugins/Animation/ControlRig/Source/ControlRigEditor/Private/ControlRigSequencerEditorLibrary.cpp
+++ b/Engine/Plugins/Animation/ControlRig/Source/ControlRigEditor/Private/ControlRigSequencerEditorLibrary.cpp
@@ -866,7 +866,7 @@ static bool LocalGetControlRigControlValues(IMovieScenePlayer* Player, UMovieSce
 			{
 				GlobalTime = FFrameRate::TransformTime(GlobalTime, MovieScene->GetDisplayRate(), MovieScene->GetTickResolution());
 			}
-
+			GlobalTime = GlobalTime * RootToLocalTransform.InverseLinearOnly();
 			FMovieSceneContext Context = FMovieSceneContext(FMovieSceneEvaluationRange(GlobalTime, TickResolution), Player->GetPlaybackStatus()).SetHasJumped(true);
 
 			Player->GetEvaluationTemplate().Evaluate(Context, *Player);
@@ -883,7 +883,7 @@ static bool GetControlRigValues(ISequencer* Sequencer, UControlRig* ControlRig,
 	if (Sequencer->GetFocusedMovieSceneSequence())
 	{
 		FMovieSceneSequenceIDRef Template = Sequencer->GetFocusedTemplateID();
-		FMovieSceneSequenceTransform RootToLocalTransform;
+		FMovieSceneSequenceTransform RootToLocalTransform = Sequencer->GetFocusedMovieSceneSequenceTransform();
 		bool bGood = LocalGetControlRigControlValues(Sequencer, Sequencer->GetFocusedMovieSceneSequence(), Template, RootToLocalTransform,
 			ControlRig, ControlName,TimeUnit, Frames, OutValues);
 		Sequencer->RequestEvaluate();
@@ -901,7 +901,7 @@ static bool GetControlRigValue(ISequencer* Sequencer, UControlRig* ControlRig, c
 		Frames.Add(Frame);
 		TArray<FRigControlValue> OutValues;
 		FMovieSceneSequenceIDRef Template = Sequencer->GetFocusedTemplateID();
-		FMovieSceneSequenceTransform RootToLocalTransform;
+		FMovieSceneSequenceTransform RootToLocalTransform = Sequencer->GetFocusedMovieSceneSequenceTransform();
 		bool bVal = LocalGetControlRigControlValues(Sequencer, Sequencer->GetFocusedMovieSceneSequence(), Template, RootToLocalTransform,
 			ControlRig, ControlName,TimeUnit, Frames, OutValues);
 		if (bVal)
diff --git a/Engine/Plugins/Animation/ControlRig/Source/ControlRigEditor/Private/Sequencer/ControlRigParameterTrackEditor.cpp b/Engine/Plugins/Animation/ControlRig/Source/ControlRigEditor/Private/Sequencer/ControlRigParameterTrackEditor.cpp
--- a/Engine/Plugins/Animation/ControlRig/Source/ControlRigEditor/Private/Sequencer/ControlRigParameterTrackEditor.cpp
+++ b/Engine/Plugins/Animation/ControlRig/Source/ControlRigEditor/Private/Sequencer/ControlRigParameterTrackEditor.cpp
@@ -3452,6 +3452,7 @@
 		FFrameNumber EndFrame = Range.GetUpperBoundValue();
 		const FFrameRate& FrameRate = SequencerPtr->GetFocusedDisplayRate();
 		const FFrameRate& TickResolution = SequencerPtr->GetFocusedTickResolution();
+		FMovieSceneSequenceTransform RootToLocalTransform = SequencerPtr->GetFocusedMovieSceneSequenceTransform();
 
 		FFrameNumber FrameRateInFrameNumber = TickResolution.AsFrameNumber(FrameRate.AsInterval());
 		TArray<FFrameNumber> Frames;
@@ -3482,6 +3483,7 @@
 		{
 			const FFrameNumber& FrameNumber = Frames[Index];
 			FFrameTime GlobalTime(FrameNumber);
+			GlobalTime = GlobalTime * RootToLocalTransform.InverseLinearOnly();
 
 			FMovieSceneContext Context = FMovieSceneContext(FMovieSceneEvaluationRange(GlobalTime, TickResolution), SequencerPtr->GetPlaybackStatus()).SetHasJumped(true);
 
diff --git a/Engine/Plugins/Animation/ControlRig/Source/ControlRigEditor/Private/Sequencer/ControlRigSpaceChannelEditors.cpp b/Engine/Plugins/Animation/ControlRig/Source/ControlRigEditor/Private/Sequencer/ControlRigSpaceChannelEditors.cpp
--- a/Engine/Plugins/Animation/ControlRig/Source/ControlRigEditor/Private/Sequencer/ControlRigSpaceChannelEditors.cpp
+++ b/Engine/Plugins/Animation/ControlRig/Source/ControlRigEditor/Private/Sequencer/ControlRigSpaceChannelEditors.cpp
@@ -426,12 +426,14 @@
 		FRigControlModifiedContext Context;
 		Context.SetKey = EControlRigSetKey::Always;
 		FFrameRate TickResolution = Sequencer->GetFocusedTickResolution();
+		FMovieSceneSequenceTransform RootToLocalTransform = Sequencer->GetFocusedMovieSceneSequenceTransform();
 
 		// set previous key if we're going to switch space
 		if (bSetPreviousKey)
 		{
-			const FFrameTime GlobalTime(Time -1);
-
+			FFrameTime GlobalTime(Time -1);
+			GlobalTime = GlobalTime * RootToLocalTransform.InverseLinearOnly();
+			
 			FMovieSceneContext SceneContext = FMovieSceneContext(FMovieSceneEvaluationRange(GlobalTime, TickResolution), Sequencer->GetPlaybackStatus()).SetHasJumped(true);
 			Sequencer->GetEvaluationTemplate().Evaluate(SceneContext, *Sequencer);
 
@@ -454,8 +456,9 @@
 		int32 FramesIndex = 0;
 		for (const FFrameNumber& Frame : Frames)
 		{
-			const FFrameTime GlobalTime(Frame);
-			
+			FFrameTime GlobalTime(Frame);
+			GlobalTime = GlobalTime * RootToLocalTransform.InverseLinearOnly();
+
 			FMovieSceneContext SceneContext = FMovieSceneContext(FMovieSceneEvaluationRange(GlobalTime, TickResolution), Sequencer->GetPlaybackStatus()).SetHasJumped(true);
 			Sequencer->GetEvaluationTemplate().Evaluate(SceneContext, *Sequencer);
 
@@ -652,10 +655,14 @@
 		FRigControlModifiedContext Context;
 		Context.SetKey = EControlRigSetKey::Always;
 		FFrameRate TickResolution = Sequencer->GetFocusedTickResolution();
+		FMovieSceneSequenceTransform RootToLocalTransform = Sequencer->GetFocusedMovieSceneSequenceTransform();
+
 		for (const FFrameNumber& Frame : Frames)
 		{
 			//evaluate sequencer
-			const FFrameTime GlobalTime(Frame);
+			FFrameTime GlobalTime(Frame);
+			GlobalTime = GlobalTime * RootToLocalTransform.InverseLinearOnly();
+
 			FMovieSceneContext SceneContext = FMovieSceneContext(FMovieSceneEvaluationRange(GlobalTime, TickResolution), Sequencer->GetPlaybackStatus()).SetHasJumped(true);
 			Sequencer->GetEvaluationTemplate().Evaluate(SceneContext, *Sequencer);
 			//make sure to set rig hierarchy correct since key is not deleted yet
@@ -918,13 +925,17 @@
 			RigHierarchy->SwitchToParent(ControlKey, Settings.TargetSpace, false, true, Dependencies, nullptr);
 			ControlRig->Evaluate_AnyThread();
 
+			FMovieSceneSequenceTransform RootToLocalTransform = Sequencer->GetFocusedMovieSceneSequenceTransform();
+
 			for (int32 Index = 0; Index < Frames.Num(); ++Index)
 			{
 				const FTransform GlobalTransform = ControlWorldTransforms[Index];
 				const FFrameNumber Frame = Frames[Index];
 
 				//evaluate sequencer
-				const FFrameTime GlobalTime(Frame);
+				FFrameTime GlobalTime(Frame);
+				GlobalTime = GlobalTime * RootToLocalTransform.InverseLinearOnly();
+
 				FMovieSceneContext SceneContext = FMovieSceneContext(FMovieSceneEvaluationRange(GlobalTime, TickResolution), Sequencer->GetPlaybackStatus()).SetHasJumped(true);
 				Sequencer->GetEvaluationTemplate().Evaluate(SceneContext, *Sequencer);
 
@@ -952,7 +963,9 @@
 				}
 
 				//evaluate sequencer
-				const FFrameTime GlobalTime(EndFrame);
+				FFrameTime GlobalTime(EndFrame);
+				GlobalTime = GlobalTime * RootToLocalTransform.InverseLinearOnly();
+
 				FMovieSceneContext SceneContext = FMovieSceneContext(FMovieSceneEvaluationRange(GlobalTime, TickResolution), Sequencer->GetPlaybackStatus()).SetHasJumped(true);
 				Sequencer->GetEvaluationTemplate().Evaluate(SceneContext, *Sequencer);
 		
diff --git a/Engine/Plugins/Animation/ControlRig/Source/ControlRigEditor/Private/Tools/ControlRigSnapper.cpp b/Engine/Plugins/Animation/ControlRig/Source/ControlRigEditor/Private/Tools/ControlRigSnapper.cpp
--- a/Engine/Plugins/Animation/ControlRig/Source/ControlRigEditor/Private/Tools/ControlRigSnapper.cpp
+++ b/Engine/Plugins/Animation/ControlRig/Source/ControlRigEditor/Private/Tools/ControlRigSnapper.cpp
@@ -131,6 +131,8 @@
 			if (CurrentFrame.IsSet() == false || CurrentFrame.GetValue() != FrameNumber)
 			{
 				FFrameTime GlobalTime(FrameNumber);
+				GlobalTime = GlobalTime * RootToLocalTransform.InverseLinearOnly(); //player evals in root time so need to go back to it.
+
 				FMovieSceneContext Context = FMovieSceneContext(FMovieSceneEvaluationRange(GlobalTime, TickResolution), Player->GetPlaybackStatus()).SetHasJumped(true);
 				Player->GetEvaluationTemplate().Evaluate(Context, *Player);
 			}
@@ -147,8 +149,8 @@
 	if (Sequencer->GetFocusedMovieSceneSequence())
 	{
 		FMovieSceneSequenceIDRef Template = Sequencer->GetFocusedTemplateID();
-		FMovieSceneSequenceTransform RootToLocalTransform;
-		TOptional<FFrameNumber> FrameNumber = Sequencer->GetLocalTime().Time.RoundToFrame();
+		FMovieSceneSequenceTransform RootToLocalTransform = Sequencer->GetFocusedMovieSceneSequenceTransform();
+		TOptional<FFrameNumber> FrameNumber = (Sequencer->GetLocalTime().Time * RootToLocalTransform.InverseLinearOnly()).RoundToFrame();
 		return LocalGetControlRigControlTransforms(Sequencer, FrameNumber, Sequencer->GetFocusedMovieSceneSequence(), Template, RootToLocalTransform,
 			ControlRig, ControlName, Frames, ParentTransforms, OutTransforms);
 	
diff --git a/Engine/Plugins/Editor/SequencerAnimTools/Source/SequencerAnimTools/Private/Sequencer/MovieSceneTransformTrail.cpp b/Engine/Plugins/Editor/SequencerAnimTools/Source/SequencerAnimTools/Private/Sequencer/MovieSceneTransformTrail.cpp
--- a/Engine/Plugins/Editor/SequencerAnimTools/Source/SequencerAnimTools/Private/Sequencer/MovieSceneTransformTrail.cpp
+++ b/Engine/Plugins/Editor/SequencerAnimTools/Source/SequencerAnimTools/Private/Sequencer/MovieSceneTransformTrail.cpp
@@ -751,13 +751,17 @@ bool FMovieSceneControlRigTransformTrail::ApplyDelta(const FVector& Pos, const F
 		Context.SetKey = EControlRigSetKey::Always;
 		IMovieScenePlayer* Player = GetSequencer().Get();
 
-		auto EvalControlRig = [&Context, Pos, TickResolution, ControlRig,Player, this](FTrailKeyInfo* KeyInfo)
+		FMovieSceneSequenceTransform RootToLocalTransform = GetSequencer()->GetFocusedMovieSceneSequenceTransform();
+
+		auto EvalControlRig = [&Context,&RootToLocalTransform, Pos, TickResolution, ControlRig,Player, this](FTrailKeyInfo* KeyInfo)
 		{
 			if (KeyInfo)
 			{
 				Context.LocalTime = TickResolution.AsSeconds(FFrameTime(KeyInfo->FrameNumber));
 				Context.KeyMask = (uint32)EControlRigContextChannelToKey::Translation;
 				FFrameTime GlobalTime(KeyInfo->FrameNumber);
+				GlobalTime = GlobalTime * RootToLocalTransform.InverseLinearOnly(); //player evals in root time so need to go back to it.
+
 				FMovieSceneContext MovieSceneContext = FMovieSceneContext(FMovieSceneEvaluationRange(GlobalTime, TickResolution), Player->GetPlaybackStatus()).SetHasJumped(true);
 
 				FTransform NewTransform(KeyInfo->Transform);
@@ -814,6 +818,9 @@ bool FMovieSceneControlRigTransformTrail::EndTracking()
 		IMovieScenePlayer* Player = GetSequencer().Get();
 		FFrameTime StartTime = GetSequencer()->GetLocalTime().Time;
 		FFrameRate TickResolution = GetSequencer()->GetFocusedTickResolution();
+		FMovieSceneSequenceTransform RootToLocalTransform = GetSequencer()->GetFocusedMovieSceneSequenceTransform();
+		StartTime = StartTime * RootToLocalTransform.InverseLinearOnly(); //player evals in root time so need to go back to it.
+
 		FMovieSceneContext MovieSceneContext = FMovieSceneContext(FMovieSceneEvaluationRange(StartTime, TickResolution), Player->GetPlaybackStatus()).SetHasJumped(true);
 		
 		Player->GetEvaluationTemplate().Evaluate(MovieSceneContext, *Player);
@@ -850,6 +857,8 @@ bool FMovieSceneControlRigTransformTrail::HandleAltClick(FEditorViewportClient*
 	FFrameTime DisplayTime = FFrameRate::TransformTime(GlobalTime, DisplayResolution, TickResolution);
 	DisplayTime = DisplayTime.RoundToFrame();
 	GlobalTime = FFrameRate::TransformTime(DisplayTime, TickResolution, DisplayResolution);
+	FMovieSceneSequenceTransform RootToLocalTransform = GetSequencer()->GetFocusedMovieSceneSequenceTransform();
+	GlobalTime = GlobalTime * RootToLocalTransform.InverseLinearOnly(); //player evals in root time so need to go back to it.
 
 	Context.LocalTime = TickResolution.AsSeconds(GlobalTime);
 	Context.KeyMask = (uint32)EControlRigContextChannelToKey::Translation;
@@ -861,6 +870,8 @@ bool FMovieSceneControlRigTransformTrail::HandleAltClick(FEditorViewportClient*
 
 	//eval back at current time
 	FFrameTime StartTime = GetSequencer()->GetLocalTime().Time;
+	StartTime = StartTime * RootToLocalTransform.InverseLinearOnly(); //player evals in root time so need to go back to it.
+
 	MovieSceneContext = FMovieSceneContext(FMovieSceneEvaluationRange(StartTime, TickResolution), Player->GetPlaybackStatus()).SetHasJumped(true);
 	Player->GetEvaluationTemplate().Evaluate(MovieSceneContext, *Player);
 	ControlRig->Evaluate_AnyThread();
diff --git a/Engine/Plugins/Editor/SequencerAnimTools/Source/SequencerAnimTools/Private/Sequencer/SequencerTrailHierarchy.cpp b/Engine/Plugins/Editor/SequencerAnimTools/Source/SequencerAnimTools/Private/Sequencer/SequencerTrailHierarchy.cpp
--- a/Engine/Plugins/Editor/SequencerAnimTools/Source/SequencerAnimTools/Private/Sequencer/SequencerTrailHierarchy.cpp
+++ b/Engine/Plugins/Editor/SequencerAnimTools/Source/SequencerAnimTools/Private/Sequencer/SequencerTrailHierarchy.cpp
@@ -231,7 +231,7 @@ void FSequencerTrailHierarchy::UpdateControlRig(const TArray<FFrameNumber> &Fram
 		FMovieSceneSequencePlaybackSettings Settings;
 		FLevelSequenceCameraSettings CameraSettings;
 		FMovieSceneSequenceIDRef Template = MovieSceneSequenceID::Root;
-		FMovieSceneSequenceTransform RootToLocalTransform;
+		FMovieSceneSequenceTransform RootToLocalTransform = Sequencer->GetFocusedMovieSceneSequenceTransform();
 		IMovieScenePlayer* Player = Sequencer.Get();
 		ULevelSequence* LevelSequence = Cast<ULevelSequence>(Sequencer->GetFocusedMovieSceneSequence());
 
@@ -257,10 +257,12 @@ void FSequencerTrailHierarchy::UpdateControlRig(const TArray<FFrameNumber> &Fram
 					TrailControlTransforms[PairIndex].Transforms.SetNum(Frames.Num());
 					++PairIndex;
 				}
+
 				for (int32 Index = 0; Index < Frames.Num(); ++Index)
 				{
 					const FFrameNumber& FrameNumber = Frames[Index];
 					FFrameTime GlobalTime(FrameNumber);
+					GlobalTime = GlobalTime * RootToLocalTransform.InverseLinearOnly(); //player evals in root time so need to go back to it.
 
 					FMovieSceneContext Context = FMovieSceneContext(FMovieSceneEvaluationRange(GlobalTime, TickResolution), Player->GetPlaybackStatus()).SetHasJumped(true);
 
@@ -288,7 +290,8 @@ void FSequencerTrailHierarchy::UpdateControlRig(const TArray<FFrameNumber> &Fram
 		}
 		if (Player)
 		{
-			const FFrameTime StartTime = Sequencer->GetLocalTime().Time;
+			FFrameTime StartTime = Sequencer->GetLocalTime().Time;
+			StartTime = StartTime * RootToLocalTransform.InverseLinearOnly(); //player evals in root time so need to go back to it.
 			FMovieSceneContext Context = FMovieSceneContext(FMovieSceneEvaluationRange(StartTime, TickResolution), Player->GetPlaybackStatus()).SetHasJumped(true);
 			Player->GetEvaluationTemplate().Evaluate(Context, *Player);
 			ControlRig->Evaluate_AnyThread();
