Control Rig: Space Switching:  Bring back space switching fix for multi control rigs. This got removed  by mistake by a Constraint submission CL 20221291

diff --git a/Engine/Plugins/Animation/ControlRig/Source/ControlRigEditor/Private/EditMode/SControlRigEditModeTools.cpp b/Engine/Plugins/Animation/ControlRig/Source/ControlRigEditor/Private/EditMode/SControlRigEditModeTools.cpp
--- a/Engine/Plugins/Animation/ControlRig/Source/ControlRigEditor/Private/EditMode/SControlRigEditModeTools.cpp
+++ b/Engine/Plugins/Animation/ControlRig/Source/ControlRigEditor/Private/EditMode/SControlRigEditModeTools.cpp
@@ -750,14 +750,12 @@ void SControlRigEditModeTools::OnRigElementSelected(UControlRig* Subject, FRigCo
 		}
 	}
 #endif
-	//mz todo handle multiple rigs
-	UControlRig* Rig = ControlRigs.Num() > 0 ? ControlRigs[0].Get() : nullptr;
-	if (Rig)
+
+	if (Subject)
 	{
 		// get the selected controls
-		TArray<FRigElementKey> SelectedControls = Rig->GetHierarchy()->GetSelectedKeys(ERigElementType::Control);
-		SpacePickerWidget->SetControls(Rig->GetHierarchy(), SelectedControls);
-
+		TArray<FRigElementKey> SelectedControls = Subject->GetHierarchy()->GetSelectedKeys(ERigElementType::Control);
+		SpacePickerWidget->SetControls(Subject->GetHierarchy(), SelectedControls);
 		if (ConstraintsEditionWidget)
 		{
 			ConstraintsEditionWidget->InvalidateConstraintList();
