Fixed use of possible null struct
note: nodes register a FPerClassInfo with a null ScriptStruct

diff --git a/Engine/Plugins/Animation/ControlRig/Source/ControlRigEditor/Private/Editor/ControlRigEditor.cpp b/Engine/Plugins/Animation/ControlRig/Source/ControlRigEditor/Private/Editor/ControlRigEditor.cpp
--- a/Engine/Plugins/Animation/ControlRig/Source/ControlRigEditor/Private/Editor/ControlRigEditor.cpp
+++ b/Engine/Plugins/Animation/ControlRig/Source/ControlRigEditor/Private/Editor/ControlRigEditor.cpp
@@ -1722,9 +1722,12 @@ bool FControlRigEditor::DetailViewShowsStruct(UScriptStruct* InStruct) const
 		{
 			if(UDetailsViewWrapperObject* WrapperObject = Cast<UDetailsViewWrapperObject>(SelectedObject.Get()))
 			{
-				if(WrapperObject->GetWrappedStruct()->IsChildOf(InStruct))
+				if(const UScriptStruct* WrappedStruct = WrapperObject->GetWrappedStruct())
 				{
-					return true;
+					if (WrappedStruct->IsChildOf(InStruct))
+					{
+						return true;
+					}
 				}
 			}
 		}
@@ -1741,11 +1744,14 @@ bool FControlRigEditor::DetailViewShowsRigElement(FRigElementKey InKey) const
 		{
 			if(UDetailsViewWrapperObject* WrapperObject = Cast<UDetailsViewWrapperObject>(SelectedObject.Get()))
 			{
-				if(WrapperObject->GetWrappedStruct()->IsChildOf(FRigBaseElement::StaticStruct()))
+				if (const UScriptStruct* WrappedStruct = WrapperObject->GetWrappedStruct())
 				{
-					if(WrapperObject->GetContent<FRigBaseElement>().GetKey() == InKey)
+					if (WrappedStruct->IsChildOf(FRigBaseElement::StaticStruct()))
 					{
-						return true;
+						if(WrapperObject->GetContent<FRigBaseElement>().GetKey() == InKey)
+						{
+							return true;
+						}
 					}
 				}
 			}
@@ -2528,12 +2534,15 @@ void FControlRigEditor::EnsureValidRigElementsInDetailPanel()
 		{
 			if(UDetailsViewWrapperObject* WrapperObject = Cast<UDetailsViewWrapperObject>(SelectedObject.Get()))
 			{
-				if(WrapperObject->GetWrappedStruct()->IsChildOf(FRigBaseElement::StaticStruct()))
+				if(const UScriptStruct* WrappedStruct = WrapperObject->GetWrappedStruct())
 				{
-					FRigElementKey Key = WrapperObject->GetContent<FRigBaseElement>().GetKey();
-					if(!Hierarchy->Contains(Key))
+					if (WrappedStruct->IsChildOf(FRigBaseElement::StaticStruct()))
 					{
-						ClearDetailObject();
+						FRigElementKey Key = WrapperObject->GetContent<FRigBaseElement>().GetKey();
+						if(!Hierarchy->Contains(Key))
+						{
+							ClearDetailObject();
+						}
 					}
 				}
 			}
@@ -2883,49 +2892,51 @@ void FControlRigEditor::HandleControlRigExecutedEvent(UControlRig* InControlRig,
 			{
 				if(UDetailsViewWrapperObject* WrapperObject = Cast<UDetailsViewWrapperObject>(SelectedObject.Get()))
 				{
-					UScriptStruct* Struct = WrapperObject->GetWrappedStruct(); 
-					if(Struct->IsChildOf(FRigBaseElement::StaticStruct()))
+					if (const UScriptStruct* Struct = WrapperObject->GetWrappedStruct())
 					{
-						const FRigElementKey Key = WrapperObject->GetContent<FRigBaseElement>().GetKey();
-
-						FRigBaseElement* Element = Hierarchy->Find(Key);
-						if(Element == nullptr)
+						if(Struct->IsChildOf(FRigBaseElement::StaticStruct()))
 						{
-							ClearDetailObject();
-							break;
-						}
+							const FRigElementKey Key = WrapperObject->GetContent<FRigBaseElement>().GetKey();
 
-						if(FRigControlElement* ControlElement = Cast<FRigControlElement>(Element))
-						{
-							// compute all transforms
-							Hierarchy->GetTransform(ControlElement, ERigTransformType::CurrentGlobal);
-							Hierarchy->GetTransform(ControlElement, ERigTransformType::CurrentLocal);
-							Hierarchy->GetTransform(ControlElement, ERigTransformType::InitialGlobal);
-							Hierarchy->GetTransform(ControlElement, ERigTransformType::InitialLocal);
-							Hierarchy->GetControlOffsetTransform(ControlElement, ERigTransformType::CurrentGlobal);
-							Hierarchy->GetControlOffsetTransform(ControlElement, ERigTransformType::CurrentLocal);
-							Hierarchy->GetControlOffsetTransform(ControlElement, ERigTransformType::InitialGlobal);
-							Hierarchy->GetControlOffsetTransform(ControlElement, ERigTransformType::InitialLocal);
-							Hierarchy->GetControlShapeTransform(ControlElement, ERigTransformType::CurrentGlobal);
-							Hierarchy->GetControlShapeTransform(ControlElement, ERigTransformType::CurrentLocal);
-							Hierarchy->GetControlShapeTransform(ControlElement, ERigTransformType::InitialGlobal);
-							Hierarchy->GetControlShapeTransform(ControlElement, ERigTransformType::InitialLocal);
-
-							WrapperObject->SetContent<FRigControlElement>(*ControlElement);
-						}
-						else if(FRigTransformElement* TransformElement = Cast<FRigTransformElement>(Element))
-						{
-							// compute all transforms
-							Hierarchy->GetTransform(TransformElement, ERigTransformType::CurrentGlobal);
-							Hierarchy->GetTransform(TransformElement, ERigTransformType::CurrentLocal);
-							Hierarchy->GetTransform(TransformElement, ERigTransformType::InitialGlobal);
-							Hierarchy->GetTransform(TransformElement, ERigTransformType::InitialLocal);
+							FRigBaseElement* Element = Hierarchy->Find(Key);
+							if(Element == nullptr)
+							{
+								ClearDetailObject();
+								break;
+							}
 
-							WrapperObject->SetContent<FRigTransformElement>(*TransformElement);
-						}
-						else
-						{
-							WrapperObject->SetContent<FRigBaseElement>(*Element);
+							if(FRigControlElement* ControlElement = Cast<FRigControlElement>(Element))
+							{
+								// compute all transforms
+								Hierarchy->GetTransform(ControlElement, ERigTransformType::CurrentGlobal);
+								Hierarchy->GetTransform(ControlElement, ERigTransformType::CurrentLocal);
+								Hierarchy->GetTransform(ControlElement, ERigTransformType::InitialGlobal);
+								Hierarchy->GetTransform(ControlElement, ERigTransformType::InitialLocal);
+								Hierarchy->GetControlOffsetTransform(ControlElement, ERigTransformType::CurrentGlobal);
+								Hierarchy->GetControlOffsetTransform(ControlElement, ERigTransformType::CurrentLocal);
+								Hierarchy->GetControlOffsetTransform(ControlElement, ERigTransformType::InitialGlobal);
+								Hierarchy->GetControlOffsetTransform(ControlElement, ERigTransformType::InitialLocal);
+								Hierarchy->GetControlShapeTransform(ControlElement, ERigTransformType::CurrentGlobal);
+								Hierarchy->GetControlShapeTransform(ControlElement, ERigTransformType::CurrentLocal);
+								Hierarchy->GetControlShapeTransform(ControlElement, ERigTransformType::InitialGlobal);
+								Hierarchy->GetControlShapeTransform(ControlElement, ERigTransformType::InitialLocal);
+
+								WrapperObject->SetContent<FRigControlElement>(*ControlElement);
+							}
+							else if(FRigTransformElement* TransformElement = Cast<FRigTransformElement>(Element))
+							{
+								// compute all transforms
+								Hierarchy->GetTransform(TransformElement, ERigTransformType::CurrentGlobal);
+								Hierarchy->GetTransform(TransformElement, ERigTransformType::CurrentLocal);
+								Hierarchy->GetTransform(TransformElement, ERigTransformType::InitialGlobal);
+								Hierarchy->GetTransform(TransformElement, ERigTransformType::InitialLocal);
+
+								WrapperObject->SetContent<FRigTransformElement>(*TransformElement);
+							}
+							else
+							{
+								WrapperObject->SetContent<FRigBaseElement>(*Element);
+							}
 						}
 					}
 				}
@@ -4432,11 +4443,11 @@ void FControlRigEditor::OnWrappedPropertyChangedChainEvent(UDetailsViewWrapperOb
 	UControlRigBlueprint* ControlRigBP = GetControlRigBlueprint();
 
 	FString PropertyPath = InPropertyPath;
-	if(InWrapperObject->GetWrappedStruct())
+	if(UScriptStruct* WrappedStruct = InWrapperObject->GetWrappedStruct())
 	{
-		if(InWrapperObject->GetWrappedStruct()->IsChildOf(FRigBaseElement::StaticStruct()))
+		if(WrappedStruct->IsChildOf(FRigBaseElement::StaticStruct()))
 		{
-			check(InWrapperObject->GetWrappedStruct() == WrapperObjects[0]->GetWrappedStruct());
+			check(WrappedStruct == WrapperObjects[0]->GetWrappedStruct());
 
 			URigHierarchy* Hierarchy = CastChecked<URigHierarchy>(InWrapperObject->GetOuter());
 			const FRigBaseElement WrappedElement = InWrapperObject->GetContent<FRigBaseElement>();
@@ -4574,9 +4585,9 @@ void FControlRigEditor::OnWrappedPropertyChangedChainEvent(UDetailsViewWrapperOb
 				ControlRigBP->MarkPackageDirty();
 			}
 		}
-		else if(InWrapperObject->GetWrappedStruct()->IsChildOf(FRigVMGraphVariableDescription::StaticStruct()))
+		else if(WrappedStruct->IsChildOf(FRigVMGraphVariableDescription::StaticStruct()))
 		{
-			check(InWrapperObject->GetWrappedStruct() == WrapperObjects[0]->GetWrappedStruct());
+			check(WrappedStruct == WrapperObjects[0]->GetWrappedStruct());
 			
 			const FRigVMGraphVariableDescription VariableDescription = InWrapperObject->GetContent<FRigVMGraphVariableDescription>();
 			URigVMGraph* Graph = CastChecked<URigVMGraph>(InWrapperObject->GetOuter());
