Control Rig: backend work for support for multi graph

diff --git a/Engine/Plugins/Animation/ControlRig/Source/ControlRigDeveloper/Private/ControlRigBlueprint.cpp b/Engine/Plugins/Animation/ControlRig/Source/ControlRigDeveloper/Private/ControlRigBlueprint.cpp
--- a/Engine/Plugins/Animation/ControlRig/Source/ControlRigDeveloper/Private/ControlRigBlueprint.cpp
+++ b/Engine/Plugins/Animation/ControlRig/Source/ControlRigDeveloper/Private/ControlRigBlueprint.cpp
@@ -27,6 +27,8 @@
 #include "AssetRegistry/AssetRegistryModule.h"
 #include "RigVMPythonUtils.h"
 #include "RigVMTypeUtils.h"
+#include "Algo/Count.h"
+#include "Algo/Transform.h"
 #include "RigVMModel/Nodes/RigVMAggregateNode.h"
 #include "Units/ControlRigNodeWorkflow.h"
 #include "Rigs/RigControlHierarchy.h"
@@ -43,6 +45,7 @@
 #include "CookOnTheSide/CookOnTheFlyServer.h"
 #include "Widgets/Notifications/SNotificationList.h"
 #include "Framework/Notifications/NotificationManager.h"
+#include "ScopedTransaction.h"
 #endif//WITH_EDITOR
 
 #define LOCTEXT_NAMESPACE "ControlRigBlueprint"
@@ -88,6 +91,7 @@
 	bSuspendModelNotificationsForSelf = false;
 	bSuspendModelNotificationsForOthers = false;
 	bSuspendAllNotifications = false;
+	bSuspendPythonMessagesForRigVMClient = true;
 
 #if WITH_EDITORONLY_DATA
 	GizmoLibrary_DEPRECATED = nullptr;
@@ -100,8 +104,13 @@
 	bIsCompiling = false;
 	VMRecompilationBracket = 0;
 
-	Model = ObjectInitializer.CreateDefaultSubobject<URigVMGraph>(this, TEXT("RigVMModel"));
-	FunctionLibrary = ObjectInitializer.CreateDefaultSubobject<URigVMFunctionLibrary>(this, TEXT("RigVMFunctionLibrary"));
+	RigVMClient.SetOuterClientHost(this, GET_MEMBER_NAME_CHECKED(UControlRigBlueprint, RigVMClient));
+	{
+		TGuardValue<bool> DisableClientNotifs(RigVMClient.bSuspendNotifications, true);
+		RigVMClient.AddModel(RigVMModelPrefix, false, &ObjectInitializer);
+		RigVMClient.GetOrCreateFunctionLibrary(false, &ObjectInitializer);
+	}
+	
 	FunctionLibraryEdGraph = ObjectInitializer.CreateDefaultSubobject<UControlRigGraph>(this, TEXT("RigVMFunctionLibraryEdGraph"));
 	FunctionLibraryEdGraph->Schema = UControlRigGraphSchema::StaticClass();
 	FunctionLibraryEdGraph->bAllowRenaming = 0;
@@ -110,11 +119,6 @@
 	FunctionLibraryEdGraph->bIsFunctionDefinition = false;
 	FunctionLibraryEdGraph->Initialize(this);
 
-	Model->SetDefaultFunctionLibrary(FunctionLibrary);
-
-	Model->SetExecuteContextStruct(FControlRigExecuteContext::StaticStruct());
-	FunctionLibrary->SetExecuteContextStruct(FControlRigExecuteContext::StaticStruct());
-
 	Validator = ObjectInitializer.CreateDefaultSubobject<UControlRigValidator>(this, TEXT("ControlRigValidator"));
 
 	DebugBoneRadius = 1.f;
@@ -147,28 +151,34 @@
 {
 	DECLARE_SCOPE_HIERARCHICAL_COUNTER_FUNC()
 
-	if (Controllers.Num() == 0)
+	if (RigVMClient.GetController(0) == nullptr)
 	{
-		Model->SetExecuteContextStruct(FControlRigExecuteContext::StaticStruct());
-		FunctionLibrary->SetExecuteContextStruct(FControlRigExecuteContext::StaticStruct());
-		GetOrCreateController(Model);
-		GetOrCreateController(FunctionLibrary);
+		check(RigVMClient.Num() == 1);
+		check(RigVMClient.GetFunctionLibrary());
+		
+		RigVMClient.GetOrCreateController(RigVMClient.GetDefaultModel());
+		RigVMClient.GetOrCreateController(RigVMClient.GetFunctionLibrary());
 
+		bool bRecompileRequired = false;
 		for (int32 i = 0; i < UbergraphPages.Num(); ++i)
 		{
 			if (UControlRigGraph* Graph = Cast<UControlRigGraph>(UbergraphPages[i]))
 			{
 				PopulateModelFromGraphForBackwardsCompatibility(Graph);
-
 				if (bRecompileVM)
 				{
-					RecompileVM();
+					bRecompileRequired = true;
 				}
 
 				Graph->Initialize(this);
 			}
 		}
 
+		if(bRecompileRequired)
+		{
+			RecompileVM();
+		}
+
 		FunctionLibraryEdGraph->Initialize(this);
 	}
 }
@@ -238,6 +248,342 @@
 	PostEditChangeChainPropertyEvent.Broadcast(PropertyChangedEvent);
 }
 
+FRigVMClient* UControlRigBlueprint::GetRigVMClient()
+{
+	return &RigVMClient;
+}
+
+const FRigVMClient* UControlRigBlueprint::GetRigVMClient() const
+{
+	return &RigVMClient;
+}
+
+UObject* UControlRigBlueprint::GetEditorObjectForRigVMGraph(URigVMGraph* InVMGraph) const
+{
+	if(InVMGraph)
+	{
+		if(InVMGraph->GetOutermost() != GetOutermost())
+		{
+			return nullptr;
+		}
+
+		TArray<UEdGraph*> EdGraphs;
+		GetAllGraphs(EdGraphs);
+
+		bool bIsFunctionDefinition = false;
+		if (URigVMLibraryNode* LibraryNode = Cast<URigVMLibraryNode>(InVMGraph->GetOuter()))
+		{
+			bIsFunctionDefinition = LibraryNode->GetGraph()->IsA<URigVMFunctionLibrary>();
+		}
+
+		for (UEdGraph* EdGraph : EdGraphs)
+		{
+			if (UControlRigGraph* RigGraph = Cast<UControlRigGraph>(EdGraph))
+			{
+				if (RigGraph->bIsFunctionDefinition != bIsFunctionDefinition)
+				{
+					continue;
+				}
+
+				if ((RigGraph->ModelNodePath == InVMGraph->GetNodePath()) ||
+					(RigGraph->ModelNodePath.IsEmpty() && (RigVMClient.GetDefaultModel() == InVMGraph)))
+				{
+					return RigGraph;
+				}
+			}
+		}
+	}
+	
+	return nullptr;
+}
+
+void UControlRigBlueprint::HandleRigVMGraphAdded(const FRigVMClient* InClient, const FString& InNodePath)
+{
+	if(URigVMGraph* Model = InClient->GetModel(InNodePath))
+	{
+		Model->SetExecuteContextStruct(FControlRigExecuteContext::StaticStruct());
+
+		if(!HasAnyFlags(RF_ClassDefaultObject | RF_NeedInitialization | RF_NeedLoad | RF_NeedPostLoad) &&
+			GetOuter() != GetTransientPackage())
+		{
+			CreateEdGraph(Model, true);
+			RecompileVM();
+		}
+
+#if WITH_EDITOR
+		if(!bSuspendPythonMessagesForRigVMClient)
+		{
+			const FString BlueprintName = URigVMController::GetSanitizedName(GetName(), true, false);
+			RigVMPythonUtils::Print(BlueprintName, 
+				FString::Printf(TEXT("blueprint.add_model('%s')"),
+					*Model->GetName()));
+		}
+#endif
+	}
+}
+
+void UControlRigBlueprint::HandleRigVMGraphRemoved(const FRigVMClient* InClient, const FString& InNodePath)
+{
+	if(URigVMGraph* Model = InClient->GetModel(InNodePath))
+	{
+		RemoveEdGraph(Model);
+		RecompileVM();
+
+#if WITH_EDITOR
+		if(!bSuspendPythonMessagesForRigVMClient)
+		{
+			const FString BlueprintName = URigVMController::GetSanitizedName(GetName(), true, false);
+			RigVMPythonUtils::Print(BlueprintName, 
+				FString::Printf(TEXT("blueprint.remove_model('%s')"),
+					*Model->GetName()));
+		}
+#endif
+	}
+}
+
+void UControlRigBlueprint::HandleRigVMGraphRenamed(const FRigVMClient* InClient, const FString& InOldNodePath, const FString& InNewNodePath)
+{
+	if(URigVMGraph* Model = InClient->GetModel(InNewNodePath))
+	{
+		TArray<UEdGraph*> EdGraphs;
+		GetAllGraphs(EdGraphs);
+
+		for (UEdGraph* EdGraph : EdGraphs)
+		{
+			if (UControlRigGraph* RigGraph = Cast<UControlRigGraph>(EdGraph))
+			{
+				RigGraph->HandleRigVMGraphRenamed(InOldNodePath, InNewNodePath);
+			}
+		}
+
+		FBlueprintEditorUtils::MarkBlueprintAsStructurallyModified(this);
+	}
+}
+
+void UControlRigBlueprint::HandleConfigureRigVMController(const FRigVMClient* InClient,
+	URigVMController* InControllerToConfigure)
+{
+	InControllerToConfigure->OnModified().AddUObject(this, &UControlRigBlueprint::HandleModifiedEvent);
+
+	InControllerToConfigure->UnfoldStructDelegate.BindLambda([](const UStruct* InStruct) -> bool {
+
+		if (InStruct == TBaseStructure<FQuat>::Get())
+		{
+			return false;
+		}
+		if (InStruct == FRuntimeFloatCurve::StaticStruct())
+		{
+			return false;
+		}
+		if (InStruct == FRigPose::StaticStruct())
+		{
+			return false;
+		}
+		return true;
+		});
+
+	TWeakObjectPtr<UControlRigBlueprint> WeakThis(this);
+
+	// this delegate is used by the controller to determine variable validity
+	// during a bind process. the controller itself doesn't own the variables,
+	// so we need a delegate to request them from the owning blueprint
+	InControllerToConfigure->GetExternalVariablesDelegate.BindLambda([](URigVMGraph* InGraph) -> TArray<FRigVMExternalVariable> {
+
+		if (InGraph)
+		{
+			if(UControlRigBlueprint* Blueprint = InGraph->GetTypedOuter<UControlRigBlueprint>())
+			{
+				if (UControlRigBlueprintGeneratedClass* RigClass = Blueprint->GetControlRigBlueprintGeneratedClass())
+				{
+                    if (UControlRig* CDO = Cast<UControlRig>(RigClass->GetDefaultObject(true /* create if needed */)))
+                    {
+                        return CDO->GetExternalVariablesImpl(true /* rely on variables within blueprint */);
+                    }
+                }
+			}
+		}
+		return TArray<FRigVMExternalVariable>();
+
+	});
+
+
+	// this delegate is used by the controller to retrieve the current bytecode of the VM
+	InControllerToConfigure->GetCurrentByteCodeDelegate.BindLambda([WeakThis]() -> const FRigVMByteCode* {
+
+		if (WeakThis.IsValid())
+		{
+			if (UControlRigBlueprintGeneratedClass* RigClass = WeakThis->GetControlRigBlueprintGeneratedClass())
+			{
+				if (UControlRig* CDO = Cast<UControlRig>(RigClass->GetDefaultObject(false)))
+				{
+					if (CDO->VM)
+					{
+						return &CDO->VM->GetByteCode();
+					}
+				}
+			}
+		}
+		return nullptr;
+
+	});
+
+	InControllerToConfigure->IsFunctionAvailableDelegate.BindLambda([WeakThis](URigVMLibraryNode* InFunction) -> bool
+	{
+		if(InFunction == nullptr)
+		{
+			return false;
+		}
+		
+		if(URigVMFunctionLibrary* Library = Cast<URigVMFunctionLibrary>(InFunction->GetOuter()))
+		{
+			if(UControlRigBlueprint* Blueprint = Cast<UControlRigBlueprint>(Library->GetOuter()))
+			{
+				if(Blueprint->IsFunctionPublic(InFunction->GetFName()))
+				{
+					return true;
+				}
+
+				// if it is private - we still see it as public if we are within the same blueprint 
+				if(WeakThis.IsValid())
+				{
+					if(WeakThis.Get() == Blueprint)
+					{
+						return true;
+					}
+				}
+			}
+		}
+
+		return false;
+	});
+
+	InControllerToConfigure->IsDependencyCyclicDelegate.BindLambda([WeakThis](UObject* InDependentObject, UObject* InDependencyObject) -> bool
+	{
+	    if(InDependentObject == nullptr || InDependencyObject == nullptr)
+	    {
+	        return false;
+	    }
+
+		UControlRigBlueprint* DependentBlueprint = InDependentObject->GetTypedOuter<UControlRigBlueprint>();
+		UControlRigBlueprint* DependencyBlueprint = InDependencyObject->GetTypedOuter<UControlRigBlueprint>();
+
+		if(DependentBlueprint == nullptr || DependencyBlueprint == nullptr)
+		{
+			return false;
+		}
+
+		if(DependentBlueprint == DependencyBlueprint)
+		{
+			return false;
+		}
+
+		const TArray<UControlRigBlueprint*> DependencyDependencies = DependencyBlueprint->GetDependencies(true);
+		return DependencyDependencies.Contains(DependentBlueprint);
+	});
+
+#if WITH_EDITOR
+
+	// this sets up three delegates:
+	// a) get external variables (mapped to Controller->GetExternalVariables)
+	// b) bind pin to variable (mapped to Controller->BindPinToVariable)
+	// c) create external variable (mapped to the passed in tfunction)
+	// the last one is defined within the blueprint since the controller
+	// doesn't own the variables and can't create one itself.
+	InControllerToConfigure->SetupDefaultUnitNodeDelegates(TDelegate<FName(FRigVMExternalVariable, FString)>::CreateLambda(
+		[WeakThis](FRigVMExternalVariable InVariableToCreate, FString InDefaultValue) -> FName {
+			if (WeakThis.IsValid())
+			{
+				return WeakThis->AddCRMemberVariableFromExternal(InVariableToCreate, InDefaultValue);
+			}
+			return NAME_None;
+		}
+	));
+
+	TWeakObjectPtr<URigVMController> WeakController = InControllerToConfigure;
+	InControllerToConfigure->RequestBulkEditDialogDelegate.BindLambda([WeakThis, WeakController](URigVMLibraryNode* InFunction, ERigVMControllerBulkEditType InEditType) -> FRigVMController_BulkEditResult 
+	{
+		if(WeakThis.IsValid() && WeakController.IsValid())
+		{
+			UControlRigBlueprint* StrongThis = WeakThis.Get();
+            URigVMController* StrongController = WeakController.Get();
+            if(StrongThis->OnRequestBulkEditDialog().IsBound())
+			{
+				return StrongThis->OnRequestBulkEditDialog().Execute(StrongThis, StrongController, InFunction, InEditType);
+			}
+		}
+		return FRigVMController_BulkEditResult();
+	});
+
+	InControllerToConfigure->RequestBreakLinksDialogDelegate.BindLambda([WeakThis, WeakController](TArray<URigVMLink*> InLinks) -> bool 
+	{
+		if(WeakThis.IsValid() && WeakController.IsValid())
+		{
+			UControlRigBlueprint* StrongThis = WeakThis.Get();
+			if(StrongThis->OnRequestBreakLinksDialog().IsBound())
+			{
+				return StrongThis->OnRequestBreakLinksDialog().Execute(InLinks);
+			}
+		}
+		return false;
+	});
+
+	InControllerToConfigure->RequestNewExternalVariableDelegate.BindLambda([WeakThis](FRigVMGraphVariableDescription InVariable, bool bInIsPublic, bool bInIsReadOnly) -> FName
+	{
+		if (WeakThis.IsValid())
+		{
+			for (FBPVariableDescription& ExistingVariable : WeakThis->NewVariables)
+			{
+				if (ExistingVariable.VarName == InVariable.Name)
+				{
+					return FName();
+				}
+			}
+
+			FRigVMExternalVariable ExternalVariable = InVariable.ToExternalVariable();
+			return WeakThis->AddMemberVariable(InVariable.Name,
+				ExternalVariable.TypeObject ? ExternalVariable.TypeObject->GetPathName() : ExternalVariable.TypeName.ToString(),
+				bInIsPublic,
+				bInIsReadOnly,
+				InVariable.DefaultValue);
+		}
+		
+		return FName();
+	});
+
+
+	InControllerToConfigure->RequestJumpToHyperlinkDelegate.BindLambda([WeakThis](const UObject* InSubject)
+	{
+		if (WeakThis.IsValid())
+		{
+			UControlRigBlueprint* StrongThis = WeakThis.Get();
+			if(StrongThis->OnRequestJumpToHyperlink().IsBound())
+			{
+				StrongThis->OnRequestJumpToHyperlink().Execute(InSubject);
+			}
+		}
+	});
+
+	InControllerToConfigure->ConfigureWorkflowOptionsDelegate.BindLambda([WeakThis](URigVMUserWorkflowOptions* Options)
+	{
+		if(UControlRigWorkflowOptions* ControlRigNodeWorkflowOptions = Cast<UControlRigWorkflowOptions>(Options))
+		{
+			ControlRigNodeWorkflowOptions->Hierarchy = nullptr;
+			ControlRigNodeWorkflowOptions->Selection.Reset();
+			
+			if(const UControlRigBlueprint* StrongThis = WeakThis.Get())
+			{
+				if(UControlRig* ControlRig = Cast<UControlRig>(StrongThis->GetObjectBeingDebugged()))
+				{
+					ControlRigNodeWorkflowOptions->Hierarchy = ControlRig->GetHierarchy();
+				}
+				ControlRigNodeWorkflowOptions->Selection = StrongThis->Hierarchy->GetSelectedKeys();
+			}
+		}
+	});
+
+#endif
+}
+
 bool UControlRigBlueprint::TryImportGraphFromText(const FString& InClipboardText, UEdGraph** OutGraphPtr)
 {
 	if (OutGraphPtr)
@@ -314,6 +660,8 @@
 
 void UControlRigBlueprint::Serialize(FArchive& Ar)
 {
+	RigVMClient.SetOuterClientHost(this, GET_MEMBER_NAME_CHECKED(UControlRigBlueprint, RigVMClient));
+	
 	Super::Serialize(Ar);
 
 	if(Ar.IsObjectReferenceCollector())
@@ -334,6 +682,15 @@
 			}
 		}
 	}
+
+	if(Ar.IsLoading())
+	{
+		if(Model_DEPRECATED || FunctionLibrary_DEPRECATED)
+		{
+			TGuardValue<bool> DisableClientNotifs(RigVMClient.bSuspendNotifications, true);
+			RigVMClient.SetFromDeprecatedData(Model_DEPRECATED, FunctionLibrary_DEPRECATED);
+		}
+	}
 }
 
 void UControlRigBlueprint::PreSave(const class ITargetPlatform* TargetPlatform)
@@ -369,7 +726,7 @@
 
 	for(FControlRigPublicFunctionData& FunctionData : PublicFunctions)
 	{
-		if(URigVMLibraryNode* FunctionNode = FunctionLibrary->FindFunction(FunctionData.Name))
+		if(URigVMLibraryNode* FunctionNode = RigVMClient.GetFunctionLibrary()->FindFunction(FunctionData.Name))
 		{
 			if(UControlRigGraph* Graph = Cast<UControlRigGraph>(GetEdGraph(FunctionNode->GetContainedGraph())))
 			{
@@ -662,7 +1019,7 @@
 
 	if(URigVMBuildData* BuildData = URigVMController::GetBuildData())
 	{
-		if(FunctionLibrary != nullptr)
+		if(URigVMFunctionLibrary* FunctionLibrary = RigVMClient.GetFunctionLibrary())
 		{
 			// for backwards compatibility load the function references from the
 			// model's storage over to the centralized build data
@@ -759,6 +1116,13 @@
 	{
 		return;
 	}
+
+	UControlRigBlueprintGeneratedClass* RigClass = GetControlRigBlueprintGeneratedClass();
+	if(RigClass == nullptr)
+	{
+		return;
+	}
+
 	TGuardValue<bool> CompilingGuard(bIsCompiling, true);
 	
 	bErrorsDuringCompilation = false;
@@ -766,7 +1130,6 @@
 	RigGraphDisplaySettings.MinMicroSeconds = RigGraphDisplaySettings.LastMinMicroSeconds = DBL_MAX;
 	RigGraphDisplaySettings.MaxMicroSeconds = RigGraphDisplaySettings.LastMaxMicroSeconds = (double)INDEX_NONE;
 
-	UControlRigBlueprintGeneratedClass* RigClass = GetControlRigBlueprintGeneratedClass();
 	UControlRig* CDO = Cast<UControlRig>(RigClass->GetDefaultObject(true /* create if needed */));
 	if (CDO->VM != nullptr)
 	{
@@ -804,7 +1167,7 @@
 
 		URigVMCompiler* Compiler = URigVMCompiler::StaticClass()->GetDefaultObject<URigVMCompiler>();
 		Compiler->Settings = (bCompileInDebugMode) ? FRigVMCompileSettings::Fast() : VMCompileSettings;
-		Compiler->Compile(Model, GetOrCreateController(), CDO->VM, CDO->GetExternalVariablesImpl(false), UserData, &PinToOperandMap);
+		Compiler->Compile(RigVMClient.GetAllModels(false, false), GetOrCreateController(), CDO->VM, CDO->GetExternalVariablesImpl(false), UserData, &PinToOperandMap);
 
 		if (bErrorsDuringCompilation)
 		{
@@ -889,9 +1252,7 @@
 {
 	TGuardValue<bool> IsCompilingGuard(bIsCompiling, true);
 	
-	TArray<URigVMGraph*> GraphsToDetach;
-	GraphsToDetach.Add(GetModel());
-	GraphsToDetach.Add(GetLocalFunctionLibrary());
+	TArray<URigVMGraph*> GraphsToDetach = RigVMClient.GetAllModels(true, false);
 
 	if (ensure(IsInGameThread()))
 	{
@@ -1085,22 +1446,29 @@
 	URigVMLibraryNode* FunctionNode = nullptr;
 	
 	// Find the node in the graph
-	URigVMNode* BreakpointNode = GetModel()->FindNode(InBreakpointNodePath);
-	if (BreakpointNode == nullptr)
+	for(const URigVMGraph* Model : RigVMClient)
 	{
-		// If we cannot find the node, it might be because it is inside a function
-		FString FunctionName = InBreakpointNodePath, Right;
-		URigVMNode::SplitNodePathAtStart(InBreakpointNodePath, FunctionName, Right);
+		URigVMNode* BreakpointNode = Model->FindNode(InBreakpointNodePath);
+		if (BreakpointNode == nullptr)
+		{
+			// If we cannot find the node, it might be because it is inside a function
+			FString FunctionName = InBreakpointNodePath, Right;
+			URigVMNode::SplitNodePathAtStart(InBreakpointNodePath, FunctionName, Right);
 
-		// Look inside the local function library
-		if (URigVMLibraryNode* LibraryNode = GetLocalFunctionLibrary()->FindFunction(FName(FunctionName)))
+			// Look inside the local function library
+			if (URigVMLibraryNode* LibraryNode = GetLocalFunctionLibrary()->FindFunction(FName(FunctionName)))
+			{
+				BreakpointNode = LibraryNode->GetContainedGraph()->FindNode(Right);
+				FunctionNode = LibraryNode;
+			}
+		}
+
+		if(BreakpointNode)
 		{
-			BreakpointNode = LibraryNode->GetContainedGraph()->FindNode(Right);
-			FunctionNode = LibraryNode;
+			return AddBreakpoint(BreakpointNode, FunctionNode);
 		}
 	}
-
-	return AddBreakpoint(BreakpointNode, FunctionNode);
+	return false;
 }
 
 bool UControlRigBlueprint::AddBreakpoint(URigVMNode* InBreakpointNode, URigVMLibraryNode* LibraryNode)
@@ -1174,28 +1542,31 @@
 	if (CDO && ByteCode)
 	{
 		FRigVMInstructionArray Instructions = ByteCode->GetInstructions();	
-
+		
 		// For each instruction, see if the node is in the callpath
 		// Only add one breakpoint for each callpath related to this node (i.e. if a node produces multiple
 		// instructions, only add a breakpoint to the first instruction)
 		for (int32 i = 0; i< Instructions.Num(); ++i)
 		{
-			const FRigVMASTProxy Proxy = FRigVMASTProxy::MakeFromCallPath(ByteCode->GetCallPathForInstruction(i), GetModel());
-			if (Proxy.GetCallstack().Contains(InBreakpointNode))
+			for(URigVMGraph* Model : RigVMClient)
 			{
-				// Find the callpath related to the breakpoint node
-				FRigVMASTProxy BreakpointProxy = Proxy;
-				while(BreakpointProxy.GetSubject() != InBreakpointNode)
+				const FRigVMASTProxy Proxy = FRigVMASTProxy::MakeFromCallPath(ByteCode->GetCallPathForInstruction(i), Model);
+				if (Proxy.GetCallstack().Contains(InBreakpointNode))
 				{
-					BreakpointProxy = BreakpointProxy.GetParent();
-				}
-				const FString& BreakpointCallPath = BreakpointProxy.GetCallstack().GetCallPath();
+					// Find the callpath related to the breakpoint node
+					FRigVMASTProxy BreakpointProxy = Proxy;
+					while(BreakpointProxy.GetSubject() != InBreakpointNode)
+					{
+						BreakpointProxy = BreakpointProxy.GetParent();
+					}
+					const FString& BreakpointCallPath = BreakpointProxy.GetCallstack().GetCallPath();
 
-				// Only add this callpath breakpoint once
-				if (!AddedCallpaths.Contains(BreakpointCallPath))
-				{
-					AddedCallpaths.Add(BreakpointCallPath);
-					CDO->AddBreakpoint(i, InBreakpointNode, BreakpointProxy.GetCallstack().Num());
+					// Only add this callpath breakpoint once
+					if (!AddedCallpaths.Contains(BreakpointCallPath))
+					{
+						AddedCallpaths.Add(BreakpointCallPath);
+						CDO->AddBreakpoint(i, InBreakpointNode, BreakpointProxy.GetCallstack().Num());
+					}
 				}
 			}
 		}
@@ -1213,30 +1584,44 @@
 bool UControlRigBlueprint::RemoveBreakpoint(const FString& InBreakpointNodePath)
 {
 	// Find the node in the graph
-	URigVMNode* BreakpointNode = GetModel()->FindNode(InBreakpointNodePath);
-	if (BreakpointNode == nullptr)
+	URigVMNode* BreakpointNode = nullptr;
+	
+	for(URigVMGraph* Model : RigVMClient)
 	{
-		// If we cannot find the node, it might be because it is inside a function
-		FString FunctionName = InBreakpointNodePath, Right;
-		URigVMNode::SplitNodePathAtStart(InBreakpointNodePath, FunctionName, Right);
+		BreakpointNode = Model->FindNode(InBreakpointNodePath);
+		if (BreakpointNode == nullptr)
+		{
+			// If we cannot find the node, it might be because it is inside a function
+			FString FunctionName = InBreakpointNodePath, Right;
+			URigVMNode::SplitNodePathAtStart(InBreakpointNodePath, FunctionName, Right);
 
-		// Look inside the local function library
-		if (URigVMLibraryNode* LibraryNode = GetLocalFunctionLibrary()->FindFunction(FName(FunctionName)))
+			// Look inside the local function library
+			if (URigVMLibraryNode* LibraryNode = GetLocalFunctionLibrary()->FindFunction(FName(FunctionName)))
+			{
+				BreakpointNode = LibraryNode->GetContainedGraph()->FindNode(Right);
+			}
+		}
+		if(BreakpointNode)
 		{
-			BreakpointNode = LibraryNode->GetContainedGraph()->FindNode(Right);
+			break;
 		}
 	}
 
-	bool bSuccess = RemoveBreakpoint(BreakpointNode);
+	if(BreakpointNode)
+	{
+		bool bSuccess = RemoveBreakpoint(BreakpointNode);
 
-	// Remove the breakpoint from all the loaded dependent blueprints
-	TArray<UControlRigBlueprint*> DependentBlueprints = GetDependentBlueprints(true, true);
-	DependentBlueprints.Remove(this);
-	for (UControlRigBlueprint* Dependent : DependentBlueprints)
-	{
-		bSuccess &= Dependent->RemoveBreakpoint(BreakpointNode);
+		// Remove the breakpoint from all the loaded dependent blueprints
+		TArray<UControlRigBlueprint*> DependentBlueprints = GetDependentBlueprints(true, true);
+		DependentBlueprints.Remove(this);
+		for (UControlRigBlueprint* Dependent : DependentBlueprints)
+		{
+			bSuccess &= Dependent->RemoveBreakpoint(BreakpointNode);
+		}
+		return bSuccess;
 	}
-	return bSuccess;
+
+	return false;
 }
 
 bool UControlRigBlueprint::RemoveBreakpoint(URigVMNode* InBreakpointNode)
@@ -1273,7 +1658,7 @@
 {
 	TArray<FRigVMReferenceNodeData> Data;
 	
-	TArray<URigVMGraph*> AllModels = GetAllModels();
+	const TArray<URigVMGraph*> AllModels = GetAllModels();
 	for (URigVMGraph* ModelToVisit : AllModels)
 	{
 		for(URigVMNode* Node : ModelToVisit->GetNodes())
@@ -1310,7 +1695,7 @@
 {
 	if (InEdGraph == nullptr)
 	{
-		return Model;
+		return GetDefaultModel();
 	}
 
 	if(InEdGraph->GetOutermost() != GetOutermost())
@@ -1321,127 +1706,60 @@
 #if WITH_EDITORONLY_DATA
 	if (InEdGraph == FunctionLibraryEdGraph)
 	{
-		return FunctionLibrary;
+		return RigVMClient.GetFunctionLibrary();
 	}
 #endif
 
 	const UControlRigGraph* RigGraph = Cast< UControlRigGraph>(InEdGraph);
 	check(RigGraph);
-
-	FString ModelNodePath = RigGraph->ModelNodePath;
-
-	if (RigGraph->bIsFunctionDefinition)
-	{
-		if (URigVMLibraryNode* LibraryNode = FunctionLibrary->FindFunction(*ModelNodePath))
-		{
-			return LibraryNode->GetContainedGraph();
-		}
-	}
-
-	if (RigGraph->GetOuter() == this)
-	{
-		return Model;
-	}
-
-	ensure(!ModelNodePath.IsEmpty());
-
-	URigVMGraph* SubModel = Model;
-	if (ModelNodePath.StartsWith(FunctionLibrary->GetNodePath()))
-	{
-		SubModel = FunctionLibrary;
-		ModelNodePath = ModelNodePath.Right(ModelNodePath.Len() - FunctionLibrary->GetNodePath().Len() - 1);
-	}
-
-	while (!ModelNodePath.IsEmpty())
-	{
-		FString NodeName = ModelNodePath;
-		if (NodeName.Contains(TEXT("|")))
-		{
-			NodeName = NodeName.Left(NodeName.Find(TEXT("|")));
-			ModelNodePath = ModelNodePath.Right(ModelNodePath.Len() - NodeName.Len() - 1);
-		}
-		else
-		{
-			ModelNodePath.Reset();
-		}
-
-		URigVMCollapseNode* CollapseNode = Cast<URigVMCollapseNode>(SubModel->FindNodeByName(*NodeName));
-		if (CollapseNode == nullptr)
-		{
-			return nullptr;
-		}
-
-		SubModel = CollapseNode->GetContainedGraph();
-	}
-
-	return SubModel;
+	return GetModel(RigGraph->ModelNodePath);
 }
 
 URigVMGraph* UControlRigBlueprint::GetModel(const FString& InNodePath) const
 {
-	if (!InNodePath.IsEmpty())
-	{
-		if (URigVMLibraryNode* LibraryNode = Cast<URigVMLibraryNode>(Model->FindNode(InNodePath)))
-		{
-			return LibraryNode->GetContainedGraph();
-		}
-
-		if(FunctionLibrary)
-		{
-			FString Left, Right;
-			if(URigVMNode::SplitNodePathAtStart(InNodePath, Left, Right))
-			{
-				if(Left == FunctionLibrary->GetNodePath())
-				{
-					if (URigVMLibraryNode* LibraryNode = Cast<URigVMLibraryNode>(FunctionLibrary->FindNode(Right)))
-					{
-						return LibraryNode->GetContainedGraph();
-					}
-				}
-			}
-		}
-		
-		return nullptr;
-	}
-	return Model;
+	return RigVMClient.GetModel(InNodePath);
 }
 
+URigVMGraph* UControlRigBlueprint::GetDefaultModel() const
+{
+	return RigVMClient.GetDefaultModel();
+}
 
 TArray<URigVMGraph*> UControlRigBlueprint::GetAllModels() const
 {
-	TArray<URigVMGraph*> Models;
-	Models.Add(GetModel());
-	Models.Append(GetModel()->GetContainedGraphs(true /* recursive */));
-	Models.Add(GetLocalFunctionLibrary());
-	Models.Append(GetLocalFunctionLibrary()->GetContainedGraphs(true /* recursive */));
-	return Models;
+	return RigVMClient.GetAllModels(true, true);
 }
 
 URigVMFunctionLibrary* UControlRigBlueprint::GetLocalFunctionLibrary() const
 {
-	return FunctionLibrary;
+	return RigVMClient.GetFunctionLibrary();
 }
 
+URigVMGraph* UControlRigBlueprint::AddModel(FString InName, bool bSetupUndoRedo, bool bPrintPythonCommand)
+{
+	const FString DesiredName = FString::Printf(TEXT("%s %s"),
+    	UControlRigBlueprint::RigVMModelPrefix, *InName);
+
+	TGuardValue<bool> EnablePythonPrint(bSuspendPythonMessagesForRigVMClient, !bPrintPythonCommand);
+	return RigVMClient.AddModel(*DesiredName, bSetupUndoRedo);
+}
+
+bool UControlRigBlueprint::RemoveModel(FString InName, bool bSetupUndoRedo, bool bPrintPythonCommand)
+{
+	TGuardValue<bool> EnablePythonPrint(bSuspendPythonMessagesForRigVMClient, !bPrintPythonCommand);
+	return RigVMClient.RemoveModel(InName, bSetupUndoRedo);
+}
+
 URigVMController* UControlRigBlueprint::GetController(const URigVMGraph* InGraph) const
 {
-	if (InGraph == nullptr)
-	{
-		InGraph = Model;
-	}
-
-	TObjectPtr<URigVMController> const* ControllerPtr = Controllers.Find(InGraph);
-	if (ControllerPtr)
-	{
-		return *ControllerPtr;
-	}
-	return nullptr;
+	return RigVMClient.GetController(InGraph);
 }
 
 URigVMController* UControlRigBlueprint::GetControllerByName(const FString InGraphName) const
 {
 	for (URigVMGraph* Graph : GetAllModels())
 	{
-		if (Graph->GetGraphName() == InGraphName)
+		if (Graph->GetName() == InGraphName || Graph->GetGraphName() == InGraphName)
 		{
 			return GetController(Graph);
 		}
@@ -1452,251 +1770,17 @@
 
 URigVMController* UControlRigBlueprint::GetOrCreateController(URigVMGraph* InGraph)
 {
-	if (URigVMController* ExistingController = GetController(InGraph))
-	{
-		return ExistingController;
-	}
-
-	if (InGraph == nullptr)
-	{
-		InGraph = Model;
-	}
-
-	URigVMController* Controller = NewObject<URigVMController>(this);
-	Controller->SetGraph(InGraph);
-	Controller->OnModified().AddUObject(this, &UControlRigBlueprint::HandleModifiedEvent);
-
-	Controller->UnfoldStructDelegate.BindLambda([](const UStruct* InStruct) -> bool {
-
-		if (InStruct == TBaseStructure<FQuat>::Get())
-		{
-			return false;
-		}
-		if (InStruct == FRuntimeFloatCurve::StaticStruct())
-		{
-			return false;
-		}
-		if (InStruct == FRigPose::StaticStruct())
-		{
-			return false;
-		}
-		return true;
-		});
-
-	TWeakObjectPtr<UControlRigBlueprint> WeakThis(this);
-
-	// this delegate is used by the controller to determine variable validity
-	// during a bind process. the controller itself doesn't own the variables,
-	// so we need a delegate to request them from the owning blueprint
-	Controller->GetExternalVariablesDelegate.BindLambda([](URigVMGraph* InGraph) -> TArray<FRigVMExternalVariable> {
-
-		if (InGraph)
-		{
-			if(UControlRigBlueprint* Blueprint = InGraph->GetTypedOuter<UControlRigBlueprint>())
-			{
-				if (UControlRigBlueprintGeneratedClass* RigClass = Blueprint->GetControlRigBlueprintGeneratedClass())
-				{
-                    if (UControlRig* CDO = Cast<UControlRig>(RigClass->GetDefaultObject(true /* create if needed */)))
-                    {
-                        return CDO->GetExternalVariablesImpl(true /* rely on variables within blueprint */);
-                    }
-                }
-			}
-		}
-		return TArray<FRigVMExternalVariable>();
-
-	});
-
-
-	// this delegate is used by the controller to retrieve the current bytecode of the VM
-	Controller->GetCurrentByteCodeDelegate.BindLambda([WeakThis]() -> const FRigVMByteCode* {
-
-		if (WeakThis.IsValid())
-		{
-			if (UControlRigBlueprintGeneratedClass* RigClass = WeakThis->GetControlRigBlueprintGeneratedClass())
-			{
-				if (UControlRig* CDO = Cast<UControlRig>(RigClass->GetDefaultObject(false)))
-				{
-					if (CDO->VM)
-					{
-						return &CDO->VM->GetByteCode();
-					}
-				}
-			}
-		}
-		return nullptr;
-
-	});
-
-	Controller->IsFunctionAvailableDelegate.BindLambda([WeakThis](URigVMLibraryNode* InFunction) -> bool
-	{
-		if(InFunction == nullptr)
-		{
-			return false;
-		}
-		
-		if(URigVMFunctionLibrary* Library = Cast<URigVMFunctionLibrary>(InFunction->GetOuter()))
-		{
-			if(UControlRigBlueprint* Blueprint = Cast<UControlRigBlueprint>(Library->GetOuter()))
-			{
-				if(Blueprint->IsFunctionPublic(InFunction->GetFName()))
-				{
-					return true;
-				}
-
-				// if it is private - we still see it as public if we are within the same blueprint 
-				if(WeakThis.IsValid())
-				{
-					if(WeakThis.Get() == Blueprint)
-					{
-						return true;
-					}
-				}
-			}
-		}
-
-		return false;
-	});
-
-	Controller->IsDependencyCyclicDelegate.BindLambda([WeakThis](UObject* InDependentObject, UObject* InDependencyObject) -> bool
-	{
-	    if(InDependentObject == nullptr || InDependencyObject == nullptr)
-	    {
-	        return false;
-	    }
-
-		UControlRigBlueprint* DependentBlueprint = InDependentObject->GetTypedOuter<UControlRigBlueprint>();
-		UControlRigBlueprint* DependencyBlueprint = InDependencyObject->GetTypedOuter<UControlRigBlueprint>();
-
-		if(DependentBlueprint == nullptr || DependencyBlueprint == nullptr)
-		{
-			return false;
-		}
-
-		if(DependentBlueprint == DependencyBlueprint)
-		{
-			return false;
-		}
-
-		const TArray<UControlRigBlueprint*> DependencyDependencies = DependencyBlueprint->GetDependencies(true);
-		return DependencyDependencies.Contains(DependentBlueprint);
-	});
-
-#if WITH_EDITOR
-
-	// this sets up three delegates:
-	// a) get external variables (mapped to Controller->GetExternalVariables)
-	// b) bind pin to variable (mapped to Controller->BindPinToVariable)
-	// c) create external variable (mapped to the passed in tfunction)
-	// the last one is defined within the blueprint since the controller
-	// doesn't own the variables and can't create one itself.
-	Controller->SetupDefaultUnitNodeDelegates(TDelegate<FName(FRigVMExternalVariable, FString)>::CreateLambda(
-		[WeakThis](FRigVMExternalVariable InVariableToCreate, FString InDefaultValue) -> FName {
-			if (WeakThis.IsValid())
-			{
-				return WeakThis->AddCRMemberVariableFromExternal(InVariableToCreate, InDefaultValue);
-			}
-			return NAME_None;
-		}
-	));
-
-	TWeakObjectPtr<URigVMController> WeakController = Controller;
-	Controller->RequestBulkEditDialogDelegate.BindLambda([WeakThis, WeakController](URigVMLibraryNode* InFunction, ERigVMControllerBulkEditType InEditType) -> FRigVMController_BulkEditResult 
-	{
-		if(WeakThis.IsValid() && WeakController.IsValid())
-		{
-			UControlRigBlueprint* StrongThis = WeakThis.Get();
-            URigVMController* StrongController = WeakController.Get();
-            if(StrongThis->OnRequestBulkEditDialog().IsBound())
-			{
-				return StrongThis->OnRequestBulkEditDialog().Execute(StrongThis, StrongController, InFunction, InEditType);
-			}
-		}
-		return FRigVMController_BulkEditResult();
-	});
-
-	Controller->RequestBreakLinksDialogDelegate.BindLambda([WeakThis, WeakController](TArray<URigVMLink*> InLinks) -> bool 
-	{
-		if(WeakThis.IsValid() && WeakController.IsValid())
-		{
-			UControlRigBlueprint* StrongThis = WeakThis.Get();
-			if(StrongThis->OnRequestBreakLinksDialog().IsBound())
-			{
-				return StrongThis->OnRequestBreakLinksDialog().Execute(InLinks);
-			}
-		}
-		return false;
-	});
-
-	Controller->RequestNewExternalVariableDelegate.BindLambda([WeakThis](FRigVMGraphVariableDescription InVariable, bool bInIsPublic, bool bInIsReadOnly) -> FName
-	{
-		if (WeakThis.IsValid())
-		{
-			for (FBPVariableDescription& ExistingVariable : WeakThis->NewVariables)
-			{
-				if (ExistingVariable.VarName == InVariable.Name)
-				{
-					return FName();
-				}
-			}
-
-			FRigVMExternalVariable ExternalVariable = InVariable.ToExternalVariable();
-			return WeakThis->AddMemberVariable(InVariable.Name,
-				ExternalVariable.TypeObject ? ExternalVariable.TypeObject->GetPathName() : ExternalVariable.TypeName.ToString(),
-				bInIsPublic,
-				bInIsReadOnly,
-				InVariable.DefaultValue);
-		}
-		
-		return FName();
-	});
-
-
-	Controller->RequestJumpToHyperlinkDelegate.BindLambda([WeakThis](const UObject* InSubject)
-	{
-		if (WeakThis.IsValid())
-		{
-			UControlRigBlueprint* StrongThis = WeakThis.Get();
-			if(StrongThis->OnRequestJumpToHyperlink().IsBound())
-			{
-				StrongThis->OnRequestJumpToHyperlink().Execute(InSubject);
-			}
-		}
-	});
-
-	Controller->ConfigureWorkflowOptionsDelegate.BindLambda([WeakThis](URigVMUserWorkflowOptions* Options)
-	{
-		if(UControlRigWorkflowOptions* ControlRigNodeWorkflowOptions = Cast<UControlRigWorkflowOptions>(Options))
-		{
-			ControlRigNodeWorkflowOptions->Hierarchy = nullptr;
-			ControlRigNodeWorkflowOptions->Selection.Reset();
-			
-			if(const UControlRigBlueprint* StrongThis = WeakThis.Get())
-			{
-				if(UControlRig* ControlRig = Cast<UControlRig>(StrongThis->GetObjectBeingDebugged()))
-				{
-					ControlRigNodeWorkflowOptions->Hierarchy = ControlRig->GetHierarchy();
-				}
-				ControlRigNodeWorkflowOptions->Selection = StrongThis->Hierarchy->GetSelectedKeys();
-			}
-		}
-	});
-
-#endif
-
-	Controller->RemoveStaleNodes();
-	Controllers.Add(InGraph, Controller);
-	return Controller;
+	return RigVMClient.GetOrCreateController(InGraph);
 }
 
 URigVMController* UControlRigBlueprint::GetController(const UEdGraph* InEdGraph) const
 {
-	return GetController(GetModel(InEdGraph));
+	return RigVMClient.GetController(InEdGraph);
 }
 
 URigVMController* UControlRigBlueprint::GetOrCreateController(const UEdGraph* InEdGraph)
 {
-	return GetOrCreateController(GetModel(InEdGraph));
+	return RigVMClient.GetOrCreateController(InEdGraph);
 }
 
 TArray<FString> UControlRigBlueprint::GeneratePythonCommands(const FString InNewBlueprintName)
@@ -1900,6 +1984,11 @@
 					}
 				}
 			}
+			else if(Graph != GetDefaultModel())
+			{
+				InternalCommands.Add(FString::Printf(TEXT("blueprint.add_model('%s')"),
+						*Graph->GetName()));
+			}
 								
 			InternalCommands.Append(Controller->GeneratePythonCommands());
 		}
@@ -1971,48 +2060,7 @@
 
 UEdGraph* UControlRigBlueprint::GetEdGraph(URigVMGraph* InModel) const
 {
-	if (InModel == nullptr)
-	{
-		return nullptr;
-	}
-
-	if(InModel->GetOutermost() != GetOutermost())
-	{
-		return nullptr;
-	}
-
-#if WITH_EDITORONLY_DATA
-	if (InModel == FunctionLibrary)
-	{
-		return FunctionLibraryEdGraph;
-	}
-#endif
-
-	TArray<UEdGraph*> EdGraphs;
-	GetAllGraphs(EdGraphs);
-
-	bool bIsFunctionDefinition = false;
-	if (URigVMLibraryNode* LibraryNode = Cast<URigVMLibraryNode>(InModel->GetOuter()))
-	{
-		bIsFunctionDefinition = LibraryNode->GetGraph()->IsA<URigVMFunctionLibrary>();
-	}
-
-	for (UEdGraph* EdGraph : EdGraphs)
-	{
-		if (UControlRigGraph* RigGraph = Cast<UControlRigGraph>(EdGraph))
-		{
-			if (RigGraph->bIsFunctionDefinition != bIsFunctionDefinition)
-			{
-				continue;
-			}
-
-			if (RigGraph->ModelNodePath == InModel->GetNodePath())
-			{
-				return RigGraph;
-			}
-		}
-	}
-	return nullptr;
+	return Cast<UEdGraph>(GetEditorObjectForRigVMGraph(InModel));
 }
 
 UEdGraph* UControlRigBlueprint::GetEdGraph(const FString& InNodePath) const
@@ -2123,7 +2171,7 @@
 	TArray<FAssetData> Dependents;
 	TArray<FName> AssetPaths;
 
-	if(FunctionLibrary)
+	if(URigVMFunctionLibrary* FunctionLibrary = RigVMClient.GetFunctionLibrary())
 	{
 		const FAssetRegistryModule& AssetRegistryModule = FModuleManager::LoadModuleChecked<FAssetRegistryModule>(TEXT("AssetRegistry"));
 		TArray<URigVMLibraryNode*> Functions = FunctionLibrary->GetFunctions();
@@ -2232,7 +2280,7 @@
 	if (TransactionEvent.GetEventType() == ETransactionObjectEventType::UndoRedo)
 	{
 		TArray<FName> PropertiesChanged = TransactionEvent.GetChangedProperties();
-		if (PropertiesChanged.Contains(TEXT("HierarchyContainer")))
+		if (PropertiesChanged.Contains(GET_MEMBER_NAME_CHECKED(UControlRigBlueprint, Hierarchy)))
 		{
 			int32 TransactionIndex = GEditor->Trans->FindTransactionIndex(TransactionEvent.GetTransactionId());
 			const FTransaction* Transaction = GEditor->Trans->GetTransaction(TransactionIndex);
@@ -2260,17 +2308,17 @@
 			MarkPackageDirty();
 		}
 
-		if (PropertiesChanged.Contains(TEXT("DrawContainer")))
+		if (PropertiesChanged.Contains(GET_MEMBER_NAME_CHECKED(UControlRigBlueprint, DrawContainer)))
 		{
 			PropagateDrawInstructionsFromBPToInstances();
 		}
 
-		if (PropertiesChanged.Contains(TEXT("VMRuntimeSettings")))
+		if (PropertiesChanged.Contains(GET_MEMBER_NAME_CHECKED(UControlRigBlueprint, VMRuntimeSettings)))
 		{
 			PropagateRuntimeSettingsFromBPToInstances();
 		}
 
-		if (PropertiesChanged.Contains(TEXT("NewVariables")))
+		if (PropertiesChanged.Contains(GET_MEMBER_NAME_CHECKED(UControlRigBlueprint, NewVariables)))
 		{
 			if (RefreshEditorEvent.IsBound())
 			{
@@ -2278,6 +2326,15 @@
 			}
 			MarkPackageDirty();			
 		}
+
+		if (PropertiesChanged.Contains(GET_MEMBER_NAME_CHECKED(UControlRigBlueprint, RigVMClient)) ||
+			PropertiesChanged.Contains(GET_MEMBER_NAME_CHECKED(UControlRigBlueprint, UbergraphPages)))
+		{
+			RigVMClient.PostTransacted(TransactionEvent);
+
+			RecompileVM();
+			MarkPackageDirty();			
+		}
 	}
 }
 
@@ -3009,16 +3066,19 @@
 
 void UControlRigBlueprint::SetupPinRedirectorsForBackwardsCompatibility()
 {
-	for (URigVMNode* Node : Model->GetNodes())
+	for(URigVMGraph* Model : RigVMClient)
 	{
-		if (URigVMUnitNode* UnitNode = Cast<URigVMUnitNode>(Node))
+		for (URigVMNode* Node : Model->GetNodes())
 		{
-			UScriptStruct* Struct = UnitNode->GetScriptStruct();
-			if (Struct == FRigUnit_SetBoneTransform::StaticStruct())
+			if (URigVMUnitNode* UnitNode = Cast<URigVMUnitNode>(Node))
 			{
-				URigVMPin* TransformPin = UnitNode->FindPin(TEXT("Transform"));
-				URigVMPin* ResultPin = UnitNode->FindPin(TEXT("Result"));
-				GetOrCreateController()->AddPinRedirector(false, true, TransformPin->GetPinPath(), ResultPin->GetPinPath());
+				UScriptStruct* Struct = UnitNode->GetScriptStruct();
+				if (Struct == FRigUnit_SetBoneTransform::StaticStruct())
+				{
+					URigVMPin* TransformPin = UnitNode->FindPin(TEXT("Transform"));
+					URigVMPin* ResultPin = UnitNode->FindPin(TEXT("Result"));
+					GetOrCreateController()->AddPinRedirector(false, true, TransformPin->GetPinPath(), ResultPin->GetPinPath());
+				}
 			}
 		}
 	}
@@ -3051,12 +3111,12 @@
 		}
 	}
 
-	TArray<URigVMGraph*> RigGraphs;
-	RigGraphs.Add(GetModel());
-	RigGraphs.Add(GetLocalFunctionLibrary());
-
-	GetOrCreateController(RigGraphs[0])->ResendAllNotifications();
-	GetOrCreateController(RigGraphs[1])->ResendAllNotifications();
+	TArray<URigVMGraph*> RigGraphs = RigVMClient.GetAllModels(true, true);
 
 	for (int32 RigGraphIndex = 0; RigGraphIndex < RigGraphs.Num(); RigGraphIndex++)
 	{
+		GetOrCreateController(RigGraphs[RigGraphIndex])->ResendAllNotifications();
+	}
+
+	for (int32 RigGraphIndex = 0; RigGraphIndex < RigGraphs.Num(); RigGraphIndex++)
+	{
@@ -3067,7 +3127,6 @@
 			if (URigVMCollapseNode* CollapseNode = Cast<URigVMCollapseNode>(RigNode))
 			{
 				CreateEdGraphForCollapseNodeIfNeeded(CollapseNode, true);
-				RigGraphs.Add(CollapseNode->GetContainedGraph());
 			}
 		}
 	}
@@ -3216,6 +3275,7 @@
 						{
 							TArray<UObject*> ArchetypeInstances;
 							DefaultObject->GetArchetypeInstances(ArchetypeInstances);
+							
 							for (UObject* ArchetypeInstance : ArchetypeInstances)
 							{
 								if (UControlRig* InstanceRig = Cast<UControlRig>(ArchetypeInstance))
@@ -3226,25 +3286,27 @@
 										{
 											return true;
 										}
-										
-										if (URigVMPin* ControlledPin = Model->FindPin(ControlElement->GetName().ToString()))
+
+										for(URigVMGraph* Model : RigVMClient)
 										{
-											URigVMPin* ControlledPinForLink = ControlledPin->GetPinForLink();
+											if (URigVMPin* ControlledPin = Model->FindPin(ControlElement->GetName().ToString()))
+											{
+												URigVMPin* ControlledPinForLink = ControlledPin->GetPinForLink();
 
-											if(ControlledPin->GetRootPin() == Pin->GetRootPin() ||
-											   ControlledPinForLink->GetRootPin() == Pin->GetRootPin())
-											{
-												InstanceRig->SetTransientControlValue(ControlledPin->GetPinForLink());
+												if(ControlledPin->GetRootPin() == Pin->GetRootPin() ||
+												   ControlledPinForLink->GetRootPin() == Pin->GetRootPin())
+												{
+													InstanceRig->SetTransientControlValue(ControlledPin->GetPinForLink());
+												}
+												else if (ControlledPin->GetNode() == Pin->GetNode() ||
+														 ControlledPinForLink->GetNode() == Pin->GetNode())
+												{
+													InstanceRig->ClearTransientControls();
+													InstanceRig->AddTransientControl(ControlledPin);
+												}
+												return false;
 											}
-											else if (ControlledPin->GetNode() == Pin->GetNode() ||
-													 ControlledPinForLink->GetNode() == Pin->GetNode())
-											{
-												InstanceRig->ClearTransientControls();
-												InstanceRig->AddTransientControl(ControlledPin);
-											}
-											return false;
 										}
-
 										return true;
 									});
 								}
@@ -3265,15 +3327,15 @@
 						RemoveBreakpoint(RigVMNode);
 					}
 				}
-					
+
 				if (URigVMCollapseNode* CollapseNode = Cast<URigVMCollapseNode>(InSubject))
 				{
 					if (InNotifType == ERigVMGraphNotifType::NodeAdded)
 					{
 						// If the controller for this graph already exist, make sure it is referencing the correct graph
-						if (TObjectPtr<URigVMController>* Controller = Controllers.Find(CollapseNode->GetContainedGraph()))
+						if (URigVMController* Controller = RigVMClient.GetController(CollapseNode->GetContainedGraph()))
 						{
-							(*Controller).Get()->SetGraph(CollapseNode->GetContainedGraph());
+							Controller->SetGraph(CollapseNode->GetContainedGraph());
 						}
 						
 						CreateEdGraphForCollapseNodeIfNeeded(CollapseNode);
@@ -3306,6 +3368,24 @@
 					break;
 				}
 
+				if (URigVMNode* RigVMNode = Cast<URigVMNode>(InSubject))
+				{
+					if(RigVMNode->IsEvent() && RigVMNode->GetGraph()->IsRootGraph())
+					{
+						// let the UI know the title for the graph may have changed.
+						RigVMClient.NotifyOuterOfPropertyChange();
+
+						if(UControlRigGraph* EdGraph = Cast<UControlRigGraph>(GetEdGraph(RigVMNode->GetGraph())))
+						{
+							// decide if this graph should be renameable
+							const int32 NumberOfEvents = Algo::CountIf(RigVMNode->GetGraph()->GetNodes(), [](const URigVMNode* NodeToCount) -> bool
+							{
+								return NodeToCount->IsEvent();
+							});
+							EdGraph->bAllowRenaming = NumberOfEvents != 1;
+						}
+					}
+				}
 				// fall through to the next case
 			}
 			case ERigVMGraphNotifType::LinkAdded:
@@ -3332,7 +3412,8 @@
 					URigVMPin* Pin = CastChecked<URigVMPin>(InSubject)->GetRootPin(); 
 					URigVMCompiler* Compiler = URigVMCompiler::StaticClass()->GetDefaultObject<URigVMCompiler>();
 					Compiler->Settings = VMCompileSettings;
-					TSharedPtr<FRigVMParserAST> RuntimeAST = Model->GetRuntimeAST();
+
+					TSharedPtr<FRigVMParserAST> RuntimeAST = GetDefaultModel()->GetRuntimeAST();
 					
 					if(Pin->RequiresWatch())
 					{
@@ -3418,11 +3499,7 @@
 						FString OldSubPathString = SubPathString.Replace(*CollapseNode->GetName(), *CollapseNode->GetPreviousFName().ToString());
 						FSoftObjectPath OldContainedGraphPath = ContainedGraphPath;
 						OldContainedGraphPath.SetSubPathString(OldSubPathString);
-						if (TObjectPtr<URigVMController>* Controller = Controllers.Find(OldContainedGraphPath))
-						{
-							Controllers.Add(ContainedGraphPath, *Controller);
-							Controllers.Remove(OldContainedGraphPath);
-						}
+						RigVMClient.OnSubGraphRenamed(OldContainedGraphPath, ContainedGraphPath);
 					}
 					
 					FString NewNodePath = CollapseNode->GetNodePath(true /* recursive */);
@@ -3434,26 +3511,7 @@
 						OldNodePath = URigVMNode::JoinNodePath(Left, OldNodePath);
 					}
 
-					FString NewNodePathPrefix = NewNodePath + TEXT("|");
-					FString OldNodePathPrefix = OldNodePath + TEXT("|");
+					HandleRigVMGraphRenamed(GetRigVMClient(), OldNodePath, NewNodePath);
 
-					TArray<UEdGraph*> EdGraphs;
-					GetAllGraphs(EdGraphs);
-
-					for (UEdGraph* EdGraph : EdGraphs)
-					{
-						if (UControlRigGraph* RigGraph = Cast<UControlRigGraph>(EdGraph))
-						{
-							if (RigGraph->ModelNodePath == OldNodePath)
-							{
-								RigGraph->ModelNodePath = NewNodePath;
-							}
-							else if (RigGraph->ModelNodePath.StartsWith(OldNodePathPrefix))
-							{
-								RigGraph->ModelNodePath = NewNodePathPrefix + RigGraph->ModelNodePath.RightChop(OldNodePathPrefix.Len());
-							}
-						}
-					}
-
 					if (UEdGraph* ContainedEdGraph = GetEdGraph(CollapseNode->GetContainedGraph()))
 					{
@@ -3542,7 +3600,7 @@
 		AddedMemberVariableMap.Add(NewVariables[VariableIndex].VarName, VariableIndex);
 	}
 
-	if (Model == nullptr)
+	if (RigVMClient.Num() == 0)
 	{
 		return;
 	}
@@ -3552,70 +3610,73 @@
 	{
 		TSharedPtr<FKismetNameValidator> NameValidator = MakeShareable(new FKismetNameValidator(this, NAME_None, nullptr));
 
-		TArray<URigVMNode*> Nodes = Model->GetNodes();
-		for (URigVMNode* Node : Nodes)
+		for(URigVMGraph* Model : RigVMClient)
 		{
-			if (URigVMVariableNode* VariableNode = Cast<URigVMVariableNode>(Node))
+			TArray<URigVMNode*> Nodes = Model->GetNodes();
+			for (URigVMNode* Node : Nodes)
 			{
-				if (URigVMPin* VariablePin = VariableNode->FindPin(TEXT("Variable")))
+				if (URigVMVariableNode* VariableNode = Cast<URigVMVariableNode>(Node))
 				{
-					if (VariablePin->GetDirection() != ERigVMPinDirection::Visible)
+					if (URigVMPin* VariablePin = VariableNode->FindPin(TEXT("Variable")))
+					{
+						if (VariablePin->GetDirection() != ERigVMPinDirection::Visible)
+						{
+							continue;
+						}
+					}
+
+					FRigVMGraphVariableDescription Description = VariableNode->GetVariableDescription();
+					if (AddedMemberVariableMap.Contains(Description.Name))
 					{
 						continue;
 					}
-				}
 
-				FRigVMGraphVariableDescription Description = VariableNode->GetVariableDescription();
-				if (AddedMemberVariableMap.Contains(Description.Name))
-				{
-					continue;
-				}
-
-				FEdGraphPinType PinType = RigVMTypeUtils::PinTypeFromExternalVariable(Description.ToExternalVariable());
-				if (!PinType.PinCategory.IsValid())
-				{
-					continue;
-				}
-
-				FName VarName = FindCRMemberVariableUniqueName(NameValidator, Description.Name.ToString());
-				int32 VariableIndex = AddCRMemberVariable(this, VarName, PinType, false, false, FString());
-				if (VariableIndex != INDEX_NONE)
-				{
-					AddedMemberVariableMap.Add(Description.Name, VariableIndex);
-					bDirtyDuringLoad = true;
-				}
-			}
-
-			// Leaving this for backwards compatibility, even though we don't support parameters anymore
-			// When a parameter node is found, we will create a variable
-			if (URigVMParameterNode* ParameterNode = Cast<URigVMParameterNode>(Node))
-			{
-				if (URigVMPin* ParameterPin = ParameterNode->FindPin(TEXT("Parameter")))
-				{
-					if (ParameterPin->GetDirection() != ERigVMPinDirection::Visible)
+					FEdGraphPinType PinType = RigVMTypeUtils::PinTypeFromExternalVariable(Description.ToExternalVariable());
+					if (!PinType.PinCategory.IsValid())
 					{
 						continue;
 					}
+
+					FName VarName = FindCRMemberVariableUniqueName(NameValidator, Description.Name.ToString());
+					int32 VariableIndex = AddCRMemberVariable(this, VarName, PinType, false, false, FString());
+					if (VariableIndex != INDEX_NONE)
+					{
+						AddedMemberVariableMap.Add(Description.Name, VariableIndex);
+						bDirtyDuringLoad = true;
+					}
 				}
 
-				FRigVMGraphParameterDescription Description = ParameterNode->GetParameterDescription();
-				if (AddedMemberVariableMap.Contains(Description.Name))
+				// Leaving this for backwards compatibility, even though we don't support parameters anymore
+				// When a parameter node is found, we will create a variable
+				if (URigVMParameterNode* ParameterNode = Cast<URigVMParameterNode>(Node))
 				{
-					continue;
-				}
+					if (URigVMPin* ParameterPin = ParameterNode->FindPin(TEXT("Parameter")))
+					{
+						if (ParameterPin->GetDirection() != ERigVMPinDirection::Visible)
+						{
+							continue;
+						}
+					}
 
-				FEdGraphPinType PinType = RigVMTypeUtils::PinTypeFromExternalVariable(Description.ToExternalVariable());
-				if (!PinType.PinCategory.IsValid())
-				{
-					continue;
-				}
+					FRigVMGraphParameterDescription Description = ParameterNode->GetParameterDescription();
+					if (AddedMemberVariableMap.Contains(Description.Name))
+					{
+						continue;
+					}
 
-				FName VarName = FindCRMemberVariableUniqueName(NameValidator, Description.Name.ToString());
-				int32 VariableIndex = AddCRMemberVariable(this, VarName, PinType, true, !Description.bIsInput, FString());
-				if (VariableIndex != INDEX_NONE)
-				{
-					AddedMemberVariableMap.Add(Description.Name, VariableIndex);
-					bDirtyDuringLoad = true;
+					FEdGraphPinType PinType = RigVMTypeUtils::PinTypeFromExternalVariable(Description.ToExternalVariable());
+					if (!PinType.PinCategory.IsValid())
+					{
+						continue;
+					}
+
+					FName VarName = FindCRMemberVariableUniqueName(NameValidator, Description.Name.ToString());
+					int32 VariableIndex = AddCRMemberVariable(this, VarName, PinType, true, !Description.bIsInput, FString());
+					if (VariableIndex != INDEX_NONE)
+					{
+						AddedMemberVariableMap.Add(Description.Name, VariableIndex);
+						bDirtyDuringLoad = true;
+					}
 				}
 			}
 		}
@@ -3722,40 +3783,42 @@
 	// If the asset was copied from one project to another, the function referenced might have a different
 	// path, even if the function is internal to the contorl rig. In that case, let's try to find the function
 	// in the local function library.
-	
-	TArray<URigVMNode*> Nodes = Model->GetNodes();
-	for (URigVMLibraryNode* Library : FunctionLibrary->GetFunctions())
+
+	for(URigVMGraph* Model : RigVMClient)
 	{
-		Nodes.Append(Library->GetContainedNodes());
-	}
-	
-	for (URigVMNode* Node : Nodes)
-	{
-		if (URigVMFunctionReferenceNode* FunctionReferenceNode = Cast<URigVMFunctionReferenceNode>(Node))
+		TArray<URigVMNode*> Nodes = Model->GetNodes();
+		for (URigVMLibraryNode* Library : RigVMClient.GetFunctionLibrary()->GetFunctions())
 		{
-			if (!FunctionReferenceNode->GetReferencedNode())
+			Nodes.Append(Library->GetContainedNodes());
+		}
+		
+		for (URigVMNode* Node : Nodes)
+		{
+			if (URigVMFunctionReferenceNode* FunctionReferenceNode = Cast<URigVMFunctionReferenceNode>(Node))
 			{
-				if(FunctionLibrary)
+				if (!FunctionReferenceNode->GetReferencedNode())
 				{
-					FString FunctionPath = FunctionReferenceNode->ReferencedNodePtr.ToSoftObjectPath().GetSubPathString();
-					
-					FString Left, Right;
-					if(FunctionPath.Split(TEXT("."), &Left, &Right))
+					if(URigVMFunctionLibrary* FunctionLibrary = RigVMClient.GetFunctionLibrary())
 					{
-						FString LibraryNodePath = FunctionLibrary->GetNodePath();
-						if(Left == FunctionLibrary->GetName())
+						FString FunctionPath = FunctionReferenceNode->ReferencedNodePtr.ToSoftObjectPath().GetSubPathString();
+						
+						FString Left, Right;
+						if(FunctionPath.Split(TEXT("."), &Left, &Right))
 						{
-							if (URigVMLibraryNode* LibraryNode = Cast<URigVMLibraryNode>(FunctionLibrary->FindNode(Right)))
+							FString LibraryNodePath = FunctionLibrary->GetNodePath();
+							if(Left == FunctionLibrary->GetName())
 							{
-								FunctionReferenceNode->SetReferencedNode(LibraryNode);
+								if (URigVMLibraryNode* LibraryNode = Cast<URigVMLibraryNode>(FunctionLibrary->FindNode(Right)))
+								{
+									FunctionReferenceNode->SetReferencedNode(LibraryNode);
+								}
 							}
 						}
-					}
-				}				
+					}				
+				}
 			}
 		}
 	}
-
 	FunctionReferenceNodeData = GetReferenceNodeData();
 }
 
@@ -3772,37 +3835,40 @@
 
 		GetOrCreateController()->ReattachLinksToPinObjects();
 
-		check(Model);
+		check(GetDefaultModel());
 
-		TArray<URigVMNode*> Nodes = Model->GetNodes();
-		for (URigVMNode* Node : Nodes)
+		for(URigVMGraph* Model : RigVMClient)
 		{
-			if (URigVMVariableNode* VariableNode = Cast<URigVMVariableNode>(Node))
+			TArray<URigVMNode*> Nodes = Model->GetNodes();
+			for (URigVMNode* Node : Nodes)
 			{
-				FRigVMGraphVariableDescription Description = VariableNode->GetVariableDescription();
-				if (!AddedMemberVariableMap.Contains(Description.Name))
+				if (URigVMVariableNode* VariableNode = Cast<URigVMVariableNode>(Node))
 				{
-					continue;
+					FRigVMGraphVariableDescription Description = VariableNode->GetVariableDescription();
+					if (!AddedMemberVariableMap.Contains(Description.Name))
+					{
+						continue;
+					}
+
+					int32 VariableIndex = AddedMemberVariableMap.FindChecked(Description.Name);
+					FName VarName = NewVariables[VariableIndex].VarName;
+					GetOrCreateController()->RefreshVariableNode(VariableNode->GetFName(), VarName, Description.CPPType, Description.CPPTypeObject, false);
+					bDirtyDuringLoad = true;
 				}
 
-				int32 VariableIndex = AddedMemberVariableMap.FindChecked(Description.Name);
-				FName VarName = NewVariables[VariableIndex].VarName;
-				GetOrCreateController()->RefreshVariableNode(VariableNode->GetFName(), VarName, Description.CPPType, Description.CPPTypeObject, false);
-				bDirtyDuringLoad = true;
+				if (URigVMParameterNode* ParameterNode = Cast<URigVMParameterNode>(Node))
+				{
+					FRigVMGraphParameterDescription Description = ParameterNode->GetParameterDescription();
+					if (!AddedMemberVariableMap.Contains(Description.Name))
+					{
+						continue;
+					}
+
+					int32 VariableIndex = AddedMemberVariableMap.FindChecked(Description.Name);
+					FName VarName = NewVariables[VariableIndex].VarName;
+					GetOrCreateController()->ReplaceParameterNodeWithVariable(ParameterNode->GetFName(), VarName, Description.CPPType, Description.CPPTypeObject, false);
+					bDirtyDuringLoad = true;
+				}
 			}
-
-			if (URigVMParameterNode* ParameterNode = Cast<URigVMParameterNode>(Node))
-			{
-				FRigVMGraphParameterDescription Description = ParameterNode->GetParameterDescription();
-				if (!AddedMemberVariableMap.Contains(Description.Name))
-				{
-					continue;
-				}
-
-				int32 VariableIndex = AddedMemberVariableMap.FindChecked(Description.Name);
-				FName VarName = NewVariables[VariableIndex].VarName;
-				GetOrCreateController()->ReplaceParameterNodeWithVariable(ParameterNode->GetFName(), VarName, Description.CPPType, Description.CPPTypeObject, false);
-				bDirtyDuringLoad = true;
-			}
 		}
 	}
@@ -4051,64 +4117,66 @@
 
 		GetOrCreateController()->ReattachLinksToPinObjects();
 
-		check(Model);
+		check(GetDefaultModel());
 
-		TArray<URigVMNode*> Nodes = Model->GetNodes();
-		for (URigVMNode* Node : Nodes)
+		for(URigVMGraph* Model : RigVMClient)
 		{
-			if (URigVMParameterNode* ParameterNode = Cast<URigVMParameterNode>(Node))
+			TArray<URigVMNode*> Nodes = Model->GetNodes();
+			for (URigVMNode* Node : Nodes)
 			{
-				FRigVMGraphParameterDescription Description = ParameterNode->GetParameterDescription();
+				if (URigVMParameterNode* ParameterNode = Cast<URigVMParameterNode>(Node))
+				{
+					FRigVMGraphParameterDescription Description = ParameterNode->GetParameterDescription();
 
-				const FEdGraphPinType PinType = RigVMTypeUtils::PinTypeFromCPPType(Description.CPPType, Description.CPPTypeObject);
+					const FEdGraphPinType PinType = RigVMTypeUtils::PinTypeFromCPPType(Description.CPPType, Description.CPPTypeObject);
 
-				int32 VariableIndex = INDEX_NONE;
-				bool bFoundName = false;
-				for (int32 i=0; i<NewVariables.Num(); ++i)
-				{
-					if (NewVariables[i].VarName == Description.Name)
+					int32 VariableIndex = INDEX_NONE;
+					bool bFoundName = false;
+					for (int32 i=0; i<NewVariables.Num(); ++i)
 					{
-						if (NewVariables[i].VarType == PinType &&
-							NewVariables[i].DefaultValue == Description.DefaultValue)
+						if (NewVariables[i].VarName == Description.Name)
 						{
-							VariableIndex = i;
+							if (NewVariables[i].VarType == PinType &&
+								NewVariables[i].DefaultValue == Description.DefaultValue)
+							{
+								VariableIndex = i;
+							}
+							bFoundName = true;
+							break;
 						}
-						bFoundName = true;
-						break;
 					}
-				}
 
-				if (VariableIndex == INDEX_NONE)
-				{
-					FName NewVariableName = Description.Name;
-					if (bFoundName)
+					if (VariableIndex == INDEX_NONE)
 					{
-						bool bFound = true;
-						int32 Count = 0;
-						while (bFound)
+						FName NewVariableName = Description.Name;
+						if (bFoundName)
 						{
-							bFound = false;
-							for (int32 i=0; i<NewVariables.Num(); ++i)
+							bool bFound = true;
+							int32 Count = 0;
+							while (bFound)
 							{
-								if (NewVariables[i].VarName == NewVariableName)
+								bFound = false;
+								for (int32 i=0; i<NewVariables.Num(); ++i)
 								{
-									bFound = true;
-									NewVariableName = *FString::Printf(TEXT("%s_%d"), *Description.Name.ToString(), ++Count);
-									break;
+									if (NewVariables[i].VarName == NewVariableName)
+									{
+										bFound = true;
+										NewVariableName = *FString::Printf(TEXT("%s_%d"), *Description.Name.ToString(), ++Count);
+										break;
+									}
 								}
 							}
 						}
+						
+						VariableIndex = AddCRMemberVariable(this, NewVariableName, PinType, false, false, Description.DefaultValue);
 					}
 					
-					VariableIndex = AddCRMemberVariable(this, NewVariableName, PinType, false, false, Description.DefaultValue);
+					FName VarName = NewVariables[VariableIndex].VarName;
+					GetOrCreateController()->ReplaceParameterNodeWithVariable(ParameterNode->GetFName(), VarName, Description.CPPType, Description.CPPTypeObject, false);
+					bDirtyDuringLoad = true;
 				}
-				
-				FName VarName = NewVariables[VariableIndex].VarName;
-				GetOrCreateController()->ReplaceParameterNodeWithVariable(ParameterNode->GetFName(), VarName, Description.CPPType, Description.CPPTypeObject, false);
-				bDirtyDuringLoad = true;
 			}
-		}		
-
+		}
 		LastNewVariables = NewVariables;
 	}
 
@@ -4665,6 +4733,81 @@
 
 #endif
 
+UEdGraph* UControlRigBlueprint::CreateEdGraph(URigVMGraph* InModel, bool bForce)
+{
+	check(InModel);
+
+	if(InModel->IsA<URigVMFunctionLibrary>())
+	{
+		return FunctionLibraryEdGraph;
+	}
+	
+	if(bForce)
+	{
+		RemoveEdGraph(InModel);
+	}
+
+	FString GraphName = InModel->GetName();
+	GraphName.RemoveFromStart(RigVMModelPrefix);
+	GraphName.TrimStartAndEndInline();
+
+	if(GraphName.IsEmpty())
+	{
+		GraphName = UControlRigGraphSchema::GraphName_ControlRig.ToString();
+	}
+
+	GraphName = RigVMClient.GetUniqueName(*GraphName).ToString();
+
+	UControlRigGraph* ControlRigGraph = NewObject<UControlRigGraph>(this, *GraphName, RF_Transactional);
+	ControlRigGraph->Schema = UControlRigGraphSchema::StaticClass();
+	ControlRigGraph->bAllowDeletion = true;
+	
+	FBlueprintEditorUtils::AddUbergraphPage(this, ControlRigGraph);
+	LastEditedDocuments.AddUnique(ControlRigGraph);
+
+	ControlRigGraph->Initialize(this);
+	ControlRigGraph->ModelNodePath = InModel->GetNodePath();
+
+	return ControlRigGraph;
+}
+
+bool UControlRigBlueprint::RemoveEdGraph(URigVMGraph* InModel)
+{
+	if(UControlRigGraph* RigGraph = Cast<UControlRigGraph>(GetEdGraph(InModel)))
+	{
+		if(UbergraphPages.Contains(RigGraph))
+		{
+			Modify();
+			UbergraphPages.Remove(RigGraph);
+		}
+		DestroyObject(RigGraph);
+		return true;
+	}
+	return false;
+}
+
+void UControlRigBlueprint::DestroyObject(UObject* InObject)
+{
+	RigVMClient.DestroyObject(InObject);
+}
+
+void UControlRigBlueprint::RenameGraph(const FString& InNodePath, const FName& InNewName)
+{
+	FName OldName = NAME_None;
+	UEdGraph* EdGraph = GetEdGraph(InNodePath);
+	if(EdGraph)
+	{
+		OldName = EdGraph->GetFName();
+	}
+	
+	RigVMClient.RenameModel(InNodePath, InNewName, true);
+
+	if(EdGraph)
+	{
+		NotifyGraphRenamed(EdGraph, OldName, EdGraph->GetFName());
+	}
+}
+
 void UControlRigBlueprint::CreateEdGraphForCollapseNodeIfNeeded(URigVMCollapseNode* InNode, bool bForce)
 {
 	check(InNode);
@@ -4780,7 +4923,7 @@
 						}
 
 						FunctionGraphs.Remove(RigFunctionGraph);
-						RigFunctionGraph->Rename(nullptr, GetTransientPackage());
+						RigFunctionGraph->Rename(nullptr, GetTransientPackage(), REN_ForceNoResetLoaders | REN_DontCreateRedirectors);
 						RigFunctionGraph->MarkAsGarbage();
 						return bNotify;
 					}
@@ -4809,7 +4952,7 @@
 						}
 
 						RigGraph->SubGraphs.Remove(SubRigGraph);
-						SubRigGraph->Rename(nullptr, GetTransientPackage());
+						SubRigGraph->Rename(nullptr, GetTransientPackage(), REN_ForceNoResetLoaders | REN_DontCreateRedirectors);
 						SubRigGraph->MarkAsGarbage();
 						return bNotify;
 					}
diff --git a/Engine/Plugins/Animation/ControlRig/Source/ControlRigDeveloper/Private/ControlRigBlueprintUtils.cpp b/Engine/Plugins/Animation/ControlRig/Source/ControlRigDeveloper/Private/ControlRigBlueprintUtils.cpp
--- a/Engine/Plugins/Animation/ControlRig/Source/ControlRigDeveloper/Private/ControlRigBlueprintUtils.cpp
+++ b/Engine/Plugins/Animation/ControlRig/Source/ControlRigDeveloper/Private/ControlRigBlueprintUtils.cpp
@@ -98,7 +98,7 @@ void FControlRigBlueprintUtils::HandleRefreshAllNodes(UBlueprint* InBlueprint)
 
 	if(UControlRigBlueprint* RigBlueprint = Cast<UControlRigBlueprint>(InBlueprint))
 	{
-		if (RigBlueprint->GetModel() == nullptr)
+		if (RigBlueprint->GetDefaultModel() == nullptr)
 		{
 			return;
 		}
diff --git a/Engine/Plugins/Animation/ControlRig/Source/ControlRigDeveloper/Private/Graph/ControlRigGraph.cpp b/Engine/Plugins/Animation/ControlRig/Source/ControlRigDeveloper/Private/Graph/ControlRigGraph.cpp
--- a/Engine/Plugins/Animation/ControlRig/Source/ControlRigDeveloper/Private/Graph/ControlRigGraph.cpp
+++ b/Engine/Plugins/Animation/ControlRig/Source/ControlRigDeveloper/Private/Graph/ControlRigGraph.cpp
@@ -35,6 +35,71 @@ UControlRigGraph::UControlRigGraph()
 	bIsFunctionDefinition = false;
 }
 
+FRigVMClient* UControlRigGraph::GetRigVMClient() const
+{
+	if (const IRigVMClientHost* Host = GetImplementingOuter<IRigVMClientHost>())
+	{
+		return (FRigVMClient*)Host->GetRigVMClient();
+	}
+	return nullptr;
+}
+
+FString UControlRigGraph::GetRigVMNodePath() const
+{
+	return ModelNodePath;
+}
+
+void UControlRigGraph::HandleRigVMGraphRenamed(const FString& InOldNodePath, const FString& InNewNodePath)
+{
+	static constexpr TCHAR NodePathPrefixFormat[] = TEXT("%s|");
+	const FString OldPrefix = FString::Printf(NodePathPrefixFormat, *InOldNodePath); 
+	const FString NewPrefix = FString::Printf(NodePathPrefixFormat, *InNewNodePath);
+	
+	if(ModelNodePath == InOldNodePath)
+	{
+		Modify();
+		ModelNodePath = InNewNodePath;
+
+		FString GraphName;
+		if(!ModelNodePath.Split(TEXT("|"), nullptr, &GraphName, ESearchCase::CaseSensitive, ESearchDir::FromEnd))
+		{
+			GraphName = ModelNodePath;
+		}
+		GraphName.RemoveFromEnd(TEXT("::"));
+		GraphName.RemoveFromStart(UControlRigBlueprint::RigVMModelPrefix);
+		GraphName.TrimStartAndEndInline();
+
+		if(GraphName.IsEmpty())
+		{
+			GraphName = UControlRigGraphSchema::GraphName_ControlRig.ToString(); 
+		}
+		GraphName = FRigVMClient::GetUniqueName(GetOuter(), *GraphName).ToString();
+
+		Rename(*GraphName, nullptr, REN_ForceNoResetLoaders | REN_DontCreateRedirectors);
+	}
+	else if(ModelNodePath.StartsWith(OldPrefix))
+	{
+		Modify();
+		ModelNodePath = NewPrefix + ModelNodePath.RightChop(OldPrefix.Len() - 1);
+	}
+	else
+	{
+		return;
+	}
+
+	for(UEdGraphNode* Node : Nodes)
+	{
+		if(UControlRigGraphNode* RigNode = Cast<UControlRigGraphNode>(Node))
+		{
+			if(RigNode->ModelNodePath.StartsWith(OldPrefix))
+			{
+				RigNode->Modify();
+				RigNode->ModelNodePath = NewPrefix + RigNode->ModelNodePath.RightChop(OldPrefix.Len() - 1);
+			}
+		}
+	}
+}
+
 void UControlRigGraph::Initialize(UControlRigBlueprint* InBlueprint)
 {
 	DECLARE_SCOPE_HIERARCHICAL_COUNTER_FUNC()
@@ -911,22 +976,31 @@ UControlRigBlueprint* UControlRigGraph::GetBlueprint() const
 
 URigVMGraph* UControlRigGraph::GetModel() const
 {
-	if (const IRigVMGraphHost* GraphHost = GetImplementingOuter<IRigVMGraphHost>())
+	if (FRigVMClient* Client = GetRigVMClient())
 	{
-		return GraphHost->GetRigVMGraph(this);
+		return Client->GetModel(this);
 	}
 	return nullptr;
 }
 
 URigVMController* UControlRigGraph::GetController() const
 {
-	if (IRigVMControllerHost* ControllerHost = GetImplementingOuter<IRigVMControllerHost>())
+	if (FRigVMClient* Client = GetRigVMClient())
 	{
-		return ControllerHost->GetOrCreateRigVMController(this);
+		return Client->GetOrCreateController(this);
 	}
 	return nullptr;
 }
 
+const UControlRigGraph* UControlRigGraph::GetRootGraph() const
+{
+	if(const UControlRigGraph* ParentGraph = Cast<UControlRigGraph>(GetOuter()))
+	{
+		return ParentGraph->GetRootGraph();
+	}
+	return this;
+}
+
 void UControlRigGraph::RemoveNode(UEdGraphNode* InNode)
 {
 	// Make sure EdGraph is not part of the transaction
@@ -947,6 +1021,8 @@ void UControlRigGraph::RemoveNode(UEdGraphNode* InNode)
 	}
 	InNode->Rename(*DeletedName, nullptr, REN_ForceNoResetLoaders | REN_DontCreateRedirectors);	
 	Super::RemoveNode(InNode);
+
+	NotifyGraphChanged(FEdGraphEditAction(EEdGraphActionType::GRAPHACTION_RemoveNode, this, InNode, false));
 }
 
 URigVMController* UControlRigGraph::GetTemplateController()
diff --git a/Engine/Plugins/Animation/ControlRig/Source/ControlRigDeveloper/Private/Graph/ControlRigGraphSchema.cpp b/Engine/Plugins/Animation/ControlRig/Source/ControlRigDeveloper/Private/Graph/ControlRigGraphSchema.cpp
--- a/Engine/Plugins/Animation/ControlRig/Source/ControlRigDeveloper/Private/Graph/ControlRigGraphSchema.cpp
+++ b/Engine/Plugins/Animation/ControlRig/Source/ControlRigDeveloper/Private/Graph/ControlRigGraphSchema.cpp
@@ -27,6 +27,7 @@
 #include "RigVMModel/RigVMVariableDescription.h"
 #include "RigVMCore/RigVMUnknownType.h"
 #include "Kismet2/Kismet2NameValidators.h"
+#include "Algo/Count.h"
 
 #if WITH_EDITOR
 #include "ControlRigEditor/Private/Editor/SControlRigFunctionLocalizationWidget.h"
@@ -573,10 +574,10 @@ bool UControlRigGraphSchema::TryCreateConnection(UEdGraphPin* PinA, UEdGraphPin*
 
 	UControlRigGraphSchema* MutableThis = (UControlRigGraphSchema*)this;
 
-	IRigVMControllerHost* ControllerHost = PinA->GetOwningNode()->GetImplementingOuter<IRigVMControllerHost>();
-	if (ControllerHost != nullptr)
+	IRigVMClientHost* Host = PinA->GetOwningNode()->GetImplementingOuter<IRigVMClientHost>();
+	if (Host != nullptr)
 	{
-		if (URigVMController* Controller = ControllerHost->GetOrCreateRigVMController(PinA->GetOwningNode()->GetGraph()))
+		if (URigVMController* Controller = Host->GetRigVMClient()->GetOrCreateController(PinA->GetOwningNode()->GetGraph()))
 		{
 			ERigVMPinDirection UserLinkDirection = ERigVMPinDirection::Output;
 			if (PinA->Direction == EGPD_Input)
@@ -1212,28 +1213,72 @@ void UControlRigGraphSchema::GetGraphDisplayInformation(const UEdGraph& Graph, /
 
 	if (UControlRigGraph* RigGraph = Cast<UControlRigGraph>((UEdGraph*)&Graph))
 	{
-		TArray<FString> NodePathParts;
-		if (URigVMNode::SplitNodePath(RigGraph->ModelNodePath, NodePathParts))
+		if(URigVMGraph* Model = RigGraph->GetModel())
 		{
-			DisplayInfo.DisplayName = FText::FromString(NodePathParts.Last());
-			DisplayInfo.PlainName = DisplayInfo.DisplayName;
+			TArray<FString> NodePathParts;
+			if (URigVMNode::SplitNodePath(RigGraph->ModelNodePath, NodePathParts))
+			{
+				if(NodePathParts.Num() > 1)
+				{
+					DisplayInfo.DisplayName = FText::FromString(NodePathParts.Last());
+					DisplayInfo.PlainName = DisplayInfo.DisplayName;
 
-			static const FText LocalFunctionText = FText::FromString(TEXT("A local function.")); 
-			DisplayInfo.Tooltip = LocalFunctionText;
+					static const FText LocalFunctionText = FText::FromString(TEXT("A local function.")); 
+					DisplayInfo.Tooltip = LocalFunctionText;
+				}
 
-			// if this is a riggraph within a collapse node - let's use that for the tooltip
-			if(URigVMGraph* Model = RigGraph->GetModel())
-			{
+				// if this is a riggraph within a collapse node - let's use that for the tooltip
 				if(URigVMCollapseNode* CollapseNode = Model->GetTypedOuter<URigVMCollapseNode>())
 				{
 					DisplayInfo.Tooltip = CollapseNode->GetToolTipText();
 				}
 			}
-		}
-		else
-		{
-			static const FText MainGraphText = FText::FromString(TEXT("The main graph for the Control Rig."));
-			DisplayInfo.Tooltip = MainGraphText;
+
+			if(Model->IsRootGraph())
+			{
+				// let's see if there is only one event
+				FString EventName;
+				if(Algo::CountIf(Model->GetNodes(), [&EventName](const URigVMNode* NodeToCount) -> bool
+				{
+					if(NodeToCount->IsEvent())
+					{
+						if(EventName.IsEmpty())
+						{
+							if(const URigVMUnitNode* UnitNode = Cast<URigVMUnitNode>(NodeToCount))
+							{
+								if(UnitNode->GetScriptStruct())
+								{
+									EventName = UnitNode->GetScriptStruct()->GetDisplayNameText().ToString();
+								}
+							}
+						}
+						return true;
+					}
+					return false;
+				}) == 1)
+				{
+					static constexpr TCHAR EventGraphNameFormat[] = TEXT("%s Graph");
+					const FString DesiredGraphName = FString::Printf(EventGraphNameFormat, *EventName);
+					DisplayInfo.DisplayName = FText::FromString(DesiredGraphName);
+					DisplayInfo.PlainName = DisplayInfo.DisplayName;
+				}
+				else
+				{
+					static const FText MainGraphText = FText::FromString(TEXT("A top level graph for the rig."));
+					DisplayInfo.Tooltip = MainGraphText;
+
+					if(!NodePathParts.IsEmpty())
+					{
+						static constexpr TCHAR NodePathSuffix[] = TEXT("::");
+						FString NodePath = NodePathParts[0];
+						NodePath.RemoveFromEnd(NodePathSuffix);
+						NodePath.RemoveFromStart(UControlRigBlueprint::RigVMModelPrefix);
+						NodePath.TrimStartAndEndInline();
+						DisplayInfo.DisplayName = FText::FromString(NodePath);
+						DisplayInfo.PlainName = DisplayInfo.DisplayName;
+					}
+				}
+			}
 		}
 	}
 }
@@ -1294,6 +1339,43 @@ FText UControlRigGraphSchema::GetGraphCategory(const UEdGraph* InGraph) const
 	return FText();
 }
 
+EGraphType UControlRigGraphSchema::GetGraphType(const UEdGraph* TestEdGraph) const
+{
+	if (UControlRigGraph* RigGraph = Cast<UControlRigGraph>((UEdGraph*)TestEdGraph))
+	{
+		if (UControlRigBlueprint* RigBlueprint = Cast<UControlRigBlueprint>(FBlueprintEditorUtils::FindBlueprintForGraph(RigGraph)))
+		{
+			if (URigVMGraph* Model = RigGraph->GetModel())
+			{
+				if(Model->IsRootGraph())
+				{
+					if(Model->IsA<URigVMFunctionLibrary>())
+					{
+						return EGraphType::GT_Function;
+					}
+					else
+					{
+						return EGraphType::GT_Ubergraph;
+					}
+				}
+				else if(URigVMLibraryNode* LibraryNode = Cast<URigVMLibraryNode>(Model->GetOuter()))
+				{
+					if(LibraryNode->GetGraph()->IsA<URigVMFunctionLibrary>())
+					{
+						return EGraphType::GT_Function;
+					}
+					else
+					{
+						// collapse nodes show up as uber graphs
+						return EGraphType::GT_Ubergraph;
+					}
+				}
+			}
+		}
+	}
+	return Super::GetGraphType(TestEdGraph);
+}
+
 FReply UControlRigGraphSchema::TrySetGraphCategory(const UEdGraph* InGraph, const FText& InCategory)
 {
 	if (UControlRigGraph* RigGraph = Cast<UControlRigGraph>((UEdGraph*)InGraph))
@@ -1326,7 +1408,12 @@ bool UControlRigGraphSchema::TryDeleteGraph(UEdGraph* GraphToDelete) const
 		{
 			if (URigVMGraph* Model = RigBlueprint->GetModel(GraphToDelete))
 			{
-				if (URigVMCollapseNode* LibraryNode = Cast<URigVMCollapseNode>(Model->GetOuter()))
+				if(Model->IsRootGraph())
+				{
+					RigBlueprint->RemoveModel(Model->GetNodePath());
+					return true;
+				}
+				else if (URigVMCollapseNode* LibraryNode = Cast<URigVMCollapseNode>(Model->GetOuter()))
 				{
 					if (URigVMController* Controller = RigBlueprint->GetOrCreateController(LibraryNode->GetGraph()))
 					{
@@ -1399,7 +1486,12 @@ bool UControlRigGraphSchema::TryRenameGraph(UEdGraph* GraphToRename, const FName
 		{
 			if (URigVMGraph* Model = RigGraph->GetModel())
 			{
-				if (URigVMGraph* RootModel = Model->GetRootGraph())
+				if(Model->IsRootGraph())
+				{
+					const FString NewName = FString::Printf(TEXT("%s %s"), RigBlueprint->RigVMModelPrefix, *InNewName.ToString()); 
+					RigBlueprint->RenameGraph(Model->GetNodePath(), *NewName);
+				}
+				else if (URigVMGraph* RootModel = Model->GetRootGraph())
 				{
 					URigVMLibraryNode* LibraryNode = Cast<URigVMLibraryNode>(RootModel->FindNode(RigGraph->ModelNodePath));
 					if (LibraryNode)
diff --git a/Engine/Plugins/Animation/ControlRig/Source/ControlRigDeveloper/Public/ControlRigBlueprint.h b/Engine/Plugins/Animation/ControlRig/Source/ControlRigDeveloper/Public/ControlRigBlueprint.h
--- a/Engine/Plugins/Animation/ControlRig/Source/ControlRigDeveloper/Public/ControlRigBlueprint.h
+++ b/Engine/Plugins/Animation/ControlRig/Source/ControlRigDeveloper/Public/ControlRigBlueprint.h
@@ -15,7 +15,7 @@
 #include "ControlRigGizmoLibrary.h"
 #include "RigVMCore/RigVM.h"
 #include "RigVMCore/RigVMStatistics.h"
-#include "RigVMModel/RigVMController.h"
+#include "RigVMModel/RigVMClient.h"
 #include "RigVMCompiler/RigVMCompiler.h"
 #include "ControlRigValidationPass.h"
 #include "Drawing/ControlRigDrawContainer.h"
@@ -196,7 +196,7 @@ struct CONTROLRIGDEVELOPER_API FControlRigPythonSettings
 };
 
 UCLASS(BlueprintType, meta=(IgnoreClassThumbnail))
-class CONTROLRIGDEVELOPER_API UControlRigBlueprint : public UBlueprint, public IInterface_PreviewMeshProvider, public IRigVMGraphHost, public IRigVMControllerHost
+class CONTROLRIGDEVELOPER_API UControlRigBlueprint : public UBlueprint, public IInterface_PreviewMeshProvider, public IRigVMClientHost
 {
 	GENERATED_UCLASS_BODY()
 
@@ -240,7 +240,7 @@ public:
 	virtual bool SupportsFunctions() const override { return true; }
 	virtual bool SupportsMacros() const override { return false; }
 	virtual bool SupportsDelegates() const override { return false; }
-	virtual bool SupportsEventGraphs() const override { return false; }
+	virtual bool SupportsEventGraphs() const override { return true; }
 	virtual bool SupportsAnimLayers() const override { return false; }
 	virtual bool ExportGraphToText(UEdGraph* InEdGraph, FString& OutText) override;
 	virtual bool TryImportGraphFromText(const FString& InClipboardText, UEdGraph** OutGraphPtr = nullptr) override;
@@ -249,15 +249,15 @@ public:
 	// UObject interface
 	virtual void PostEditChangeChainProperty(struct FPropertyChangedChainEvent& PropertyChangedEvent) override;
 
-	// IRigVMGraphHost interface
-	virtual URigVMGraph* GetRigVMGraph(const UObject* InEditorObject) const override { return GetModel(CastChecked<UEdGraph>(InEditorObject)); }
+	// IRigVMClientHost interface
+	virtual FRigVMClient* GetRigVMClient() override;
+	virtual const FRigVMClient* GetRigVMClient() const override;
+	virtual UObject* GetEditorObjectForRigVMGraph(URigVMGraph* InVMGraph) const override;
+	virtual void HandleRigVMGraphAdded(const FRigVMClient* InClient, const FString& InNodePath) override;
+	virtual void HandleRigVMGraphRemoved(const FRigVMClient* InClient, const FString& InNodePath) override;
+	virtual void HandleRigVMGraphRenamed(const FRigVMClient* InClient, const FString& InOldNodePath, const FString& InNewNodePath) override;
+	virtual void HandleConfigureRigVMController(const FRigVMClient* InClient, URigVMController* InControllerToConfigure) override;
 
-	// IRigVMControllerHost interface
-	virtual URigVMController* GetRigVMController(const URigVMGraph* InRigVMGraph) const override { return GetController(InRigVMGraph); }
-	virtual URigVMController* GetRigVMController(const UObject* InEditorObject) const override { return GetController(CastChecked<UEdGraph>(InEditorObject)); }
-	virtual URigVMController* GetOrCreateRigVMController(URigVMGraph* InRigVMGraph) override { return GetOrCreateController(InRigVMGraph); }
-	virtual URigVMController* GetOrCreateRigVMController(const UObject* InEditorObject) override { return GetOrCreateController(CastChecked<UEdGraph>(InEditorObject)); }
-	
 	FOnRequestInspectObject& OnRequestInspectObject() { return OnRequestInspectObjectEvent; }
 	void RequestInspectObject(const TArray<UObject*>& InObjects) { OnRequestInspectObjectEvent.Broadcast(InObjects); }
 
@@ -297,12 +297,21 @@ public:
 	URigVMGraph* GetModel(const UEdGraph* InEdGraph = nullptr) const;
 	URigVMGraph* GetModel(const FString& InNodePath) const;
 
+	UFUNCTION(BlueprintCallable, Category = "Control Rig Blueprint")
+	URigVMGraph* GetDefaultModel() const;
+
 	UFUNCTION(BlueprintCallable, Category = "Control Rig Blueprint")
 	TArray<URigVMGraph*> GetAllModels() const;
 
 	UFUNCTION(BlueprintCallable, Category = "Control Rig Blueprint")
 	URigVMFunctionLibrary* GetLocalFunctionLibrary() const;
 
+	UFUNCTION(BlueprintCallable, Category = "Control Rig Blueprint")
+	URigVMGraph* AddModel(FString InName = TEXT("Rig Graph"), bool bSetupUndoRedo = true, bool bPrintPythonCommand = true);
+
+	UFUNCTION(BlueprintCallable, Category = "Control Rig Blueprint")
+	bool RemoveModel(FString InName = TEXT("Rig Graph"), bool bSetupUndoRedo = true, bool bPrintPythonCommand = true);
+
 	UFUNCTION(BlueprintCallable, Category = "Control Rig Blueprint")
 	URigVMController* GetController(const URigVMGraph* InGraph = nullptr) const;
 
@@ -364,11 +373,14 @@ public:
 
 protected:
 
-	UPROPERTY(BlueprintReadOnly, Category = "VM")
-	TObjectPtr<URigVMGraph> Model;
+	UPROPERTY()
+	TObjectPtr<URigVMGraph> Model_DEPRECATED;
+
+	UPROPERTY()
+	TObjectPtr<URigVMFunctionLibrary> FunctionLibrary_DEPRECATED;
 
-	UPROPERTY(BlueprintReadOnly, Category = "VM")
-	TObjectPtr<URigVMFunctionLibrary> FunctionLibrary;
+	UPROPERTY()
+	FRigVMClient RigVMClient;
 
 	/** Asset searchable information about exposed public functions on this rig */
 	UPROPERTY(AssetRegistrySearchable)
@@ -378,10 +390,6 @@ protected:
 	UPROPERTY(AssetRegistrySearchable)
 	TArray<FRigVMReferenceNodeData> FunctionReferenceNodeData;
 
-	// Related URigVMGraph path and controller
-	UPROPERTY(BlueprintReadOnly, transient, Category = "VM")
-	TMap<FSoftObjectPath, TObjectPtr<URigVMController>> Controllers;
-
 #if WITH_EDITORONLY_DATA
 
 	UPROPERTY(transient, DuplicateTransient)
@@ -685,6 +693,10 @@ private:
 
 #endif
 
+	UEdGraph* CreateEdGraph(URigVMGraph* InModel, bool bForce = false);
+	bool RemoveEdGraph(URigVMGraph* InModel);
+	void DestroyObject(UObject* InObject);
+	void RenameGraph(const FString& InNodePath, const FName& InNewName);
 	void CreateEdGraphForCollapseNodeIfNeeded(URigVMCollapseNode* InNode, bool bForce = false);
 	bool RemoveEdGraphForCollapseNode(URigVMCollapseNode* InNode, bool bNotify = false);
 	void HandleReportFromCompiler(EMessageSeverity::Type InSeverity, UObject* InSubject, const FString& InMessage);
@@ -734,9 +746,12 @@ public:
 
 #endif
 
+	static constexpr TCHAR RigVMModelPrefix[] = TEXT("RigVMModel");
+
 private:
 	bool bDirtyDuringLoad;
 	bool bErrorsDuringCompilation;
+	bool bSuspendPythonMessagesForRigVMClient;
 
 	friend class FControlRigBlueprintCompilerContext;
 	friend class SRigHierarchy;
@@ -754,6 +769,7 @@ private:
 	friend class FControlRigEditorModule;
 	friend class UControlRigComponent;
 	friend struct FControlRigGraphSchemaAction_PromoteToVariable;
+	friend class UControlRigGraphSchema;
 };
 
 class CONTROLRIGDEVELOPER_API FControlRigBlueprintVMCompileScope
diff --git a/Engine/Plugins/Animation/ControlRig/Source/ControlRigDeveloper/Public/Graph/ControlRigGraph.h b/Engine/Plugins/Animation/ControlRig/Source/ControlRigDeveloper/Public/Graph/ControlRigGraph.h
--- a/Engine/Plugins/Animation/ControlRig/Source/ControlRigDeveloper/Public/Graph/ControlRigGraph.h
+++ b/Engine/Plugins/Animation/ControlRig/Source/ControlRigDeveloper/Public/Graph/ControlRigGraph.h
@@ -21,13 +21,18 @@ struct FRigCurveContainer;
 DECLARE_MULTICAST_DELEGATE_OneParam(FControlRigGraphNodeClicked, UControlRigGraphNode*);
 
 UCLASS()
-class CONTROLRIGDEVELOPER_API UControlRigGraph : public UEdGraph
+class CONTROLRIGDEVELOPER_API UControlRigGraph : public UEdGraph, public IRigVMEditorSideObject
 {
 	GENERATED_BODY()
 
 public:
 	UControlRigGraph();
 
+	/** IRigVMEditorSideObject interface */
+	virtual FRigVMClient* GetRigVMClient() const override;
+	virtual FString GetRigVMNodePath() const override;
+	virtual void HandleRigVMGraphRenamed(const FString& InOldNodePath, const FString& InNewNodePath) override;
+
 	/** Set up this graph */
 	void Initialize(UControlRigBlueprint* InBlueprint);
 
@@ -69,8 +74,10 @@ public:
 	UEdGraphNode* FindNodeForModelNodeName(const FName& InModelNodeName, const bool bCacheIfRequired = true);
 
 	UControlRigBlueprint* GetBlueprint() const;
-	virtual URigVMGraph* GetModel() const;
-	virtual URigVMController* GetController() const;
+	URigVMGraph* GetModel() const;
+	URigVMController* GetController() const;
+	bool IsRootGraph() const { return GetRootGraph() == this; }
+	const UControlRigGraph* GetRootGraph() const;
 
 	void HandleModifiedEvent(ERigVMGraphNotifType InNotifType, URigVMGraph* InGraph, UObject* InSubject);
 
diff --git a/Engine/Plugins/Animation/ControlRig/Source/ControlRigDeveloper/Public/Graph/ControlRigGraphSchema.h b/Engine/Plugins/Animation/ControlRig/Source/ControlRigDeveloper/Public/Graph/ControlRigGraphSchema.h
--- a/Engine/Plugins/Animation/ControlRig/Source/ControlRigDeveloper/Public/Graph/ControlRigGraphSchema.h
+++ b/Engine/Plugins/Animation/ControlRig/Source/ControlRigDeveloper/Public/Graph/ControlRigGraphSchema.h
@@ -227,6 +227,7 @@ public:
 	virtual bool GetLocalVariables(const UEdGraph* InGraph, TArray<FBPVariableDescription>& OutLocalVariables) const override;
 	virtual TSharedPtr<FEdGraphSchemaAction> MakeActionFromVariableDescription(const UEdGraph* InEdGraph, const FBPVariableDescription& Variable) const override;
 	virtual FText GetGraphCategory(const UEdGraph* InGraph) const override;
+	virtual EGraphType GetGraphType(const UEdGraph* TestEdGraph) const override;
 	virtual FReply TrySetGraphCategory(const UEdGraph* InGraph, const FText& InCategory) override;
 	virtual bool TryDeleteGraph(UEdGraph* GraphToDelete) const override;
 	virtual bool TryRenameGraph(UEdGraph* GraphToRename, const FName& InNewName) const override;
@@ -282,4 +283,3 @@ private:
 	friend class UControlRigUnitNodeSpawner;
 	friend class UControlRigArrayNodeSpawner;
 };
-
diff --git a/Engine/Plugins/Animation/ControlRig/Source/ControlRigEditor/Private/ControlRigBlueprintEditorLibrary.cpp b/Engine/Plugins/Animation/ControlRig/Source/ControlRigEditor/Private/ControlRigBlueprintEditorLibrary.cpp
--- a/Engine/Plugins/Animation/ControlRig/Source/ControlRigEditor/Private/ControlRigBlueprintEditorLibrary.cpp
+++ b/Engine/Plugins/Animation/ControlRig/Source/ControlRigEditor/Private/ControlRigBlueprintEditorLibrary.cpp
@@ -73,7 +73,7 @@ URigVMGraph* UControlRigBlueprintEditorLibrary::GetModel(UControlRigBlueprint* I
 	{
 		return nullptr;
 	}
-	return InRigBlueprint->GetModel();
+	return InRigBlueprint->GetDefaultModel();
 }
 
 URigVMController* UControlRigBlueprintEditorLibrary::GetController(UControlRigBlueprint* InRigBlueprint)
diff --git a/Engine/Plugins/Animation/ControlRig/Source/ControlRigEditor/Private/ControlRigCompilerDetails.cpp b/Engine/Plugins/Animation/ControlRig/Source/ControlRigEditor/Private/ControlRigCompilerDetails.cpp
--- a/Engine/Plugins/Animation/ControlRig/Source/ControlRigEditor/Private/ControlRigCompilerDetails.cpp
+++ b/Engine/Plugins/Animation/ControlRig/Source/ControlRigEditor/Private/ControlRigCompilerDetails.cpp
@@ -178,9 +178,9 @@ FReply FRigVMCompileSettingsDetails::OnCopyASTClicked()
 {
 	if (BlueprintBeingCustomized)
 	{
-		if (BlueprintBeingCustomized->GetModel())
+		if (BlueprintBeingCustomized->GetDefaultModel())
 		{
-			FString DotContent = BlueprintBeingCustomized->GetModel()->GetRuntimeAST()->DumpDot();
+			FString DotContent = BlueprintBeingCustomized->GetDefaultModel()->GetRuntimeAST()->DumpDot();
 			FPlatformApplicationMisc::ClipboardCopy(*DotContent);
 		}
 	}
@@ -191,7 +191,7 @@ FReply FRigVMCompileSettingsDetails::OnCopyByteCodeClicked()
 {
 	if (BlueprintBeingCustomized)
 	{
-		if (BlueprintBeingCustomized->GetModel())
+		if (BlueprintBeingCustomized->GetDefaultModel())
 		{
 			if(UControlRig* ControlRig = Cast<UControlRig>(BlueprintBeingCustomized->GetObjectBeingDebugged()))
 			{
@@ -241,7 +241,7 @@ FReply FRigVMCompileSettingsDetails::OnCopyGeneratedCodeClicked()
 					}
 					
 					FRigVMCodeGenerator CodeGenerator(ClassName,
-						TEXT("TestModule"), BlueprintBeingCustomized->GetModel(), CDO->GetVM(), BlueprintBeingCustomized->PinToOperandMap);
+						TEXT("TestModule"), BlueprintBeingCustomized->GetDefaultModel(), CDO->GetVM(), BlueprintBeingCustomized->PinToOperandMap);
 					const FString Content = CodeGenerator.DumpHeader() + TEXT("\r\n\r\n") + CodeGenerator.DumpSource();
 					FPlatformApplicationMisc::ClipboardCopy(*Content);
 				}
diff --git a/Engine/Plugins/Animation/ControlRig/Source/ControlRigEditor/Private/ControlRigEditorModule.cpp b/Engine/Plugins/Animation/ControlRig/Source/ControlRigEditor/Private/ControlRigEditorModule.cpp
--- a/Engine/Plugins/Animation/ControlRig/Source/ControlRigEditor/Private/ControlRigEditorModule.cpp
+++ b/Engine/Plugins/Animation/ControlRig/Source/ControlRigEditor/Private/ControlRigEditorModule.cpp
@@ -2518,7 +2518,7 @@ void FControlRigEditorModule::PostChange(const UUserDefinedStruct* Changed,
 				if (Pin->GetGraph() == RigBlueprint->GetLocalFunctionLibrary())
 				{
 					TArray< TSoftObjectPtr<URigVMFunctionReferenceNode> > References;
-					References = RigBlueprint->FunctionLibrary->GetReferencesForFunction(Pin->GetNode()->GetFName());
+					References = RigBlueprint->RigVMClient.GetOrCreateFunctionLibrary(false)->GetReferencesForFunction(Pin->GetNode()->GetFName());
 					
 					for (const TSoftObjectPtr<URigVMFunctionReferenceNode>& Reference : References)
 					{
diff --git a/Engine/Plugins/Animation/ControlRig/Source/ControlRigEditor/Private/Editor/ControlRigEditor.cpp b/Engine/Plugins/Animation/ControlRig/Source/ControlRigEditor/Private/Editor/ControlRigEditor.cpp
--- a/Engine/Plugins/Animation/ControlRig/Source/ControlRigEditor/Private/Editor/ControlRigEditor.cpp
+++ b/Engine/Plugins/Animation/ControlRig/Source/ControlRigEditor/Private/Editor/ControlRigEditor.cpp
@@ -229,6 +229,11 @@ FControlRigEditor::~FControlRigEditor()
 
 	ControlRigEditorClosedDelegate.Broadcast(this, RigBlueprint);
 
+	if(PropertyChangedHandle.IsValid())
+	{
+		FCoreUObjectDelegates::OnObjectPropertyChanged.Remove(PropertyChangedHandle);
+	}
+
 	if (RigBlueprint)
 	{
 		// clear editor related data from the debugged control rig instance 
@@ -415,17 +420,14 @@ void FControlRigEditor::InitControlRigEditor(const EToolkitMode::Type Mode, cons
 		bool bBroughtGraphToFront = false;
 		for(UEdGraph* Graph : ControlRigBlueprints[0]->UbergraphPages)
 		{
-			if (Graph->GetFName().IsEqual(UControlRigGraphSchema::GraphName_ControlRig))
+			if (UControlRigGraph* RigGraph = Cast<UControlRigGraph>(Graph))
 			{
 				if (!bBroughtGraphToFront)
 				{
 					OpenGraphAndBringToFront(Graph, false);
 					bBroughtGraphToFront = true;
 				}
-			}
 
-			if (UControlRigGraph* RigGraph = Cast<UControlRigGraph>(Graph))
-			{
 				RigGraph->OnGraphNodeClicked.AddSP(this, &FControlRigEditor::OnGraphNodeClicked);
 				ActiveTabNodePath = RigGraph->ModelNodePath;
 			}
@@ -433,7 +435,7 @@ void FControlRigEditor::InitControlRigEditor(const EToolkitMode::Type Mode, cons
 	}
 
 	{
-		if (URigVMGraph* Model = InControlRigBlueprint->GetModel())
+		if (URigVMGraph* Model = InControlRigBlueprint->GetDefaultModel())
 		{
 			if (Model->GetNodes().Num() == 0)
 			{
@@ -578,6 +580,8 @@ void FControlRigEditor::InitControlRigEditor(const EToolkitMode::Type Mode, cons
 				return FControlRigGraphMathTypeDetails::MakeInstance();
 			}));
 	}
+
+	PropertyChangedHandle = FCoreUObjectDelegates::OnObjectPropertyChanged.AddSP(this, &FControlRigEditor::OnPropertyChanged);
 }
 
 void FControlRigEditor::InitFunctionReferences()
@@ -2135,7 +2139,10 @@ void FControlRigEditor::InitToolMenuContext(FToolMenuContext& MenuContext)
 			if (GraphNodeContext->Node)
 			{
 				Model = RigBlueprint->GetModel(GraphNodeContext->Graph);
-				Node = Model->FindNodeByName(GraphNodeContext->Node->GetFName());
+				if(Model)
+				{
+					Node = Model->FindNodeByName(GraphNodeContext->Node->GetFName());
+				}
 			}
 		
 			if (GraphNodeContext->Pin && Node)
@@ -2522,6 +2529,17 @@ void FControlRigEditor::PostTransaction(bool bSuccess, const FTransaction* Trans
 	}
 }
 
+void FControlRigEditor::JumpToHyperlink(const UObject* ObjectReference, bool bRequestRename)
+{
+	if(const UControlRigGraph* Graph = Cast<UControlRigGraph>(ObjectReference))
+	{
+		OpenGraphAndBringToFront((UEdGraph*)Graph, true);
+		return;
+	}
+	
+	IControlRigEditor::JumpToHyperlink(ObjectReference, bRequestRename);
+}
+
 void FControlRigEditor::EnsureValidRigElementsInDetailPanel()
 {
 	UControlRigBlueprint* ControlRigBP = GetControlRigBlueprint();
@@ -2655,23 +2673,35 @@ void FControlRigEditor::CreateDefaultTabContents(const TArray<UBlueprint*>& InBl
 
 void FControlRigEditor::NewDocument_OnClicked(ECreatedDocumentType GraphType)
 {
-	if (GraphType != FBlueprintEditor::CGT_NewFunctionGraph)
+	if (GraphType == FBlueprintEditor::CGT_NewFunctionGraph)
 	{
-		return;
-	}
+		if (UControlRigBlueprint* Blueprint = GetControlRigBlueprint())
+		{
+			if (URigVMController* Controller = Blueprint->GetOrCreateController(Blueprint->GetLocalFunctionLibrary()))
+			{
+				if (const URigVMLibraryNode* FunctionNode = Controller->AddFunctionToLibrary(TEXT("New Function"), true, FVector2D::ZeroVector, true, true))
+				{
+					if (const UEdGraph* NewGraph = Blueprint->GetEdGraph(FunctionNode->GetContainedGraph()))
+					{
+						OpenDocument(NewGraph, FDocumentTracker::OpenNewDocument);
+						RenameNewlyAddedAction(FunctionNode->GetFName());
+					}
 
-	if (UControlRigBlueprint* Blueprint = GetControlRigBlueprint())
+				}
+			}
+		}
+	}
+	else if(GraphType == FBlueprintEditor::CGT_NewEventGraph)
 	{
-		if (URigVMController* Controller = Blueprint->GetOrCreateController(Blueprint->GetLocalFunctionLibrary()))
+		if (UControlRigBlueprint* Blueprint = GetControlRigBlueprint())
 		{
-			if (URigVMLibraryNode* FunctionNode = Controller->AddFunctionToLibrary(TEXT("New Function"), true, FVector2D::ZeroVector, true, true))
+			if(URigVMGraph* Model = Blueprint->AddModel(UControlRigGraphSchema::GraphName_ControlRig.ToString()))
 			{
-				if (UEdGraph* NewGraph = Blueprint->GetEdGraph(FunctionNode->GetContainedGraph()))
+				if (const UEdGraph* NewGraph = Blueprint->GetEdGraph(Model))
 				{
 					OpenDocument(NewGraph, FDocumentTracker::OpenNewDocument);
-					RenameNewlyAddedAction(FunctionNode->GetFName());
+					RenameNewlyAddedAction(NewGraph->GetFName());
 				}
-
 			}
 		}
 	}
@@ -4412,20 +4442,20 @@ void FControlRigEditor::OnFinishedChangingProperties(const FPropertyChangedEvent
 
 	if (ControlRigBP)
 	{
-		if (PropertyChangedEvent.MemberProperty->GetNameCPP() == TEXT("VMCompileSettings"))
+		if (PropertyChangedEvent.MemberProperty->GetNameCPP() == GET_MEMBER_NAME_STRING_CHECKED(UControlRigBlueprint, VMCompileSettings))
 		{
 			ControlRigBP->RecompileVM();
 			return;
 		}
 
-		if (PropertyChangedEvent.MemberProperty->GetNameCPP() == TEXT("VMRuntimeSettings"))
+		if (PropertyChangedEvent.MemberProperty->GetNameCPP() == GET_MEMBER_NAME_STRING_CHECKED(UControlRigBlueprint, VMRuntimeSettings))
 		{
 			ControlRigBP->VMRuntimeSettings.Validate();
 			ControlRigBP->PropagateRuntimeSettingsFromBPToInstances();
 			return;
 		}
 
-		if (PropertyChangedEvent.MemberProperty->GetNameCPP() == TEXT("DrawContainer"))
+		if (PropertyChangedEvent.MemberProperty->GetNameCPP() == GET_MEMBER_NAME_STRING_CHECKED(UControlRigBlueprint, DrawContainer))
 		{
 			ControlRigBP->PropagateDrawInstructionsFromBPToInstances();
 			return;
@@ -4433,6 +4463,21 @@ void FControlRigEditor::OnFinishedChangingProperties(const FPropertyChangedEvent
 	}
 }
 
+void FControlRigEditor::OnPropertyChanged(UObject* InObject, FPropertyChangedEvent& InEvent)
+{
+	UControlRigBlueprint* ControlRigBP = GetControlRigBlueprint();
+
+	if (ControlRigBP && InObject == ControlRigBP)
+	{
+		// if the models have changed - we may need to close a document
+		if(InEvent.MemberProperty == ControlRigBP->GetClass()->FindPropertyByName(GET_MEMBER_NAME_CHECKED(UControlRigBlueprint, RigVMClient)) ||
+			InEvent.MemberProperty == ControlRigBP->GetClass()->FindPropertyByName(GET_MEMBER_NAME_CHECKED(UControlRigBlueprint, UbergraphPages)))
+		{
+			DocumentManager->CleanInvalidTabs();
+		}
+	}
+}
+
 void FControlRigEditor::OnWrappedPropertyChangedChainEvent(UDetailsViewWrapperObject* InWrapperObject, const FString& InPropertyPath, FPropertyChangedChainEvent& InPropertyChangedChainEvent)
 {
 	check(InWrapperObject);
@@ -4458,19 +4503,20 @@ void FControlRigEditor::OnWrappedPropertyChangedChainEvent(UDetailsViewWrapperOb
 				return;
 			}
 
-			const FString PoseString = TEXT("Pose->");
-			const FString OffsetString = TEXT("Offset->");
-			const FString ShapeString = TEXT("Shape->");
-			const FString SettingsString = TEXT("Settings->");
+			static constexpr TCHAR PropertyChainElementFormat[] = TEXT("%s->");
+			static const FString PoseString = FString::Printf(PropertyChainElementFormat, GET_MEMBER_NAME_STRING_CHECKED(FRigTransformElement, Pose));
+			static const FString OffsetString = FString::Printf(PropertyChainElementFormat, GET_MEMBER_NAME_STRING_CHECKED(FRigControlElement, Offset));
+			static const FString ShapeString = FString::Printf(PropertyChainElementFormat, GET_MEMBER_NAME_STRING_CHECKED(FRigControlElement, Shape));
+			static const FString SettingsString = FString::Printf(PropertyChainElementFormat, GET_MEMBER_NAME_STRING_CHECKED(FRigControlElement, Settings));
 
 			struct Local
 			{
 				static ERigTransformType::Type GetTransformTypeFromPath(FString& PropertyPath)
 				{
-					const FString InitialString = TEXT("Initial->");
-					const FString CurrentString = TEXT("Current->");
-					const FString GlobalString = TEXT("Global->");
-					const FString LocalString = TEXT("Local->");
+					static const FString InitialString = FString::Printf(PropertyChainElementFormat, GET_MEMBER_NAME_STRING_CHECKED(FRigCurrentAndInitialTransform, Initial));
+					static const FString CurrentString = FString::Printf(PropertyChainElementFormat, GET_MEMBER_NAME_STRING_CHECKED(FRigCurrentAndInitialTransform, Current));
+					static const FString GlobalString = FString::Printf(PropertyChainElementFormat, GET_MEMBER_NAME_STRING_CHECKED(FRigLocalAndGlobalTransform, Global));
+					static const FString LocalString = FString::Printf(PropertyChainElementFormat, GET_MEMBER_NAME_STRING_CHECKED(FRigLocalAndGlobalTransform, Local));
 
 					ERigTransformType::Type TransformType = ERigTransformType::CurrentLocal;
 
@@ -4670,7 +4716,7 @@ void FControlRigEditor::OnRequestLocalizeFunctionDialog(URigVMLibraryNode* InFun
 		return;
 	}
 	
-	if(URigVMController* TargetController = InTargetBlueprint->GetController(InTargetBlueprint->GetModel()))
+	if(URigVMController* TargetController = InTargetBlueprint->GetController(InTargetBlueprint->GetDefaultModel()))
 	{
 		if(URigVMFunctionLibrary* FunctionLibrary = Cast<URigVMFunctionLibrary>(InFunction->GetOuter()))
 		{
@@ -5992,7 +6038,7 @@ void FControlRigEditor::HandleOnControlModified(UControlRig* Subject, FRigContro
 		FRigControlValue ControlValue = Hierarchy->GetControlValue(ControlElement, ERigControlValueType::Current);
 
 		const FString PinPath = UControlRig::GetPinNameFromTransientControl(ControlElement->GetKey());
-		if (URigVMPin* Pin = Blueprint->Model->FindPin(PinPath))
+		if (URigVMPin* Pin = Blueprint->GetRigVMClient()->FindPin(PinPath))
 		{
 			FString NewDefaultValue;
 			switch (ControlElement->Settings.ControlType)
@@ -6423,7 +6469,7 @@ URigVMController* FControlRigEditor::GetFocusedController() const
 	{
 		return nullptr;
 	}
-	return Blueprint->GetController(GetFocusedModel());
+	return Blueprint->GetOrCreateController(GetFocusedModel());
 }
 
 TSharedPtr<SGraphEditor> FControlRigEditor::GetGraphEditor(UEdGraph* InEdGraph) const
diff --git a/Engine/Plugins/Animation/ControlRig/Source/ControlRigEditor/Private/Editor/ControlRigEditor.h b/Engine/Plugins/Animation/ControlRig/Source/ControlRigEditor/Private/Editor/ControlRigEditor.h
--- a/Engine/Plugins/Animation/ControlRig/Source/ControlRigEditor/Private/Editor/ControlRigEditor.h
+++ b/Engine/Plugins/Animation/ControlRig/Source/ControlRigEditor/Private/Editor/ControlRigEditor.h
@@ -143,6 +143,8 @@ public:
 
 	virtual void PostTransaction(bool bSuccess, const FTransaction* Transaction, bool bIsRedo);
 
+	virtual void JumpToHyperlink(const UObject* ObjectReference, bool bRequestRename = false) override;
+
 	void EnsureValidRigElementsInDetailPanel();
 
 	virtual void OnStartWatchingPin();
@@ -436,6 +438,7 @@ protected:
 	virtual void NotifyPostChange(const FPropertyChangedEvent& PropertyChangedEvent, FProperty* PropertyThatChanged) override;
 	/** delegate for changing property */
 	virtual void OnFinishedChangingProperties(const FPropertyChangedEvent& PropertyChangedEvent) override;
+	void OnPropertyChanged(UObject* InObject, FPropertyChangedEvent& InEvent);
 	void OnWrappedPropertyChangedChainEvent(UDetailsViewWrapperObject* InWrapperObject, const FString& InPropertyPath, FPropertyChangedChainEvent& InPropertyChangedChainEvent);
 	void OnRequestLocalizeFunctionDialog(URigVMLibraryNode* InFunction, UControlRigBlueprint* InTargetBlueprint, bool bForce);
 	FRigVMController_BulkEditResult OnRequestBulkEditDialog(UControlRigBlueprint* InBlueprint, URigVMController* InController, URigVMLibraryNode* InFunction, ERigVMControllerBulkEditType InEditType);
@@ -500,6 +503,8 @@ protected:
 
 	TArray<TStrongObjectPtr<UDetailsViewWrapperObject>> WrapperObjects;
 	TWeakObjectPtr<AStaticMeshActor> WeakGroundActorPtr;
+
+	FDelegateHandle PropertyChangedHandle;
 	
 	friend class FControlRigEditorMode;
 	friend class SControlRigStackView;
diff --git a/Engine/Plugins/Animation/ControlRig/Source/ControlRigEditor/Private/Editor/SControlRigStackView.cpp b/Engine/Plugins/Animation/ControlRig/Source/ControlRigEditor/Private/Editor/SControlRigStackView.cpp
--- a/Engine/Plugins/Animation/ControlRig/Source/ControlRigEditor/Private/Editor/SControlRigStackView.cpp
+++ b/Engine/Plugins/Animation/ControlRig/Source/ControlRigEditor/Private/Editor/SControlRigStackView.cpp
@@ -442,11 +442,13 @@ void SControlRigStackView::PopulateStackView(URigVM* InVM)
 	if (InVM)
 	{
 		UControlRigBlueprint* Blueprint = ControlRigEditor.Pin()->GetControlRigBlueprint();
-		URigVMGraph* RootGraph = Blueprint->GetModel();
 
 		const FRigVMInstructionArray Instructions = InVM->GetInstructions();
 		const FRigVMByteCode& ByteCode = InVM->GetByteCode();
 
+		TArray<URigVMGraph*> RootGraphs;
+		RootGraphs.AddZeroed(Instructions.Num());
+
 		const bool bUseSimpleLabels = !CVarControlRigExecutionStackDetailedLabels.GetValueOnAnyThread();
 		if(bUseSimpleLabels)
 		{
@@ -459,6 +461,7 @@ void SControlRigStackView::PopulateStackView(URigVM* InVM)
 				if(Node)
 				{
 					DisplayName = Node->GetName();
+					RootGraphs[InstructionIndex] = Node->GetRootGraph();
 					
 					// only unit nodes among all nodes has StaticExecute() that generates actual instructions
 					if (URigVMUnitNode* UnitNode = Cast<URigVMUnitNode>(Node))
@@ -638,7 +641,7 @@ void SControlRigStackView::PopulateStackView(URigVM* InVM)
 		TMap<FString, FString> OperandFormatMap;
 		for (int32 InstructionIndex = 0; InstructionIndex < Instructions.Num(); InstructionIndex++)
 		{
-			const FRigVMASTProxy Proxy = FRigVMASTProxy::MakeFromCallPath(ByteCode.GetCallPathForInstruction(InstructionIndex), RootGraph);
+			const FRigVMASTProxy Proxy = FRigVMASTProxy::MakeFromCallPath(ByteCode.GetCallPathForInstruction(InstructionIndex), RootGraphs[InstructionIndex]);
 			if(URigVMNode* Node = Proxy.GetSubject<URigVMNode>())
 			{
 				FString DisplayName = Node->GetName();
@@ -692,7 +695,7 @@ void SControlRigStackView::PopulateStackView(URigVM* InVM)
 		for (int32 InstructionIndex = 0; InstructionIndex < Labels.Num(); InstructionIndex++)
 		{
 			FString Label = Labels[InstructionIndex];
-			const FRigVMASTProxy Proxy = FRigVMASTProxy::MakeFromCallPath(ByteCode.GetCallPathForInstruction(InstructionIndex), RootGraph);
+			const FRigVMASTProxy Proxy = FRigVMASTProxy::MakeFromCallPath(ByteCode.GetCallPathForInstruction(InstructionIndex), RootGraphs[InstructionIndex]);
 
 			if(URigVMNode* Node = Proxy.GetSubject<URigVMNode>())
 			{
diff --git a/Engine/Plugins/Animation/ControlRig/Source/ControlRigEditor/Private/Graph/NodeSpawners/ControlRigUnitNodeSpawner.cpp b/Engine/Plugins/Animation/ControlRig/Source/ControlRigEditor/Private/Graph/NodeSpawners/ControlRigUnitNodeSpawner.cpp
--- a/Engine/Plugins/Animation/ControlRig/Source/ControlRigEditor/Private/Graph/NodeSpawners/ControlRigUnitNodeSpawner.cpp
+++ b/Engine/Plugins/Animation/ControlRig/Source/ControlRigEditor/Private/Graph/NodeSpawners/ControlRigUnitNodeSpawner.cpp
@@ -120,6 +120,32 @@ UControlRigGraphNode* UControlRigUnitNodeSpawner::SpawnNode(UEdGraph* ParentGrap
 		bool const bIsUserFacingNode = !bIsTemplateNode;
 
 		FName Name = bIsTemplateNode ? *StructTemplate->GetStructCPPName() : FControlRigBlueprintUtils::ValidateName(RigBlueprint, StructTemplate->GetFName().ToString());
+
+		// events can only exist once across all uber graphs
+		if(bIsUserFacingNode)
+		{
+			FStructOnScope StructScope(StructTemplate);
+			const FRigVMStruct* StructInstance = (const FRigVMStruct*)StructScope.GetStructMemory();
+			const FName EventName = StructInstance->GetEventName();
+			if(!EventName.IsNone())
+			{
+				const TArray<URigVMGraph*> Models = RigBlueprint->GetAllModels();
+				for(URigVMGraph* Model : Models)
+				{
+					if(Model->IsRootGraph() && !Model->IsA<URigVMFunctionLibrary>())
+					{
+						for(const URigVMNode* Node : Model->GetNodes())
+						{
+							if(Node->GetEventName() == EventName)
+							{
+								RigBlueprint->OnRequestJumpToHyperlink().Execute(Node);
+								return nullptr;
+							}
+						}
+					}
+				}
+			}
+		}
 		
 		URigVMController* Controller = bIsTemplateNode ? RigGraph->GetTemplateController() : RigBlueprint->GetController(ParentGraph);
 
diff --git a/Engine/Plugins/Experimental/DataInterfaceGraph/Source/DataInterfaceGraphEditor/Private/DataInterfaceGraphEditor.cpp b/Engine/Plugins/Experimental/DataInterfaceGraph/Source/DataInterfaceGraphEditor/Private/DataInterfaceGraphEditor.cpp
--- a/Engine/Plugins/Experimental/DataInterfaceGraph/Source/DataInterfaceGraphEditor/Private/DataInterfaceGraphEditor.cpp
+++ b/Engine/Plugins/Experimental/DataInterfaceGraph/Source/DataInterfaceGraphEditor/Private/DataInterfaceGraphEditor.cpp
@@ -145,12 +145,12 @@ UEdGraph* FGraphEditor::GetFocusedGraph() const
 URigVMGraph* FGraphEditor::GetFocusedVMGraph() const
 {
 	UDataInterfaceGraph_EdGraph* EdGraph = Cast<UDataInterfaceGraph_EdGraph>(GetFocusedGraph());
-	return DataInterfaceGraph_EditorData->GetRigVMGraph(EdGraph);
+	return DataInterfaceGraph_EditorData->GetVMGraphForEdGraph(EdGraph);
 }
 
 URigVMController* FGraphEditor::GetFocusedVMController() const
 {
-	return DataInterfaceGraph_EditorData->GetRigVMController(GetFocusedVMGraph());
+	return DataInterfaceGraph_EditorData->GetRigVMClient()->GetController(GetFocusedVMGraph());
 }
 
 TSharedRef<SGraphEditor> FGraphEditor::CreateGraphEditorWidget(TSharedRef<FTabInfo> InTabInfo, UEdGraph* InGraph)
diff --git a/Engine/Plugins/Experimental/DataInterfaceGraph/Source/DataInterfaceGraphEditor/Private/DataInterfaceGraphEditorSchemaActions.cpp b/Engine/Plugins/Experimental/DataInterfaceGraph/Source/DataInterfaceGraphEditor/Private/DataInterfaceGraphEditorSchemaActions.cpp
--- a/Engine/Plugins/Experimental/DataInterfaceGraph/Source/DataInterfaceGraphEditor/Private/DataInterfaceGraphEditorSchemaActions.cpp
+++ b/Engine/Plugins/Experimental/DataInterfaceGraph/Source/DataInterfaceGraphEditor/Private/DataInterfaceGraphEditorSchemaActions.cpp
@@ -19,7 +19,7 @@ UEdGraphNode* FDataInterfaceGraphSchemaAction_RigUnit::PerformAction(UEdGraph* P
 	if (EditorData != nullptr && EdGraph != nullptr)
 	{
 		FName Name = UE::DataInterfaceGraphEditor::FUtils::ValidateName(EditorData, StructTemplate->GetFName().ToString());
-		URigVMController* Controller = EditorData->GetRigVMController(ParentGraph);
+		URigVMController* Controller = EditorData->GetRigVMClient()->GetController(ParentGraph);
 
 		Controller->OpenUndoBracket(FString::Printf(TEXT("Add '%s' Node"), *Name.ToString()));
 
diff --git a/Engine/Plugins/Experimental/DataInterfaceGraph/Source/DataInterfaceGraphEditor/Private/DataInterfaceGraphFactory.cpp b/Engine/Plugins/Experimental/DataInterfaceGraph/Source/DataInterfaceGraphEditor/Private/DataInterfaceGraphFactory.cpp
--- a/Engine/Plugins/Experimental/DataInterfaceGraph/Source/DataInterfaceGraphEditor/Private/DataInterfaceGraphFactory.cpp
+++ b/Engine/Plugins/Experimental/DataInterfaceGraph/Source/DataInterfaceGraphEditor/Private/DataInterfaceGraphFactory.cpp
@@ -34,18 +34,18 @@ UObject* UDataInterfaceGraphFactory::FactoryCreateNew(UClass* Class, UObject* In
 	EditorData->Initialize(/*bRecompileVM*/false);
 
 	// Add initial execution unit
-	URigVMController* RootController = EditorData->GetRigVMController(EditorData->RootGraph);
+	URigVMController* RootController = EditorData->GetRigVMClient()->GetController(EditorData->RootGraph);
 	URigVMUnitNode* MainEntryPointNode = RootController->AddUnitNode(FRigUnit_DataInterfaceBeginExecution::StaticStruct(), FRigUnit::GetMethodName(), FVector2D(-400.0f, 0.0f), FString(), false);
 	URigVMPin* BeginExecutePin = MainEntryPointNode->FindPin(GET_MEMBER_NAME_STRING_CHECKED(FRigUnit_DataInterfaceBeginExecution, ExecuteContext));
 	check(BeginExecutePin);
 	check(BeginExecutePin->GetDirection() == ERigVMPinDirection::Output);
 	
 	// Add function to function lib
-	URigVMController* FunctionLibraryController = EditorData->GetRigVMController(EditorData->RigVMFunctionLibrary);
+	URigVMController* FunctionLibraryController = EditorData->GetRigVMClient()->GetController(EditorData->GetRigVMClient()->GetFunctionLibrary());
 	EditorData->EntryPoint = FunctionLibraryController->AddFunctionToLibrary(UE::DataInterfaceGraph::EntryPointName, true, FVector2D::ZeroVector, false, false);
 	
 	// Add exposed pin for result to function
-	URigVMController* EntryPointController = EditorData->GetOrCreateRigVMController(EditorData->EntryPointGraph);
+	URigVMController* EntryPointController = EditorData->GetRigVMClient()->GetOrCreateController(EditorData->EntryPointGraph);
 	
 	// TODO: using float for now, but needs to be user-driven
 	{
diff --git a/Engine/Plugins/Experimental/DataInterfaceGraph/Source/DataInterfaceGraphUncookedOnly/Internal/DataInterfaceGraph_EdGraph.h b/Engine/Plugins/Experimental/DataInterfaceGraph/Source/DataInterfaceGraphUncookedOnly/Internal/DataInterfaceGraph_EdGraph.h
--- a/Engine/Plugins/Experimental/DataInterfaceGraph/Source/DataInterfaceGraphUncookedOnly/Internal/DataInterfaceGraph_EdGraph.h
+++ b/Engine/Plugins/Experimental/DataInterfaceGraph/Source/DataInterfaceGraphUncookedOnly/Internal/DataInterfaceGraph_EdGraph.h
@@ -21,8 +21,7 @@ class UDataInterfaceGraph_EdGraph : public UControlRigGraph
 	friend class UDataInterfaceGraph_EditorData;
 
 	// UControlRigGraph interface
-	virtual URigVMGraph* GetModel() const override;
-	virtual URigVMController* GetController() const override;
+	virtual FRigVMClient* GetRigVMClient() const override;
 	
 	void Initialize(UDataInterfaceGraph_EditorData* InEditorData);
 };
\ No newline at end of file
diff --git a/Engine/Plugins/Experimental/DataInterfaceGraph/Source/DataInterfaceGraphUncookedOnly/Internal/DataInterfaceGraph_EditorData.h b/Engine/Plugins/Experimental/DataInterfaceGraph/Source/DataInterfaceGraphUncookedOnly/Internal/DataInterfaceGraph_EditorData.h
--- a/Engine/Plugins/Experimental/DataInterfaceGraph/Source/DataInterfaceGraphUncookedOnly/Internal/DataInterfaceGraph_EditorData.h
+++ b/Engine/Plugins/Experimental/DataInterfaceGraph/Source/DataInterfaceGraphUncookedOnly/Internal/DataInterfaceGraph_EditorData.h
@@ -20,7 +20,7 @@ namespace UE::DataInterfaceGraphEditor
 }
 
 UCLASS(Optional, MinimalAPI)
-class UDataInterfaceGraph_EditorData : public UObject, public IRigVMGraphHost, public IRigVMControllerHost
+class UDataInterfaceGraph_EditorData : public UObject, public IRigVMClientHost
 {
 	GENERATED_BODY()
 
@@ -33,17 +33,18 @@ class UDataInterfaceGraph_EditorData : public UObject, public IRigVMGraphHost, p
 	friend struct FDataInterfaceGraphSchemaAction_RigUnit;
 
 	// UObject interface
+	virtual void Serialize(FArchive& Ar) override;
 	virtual void PostLoad() override;
 	virtual bool IsEditorOnly() const override { return true; }
 	
-	// IRigVMGraphHost interface
-	virtual URigVMGraph* GetRigVMGraph(const UObject* InEditorObject) const override;
-
-	// IRigVMControllerHost interface
-	virtual URigVMController* GetRigVMController(const URigVMGraph* InRigVMGraph) const override;
-	virtual URigVMController* GetRigVMController(const UObject* InEditorObject) const override;
-	virtual URigVMController* GetOrCreateRigVMController(URigVMGraph* InRigVMGraph) override;
-	virtual URigVMController* GetOrCreateRigVMController(const UObject* InEditorObject) override;
+	// IRigVMClientHost interface
+	virtual FRigVMClient* GetRigVMClient() override;
+	virtual const FRigVMClient* GetRigVMClient() const override;
+	virtual void HandleRigVMGraphAdded(const FRigVMClient* InClient, const FString& InNodePath) override;
+	virtual void HandleRigVMGraphRemoved(const FRigVMClient* InClient, const FString& InNodePath) override {}
+	virtual void HandleRigVMGraphRenamed(const FRigVMClient* InClient, const FString& InOldNodePath, const FString& InNewNodePath) override {}
+	virtual void HandleConfigureRigVMController(const FRigVMClient* InClient, URigVMController* InControllerToConfigure) override;
+	virtual UObject* GetEditorObjectForRigVMGraph(URigVMGraph* InVMGraph) const override;
 
 	DATAINTERFACEGRAPHUNCOOKEDONLY_API void Initialize(bool bRecompileVM);
 
@@ -55,7 +56,7 @@ class UDataInterfaceGraph_EditorData : public UObject, public IRigVMGraphHost, p
 	
 	void HandleModifiedEvent(ERigVMGraphNotifType InNotifType, URigVMGraph* InGraph, UObject* InSubject);
 
-	URigVMGraph* GetVMGraphForEdGraph(const UEdGraph* InGraph) const;
+	DATAINTERFACEGRAPHUNCOOKEDONLY_API URigVMGraph* GetVMGraphForEdGraph(const UEdGraph* InGraph) const;
 
 	void CreateEdGraphForCollapseNode(URigVMCollapseNode* InNode);
 	
@@ -67,12 +68,15 @@ class UDataInterfaceGraph_EditorData : public UObject, public IRigVMGraphHost, p
 	
 	UPROPERTY()
 	TObjectPtr<UDataInterfaceGraph_EdGraph> FunctionLibraryEdGraph;
-	
+
+	UPROPERTY()
+	FRigVMClient RigVMClient;
+
 	UPROPERTY()
-	TObjectPtr<URigVMGraph> RigVMGraph;
+	TObjectPtr<URigVMGraph> RigVMGraph_DEPRECATED;
 	
 	UPROPERTY()
-	TObjectPtr<URigVMFunctionLibrary> RigVMFunctionLibrary;
+	TObjectPtr<URigVMFunctionLibrary> RigVMFunctionLibrary_DEPRECATED;
 
 	UPROPERTY()
 	TObjectPtr<URigVMLibraryNode> EntryPoint;
diff --git a/Engine/Plugins/Experimental/DataInterfaceGraph/Source/DataInterfaceGraphUncookedOnly/Private/DataInterfaceGraphUncookedOnlyUtils.cpp b/Engine/Plugins/Experimental/DataInterfaceGraph/Source/DataInterfaceGraphUncookedOnly/Private/DataInterfaceGraphUncookedOnlyUtils.cpp
--- a/Engine/Plugins/Experimental/DataInterfaceGraph/Source/DataInterfaceGraphUncookedOnly/Private/DataInterfaceGraphUncookedOnlyUtils.cpp
+++ b/Engine/Plugins/Experimental/DataInterfaceGraph/Source/DataInterfaceGraphUncookedOnly/Private/DataInterfaceGraphUncookedOnlyUtils.cpp
@@ -56,8 +56,8 @@ void FUtils::Compile(UDataInterfaceGraph* InGraph)
 
 	URigVMCompiler* Compiler = URigVMCompiler::StaticClass()->GetDefaultObject<URigVMCompiler>();
 	Compiler->Settings = (EditorData->bCompileInDebugMode) ? FRigVMCompileSettings::Fast() : EditorData->VMCompileSettings;
-	URigVMController* RootController = EditorData->GetOrCreateRigVMController(EditorData->RigVMGraph);
-	Compiler->Compile(EditorData->RigVMGraph, RootController, InGraph->RigVM, InGraph->GetRigVMExternalVariables(), UserData, &EditorData->PinToOperandMap);
+	URigVMController* RootController = EditorData->GetRigVMClient()->GetOrCreateController(EditorData->GetRigVMClient()->GetDefaultModel());
+	Compiler->Compile(EditorData->GetRigVMClient()->GetAllModels(false, false), RootController, InGraph->RigVM, InGraph->GetRigVMExternalVariables(), UserData, &EditorData->PinToOperandMap);
 
 	if (EditorData->bErrorsDuringCompilation)
 	{
diff --git a/Engine/Plugins/Experimental/DataInterfaceGraph/Source/DataInterfaceGraphUncookedOnly/Private/DataInterfaceGraph_EdGraph.cpp b/Engine/Plugins/Experimental/DataInterfaceGraph/Source/DataInterfaceGraphUncookedOnly/Private/DataInterfaceGraph_EdGraph.cpp
--- a/Engine/Plugins/Experimental/DataInterfaceGraph/Source/DataInterfaceGraphUncookedOnly/Private/DataInterfaceGraph_EdGraph.cpp
+++ b/Engine/Plugins/Experimental/DataInterfaceGraph/Source/DataInterfaceGraphUncookedOnly/Private/DataInterfaceGraph_EdGraph.cpp
@@ -12,20 +12,11 @@ void UDataInterfaceGraph_EdGraph::Initialize(UDataInterfaceGraph_EditorData* InE
 	InEditorData->VMCompiledEvent.AddUObject(this, &UDataInterfaceGraph_EdGraph::HandleVMCompiledEvent);
 }
 
-URigVMGraph* UDataInterfaceGraph_EdGraph::GetModel() const
+FRigVMClient* UDataInterfaceGraph_EdGraph::GetRigVMClient() const
 {
 	if (const UDataInterfaceGraph_EditorData* EditorData = GetTypedOuter<UDataInterfaceGraph_EditorData>())
 	{
-		return EditorData->GetRigVMGraph(this);
+		return (FRigVMClient*)EditorData->GetRigVMClient();
 	}
 	return nullptr;
 }
-
-URigVMController* UDataInterfaceGraph_EdGraph::GetController() const
-{
-	if (const UDataInterfaceGraph_EditorData* EditorData = GetTypedOuter<UDataInterfaceGraph_EditorData>())
-	{
-		return EditorData->GetRigVMController(this);
-	}
-	return nullptr;
-}
\ No newline at end of file
diff --git a/Engine/Plugins/Experimental/DataInterfaceGraph/Source/DataInterfaceGraphUncookedOnly/Private/DataInterfaceGraph_EditorData.cpp b/Engine/Plugins/Experimental/DataInterfaceGraph/Source/DataInterfaceGraphUncookedOnly/Private/DataInterfaceGraph_EditorData.cpp
--- a/Engine/Plugins/Experimental/DataInterfaceGraph/Source/DataInterfaceGraphUncookedOnly/Private/DataInterfaceGraph_EditorData.cpp
+++ b/Engine/Plugins/Experimental/DataInterfaceGraph/Source/DataInterfaceGraphUncookedOnly/Private/DataInterfaceGraph_EditorData.cpp
@@ -15,12 +15,13 @@
 
 UDataInterfaceGraph_EditorData::UDataInterfaceGraph_EditorData(const FObjectInitializer& ObjectInitializer)
 {
-	RigVMGraph = ObjectInitializer.CreateDefaultSubobject<URigVMGraph>(this, TEXT("RigVMGraph"));
-	RigVMFunctionLibrary = ObjectInitializer.CreateDefaultSubobject<URigVMFunctionLibrary>(this, TEXT("RigVMFunctionLibrary"));
-	RigVMGraph->SetDefaultFunctionLibrary(RigVMFunctionLibrary);
-
-	RigVMGraph->SetExecuteContextStruct(FRigVMExecuteContext::StaticStruct());
-	RigVMFunctionLibrary->SetExecuteContextStruct(FRigVMExecuteContext::StaticStruct());
+	RigVMClient.Reset();
+	RigVMClient.SetOuterClientHost(this, GET_MEMBER_NAME_CHECKED(UDataInterfaceGraph_EditorData, RigVMClient));
+	{
+		TGuardValue<bool> DisableClientNotifs(RigVMClient.bSuspendNotifications, true);
+		RigVMClient.AddModel(TEXT("RigVMGraph"), false);
+		RigVMClient.GetOrCreateFunctionLibrary(false);
+	}
 	
 	auto MakeEdGraph = [this, &ObjectInitializer](FName InName) -> UDataInterfaceGraph_EdGraph*
 	{
@@ -39,19 +40,41 @@ UDataInterfaceGraph_EditorData::UDataInterfaceGraph_EditorData(const FObjectInit
 	FunctionLibraryEdGraph = MakeEdGraph(TEXT("RigVMFunctionLibraryEdGraph"));
 }
 
+void UDataInterfaceGraph_EditorData::Serialize(FArchive& Ar)
+{
+	RigVMClient.SetOuterClientHost(this, GET_MEMBER_NAME_CHECKED(UDataInterfaceGraph_EditorData, RigVMClient));
+
+	UObject::Serialize(Ar);
+
+	if(Ar.IsLoading())
+	{
+		if(Ar.IsLoading())
+		{
+			if(RigVMGraph_DEPRECATED || RigVMFunctionLibrary_DEPRECATED)
+			{
+				TGuardValue<bool> DisableClientNotifs(RigVMClient.bSuspendNotifications, true);
+				RigVMClient.SetFromDeprecatedData(RigVMGraph_DEPRECATED, RigVMFunctionLibrary_DEPRECATED);
+			}
+		}
+	}
+}
+
 void UDataInterfaceGraph_EditorData::Initialize(bool bRecompileVM)
 {
 	UDataInterfaceGraph* DataInterfaceGraph = GetTypedOuter<UDataInterfaceGraph>();
-	
-	if (Controllers.Num() == 0)
+
+	if (RigVMClient.GetController(0) == nullptr)
 	{
-		GetOrCreateRigVMController(RigVMGraph);
+		check(RigVMClient.Num() == 1);
+		check(RigVMClient.GetFunctionLibrary());
+		
+		RigVMClient.GetOrCreateController(RigVMClient.GetDefaultModel());
+		RigVMClient.GetOrCreateController(RigVMClient.GetFunctionLibrary());
 
 		// Init function library controllers
-		GetOrCreateRigVMController(RigVMFunctionLibrary);
-		for(URigVMLibraryNode* LibraryNode : RigVMFunctionLibrary->GetFunctions())
+		for(URigVMLibraryNode* LibraryNode : RigVMClient.GetFunctionLibrary()->GetFunctions())
 		{
-			GetOrCreateRigVMController(LibraryNode->GetContainedGraph());
+			RigVMClient.GetOrCreateController(LibraryNode->GetContainedGraph());
 		}
 		
 		if(bRecompileVM)
@@ -60,9 +83,6 @@ void UDataInterfaceGraph_EditorData::Initialize(bool bRecompileVM)
 		}
 	}
 
-	RigVMGraph->SetExecuteContextStruct(FRigVMExecuteContext::StaticStruct());
-	RigVMFunctionLibrary->SetExecuteContextStruct(FRigVMExecuteContext::StaticStruct());
-	
 	RootGraph->Initialize(this);
 	FunctionLibraryEdGraph->Initialize(this);
 	if(EntryPointGraph)
@@ -78,119 +98,30 @@ void UDataInterfaceGraph_EditorData::PostLoad()
 	Initialize(/*bRecompileVM*/false);
 }
 
-void UDataInterfaceGraph_EditorData::RecompileVM()
+FRigVMClient* UDataInterfaceGraph_EditorData::GetRigVMClient()
 {
-	UE::DataInterfaceGraphUncookedOnly::FUtils::Compile(GetTypedOuter<UDataInterfaceGraph>());
+	return &RigVMClient;
 }
 
-void UDataInterfaceGraph_EditorData::RecompileVMIfRequired()
+const FRigVMClient* UDataInterfaceGraph_EditorData::GetRigVMClient() const
 {
-	if (bVMRecompilationRequired)
-	{
-		RecompileVM();
-	}
+	return &RigVMClient;
 }
 
-void UDataInterfaceGraph_EditorData::RequestAutoVMRecompilation()
+void UDataInterfaceGraph_EditorData::HandleRigVMGraphAdded(const FRigVMClient* InClient, const FString& InNodePath)
 {
-	bVMRecompilationRequired = true;
-	if (bAutoRecompileVM)
+	if(URigVMGraph* RigVMGraph = InClient->GetModel(InNodePath))
 	{
-		RecompileVMIfRequired();
+		RigVMGraph->SetExecuteContextStruct(FRigVMExecuteContext::StaticStruct());
 	}
 }
 
-void UDataInterfaceGraph_EditorData::HandleModifiedEvent(ERigVMGraphNotifType InNotifType, URigVMGraph* InGraph, UObject* InSubject)
+void UDataInterfaceGraph_EditorData::HandleConfigureRigVMController(const FRigVMClient* InClient,
+                                                                    URigVMController* InControllerToConfigure)
 {
-	bool bNotifForOthersPending = true;
+	InControllerToConfigure->OnModified().AddUObject(this, &UDataInterfaceGraph_EditorData::HandleModifiedEvent);
 
-	switch(InNotifType)
-	{
-	case ERigVMGraphNotifType::NodeAdded:
-		{
-			if (URigVMCollapseNode* CollapseNode = Cast<URigVMCollapseNode>(InSubject))
-			{
-				CreateEdGraphForCollapseNode(CollapseNode);
-				break;
-			}
-		}
-		// Fall through to the next case
-	case ERigVMGraphNotifType::LinkAdded:
-	case ERigVMGraphNotifType::LinkRemoved:
-	case ERigVMGraphNotifType::PinArraySizeChanged:
-	case ERigVMGraphNotifType::PinDirectionChanged:
-	case ERigVMGraphNotifType::PinDefaultValueChanged:
-		{
-			RequestAutoVMRecompilation();
-			break;
-		}
-	}
-	
-	// if the notification still has to be sent...
-	if (bNotifForOthersPending && !bSuspendModelNotificationsForOthers)
-	{
-		if (ModifiedEvent.IsBound())
-		{
-			ModifiedEvent.Broadcast(InNotifType, InGraph, InSubject);
-		}
-	}
-}
-
-URigVMGraph* UDataInterfaceGraph_EditorData::GetRigVMGraph(const UObject* InEditorObject) const
-{
-	return GetVMGraphForEdGraph(CastChecked<UEdGraph>(InEditorObject));
-}
-
-URigVMGraph* UDataInterfaceGraph_EditorData::GetVMGraphForEdGraph(const UEdGraph* InGraph) const
-{
-	if (InGraph == RootGraph)
-	{
-		return RigVMGraph;
-	}
-	else
-	{
-		const UDataInterfaceGraph_EdGraph* Graph = Cast<UDataInterfaceGraph_EdGraph>(InGraph);
-		check(Graph);
-
-		if (Graph->bIsFunctionDefinition)
-		{
-			if (URigVMLibraryNode* LibraryNode = RigVMFunctionLibrary->FindFunction(*Graph->ModelNodePath))
-			{
-				return LibraryNode->GetContainedGraph();
-			}
-		}
-	}
-
-	return nullptr;
-}
-
-URigVMController* UDataInterfaceGraph_EditorData::GetRigVMController(const URigVMGraph* InRigVMGraph) const
-{
-	TObjectPtr<URigVMController> const* ControllerPtr = Controllers.Find(InRigVMGraph);
-	if (ControllerPtr)
-	{
-		return *ControllerPtr;
-	}
-	return nullptr;
-}
-
-URigVMController* UDataInterfaceGraph_EditorData::GetRigVMController(const UObject* InEditorObject) const
-{
-	return GetRigVMController(GetVMGraphForEdGraph(CastChecked<UEdGraph>(InEditorObject)));
-}
-
-URigVMController* UDataInterfaceGraph_EditorData::GetOrCreateRigVMController(URigVMGraph* InRigVMGraph)
-{
-	if (URigVMController* ExistingController = GetRigVMController(InRigVMGraph))
-	{
-		return ExistingController;
-	}
-
-	URigVMController* Controller = NewObject<URigVMController>(this);
-	Controller->SetGraph(InRigVMGraph);
-	Controller->OnModified().AddUObject(this, &UDataInterfaceGraph_EditorData::HandleModifiedEvent);
-
-	Controller->UnfoldStructDelegate.BindLambda([](const UStruct* InStruct) -> bool
+	InControllerToConfigure->UnfoldStructDelegate.BindLambda([](const UStruct* InStruct) -> bool
 	{
 		if (InStruct == TBaseStructure<FQuat>::Get())
 		{
@@ -212,7 +143,7 @@ URigVMController* UDataInterfaceGraph_EditorData::GetOrCreateRigVMController(URi
 	// this delegate is used by the controller to determine variable validity
 	// during a bind process. the controller itself doesn't own the variables,
 	// so we need a delegate to request them from the owning blueprint
-	Controller->GetExternalVariablesDelegate.BindLambda([](URigVMGraph* InGraph) -> TArray<FRigVMExternalVariable>
+	InControllerToConfigure->GetExternalVariablesDelegate.BindLambda([](URigVMGraph* InGraph) -> TArray<FRigVMExternalVariable>
 	{
 		if (InGraph)
 		{
@@ -228,7 +159,7 @@ URigVMController* UDataInterfaceGraph_EditorData::GetOrCreateRigVMController(URi
 	});
 
 	// this delegate is used by the controller to retrieve the current bytecode of the VM
-	Controller->GetCurrentByteCodeDelegate.BindLambda([WeakThis]() -> const FRigVMByteCode*
+	InControllerToConfigure->GetCurrentByteCodeDelegate.BindLambda([WeakThis]() -> const FRigVMByteCode*
 	{
 		if (WeakThis.IsValid())
 		{
@@ -244,18 +175,18 @@ URigVMController* UDataInterfaceGraph_EditorData::GetOrCreateRigVMController(URi
 
 	});
 
-	Controller->IsFunctionAvailableDelegate.BindLambda([](URigVMLibraryNode* InFunction) -> bool
+	InControllerToConfigure->IsFunctionAvailableDelegate.BindLambda([](URigVMLibraryNode* InFunction) -> bool
 	{
 		return true;	// @TODO: should only allow main entry point function here
 	});
 
-	Controller->IsDependencyCyclicDelegate.BindLambda([](UObject* InDependentObject, UObject* InDependencyObject) -> bool
+	InControllerToConfigure->IsDependencyCyclicDelegate.BindLambda([](UObject* InDependentObject, UObject* InDependencyObject) -> bool
 	{
 		return false;
 	});
 
 #if WITH_EDITOR
-	Controller->SetupDefaultUnitNodeDelegates(TDelegate<FName(FRigVMExternalVariable, FString)>::CreateLambda(
+	InControllerToConfigure->SetupDefaultUnitNodeDelegates(TDelegate<FName(FRigVMExternalVariable, FString)>::CreateLambda(
 		[](FRigVMExternalVariable InVariableToCreate, FString InDefaultValue) -> FName
 		{
 			return NAME_None;
@@ -263,14 +194,106 @@ URigVMController* UDataInterfaceGraph_EditorData::GetOrCreateRigVMController(URi
 	));
 #endif
 
-	Controller->RemoveStaleNodes();
-	Controllers.Add(InRigVMGraph, Controller);
-	return Controller;
 }
 
-URigVMController* UDataInterfaceGraph_EditorData::GetOrCreateRigVMController(const UObject* InEditorObject)
+UObject* UDataInterfaceGraph_EditorData::GetEditorObjectForRigVMGraph(URigVMGraph* InVMGraph) const
+{
+	if(InVMGraph)
+	{
+		TArray<UDataInterfaceGraph_EdGraph*> AllGraphs = {RootGraph, EntryPointGraph, FunctionLibraryEdGraph};
+		for(UDataInterfaceGraph_EdGraph* EdGraph : AllGraphs)
+		{
+			if(EdGraph)
+			{
+				if(EdGraph->ModelNodePath == InVMGraph->GetNodePath())
+				{
+					return EdGraph;
+				}
+			}
+		}
+	}
+	return nullptr;
+}
+
+void UDataInterfaceGraph_EditorData::RecompileVM()
+{
+	UE::DataInterfaceGraphUncookedOnly::FUtils::Compile(GetTypedOuter<UDataInterfaceGraph>());
+}
+
+void UDataInterfaceGraph_EditorData::RecompileVMIfRequired()
+{
+	if (bVMRecompilationRequired)
+	{
+		RecompileVM();
+	}
+}
+
+void UDataInterfaceGraph_EditorData::RequestAutoVMRecompilation()
+{
+	bVMRecompilationRequired = true;
+	if (bAutoRecompileVM)
+	{
+		RecompileVMIfRequired();
+	}
+}
+
+void UDataInterfaceGraph_EditorData::HandleModifiedEvent(ERigVMGraphNotifType InNotifType, URigVMGraph* InGraph, UObject* InSubject)
+{
+	bool bNotifForOthersPending = true;
+
+	switch(InNotifType)
+	{
+	case ERigVMGraphNotifType::NodeAdded:
+		{
+			if (URigVMCollapseNode* CollapseNode = Cast<URigVMCollapseNode>(InSubject))
+			{
+				CreateEdGraphForCollapseNode(CollapseNode);
+				break;
+			}
+		}
+		// Fall through to the next case
+	case ERigVMGraphNotifType::LinkAdded:
+	case ERigVMGraphNotifType::LinkRemoved:
+	case ERigVMGraphNotifType::PinArraySizeChanged:
+	case ERigVMGraphNotifType::PinDirectionChanged:
+	case ERigVMGraphNotifType::PinDefaultValueChanged:
+		{
+			RequestAutoVMRecompilation();
+			break;
+		}
+	}
+	
+	// if the notification still has to be sent...
+	if (bNotifForOthersPending && !bSuspendModelNotificationsForOthers)
+	{
+		if (ModifiedEvent.IsBound())
+		{
+			ModifiedEvent.Broadcast(InNotifType, InGraph, InSubject);
+		}
+	}
+}
+
+URigVMGraph* UDataInterfaceGraph_EditorData::GetVMGraphForEdGraph(const UEdGraph* InGraph) const
 {
-	return GetOrCreateRigVMController(GetVMGraphForEdGraph(CastChecked<UEdGraph>(InEditorObject)));
+	if (InGraph == RootGraph)
+	{
+		return RigVMClient.GetDefaultModel();
+	}
+	else
+	{
+		const UDataInterfaceGraph_EdGraph* Graph = Cast<UDataInterfaceGraph_EdGraph>(InGraph);
+		check(Graph);
+
+		if (Graph->bIsFunctionDefinition)
+		{
+			if (URigVMLibraryNode* LibraryNode = RigVMClient.GetFunctionLibrary()->FindFunction(*Graph->ModelNodePath))
+			{
+				return LibraryNode->GetContainedGraph();
+			}
+		}
+	}
+
+	return nullptr;
 }
 
 void UDataInterfaceGraph_EditorData::CreateEdGraphForCollapseNode(URigVMCollapseNode* InNode)
@@ -293,7 +316,7 @@ void UDataInterfaceGraph_EditorData::CreateEdGraphForCollapseNode(URigVMCollapse
 				EntryPointGraph = RigFunctionGraph;
 				RigFunctionGraph->Initialize(this);
 
-				GetOrCreateRigVMController(ContainedGraph)->ResendAllNotifications();
+				RigVMClient.GetOrCreateController(ContainedGraph)->ResendAllNotifications();
 			}
 		}
 	}
diff --git a/Engine/Source/Developer/RigVMDeveloper/Private/RigVMCompiler/RigVMAST.cpp b/Engine/Source/Developer/RigVMDeveloper/Private/RigVMCompiler/RigVMAST.cpp
--- a/Engine/Source/Developer/RigVMDeveloper/Private/RigVMCompiler/RigVMAST.cpp
+++ b/Engine/Source/Developer/RigVMDeveloper/Private/RigVMCompiler/RigVMAST.cpp
@@ -709,7 +709,7 @@ void FRigVMParserASTSettings::Report(EMessageSeverity::Type InSeverity, UObject*
 
 const TArray<FRigVMASTProxy> FRigVMParserAST::EmptyProxyArray;
 
-FRigVMParserAST::FRigVMParserAST(URigVMGraph* InGraph, URigVMController* InController, const FRigVMParserASTSettings& InSettings, const TArray<FRigVMExternalVariable>& InExternalVariables, const TArray<FRigVMUserDataArray>& InRigVMUserData)
+FRigVMParserAST::FRigVMParserAST(TArray<URigVMGraph*> InGraphs, URigVMController* InController, const FRigVMParserASTSettings& InSettings, const TArray<FRigVMExternalVariable>& InExternalVariables, const TArray<FRigVMUserDataArray>& InRigVMUserData)
 {
 	DECLARE_SCOPE_HIERARCHICAL_COUNTER_FUNC()
 
@@ -718,8 +718,10 @@ FRigVMParserAST::FRigVMParserAST(URigVMGraph* InGraph, URigVMController* InContr
 	LastCycleCheckExpr = nullptr;
 	LinksToSkip = InSettings.LinksToSkip;
 
+	check(!InGraphs.IsEmpty());
+
 	// construct the inlined nodes and links information
-	Inline(InGraph);
+	Inline(InGraphs);
 
 	for (const FRigVMASTProxy& NodeProxy : NodeProxies)
 	{
@@ -769,11 +771,11 @@ FRigVMParserAST::FRigVMParserAST(URigVMGraph* InGraph, URigVMController* InContr
 	while (bContinueToFoldConstantBranches)
 	{
 		bContinueToFoldConstantBranches = false;
-		if (FoldConstantValuesToLiterals(InGraph, InController, InExternalVariables, InRigVMUserData))
+		if (FoldConstantValuesToLiterals(InGraphs, InController, InExternalVariables, InRigVMUserData))
 		{
 			bContinueToFoldConstantBranches = true;
 		}
-		if (FoldUnreachableBranches(InGraph))
+		if (FoldUnreachableBranches(InGraphs))
 		{
 			bContinueToFoldConstantBranches = true;
 		}
@@ -792,8 +794,10 @@ FRigVMParserAST::FRigVMParserAST(URigVMGraph* InGraph, URigVMController* InContr
 	}
 }
 
-FRigVMParserAST::FRigVMParserAST(URigVMGraph* InGraph, const TArray<FRigVMASTProxy>& InNodesToCompute)
+FRigVMParserAST::FRigVMParserAST(TArray<URigVMGraph*> InGraphs, const TArray<FRigVMASTProxy>& InNodesToCompute)
 {
+	check(!InGraphs.IsEmpty());
+
 	LastCycleCheckExpr = nullptr;
 
 	FRigVMBlockExprAST* Block = MakeExpr<FRigVMBlockExprAST>(FRigVMExprAST::EType::Block, FRigVMASTProxy());
@@ -802,7 +806,7 @@ FRigVMParserAST::FRigVMParserAST(URigVMGraph* InGraph, const TArray<FRigVMASTPro
 
 	NodeProxies = InNodesToCompute;
 	
-	Inline(InGraph, InNodesToCompute);
+	Inline(InGraphs, InNodesToCompute);
 
 	for (const FRigVMASTProxy& NodeProxy : NodeProxies)
 	{
@@ -1635,8 +1639,7 @@ void FRigVMParserAST::FoldAssignments()
 
 	RemoveExpressions(ExpressionsToRemove);
 }
-
-bool FRigVMParserAST::FoldConstantValuesToLiterals(URigVMGraph* InGraph, URigVMController* InController, const TArray<FRigVMExternalVariable>& InExternalVariables, const TArray<FRigVMUserDataArray>& InRigVMUserData)
+bool FRigVMParserAST::FoldConstantValuesToLiterals(TArray<URigVMGraph*> InGraphs, URigVMController* InController, const TArray<FRigVMExternalVariable>& InExternalVariables, const TArray<FRigVMUserDataArray>& InRigVMUserData)
 {
 	DECLARE_SCOPE_HIERARCHICAL_COUNTER_FUNC()
 
@@ -1777,7 +1780,7 @@ bool FRigVMParserAST::FoldConstantValuesToLiterals(URigVMGraph* InGraph, URigVMC
 
 	// we now know the node we need to run.
 	// let's build an temporary AST which has only those nodes
-	TSharedPtr<FRigVMParserAST> TempAST = MakeShareable(new FRigVMParserAST(InGraph, NodesToCompute));
+	TSharedPtr<FRigVMParserAST> TempAST = MakeShareable(new FRigVMParserAST(InGraphs, NodesToCompute));
 
 	// share the pin overrides with the constant folding AST to ensure
 	// the complete view of default values across function references is available
@@ -1786,7 +1789,8 @@ bool FRigVMParserAST::FoldConstantValuesToLiterals(URigVMGraph* InGraph, URigVMC
 
 	// build the VM to run this AST
 	TMap<FString, FRigVMOperand> Operands;
-	URigVM* TempVM = NewObject<URigVM>(InGraph);
+	check(!InGraphs.IsEmpty());
+	URigVM* TempVM = NewObject<URigVM>(InGraphs[0]);
 	
 	URigVMCompiler* TempCompiler = NewObject<URigVMCompiler>(GetTransientPackage());
 	TempCompiler->Settings.SetupNodeInstructionIndex = false;
@@ -1794,7 +1798,7 @@ bool FRigVMParserAST::FoldConstantValuesToLiterals(URigVMGraph* InGraph, URigVMC
 	TempCompiler->Settings.EnablePinWatches = false;
 	TempCompiler->Settings.ASTSettings = FRigVMParserASTSettings::Fast();
 
-	TempCompiler->Compile(InGraph, InController, TempVM, InExternalVariables, InRigVMUserData, &Operands, TempAST);
+	TempCompiler->Compile(InGraphs, InController, TempVM, InExternalVariables, InRigVMUserData, &Operands, TempAST);
 
 	TArray<URigVMMemoryStorage*> Memory;
 	Memory.Add(TempVM->GetWorkMemory());
@@ -1972,7 +1976,7 @@ bool FRigVMParserAST::FoldConstantValuesToLiterals(URigVMGraph* InGraph, URigVMC
 	return ExpressionsToRemove.Num() > 0;
 }
 
-bool FRigVMParserAST::FoldUnreachableBranches(URigVMGraph* InGraph)
+bool FRigVMParserAST::FoldUnreachableBranches(TArray<URigVMGraph*> InGraphs)
 {
 	DECLARE_SCOPE_HIERARCHICAL_COUNTER_FUNC()
 
@@ -2795,17 +2799,20 @@ const FRigVMASTLinkDescription& FRigVMParserAST::GetLink(int32 InLinkIndex) cons
 	return Links[InLinkIndex];
 }
 
-void FRigVMParserAST::Inline(URigVMGraph* InGraph)
+void FRigVMParserAST::Inline(TArray<URigVMGraph*> InGraphs)
 {
 	TArray<FRigVMASTProxy> LocalNodeProxies;
-	for (URigVMNode* LocalNode : InGraph->GetNodes())
+	for(URigVMGraph* Graph : InGraphs)
 	{
-		LocalNodeProxies.Add(FRigVMASTProxy::MakeFromUObject(LocalNode));
+		for (URigVMNode* LocalNode : Graph->GetNodes())
+		{
+			LocalNodeProxies.Add(FRigVMASTProxy::MakeFromUObject(LocalNode));
+		}
 	}
-	Inline(InGraph, LocalNodeProxies);
+	Inline(InGraphs, LocalNodeProxies);
 }
 
-void FRigVMParserAST::Inline(URigVMGraph* InGraph, const TArray<FRigVMASTProxy>& InNodeProxies)
+void FRigVMParserAST::Inline(TArray<URigVMGraph*> InGraphs, const TArray<FRigVMASTProxy>& InNodeProxies)
 {
 	DECLARE_SCOPE_HIERARCHICAL_COUNTER_FUNC()
 
diff --git a/Engine/Source/Developer/RigVMDeveloper/Private/RigVMCompiler/RigVMASTProxy.cpp b/Engine/Source/Developer/RigVMDeveloper/Private/RigVMCompiler/RigVMASTProxy.cpp
--- a/Engine/Source/Developer/RigVMDeveloper/Private/RigVMCompiler/RigVMASTProxy.cpp
+++ b/Engine/Source/Developer/RigVMDeveloper/Private/RigVMCompiler/RigVMASTProxy.cpp
@@ -13,11 +13,19 @@ FString FRigVMCallstack::GetCallPath(bool bIncludeLast) const
 	{
 		if (URigVMNode* Node = Cast<URigVMNode>(Entry))
 		{
-			Segments.Add(Node->GetName());
+			if(Node->GetGraph()->IsRootGraph())
+			{
+				Segments.Add(Node->GetNodePath(true));
+			}
+			else
+			{
+				Segments.Add(Node->GetName());
+			}
 		}
 		else if (URigVMPin* Pin = Cast<URigVMPin>(Entry))
 		{
-			Segments.Add(Pin->GetPinPath());
+			const bool bUseNodePath = Pin->GetGraph()->IsRootGraph();
+			Segments.Add(Pin->GetPinPath(bUseNodePath));
 		}
 	}
 
@@ -64,6 +72,19 @@ bool FRigVMCallstack::Contains(const UObject* InEntry) const
 	return Stack.Contains(InEntry);
 }
 
+FRigVMCallstack FRigVMCallstack::GetCallStackUpTo(int32 InIndex) const
+{
+	if(!ensure(Stack.IsValidIndex(InIndex)))
+	{
+		return FRigVMCallstack();
+	}
+	
+	FRigVMCallstack Partial;
+	Partial.Stack = Stack;
+	Partial.Stack.SetNum(InIndex + 1);
+	return Partial;
+}
+
 FRigVMASTProxy FRigVMASTProxy::MakeFromUObject(UObject* InSubject)
 {
 	UObject* Subject = InSubject;
@@ -133,7 +154,14 @@ FRigVMASTProxy FRigVMASTProxy::MakeFromCallPath(const FString& InCallPath, UObje
 
 		if(URigVMGraph* Graph = Cast<URigVMGraph>(ParentObject))
 		{
-			if(URigVMNode* FoundNode = Graph->FindNodeByName(*Left))
+			if(Graph->IsRootGraph() && Left.EndsWith(TEXT("::")))
+			{
+				if(Left != Graph->GetNodePath())
+				{
+					return FRigVMASTProxy();
+				}
+			}
+			else if(URigVMNode* FoundNode = Graph->FindNodeByName(*Left))
 			{
 				Proxy.Callstack.Stack.Push(FoundNode);
 				ParentObject = FoundNode;
diff --git a/Engine/Source/Developer/RigVMDeveloper/Private/RigVMCompiler/RigVMCompiler.cpp b/Engine/Source/Developer/RigVMDeveloper/Private/RigVMCompiler/RigVMCompiler.cpp
--- a/Engine/Source/Developer/RigVMDeveloper/Private/RigVMCompiler/RigVMCompiler.cpp
+++ b/Engine/Source/Developer/RigVMDeveloper/Private/RigVMCompiler/RigVMCompiler.cpp
@@ -129,19 +129,29 @@ URigVMCompiler::URigVMCompiler()
 {
 }
 
-bool URigVMCompiler::Compile(URigVMGraph* InGraph, URigVMController* InController, URigVM* OutVM, const TArray<FRigVMExternalVariable>& InExternalVariables, const TArray<FRigVMUserDataArray>& InRigVMUserData, TMap<FString, FRigVMOperand>* OutOperands, TSharedPtr<FRigVMParserAST> InAST)
+bool URigVMCompiler::Compile(TArray<URigVMGraph*> InGraphs, URigVMController* InController, URigVM* OutVM, const TArray<FRigVMExternalVariable>& InExternalVariables, const TArray<FRigVMUserDataArray>& InRigVMUserData, TMap<FString, FRigVMOperand>* OutOperands, TSharedPtr<FRigVMParserAST> InAST)
 {
-	if (InGraph == nullptr)
+	if (InGraphs.IsEmpty() || InGraphs.Contains(nullptr))
 	{
 		ReportError(TEXT("Provided graph is nullptr."));
 		return false;
 	}
+	
 	if (OutVM == nullptr)
 	{
 		ReportError(TEXT("Provided vm is nullptr."));
 		return false;
 	}
 
+	for(int32 Index = 1; Index < InGraphs.Num(); Index++)
+	{
+		if(InGraphs[0]->GetOuter() != InGraphs[Index]->GetOuter())
+		{
+			ReportError(TEXT("Provided graphs don't share a common outer / package."));
+			return false;
+		}
+	}
+
 	if(OutVM->GetClass()->IsChildOf(URigVMNativized::StaticClass()))
 	{
 		ReportError(TEXT("Provided vm is nativized."));
@@ -170,7 +180,7 @@ bool URigVMCompiler::Compile(URigVMGraph* InGraph, URigVMController* InControlle
 	// traverse all graphs and try to clear out orphan pins
 	// also check on function references with unmapped variables
 	TArray<URigVMGraph*> VisitedGraphs;
-	VisitedGraphs.Add(InGraph);
+	VisitedGraphs.Append(InGraphs);
 
 	bool bEncounteredGraphError = false;
 	for(int32 GraphIndex=0; GraphIndex<VisitedGraphs.Num(); GraphIndex++)
@@ -343,8 +353,11 @@ bool URigVMCompiler::Compile(URigVMGraph* InGraph, URigVMController* InControlle
 	WorkData.AST = InAST;
 	if (!WorkData.AST.IsValid())
 	{
-		WorkData.AST = MakeShareable(new FRigVMParserAST(InGraph, InController, Settings.ASTSettings, InExternalVariables, UserData));
-		InGraph->RuntimeAST = WorkData.AST;
+		WorkData.AST = MakeShareable(new FRigVMParserAST(InGraphs, InController, Settings.ASTSettings, InExternalVariables, UserData));
+		for(URigVMGraph* Graph : InGraphs)
+		{
+			Graph->RuntimeAST = WorkData.AST;
+		}
 #if UE_BUILD_DEBUG
 		//UE_LOG(LogRigVMDeveloper, Display, TEXT("%s"), *AST->DumpDot());
 #endif
@@ -373,7 +386,7 @@ bool URigVMCompiler::Compile(URigVMGraph* InGraph, URigVMController* InControlle
 #if WITH_EDITOR
 	// If in editor, make sure we visit all the graphs to initialize local variables
 	// in case the user wants to edit default values
-	URigVMFunctionLibrary* FunctionLibrary = InGraph->GetDefaultFunctionLibrary();
+	URigVMFunctionLibrary* FunctionLibrary = InGraphs[0]->GetDefaultFunctionLibrary();
 	if (FunctionLibrary)
 	{
 		for (URigVMLibraryNode* LibraryNode : FunctionLibrary->GetFunctions())
@@ -387,7 +400,7 @@ bool URigVMCompiler::Compile(URigVMGraph* InGraph, URigVMController* InControlle
 
 			for (FRigVMGraphVariableDescription& Variable : LibraryNode->GetContainedGraph()->LocalVariables)
 			{
-				FString Path = FString::Printf(TEXT("LocalVariableDefault::%s|%s::Const"), *LibraryNode->GetFName().ToString(), *Variable.Name.ToString());
+				FString Path = FString::Printf(TEXT("LocalVariableDefault::%s|%s::Const"), *LibraryNode->GetContainedGraph()->GetGraphName(), *Variable.Name.ToString());
 				FRigVMOperand Operand = WorkData.AddProperty(ERigVMMemoryType::Literal, *Path, Variable.CPPType, Variable.CPPTypeObject, Variable.DefaultValue);
 				WorkData.PinPathToOperand->Add(Path, Operand);
 
@@ -499,7 +512,7 @@ bool URigVMCompiler::Compile(URigVMGraph* InGraph, URigVMController* InControlle
 
 	for(ERigVMMemoryType MemoryType : MemoryTypes)
 	{
-		UPackage* Package = InGraph->GetOutermost();
+		UPackage* Package = InGraphs[0]->GetOutermost();
 
 		TArray<FRigVMPropertyDescription>* Properties = WorkData.PropertyDescriptions.Find(MemoryType);
 		if(Properties == nullptr)
@@ -1906,7 +1919,8 @@ void URigVMCompiler::InitializeLocalVariables(const FRigVMExprAST* InExpr, FRigV
 					{
 						for (FRigVMGraphVariableDescription Variable : Node->GetContainedGraph()->LocalVariables)
 						{
-							FString TargetPath = FString::Printf(TEXT("LocalVariable::%s|%s"), *Node->GetNodePath(), *Variable.Name.ToString());
+							const FRigVMCallstack LocalCallstack = Callstack.GetCallStackUpTo(SubjectIndex);
+							FString TargetPath = FString::Printf(TEXT("LocalVariable::%s|%s"), *LocalCallstack.GetCallPath(), *Variable.Name.ToString());
 							FString SourcePath = FString::Printf(TEXT("LocalVariableDefault::%s|%s::Const"), *Node->GetContainedGraph()->GetGraphName(), *Variable.Name.ToString());
 							FRigVMOperand* TargetPtr = WorkData.PinPathToOperand->Find(TargetPath);
 							FRigVMOperand* SourcePtr = WorkData.PinPathToOperand->Find(SourcePath);
@@ -2020,14 +2034,13 @@ FString URigVMCompiler::GetPinHashImpl(const URigVMPin* InPin, const FRigVMVarEx
 
 							if(URigVMLibraryNode* LibraryNode = ParentProxy.GetSubject<URigVMLibraryNode>())
 							{
-								// Local variables for non-root graphs are in the format "LocalVariable::PathToGraph|VariableName"
-								const FString GraphPath = ParentProxy.GetCallstack().GetCallPath(true);
-								return FString::Printf(TEXT("%sLocalVariable::%s|%s%s"), *Prefix, *GraphPath, *VariableName.ToString(), *Suffix);
+								break;
 							}
 						}
 
-						// Local variables for root graphs are in the format "LocalVariable::VariableName"
-						return FString::Printf(TEXT("%sLocalVariable::%s%s"), *Prefix, *VariableName.ToString(), *Suffix);
+						// Local variables for root / non-root graphs are in the format "LocalVariable::PathToGraph|VariableName"
+						const FString GraphPath = ParentProxy.GetCallstack().GetCallPath(true);
+						return FString::Printf(TEXT("%sLocalVariable::%s|%s%s"), *Prefix, *GraphPath, *VariableName.ToString(), *Suffix);
 					}
 				}
 			}
diff --git a/Engine/Source/Developer/RigVMDeveloper/Private/RigVMModel/RigVMClient.cpp b/Engine/Source/Developer/RigVMDeveloper/Private/RigVMModel/RigVMClient.cpp
new file mode 100644
--- /dev/null
+++ b/Engine/Source/Developer/RigVMDeveloper/Private/RigVMModel/RigVMClient.cpp
@@ -0,0 +1,621 @@
+// Copyright Epic Games, Inc. All Rights Reserved.
+
+#include "RigVMModel/RigVMClient.h"
+#if WITH_EDITOR
+#include "ScopedTransaction.h"
+#endif
+
+void FRigVMClient::SetOuterClientHost(UObject* InOuterClientHost, const FName& InOuterClientHostPropertyName)
+{
+	OuterClientHost = InOuterClientHost;
+	OuterClientPropertyName = InOuterClientHostPropertyName;
+
+	check(OuterClientHost->Implements<URigVMClientHost>());
+	check(GetOuterClientProperty() != nullptr);
+}
+
+void FRigVMClient::SetFromDeprecatedData(URigVMGraph* InDefaultGraph, URigVMFunctionLibrary* InFunctionLibrary)
+{
+	if(GetDefaultModel() != InDefaultGraph ||
+		GetFunctionLibrary() != InFunctionLibrary)
+	{
+		if(GetDefaultModel() == InDefaultGraph)
+		{
+			Models.Reset();
+		}
+		
+		if(InFunctionLibrary == nullptr)
+		{
+			Swap(FunctionLibrary, InFunctionLibrary);
+		}
+
+		Reset();
+		if(InDefaultGraph)
+		{
+			AddModel(InDefaultGraph, false);
+		}
+		if(InFunctionLibrary)
+		{
+			AddModel(InFunctionLibrary, false);
+		}
+	}
+}
+
+void FRigVMClient::Reset()
+{
+	for(URigVMGraph* Model : Models)
+	{
+		DestroyObject(Model);
+	}
+	for(auto Pair : Controllers)
+	{
+		DestroyObject(Pair.Value);
+	}
+	DestroyObject(FunctionLibrary);
+
+	Models.Reset();
+	Controllers.Reset();
+	FunctionLibrary = nullptr;
+}
+
+URigVMGraph* FRigVMClient::GetDefaultModel() const
+{
+	if(Models.IsEmpty())
+	{
+		return nullptr;
+	}
+	return GetModel(0);
+}
+
+URigVMGraph* FRigVMClient::GetModel(const FString& InNodePathOrName) const
+{
+	if(InNodePathOrName.IsEmpty())
+	{
+		return GetDefaultModel();
+	}
+
+	TArray<URigVMGraph*> ModelsAndFunctionLibrary = GetAllModels(true, false);
+	for(URigVMGraph* Model : ModelsAndFunctionLibrary)
+	{
+		if(Model->GetNodePath() == InNodePathOrName || Model->GetName() == InNodePathOrName)
+		{
+			return Model;
+		}
+
+		static constexpr TCHAR NodePathPrefixFormat[] = TEXT("%s|");
+		const FString NodePathPrefix = FString::Printf(NodePathPrefixFormat, *Model->GetNodePath());
+		if(InNodePathOrName.StartsWith(NodePathPrefix))
+		{
+			const FString RemainingNodePath = InNodePathOrName.Mid(NodePathPrefix.Len());
+			if(const URigVMLibraryNode* LibraryNode = Cast<URigVMLibraryNode>(Model->FindNode(RemainingNodePath)))
+			{
+				return LibraryNode->GetContainedGraph();
+			}
+		}
+	}
+	return nullptr;
+}
+
+URigVMGraph* FRigVMClient::GetModel(const UObject* InEditorSideObject) const
+{
+	check(InEditorSideObject);
+
+	if (InEditorSideObject->Implements<URigVMEditorSideObject>())
+	{
+		const IRigVMEditorSideObject* UserInterfaceGraph = Cast<IRigVMEditorSideObject>(InEditorSideObject);
+		return GetModel(UserInterfaceGraph->GetRigVMNodePath());
+	}
+
+	return nullptr;
+}
+
+TArray<URigVMGraph*> FRigVMClient::GetAllModels(bool bIncludeFunctionLibrary, bool bRecursive) const
+{
+	TArray<URigVMGraph*> AllModels = Models;
+	if(bRecursive)
+	{
+		for(URigVMGraph* Model : Models)
+		{
+			AllModels.Append(Model->GetContainedGraphs(true /* recursive */));
+		}
+	}
+	if(bIncludeFunctionLibrary && FunctionLibrary)
+	{
+		AllModels.Add(FunctionLibrary);
+		if(bRecursive)
+		{
+			AllModels.Append(FunctionLibrary->GetContainedGraphs(true /* recursive */));
+		}
+	}
+	return AllModels;
+}
+
+URigVMController* FRigVMClient::GetController(int32 InIndex) const
+{
+	return GetController(GetModel(InIndex));
+}
+
+URigVMController* FRigVMClient::GetController(const FString& InNodePathOrName) const
+{
+	return GetController(GetModel(InNodePathOrName));
+}
+
+URigVMController* FRigVMClient::GetController(const URigVMGraph* InModel) const
+{
+	if(InModel == nullptr)
+	{
+		InModel = GetDefaultModel();
+	}
+	
+	if(InModel)
+	{
+		const FSoftObjectPath Key(InModel);
+		if(const TObjectPtr<URigVMController>* Controller = Controllers.Find(Key))
+		{
+			return Controller->Get();
+		}
+	}
+	return nullptr;
+}
+
+URigVMController* FRigVMClient::GetController(const UObject* InEditorSideObject) const
+{
+	check(InEditorSideObject);
+
+	if (InEditorSideObject->Implements<URigVMEditorSideObject>())
+	{
+		const IRigVMEditorSideObject* UserInterfaceGraph = Cast<IRigVMEditorSideObject>(InEditorSideObject);
+		return GetController(UserInterfaceGraph->GetRigVMNodePath());
+	}
+
+	return nullptr;
+}
+
+URigVMController* FRigVMClient::GetOrCreateController(int32 InIndex)
+{
+	return GetOrCreateController(GetModel(InIndex));
+}
+
+URigVMController* FRigVMClient::GetOrCreateController(const FString& InNodePathOrName)
+{
+	return GetOrCreateController(GetModel(InNodePathOrName));
+}
+
+URigVMController* FRigVMClient::GetOrCreateController(const URigVMGraph* InModel)
+{
+	if(InModel == nullptr)
+	{
+		InModel = GetDefaultModel();
+	}
+
+	if(InModel)
+	{
+		if(URigVMController* Controller = GetController(InModel))
+		{
+			return Controller;
+		}
+		return CreateController(InModel);
+	}
+	return nullptr; 
+}
+
+URigVMController* FRigVMClient::GetOrCreateController(const UObject* InEditorSideObject)
+{
+	check(InEditorSideObject);
+
+	if (InEditorSideObject->Implements<URigVMEditorSideObject>())
+	{
+		const IRigVMEditorSideObject* UserInterfaceGraph = Cast<IRigVMEditorSideObject>(InEditorSideObject);
+		return GetOrCreateController(UserInterfaceGraph->GetRigVMNodePath());
+	}
+
+	return nullptr;
+}
+
+URigVMGraph* FRigVMClient::AddModel(const FName& InName, bool bSetupUndoRedo, const FObjectInitializer* ObjectInitializer)
+{
+#if WITH_EDITOR
+	TSharedPtr<FScopedTransaction> Transaction;
+	if(bSetupUndoRedo)
+	{
+		Transaction = MakeShared<FScopedTransaction>(NSLOCTEXT("RigVMClient", "AddModel", "Add new root graph"));
+		GetOuter()->Modify();
+	}
+#endif
+
+	const FName SafeGraphName = GetUniqueName(InName);
+	URigVMGraph* Model = nullptr;
+	if(ObjectInitializer)
+	{
+		Model = ObjectInitializer->CreateDefaultSubobject<URigVMGraph>(GetOuter(), SafeGraphName);
+	}
+	else
+	{
+		Model = NewObject<URigVMGraph>(GetOuter(), SafeGraphName);
+	}
+	AddModel(Model, true);
+
+	if(bSetupUndoRedo)
+	{
+		GetOuter()->Modify();
+		UndoRedoIndex++;
+
+		FRigVMClientAction Action;
+		Action.Type = ERigVMClientAction_AddModel;
+		Action.NodePath = Model->GetNodePath();
+		UndoStack.Push(Action);
+		RedoStack.Reset();
+	}
+	return Model;
+}
+
+void FRigVMClient::AddModel(URigVMGraph* InModel, bool bCreateController)
+{
+	check(InModel);
+
+	if(InModel->IsA<URigVMFunctionLibrary>())
+	{
+		check(FunctionLibrary == nullptr);
+		FunctionLibrary = Cast<URigVMFunctionLibrary>(InModel);
+	}
+	else
+	{
+		Models.Add(InModel);
+	}
+
+	if(bCreateController)
+	{
+		CreateController(InModel);
+	}
+
+	if(InModel->IsA<URigVMFunctionLibrary>())
+	{
+		for(URigVMGraph* Model : Models)
+		{
+			Model->SetDefaultFunctionLibrary(FunctionLibrary);
+		}
+	}
+	else if(FunctionLibrary)
+	{
+		InModel->SetDefaultFunctionLibrary(FunctionLibrary);
+	}
+
+	if (GetOuter()->Implements<URigVMClientHost>())
+	{
+		IRigVMClientHost* ClientHost = Cast<IRigVMClientHost>(GetOuter());
+		ClientHost->HandleRigVMGraphAdded(this, InModel->GetNodePath());
+	}
+		
+	NotifyOuterOfPropertyChange();
+
+}
+
+URigVMFunctionLibrary* FRigVMClient::GetOrCreateFunctionLibrary(bool bSetupUndoRedo, const FObjectInitializer* ObjectInitializer)
+{
+	if(FunctionLibrary)
+	{
+		return FunctionLibrary;
+	}
+	
+#if WITH_EDITOR
+	TSharedPtr<FScopedTransaction> Transaction;
+	if(bSetupUndoRedo)
+	{
+		Transaction = MakeShared<FScopedTransaction>(NSLOCTEXT("RigVMClient", "AddModel", "Add new root graph"));
+		GetOuter()->Modify();
+	}
+#endif
+
+	static constexpr TCHAR FunctionLibraryName[] = TEXT("RigVMFunctionLibrary");
+	const FName SafeGraphName = GetUniqueName(FunctionLibraryName);
+	URigVMFunctionLibrary* NewFunctionLibrary = nullptr;
+	if(ObjectInitializer)
+	{
+		NewFunctionLibrary = ObjectInitializer->CreateDefaultSubobject<URigVMFunctionLibrary>(GetOuter(), SafeGraphName);
+	}
+	else
+	{
+		NewFunctionLibrary = NewObject<URigVMFunctionLibrary>(GetOuter(), SafeGraphName);
+	}
+	AddModel(NewFunctionLibrary, true);
+	return NewFunctionLibrary;
+}
+
+bool FRigVMClient::RemoveModel(const FString& InNodePathOrName, bool bSetupUndoRedo)
+{
+	if(URigVMGraph* Model = GetModel(InNodePathOrName))
+	{
+		if(Model == GetDefaultModel())
+		{
+#if WITH_EDITOR
+			static constexpr TCHAR Message[] = TEXT("Cannot remove the default model.");
+			FScriptExceptionHandler::Get().HandleException(ELogVerbosity::Error, Message, *FString());
+#endif
+			return false;
+		}
+
+		if(Model == FunctionLibrary)
+		{
+#if WITH_EDITOR
+			static constexpr TCHAR Message[] = TEXT("Cannot remove the function library.");
+			FScriptExceptionHandler::Get().HandleException(ELogVerbosity::Error, Message, *FString());
+#endif
+			return false;
+		}
+
+#if WITH_EDITOR
+		TSharedPtr<FScopedTransaction> Transaction;
+		if(bSetupUndoRedo)
+		{
+			Transaction = MakeShared<FScopedTransaction>(NSLOCTEXT("RigVMClient", "RemoveModel", "Remove a root graph"));
+			GetOuter()->Modify();
+		}
+#endif
+
+		// remove the controller
+		if(URigVMController* Controller = GetController(Model))
+		{
+			verify(Controller == Controllers.FindAndRemoveChecked(Model));
+		}
+
+		if (GetOuter()->Implements<URigVMClientHost>())
+		{
+			IRigVMClientHost* ClientHost = Cast<IRigVMClientHost>(GetOuter());
+			ClientHost->HandleRigVMGraphRemoved(this, InNodePathOrName);
+		}
+
+		if(bSetupUndoRedo)
+		{
+			GetOuter()->Modify();
+			UndoRedoIndex++;
+
+			FRigVMClientAction Action;
+			Action.Type = ERigVMClientAction_RemoveModel;
+			Action.NodePath = Model->GetNodePath();
+			UndoStack.Push(Action);
+			RedoStack.Reset();
+		}
+
+		// clean up the model
+		Models.Remove(Model);
+
+		NotifyOuterOfPropertyChange();
+		return true;
+	}
+	return false;
+}
+
+FName FRigVMClient::RenameModel(const FString& InNodePathOrName, const FName& InNewName, bool bSetupUndoRedo)
+{
+	if(URigVMGraph* Model = GetModel(InNodePathOrName))
+	{
+		if(Model == FunctionLibrary)
+		{
+#if WITH_EDITOR
+			static constexpr TCHAR Message[] = TEXT("Cannot rename the function library.");
+			FScriptExceptionHandler::Get().HandleException(ELogVerbosity::Error, Message, *FString());
+#endif
+			return NAME_None;
+		}
+
+		if(Model->GetFName() == InNewName)
+		{
+			return InNewName;
+		}
+
+#if WITH_EDITOR
+		TSharedPtr<FScopedTransaction> Transaction;
+		if(bSetupUndoRedo)
+		{
+			Transaction = MakeShared<FScopedTransaction>(NSLOCTEXT("RigVMClient", "RenameModel", "Rename a root graph"));
+		}
+#endif
+
+		const FString OldNodePath = Model->GetNodePath();
+		const FName SafeNewName = GetUniqueName(InNewName);
+		Model->Rename(*SafeNewName.ToString(), nullptr, REN_ForceNoResetLoaders | REN_DontCreateRedirectors);
+		const FString NewNodePath = Model->GetNodePath();
+
+		if(bSetupUndoRedo)
+		{
+			GetOuter()->Modify();
+			UndoRedoIndex++;
+
+			FRigVMClientAction Action;
+			Action.Type = ERigVMClientAction_RenameModel;
+			Action.NodePath = OldNodePath;
+			Action.OtherNodePath = NewNodePath;
+			UndoStack.Push(Action);
+			RedoStack.Reset();
+		}
+
+		if (GetOuter()->Implements<URigVMClientHost>())
+		{
+			IRigVMClientHost* ClientHost = Cast<IRigVMClientHost>(GetOuter());
+			ClientHost->HandleRigVMGraphRenamed(this, OldNodePath, NewNodePath);
+		}
+		
+		NotifyOuterOfPropertyChange();
+		return SafeNewName;
+	}
+
+	return NAME_None;
+}
+
+void FRigVMClient::PostTransacted(const FTransactionObjectEvent& TransactionEvent)
+{
+	IRigVMClientHost* ClientHost = CastChecked<IRigVMClientHost>(GetOuter());
+
+	if (TransactionEvent.GetEventType() == ETransactionObjectEventType::UndoRedo)
+	{
+		auto PerformAction = [this, ClientHost](const FRigVMClientAction& InAction, bool bUndo)
+		{
+			switch(InAction.Type)
+			{
+				case ERigVMClientAction_AddModel:
+				case ERigVMClientAction_RemoveModel:
+				{
+					if(InAction.Type == ERigVMClientAction_RemoveModel)
+					{
+						bUndo = !bUndo;
+					}
+						
+					if(URigVMGraph* Model = GetModel(InAction.NodePath))
+					{
+						if(bUndo)
+						{
+							ClientHost->HandleRigVMGraphRemoved(this, InAction.NodePath);
+						}
+						else
+						{
+							URigVMController* Controller = GetOrCreateController(Model);
+							ClientHost->HandleRigVMGraphAdded(this, InAction.NodePath);
+							if(Controller)
+							{
+								Controller->ResendAllNotifications();
+							}
+						}
+					}
+					break;
+				}
+				case ERigVMClientAction_RenameModel:
+				{
+					const FString NodePathA = bUndo ? InAction.OtherNodePath : InAction.NodePath;
+					const FString NodePathB = bUndo ? InAction.NodePath : InAction.OtherNodePath;
+					FString NodeNameB = NodePathB;
+					NodeNameB.Split(TEXT("|"), nullptr, &NodeNameB, ESearchCase::IgnoreCase, ESearchDir::FromEnd);
+					NodeNameB.RemoveFromEnd(TEXT("::"));
+
+					RenameModel(NodePathA, *NodeNameB, false);
+					break;
+				}
+			}
+		};
+		
+		while(UndoStack.Num() != UndoRedoIndex)
+		{
+			// do we need to undo - or redo?
+			if(UndoStack.Num() > UndoRedoIndex)
+			{
+				FRigVMClientAction Action = UndoStack.Pop();
+				PerformAction(Action, true);
+				RedoStack.Push(Action);
+			}
+			else
+			{
+				FRigVMClientAction Action = RedoStack.Pop();
+				PerformAction(Action, false);
+				UndoStack.Push(Action);
+			}
+		}
+	}
+}
+
+void FRigVMClient::OnSubGraphRenamed(const FSoftObjectPath& OldObjectPath, const FSoftObjectPath& NewObjectPath)
+{
+	if(OldObjectPath != NewObjectPath)
+	{
+		if (TObjectPtr<URigVMController>* Controller = Controllers.Find(OldObjectPath))
+		{
+			Controllers.Add(NewObjectPath, *Controller);
+			Controllers.Remove(OldObjectPath);
+		}
+	}
+}
+
+URigVMNode* FRigVMClient::FindNode(const FString& InNodePathOrName) const
+{
+	for(const URigVMGraph* Model : Models)
+	{
+		if(URigVMNode* Node = Model->FindNode(InNodePathOrName))
+		{
+			return Node;
+		}
+	}
+	if(FunctionLibrary)
+	{
+		return FunctionLibrary->FindNode(InNodePathOrName);
+	}
+	return nullptr;
+}
+
+URigVMPin* FRigVMClient::FindPin(const FString& InPinPath) const
+{
+	for(const URigVMGraph* Model : Models)
+	{
+		if(URigVMPin* Pin = Model->FindPin(InPinPath))
+		{
+			return Pin;
+		}
+	}
+	if(FunctionLibrary)
+	{
+		return FunctionLibrary->FindPin(InPinPath);
+	}
+	return nullptr;
+}
+
+UObject* FRigVMClient::GetOuter() const
+{
+	UObject* Outer = OuterClientHost.Get();
+	check(Outer);
+	return Outer;
+}
+
+FProperty* FRigVMClient::GetOuterClientProperty() const
+{
+	return GetOuter()->GetClass()->FindPropertyByName(OuterClientPropertyName);
+}
+
+void FRigVMClient::NotifyOuterOfPropertyChange(EPropertyChangeType::Type ChangeType) const
+{
+	if(bSuspendNotifications)
+	{
+		return;
+	}
+	FProperty* Property = GetOuterClientProperty();
+	FPropertyChangedEvent PropertyChangedEvent(Property, ChangeType);
+	GetOuter()->PostEditChangeProperty(PropertyChangedEvent);
+}
+
+URigVMController* FRigVMClient::CreateController(const URigVMGraph* InModel)
+{
+	const FName SafeControllerName = GetUniqueName(*FString::Printf(TEXT("%s_Controller"), *InModel->GetName()));
+	URigVMController* Controller = NewObject<URigVMController>(GetOuter(), SafeControllerName);
+	Controllers.Add(InModel, Controller);
+	Controller->SetGraph((URigVMGraph*)InModel);
+
+	if (GetOuter()->Implements<URigVMClientHost>())
+	{
+		IRigVMClientHost* ClientHost = Cast<IRigVMClientHost>(GetOuter());
+		ClientHost->HandleConfigureRigVMController(this, Controller);
+	}
+	
+	Controller->RemoveStaleNodes();
+	return Controller;
+}
+
+FName FRigVMClient::GetUniqueName(const FName& InDesiredName) const
+{
+	return GetUniqueName(GetOuter(), InDesiredName);
+}
+
+FName FRigVMClient::GetUniqueName(UObject* InOuter, const FName& InDesiredName)
+{
+	return URigVMController::GetUniqueName(*InDesiredName.ToString(), [InOuter](const FName& InName) -> bool
+	{
+		return FindObjectWithOuter(InOuter, nullptr, InName) == nullptr; 
+	}, false, true);
+}
+
+void FRigVMClient::DestroyObject(UObject* InObject)
+{
+	if(InObject)
+	{
+		static int32 ObjectIndexToBeDestroyed = 0;
+		static constexpr TCHAR ObjectNameFormat[] = TEXT("RigVMClient_ObjectToBeDestroyed_%d");
+		const FString NewObjectName = FString::Printf(ObjectNameFormat, ObjectIndexToBeDestroyed++);
+		InObject->Rename(*NewObjectName, GetTransientPackage(), REN_ForceNoResetLoaders | REN_DontCreateRedirectors);
+		InObject->MarkAsGarbage();
+	}
+}
diff --git a/Engine/Source/Developer/RigVMDeveloper/Private/RigVMModel/RigVMFunctionLibrary.cpp b/Engine/Source/Developer/RigVMDeveloper/Private/RigVMModel/RigVMFunctionLibrary.cpp
--- a/Engine/Source/Developer/RigVMDeveloper/Private/RigVMModel/RigVMFunctionLibrary.cpp
+++ b/Engine/Source/Developer/RigVMDeveloper/Private/RigVMModel/RigVMFunctionLibrary.cpp
@@ -12,7 +12,8 @@ URigVMFunctionLibrary::URigVMFunctionLibrary()
 
 FString URigVMFunctionLibrary::GetNodePath() const
 {
-	return FString::Printf(TEXT("FunctionLibrary::%s"), *Super::GetNodePath());
+	static constexpr TCHAR FunctionLibraryNodePath[] = TEXT("FunctionLibrary::");
+	return FunctionLibraryNodePath;
 }
 
 URigVMFunctionLibrary* URigVMFunctionLibrary::GetDefaultFunctionLibrary() const
diff --git a/Engine/Source/Developer/RigVMDeveloper/Private/RigVMModel/RigVMGraph.cpp b/Engine/Source/Developer/RigVMDeveloper/Private/RigVMModel/RigVMGraph.cpp
--- a/Engine/Source/Developer/RigVMDeveloper/Private/RigVMModel/RigVMGraph.cpp
+++ b/Engine/Source/Developer/RigVMDeveloper/Private/RigVMModel/RigVMGraph.cpp
@@ -110,7 +110,9 @@ FString URigVMGraph::GetNodePath() const
 	{
 		return CollapseNode->GetNodePath(true /* recursive */);
 	}
-	return FString();
+
+	static constexpr TCHAR NodePathFormat[] = TEXT("%s::");
+	return FString::Printf(NodePathFormat, *GetName());
 }
 
 FString URigVMGraph::GetGraphName() const
@@ -119,7 +121,7 @@ FString URigVMGraph::GetGraphName() const
 	{
 		return CollapseNode->GetNodePath(false /* recursive */);
 	}
-	return FString();
+	return GetName();
 }
 
 URigVMNode* URigVMGraph::FindNodeByName(const FName& InNodeName) const
@@ -147,15 +149,13 @@ URigVMNode* URigVMGraph::FindNode(const FString& InNodePath) const
 	}
 
 	FString Path = InNodePath;
-	if (Path.StartsWith(TEXT("FunctionLibrary::|")))
+
+	if(IsRootGraph())
 	{
-		if (GetRootGraph()->IsA<URigVMFunctionLibrary>())
-		{
-			Path.RightChopInline(18);
-		}
-		else
+		const FString MyNodePath = GetNodePath();
+		if(Path.StartsWith(MyNodePath))
 		{
-			return nullptr;
+			Path.RightChopInline(MyNodePath.Len() + 1);
 		}
 	}
 
@@ -372,7 +372,7 @@ TSharedPtr<FRigVMParserAST> URigVMGraph::GetDiagnosticsAST(bool bForceRefresh, T
 	{
 		FRigVMParserASTSettings Settings = FRigVMParserASTSettings::Fast();
 		Settings.LinksToSkip = InLinksToSkip;
-		DiagnosticsAST = MakeShareable(new FRigVMParserAST(this, nullptr, Settings));
+		DiagnosticsAST = MakeShareable(new FRigVMParserAST({this}, nullptr, Settings));
 	}
 	return DiagnosticsAST;
 }
@@ -381,7 +381,7 @@ TSharedPtr<FRigVMParserAST> URigVMGraph::GetRuntimeAST(const FRigVMParserASTSett
 {
 	if (RuntimeAST == nullptr || bForceRefresh)
 	{
-		RuntimeAST = MakeShareable(new FRigVMParserAST(this, nullptr, InSettings));
+		RuntimeAST = MakeShareable(new FRigVMParserAST({this}, nullptr, InSettings));
 	}
 	return RuntimeAST;
 }
diff --git a/Engine/Source/Developer/RigVMDeveloper/Public/RigVMCompiler/RigVMAST.h b/Engine/Source/Developer/RigVMDeveloper/Public/RigVMCompiler/RigVMAST.h
--- a/Engine/Source/Developer/RigVMDeveloper/Public/RigVMCompiler/RigVMAST.h
+++ b/Engine/Source/Developer/RigVMDeveloper/Public/RigVMCompiler/RigVMAST.h
@@ -1226,7 +1226,7 @@ public:
 	// default constructor
 	// @param InGraph The graph / model to parse
 	// @param InSettings The parse settings to use
-	FRigVMParserAST(URigVMGraph* InGraph, URigVMController* InController = nullptr, const FRigVMParserASTSettings& InSettings = FRigVMParserASTSettings::Fast(), const TArray<FRigVMExternalVariable>& InExternalVariables = TArray<FRigVMExternalVariable>(), const TArray<FRigVMUserDataArray>& InRigVMUserData = TArray<FRigVMUserDataArray>());
+	FRigVMParserAST(TArray<URigVMGraph*> InGraphs, URigVMController* InController = nullptr, const FRigVMParserASTSettings& InSettings = FRigVMParserASTSettings::Fast(), const TArray<FRigVMExternalVariable>& InExternalVariables = TArray<FRigVMExternalVariable>(), const TArray<FRigVMUserDataArray>& InRigVMUserData = TArray<FRigVMUserDataArray>());
 
 	// default destructor
 	~FRigVMParserAST();
@@ -1296,7 +1296,7 @@ public:
 private:
 
 	// private constructor for a partial build
-	FRigVMParserAST(URigVMGraph* InGraph, const TArray<FRigVMASTProxy>& InNodesToCompute);
+	FRigVMParserAST(TArray<URigVMGraph*> InGraphs, const TArray<FRigVMASTProxy>& InNodesToCompute);
 
 	// make function to create an expression
 	template<class ExprType>
@@ -1408,14 +1408,14 @@ private:
 	void FoldAssignments();
 
 	// helper function to fold constant values into literals
-	bool FoldConstantValuesToLiterals(URigVMGraph* InGraph, URigVMController* InController, const TArray<FRigVMExternalVariable>& InExternalVariables, const TArray<FRigVMUserDataArray>& InRigVMUserData);
+	bool FoldConstantValuesToLiterals(TArray<URigVMGraph*> InGraphs, URigVMController* InController, const TArray<FRigVMExternalVariable>& InExternalVariables, const TArray<FRigVMUserDataArray>& InRigVMUserData);
 
 	// helper function to fold unreachable branches 
-	bool FoldUnreachableBranches(URigVMGraph* InGraph);
+	bool FoldUnreachableBranches(TArray<URigVMGraph*> InGraphs);
 
 	// helper function to inline all contributing nodes of the graph
-	void Inline(URigVMGraph* InGraph);
-	void Inline(URigVMGraph* InGraph, const TArray<FRigVMASTProxy>& InNodeProxies);
+	void Inline(TArray<URigVMGraph*> InGraphs);
+	void Inline(TArray<URigVMGraph*> InGraphs, const TArray<FRigVMASTProxy>& InNodeProxies);
 
 	// helper functions to retrieve links for a given pin
 	TArray<int32> GetSourceLinkIndices(const FRigVMASTProxy& InPinProxy, bool bRecursive = false) const;
diff --git a/Engine/Source/Developer/RigVMDeveloper/Public/RigVMCompiler/RigVMASTProxy.h b/Engine/Source/Developer/RigVMDeveloper/Public/RigVMCompiler/RigVMASTProxy.h
--- a/Engine/Source/Developer/RigVMDeveloper/Public/RigVMCompiler/RigVMASTProxy.h
+++ b/Engine/Source/Developer/RigVMDeveloper/Public/RigVMCompiler/RigVMASTProxy.h
@@ -97,6 +97,8 @@ public:
 
 	FORCEINLINE TArray<UObject*> GetStack() const { return Stack; }
 
+	FRigVMCallstack GetCallStackUpTo(int32 InIndex) const;
+
 private:
 
 	FORCEINLINE uint32 GetEntryTypeHash(int32 Index) const
diff --git a/Engine/Source/Developer/RigVMDeveloper/Public/RigVMCompiler/RigVMCompiler.h b/Engine/Source/Developer/RigVMDeveloper/Public/RigVMCompiler/RigVMCompiler.h
--- a/Engine/Source/Developer/RigVMDeveloper/Public/RigVMCompiler/RigVMCompiler.h
+++ b/Engine/Source/Developer/RigVMDeveloper/Public/RigVMCompiler/RigVMCompiler.h
@@ -137,12 +137,12 @@ public:
 	FRigVMCompileSettings Settings;
 
 	UFUNCTION(BlueprintCallable, Category = FRigVMCompiler)
-		bool Compile(URigVMGraph* InGraph, URigVMController* InController, URigVM* OutVM)
+	bool Compile(TArray<URigVMGraph*> InGraphs, URigVMController* InController, URigVM* OutVM)
 	{
-		return Compile(InGraph, InController, OutVM, TArray<FRigVMExternalVariable>(), TArray<FRigVMUserDataArray>(), nullptr);
+		return Compile(InGraphs, InController, OutVM, TArray<FRigVMExternalVariable>(), TArray<FRigVMUserDataArray>(), nullptr);
 	}
 
-	bool Compile(URigVMGraph* InGraph, URigVMController* InController, URigVM* OutVM, const TArray<FRigVMExternalVariable>& InExternalVariables, const TArray<FRigVMUserDataArray>& InRigVMUserData, TMap<FString, FRigVMOperand>* OutOperands, TSharedPtr<FRigVMParserAST> InAST = TSharedPtr<FRigVMParserAST>());
+	bool Compile(TArray<URigVMGraph*> InGraphs, URigVMController* InController, URigVM* OutVM, const TArray<FRigVMExternalVariable>& InExternalVariables, const TArray<FRigVMUserDataArray>& InRigVMUserData, TMap<FString, FRigVMOperand>* OutOperands, TSharedPtr<FRigVMParserAST> InAST = TSharedPtr<FRigVMParserAST>());
 
 	static UScriptStruct* GetScriptStructForCPPType(const FString& InCPPType);
 	static FString GetPinHash(const URigVMPin* InPin, const FRigVMVarExprAST* InVarExpr, bool bIsDebugValue = false, const FRigVMASTProxy& InPinProxy = FRigVMASTProxy());
diff --git a/Engine/Source/Developer/RigVMDeveloper/Public/RigVMModel/RigVMClient.h b/Engine/Source/Developer/RigVMDeveloper/Public/RigVMModel/RigVMClient.h
new file mode 100644
--- /dev/null
+++ b/Engine/Source/Developer/RigVMDeveloper/Public/RigVMModel/RigVMClient.h
@@ -0,0 +1,167 @@
+// Copyright Epic Games, Inc. All Rights Reserved.
+
+#pragma once
+
+#include "RigVMGraph.h"
+#include "RigVMFunctionLibrary.h"
+#include "RigVMController.h"
+#include "RigVMClient.generated.h"
+
+struct FRigVMClient;
+
+UINTERFACE()
+class RIGVMDEVELOPER_API URigVMClientHost : public UInterface
+{
+	GENERATED_BODY()
+};
+
+// Interface that allows an object to host a rig VM client. Used by graph edting code to interact with the controller.
+class RIGVMDEVELOPER_API IRigVMClientHost
+{
+	GENERATED_BODY()
+
+public:
+	
+	// Returns the rigvm client for this host
+	virtual FRigVMClient* GetRigVMClient() = 0;
+
+	// Returns the rigvm client for this host
+	virtual const FRigVMClient* GetRigVMClient() const = 0;
+
+	// Returns the editor object corresponding with the supplied editor object
+	virtual UObject* GetEditorObjectForRigVMGraph(URigVMGraph* InVMGraph) const = 0;
+
+	// Reacts to adding a graph
+	virtual void HandleRigVMGraphAdded(const FRigVMClient* InClient, const FString& InNodePathOrName) = 0;
+
+	// Reacts to removing a graph
+	virtual void HandleRigVMGraphRemoved(const FRigVMClient* InClient, const FString& InNodePathOrName) = 0;
+
+	// Reacts to renaming a graph
+	virtual void HandleRigVMGraphRenamed(const FRigVMClient* InClient, const FString& InOldNodePath, const FString& InNewNodePath) = 0;
+
+	// Reacts to a request to configure a controller
+	virtual void HandleConfigureRigVMController(const FRigVMClient* InClient, URigVMController* InControllerToConfigure) = 0;
+};
+
+UINTERFACE()
+class RIGVMDEVELOPER_API URigVMEditorSideObject : public UInterface
+{
+	GENERATED_BODY()
+};
+
+// Interface that allows a UI graph to identify itself against a model graph
+class RIGVMDEVELOPER_API IRigVMEditorSideObject
+{
+	GENERATED_BODY()
+
+public:
+
+	// Returns the corresponding VM graph
+	virtual FRigVMClient* GetRigVMClient() const = 0;
+
+	// Returns the nodepath for this UI graph
+	virtual FString GetRigVMNodePath() const = 0;
+
+	// Reacts to renaming the model
+	virtual void HandleRigVMGraphRenamed(const FString& InOldNodePath, const FString& InNewNodePath) = 0;
+};
+
+// A management struct containing graphs and controllers.
+USTRUCT()
+struct RIGVMDEVELOPER_API FRigVMClient
+{
+public:
+
+	GENERATED_BODY()
+
+	FRigVMClient()
+		: FunctionLibrary(nullptr)
+		, bSuspendNotifications(false)
+		, OuterClientHost(nullptr)
+		, OuterClientPropertyName(NAME_None)
+	{
+	}
+
+	void SetOuterClientHost(UObject* InOuterClientHost, const FName& InOuterClientHostPropertyName);
+	void SetFromDeprecatedData(URigVMGraph* InDefaultGraph, URigVMFunctionLibrary* InFunctionLibrary);
+
+	void Reset();
+	FORCEINLINE int32 Num() const { return Models.Num(); }
+	URigVMGraph* GetDefaultModel() const;
+	FORCEINLINE URigVMGraph* GetModel(int32 InIndex) const { return Models[InIndex]; }
+	URigVMGraph* GetModel(const FString& InNodePathOrName) const;
+	URigVMGraph* GetModel(const UObject* InEditorSideObject) const;
+	TArray<URigVMGraph*> GetAllModels(bool bIncludeFunctionLibrary, bool bRecursive) const;
+	URigVMController* GetController(int32 InIndex) const;
+	URigVMController* GetController(const FString& InNodePathOrName) const;
+	URigVMController* GetController(const URigVMGraph* InModel) const;
+	URigVMController* GetController(const UObject* InEditorSideObject) const;
+	URigVMController* GetOrCreateController(int32 InIndex);
+	URigVMController* GetOrCreateController(const FString& InNodePathOrName);
+	URigVMController* GetOrCreateController(const URigVMGraph* InModel);
+	URigVMController* GetOrCreateController(const UObject* InEditorSideObject);
+	URigVMFunctionLibrary* GetFunctionLibrary() const { return FunctionLibrary; }
+	URigVMFunctionLibrary* GetOrCreateFunctionLibrary(bool bSetupUndoRedo, const FObjectInitializer* ObjectInitializer = nullptr);
+
+	URigVMGraph* AddModel(const FName& InName, bool bSetupUndoRedo, const FObjectInitializer* ObjectInitializer = nullptr);
+	void AddModel(URigVMGraph* InModel, bool bCreateController);
+	bool RemoveModel(const FString& InNodePathOrName, bool bSetupUndoRedo);
+	FName RenameModel(const FString& InNodePathOrName, const FName& InNewName, bool bSetupUndoRedo);
+	void PostTransacted(const FTransactionObjectEvent& TransactionEvent);
+	void OnSubGraphRenamed(const FSoftObjectPath& OldObjectPath, const FSoftObjectPath& NewObjectPath);
+
+	URigVMNode* FindNode(const FString& InNodePathOrName) const;
+	URigVMPin* FindPin(const FString& InPinPath) const;
+	
+	FORCEINLINE TArray<TObjectPtr<URigVMGraph>>::RangedForIteratorType      begin() { return Models.begin(); }
+	FORCEINLINE TArray<TObjectPtr<URigVMGraph>>::RangedForConstIteratorType begin() const { return Models.begin(); }
+	FORCEINLINE TArray<TObjectPtr<URigVMGraph>>::RangedForIteratorType      end() { return Models.end(); }
+	FORCEINLINE TArray<TObjectPtr<URigVMGraph>>::RangedForConstIteratorType end() const { return Models.end(); }
+
+	UObject* GetOuter() const;
+	FProperty* GetOuterClientProperty() const;
+	void NotifyOuterOfPropertyChange(EPropertyChangeType::Type ChangeType = EPropertyChangeType::Interactive) const;
+	FName GetUniqueName(const FName& InDesiredName) const;
+	static FName GetUniqueName(UObject* InOuter, const FName& InDesiredName);
+	static void DestroyObject(UObject* InObject);
+
+private:
+
+	enum ERigVMClientAction
+	{
+		ERigVMClientAction_AddModel,
+		ERigVMClientAction_RemoveModel,
+		ERigVMClientAction_RenameModel
+	};
+
+	struct FRigVMClientAction
+	{
+		ERigVMClientAction Type;
+		FString NodePath;
+		FString OtherNodePath;
+	};
+
+	URigVMController* CreateController(const URigVMGraph* InModel);
+
+	UPROPERTY()
+	TArray<TObjectPtr<URigVMGraph>> Models;
+
+	UPROPERTY()
+	TObjectPtr<URigVMFunctionLibrary> FunctionLibrary;
+
+	UPROPERTY(transient)
+	TMap<FSoftObjectPath, TObjectPtr<URigVMController>> Controllers;
+
+	UPROPERTY(transient)
+	int32 UndoRedoIndex = 0;
+	
+	TArray<FRigVMClientAction> UndoStack;
+	TArray<FRigVMClientAction> RedoStack;
+
+public:
+	bool bSuspendNotifications;
+private:
+	TWeakObjectPtr<UObject> OuterClientHost;
+	FName OuterClientPropertyName;
+};
diff --git a/Engine/Source/Developer/RigVMDeveloper/Public/RigVMModel/RigVMController.h b/Engine/Source/Developer/RigVMDeveloper/Public/RigVMModel/RigVMController.h
--- a/Engine/Source/Developer/RigVMDeveloper/Public/RigVMModel/RigVMController.h
+++ b/Engine/Source/Developer/RigVMDeveloper/Public/RigVMModel/RigVMController.h
@@ -93,12 +93,6 @@ DECLARE_DELEGATE_OneParam(FRigVMController_RequestJumpToHyperlinkDelegate, const
 DECLARE_DELEGATE_OneParam(FRigVMController_ConfigureWorkflowOptionsDelegate, URigVMUserWorkflowOptions* InOutOptions);
 DECLARE_DELEGATE_RetVal_TwoParams(bool, FRigVMController_CheckPinComatibilityDelegate, URigVMPin*, URigVMPin*);
 
-UINTERFACE()
-class RIGVMDEVELOPER_API URigVMControllerHost : public UInterface
-{
-	GENERATED_BODY()
-};
-
 USTRUCT(BlueprintType)
 struct RIGVMDEVELOPER_API FRigStructScope
 {
@@ -135,25 +129,6 @@ protected:
 	const uint8* Memory;
 };
 
-// Interface that allows an object to host a rig VM controller. Used by graph edting code to interact with the controller.
-class RIGVMDEVELOPER_API IRigVMControllerHost
-{
-	GENERATED_BODY()
-
-public:
-	// Get an existing Rig VM controller corresponding to a URigVMGraph  
-	virtual URigVMController* GetRigVMController(const URigVMGraph* InRigVMGraph) const = 0;
-	
-	// Get an existing Rig VM controller corresponding to the supplied editor object that represents it (usually a UEdGraph) 
-	virtual URigVMController* GetRigVMController(const UObject* InEditorObject) const = 0;
-
-	// Get or create Rig VM controller corresponding to a URigVMGraph
-	virtual URigVMController* GetOrCreateRigVMController(URigVMGraph* InRigVMGraph) = 0;
-	
-	// Get or create a Rig VM controller corresponding to the supplied editor object that represents it (usually a UEdGraph) 
-	virtual URigVMController* GetOrCreateRigVMController(const UObject* InEditorObject) = 0;
-};
-
 /**
  * The Controller is the sole authority to perform changes
  * on the Graph. The Controller itself is stateless.
diff --git a/Engine/Source/Developer/RigVMDeveloper/Public/RigVMModel/RigVMGraph.h b/Engine/Source/Developer/RigVMDeveloper/Public/RigVMModel/RigVMGraph.h
--- a/Engine/Source/Developer/RigVMDeveloper/Public/RigVMModel/RigVMGraph.h
+++ b/Engine/Source/Developer/RigVMDeveloper/Public/RigVMModel/RigVMGraph.h
@@ -14,22 +14,7 @@
 #include "RigVMGraph.generated.h"
 
 class URigVMFunctionLibrary;
-
-UINTERFACE()
-class RIGVMDEVELOPER_API URigVMGraphHost : public UInterface
-{
-	GENERATED_BODY()
-};
-
-// Interface that allows an object to host a rig VM graph. Used by graph edting code to interact with the graph.
-class RIGVMDEVELOPER_API IRigVMGraphHost
-{
-	GENERATED_BODY()
-
-public:
-	// Get a Rig VM graph corresponding to the supplied editor object that represents it (usually a UEdGraph)
-	virtual URigVMGraph* GetRigVMGraph(const UObject* InEditorObject) const = 0;
-};
+struct FRigVMClient;
 
 /**
  * The Graph represents a Function definition
