Control Rig: Rename update and inverse to forwards and backwards solve

diff --git a/Engine/Plugins/Animation/ControlRig/Source/ControlRig/Private/ControlRig.cpp b/Engine/Plugins/Animation/ControlRig/Source/ControlRig/Private/ControlRig.cpp
--- a/Engine/Plugins/Animation/ControlRig/Source/ControlRig/Private/ControlRig.cpp
+++ b/Engine/Plugins/Animation/ControlRig/Source/ControlRig/Private/ControlRig.cpp
@@ -1597,6 +1597,25 @@ void UControlRig::PostLoad()
 		VMSnapshotBeforeExecution->SetFlags(VMSnapshotBeforeExecution->GetFlags() | RF_Transient);
 	}
 #endif
+
+	FRigInfluenceMapPerEvent NewInfluences;
+	for(int32 MapIndex = 0; MapIndex < Influences.Num(); MapIndex++)
+	{
+		FRigInfluenceMap Map = Influences[MapIndex];
+		FName EventName = Map.GetEventName();
+
+		if(EventName == TEXT("Update"))
+		{
+			EventName = FRigUnit_BeginExecution::EventName;
+		}
+		else if(EventName == TEXT("Inverse"))
+		{
+			EventName = FRigUnit_InverseExecution::EventName;
+		}
+		
+		NewInfluences.FindOrAdd(EventName).Merge(Map, true);
+	}
+	Influences = NewInfluences;
 }
 
 TArray<FRigControlElement*> UControlRig::AvailableControls() const
diff --git a/Engine/Plugins/Animation/ControlRig/Source/ControlRig/Private/ControlRigComponent.cpp b/Engine/Plugins/Animation/ControlRig/Source/ControlRig/Private/ControlRigComponent.cpp
--- a/Engine/Plugins/Animation/ControlRig/Source/ControlRig/Private/ControlRigComponent.cpp
+++ b/Engine/Plugins/Animation/ControlRig/Source/ControlRig/Private/ControlRigComponent.cpp
@@ -1791,12 +1791,12 @@
 		{
 			if (InCallingFunctionName)
 			{
-				ReportError(FString::Printf(TEXT("%s cannot be called during the UpdateEvent - use ConstructionScript instead."), InCallingFunctionName));
+				ReportError(FString::Printf(TEXT("%s cannot be called during the ForwardsSolveEvent - use ConstructionScript instead."), InCallingFunctionName));
 				return false;
 			}
 			else
 			{
-				ReportError(FString::Printf(TEXT("Cannot be called during the UpdateEvent - use ConstructionScript instead."), InCallingFunctionName));
+				ReportError(FString::Printf(TEXT("Cannot be called during the ForwardsSolveEvent - use ConstructionScript instead."), InCallingFunctionName));
 				return false;
 			}
 		}
diff --git a/Engine/Plugins/Animation/ControlRig/Source/ControlRig/Private/Rigs/RigInfluenceMap.cpp b/Engine/Plugins/Animation/ControlRig/Source/ControlRig/Private/Rigs/RigInfluenceMap.cpp
--- a/Engine/Plugins/Animation/ControlRig/Source/ControlRig/Private/Rigs/RigInfluenceMap.cpp
+++ b/Engine/Plugins/Animation/ControlRig/Source/ControlRig/Private/Rigs/RigInfluenceMap.cpp
@@ -85,23 +85,7 @@ FRigInfluenceMap FRigInfluenceMap::Inverse() const
 
 bool FRigInfluenceMap::Merge(const FRigInfluenceMap& Other)
 {
-	if(Other.EventName != EventName)
-	{
-		return false;
-	}
-
-	FRigInfluenceMap Temp = *this;
-	for(const FRigInfluenceEntry& OtherEntry : Other)
-	{
-		FRigInfluenceEntry& Entry = Temp.FindOrAdd(OtherEntry.Source);
-		if(!Entry.Merge(OtherEntry))
-		{
-			return false;
-		}
-	}
-
-	*this = Temp;
-	return true;
+	return Merge(Other, false);
 }
 
 FRigInfluenceEntryModifier FRigInfluenceMap::GetEntryModifier(const FRigElementKey& InKey) const
@@ -162,6 +146,27 @@ void FRigInfluenceMap::OnKeyRenamed(const FRigElementKey& InOldKey, const FRigEl
 	}
 }
 
+bool FRigInfluenceMap::Merge(const FRigInfluenceMap& Other, bool bIgnoreEventName)
+{
+	if(!bIgnoreEventName && (Other.EventName != EventName))
+	{
+		return false;
+	}
+
+	FRigInfluenceMap Temp = *this;
+	for(const FRigInfluenceEntry& OtherEntry : Other)
+	{
+		FRigInfluenceEntry& Entry = Temp.FindOrAdd(OtherEntry.Source);
+		if(!Entry.Merge(OtherEntry))
+		{
+			return false;
+		}
+	}
+
+	*this = Temp;
+	return true;
+}
+
 ////////////////////////////////////////////////////////////////////////////////
 // FRigInfluenceMapPerEvent
 ////////////////////////////////////////////////////////////////////////////////
diff --git a/Engine/Plugins/Animation/ControlRig/Source/ControlRig/Private/Units/Execution/RigUnit_BeginExecution.cpp b/Engine/Plugins/Animation/ControlRig/Source/ControlRig/Private/Units/Execution/RigUnit_BeginExecution.cpp
--- a/Engine/Plugins/Animation/ControlRig/Source/ControlRig/Private/Units/Execution/RigUnit_BeginExecution.cpp
+++ b/Engine/Plugins/Animation/ControlRig/Source/ControlRig/Private/Units/Execution/RigUnit_BeginExecution.cpp
@@ -3,7 +3,7 @@
 #include "RigUnit_BeginExecution.h"
 #include "Units/RigUnitContext.h"
 
-FName FRigUnit_BeginExecution::EventName = TEXT("Update");
+FName FRigUnit_BeginExecution::EventName = TEXT("Forwards Solve");
 
 FRigUnit_BeginExecution_Execute()
 {
diff --git a/Engine/Plugins/Animation/ControlRig/Source/ControlRig/Private/Units/Execution/RigUnit_InverseExecution.cpp b/Engine/Plugins/Animation/ControlRig/Source/ControlRig/Private/Units/Execution/RigUnit_InverseExecution.cpp
--- a/Engine/Plugins/Animation/ControlRig/Source/ControlRig/Private/Units/Execution/RigUnit_InverseExecution.cpp
+++ b/Engine/Plugins/Animation/ControlRig/Source/ControlRig/Private/Units/Execution/RigUnit_InverseExecution.cpp
@@ -3,7 +3,7 @@
 #include "RigUnit_InverseExecution.h"
 #include "Units/RigUnitContext.h"
 
-FName FRigUnit_InverseExecution::EventName = TEXT("Inverse");
+FName FRigUnit_InverseExecution::EventName = TEXT("Backwards Solve");
 
 FRigUnit_InverseExecution_Execute()
 {
diff --git a/Engine/Plugins/Animation/ControlRig/Source/ControlRig/Public/Rigs/RigInfluenceMap.h b/Engine/Plugins/Animation/ControlRig/Source/ControlRig/Public/Rigs/RigInfluenceMap.h
--- a/Engine/Plugins/Animation/ControlRig/Source/ControlRig/Public/Rigs/RigInfluenceMap.h
+++ b/Engine/Plugins/Animation/ControlRig/Source/ControlRig/Public/Rigs/RigInfluenceMap.h
@@ -156,6 +156,8 @@ public:
 
 protected:
 
+	bool Merge(const FRigInfluenceMap& Other, bool bIgnoreEventName);
+
 	UPROPERTY()
 	FName EventName;
 
@@ -166,6 +168,7 @@ protected:
 	TMap<FRigElementKey, int32> KeyToIndex;
 
 	friend struct FRigInfluenceMapPerEvent;
+	friend class UControlRig;
 };
 
 USTRUCT(BlueprintType)
diff --git a/Engine/Plugins/Animation/ControlRig/Source/ControlRigEditor/Private/ControlRigBlueprintCommands.cpp b/Engine/Plugins/Animation/ControlRig/Source/ControlRigEditor/Private/ControlRigBlueprintCommands.cpp
--- a/Engine/Plugins/Animation/ControlRig/Source/ControlRigEditor/Private/ControlRigBlueprintCommands.cpp
+++ b/Engine/Plugins/Animation/ControlRig/Source/ControlRigEditor/Private/ControlRigBlueprintCommands.cpp
@@ -11,9 +11,9 @@ void FControlRigBlueprintCommands::RegisterCommands()
 	UI_COMMAND(AutoCompileGraph, "Auto Compile", "Auto-compile the rig graph if On.", EUserInterfaceActionType::ToggleButton, FInputChord());
 	UI_COMMAND(ToggleEventQueue, "Toggle Event", "Toggle between the current and last running event", EUserInterfaceActionType::Button, FInputChord());
 	UI_COMMAND(SetupEvent, "Setup Event", "Enable the setup mode for the rig", EUserInterfaceActionType::Button, FInputChord());
-	UI_COMMAND(UpdateEvent, "Forwards Solve", "Run the normal update graph", EUserInterfaceActionType::Button, FInputChord());
-	UI_COMMAND(InverseEvent, "Backwards Solve", "Run the inverse graph", EUserInterfaceActionType::Button, FInputChord());
-	UI_COMMAND(InverseAndUpdateEvent, "Backwards and Forwards", "Run the inverse graph followed by the update graph", EUserInterfaceActionType::Button, FInputChord());
+	UI_COMMAND(ForwardsSolveEvent, "Forwards Solve", "Run the forwards solve graph", EUserInterfaceActionType::Button, FInputChord());
+	UI_COMMAND(BackwardsSolveEvent, "Backwards Solve", "Run the backwards solve graph", EUserInterfaceActionType::Button, FInputChord());
+	UI_COMMAND(BackwardsAndForwardsSolveEvent, "Backwards and Forwards", "Run backwards solve followed by forwards solve", EUserInterfaceActionType::Button, FInputChord());
 	UI_COMMAND(ToggleExecutionMode, "Toggle Execution Mode", "Toggle between Release and Debug execution mode", EUserInterfaceActionType::Button, FInputChord());
 	UI_COMMAND(ReleaseMode, "Release Mode", "Compiles and Executes the rig, ignoring debug data.", EUserInterfaceActionType::Button, FInputChord());
 	UI_COMMAND(DebugMode, "Debug Mode", "Compiles and Executes the unoptimized rig, stopping at breakpoints.", EUserInterfaceActionType::Button, FInputChord());
diff --git a/Engine/Plugins/Animation/ControlRig/Source/ControlRigEditor/Private/ControlRigBlueprintCommands.h b/Engine/Plugins/Animation/ControlRig/Source/ControlRigEditor/Private/ControlRigBlueprintCommands.h
--- a/Engine/Plugins/Animation/ControlRig/Source/ControlRigEditor/Private/ControlRigBlueprintCommands.h
+++ b/Engine/Plugins/Animation/ControlRig/Source/ControlRigEditor/Private/ControlRigBlueprintCommands.h
@@ -33,14 +33,14 @@ public:
 	/** Enable the setup mode for the rig */
 	TSharedPtr< FUICommandInfo > SetupEvent;
 
-	/** Run the normal update graph */
-	TSharedPtr< FUICommandInfo > UpdateEvent;
+	/** Run the forwards solve graph */
+	TSharedPtr< FUICommandInfo > ForwardsSolveEvent;
 
-	/** Run the inverse graph */
-	TSharedPtr< FUICommandInfo > InverseEvent;
+	/** Run the backwards solve graph */
+	TSharedPtr< FUICommandInfo > BackwardsSolveEvent;
 
-	/** Run the inverse graph followed by the update graph */
-	TSharedPtr< FUICommandInfo > InverseAndUpdateEvent;
+	/** Run the backwards solve graph followed by the forwards solve graph */
+	TSharedPtr< FUICommandInfo > BackwardsAndForwardsSolveEvent;
 
 	/** Toggle between Release and Debug execution mode */
 	TSharedPtr< FUICommandInfo > ToggleExecutionMode;
diff --git a/Engine/Plugins/Animation/ControlRig/Source/ControlRigEditor/Private/ControlRigEditorStyle.h b/Engine/Plugins/Animation/ControlRig/Source/ControlRigEditor/Private/ControlRigEditorStyle.h
--- a/Engine/Plugins/Animation/ControlRig/Source/ControlRigEditor/Private/ControlRigEditorStyle.h
+++ b/Engine/Plugins/Animation/ControlRig/Source/ControlRigEditor/Private/ControlRigEditorStyle.h
@@ -106,10 +106,9 @@ public:
 				FContentRootBracket Bracket(this, EngineEditorSlateDir);
 				Set("ControlRig.SetupMode", new IMAGE_BRUSH_SVG("Starship/Common/Adjust", Icon40x40));
 				Set("ControlRig.SetupMode.Small", new IMAGE_BRUSH_SVG("Starship/Common/Adjust", Icon20x20));
-				Set("ControlRig.UpdateEvent", new IMAGE_BRUSH("Icons/diff_next_40x", Icon40x40));
-				Set("ControlRig.InverseEvent", new IMAGE_BRUSH("Icons/diff_prev_40x", Icon40x40));
-				Set("ControlRig.UpdateAndInverse", new IMAGE_BRUSH("Icons/Loop_40x", Icon40x40));
-				Set("ControlRig.InverseAndUpdate", new IMAGE_BRUSH("Icons/Loop_40x", Icon40x40));
+				Set("ControlRig.ForwardsSolveEvent", new IMAGE_BRUSH("Icons/diff_next_40x", Icon40x40));
+				Set("ControlRig.BackwardsSolveEvent", new IMAGE_BRUSH("Icons/diff_prev_40x", Icon40x40));
+				Set("ControlRig.BackwardsAndForwardsSolveEvent", new IMAGE_BRUSH("Icons/Loop_40x", Icon40x40));
 			}
 
 			Set("ControlRig.Bug.Dot", new IMAGE_BRUSH("Slate/ControlRig_BugDot_32x", Icon16x16));
diff --git a/Engine/Plugins/Animation/ControlRig/Source/ControlRigEditor/Private/Editor/ControlRigEditor.cpp b/Engine/Plugins/Animation/ControlRig/Source/ControlRigEditor/Private/Editor/ControlRigEditor.cpp
--- a/Engine/Plugins/Animation/ControlRig/Source/ControlRigEditor/Private/Editor/ControlRigEditor.cpp
+++ b/Engine/Plugins/Animation/ControlRig/Source/ControlRigEditor/Private/Editor/ControlRigEditor.cpp
@@ -661,18 +661,18 @@ void FControlRigEditor::BindCommands()
 		FCanExecuteAction());
 
 	GetToolkitCommands()->MapAction(
-		FControlRigBlueprintCommands::Get().UpdateEvent,
-		FExecuteAction::CreateSP(this, &FControlRigEditor::SetEventQueue, EControlRigEditorEventQueue::Update),
+		FControlRigBlueprintCommands::Get().ForwardsSolveEvent,
+		FExecuteAction::CreateSP(this, &FControlRigEditor::SetEventQueue, EControlRigEditorEventQueue::ForwardsSolve),
 		FCanExecuteAction());
 
 	GetToolkitCommands()->MapAction(
-		FControlRigBlueprintCommands::Get().InverseEvent,
-		FExecuteAction::CreateSP(this, &FControlRigEditor::SetEventQueue, EControlRigEditorEventQueue::Inverse),
+		FControlRigBlueprintCommands::Get().BackwardsSolveEvent,
+		FExecuteAction::CreateSP(this, &FControlRigEditor::SetEventQueue, EControlRigEditorEventQueue::BackwardsSolve),
 		FCanExecuteAction());
 
 	GetToolkitCommands()->MapAction(
-		FControlRigBlueprintCommands::Get().InverseAndUpdateEvent,
-		FExecuteAction::CreateSP(this, &FControlRigEditor::SetEventQueue, EControlRigEditorEventQueue::InverseAndUpdate),
+		FControlRigBlueprintCommands::Get().BackwardsAndForwardsSolveEvent,
+		FExecuteAction::CreateSP(this, &FControlRigEditor::SetEventQueue, EControlRigEditorEventQueue::BackwardsAndForwardsSolve),
 		FCanExecuteAction());
 
 	GetToolkitCommands()->MapAction(
@@ -778,12 +778,12 @@ TSharedRef<SWidget> FControlRigEditor::GenerateEventQueueMenuContent()
 
 	MenuBuilder.BeginSection(TEXT("Events"));
 	MenuBuilder.AddMenuEntry(FControlRigBlueprintCommands::Get().SetupEvent, TEXT("Setup"), TAttribute<FText>(), TAttribute<FText>(), GetEventQueueIcon(EControlRigEditorEventQueue::Setup));
-	MenuBuilder.AddMenuEntry(FControlRigBlueprintCommands::Get().UpdateEvent, TEXT("Update"), TAttribute<FText>(), TAttribute<FText>(), GetEventQueueIcon(EControlRigEditorEventQueue::Update));
-	MenuBuilder.AddMenuEntry(FControlRigBlueprintCommands::Get().InverseEvent, TEXT("Inverse"), TAttribute<FText>(), TAttribute<FText>(), GetEventQueueIcon(EControlRigEditorEventQueue::Inverse));
+	MenuBuilder.AddMenuEntry(FControlRigBlueprintCommands::Get().ForwardsSolveEvent, TEXT("Forwards Solve"), TAttribute<FText>(), TAttribute<FText>(), GetEventQueueIcon(EControlRigEditorEventQueue::ForwardsSolve));
+	MenuBuilder.AddMenuEntry(FControlRigBlueprintCommands::Get().BackwardsSolveEvent, TEXT("Backwards Solve"), TAttribute<FText>(), TAttribute<FText>(), GetEventQueueIcon(EControlRigEditorEventQueue::BackwardsSolve));
 	MenuBuilder.EndSection();
 
 	MenuBuilder.BeginSection(TEXT("Validation"));
-	MenuBuilder.AddMenuEntry(FControlRigBlueprintCommands::Get().InverseAndUpdateEvent, TEXT("InverseAndUpdate"), TAttribute<FText>(), TAttribute<FText>(), GetEventQueueIcon(EControlRigEditorEventQueue::InverseAndUpdate));
+	MenuBuilder.AddMenuEntry(FControlRigBlueprintCommands::Get().BackwardsAndForwardsSolveEvent, TEXT("BackwardsAndForwards"), TAttribute<FText>(), TAttribute<FText>(), GetEventQueueIcon(EControlRigEditorEventQueue::BackwardsAndForwardsSolve));
 	MenuBuilder.EndSection();
 
 	return MenuBuilder.MakeWidget();
@@ -928,11 +928,11 @@ EControlRigEditorEventQueue FControlRigEditor::GetEventQueue() const
 			}
 			else if (EventQueue[0] == FRigUnit_BeginExecution::EventName)
 			{
-				return EControlRigEditorEventQueue::Update;
+				return EControlRigEditorEventQueue::ForwardsSolve;
 			}
 			else if (EventQueue[0] == FRigUnit_InverseExecution::EventName)
 			{
-				return EControlRigEditorEventQueue::Inverse;
+				return EControlRigEditorEventQueue::BackwardsSolve;
 			}
 		}
 		else if (EventQueue.Num() == 2)
@@ -940,12 +940,12 @@ EControlRigEditorEventQueue FControlRigEditor::GetEventQueue() const
 			if (EventQueue[0] == FRigUnit_InverseExecution::EventName &&
 				EventQueue[1] == FRigUnit_BeginExecution::EventName)
 			{
-				return EControlRigEditorEventQueue::InverseAndUpdate;
+				return EControlRigEditorEventQueue::BackwardsAndForwardsSolve;
 			}
 		}
 	}
 
-	return EControlRigEditorEventQueue::Update;
+	return EControlRigEditorEventQueue::ForwardsSolve;
 }
 
 void FControlRigEditor::SetEventQueue(EControlRigEditorEventQueue InEventQueue)
@@ -972,17 +972,17 @@ void FControlRigEditor::SetEventQueue(EControlRigEditorEventQueue InEventQueue)
 				}
 				return;
 			}
-			case EControlRigEditorEventQueue::Update:
+			case EControlRigEditorEventQueue::ForwardsSolve:
 			{
 				EventNames.Add(FRigUnit_BeginExecution::EventName);
 				break;
 			}
-			case EControlRigEditorEventQueue::Inverse:
+			case EControlRigEditorEventQueue::BackwardsSolve:
 			{
 				EventNames.Add(FRigUnit_InverseExecution::EventName);
 				break;
 			}
-			case EControlRigEditorEventQueue::InverseAndUpdate:
+			case EControlRigEditorEventQueue::BackwardsAndForwardsSolve:
 			{
 				EventNames.Add(FRigUnit_InverseExecution::EventName);
 				EventNames.Add(FRigUnit_BeginExecution::EventName);
@@ -1007,7 +1007,7 @@ void FControlRigEditor::SetEventQueue(EControlRigEditorEventQueue InEventQueue)
 
 		// Reset transforms only for setup and forward solve to not inturrupt any animation that might be playing
 		if (InEventQueue == EControlRigEditorEventQueue::Setup ||
-			InEventQueue == EControlRigEditorEventQueue::Update)
+			InEventQueue == EControlRigEditorEventQueue::ForwardsSolve)
 		{
 			ControlRig->GetHierarchy()->ResetPoseToInitial(ERigElementType::All);
 		}
@@ -1028,15 +1028,15 @@ FText FControlRigEditor::GetEventQueueLabel() const
 		{
 			return FRigUnit_PrepareForExecution::StaticStruct()->GetDisplayNameText();
 		}
-		case EControlRigEditorEventQueue::Update:
+		case EControlRigEditorEventQueue::ForwardsSolve:
 		{
 			return FRigUnit_BeginExecution::StaticStruct()->GetDisplayNameText();
 		}
-		case EControlRigEditorEventQueue::Inverse:
+		case EControlRigEditorEventQueue::BackwardsSolve:
 		{
 			return FRigUnit_InverseExecution::StaticStruct()->GetDisplayNameText();
 		}
-		case EControlRigEditorEventQueue::InverseAndUpdate:
+		case EControlRigEditorEventQueue::BackwardsAndForwardsSolve:
 		{
 			return FText::FromString(FString::Printf(TEXT("%s and %s"),
 				*FRigUnit_InverseExecution::StaticStruct()->GetDisplayNameText().ToString(),
@@ -1054,17 +1054,17 @@ FSlateIcon FControlRigEditor::GetEventQueueIcon(EControlRigEditorEventQueue InEv
 		{
 			return FSlateIcon(FControlRigEditorStyle::Get().GetStyleSetName(), "ControlRig.SetupMode");
 		}
-		case EControlRigEditorEventQueue::Update:
+		case EControlRigEditorEventQueue::ForwardsSolve:
 		{
-			return FSlateIcon(FControlRigEditorStyle::Get().GetStyleSetName(), "ControlRig.UpdateEvent");
+			return FSlateIcon(FControlRigEditorStyle::Get().GetStyleSetName(), "ControlRig.ForwardsSolveEvent");
 		}
-		case EControlRigEditorEventQueue::Inverse:
+		case EControlRigEditorEventQueue::BackwardsSolve:
 		{
-			return FSlateIcon(FControlRigEditorStyle::Get().GetStyleSetName(), "ControlRig.InverseEvent");
+			return FSlateIcon(FControlRigEditorStyle::Get().GetStyleSetName(), "ControlRig.BackwardsSolveEvent");
 		}
-		case EControlRigEditorEventQueue::InverseAndUpdate:
+		case EControlRigEditorEventQueue::BackwardsAndForwardsSolve:
 		{
-			return FSlateIcon(FControlRigEditorStyle::Get().GetStyleSetName(), "ControlRig.InverseAndUpdate");
+			return FSlateIcon(FControlRigEditorStyle::Get().GetStyleSetName(), "ControlRig.BackwardsAndForwardsSolveEvent");
 		}
 	}
 
@@ -3673,10 +3673,10 @@ void FControlRigEditor::HandleViewportCreated(const TSharedRef<class IPersonaVie
 			case EControlRigEditorEventQueue::Setup:
 				Color = UControlRigEditorSettings::Get()->SetupEventBorderColor;
 				break;
-			case EControlRigEditorEventQueue::Inverse:
+			case EControlRigEditorEventQueue::BackwardsSolve:
 				Color = UControlRigEditorSettings::Get()->BackwardsSolveBorderColor;
 				break;
-			case EControlRigEditorEventQueue::InverseAndUpdate:
+			case EControlRigEditorEventQueue::BackwardsAndForwardsSolve:
 				Color = UControlRigEditorSettings::Get()->BackwardsAndForwardsBorderColor;
 				break;
 			default:
@@ -3689,7 +3689,7 @@ void FControlRigEditor::HandleViewportCreated(const TSharedRef<class IPersonaVie
 	auto GetBorderVisibility = [this]()
 	{
 		EVisibility Visibility = EVisibility::Collapsed;
-		if (GetEventQueue() != EControlRigEditorEventQueue::Update)
+		if (GetEventQueue() != EControlRigEditorEventQueue::ForwardsSolve)
 		{
 			Visibility = EVisibility::HitTestInvisible;
 		}
diff --git a/Engine/Plugins/Animation/ControlRig/Source/ControlRigEditor/Private/Editor/ControlRigEditor.h b/Engine/Plugins/Animation/ControlRig/Source/ControlRigEditor/Private/Editor/ControlRigEditor.h
--- a/Engine/Plugins/Animation/ControlRig/Source/ControlRigEditor/Private/Editor/ControlRigEditor.h
+++ b/Engine/Plugins/Animation/ControlRig/Source/ControlRigEditor/Private/Editor/ControlRigEditor.h
@@ -39,14 +39,14 @@ enum class EControlRigEditorEventQueue : uint8
 	/** Setup Event */
 	Setup,
 
-	/** Update Event */
-	Update,
+	/** Forwards Solve Event */
+	ForwardsSolve,
 
-	/** Inverse Event */
-	Inverse,
+	/** Backwards Solve Event */
+	BackwardsSolve,
 
-	/** Inverse -> Update */
-	InverseAndUpdate,
+	/** Backwards Solve Event -> Forwards Solve Event */
+	BackwardsAndForwardsSolve,
 
 	/** MAX - invalid */
 	Max UMETA(Hidden),
