Control Rig: Maintain skip flag for selection once it's been set

diff --git a/Engine/Plugins/Animation/ControlRig/Source/ControlRigEditor/Private/Sequencer/ControlRigParameterTrackEditor.cpp b/Engine/Plugins/Animation/ControlRig/Source/ControlRigEditor/Private/Sequencer/ControlRigParameterTrackEditor.cpp
--- a/Engine/Plugins/Animation/ControlRig/Source/ControlRigEditor/Private/Sequencer/ControlRigParameterTrackEditor.cpp
+++ b/Engine/Plugins/Animation/ControlRig/Source/ControlRigEditor/Private/Sequencer/ControlRigParameterTrackEditor.cpp
@@ -2389,7 +2389,8 @@ void FControlRigParameterTrackEditor::HandleControlSelected(UControlRig* Subject
 			SetUpEditModeIfNeeded(Subject);
 
 			//Force refresh later, not now
-			bSkipNextSelectionFromTimer = bIsSelectingIndirectControl && ControlElement->Settings.AnimationType == ERigControlAnimationType::AnimationControl;
+			bSkipNextSelectionFromTimer = bSkipNextSelectionFromTimer ||
+				(bIsSelectingIndirectControl && ControlElement->Settings.AnimationType == ERigControlAnimationType::AnimationControl);
 			GetSequencer()->NotifyMovieSceneDataChanged(EMovieSceneDataChangeType::RefreshTree);
 		}
 	}
