Sequencer: Snapper: Clear snapper objects when sub sequence changes
#trivial

diff --git a/Engine/Plugins/Animation/ControlRig/Source/ControlRigEditor/Private/EditMode/SControlRigSnapper.cpp b/Engine/Plugins/Animation/ControlRig/Source/ControlRigEditor/Private/EditMode/SControlRigSnapper.cpp
--- a/Engine/Plugins/Animation/ControlRig/Source/ControlRigEditor/Private/EditMode/SControlRigSnapper.cpp
+++ b/Engine/Plugins/Animation/ControlRig/Source/ControlRigEditor/Private/EditMode/SControlRigSnapper.cpp
@@ -367,8 +367,29 @@ void SControlRigSnapper::Construct(const FArguments& InArgs)
 				]
 			]
 		];
+
+	TWeakPtr<ISequencer> Sequencer = Snapper.GetSequencer();
+	if (Sequencer.IsValid())
+	{
+		Sequencer.Pin()->OnActivateSequence().AddRaw(this, &SControlRigSnapper::OnActivateSequenceChanged);
+	}
 }
 
+SControlRigSnapper::~SControlRigSnapper()
+{
+	TWeakPtr<ISequencer> Sequencer = Snapper.GetSequencer();
+	if (Sequencer.IsValid())
+	{
+		Sequencer.Pin()->OnActivateSequence().RemoveAll(this);
+	}
+}
+
+void SControlRigSnapper::OnActivateSequenceChanged(FMovieSceneSequenceIDRef ID)
+{
+	ClearActors();
+}
+
+
 FReply SControlRigSnapper::OnActorToSnapClicked()
 {
 	ActorToSnap = GetSelection(true);
diff --git a/Engine/Plugins/Animation/ControlRig/Source/ControlRigEditor/Private/EditMode/SControlRigSnapper.h b/Engine/Plugins/Animation/ControlRig/Source/ControlRigEditor/Private/EditMode/SControlRigSnapper.h
--- a/Engine/Plugins/Animation/ControlRig/Source/ControlRigEditor/Private/EditMode/SControlRigSnapper.h
+++ b/Engine/Plugins/Animation/ControlRig/Source/ControlRigEditor/Private/EditMode/SControlRigSnapper.h
@@ -11,6 +11,7 @@
 #include "Tools/ControlRigSnapper.h"
 #include "Misc/FrameNumber.h"
 #include "IDetailsView.h"
+#include "MovieSceneSequenceID.h"
 
 class UControlRig;
 class ISequencer;
@@ -22,9 +23,7 @@ class SControlRigSnapper : public SCompoundWidget
 	SLATE_BEGIN_ARGS(SControlRigSnapper) {}
 
 	SLATE_END_ARGS()
-	~SControlRigSnapper()
-	{
-	}
+	~SControlRigSnapper();
 
 	void Construct(const FArguments& InArgs);
 
@@ -39,6 +38,7 @@ private:
 	void ActorParentPicked(FActorForWorldTransforms Selection);
 	void ActorParentSocketPicked(const FName SocketName, FActorForWorldTransforms Selection);
 	void ActorParentComponentPicked(FName ComponentName, FActorForWorldTransforms Selection);
+	void OnActivateSequenceChanged(FMovieSceneSequenceIDRef ID);
 
 	FReply OnActorToSnapClicked();
 	FReply OnParentToSnapToClicked();
