Sequencer: Space Switches, Baking fix again, error got wiped over.
#trivial

diff --git a/Engine/Plugins/Animation/ControlRig/Source/ControlRigEditor/Private/EditMode/SControlRigEditModeTools.cpp b/Engine/Plugins/Animation/ControlRig/Source/ControlRigEditor/Private/EditMode/SControlRigEditModeTools.cpp
--- a/Engine/Plugins/Animation/ControlRig/Source/ControlRigEditor/Private/EditMode/SControlRigEditModeTools.cpp
+++ b/Engine/Plugins/Animation/ControlRig/Source/ControlRigEditor/Private/EditMode/SControlRigEditModeTools.cpp
@@ -883,7 +883,7 @@ FReply SControlRigEditModeTools::OnBakeControlsToNewSpaceButtonClicked()
 	bool bNoValidControlRig = true;
 	for (TWeakObjectPtr<UControlRig>& ControlRig : ControlRigs)
 	{
-		if (ControlRig.IsValid())
+		if (ControlRig.IsValid() && SpacePickerWidget->GetHierarchy() == ControlRig->GetHierarchy())
 		{
 			bNoValidControlRig = false;
 			break;
@@ -901,7 +901,7 @@ FReply SControlRigEditModeTools::OnBakeControlsToNewSpaceButtonClicked()
 	}
 	for (TWeakObjectPtr<UControlRig>& ControlRig : ControlRigs)
 	{
-		if (ControlRig.IsValid())
+		if (ControlRig.IsValid() && SpacePickerWidget->GetHierarchy() == ControlRig->GetHierarchy())
 		{
 
 			FRigSpacePickerBakeSettings Settings;
