reset hierachy when disactivating the dialog window as we don't need to update the widget anymore

diff --git a/Engine/Plugins/Animation/ControlRig/Source/ControlRigEditor/Private/Editor/SRigSpacePickerWidget.cpp b/Engine/Plugins/Animation/ControlRig/Source/ControlRigEditor/Private/Editor/SRigSpacePickerWidget.cpp
--- a/Engine/Plugins/Animation/ControlRig/Source/ControlRigEditor/Private/Editor/SRigSpacePickerWidget.cpp
+++ b/Engine/Plugins/Animation/ControlRig/Source/ControlRigEditor/Private/Editor/SRigSpacePickerWidget.cpp
@@ -251,6 +251,13 @@ FReply SRigSpacePickerWidget::OpenDialog(bool bModal)
 	];
 	
 	Window->SetWidgetToFocusOnActivate(AsShared());
+	if (!Window->GetOnWindowDeactivatedEvent().IsBoundToObject(this))
+	{
+		Window->GetOnWindowDeactivatedEvent().AddLambda([this]()
+		{
+			SetControls(nullptr, {});
+		});
+	}
 	
 	DialogWindow = Window;
 
@@ -282,6 +289,7 @@ void SRigSpacePickerWidget::CloseDialog()
 	
 	if ( DialogWindow.IsValid() )
 	{
+		DialogWindow.Pin()->GetOnWindowDeactivatedEvent().RemoveAll(this);
 		DialogWindow.Pin()->RequestDestroyWindow();
 		DialogWindow.Reset();
 	}
