Control Rig: Enable translucent selection by default

diff --git a/Engine/Plugins/Animation/ControlRig/Source/ControlRigEditor/Private/EditMode/ControlRigEditMode.cpp b/Engine/Plugins/Animation/ControlRig/Source/ControlRigEditor/Private/EditMode/ControlRigEditMode.cpp
--- a/Engine/Plugins/Animation/ControlRig/Source/ControlRigEditor/Private/EditMode/ControlRigEditMode.cpp
+++ b/Engine/Plugins/Animation/ControlRig/Source/ControlRigEditor/Private/EditMode/ControlRigEditMode.cpp
@@ -59,6 +59,7 @@
 #include "LevelEditor.h"
 #include "InteractiveToolManager.h"
 #include "EdModeInteractiveToolsContext.h"
+#include "Editor/EditorPerProjectUserSettings.h"
 
 void UControlRigEditModeDelegateHelper::OnPoseInitialized()
 {
@@ -2995,6 +2996,12 @@
 	FActorSpawnParameters ActorSpawnParameters;
 	ActorSpawnParameters.bTemporaryEditorActor = true;
 
+	if(bShowControlsAsOverlay)
+	{
+		// enable translucent selection
+		GetMutableDefault<UEditorPerProjectUserSettings>()->bAllowSelectTranslucent = true;
+	}
+
 	TArray<FRigControlElement*> Controls = ControlRig->AvailableControls();
 	const TArray<TSoftObjectPtr<UControlRigShapeLibrary>> ShapeLibraries = ControlRig->GetShapeLibraries();
 	int32 ControlRigIndex = RuntimeControlRigs.Find(ControlRig);
