Control Rig: Provide feedback when changing initial transform during forward solve

diff --git a/Engine/Plugins/Animation/ControlRig/Source/ControlRig/Private/Units/Hierarchy/RigUnit_SetTransform.cpp b/Engine/Plugins/Animation/ControlRig/Source/ControlRig/Private/Units/Hierarchy/RigUnit_SetTransform.cpp
--- a/Engine/Plugins/Animation/ControlRig/Source/ControlRig/Private/Units/Hierarchy/RigUnit_SetTransform.cpp
+++ b/Engine/Plugins/Animation/ControlRig/Source/ControlRig/Private/Units/Hierarchy/RigUnit_SetTransform.cpp
@@ -4,7 +4,7 @@
 #include "RigUnit_GetTransform.h"
 #include "Units/RigUnitContext.h"
 #include "Units/Execution/RigUnit_PrepareForExecution.h"
-#include "Units/Execution/RigUnit_InverseExecution.h"
+#include "Units/Execution/RigUnit_BeginExecution.h"
 #include "Units/Hierarchy/RigUnit_SetControlOffset.h"
 
 FString FRigUnit_SetTransform::GetUnitLabel() const
@@ -41,6 +41,17 @@ FRigUnit_SetTransform_Execute()
 				}
 				else
 				{
+#if WITH_EDITOR
+					if(bInitial)
+					{
+						// provide some user feedback when changing initial transforms during forward solve
+						if(RigVMExecuteContext.GetEventName() == FRigUnit_BeginExecution::EventName)
+						{
+							UE_CONTROLRIG_RIGUNIT_LOG_MESSAGE(TEXT("Changing initial transforms during %s is not recommended."), *RigVMExecuteContext.GetEventName().ToString());
+						}
+					}
+#endif
+
 					// for controls - set the control offset transform instead
 					if(bInitial && (CachedIndex.GetKey().Type == ERigElementType::Control))
 					{
