fixed missing include

diff --git a/Engine/Plugins/Animation/ControlRig/Source/ControlRigEditor/Public/Tools/BakingHelper.h b/Engine/Plugins/Animation/ControlRig/Source/ControlRigEditor/Public/Tools/BakingHelper.h
--- a/Engine/Plugins/Animation/ControlRig/Source/ControlRigEditor/Public/Tools/BakingHelper.h
+++ b/Engine/Plugins/Animation/ControlRig/Source/ControlRigEditor/Public/Tools/BakingHelper.h
@@ -4,6 +4,7 @@
 
 #include "Misc/FrameNumber.h"
 #include "Templates/SharedPointer.h"
+#include "Math/Transform.h"
 
 class UMovieScene3DTransformSection;
 class AActor;
