Control Rig: Allow to get and set control shape transform

diff --git a/Engine/Plugins/Animation/ControlRig/Source/ControlRig/Private/Units/Hierarchy/RigUnit_SetControlOffset.cpp b/Engine/Plugins/Animation/ControlRig/Source/ControlRig/Private/Units/Hierarchy/RigUnit_SetControlOffset.cpp
--- a/Engine/Plugins/Animation/ControlRig/Source/ControlRig/Private/Units/Hierarchy/RigUnit_SetControlOffset.cpp
+++ b/Engine/Plugins/Animation/ControlRig/Source/ControlRig/Private/Units/Hierarchy/RigUnit_SetControlOffset.cpp
@@ -48,3 +48,71 @@ FRigUnit_SetControlOffset_Execute()
 		}
 	}
 }
+
+FRigUnit_GetShapeTransform_Execute()
+{
+	DECLARE_SCOPE_HIERARCHICAL_COUNTER_RIGUNIT()
+
+	const URigHierarchy* Hierarchy = Context.Hierarchy;
+	if (Hierarchy)
+	{
+		switch (Context.State)
+		{
+			case EControlRigState::Init:
+			{
+				CachedControlIndex.Reset();
+				break;
+			}
+			case EControlRigState::Update:
+			{
+				if (!CachedControlIndex.UpdateCache(FRigElementKey(Control, ERigElementType::Control), Hierarchy))
+				{
+					UE_CONTROLRIG_RIGUNIT_REPORT_WARNING(TEXT("Control '%s' is not valid."), *Control.ToString());
+					return;
+				}
+
+				const FRigControlElement* ControlElement = Hierarchy->Get<FRigControlElement>(CachedControlIndex);
+				Transform = Hierarchy->GetControlShapeTransform((FRigControlElement*)ControlElement, ERigTransformType::CurrentLocal);
+				break;
+			}
+			default:
+			{
+				break;
+			}
+		}
+	}
+}
+
+FRigUnit_SetShapeTransform_Execute()
+{
+	DECLARE_SCOPE_HIERARCHICAL_COUNTER_RIGUNIT()
+
+	URigHierarchy* Hierarchy = ExecuteContext.Hierarchy;
+	if (Hierarchy)
+	{
+		switch (Context.State)
+		{
+			case EControlRigState::Init:
+			{
+				CachedControlIndex.Reset();
+			break;
+			}
+			case EControlRigState::Update:
+			{
+				if (!CachedControlIndex.UpdateCache(FRigElementKey(Control, ERigElementType::Control), Hierarchy))
+				{
+					UE_CONTROLRIG_RIGUNIT_REPORT_WARNING(TEXT("Control '%s' is not valid."), *Control.ToString());
+					return;
+				}
+
+				FRigControlElement* ControlElement = Hierarchy->Get<FRigControlElement>(CachedControlIndex);
+				Hierarchy->SetControlShapeTransform(ControlElement, Transform, ERigTransformType::CurrentLocal);
+				break;
+			}
+			default:
+			{
+				break;
+			}
+		}
+	}
+}
\ No newline at end of file
diff --git a/Engine/Plugins/Animation/ControlRig/Source/ControlRig/Private/Units/Hierarchy/RigUnit_SetControlOffset.h b/Engine/Plugins/Animation/ControlRig/Source/ControlRig/Private/Units/Hierarchy/RigUnit_SetControlOffset.h
--- a/Engine/Plugins/Animation/ControlRig/Source/ControlRig/Private/Units/Hierarchy/RigUnit_SetControlOffset.h
+++ b/Engine/Plugins/Animation/ControlRig/Source/ControlRig/Private/Units/Hierarchy/RigUnit_SetControlOffset.h
@@ -9,7 +9,7 @@
  * SetControlOffset is used to perform a change in the hierarchy by setting a single control's transform.
  * This is typically only used during the Setup Event.
  */
-USTRUCT(meta=(DisplayName="Set Control Offset", Category="Controls", DocumentationPolicy="Strict", Keywords = "SetControlOffset,Initial,InitialTransform,SetInitialTransform,SetInitialControlTransform", TemplateName="SetControlValue", NodeColor="0, 0.364706, 1.0"))
+USTRUCT(meta=(DisplayName="Set Control Offset", Category="Controls", DocumentationPolicy="Strict", Keywords = "SetControlOffset,Initial,InitialTransform,SetInitialTransform,SetInitialControlTransform", NodeColor="0, 0.364706, 1.0"))
 struct CONTROLRIG_API FRigUnit_SetControlOffset : public FRigUnitMutable
 {
 	GENERATED_BODY()
@@ -17,7 +17,8 @@ struct CONTROLRIG_API FRigUnit_SetControlOffset : public FRigUnitMutable
 	FRigUnit_SetControlOffset()
 		: Control(NAME_None)
 		, Space(EBoneGetterSetterMode::GlobalSpace)
-	{}
+	{
+	}
 
 	RIGVM_METHOD()
 	virtual void Execute(const FRigUnitContext& Context) override;
@@ -29,7 +30,7 @@ struct CONTROLRIG_API FRigUnit_SetControlOffset : public FRigUnitMutable
 	FName Control;
 
 	/**
-	 * The offse transform to set for the control
+	 * The offset transform to set for the control
 	 */
 	UPROPERTY(meta = (Input, Output))
 	FTransform Offset;
@@ -45,3 +46,73 @@ struct CONTROLRIG_API FRigUnit_SetControlOffset : public FRigUnitMutable
 	UPROPERTY()
 	FCachedRigElement CachedControlIndex;
 };
+
+/**
+ * GetShapeTransform is used to retrieve single control's shape transform.
+ * This is typically only used during the Setup Event.
+ */
+USTRUCT(meta=(DisplayName="Get Shape Transform", Category="Controls", DocumentationPolicy="Strict", Keywords = "GetControlShapeTransform,Gizmo,GizmoTransform,MeshTransform", NodeColor="0, 0.364706, 1.0"))
+struct CONTROLRIG_API FRigUnit_GetShapeTransform : public FRigUnit
+{
+	GENERATED_BODY()
+
+	FRigUnit_GetShapeTransform()
+		: Control(NAME_None)
+	{
+		Transform = FTransform::Identity;
+	}
+
+	RIGVM_METHOD()
+	virtual void Execute(const FRigUnitContext& Context) override;
+
+	/**
+	 * The name of the Control to set the transform for.
+	 */
+	UPROPERTY(meta = (Input, CustomWidget = "ControlName" ))
+	FName Control;
+
+	/**
+	 * The shape transform to set for the control
+	 */
+	UPROPERTY(meta = (Output))
+	FTransform Transform;
+
+	// user to internally cache the index of the bone
+	UPROPERTY()
+	FCachedRigElement CachedControlIndex;
+};
+
+/**
+ * SetShapeTransform is used to perform a change in the hierarchy by setting a single control's shape transform.
+ * This is typically only used during the Setup Event.
+ */
+USTRUCT(meta=(DisplayName="Set Shape Transform", Category="Controls", DocumentationPolicy="Strict", Keywords = "SetControlShapeTransform,Gizmo,GizmoTransform,MeshTransform", NodeColor="0, 0.364706, 1.0"))
+struct CONTROLRIG_API FRigUnit_SetShapeTransform : public FRigUnitMutable
+{
+	GENERATED_BODY()
+
+	FRigUnit_SetShapeTransform()
+		: Control(NAME_None)
+	{
+		Transform = FTransform::Identity;
+	}
+
+	RIGVM_METHOD()
+	virtual void Execute(const FRigUnitContext& Context) override;
+
+	/**
+	 * The name of the Control to set the transform for.
+	 */
+	UPROPERTY(meta = (Input, CustomWidget = "ControlName" ))
+	FName Control;
+
+	/**
+	 * The shape transform to set for the control
+	 */
+	UPROPERTY(meta = (Input))
+	FTransform Transform;
+
+	// user to internally cache the index of the bone
+	UPROPERTY()
+	FCachedRigElement CachedControlIndex;
+};
