ControlRig: user defined events and event invocation

diff --git a/Engine/Plugins/Animation/ControlRig/Source/ControlRig/Private/Units/Execution/RigUnit_BeginExecution.h b/Engine/Plugins/Animation/ControlRig/Source/ControlRig/Private/Units/Execution/RigUnit_BeginExecution.h
--- a/Engine/Plugins/Animation/ControlRig/Source/ControlRig/Private/Units/Execution/RigUnit_BeginExecution.h
+++ b/Engine/Plugins/Animation/ControlRig/Source/ControlRig/Private/Units/Execution/RigUnit_BeginExecution.h
@@ -18,6 +18,7 @@ struct CONTROLRIG_API FRigUnit_BeginExecution : public FRigUnit
 	virtual void Execute(const FRigUnitContext& Context) override;
 
 	virtual FName GetEventName() const override { return EventName; }
+	virtual bool CanOnlyExistOnce() const override { return true; }
 
 	// The execution result
 	UPROPERTY(EditAnywhere, Transient, DisplayName = "Execute", Category = "BeginExecution", meta = (Output))
diff --git a/Engine/Plugins/Animation/ControlRig/Source/ControlRig/Private/Units/Execution/RigUnit_InverseExecution.h b/Engine/Plugins/Animation/ControlRig/Source/ControlRig/Private/Units/Execution/RigUnit_InverseExecution.h
--- a/Engine/Plugins/Animation/ControlRig/Source/ControlRig/Private/Units/Execution/RigUnit_InverseExecution.h
+++ b/Engine/Plugins/Animation/ControlRig/Source/ControlRig/Private/Units/Execution/RigUnit_InverseExecution.h
@@ -18,6 +18,7 @@ struct CONTROLRIG_API FRigUnit_InverseExecution : public FRigUnit
 	virtual void Execute(const FRigUnitContext& Context) override;
 
 	virtual FName GetEventName() const override { return EventName; }
+	virtual bool CanOnlyExistOnce() const override { return true; }
 
 	// The execution result
 	UPROPERTY(EditAnywhere, Transient, DisplayName = "Execute", Category = "InverseExecution", meta = (Output))
diff --git a/Engine/Plugins/Animation/ControlRig/Source/ControlRig/Private/Units/Execution/RigUnit_PrepareForExecution.h b/Engine/Plugins/Animation/ControlRig/Source/ControlRig/Private/Units/Execution/RigUnit_PrepareForExecution.h
--- a/Engine/Plugins/Animation/ControlRig/Source/ControlRig/Private/Units/Execution/RigUnit_PrepareForExecution.h
+++ b/Engine/Plugins/Animation/ControlRig/Source/ControlRig/Private/Units/Execution/RigUnit_PrepareForExecution.h
@@ -18,6 +18,7 @@ struct CONTROLRIG_API FRigUnit_PrepareForExecution : public FRigUnit
 	virtual void Execute(const FRigUnitContext& Context) override;
 
 	virtual FName GetEventName() const override { return EventName; }
+	virtual bool CanOnlyExistOnce() const override { return true; }
 
 	// The execution result
 	UPROPERTY(EditAnywhere, Transient, DisplayName = "Execute", Category = "PrepareForExecution", meta = (Output))
diff --git a/Engine/Plugins/Animation/ControlRig/Source/ControlRig/Private/Units/Execution/RigUnit_UserDefinedEvent.cpp b/Engine/Plugins/Animation/ControlRig/Source/ControlRig/Private/Units/Execution/RigUnit_UserDefinedEvent.cpp
new file mode 100644
--- /dev/null
+++ b/Engine/Plugins/Animation/ControlRig/Source/ControlRig/Private/Units/Execution/RigUnit_UserDefinedEvent.cpp
@@ -0,0 +1,27 @@
+// Copyright Epic Games, Inc. All Rights Reserved.
+
+#include "RigUnit_UserDefinedEvent.h"
+#include "Units/RigUnitContext.h"
+
+FString FRigUnit_UserDefinedEvent::GetUnitLabel() const
+{
+	if(EventName.IsNone())
+	{
+		static const FString InvalidEvent = TEXT("Invalid Event");
+		return InvalidEvent;
+	}
+	return *EventName.ToString();
+}
+
+FRigUnit_UserDefinedEvent_Execute()
+{
+    DECLARE_SCOPE_HIERARCHICAL_COUNTER_RIGUNIT()
+	ExecuteContext.CopyFrom(RigVMExecuteContext);
+	ExecuteContext.Hierarchy = Context.Hierarchy;
+	ExecuteContext.EventName = EventName;
+
+	if(EventName.IsNone())
+	{
+		UE_CONTROLRIG_RIGUNIT_REPORT_ERROR(TEXT("Unable to run event - EventName is None."));
+	}
+}
diff --git a/Engine/Plugins/Animation/ControlRig/Source/ControlRig/Private/Units/Execution/RigUnit_UserDefinedEvent.h b/Engine/Plugins/Animation/ControlRig/Source/ControlRig/Private/Units/Execution/RigUnit_UserDefinedEvent.h
new file mode 100644
--- /dev/null
+++ b/Engine/Plugins/Animation/ControlRig/Source/ControlRig/Private/Units/Execution/RigUnit_UserDefinedEvent.h
@@ -0,0 +1,36 @@
+// Copyright Epic Games, Inc. All Rights Reserved.
+
+#pragma once
+
+#include "Units/RigUnit.h"
+#include "ControlRigDefines.h"
+#include "RigUnit_UserDefinedEvent.generated.h"
+
+/**
+ * User Defined Event for running custom logic
+ */
+USTRUCT(meta=(DisplayName="User Defined Event", Category="Events", NodeColor="1, 0, 0", Keywords="Event,Entry,MyEvent"))
+struct CONTROLRIG_API FRigUnit_UserDefinedEvent : public FRigUnit
+{
+	GENERATED_BODY()
+
+	FRigUnit_UserDefinedEvent()
+	{
+		EventName = TEXT("My Event");
+	}
+
+	RIGVM_METHOD()
+	virtual void Execute(const FRigUnitContext& Context) override;
+
+	virtual FString GetUnitLabel() const override;
+	virtual FName GetEventName() const override { return EventName; }
+	virtual bool CanOnlyExistOnce() const override { return false; }
+
+	// The execution result
+	UPROPERTY(EditAnywhere, Transient, DisplayName = "Execute", Category = "UserDefinedEvent", meta = (Output))
+	FControlRigExecuteContext ExecuteContext;
+
+	// True if the current interaction is a rotation
+	UPROPERTY(EditAnywhere, Transient, Category = "UserDefinedEvent", meta = (Input,Constant,DetailsOnly))
+	FName EventName;
+};
diff --git a/Engine/Plugins/Animation/ControlRig/Source/ControlRigDeveloper/Private/ControlRigBlueprint.cpp b/Engine/Plugins/Animation/ControlRig/Source/ControlRigDeveloper/Private/ControlRigBlueprint.cpp
--- a/Engine/Plugins/Animation/ControlRig/Source/ControlRigDeveloper/Private/ControlRigBlueprint.cpp
+++ b/Engine/Plugins/Animation/ControlRig/Source/ControlRigDeveloper/Private/ControlRigBlueprint.cpp
@@ -3433,7 +3433,7 @@ void UControlRigBlueprint::HandleModifiedEvent(ERigVMGraphNotifType InNotifType,
 							// decide if this graph should be renameable
 							const int32 NumberOfEvents = Algo::CountIf(RigVMNode->GetGraph()->GetNodes(), [](const URigVMNode* NodeToCount) -> bool
 							{
-								return NodeToCount->IsEvent();
+								return NodeToCount->IsEvent() && NodeToCount->CanOnlyExistOnce();
 							});
 							EdGraph->bAllowRenaming = NumberOfEvents != 1;
 						}
diff --git a/Engine/Plugins/Animation/ControlRig/Source/ControlRigDeveloper/Private/Graph/ControlRigGraph.cpp b/Engine/Plugins/Animation/ControlRig/Source/ControlRigDeveloper/Private/Graph/ControlRigGraph.cpp
--- a/Engine/Plugins/Animation/ControlRig/Source/ControlRigDeveloper/Private/Graph/ControlRigGraph.cpp
+++ b/Engine/Plugins/Animation/ControlRig/Source/ControlRigDeveloper/Private/Graph/ControlRigGraph.cpp
@@ -187,6 +187,18 @@ void UControlRigGraph::CacheNameLists(URigHierarchy* InHierarchy, const FControl
 		LastHierarchyTopologyVersion = InHierarchy->GetTopologyVersion();
 	}
 	CacheNameList<FControlRigDrawContainer>(*DrawContainer, DrawingNameList);
+
+	EntryNameList.Reset();
+	EntryNameList.Add(MakeShared<FString>(FName(NAME_None).ToString()));
+
+	if(const UControlRigBlueprint* Blueprint = GetBlueprint())
+	{
+		const TArray<FName> EntryNames = Blueprint->GetRigVMClient()->GetEntryNames();
+		for (const FName& EntryName : EntryNames)
+		{
+			EntryNameList.Add(MakeShared<FString>(EntryName.ToString()));
+		}
+	}
 }
 
 const TArray<TSharedPtr<FString>>* UControlRigGraph::GetElementNameList(ERigElementType InElementType) const
@@ -273,6 +285,15 @@ const TArray<TSharedPtr<FString>>* UControlRigGraph::GetDrawingNameList(URigVMPi
 	return &DrawingNameList;
 }
 
+const TArray<TSharedPtr<FString>>* UControlRigGraph::GetEntryNameList(URigVMPin* InPin) const
+{
+	if (UControlRigGraph* OuterGraph = Cast<UControlRigGraph>(GetOuter()))
+	{
+		return OuterGraph->GetEntryNameList(InPin);
+	}
+	return &EntryNameList;
+}
+
 void UControlRigGraph::HandleModifiedEvent(ERigVMGraphNotifType InNotifType, URigVMGraph* InGraph, UObject* InSubject)
 {
 	DECLARE_SCOPE_HIERARCHICAL_COUNTER_FUNC()
@@ -647,6 +668,10 @@ void UControlRigGraph::HandleModifiedEvent(ERigVMGraphNotifType InNotifType, URi
 					UEdGraphPin* RigNodePin = RigNode->FindPin(ModelPin->GetPinPath());
 					if (RigNodePin == nullptr)
 					{
+						if(ModelPin->GetNode()->IsEvent())
+						{
+							RigNode->InvalidateNodeTitle();
+						}
 						break;
 					}
 
diff --git a/Engine/Plugins/Animation/ControlRig/Source/ControlRigDeveloper/Private/Graph/ControlRigGraphNode.cpp b/Engine/Plugins/Animation/ControlRig/Source/ControlRigDeveloper/Private/Graph/ControlRigGraphNode.cpp
--- a/Engine/Plugins/Animation/ControlRig/Source/ControlRigDeveloper/Private/Graph/ControlRigGraphNode.cpp
+++ b/Engine/Plugins/Animation/ControlRig/Source/ControlRigDeveloper/Private/Graph/ControlRigGraphNode.cpp
@@ -25,11 +25,13 @@
 #include "RigVMCore/RigVMExecuteContext.h"
 #include "ControlRigDeveloper.h"
 #include "ControlRigObjectVersion.h"
+#include "GraphEditAction.h"
 #include "RigVMModel/Nodes/RigVMAggregateNode.h"
 #include "RigVMModel/Nodes/RigVMFunctionReferenceNode.h"
 #include "RigVMModel/Nodes/RigVMFunctionEntryNode.h"
 #include "RigVMModel/Nodes/RigVMFunctionReturnNode.h"
 #include "RigVMModel/Nodes/RigVMCollapseNode.h"
+#include "RigVMModel/Nodes/RigVMInvokeEntryNode.h"
 
 #if WITH_EDITOR
 #include "IControlRigEditorModule.h"
@@ -778,6 +780,7 @@ bool UControlRigGraphNode::ShowPaletteIconOnNode() const
 	if (URigVMNode* ModelNode = GetModelNode())
 	{
 		return ModelNode->IsEvent() ||
+			ModelNode->IsA<URigVMInvokeEntryNode>() ||
 			ModelNode->IsA<URigVMFunctionEntryNode>() ||
 			ModelNode->IsA<URigVMFunctionReturnNode>() ||
 			ModelNode->IsA<URigVMFunctionReferenceNode>() ||
@@ -801,7 +804,7 @@ FSlateIcon UControlRigGraphNode::GetIconAndTint(FLinearColor& OutColor) const
 
 	if (URigVMNode* ModelNode = GetModelNode())
 	{
-		if (ModelNode->IsEvent())
+		if (ModelNode->IsEvent() || ModelNode->IsA<URigVMInvokeEntryNode>())
 		{
 			return EventIcon;
 		}
diff --git a/Engine/Plugins/Animation/ControlRig/Source/ControlRigDeveloper/Private/Graph/ControlRigGraphSchema.cpp b/Engine/Plugins/Animation/ControlRig/Source/ControlRigDeveloper/Private/Graph/ControlRigGraphSchema.cpp
--- a/Engine/Plugins/Animation/ControlRig/Source/ControlRigDeveloper/Private/Graph/ControlRigGraphSchema.cpp
+++ b/Engine/Plugins/Animation/ControlRig/Source/ControlRigDeveloper/Private/Graph/ControlRigGraphSchema.cpp
@@ -1240,7 +1240,7 @@ void UControlRigGraphSchema::GetGraphDisplayInformation(const UEdGraph& Graph, /
 				FString EventName;
 				if(Algo::CountIf(Model->GetNodes(), [&EventName](const URigVMNode* NodeToCount) -> bool
 				{
-					if(NodeToCount->IsEvent())
+					if(NodeToCount->IsEvent() && NodeToCount->CanOnlyExistOnce())
 					{
 						if(EventName.IsEmpty())
 						{
diff --git a/Engine/Plugins/Animation/ControlRig/Source/ControlRigDeveloper/Public/Graph/ControlRigGraph.h b/Engine/Plugins/Animation/ControlRig/Source/ControlRigDeveloper/Public/Graph/ControlRigGraph.h
--- a/Engine/Plugins/Animation/ControlRig/Source/ControlRigDeveloper/Public/Graph/ControlRigGraph.h
+++ b/Engine/Plugins/Animation/ControlRig/Source/ControlRigDeveloper/Public/Graph/ControlRigGraph.h
@@ -67,6 +67,7 @@ public:
 	const TArray<TSharedPtr<FString>>* GetElementNameList(URigVMPin* InPin = nullptr) const;
 	const TArray<TSharedPtr<FString>> GetSelectedElementsNameList() const;
 	const TArray<TSharedPtr<FString>>* GetDrawingNameList(URigVMPin* InPin = nullptr) const;
+	const TArray<TSharedPtr<FString>>* GetEntryNameList(URigVMPin* InPin = nullptr) const;
 
 	bool bSuspendModelNotifications;
 	bool bIsTemporaryGraphForCopyPaste;
@@ -134,6 +135,7 @@ private:
 
 	TMap<ERigElementType, TArray<TSharedPtr<FString>>> ElementNameLists;
 	TArray<TSharedPtr<FString>> DrawingNameList;
+	TArray<TSharedPtr<FString>> EntryNameList;
 	int32 LastHierarchyTopologyVersion;
 
 	bool bIsSelecting;
@@ -173,6 +175,7 @@ private:
 	friend class UControlRigEnumNodeSpawner;
 	friend class UControlRigFunctionRefNodeSpawner;
 	friend class UControlRigArrayNodeSpawner;
+	friend class UControlRigInvokeEntryNodeSpawner;
 
 #endif
 	friend class UControlRigGraphNode;
diff --git a/Engine/Plugins/Animation/ControlRig/Source/ControlRigEditor/Private/ControlRigEditorModule.cpp b/Engine/Plugins/Animation/ControlRig/Source/ControlRigEditor/Private/ControlRigEditorModule.cpp
--- a/Engine/Plugins/Animation/ControlRig/Source/ControlRigEditor/Private/ControlRigEditorModule.cpp
+++ b/Engine/Plugins/Animation/ControlRig/Source/ControlRigEditor/Private/ControlRigEditorModule.cpp
@@ -56,6 +56,7 @@
 #include "Graph/NodeSpawners/ControlRigEnumNodeSpawner.h"
 #include "Graph/NodeSpawners/ControlRigFunctionRefNodeSpawner.h"
 #include "Graph/NodeSpawners/ControlRigArrayNodeSpawner.h"
+#include "Graph/NodeSpawners/ControlRigInvokeEntryNodeSpawner.h"
 #include "Kismet2/BlueprintEditorUtils.h"
 #include "Kismet2/KismetDebugUtilities.h"
 #include "Graph/ControlRigGraphNode.h"
@@ -1159,8 +1160,7 @@ void FControlRigEditorModule::GetInstanceActions(UControlRigBlueprint* CRB, FBlu
 	{
 		if (UControlRig* CDO = Cast<UControlRig>(GeneratedClass->GetDefaultObject()))
 		{
-			static const FString CategoryDelimiter(TEXT("|"));
-			FText NodeCategory = LOCTEXT("Variables", "Variables");
+			static const FText NodeCategory = LOCTEXT("Variables", "Variables");
 
 			TArray<FRigVMExternalVariable> ExternalVariables = CDO->GetExternalVariables();
 			for (const FRigVMExternalVariable& ExternalVariable : ExternalVariables)
@@ -1184,8 +1184,7 @@ void FControlRigEditorModule::GetInstanceActions(UControlRigBlueprint* CRB, FBlu
 				ActionRegistrar.AddBlueprintAction(GeneratedClass, NodeSpawner);
 			}
 
-			static const FString CategoryDelimiter(TEXT("|"));
-			FText NodeCategory = LOCTEXT("LocalVariables", "Local Variables");
+			static const FText NodeCategory = LOCTEXT("LocalVariables", "Local Variables");
 			for (URigVMLibraryNode* Function : Functions)
 			{
 				for (const FRigVMGraphVariableDescription& LocalVariable : Function->GetContainedGraph()->GetLocalVariables())
@@ -1204,7 +1203,7 @@ void FControlRigEditorModule::GetInstanceActions(UControlRigBlueprint* CRB, FBlu
 		{
 			if (Graph->GetEntryNode())
 			{
-				FText NodeCategory = LOCTEXT("InputArguments", "Input Arguments");
+				static const FText NodeCategory = LOCTEXT("InputArguments", "Input Arguments");
 				for (const FRigVMGraphVariableDescription& InputArgument : Graph->GetInputArguments())
 				{
 					FText MenuDesc = FText::FromName(InputArgument.Name);
@@ -1213,6 +1212,21 @@ void FControlRigEditorModule::GetInstanceActions(UControlRigBlueprint* CRB, FBlu
 				}			
 			}
 		}
+
+		const TArray<FName> EntryNames = CRB->GetRigVMClient()->GetEntryNames();
+		if(!EntryNames.IsEmpty())
+		{
+			static const FText NodeCategory = LOCTEXT("Events", "Events");
+			for (const FName& EntryName : EntryNames)
+			{
+				static constexpr TCHAR EventStr[] = TEXT("Event");
+				static const FString EventSuffix = FString::Printf(TEXT(" %s"), EventStr);
+				FString Suffix = EntryName.ToString().EndsWith(EventStr) ? FString() : EventSuffix;
+				FText MenuDesc = FText::FromString(FString::Printf(TEXT("Run %s%s"), *EntryName.ToString(), *Suffix));
+				FText ToolTip = FText::FromString(FString::Printf(TEXT("Runs the %s%s"), *EntryName.ToString(), *Suffix));
+				ActionRegistrar.AddBlueprintAction(GeneratedClass, UControlRigInvokeEntryNodeSpawner::CreateForEntry(CRB, EntryName, MenuDesc, NodeCategory, ToolTip));
+			}
+		}
 	}
 }
 
diff --git a/Engine/Plugins/Animation/ControlRig/Source/ControlRigEditor/Private/Editor/ControlRigEditor.cpp b/Engine/Plugins/Animation/ControlRig/Source/ControlRigEditor/Private/Editor/ControlRigEditor.cpp
--- a/Engine/Plugins/Animation/ControlRig/Source/ControlRigEditor/Private/Editor/ControlRigEditor.cpp
+++ b/Engine/Plugins/Animation/ControlRig/Source/ControlRigEditor/Private/Editor/ControlRigEditor.cpp
@@ -2820,6 +2820,15 @@ void FControlRigEditor::HandleModifiedEvent(ERigVMGraphNotifType InNotifType, UR
 						}
 					}
 				}
+
+				if(const URigVMUnitNode* UnitNode = Cast<URigVMUnitNode>(RootPin->GetNode()))
+				{
+					if(UnitNode->IsEvent())
+					{
+						FBlueprintEditorUtils::MarkBlueprintAsStructurallyModified(GetBlueprintObj());
+						CacheNameLists();
+					}
+				}
 			}
 	
 			break;
@@ -2852,6 +2861,20 @@ void FControlRigEditor::HandleModifiedEvent(ERigVMGraphNotifType InNotifType, UR
 			{
 				ClearDetailObject();
 			}
+				
+			// fall through next case since we want to refresh the name lists
+			// both for removing or adding an event
+		}
+		case ERigVMGraphNotifType::NodeAdded:
+		{
+			if(URigVMUnitNode* UnitNode = Cast<URigVMUnitNode>(InSubject))
+			{
+				if(UnitNode->IsEvent())
+				{
+					FBlueprintEditorUtils::MarkBlueprintAsStructurallyModified(GetBlueprintObj());
+					CacheNameLists();
+				}
+			}
 			break;
 		}
 		case ERigVMGraphNotifType::NodeSelectionChanged:
@@ -4695,6 +4718,12 @@ void FControlRigEditor::OnWrappedPropertyChangedChainEvent(UDetailsViewWrapperOb
 		URigVMNode* Node = CastChecked<URigVMNode>(InWrapperObject->GetOuter());
 
 		FString DefaultValue = FRigVMStruct::ExportToFullyQualifiedText(TargetProperty, FirstPropertyStorage);
+
+		if(TargetProperty->IsA<FStrProperty>() || TargetProperty->IsA<FNameProperty>())
+		{
+			DefaultValue.TrimCharInline(TEXT('\"'), nullptr);
+		}
+		
 		URigVMController* Controller = GetControlRigBlueprint()->GetController(Node->GetGraph());
 
 		if (!DefaultValue.IsEmpty())
diff --git a/Engine/Plugins/Animation/ControlRig/Source/ControlRigEditor/Private/Editor/SControlRigStackView.cpp b/Engine/Plugins/Animation/ControlRig/Source/ControlRigEditor/Private/Editor/SControlRigStackView.cpp
--- a/Engine/Plugins/Animation/ControlRig/Source/ControlRigEditor/Private/Editor/SControlRigStackView.cpp
+++ b/Engine/Plugins/Animation/ControlRig/Source/ControlRigEditor/Private/Editor/SControlRigStackView.cpp
@@ -478,6 +478,15 @@ void SControlRigStackView::PopulateStackView(URigVM* InVM)
 						{
 							DisplayName = FString::Printf(TEXT("%s %s"), *UnitNode->GetNodeTitle(), *MenuDescSuffixMetadata);
 						}
+
+						if(UnitNode->IsEvent())
+						{
+							DisplayName = Node->GetEventName().ToString();
+							if(!DisplayName.EndsWith(TEXT("Event")))
+							{
+								DisplayName += TEXT(" Event");
+							}
+						}
 	#endif
 					}
 				}
@@ -609,6 +618,12 @@ void SControlRigStackView::PopulateStackView(URigVM* InVM)
 						}
 						break;
 					}
+					case ERigVMOpCode::InvokeEntry:
+					{
+						const FRigVMInvokeEntryOp& Op = ByteCode.GetOpAt<FRigVMInvokeEntryOp>(Instructions[InstructionIndex]);
+						Label = FString::Printf(TEXT("Run %s Event"), *Op.EntryName.ToString());
+						break;
+					}
 					case ERigVMOpCode::Exit:
 					{
 						Label = TEXT("Exit");
diff --git a/Engine/Plugins/Animation/ControlRig/Source/ControlRigEditor/Private/Graph/ControlRigGraphPanelPinFactory.cpp b/Engine/Plugins/Animation/ControlRig/Source/ControlRigEditor/Private/Graph/ControlRigGraphPanelPinFactory.cpp
--- a/Engine/Plugins/Animation/ControlRig/Source/ControlRigEditor/Private/Graph/ControlRigGraphPanelPinFactory.cpp
+++ b/Engine/Plugins/Animation/ControlRig/Source/ControlRigEditor/Private/Graph/ControlRigGraphPanelPinFactory.cpp
@@ -83,6 +83,12 @@ TSharedPtr<SGraphPin> FControlRigGraphPanelPinFactory::CreatePin(UEdGraphPin* In
 						.OnGetNameFromSelection_UObject(RigGraph, &UControlRigGraph::GetSelectedElementsNameList)
 						.OnGetNameListContent_UObject(RigGraph, &UControlRigGraph::GetElementNameList);
 				}
+				else if (CustomWidgetName == TEXT("EntryName"))
+				{
+					return SNew(SControlRigGraphPinNameList, InPin)
+						.ModelPin(ModelPin)
+						.OnGetNameListContent_UObject(RigGraph, &UControlRigGraph::GetEntryNameList);
+				}
 				else if (CustomWidgetName == TEXT("DrawingName"))
 				{
 					return SNew(SControlRigGraphPinNameList, InPin)
diff --git a/Engine/Plugins/Animation/ControlRig/Source/ControlRigEditor/Private/Graph/NodeSpawners/ControlRigInvokeEntryNodeSpawner.cpp b/Engine/Plugins/Animation/ControlRig/Source/ControlRigEditor/Private/Graph/NodeSpawners/ControlRigInvokeEntryNodeSpawner.cpp
new file mode 100644
--- /dev/null
+++ b/Engine/Plugins/Animation/ControlRig/Source/ControlRigEditor/Private/Graph/NodeSpawners/ControlRigInvokeEntryNodeSpawner.cpp
@@ -0,0 +1,140 @@
+// Copyright Epic Games, Inc. All Rights Reserved.
+
+#include "ControlRigInvokeEntryNodeSpawner.h"
+#include "Graph/ControlRigGraph.h"
+#include "Graph/ControlRigGraphNode.h"
+#include "EdGraphSchema_K2.h"
+#include "Kismet2/BlueprintEditorUtils.h"
+#include "Settings/EditorStyleSettings.h"
+#include "Editor/EditorEngine.h"
+#include "ObjectEditorUtils.h"
+#include "EditorCategoryUtils.h"
+#include "BlueprintNodeTemplateCache.h"
+#include "RigVMModel/RigVMController.h"
+#include "ControlRigBlueprint.h"
+#include "ControlRigBlueprintUtils.h"
+#include "RigVMModel/Nodes/RigVMInvokeEntryNode.h"
+
+#if WITH_EDITOR
+#include "Editor.h"
+#endif
+
+#define LOCTEXT_NAMESPACE "ControlRigInvokeEntryNodeSpawner"
+
+UControlRigInvokeEntryNodeSpawner* UControlRigInvokeEntryNodeSpawner::CreateForEntry(UControlRigBlueprint* InBlueprint, const FName& InEntryName, const FText& InMenuDesc, const FText& InCategory, const FText& InTooltip)
+{
+	UControlRigInvokeEntryNodeSpawner* NodeSpawner = NewObject<UControlRigInvokeEntryNodeSpawner>(GetTransientPackage());
+	NodeSpawner->Blueprint = InBlueprint;
+	NodeSpawner->EntryName = InEntryName;
+	NodeSpawner->NodeClass = UControlRigGraphNode::StaticClass();
+
+	FBlueprintActionUiSpec& MenuSignature = NodeSpawner->DefaultMenuSignature;
+	
+	MenuSignature.MenuName = InMenuDesc;
+	MenuSignature.Tooltip  = InTooltip;
+	MenuSignature.Category = InCategory;
+	MenuSignature.Keywords = FText::FromString(TEXT("Event,Entry,Invoke,Run,Launch,Start"));
+	MenuSignature.Icon = FSlateIcon(FAppStyle::GetAppStyleSetName(), "GraphEditor.Event_16x");
+
+	return NodeSpawner;
+}
+
+void UControlRigInvokeEntryNodeSpawner::Prime()
+{
+	// we expect that you don't need a node template to construct menu entries
+	// from this, so we choose not to pre-cache one here
+}
+
+FBlueprintNodeSignature UControlRigInvokeEntryNodeSpawner::GetSpawnerSignature() const
+{
+	return FBlueprintNodeSignature(FString("InvokeEntryNode=" + EntryName.ToString()));
+}
+
+FBlueprintActionUiSpec UControlRigInvokeEntryNodeSpawner::GetUiSpec(FBlueprintActionContext const& Context, FBindingSet const& Bindings) const
+{
+	UEdGraph* TargetGraph = (Context.Graphs.Num() > 0) ? Context.Graphs[0] : nullptr;
+	FBlueprintActionUiSpec MenuSignature = PrimeDefaultUiSpec(TargetGraph);
+
+	DynamicUiSignatureGetter.ExecuteIfBound(Context, Bindings, &MenuSignature);
+	return MenuSignature;
+}
+
+UEdGraphNode* UControlRigInvokeEntryNodeSpawner::Invoke(UEdGraph* ParentGraph, FBindingSet const& Bindings, FVector2D const Location) const
+{
+	UControlRigGraphNode* NewNode = nullptr;
+
+	bool const bIsTemplateNode = FBlueprintNodeTemplateCache::IsTemplateOuter(ParentGraph);
+	bool const bIsUserFacingNode = !bIsTemplateNode;
+
+	// First create a backing member for our node
+	UControlRigGraph* RigGraph = Cast<UControlRigGraph>(ParentGraph);
+	if(RigGraph == nullptr) return nullptr;
+	UControlRigBlueprint* RigBlueprint = Cast<UControlRigBlueprint>(FBlueprintEditorUtils::FindBlueprintForGraph(ParentGraph));
+	check(RigBlueprint);
+
+#if WITH_EDITOR
+	if (GEditor && !bIsTemplateNode)
+	{
+		GEditor->CancelTransaction(0);
+	}
+#endif
+
+	URigVMController* Controller = bIsTemplateNode ? RigGraph->GetTemplateController() : RigBlueprint->GetController(ParentGraph);
+
+	if (bIsUserFacingNode)
+	{
+		Controller->OpenUndoBracket(TEXT("Add Run Event"));
+	}
+
+	FString NodeName;
+	if(bIsTemplateNode)
+	{
+		// since we are removing the node at the end of this function
+		// we need to create a unique here.
+		static constexpr TCHAR InvokeEntryNodeNameFormat[] = TEXT("InvokeEntryNode_%s");
+		NodeName = FString::Printf(InvokeEntryNodeNameFormat, *EntryName.ToString());
+	}
+
+	if (URigVMNode* ModelNode = Controller->AddInvokeEntryNode(EntryName, Location, NodeName, !bIsTemplateNode, !bIsTemplateNode))
+	{
+		for (UEdGraphNode* Node : ParentGraph->Nodes)
+		{
+			if (UControlRigGraphNode* RigNode = Cast<UControlRigGraphNode>(Node))
+			{
+				if (RigNode->GetModelNodeName() == ModelNode->GetFName())
+				{
+					NewNode = RigNode;
+					break;
+				}
+			}
+		}
+
+		if (bIsUserFacingNode)
+		{
+			if (NewNode)
+			{
+				Controller->ClearNodeSelection(true);
+				Controller->SelectNode(ModelNode, true, true);
+			}
+			Controller->CloseUndoBracket();
+		}
+		else
+		{
+			// similar to UBlueprintNodeSpawner::Invoke -> UBlueprintNodeSpawner::SpawnEdGraphNode
+			// we simply want the node, but not actually adding it to a graph
+			Controller->RemoveNode(ModelNode, false);
+		}
+	}
+	else
+	{
+		if (!bIsTemplateNode)
+		{
+			Controller->CancelUndoBracket();
+		}
+	}
+
+
+	return NewNode;
+}
+
+#undef LOCTEXT_NAMESPACE
diff --git a/Engine/Plugins/Animation/ControlRig/Source/ControlRigEditor/Private/Graph/NodeSpawners/ControlRigInvokeEntryNodeSpawner.h b/Engine/Plugins/Animation/ControlRig/Source/ControlRigEditor/Private/Graph/NodeSpawners/ControlRigInvokeEntryNodeSpawner.h
new file mode 100644
--- /dev/null
+++ b/Engine/Plugins/Animation/ControlRig/Source/ControlRigEditor/Private/Graph/NodeSpawners/ControlRigInvokeEntryNodeSpawner.h
@@ -0,0 +1,39 @@
+// Copyright Epic Games, Inc. All Rights Reserved.
+
+#pragma once
+
+#include "CoreMinimal.h"
+#include "BlueprintActionFilter.h"
+#include "BlueprintNodeSignature.h"
+#include "BlueprintFieldNodeSpawner.h"
+#include "RigVMModel/RigVMGraph.h"
+#include "ControlRigInvokeEntryNodeSpawner.generated.h"
+
+class UControlRigGraphNode;
+class UControlRigBlueprint;
+
+UCLASS(Transient)
+class CONTROLRIGEDITOR_API UControlRigInvokeEntryNodeSpawner : public UBlueprintNodeSpawner
+{
+	GENERATED_BODY()
+
+public:
+
+	static UControlRigInvokeEntryNodeSpawner* CreateForEntry(UControlRigBlueprint* InBlueprint, const FName& InEntryName, const FText& InMenuDesc, const FText& InCategory, const FText& InTooltip);
+
+	// UBlueprintNodeSpawner interface
+	virtual void Prime() override;
+	virtual FBlueprintNodeSignature GetSpawnerSignature() const override;
+	virtual FBlueprintActionUiSpec GetUiSpec(FBlueprintActionContext const& Context, FBindingSet const& Bindings) const override;
+	virtual UEdGraphNode* Invoke(UEdGraph* ParentGraph, FBindingSet const& Bindings, FVector2D const Location) const override;
+	// End UBlueprintNodeSpawner interface
+
+private:
+
+	/** The pin type we will spawn */
+	TWeakObjectPtr<UControlRigBlueprint> Blueprint;
+	TWeakObjectPtr<URigVMGraph> GraphOwner;
+	FName EntryName;
+
+	friend class UEngineTestControlRig;
+};
diff --git a/Engine/Plugins/Animation/ControlRig/Source/ControlRigEditor/Private/Graph/NodeSpawners/ControlRigUnitNodeSpawner.cpp b/Engine/Plugins/Animation/ControlRig/Source/ControlRigEditor/Private/Graph/NodeSpawners/ControlRigUnitNodeSpawner.cpp
--- a/Engine/Plugins/Animation/ControlRig/Source/ControlRigEditor/Private/Graph/NodeSpawners/ControlRigUnitNodeSpawner.cpp
+++ b/Engine/Plugins/Animation/ControlRig/Source/ControlRigEditor/Private/Graph/NodeSpawners/ControlRigUnitNodeSpawner.cpp
@@ -17,6 +17,11 @@
 #include "ControlRig/Private/Units/Execution/RigUnit_BeginExecution.h"
 #include "ControlRig.h"
 #include "Settings/ControlRigSettings.h"
+#include "ControlRig/Private/Units/Execution/RigUnit_BeginExecution.h"
+#include "ControlRig/Private/Units/Execution/RigUnit_InverseExecution.h"
+#include "ControlRig/Private/Units/Execution/RigUnit_PrepareForExecution.h"
+#include "ControlRig/Private/Units/Execution/RigUnit_InteractionExecution.h"
+#include "ControlRig/Private/Units/Execution/RigUnit_UserDefinedEvent.h"
 
 #if WITH_EDITOR
 #include "Editor.h"
@@ -63,7 +68,52 @@ UControlRigUnitNodeSpawner* UControlRigUnitNodeSpawner::CreateFromStruct(UScript
 	}
 
 	// @TODO: should use details customization-like extensibility system to provide editor only data like this
-	MenuSignature.Icon = FSlateIcon(TEXT("ControlRigEditorStyle"), TEXT("ControlRig.RigUnit"));
+	FStructOnScope StructScope(InStruct);
+	const FRigVMStruct* StructInstance = (const FRigVMStruct*)StructScope.GetStructMemory();
+	if(StructInstance->GetEventName().IsNone())
+	{
+		if(InStruct->HasMetaDataHierarchical(FRigVMStruct::IconMetaName))
+		{
+			FString IconPath;
+			const int32 NumOfIconPathNames = 4;
+			
+			FName IconPathNames[NumOfIconPathNames] = {
+				NAME_None, // StyleSetName
+				NAME_None, // StyleName
+				NAME_None, // SmallStyleName
+				NAME_None  // StatusOverlayStyleName
+			};
+
+			InStruct->GetStringMetaDataHierarchical(FRigVMStruct::IconMetaName, &IconPath);
+
+			int32 NameIndex = 0;
+
+			while (!IconPath.IsEmpty() && NameIndex < NumOfIconPathNames)
+			{
+				FString Left;
+				FString Right;
+
+				if (!IconPath.Split(TEXT("|"), &Left, &Right))
+				{
+					Left = IconPath;
+				}
+
+				IconPathNames[NameIndex] = FName(*Left);
+
+				NameIndex++;
+				IconPath = Right;
+			}
+			MenuSignature.Icon = FSlateIcon(IconPathNames[0], IconPathNames[1], IconPathNames[2], IconPathNames[3]);
+		}
+		else
+		{
+			MenuSignature.Icon = FSlateIcon(TEXT("ControlRigEditorStyle"), TEXT("ControlRig.RigUnit"));
+		}
+	}
+	else
+	{
+		MenuSignature.Icon = FSlateIcon(FAppStyle::GetAppStyleSetName(), "GraphEditor.Event_16x");
+	}
 
 	return NodeSpawner;
 }
@@ -121,25 +171,30 @@ UControlRigGraphNode* UControlRigUnitNodeSpawner::SpawnNode(UEdGraph* ParentGrap
 
 		FName Name = bIsTemplateNode ? *StructTemplate->GetStructCPPName() : FControlRigBlueprintUtils::ValidateName(RigBlueprint, StructTemplate->GetFName().ToString());
 
+		FStructOnScope StructScope(StructTemplate);
+		const FRigVMStruct* StructInstance = (const FRigVMStruct*)StructScope.GetStructMemory();
+		const FName EventName = StructInstance->GetEventName();
+
 		// events can only exist once across all uber graphs
 		if(bIsUserFacingNode)
 		{
-			FStructOnScope StructScope(StructTemplate);
-			const FRigVMStruct* StructInstance = (const FRigVMStruct*)StructScope.GetStructMemory();
-			const FName EventName = StructInstance->GetEventName();
 			if(!EventName.IsNone())
 			{
-				const TArray<URigVMGraph*> Models = RigBlueprint->GetAllModels();
-				for(URigVMGraph* Model : Models)
+				if(StructInstance->CanOnlyExistOnce() &&
+					(StructTemplate != FRigUnit_UserDefinedEvent::StaticStruct()))
 				{
-					if(Model->IsRootGraph() && !Model->IsA<URigVMFunctionLibrary>())
+					const TArray<URigVMGraph*> Models = RigBlueprint->GetAllModels();
+					for(URigVMGraph* Model : Models)
 					{
-						for(const URigVMNode* Node : Model->GetNodes())
+						if(Model->IsRootGraph() && !Model->IsA<URigVMFunctionLibrary>())
 						{
-							if(Node->GetEventName() == EventName)
+							for(const URigVMNode* Node : Model->GetNodes())
 							{
-								RigBlueprint->OnRequestJumpToHyperlink().Execute(Node);
-								return nullptr;
+								if(Node->GetEventName() == EventName)
+								{
+									RigBlueprint->OnRequestJumpToHyperlink().Execute(Node);
+									return nullptr;
+								}
 							}
 						}
 					}
@@ -164,6 +219,40 @@ UControlRigGraphNode* UControlRigUnitNodeSpawner::SpawnNode(UEdGraph* ParentGrap
 
 			if (NewNode && bIsUserFacingNode)
 			{
+				if(StructTemplate == FRigUnit_UserDefinedEvent::StaticStruct())
+				{
+					// ensure uniqueness for the event name
+					TArray<FName> ExistingEventNames = {
+						FRigUnit_BeginExecution::EventName,
+						FRigUnit_InverseExecution::EventName,
+						FRigUnit_PrepareForExecution::EventName,
+						FRigUnit_InteractionExecution::EventName
+					};
+
+					const TArray<URigVMGraph*> Models = RigBlueprint->GetAllModels();
+					for(URigVMGraph* Model : Models)
+					{
+						for(const URigVMNode* Node : Model->GetNodes())
+						{
+							if(Node != ModelNode)
+							{
+								const FName ExistingEventName = Node->GetEventName();
+								if(!ExistingEventName.IsNone())
+								{
+									ExistingEventNames.AddUnique(ExistingEventName);
+								}
+							}
+						}
+					}
+
+					FName SafeEventName = URigVMController::GetUniqueName(EventName, [ExistingEventNames](const FName& NameToCheck) -> bool
+					{
+						return !ExistingEventNames.Contains(NameToCheck);
+					}, false, true);
+
+					Controller->SetPinDefaultValue(ModelNode->FindPin(TEXT("EventName"))->GetPinPath(), SafeEventName.ToString(), false, true, false, true);
+				}
+
 				Controller->ClearNodeSelection(true);
 				Controller->SelectNode(ModelNode, true, true);
 
diff --git a/Engine/Plugins/Animation/ControlRig/Source/ControlRigEditor/Private/Graph/SControlRigGraphPinNameList.cpp b/Engine/Plugins/Animation/ControlRig/Source/ControlRigEditor/Private/Graph/SControlRigGraphPinNameList.cpp
--- a/Engine/Plugins/Animation/ControlRig/Source/ControlRigEditor/Private/Graph/SControlRigGraphPinNameList.cpp
+++ b/Engine/Plugins/Animation/ControlRig/Source/ControlRigEditor/Private/Graph/SControlRigGraphPinNameList.cpp
@@ -81,6 +81,7 @@ TSharedRef<SWidget>	SControlRigGraphPinNameList::GetDefaultValueWidget()
 				.OnClicked(this, &SControlRigGraphPinNameList::OnGetSelectedClicked)
 				.ContentPadding(1.f)
 				.ToolTipText(NSLOCTEXT("ControlRigElementDetails", "ObjectGraphPin_Use_Tooltip", "Use item selected"))
+				.Visibility(OnGetNameFromSelection.IsBound() ? EVisibility::Visible : EVisibility::Collapsed)
 				[
 					SNew(SImage)
 					.ColorAndOpacity( this, &SControlRigGraphPinNameList::OnGetWidgetForeground )
@@ -100,6 +101,7 @@ TSharedRef<SWidget>	SControlRigGraphPinNameList::GetDefaultValueWidget()
 				.OnClicked(this, &SControlRigGraphPinNameList::OnBrowseClicked)
 				.ContentPadding(0)
 				.ToolTipText(NSLOCTEXT("SControlRigGraphPinNameList", "ObjectGraphPin_Browse_Tooltip", "Browse"))
+				.Visibility(OnGetNameFromSelection.IsBound() ? EVisibility::Visible : EVisibility::Collapsed)
 				[
 					SNew(SImage)
 					.ColorAndOpacity( this, &SControlRigGraphPinNameList::OnGetWidgetForeground )
diff --git a/Engine/Source/Developer/RigVMDeveloper/Private/RigVMCompiler/RigVMCompiler.cpp b/Engine/Source/Developer/RigVMDeveloper/Private/RigVMCompiler/RigVMCompiler.cpp
--- a/Engine/Source/Developer/RigVMDeveloper/Private/RigVMCompiler/RigVMCompiler.cpp
+++ b/Engine/Source/Developer/RigVMDeveloper/Private/RigVMCompiler/RigVMCompiler.cpp
@@ -1753,7 +1753,14 @@ void URigVMCompiler::TraverseInvokeEntry(const FRigVMInvokeEntryExprAST* InExpr,
 	}
 	else
 	{
+		const int32 InstructionIndex = WorkData.VM->GetByteCode().GetNumInstructions();
 		WorkData.VM->GetByteCode().AddInvokeEntryOp(InvokeEntryNode->GetEntryName());
+
+		if (Settings.SetupNodeInstructionIndex)
+		{
+			const FRigVMCallstack Callstack = InExpr->GetProxy().GetSibling(InvokeEntryNode).GetCallstack();
+			WorkData.VM->GetByteCode().SetSubject(InstructionIndex, Callstack.GetCallPath(), Callstack.GetStack());
+		}
 	}
 }
 
diff --git a/Engine/Source/Developer/RigVMDeveloper/Private/RigVMModel/Nodes/RigVMInvokeEntryNode.cpp b/Engine/Source/Developer/RigVMDeveloper/Private/RigVMModel/Nodes/RigVMInvokeEntryNode.cpp
--- a/Engine/Source/Developer/RigVMDeveloper/Private/RigVMModel/Nodes/RigVMInvokeEntryNode.cpp
+++ b/Engine/Source/Developer/RigVMDeveloper/Private/RigVMModel/Nodes/RigVMInvokeEntryNode.cpp
@@ -12,7 +12,7 @@ URigVMInvokeEntryNode::URigVMInvokeEntryNode()
 
 FString URigVMInvokeEntryNode::GetNodeTitle() const
 {
-	return FString::Printf(TEXT("Invoke %s"), *GetEntryName().ToString());
+	return FString::Printf(TEXT("Run %s"), *GetEntryName().ToString());
 }
 
 FName URigVMInvokeEntryNode::GetEntryName() const
diff --git a/Engine/Source/Developer/RigVMDeveloper/Private/RigVMModel/Nodes/RigVMUnitNode.cpp b/Engine/Source/Developer/RigVMDeveloper/Private/RigVMModel/Nodes/RigVMUnitNode.cpp
--- a/Engine/Source/Developer/RigVMDeveloper/Private/RigVMModel/Nodes/RigVMUnitNode.cpp
+++ b/Engine/Source/Developer/RigVMDeveloper/Private/RigVMModel/Nodes/RigVMUnitNode.cpp
@@ -64,13 +64,24 @@ bool URigVMUnitNode::IsDefinedAsVarying() const
 
 FName URigVMUnitNode::GetEventName() const
 {
-	TSharedPtr<FStructOnScope> StructOnScope = ConstructStructInstance(true);
+	TSharedPtr<FStructOnScope> StructOnScope = ConstructStructInstance(false);
 	if (StructOnScope.IsValid())
 	{
 		const FRigVMStruct* StructMemory = (FRigVMStruct*)StructOnScope->GetStructMemory();
 		return StructMemory->GetEventName();
 	}
-	return NAME_None;
+	return Super::GetEventName();
+}
+
+bool URigVMUnitNode::CanOnlyExistOnce() const
+{
+	TSharedPtr<FStructOnScope> StructOnScope = ConstructStructInstance(false);
+	if (StructOnScope.IsValid())
+	{
+		const FRigVMStruct* StructMemory = (FRigVMStruct*)StructOnScope->GetStructMemory();
+		return StructMemory->CanOnlyExistOnce();
+	}
+	return Super::CanOnlyExistOnce();
 }
 
 FText URigVMUnitNode::GetToolTipTextForPin(const URigVMPin* InPin) const
diff --git a/Engine/Source/Developer/RigVMDeveloper/Private/RigVMModel/RigVMClient.cpp b/Engine/Source/Developer/RigVMDeveloper/Private/RigVMModel/RigVMClient.cpp
--- a/Engine/Source/Developer/RigVMDeveloper/Private/RigVMModel/RigVMClient.cpp
+++ b/Engine/Source/Developer/RigVMDeveloper/Private/RigVMModel/RigVMClient.cpp
@@ -322,6 +322,23 @@ URigVMFunctionLibrary* FRigVMClient::GetOrCreateFunctionLibrary(bool bSetupUndoR
 	return NewFunctionLibrary;
 }
 
+TArray<FName> FRigVMClient::GetEntryNames() const
+{
+	TArray<FName> EntryNames;
+	for(const URigVMGraph* Model : Models)
+	{
+		for(const URigVMNode* Node : Model->GetNodes())
+		{
+			const FName EntryName = Node->GetEventName();
+			if(!EntryName.IsNone())
+			{
+				EntryNames.Add(EntryName);
+			}
+		}
+	}
+	return EntryNames;
+}
+
 bool FRigVMClient::RemoveModel(const FString& InNodePathOrName, bool bSetupUndoRedo)
 {
 	if(URigVMGraph* Model = GetModel(InNodePathOrName))
diff --git a/Engine/Source/Developer/RigVMDeveloper/Private/RigVMModel/RigVMController.cpp b/Engine/Source/Developer/RigVMDeveloper/Private/RigVMModel/RigVMController.cpp
--- a/Engine/Source/Developer/RigVMDeveloper/Private/RigVMModel/RigVMController.cpp
+++ b/Engine/Source/Developer/RigVMDeveloper/Private/RigVMModel/RigVMController.cpp
@@ -660,7 +660,7 @@ TArray<FString> URigVMController::GetAddNodePythonCommands(URigVMNode* Node) con
 		Commands.Add(FString::Printf(TEXT("blueprint.get_controller_by_name('%s').add_invoke_entry_node('%s', %s, '%s')"),
 						*GraphName,
 						*InvokeEntryNode->GetEntryName().ToString(),
-						*RigVMPythonUtils::Vector2DToPythonString(EnumNode->GetPosition()),
+						*RigVMPythonUtils::Vector2DToPythonString(InvokeEntryNode->GetPosition()),
 						*NodeName));
 	}
 	else if (Node->IsA<URigVMFunctionEntryNode>() || Node->IsA<URigVMFunctionReturnNode>())
@@ -835,14 +835,17 @@ URigVMUnitNode* URigVMController::AddUnitNode(UScriptStruct* InScriptStruct, con
 			return nullptr;
 		}
 
-		// don't allow several event nodes in the main graph
-		TObjectPtr<URigVMNode> EventNode = FindEventNode(InScriptStruct);
-		if (EventNode != nullptr)
+		if(StructMemory->CanOnlyExistOnce())
 		{
-			const FString ErrorMessage = FString::Printf(TEXT("Rig Graph can only contain one single %s node."),
-															*InScriptStruct->GetDisplayNameText().ToString());
-			ReportAndNotifyError(ErrorMessage);
-			return Cast<URigVMUnitNode>(EventNode);
+			// don't allow several event nodes in the main graph
+			TObjectPtr<URigVMNode> EventNode = FindEventNode(InScriptStruct);
+			if (EventNode != nullptr)
+			{
+				const FString ErrorMessage = FString::Printf(TEXT("Rig Graph can only contain one single %s node."),
+																*InScriptStruct->GetDisplayNameText().ToString());
+				ReportAndNotifyError(ErrorMessage);
+				return Cast<URigVMUnitNode>(EventNode);
+			}
 		}
 	}
 	
diff --git a/Engine/Source/Developer/RigVMDeveloper/Private/RigVMModel/RigVMNode.cpp b/Engine/Source/Developer/RigVMDeveloper/Private/RigVMModel/RigVMNode.cpp
--- a/Engine/Source/Developer/RigVMDeveloper/Private/RigVMModel/RigVMNode.cpp
+++ b/Engine/Source/Developer/RigVMDeveloper/Private/RigVMModel/RigVMNode.cpp
@@ -308,7 +308,7 @@ bool URigVMNode::HasWildCardPin() const
 
 bool URigVMNode::IsEvent() const
 {
-	return IsMutable() && !HasInputPin(true /* include io */) && !GetEventName().IsNone();
+	return IsMutable() && !GetEventName().IsNone();
 }
 
 FName URigVMNode::GetEventName() const
@@ -316,6 +316,11 @@ FName URigVMNode::GetEventName() const
 	return NAME_None;
 }
 
+bool URigVMNode::CanOnlyExistOnce() const
+{
+	return false;
+}
+
 bool URigVMNode::HasInputPin(bool bIncludeIO) const
 {
 	if (HasPinOfDirection(ERigVMPinDirection::Input))
diff --git a/Engine/Source/Developer/RigVMDeveloper/Private/RigVMModel/RigVMPin.cpp b/Engine/Source/Developer/RigVMDeveloper/Private/RigVMModel/RigVMPin.cpp
--- a/Engine/Source/Developer/RigVMDeveloper/Private/RigVMModel/RigVMPin.cpp
+++ b/Engine/Source/Developer/RigVMDeveloper/Private/RigVMModel/RigVMPin.cpp
@@ -1219,13 +1219,6 @@ bool URigVMPin::ShowInDetailsPanelOnly() const
 			}
 		}
 	}
-	else if(URigVMInvokeEntryNode* InvokeEntryNode = Cast<URigVMInvokeEntryNode>(GetNode()))
-	{
-		if(GetName() == URigVMInvokeEntryNode::EntryName)
-		{
-			return true;
-		}
-	}
 #endif
 	return false;
 }
diff --git a/Engine/Source/Developer/RigVMDeveloper/Public/RigVMModel/Nodes/RigVMInvokeEntryNode.h b/Engine/Source/Developer/RigVMDeveloper/Public/RigVMModel/Nodes/RigVMInvokeEntryNode.h
--- a/Engine/Source/Developer/RigVMDeveloper/Public/RigVMModel/Nodes/RigVMInvokeEntryNode.h
+++ b/Engine/Source/Developer/RigVMDeveloper/Public/RigVMModel/Nodes/RigVMInvokeEntryNode.h
@@ -27,7 +27,7 @@ public:
 	FName GetEntryName() const;
 
 	// Override of node title
-	virtual FLinearColor GetNodeColor() const override { return FLinearColor::Blue; }
+	virtual FLinearColor GetNodeColor() const override { return FLinearColor::Red; }
 	virtual bool IsDefinedAsVarying() const override { return true; }
 
 private:
diff --git a/Engine/Source/Developer/RigVMDeveloper/Public/RigVMModel/Nodes/RigVMUnitNode.h b/Engine/Source/Developer/RigVMDeveloper/Public/RigVMModel/Nodes/RigVMUnitNode.h
--- a/Engine/Source/Developer/RigVMDeveloper/Public/RigVMModel/Nodes/RigVMUnitNode.h
+++ b/Engine/Source/Developer/RigVMDeveloper/Public/RigVMModel/Nodes/RigVMUnitNode.h
@@ -28,6 +28,7 @@ public:
 	virtual bool IsDefinedAsConstant() const override;
 	virtual bool IsDefinedAsVarying() const override;
 	virtual FName GetEventName() const override;
+	virtual bool CanOnlyExistOnce() const override;
 	virtual bool IsLoopNode() const override;
 	virtual bool CanBeUpgraded() const override { return GetUpgradeInfo().IsValid(); }
 	virtual TArray<FRigVMUserWorkflow> GetSupportedWorkflows(ERigVMUserWorkflowType InType, const UObject* InSubject) const override;
diff --git a/Engine/Source/Developer/RigVMDeveloper/Public/RigVMModel/RigVMClient.h b/Engine/Source/Developer/RigVMDeveloper/Public/RigVMModel/RigVMClient.h
--- a/Engine/Source/Developer/RigVMDeveloper/Public/RigVMModel/RigVMClient.h
+++ b/Engine/Source/Developer/RigVMDeveloper/Public/RigVMModel/RigVMClient.h
@@ -103,6 +103,7 @@ public:
 	URigVMController* GetOrCreateController(const UObject* InEditorSideObject);
 	URigVMFunctionLibrary* GetFunctionLibrary() const { return FunctionLibrary; }
 	URigVMFunctionLibrary* GetOrCreateFunctionLibrary(bool bSetupUndoRedo, const FObjectInitializer* ObjectInitializer = nullptr);
+	TArray<FName> GetEntryNames() const;
 
 	URigVMGraph* AddModel(const FName& InName, bool bSetupUndoRedo, const FObjectInitializer* ObjectInitializer = nullptr);
 	void AddModel(URigVMGraph* InModel, bool bCreateController);
diff --git a/Engine/Source/Developer/RigVMDeveloper/Public/RigVMModel/RigVMNode.h b/Engine/Source/Developer/RigVMDeveloper/Public/RigVMModel/RigVMNode.h
--- a/Engine/Source/Developer/RigVMDeveloper/Public/RigVMModel/RigVMNode.h
+++ b/Engine/Source/Developer/RigVMDeveloper/Public/RigVMModel/RigVMNode.h
@@ -153,6 +153,10 @@ public:
 	UFUNCTION(BlueprintCallable, Category = RigVMNode)
 	virtual FName GetEventName() const;
 
+	// Returns true if this node can only exist once in a graph
+	UFUNCTION(BlueprintCallable, Category = RigVMNode)
+	virtual bool CanOnlyExistOnce() const;
+
 	// Returns true if the node has any input pins
 	UFUNCTION(BlueprintCallable, Category = RigVMNode)
 	virtual bool HasInputPin(bool bIncludeIO = true) const;
diff --git a/Engine/Source/Runtime/RigVM/Public/RigVMCore/RigVMExecuteContext.h b/Engine/Source/Runtime/RigVM/Public/RigVMCore/RigVMExecuteContext.h
--- a/Engine/Source/Runtime/RigVM/Public/RigVMCore/RigVMExecuteContext.h
+++ b/Engine/Source/Runtime/RigVM/Public/RigVMCore/RigVMExecuteContext.h
@@ -281,6 +281,7 @@ private:
 	friend struct FRigUnit_InverseExecution;
 	friend struct FRigUnit_PrepareForExecution;
 	friend struct FRigUnit_InteractionExecution;
+	friend struct FRigUnit_UserDefinedEvent;
 	friend struct FEngineTestRigVM_Begin;
 	friend struct FEngineTestRigVM_Setup; 
 };
diff --git a/Engine/Source/Runtime/RigVM/Public/RigVMCore/RigVMStruct.h b/Engine/Source/Runtime/RigVM/Public/RigVMCore/RigVMStruct.h
--- a/Engine/Source/Runtime/RigVM/Public/RigVMCore/RigVMStruct.h
+++ b/Engine/Source/Runtime/RigVM/Public/RigVMCore/RigVMStruct.h
@@ -221,6 +221,7 @@ struct RIGVM_API FRigVMStruct
 	virtual ~FRigVMStruct() {}
 	virtual FString ProcessPinLabelForInjection(const FString& InLabel) const { return InLabel; }
 	virtual FName GetEventName() const { return NAME_None; }
+	virtual bool CanOnlyExistOnce() const { return false; }
 
 public:
 
