Add dll definition to FControlRigEditMode

diff --git a/Engine/Plugins/Animation/ControlRig/Source/ControlRigEditor/Private/EditMode/ControlRigEditMode.h b/Engine/Plugins/Animation/ControlRig/Source/ControlRigEditor/Private/EditMode/ControlRigEditMode.h
--- a/Engine/Plugins/Animation/ControlRig/Source/ControlRigEditor/Private/EditMode/ControlRigEditMode.h
+++ b/Engine/Plugins/Animation/ControlRig/Source/ControlRigEditor/Private/EditMode/ControlRigEditMode.h
@@ -80,7 +80,7 @@
 	FDelegateHandle OnBoneTransformsFinalizedHandle;
 };
 
-class FControlRigEditMode : public IPersonaEditMode
+class CONTROLRIGEDITOR_API FControlRigEditMode : public IPersonaEditMode
 {
 public:
 	static FName ModeName;
