Fix use of IsChildOf with possible nullptr (cf. CL 20439034)
todo: check for all use of IsChildOf within ControlRig to avoid future issues

diff --git a/Engine/Plugins/Animation/ControlRig/Source/ControlRigDeveloper/Private/Graph/ControlRigGraphNode.cpp b/Engine/Plugins/Animation/ControlRig/Source/ControlRigDeveloper/Private/Graph/ControlRigGraphNode.cpp
--- a/Engine/Plugins/Animation/ControlRig/Source/ControlRigDeveloper/Private/Graph/ControlRigGraphNode.cpp
+++ b/Engine/Plugins/Animation/ControlRig/Source/ControlRigDeveloper/Private/Graph/ControlRigGraphNode.cpp
@@ -64,7 +64,8 @@ FText UControlRigGraphNode::GetNodeTitle(ENodeTitleType::Type TitleType) const
 		{
 			if (URigVMUnitNode* UnitNode = Cast<URigVMUnitNode>(ModelNode))
 			{
-				if (UnitNode->GetScriptStruct()->IsChildOf(FRigUnit::StaticStruct()))
+				const UScriptStruct* ScriptStruct = UnitNode->GetScriptStruct();
+				if (ScriptStruct && ScriptStruct->IsChildOf(FRigUnit::StaticStruct()))
 				{
 					if (TSharedPtr<FStructOnScope> StructOnScope = UnitNode->ConstructStructInstance())
 					{
@@ -519,7 +520,8 @@ void UControlRigGraphNode::AllocateDefaultPins()
 				{
 					if (ModelPin->IsStruct())
 					{
-						if (ModelPin->GetScriptStruct()->IsChildOf(FRigVMExecuteContext::StaticStruct()))
+						const UScriptStruct* ScriptStruct = ModelPin->GetScriptStruct();
+						if (ScriptStruct && ScriptStruct->IsChildOf(FRigVMExecuteContext::StaticStruct()))
 						{
 							ExecutePins.Add(ModelPin);
 							continue;
