Fix use of IsChildOf with possible nullptr (cf. CL 20439034)

diff --git a/Engine/Plugins/Animation/ControlRig/Source/ControlRig/Private/ControlRig.cpp b/Engine/Plugins/Animation/ControlRig/Source/ControlRig/Private/ControlRig.cpp
--- a/Engine/Plugins/Animation/ControlRig/Source/ControlRig/Private/ControlRig.cpp
+++ b/Engine/Plugins/Animation/ControlRig/Source/ControlRig/Private/ControlRig.cpp
@@ -1545,7 +1545,8 @@ TArray<UControlRig*> UControlRig::FindControlRigs(UObject* Outer, TSubclassOf<UC
 	for (TObjectIterator<UControlRig> Itr; Itr; ++Itr)
 	{
 		UControlRig* RigInstance = *Itr;
-		if (OptionalClass == nullptr || RigInstance->GetClass()->IsChildOf(OptionalClass))
+		const UClass* RigInstanceClass = RigInstance ? RigInstance->GetClass() : nullptr;
+		if (OptionalClass == nullptr || (RigInstanceClass && RigInstanceClass->IsChildOf(OptionalClass)))
 		{
 			if(RigInstance->IsInOuter(Outer))
 			{
diff --git a/Engine/Plugins/Animation/ControlRig/Source/ControlRigDeveloper/Private/ControlRigBlueprint.cpp b/Engine/Plugins/Animation/ControlRig/Source/ControlRigDeveloper/Private/ControlRigBlueprint.cpp
--- a/Engine/Plugins/Animation/ControlRig/Source/ControlRigDeveloper/Private/ControlRigBlueprint.cpp
+++ b/Engine/Plugins/Animation/ControlRig/Source/ControlRigDeveloper/Private/ControlRigBlueprint.cpp
@@ -3038,7 +3038,8 @@ void UControlRigBlueprint::PopulateModelFromGraphForBackwardsCompatibility(UCont
 							// check the material + mesh pins for deprecated control nodes
 							if (URigVMUnitNode* ModelUnitNode = Cast<URigVMUnitNode>(ModelNode))
 							{
-								if (ModelUnitNode->GetScriptStruct()->IsChildOf(FRigUnit_Control::StaticStruct()))
+								const UScriptStruct* ScriptStruct = ModelUnitNode->GetScriptStruct();
+								if (ScriptStruct && ScriptStruct->IsChildOf(FRigUnit_Control::StaticStruct()))
 								{
 									if (Pin->GetName().EndsWith(TEXT(".StaticMesh")) || Pin->GetName().EndsWith(TEXT(".Materials")))
 									{
diff --git a/Engine/Plugins/Animation/ControlRig/Source/ControlRigDeveloper/Private/ControlRigBlueprintUtils.cpp b/Engine/Plugins/Animation/ControlRig/Source/ControlRigDeveloper/Private/ControlRigBlueprintUtils.cpp
--- a/Engine/Plugins/Animation/ControlRig/Source/ControlRigDeveloper/Private/ControlRigBlueprintUtils.cpp
+++ b/Engine/Plugins/Animation/ControlRig/Source/ControlRigDeveloper/Private/ControlRigBlueprintUtils.cpp
@@ -75,11 +75,14 @@ void FControlRigBlueprintUtils::ForAllRigUnits(TFunction<void(UScriptStruct*)> I
 	// Run over all unit types
 	for(TObjectIterator<UStruct> StructIt; StructIt; ++StructIt)
 	{
-		if(StructIt->IsChildOf(FRigUnit::StaticStruct()) && !StructIt->HasMetaData(FRigVMStruct::AbstractMetaName))
+		if (*StructIt)
 		{
-			if (UScriptStruct* ScriptStruct = Cast<UScriptStruct>(*StructIt))
+			if(StructIt->IsChildOf(FRigUnit::StaticStruct()) && !StructIt->HasMetaData(FRigVMStruct::AbstractMetaName))
 			{
-				InFunction(ScriptStruct);
+				if (UScriptStruct* ScriptStruct = Cast<UScriptStruct>(*StructIt))
+				{
+					InFunction(ScriptStruct);
+				}
 			}
 		}
 	}
diff --git a/Engine/Plugins/Animation/ControlRig/Source/ControlRigEditor/Private/ControlRigEditorModule.cpp b/Engine/Plugins/Animation/ControlRig/Source/ControlRigEditor/Private/ControlRigEditorModule.cpp
--- a/Engine/Plugins/Animation/ControlRig/Source/ControlRigEditor/Private/ControlRigEditorModule.cpp
+++ b/Engine/Plugins/Animation/ControlRig/Source/ControlRigEditor/Private/ControlRigEditorModule.cpp
@@ -524,15 +524,17 @@ TSharedRef< SWidget > FControlRigEditorModule::GenerateAnimationMenu(TWeakPtr<IA
 									}
 									bool IsClassAllowed(const FClassViewerInitializationOptions& InInitOptions, const UClass* InClass, TSharedRef< FClassViewerFilterFuncs > InFilterFuncs) override
 									{
-										const bool bChildOfObjectClass = InClass->IsChildOf(UControlRig::StaticClass());
-										const bool bMatchesFlags = !InClass->HasAnyClassFlags(CLASS_Hidden | CLASS_HideDropDown | CLASS_Deprecated | CLASS_Abstract);
-										const bool bNotNative = !InClass->IsNative();
-
-										if (bChildOfObjectClass && bMatchesFlags && bNotNative)
+										if (InClass)
 										{
-											FAssetData AssetData(InClass);
-											return MatchesFilter(AssetData);
+											const bool bChildOfObjectClass = InClass->IsChildOf(UControlRig::StaticClass());
+											const bool bMatchesFlags = !InClass->HasAnyClassFlags(CLASS_Hidden | CLASS_HideDropDown | CLASS_Deprecated | CLASS_Abstract);
+											const bool bNotNative = !InClass->IsNative();
 
+											if (bChildOfObjectClass && bMatchesFlags && bNotNative)
+											{
+												FAssetData AssetData(InClass);
+												return MatchesFilter(AssetData);
+											}
 										}
 										return false;
 									}
@@ -1020,7 +1022,7 @@ void FControlRigEditorModule::GetTypeActions(UControlRigBlueprint* CRB, FBluepri
 	for(const FRigVMFunction& Function : FRigVMRegistry::Get().GetFunctions())
 	{
 		UScriptStruct* Struct = Function.Struct;
-		if (!Struct->IsChildOf(FRigUnit::StaticStruct()))
+		if (Struct == nullptr || !Struct->IsChildOf(FRigUnit::StaticStruct()))
 		{
 			continue;
 		}
diff --git a/Engine/Plugins/Animation/ControlRig/Source/ControlRigEditor/Private/ControlRigSequencerEditorLibrary.cpp b/Engine/Plugins/Animation/ControlRig/Source/ControlRigEditor/Private/ControlRigSequencerEditorLibrary.cpp
--- a/Engine/Plugins/Animation/ControlRig/Source/ControlRigEditor/Private/ControlRigSequencerEditorLibrary.cpp
+++ b/Engine/Plugins/Animation/ControlRig/Source/ControlRigEditor/Private/ControlRigSequencerEditorLibrary.cpp
@@ -149,7 +149,7 @@ static void AcquireSkeletonAndSkelMeshCompFromObject(UObject* BoundObject, USkel
 
 			for (USCS_Node* Node : ActorBlueprintNodes)
 			{
-				if (Node->ComponentClass->IsChildOf(USkeletalMeshComponent::StaticClass()))
+				if (Node->ComponentClass && Node->ComponentClass->IsChildOf(USkeletalMeshComponent::StaticClass()))
 				{
 					USkeletalMeshComponent* SkeletalMeshComp = Cast<USkeletalMeshComponent>(Node->GetActualComponentTemplate(ActorBlueprintGeneratedClass));
 					if (SkeletalMeshComp)
diff --git a/Engine/Plugins/Animation/ControlRig/Source/ControlRigEditor/Private/ControlRigVariableDetailsCustomization.cpp b/Engine/Plugins/Animation/ControlRig/Source/ControlRigEditor/Private/ControlRigVariableDetailsCustomization.cpp
--- a/Engine/Plugins/Animation/ControlRig/Source/ControlRigEditor/Private/ControlRigVariableDetailsCustomization.cpp
+++ b/Engine/Plugins/Animation/ControlRig/Source/ControlRigEditor/Private/ControlRigVariableDetailsCustomization.cpp
@@ -22,7 +22,7 @@ TSharedPtr<IDetailCustomization> FControlRigVariableDetailsCustomization::MakeIn
 	{
 		if (UBlueprint* Blueprint = Cast<UBlueprint>((*Objects)[0]))
 		{
-			if (Blueprint->ParentClass->IsChildOf(UControlRig::StaticClass()))
+			if (Blueprint->ParentClass && Blueprint->ParentClass->IsChildOf(UControlRig::StaticClass()))
 			{
 				return MakeShareable(new FControlRigVariableDetailsCustomization(InBlueprintEditor, Blueprint));
 			}
diff --git a/Engine/Plugins/Animation/ControlRig/Source/ControlRigEditor/Private/Editor/DetailsViewWrapperObject.h b/Engine/Plugins/Animation/ControlRig/Source/ControlRigEditor/Private/Editor/DetailsViewWrapperObject.h
--- a/Engine/Plugins/Animation/ControlRig/Source/ControlRigEditor/Private/Editor/DetailsViewWrapperObject.h
+++ b/Engine/Plugins/Animation/ControlRig/Source/ControlRigEditor/Private/Editor/DetailsViewWrapperObject.h
@@ -31,7 +31,8 @@ public:
 	
 	bool IsChildOf(const UStruct* InStruct) const
 	{
-		return GetWrappedStruct()->IsChildOf(InStruct);
+		const UScriptStruct* WrappedStruct = GetWrappedStruct();
+		return WrappedStruct && WrappedStruct->IsChildOf(InStruct);
 	}
 
 	template<typename T>
diff --git a/Engine/Plugins/Animation/ControlRig/Source/ControlRigEditor/Private/Editor/SControlRigValidationWidget.cpp b/Engine/Plugins/Animation/ControlRig/Source/ControlRigEditor/Private/Editor/SControlRigValidationWidget.cpp
--- a/Engine/Plugins/Animation/ControlRig/Source/ControlRigEditor/Private/Editor/SControlRigValidationWidget.cpp
+++ b/Engine/Plugins/Animation/ControlRig/Source/ControlRigEditor/Private/Editor/SControlRigValidationWidget.cpp
@@ -116,7 +116,8 @@ void SControlRigValidationWidget::Construct(const FArguments& InArgs, UControlRi
 	ClassItems.Reset();
 	for (TObjectIterator<UClass> ClassIterator; ClassIterator; ++ClassIterator)
 	{
-		if (ClassIterator->IsChildOf(UControlRigValidationPass::StaticClass()) && !ClassIterator->HasAnyClassFlags(CLASS_Abstract | CLASS_Deprecated | CLASS_NewerVersionExists))
+		const bool bIsValidationPassChild = (*ClassIterator) && ClassIterator->IsChildOf(UControlRigValidationPass::StaticClass());
+		if (bIsValidationPassChild && !ClassIterator->HasAnyClassFlags(CLASS_Abstract | CLASS_Deprecated | CLASS_NewerVersionExists))
 		{
 			ClassItems.Add(MakeShared<FControlRigValidationPassItem>(*ClassIterator));
 		}
diff --git a/Engine/Plugins/Animation/ControlRig/Source/ControlRigEditor/Private/Graph/ControlRigGraphPanelPinFactory.cpp b/Engine/Plugins/Animation/ControlRig/Source/ControlRigEditor/Private/Graph/ControlRigGraphPanelPinFactory.cpp
--- a/Engine/Plugins/Animation/ControlRig/Source/ControlRigEditor/Private/Graph/ControlRigGraphPanelPinFactory.cpp
+++ b/Engine/Plugins/Animation/ControlRig/Source/ControlRigEditor/Private/Graph/ControlRigGraphPanelPinFactory.cpp
@@ -104,7 +104,8 @@ TSharedPtr<SGraphPin> FControlRigGraphPanelPinFactory::CreatePin(UEdGraphPin* In
 
 			if (InPin->PinType.PinCategory == UEdGraphSchema_K2::PC_Struct)
 			{
-				if (Cast<UStruct>(InPin->PinType.PinSubCategoryObject)->IsChildOf(FRigVMExecuteContext::StaticStruct()))
+				const UStruct* Struct = Cast<UStruct>(InPin->PinType.PinSubCategoryObject);
+				if (Struct && Struct->IsChildOf(FRigVMExecuteContext::StaticStruct()))
 				{
 					return SNew(SGraphPinExec, InPin);
 				}
diff --git a/Engine/Plugins/Animation/ControlRig/Source/ControlRigEditor/Private/Sequencer/ControlRigParameterTrackEditor.cpp b/Engine/Plugins/Animation/ControlRig/Source/ControlRigEditor/Private/Sequencer/ControlRigParameterTrackEditor.cpp
--- a/Engine/Plugins/Animation/ControlRig/Source/ControlRigEditor/Private/Sequencer/ControlRigParameterTrackEditor.cpp
+++ b/Engine/Plugins/Animation/ControlRig/Source/ControlRigEditor/Private/Sequencer/ControlRigParameterTrackEditor.cpp
@@ -166,7 +166,7 @@ static USkeleton* AcquireSkeletonFromObjectGuid(const FGuid& Guid, UObject** Obj
 
 			for (USCS_Node* Node : ActorBlueprintNodes)
 			{
-				if (Node->ComponentClass->IsChildOf(USkeletalMeshComponent::StaticClass()))
+				if (Node->ComponentClass && Node->ComponentClass->IsChildOf(USkeletalMeshComponent::StaticClass()))
 				{
 					if (USkeletalMeshComponent* SkeletalMeshComponent = Cast<USkeletalMeshComponent>(Node->GetActualComponentTemplate(ActorBlueprintGeneratedClass)))
 					{
@@ -546,7 +546,14 @@ TSharedRef<ISequencerSection> FControlRigParameterTrackEditor::MakeSectionInterf
 
 void FControlRigParameterTrackEditor::BuildObjectBindingContextMenu(FMenuBuilder& MenuBuilder, const TArray<FGuid>& ObjectBindings, const UClass* ObjectClass)
 {
-	if (ObjectClass->IsChildOf(USkeletalMeshComponent::StaticClass()) || ObjectClass->IsChildOf(AActor::StaticClass()) || ObjectClass->IsChildOf(UChildActorComponent::StaticClass()))
+	if(!ObjectClass)
+	{
+		return;
+	}
+	
+	if (ObjectClass->IsChildOf(USkeletalMeshComponent::StaticClass()) ||
+		ObjectClass->IsChildOf(AActor::StaticClass()) ||
+		ObjectClass->IsChildOf(UChildActorComponent::StaticClass()))
 	{
 		const TSharedPtr<ISequencer> ParentSequencer = GetSequencer();
 		UObject* BoundObject = nullptr;
@@ -674,15 +681,18 @@ public:
 	}
 	bool IsClassAllowed(const FClassViewerInitializationOptions& InInitOptions, const UClass* InClass, TSharedRef< FClassViewerFilterFuncs > InFilterFuncs) override
 	{
-		const bool bChildOfObjectClass = InClass->IsChildOf(UControlRig::StaticClass());
-		const bool bMatchesFlags = !InClass->HasAnyClassFlags(CLASS_Hidden | CLASS_HideDropDown | CLASS_Deprecated | CLASS_Abstract);
-		const bool bNotNative = !InClass->IsNative();
-
-		if (bChildOfObjectClass && bMatchesFlags && bNotNative)
+		if(InClass)
 		{
-			FAssetData AssetData(InClass);
-			return MatchesFilter(AssetData);
+			const bool bChildOfObjectClass = InClass->IsChildOf(UControlRig::StaticClass());
+			const bool bMatchesFlags = !InClass->HasAnyClassFlags(CLASS_Hidden | CLASS_HideDropDown | CLASS_Deprecated | CLASS_Abstract);
+			const bool bNotNative = !InClass->IsNative();
 
+			if (bChildOfObjectClass && bMatchesFlags && bNotNative)
+			{
+				FAssetData AssetData(InClass);
+				return MatchesFilter(AssetData);
+
+			}
 		}
 		return false;
 	}
@@ -1135,7 +1145,14 @@ void FControlRigParameterTrackEditor::BakeToControlRig(UClass* InClass, FGuid Ob
 
 void FControlRigParameterTrackEditor::BuildObjectBindingTrackMenu(FMenuBuilder& MenuBuilder, const TArray<FGuid>& ObjectBindings, const UClass* ObjectClass)
 {
-	if (ObjectClass->IsChildOf(USkeletalMeshComponent::StaticClass()) || ObjectClass->IsChildOf(AActor::StaticClass()) || ObjectClass->IsChildOf(UChildActorComponent::StaticClass()))
+	if(!ObjectClass)
+	{
+		return;
+	}
+	
+	if (ObjectClass->IsChildOf(USkeletalMeshComponent::StaticClass()) ||
+		ObjectClass->IsChildOf(AActor::StaticClass()) ||
+		ObjectClass->IsChildOf(UChildActorComponent::StaticClass()))
 	{
 		const TSharedPtr<ISequencer> ParentSequencer = GetSequencer();
 		UObject* BoundObject = nullptr;
@@ -2189,7 +2206,8 @@ bool FControlRigParameterTrackEditor::MatchesContext(const FTransactionContext&
 		UObject* Object = TransactionObjectPair.Key;
 		while (Object != nullptr)
 		{
-			if (Object->GetClass()->IsChildOf(UMovieSceneControlRigParameterSection::StaticClass()))
+			const UClass* ObjectClass = Object->GetClass();
+			if (ObjectClass && ObjectClass->IsChildOf(UMovieSceneControlRigParameterSection::StaticClass()))
 			{
 				UMovieSceneControlRigParameterSection* Section = Cast< UMovieSceneControlRigParameterSection>(Object);
 				if (Section)
