Control Rig: Event API

diff --git a/Engine/Plugins/Animation/ControlRig/Source/ControlRig/Private/ControlRig.cpp b/Engine/Plugins/Animation/ControlRig/Source/ControlRig/Private/ControlRig.cpp
--- a/Engine/Plugins/Animation/ControlRig/Source/ControlRig/Private/ControlRig.cpp
+++ b/Engine/Plugins/Animation/ControlRig/Source/ControlRig/Private/ControlRig.cpp
@@ -1008,7 +1008,11 @@ void UControlRig::Execute(const EControlRigState InState, const FName& InEventNa
 	}
 	else if (InState == EControlRigState::Update)
 	{
-		DeltaTime = 0.f;
+		if(InEventName == GetEventQueue().Last() ||
+			!GetEventQueue().Contains(InEventName))
+		{
+			DeltaTime = 0.f;
+		}
 
 		if (ExecutedEvent.IsBound())
 		{
@@ -1401,6 +1405,35 @@ void UControlRig::ExecuteUnits(FRigUnitContext& InOutContext, const FName& InEve
 	}
 }
 
+bool UControlRig::ContainsEvent(const FName& InEventName) const
+{
+	if(VM)
+	{
+		return VM->ContainsEntry(InEventName);
+	}
+	return false;
+}
+
+TArray<FName> UControlRig::GetEvents() const
+{
+	if(VM)
+	{
+		return VM->GetEntryNames();
+	}
+	return TArray<FName>();
+}
+
+bool UControlRig::ExecuteEvent(const FName& InEventName)
+{
+	if(ContainsEvent(InEventName))
+	{
+		TGuardValue<TArray<FName>> EventQueueGuard(EventQueue, {InEventName});
+		Evaluate_AnyThread();
+		return true;
+	}
+	return false;
+}
+
 void UControlRig::RequestInit()
 {
 	bRequiresInitExecution = true;
diff --git a/Engine/Plugins/Animation/ControlRig/Source/ControlRig/Public/ControlRig.h b/Engine/Plugins/Animation/ControlRig/Source/ControlRig/Public/ControlRig.h
--- a/Engine/Plugins/Animation/ControlRig/Source/ControlRig/Public/ControlRig.h
+++ b/Engine/Plugins/Animation/ControlRig/Source/ControlRig/Public/ControlRig.h
@@ -282,6 +282,18 @@ public:
 	/** ExecuteUnits */
 	virtual void ExecuteUnits(FRigUnitContext& InOutContext, const FName& InEventName);
 
+	/** Returns true if this rig contains a given event */
+	UFUNCTION(BlueprintPure, Category = "Control Rig")
+	bool ContainsEvent(const FName& InEventName) const;
+
+	/** Returns the user defined events */
+	UFUNCTION(BlueprintPure, Category = "Control Rig")
+	TArray<FName> GetEvents() const;
+
+	/** Execute a user defined event */
+	UFUNCTION(BlueprintCallable, Category = "Control Rig")
+	bool ExecuteEvent(const FName& InEventName);
+
 	/** Requests to perform an init during the next execution */
 	UFUNCTION(BlueprintCallable, Category = "Control Rig")
 	void RequestInit();
