Control Rig: Allow access to event queue in AnimNode

diff --git a/Engine/Plugins/Animation/ControlRig/Source/ControlRig/Private/AnimNode_ControlRigBase.cpp b/Engine/Plugins/Animation/ControlRig/Source/ControlRig/Private/AnimNode_ControlRigBase.cpp
--- a/Engine/Plugins/Animation/ControlRig/Source/ControlRig/Private/AnimNode_ControlRigBase.cpp
+++ b/Engine/Plugins/Animation/ControlRig/Source/ControlRig/Private/AnimNode_ControlRigBase.cpp
@@ -7,6 +7,7 @@
 #include "Animation/NodeMappingContainer.h"
 #include "AnimationRuntime.h"
 #include "Units/Execution/RigUnit_BeginExecution.h"
+#include "Algo/Transform.h"
 
 #if ENABLE_ANIM_DEBUG
 TAutoConsoleVariable<int32> CVarAnimNodeControlRigDebug(TEXT("a.AnimNode.ControlRig.Debug"), 0, TEXT("Set to 1 to turn on debug drawing for AnimNode_ControlRigBase"));
@@ -504,7 +505,22 @@ void FAnimNode_ControlRigBase::ExecuteControlRig(FPoseContext& InOutput)
 			}
 #endif
 
-			// first evaluate control rig
+			// pick the event to run
+			if(EventQueue.IsEmpty())
+			{
+				ControlRig->SetEventQueue({FRigUnit_BeginExecution::EventName});
+			}
+			else
+			{
+				TArray<FName> EventNames;
+				Algo::Transform(EventQueue, EventNames, [](const FControlRigAnimNodeEventName& InEventName) 
+				{
+					return InEventName.EventName;
+				});
+				ControlRig->SetEventQueue(EventNames);
+			}
+			
+			// evaluate control rig
 			ControlRig->Evaluate_AnyThread();
 
 #if ENABLE_ANIM_DEBUG 
diff --git a/Engine/Plugins/Animation/ControlRig/Source/ControlRig/Private/AnimNode_ControlRig_ExternalSource.cpp b/Engine/Plugins/Animation/ControlRig/Source/ControlRig/Private/AnimNode_ControlRig_ExternalSource.cpp
--- a/Engine/Plugins/Animation/ControlRig/Source/ControlRig/Private/AnimNode_ControlRig_ExternalSource.cpp
+++ b/Engine/Plugins/Animation/ControlRig/Source/ControlRig/Private/AnimNode_ControlRig_ExternalSource.cpp
@@ -18,3 +18,12 @@ UControlRig* FAnimNode_ControlRig_ExternalSource::GetControlRig() const
 	return (ControlRig.IsValid()? ControlRig.Get() : nullptr);
 }
 
+TSubclassOf<UControlRig> FAnimNode_ControlRig_ExternalSource::GetControlRigClass() const
+{
+	if(UControlRig* CR = GetControlRig())
+	{
+		return CR->GetClass();
+	}
+	return nullptr;
+}
+
diff --git a/Engine/Plugins/Animation/ControlRig/Source/ControlRig/Public/AnimNode_ControlRig.h b/Engine/Plugins/Animation/ControlRig/Source/ControlRig/Public/AnimNode_ControlRig.h
--- a/Engine/Plugins/Animation/ControlRig/Source/ControlRig/Public/AnimNode_ControlRig.h
+++ b/Engine/Plugins/Animation/ControlRig/Source/ControlRig/Public/AnimNode_ControlRig.h
@@ -19,7 +19,8 @@ struct CONTROLRIG_API FAnimNode_ControlRig : public FAnimNode_ControlRigBase
 
 	FAnimNode_ControlRig();
 
-	UControlRig* GetControlRig() const { return ControlRig; }
+	virtual UControlRig* GetControlRig() const override { return ControlRig; }
+	virtual TSubclassOf<UControlRig> GetControlRigClass() const override { return ControlRigClass; }
 
 	// FAnimNode_Base interface
 	virtual void OnInitializeAnimInstance(const FAnimInstanceProxy* InProxy, const UAnimInstance* InAnimInstance) override;
diff --git a/Engine/Plugins/Animation/ControlRig/Source/ControlRig/Public/AnimNode_ControlRigBase.h b/Engine/Plugins/Animation/ControlRig/Source/ControlRig/Public/AnimNode_ControlRigBase.h
--- a/Engine/Plugins/Animation/ControlRig/Source/ControlRig/Public/AnimNode_ControlRigBase.h
+++ b/Engine/Plugins/Animation/ControlRig/Source/ControlRig/Public/AnimNode_ControlRigBase.h
@@ -38,6 +38,19 @@ struct CONTROLRIG_API FControlRigIOSettings
 	bool bUpdateCurves;
 };
 
+USTRUCT()
+struct CONTROLRIG_API FControlRigAnimNodeEventName
+{
+	GENERATED_BODY()
+
+	FControlRigAnimNodeEventName()
+	: EventName(NAME_None)
+	{}
+
+	UPROPERTY(EditAnywhere, Category = Links)
+	FName EventName;	
+};
+
 /**
  * Animation node that allows animation ControlRig output to be used in an animation graph
  */
@@ -50,7 +63,8 @@ struct CONTROLRIG_API FAnimNode_ControlRigBase : public FAnimNode_CustomProperty
 
 	/* return Control Rig of current object */
 	virtual UControlRig* GetControlRig() const PURE_VIRTUAL(FAnimNode_ControlRigBase::GetControlRig, return nullptr; );
-
+	virtual TSubclassOf<UControlRig> GetControlRigClass() const PURE_VIRTUAL(FAnimNode_ControlRigBase::GetControlRigClass, return nullptr; );
+	
 	// FAnimNode_Base interface
 	virtual void OnInitializeAnimInstance(const FAnimInstanceProxy* InProxy, const UAnimInstance* InAnimInstance) override;
 	virtual bool NeedsOnInitializeAnimInstance() const override { return true; }
@@ -137,6 +151,10 @@ protected:
 	// The below is alpha value support for control rig
 	float InternalBlendAlpha;
 
+	// The customized event queue to run
+	UPROPERTY(EditAnywhere, AdvancedDisplay, Category = Settings)
+	TArray<FControlRigAnimNodeEventName> EventQueue;
+
 	virtual bool CanExecute();
 	// update input/output to control rig
 	virtual void UpdateInput(UControlRig* ControlRig, const FPoseContext& InOutput);
diff --git a/Engine/Plugins/Animation/ControlRig/Source/ControlRig/Public/AnimNode_ControlRig_ExternalSource.h b/Engine/Plugins/Animation/ControlRig/Source/ControlRig/Public/AnimNode_ControlRig_ExternalSource.h
--- a/Engine/Plugins/Animation/ControlRig/Source/ControlRig/Public/AnimNode_ControlRig_ExternalSource.h
+++ b/Engine/Plugins/Animation/ControlRig/Source/ControlRig/Public/AnimNode_ControlRig_ExternalSource.h
@@ -16,7 +16,8 @@ struct CONTROLRIG_API FAnimNode_ControlRig_ExternalSource : public FAnimNode_Con
 	FAnimNode_ControlRig_ExternalSource();
 
 	void SetControlRig(UControlRig* InControlRig);
-	virtual UControlRig* GetControlRig() const;
+	virtual UControlRig* GetControlRig() const override;
+	virtual TSubclassOf<UControlRig> GetControlRigClass() const override;
 
 private:
 	UPROPERTY(transient)
diff --git a/Engine/Plugins/Animation/ControlRig/Source/ControlRigEditor/Private/ControlRigAnimGraphDetails.cpp b/Engine/Plugins/Animation/ControlRig/Source/ControlRigEditor/Private/ControlRigAnimGraphDetails.cpp
new file mode 100644
--- /dev/null
+++ b/Engine/Plugins/Animation/ControlRig/Source/ControlRigEditor/Private/ControlRigAnimGraphDetails.cpp
@@ -0,0 +1,174 @@
+// Copyright Epic Games, Inc. All Rights Reserved.
+
+#include "ControlRigAnimGraphDetails.h"
+
+#include "ControlRig.h"
+#include "Widgets/SWidget.h"
+#include "DetailLayoutBuilder.h"
+#include "Styling/AppStyle.h"
+#include "PropertyCustomizationHelpers.h"
+#include "AnimGraphNode_ControlRig.h"
+#include "Algo/Transform.h"
+
+#define LOCTEXT_NAMESPACE "ControlRigAnimGraphDetails"
+static const FText ControlRigAnimDetailsMultipleValues = LOCTEXT("MultipleValues", "Multiple Values");
+
+void FControlRigAnimNodeEventNameDetails::CustomizeHeader(TSharedRef<IPropertyHandle> InStructPropertyHandle, FDetailWidgetRow& HeaderRow, IPropertyTypeCustomizationUtils& StructCustomizationUtils)
+{
+	AnimNodeBeingCustomized = nullptr;
+
+	TArray<UObject*> Objects;
+	InStructPropertyHandle->GetOuterObjects(Objects);
+	for (UObject* Object : Objects)
+	{
+		if(UAnimGraphNode_ControlRig* GraphNode = Cast<UAnimGraphNode_ControlRig>(Object))
+		{
+			AnimNodeBeingCustomized = &GraphNode->Node;
+			break;
+		}
+	}
+
+	if (AnimNodeBeingCustomized == nullptr)
+	{
+		HeaderRow
+		.NameContent()
+		[
+			InStructPropertyHandle->CreatePropertyNameWidget()
+		]
+		.ValueContent()
+		[
+			InStructPropertyHandle->CreatePropertyValueWidget()
+		];
+	}
+	else
+	{
+		NameHandle = InStructPropertyHandle->GetChildHandle(GET_MEMBER_NAME_CHECKED(FControlRigAnimNodeEventName, EventName));
+		UpdateEntryNameList();
+
+		HeaderRow
+		.NameContent()
+		[
+			InStructPropertyHandle->CreatePropertyNameWidget()
+		]
+		.ValueContent()
+		.MinDesiredWidth(250.f)
+		[
+			SNew(SHorizontalBox)
+
+			+ SHorizontalBox::Slot()
+			.AutoWidth()
+			.Padding(0.f, 0.f, 0.f, 0.f)
+			[
+				SAssignNew(SearchableComboBox, SSearchableComboBox)
+				.OptionsSource(&EntryNameList)
+				.OnSelectionChanged(this, &FControlRigAnimNodeEventNameDetails::OnEntryNameChanged)
+				.OnGenerateWidget(this, &FControlRigAnimNodeEventNameDetails::OnGetEntryNameWidget)
+				.IsEnabled(!NameHandle->IsEditConst())
+				.Content()
+				[
+					SNew(STextBlock)
+					.Text(this, &FControlRigAnimNodeEventNameDetails::GetEntryNameAsText)
+					.Font(IDetailLayoutBuilder::GetDetailFont())
+				]
+			]
+		];
+	}
+}
+
+void FControlRigAnimNodeEventNameDetails::CustomizeChildren(TSharedRef<IPropertyHandle> InStructPropertyHandle, IDetailChildrenBuilder& StructBuilder, IPropertyTypeCustomizationUtils& StructCustomizationUtils)
+{
+	if (InStructPropertyHandle->IsValidHandle())
+	{
+		// only fill the children if the blueprint cannot be found
+		if (AnimNodeBeingCustomized == nullptr)
+		{
+			uint32 NumChildren = 0;
+			InStructPropertyHandle->GetNumChildren(NumChildren);
+
+			for (uint32 ChildIndex = 0; ChildIndex < NumChildren; ChildIndex++)
+			{
+				StructBuilder.AddProperty(InStructPropertyHandle->GetChildHandle(ChildIndex).ToSharedRef());
+			}
+		}
+	}
+}
+
+FString FControlRigAnimNodeEventNameDetails::GetEntryName() const
+{
+	FString EntryNameStr;
+	if (NameHandle.IsValid())
+	{
+		for(int32 ObjectIndex = 0; ObjectIndex < NameHandle->GetNumPerObjectValues(); ObjectIndex++)
+		{
+			FString PerObjectValue;
+			NameHandle->GetPerObjectValue(ObjectIndex, PerObjectValue);
+
+			if(ObjectIndex == 0)
+			{
+				EntryNameStr = PerObjectValue;
+			}
+			else if(EntryNameStr != PerObjectValue)
+			{
+				return ControlRigAnimDetailsMultipleValues.ToString();
+			}
+		}
+	}
+	return EntryNameStr;
+}
+
+void FControlRigAnimNodeEventNameDetails::SetEntryName(FString InName)
+{
+	if (NameHandle.IsValid())
+	{
+		NameHandle->SetValue(InName);
+	}
+}
+
+void FControlRigAnimNodeEventNameDetails::UpdateEntryNameList()
+{
+	EntryNameList.Reset();
+
+	if (AnimNodeBeingCustomized)
+	{
+		if(const UClass* Class = AnimNodeBeingCustomized->GetControlRigClass())
+		{
+			if(const UControlRig* CDO = Cast<UControlRig>(Class->GetDefaultObject(true)))
+			{
+				Algo::Transform(CDO->GetEvents(), EntryNameList,[](const FName& InEntryName)
+				{
+					return MakeShareable(new FString(InEntryName.ToString()));
+				});
+				if(SearchableComboBox.IsValid())
+				{
+					SearchableComboBox->RefreshOptions();
+				}
+			}
+		}
+	}
+}
+
+void FControlRigAnimNodeEventNameDetails::OnEntryNameChanged(TSharedPtr<FString> InItem, ESelectInfo::Type InSelectionInfo)
+{
+	if (InItem.IsValid())
+	{
+		SetEntryName(*InItem);
+	}
+	else
+	{
+		SetEntryName(FString());
+	}
+}
+
+TSharedRef<SWidget> FControlRigAnimNodeEventNameDetails::OnGetEntryNameWidget(TSharedPtr<FString> InItem)
+{
+	return SNew(STextBlock)
+		.Text(FText::FromString(InItem.IsValid() ? *InItem : FString()))
+		.Font(IDetailLayoutBuilder::GetDetailFont());
+}
+
+FText FControlRigAnimNodeEventNameDetails::GetEntryNameAsText() const
+{
+	return FText::FromString(GetEntryName());
+}
+
+#undef LOCTEXT_NAMESPACE
diff --git a/Engine/Plugins/Animation/ControlRig/Source/ControlRigEditor/Private/ControlRigAnimGraphDetails.h b/Engine/Plugins/Animation/ControlRig/Source/ControlRigEditor/Private/ControlRigAnimGraphDetails.h
new file mode 100644
--- /dev/null
+++ b/Engine/Plugins/Animation/ControlRig/Source/ControlRigEditor/Private/ControlRigAnimGraphDetails.h
@@ -0,0 +1,42 @@
+// Copyright Epic Games, Inc. All Rights Reserved.
+
+#pragma once
+
+#include "CoreMinimal.h"
+#include "IPropertyTypeCustomization.h"
+#include "DetailCategoryBuilder.h"
+#include "IDetailChildrenBuilder.h"
+#include "SSearchableComboBox.h"
+#include "AnimNode_ControlRigBase.h"
+
+class IDetailLayoutBuilder;
+
+class IPropertyHandle;
+
+class FControlRigAnimNodeEventNameDetails : public IPropertyTypeCustomization
+{
+public:
+
+	static TSharedRef<IPropertyTypeCustomization> MakeInstance()
+	{
+		return MakeShareable(new FControlRigAnimNodeEventNameDetails);
+	}
+
+	/** IPropertyTypeCustomization interface */
+	virtual void CustomizeHeader(TSharedRef<class IPropertyHandle> InStructPropertyHandle, class FDetailWidgetRow& HeaderRow, IPropertyTypeCustomizationUtils& StructCustomizationUtils) override;
+	virtual void CustomizeChildren(TSharedRef<class IPropertyHandle> InStructPropertyHandle, class IDetailChildrenBuilder& StructBuilder, IPropertyTypeCustomizationUtils& StructCustomizationUtils) override;
+
+protected:
+
+	FString GetEntryName() const;
+	void SetEntryName(FString InName);
+	void UpdateEntryNameList();
+	void OnEntryNameChanged(TSharedPtr<FString> InItem, ESelectInfo::Type InSelectionInfo);
+	TSharedRef<SWidget> OnGetEntryNameWidget(TSharedPtr<FString> InItem);
+	FText GetEntryNameAsText() const;
+
+	TSharedPtr<IPropertyHandle> NameHandle;
+	TArray<TSharedPtr<FString>> EntryNameList;
+	TSharedPtr<SSearchableComboBox> SearchableComboBox;
+	FAnimNode_ControlRigBase* AnimNodeBeingCustomized;
+};
\ No newline at end of file
diff --git a/Engine/Plugins/Animation/ControlRig/Source/ControlRigEditor/Private/ControlRigEditorModule.cpp b/Engine/Plugins/Animation/ControlRig/Source/ControlRigEditor/Private/ControlRigEditorModule.cpp
--- a/Engine/Plugins/Animation/ControlRig/Source/ControlRigEditor/Private/ControlRigEditorModule.cpp
+++ b/Engine/Plugins/Animation/ControlRig/Source/ControlRigEditor/Private/ControlRigEditorModule.cpp
@@ -73,6 +73,7 @@
 #include "ControlRigCompilerDetails.h"
 #include "ControlRigDrawingDetails.h"
 #include "ControlRigGraphDetails.h"
+#include "ControlRigAnimGraphDetails.h"
 #include "ControlRigLocalVariableDetails.h"
 #include "ControlRigInfluenceMapDetails.h"
 #include "Animation/AnimSequence.h"
@@ -184,6 +185,9 @@ void FControlRigEditorModule::StartupModule()
 	PropertiesToUnregisterOnShutdown.Add(FRigComputedTransform::StaticStruct()->GetFName());
 	PropertyEditorModule.RegisterCustomPropertyTypeLayout(PropertiesToUnregisterOnShutdown.Last(), FOnGetPropertyTypeCustomizationInstance::CreateStatic(&FRigComputedTransformDetails::MakeInstance));
 
+	PropertiesToUnregisterOnShutdown.Add(FControlRigAnimNodeEventName::StaticStruct()->GetFName());
+	PropertyEditorModule.RegisterCustomPropertyTypeLayout(PropertiesToUnregisterOnShutdown.Last(), FOnGetPropertyTypeCustomizationInstance::CreateStatic(&FControlRigAnimNodeEventNameDetails::MakeInstance));
+
 	FRigBaseElementDetails::RegisterSectionMappings(PropertyEditorModule);
 
 	// Register asset tools
