ControlRig: Add missing dependency on AnimationEditMode.

#rnx

diff --git a/Engine/Plugins/Animation/ControlRig/Source/ControlRigEditor/ControlRigEditor.Build.cs b/Engine/Plugins/Animation/ControlRig/Source/ControlRigEditor/ControlRigEditor.Build.cs
--- a/Engine/Plugins/Animation/ControlRig/Source/ControlRigEditor/ControlRigEditor.Build.cs
+++ b/Engine/Plugins/Animation/ControlRig/Source/ControlRigEditor/ControlRigEditor.Build.cs
@@ -87,7 +87,8 @@
 					"ToolWidgets",
                     "AnimationWidgets",
                     "ActorPickerMode",
-                    "Constraints"
+                    "Constraints",
+                    "AnimationEditMode"
 				}
             );
 
