Control Rig: ControlRigComponent - fix static analysis

diff --git a/Engine/Plugins/Animation/ControlRig/Source/ControlRig/Private/ControlRigComponent.cpp b/Engine/Plugins/Animation/ControlRig/Source/ControlRig/Private/ControlRigComponent.cpp
--- a/Engine/Plugins/Animation/ControlRig/Source/ControlRig/Private/ControlRigComponent.cpp
+++ b/Engine/Plugins/Animation/ControlRig/Source/ControlRig/Private/ControlRigComponent.cpp
@@ -1917,7 +1917,7 @@ void FControlRigSceneProxy::GetDynamicMeshElements(const TArray<const FSceneView
 
 					const float BoneLength = (End - Start).Size();
 					// clamp by bound, we don't want too long or big
-					const float Radius = FMath::Clamp(BoneLength * 0.05f, 0.1f, MaxDrawRadius) * RadiusMultiplier;
+					const float Radius = FMath::Clamp<float>(BoneLength * 0.05f, 0.1f, MaxDrawRadius) * RadiusMultiplier;
 
 					//Render Sphere for bone end point and a cone between it and its parent.
 					SkeletalDebugRendering::DrawWireBone(PDI, Start, End, LineColor, SDPG_Foreground, Radius);
