Control Rig: Allow events to be invoked from top toolbar and context menu

diff --git a/Engine/Plugins/Animation/ControlRig/Source/ControlRig/Private/AnimNode_ControlRigBase.cpp b/Engine/Plugins/Animation/ControlRig/Source/ControlRig/Private/AnimNode_ControlRigBase.cpp
--- a/Engine/Plugins/Animation/ControlRig/Source/ControlRig/Private/AnimNode_ControlRigBase.cpp
+++ b/Engine/Plugins/Animation/ControlRig/Source/ControlRig/Private/AnimNode_ControlRigBase.cpp
@@ -508,7 +508,11 @@ void FAnimNode_ControlRigBase::ExecuteControlRig(FPoseContext& InOutput)
 			// pick the event to run
 			if(EventQueue.IsEmpty())
 			{
-				ControlRig->SetEventQueue({FRigUnit_BeginExecution::EventName});
+				if(bClearEventQueueRequired)
+				{
+					ControlRig->SetEventQueue({FRigUnit_BeginExecution::EventName});
+					bClearEventQueueRequired = false;
+				}
 			}
 			else
 			{
@@ -518,6 +522,7 @@ void FAnimNode_ControlRigBase::ExecuteControlRig(FPoseContext& InOutput)
 					return InEventName.EventName;
 				});
 				ControlRig->SetEventQueue(EventNames);
+				bClearEventQueueRequired = true;
 			}
 			
 			// evaluate control rig
diff --git a/Engine/Plugins/Animation/ControlRig/Source/ControlRig/Private/ControlRig.cpp b/Engine/Plugins/Animation/ControlRig/Source/ControlRig/Private/ControlRig.cpp
--- a/Engine/Plugins/Animation/ControlRig/Source/ControlRig/Private/ControlRig.cpp
+++ b/Engine/Plugins/Animation/ControlRig/Source/ControlRig/Private/ControlRig.cpp
@@ -275,22 +275,26 @@ void UControlRig::Evaluate_AnyThread()
 	DECLARE_SCOPE_HIERARCHICAL_COUNTER_FUNC()
 	QUICK_SCOPE_CYCLE_COUNTER(STAT_ControlRig_Evaluate);
 
-	TGuardValue<TArray<FName>> EventQueueGuard(EventQueue, EventQueue);
+	// create a copy since we need to change it here temporarily,
+	// and UI / the rig may change the event queue while it is running
+	TArray<FName> EventQueueToRun = EventQueue;
+	
 	if(InteractionType != (uint8)EControlRigInteractionType::None)
 	{
-		if(EventQueue.IsEmpty())
+		if(EventQueueToRun.IsEmpty())
 		{
-			EventQueue.Add(FRigUnit_InteractionExecution::EventName);
+			EventQueueToRun.Add(FRigUnit_InteractionExecution::EventName);
 		}
 		else
 		{
 			// insert just before the last event so the interaction runs prior to
 			// forward solve or backwards solve.
-			EventQueue.Insert(FRigUnit_InteractionExecution::EventName, EventQueue.Num() - 1);
+			EventQueueToRun.Insert(FRigUnit_InteractionExecution::EventName, EventQueueToRun.Num() - 1);
 		}
 	}
 
-	for (const FName& EventName : EventQueue)
+	
+	for (const FName& EventName : EventQueueToRun)
 	{
 		Execute(EControlRigState::Update, EventName);
 
diff --git a/Engine/Plugins/Animation/ControlRig/Source/ControlRig/Public/AnimNode_ControlRigBase.h b/Engine/Plugins/Animation/ControlRig/Source/ControlRig/Public/AnimNode_ControlRigBase.h
--- a/Engine/Plugins/Animation/ControlRig/Source/ControlRig/Public/AnimNode_ControlRigBase.h
+++ b/Engine/Plugins/Animation/ControlRig/Source/ControlRig/Public/AnimNode_ControlRigBase.h
@@ -155,6 +155,8 @@ protected:
 	UPROPERTY(EditAnywhere, AdvancedDisplay, Category = Settings)
 	TArray<FControlRigAnimNodeEventName> EventQueue;
 
+	bool bClearEventQueueRequired = false;
+
 	virtual bool CanExecute();
 	// update input/output to control rig
 	virtual void UpdateInput(UControlRig* ControlRig, const FPoseContext& InOutput);
diff --git a/Engine/Plugins/Animation/ControlRig/Source/ControlRigDeveloper/Private/Graph/ControlRigGraphSchema.cpp b/Engine/Plugins/Animation/ControlRig/Source/ControlRigDeveloper/Private/Graph/ControlRigGraphSchema.cpp
--- a/Engine/Plugins/Animation/ControlRig/Source/ControlRigDeveloper/Private/Graph/ControlRigGraphSchema.cpp
+++ b/Engine/Plugins/Animation/ControlRig/Source/ControlRigDeveloper/Private/Graph/ControlRigGraphSchema.cpp
@@ -28,6 +28,10 @@
 #include "RigVMCore/RigVMUnknownType.h"
 #include "Kismet2/Kismet2NameValidators.h"
 #include "Algo/Count.h"
+#include "ControlRig/Private/Units/Execution/RigUnit_BeginExecution.h"
+#include "ControlRig/Private/Units/Execution/RigUnit_PrepareForExecution.h"
+#include "ControlRig/Private/Units/Execution/RigUnit_InverseExecution.h"
+#include "ControlRig/Private/Units/Execution/RigUnit_InteractionExecution.h"
 
 #if WITH_EDITOR
 #include "ControlRigEditor/Private/Editor/SControlRigFunctionLocalizationWidget.h"
@@ -2064,4 +2068,12 @@ void UControlRigGraphSchema::HandleModifiedEvent(ERigVMGraphNotifType InNotifTyp
 	}
 }
 
+bool UControlRigGraphSchema::IsControlRigDefaultEvent(const FName& InEventName)
+{
+	return InEventName == FRigUnit_BeginExecution::EventName ||
+		InEventName == FRigUnit_InverseExecution::EventName ||
+		InEventName == FRigUnit_PrepareForExecution::EventName ||
+		InEventName == FRigUnit_InteractionExecution::EventName;
+}
+
 #undef LOCTEXT_NAMESPACE
diff --git a/Engine/Plugins/Animation/ControlRig/Source/ControlRigDeveloper/Public/Graph/ControlRigGraphSchema.h b/Engine/Plugins/Animation/ControlRig/Source/ControlRigDeveloper/Public/Graph/ControlRigGraphSchema.h
--- a/Engine/Plugins/Animation/ControlRig/Source/ControlRigDeveloper/Public/Graph/ControlRigGraphSchema.h
+++ b/Engine/Plugins/Animation/ControlRig/Source/ControlRigDeveloper/Public/Graph/ControlRigGraphSchema.h
@@ -266,6 +266,8 @@ public:
 	// Allow derived classes to spawn derived node classes
 	virtual TSubclassOf<UControlRigGraphNode> GetGraphNodeClass() const { return UControlRigGraphNode::StaticClass(); }
 
+	static bool IsControlRigDefaultEvent(const FName& InEventName);
+
 private:
 
 	const UEdGraphPin* LastPinForCompatibleCheck = nullptr;
diff --git a/Engine/Plugins/Animation/ControlRig/Source/ControlRigEditor/Private/ControlRigEditorModule.cpp b/Engine/Plugins/Animation/ControlRig/Source/ControlRigEditor/Private/ControlRigEditorModule.cpp
--- a/Engine/Plugins/Animation/ControlRig/Source/ControlRigEditor/Private/ControlRigEditorModule.cpp
+++ b/Engine/Plugins/Animation/ControlRig/Source/ControlRigEditor/Private/ControlRigEditorModule.cpp
@@ -2079,6 +2079,58 @@ void FControlRigEditorModule::GetContextMenuActions(const UControlRigGraphSchema
 				{
 					if (URigVMNode* ModelNode = RigNode->GetModelNode())
 					{
+						if(ModelNode->IsEvent())
+						{
+							const FName& EventName = ModelNode->GetEventName();
+							const bool bCanRunOnce = !UControlRigGraphSchema::IsControlRigDefaultEvent(EventName);
+
+							FToolMenuSection& EventsSection = Menu->AddSection("EdGraphSchemaEvents", LOCTEXT("EventsHeader", "Events"));
+
+							EventsSection.AddMenuEntry(
+								"Switch to Event",
+								LOCTEXT("SwitchToEvent", "Switch to Event"),
+								LOCTEXT("SwitchToEvent_Tooltip", "Switches the Control Rig Editor to run this event permanently."),
+								FSlateIcon(),
+								FUIAction(FExecuteAction::CreateLambda([RigBlueprint, EventName]()
+								{
+									if(UControlRig* ControlRig = Cast<UControlRig>(RigBlueprint->GetObjectBeingDebugged()))
+									{
+										ControlRig->SetEventQueue({EventName});
+									}
+								}))
+							);
+
+							EventsSection.AddMenuEntry(
+								"Run Event Once",
+								LOCTEXT("RuntEventOnce", "Run Event Once"),
+								LOCTEXT("RuntEventOnce_Tooltip", "Runs the event once (for testing)"),
+								FSlateIcon(),
+								FUIAction(FExecuteAction::CreateLambda([RigBlueprint, EventName]()
+								{
+									if(UControlRig* ControlRig = Cast<UControlRig>(RigBlueprint->GetObjectBeingDebugged()))
+									{
+										const TArray<FName> PreviousEventQueue = ControlRig->GetEventQueue();
+										TArray<FName> NewEventQueue = PreviousEventQueue;
+										NewEventQueue.Add(EventName);
+										ControlRig->SetEventQueue(NewEventQueue);
+
+										TSharedPtr<FDelegateHandle> RunOnceHandle = MakeShareable(new FDelegateHandle);
+										*(RunOnceHandle.Get()) = ControlRig->OnExecuted_AnyThread().AddLambda(
+											[RunOnceHandle, EventName, PreviousEventQueue](UControlRig* InRig, const EControlRigState InState, const FName& InEventName)
+											{
+												if(InEventName == EventName)
+												{
+													InRig->SetEventQueue(PreviousEventQueue);
+													InRig->OnExecuted_AnyThread().Remove(*RunOnceHandle.Get());
+												}
+											}
+										);
+									}
+								}),
+								FCanExecuteAction::CreateLambda([bCanRunOnce](){ return bCanRunOnce; }))
+							);
+						}
+						
 						const TArray<FRigVMUserWorkflow> Workflows = ModelNode->GetSupportedWorkflows(ERigVMUserWorkflowType::NodeContext, ModelNode);
 						if(!Workflows.IsEmpty())
 						{
diff --git a/Engine/Plugins/Animation/ControlRig/Source/ControlRigEditor/Private/Editor/ControlRigEditor.cpp b/Engine/Plugins/Animation/ControlRig/Source/ControlRigEditor/Private/Editor/ControlRigEditor.cpp
--- a/Engine/Plugins/Animation/ControlRig/Source/ControlRigEditor/Private/Editor/ControlRigEditor.cpp
+++ b/Engine/Plugins/Animation/ControlRig/Source/ControlRigEditor/Private/Editor/ControlRigEditor.cpp
@@ -98,6 +98,7 @@
 #include "ToolMenus.h"
 #include "Styling/AppStyle.h"
 #include "AssetRegistry/AssetRegistryModule.h"
+#include "ControlRig/Private/Units/Execution/RigUnit_InteractionExecution.h"
 
 #define LOCTEXT_NAMESPACE "ControlRigEditor"
 
@@ -203,6 +204,11 @@ struct FControlRigZoomLevelsContainer : public FZoomLevelsContainer
 	TArray<FControlRigZoomLevelEntry> ZoomLevels;
 };
 
+const TArray<FName> FControlRigEditor::ForwardsSolveEventQueue = {FRigUnit_BeginExecution::EventName};
+const TArray<FName> FControlRigEditor::BackwardsSolveEventQueue = {FRigUnit_InverseExecution::EventName};
+const TArray<FName> FControlRigEditor::SetupEventQueue = {FRigUnit_PrepareForExecution::EventName};
+const TArray<FName> FControlRigEditor::BackwardsAndForwardsSolveEventQueue = {FRigUnit_InverseExecution::EventName, FRigUnit_BeginExecution::EventName};
+
 FControlRigEditor::FControlRigEditor()
 	: ControlRig(nullptr)
 	, ActiveController(nullptr)
@@ -212,7 +218,7 @@ FControlRigEditor::FControlRigEditor()
 	, bSetupModeEnabled(false)
 	, bIsCompilingThroughUI(false)
 	, bAnyErrorsLeft(false)
-	, LastEventQueue(EControlRigEditorEventQueue::Setup)
+	, LastEventQueue(SetupEventQueue)
 	, ExecutionMode(EControlRigExecutionModeType::EControlRigExecutionModeType_Release)
 	, LastDebuggedRig()
 	, RigHierarchyTabCount(0)
@@ -657,22 +663,22 @@ void FControlRigEditor::BindCommands()
 
 	GetToolkitCommands()->MapAction(
 		FControlRigBlueprintCommands::Get().SetupEvent,
-		FExecuteAction::CreateSP(this, &FControlRigEditor::SetEventQueue, EControlRigEditorEventQueue::Setup),
+		FExecuteAction::CreateSP(this, &FControlRigEditor::SetEventQueue, TArray<FName>(SetupEventQueue)),
 		FCanExecuteAction());
 
 	GetToolkitCommands()->MapAction(
 		FControlRigBlueprintCommands::Get().ForwardsSolveEvent,
-		FExecuteAction::CreateSP(this, &FControlRigEditor::SetEventQueue, EControlRigEditorEventQueue::ForwardsSolve),
+		FExecuteAction::CreateSP(this, &FControlRigEditor::SetEventQueue, TArray<FName>(ForwardsSolveEventQueue)),
 		FCanExecuteAction());
 
 	GetToolkitCommands()->MapAction(
 		FControlRigBlueprintCommands::Get().BackwardsSolveEvent,
-		FExecuteAction::CreateSP(this, &FControlRigEditor::SetEventQueue, EControlRigEditorEventQueue::BackwardsSolve),
+		FExecuteAction::CreateSP(this, &FControlRigEditor::SetEventQueue, TArray<FName>(BackwardsSolveEventQueue)),
 		FCanExecuteAction());
 
 	GetToolkitCommands()->MapAction(
 		FControlRigBlueprintCommands::Get().BackwardsAndForwardsSolveEvent,
-		FExecuteAction::CreateSP(this, &FControlRigEditor::SetEventQueue, EControlRigEditorEventQueue::BackwardsAndForwardsSolve),
+		FExecuteAction::CreateSP(this, &FControlRigEditor::SetEventQueue, TArray<FName>(BackwardsAndForwardsSolveEventQueue)),
 		FCanExecuteAction());
 
 	GetToolkitCommands()->MapAction(
@@ -777,15 +783,50 @@ TSharedRef<SWidget> FControlRigEditor::GenerateEventQueueMenuContent()
 	FMenuBuilder MenuBuilder(true, GetToolkitCommands());
 
 	MenuBuilder.BeginSection(TEXT("Events"));
-	MenuBuilder.AddMenuEntry(FControlRigBlueprintCommands::Get().SetupEvent, TEXT("Setup"), TAttribute<FText>(), TAttribute<FText>(), GetEventQueueIcon(EControlRigEditorEventQueue::Setup));
-	MenuBuilder.AddMenuEntry(FControlRigBlueprintCommands::Get().ForwardsSolveEvent, TEXT("Forwards Solve"), TAttribute<FText>(), TAttribute<FText>(), GetEventQueueIcon(EControlRigEditorEventQueue::ForwardsSolve));
-	MenuBuilder.AddMenuEntry(FControlRigBlueprintCommands::Get().BackwardsSolveEvent, TEXT("Backwards Solve"), TAttribute<FText>(), TAttribute<FText>(), GetEventQueueIcon(EControlRigEditorEventQueue::BackwardsSolve));
+	MenuBuilder.AddMenuEntry(FControlRigBlueprintCommands::Get().SetupEvent, TEXT("Setup"), TAttribute<FText>(), TAttribute<FText>(), GetEventQueueIcon(SetupEventQueue));
+	MenuBuilder.AddMenuEntry(FControlRigBlueprintCommands::Get().ForwardsSolveEvent, TEXT("Forwards Solve"), TAttribute<FText>(), TAttribute<FText>(), GetEventQueueIcon(ForwardsSolveEventQueue));
+	MenuBuilder.AddMenuEntry(FControlRigBlueprintCommands::Get().BackwardsSolveEvent, TEXT("Backwards Solve"), TAttribute<FText>(), TAttribute<FText>(), GetEventQueueIcon(BackwardsSolveEventQueue));
 	MenuBuilder.EndSection();
 
 	MenuBuilder.BeginSection(TEXT("Validation"));
-	MenuBuilder.AddMenuEntry(FControlRigBlueprintCommands::Get().BackwardsAndForwardsSolveEvent, TEXT("BackwardsAndForwards"), TAttribute<FText>(), TAttribute<FText>(), GetEventQueueIcon(EControlRigEditorEventQueue::BackwardsAndForwardsSolve));
+	MenuBuilder.AddMenuEntry(FControlRigBlueprintCommands::Get().BackwardsAndForwardsSolveEvent, TEXT("BackwardsAndForwards"), TAttribute<FText>(), TAttribute<FText>(), GetEventQueueIcon(BackwardsAndForwardsSolveEventQueue));
 	MenuBuilder.EndSection();
 
+	if (const UControlRigBlueprint* RigBlueprint = Cast<UControlRigBlueprint>(GetBlueprintObj()))
+	{
+		bool bFoundUserDefinedEvent = false;
+		const TArray<FName> EntryNames = RigBlueprint->GetRigVMClient()->GetEntryNames();
+		for(const FName& EntryName : EntryNames)
+		{
+			if(UControlRigGraphSchema::IsControlRigDefaultEvent(EntryName))
+			{
+				continue;
+			}
+
+			if(!bFoundUserDefinedEvent)
+			{
+				MenuBuilder.AddSeparator();
+				bFoundUserDefinedEvent = true;
+			}
+
+			FString EventNameStr = EntryName.ToString();
+			if(!EventNameStr.EndsWith(TEXT("Event")))
+			{
+				EventNameStr += TEXT(" Event");
+			}
+
+			MenuBuilder.AddMenuEntry(
+				FText::FromString(EventNameStr),
+				FText::FromString(FString::Printf(TEXT("Runs the user defined %s"), *EventNameStr)),
+				GetEventQueueIcon({EntryName}),
+				FUIAction(
+					FExecuteAction::CreateSP(this, &FControlRigEditor::SetEventQueue, TArray<FName>({EntryName})),
+					FCanExecuteAction()
+				)
+			);
+		}
+	}
+
 	return MenuBuilder.MakeWidget();
 }
 
@@ -910,45 +951,22 @@ void FControlRigEditor::FillToolbar(FToolBarBuilder& ToolbarBuilder)
 	ToolbarBuilder.EndSection();
 }
 
-EControlRigEditorEventQueue FControlRigEditor::GetEventQueue() const
+TArray<FName> FControlRigEditor::GetEventQueue() const
 {
 	if (ControlRig)
 	{
 		if (bSetupModeEnabled)
 		{
-			return EControlRigEditorEventQueue::Setup;
+			return SetupEventQueue;
 		}
 
-		TArray<FName> EventQueue = ControlRig->EventQueue;
-		if (EventQueue.Num() == 1)
-		{
-			if (EventQueue[0] == FRigUnit_PrepareForExecution::EventName)
-			{
-				return EControlRigEditorEventQueue::Setup;
-			}
-			else if (EventQueue[0] == FRigUnit_BeginExecution::EventName)
-			{
-				return EControlRigEditorEventQueue::ForwardsSolve;
-			}
-			else if (EventQueue[0] == FRigUnit_InverseExecution::EventName)
-			{
-				return EControlRigEditorEventQueue::BackwardsSolve;
-			}
-		}
-		else if (EventQueue.Num() == 2)
-		{
-			if (EventQueue[0] == FRigUnit_InverseExecution::EventName &&
-				EventQueue[1] == FRigUnit_BeginExecution::EventName)
-			{
-				return EControlRigEditorEventQueue::BackwardsAndForwardsSolve;
-			}
-		}
+		return ControlRig->GetEventQueue();
 	}
 
-	return EControlRigEditorEventQueue::ForwardsSolve;
+	return ForwardsSolveEventQueue;
 }
 
-void FControlRigEditor::SetEventQueue(EControlRigEditorEventQueue InEventQueue)
+void FControlRigEditor::SetEventQueue(TArray<FName> InEventQueue)
 {
 	if (GetEventQueue() == InEventQueue)
 	{
@@ -961,10 +979,9 @@ void FControlRigEditor::SetEventQueue(EControlRigEditorEventQueue InEventQueue)
 
 	if (ControlRig)
 	{
-		TArray<FName> EventNames;
-		switch (InEventQueue)
+		if (InEventQueue.Num() > 0)
 		{
-			case EControlRigEditorEventQueue::Setup:
+			if(InEventQueue == SetupEventQueue)
 			{
 				if (!bSetupModeEnabled)
 				{
@@ -972,27 +989,8 @@ void FControlRigEditor::SetEventQueue(EControlRigEditorEventQueue InEventQueue)
 				}
 				return;
 			}
-			case EControlRigEditorEventQueue::ForwardsSolve:
-			{
-				EventNames.Add(FRigUnit_BeginExecution::EventName);
-				break;
-			}
-			case EControlRigEditorEventQueue::BackwardsSolve:
-			{
-				EventNames.Add(FRigUnit_InverseExecution::EventName);
-				break;
-			}
-			case EControlRigEditorEventQueue::BackwardsAndForwardsSolve:
-			{
-				EventNames.Add(FRigUnit_InverseExecution::EventName);
-				EventNames.Add(FRigUnit_BeginExecution::EventName);
-				break;
-			}
-		}
-
-		if (EventNames.Num() > 0)
-		{
-			ControlRig->SetEventQueue(EventNames);
+			
+			ControlRig->SetEventQueue(InEventQueue);
 
 			if (UControlRigBlueprint* RigBlueprint = Cast<UControlRigBlueprint>(GetBlueprintObj()))
 			{
@@ -1006,8 +1004,8 @@ void FControlRigEditor::SetEventQueue(EControlRigEditorEventQueue InEventQueue)
 		}
 
 		// Reset transforms only for setup and forward solve to not inturrupt any animation that might be playing
-		if (InEventQueue == EControlRigEditorEventQueue::Setup ||
-			InEventQueue == EControlRigEditorEventQueue::ForwardsSolve)
+		if (InEventQueue.Contains(FRigUnit_PrepareForExecution::EventName) ||
+			InEventQueue.Contains(FRigUnit_BeginExecution::EventName))
 		{
 			ControlRig->GetHierarchy()->ResetPoseToInitial(ERigElementType::All);
 		}
@@ -1016,59 +1014,81 @@ void FControlRigEditor::SetEventQueue(EControlRigEditorEventQueue InEventQueue)
 
 int32 FControlRigEditor::GetEventQueueComboValue() const
 {
-	return (int32)GetEventQueue();
+	const TArray<FName> EventQueue = GetEventQueue();
+	if(EventQueue == ForwardsSolveEventQueue)
+	{
+		return 0;
+	}
+	if(EventQueue == SetupEventQueue)
+	{
+		return 1;
+	}
+	if(EventQueue == BackwardsSolveEventQueue)
+	{
+		return 2;
+	}
+	if(EventQueue == BackwardsAndForwardsSolveEventQueue)
+	{
+		return 3;
+	}
+	return INDEX_NONE;
 }
 
 FText FControlRigEditor::GetEventQueueLabel() const
 {
-	EControlRigEditorEventQueue EventQueue = GetEventQueue();
-	switch (EventQueue)
+	TArray<FName> EventQueue = GetEventQueue();
+
+	if(EventQueue == SetupEventQueue)
 	{
-		case EControlRigEditorEventQueue::Setup:
-		{
-			return FRigUnit_PrepareForExecution::StaticStruct()->GetDisplayNameText();
-		}
-		case EControlRigEditorEventQueue::ForwardsSolve:
-		{
-			return FRigUnit_BeginExecution::StaticStruct()->GetDisplayNameText();
-		}
-		case EControlRigEditorEventQueue::BackwardsSolve:
-		{
-			return FRigUnit_InverseExecution::StaticStruct()->GetDisplayNameText();
-		}
-		case EControlRigEditorEventQueue::BackwardsAndForwardsSolve:
+		return FRigUnit_PrepareForExecution::StaticStruct()->GetDisplayNameText();
+	}
+	if(EventQueue == ForwardsSolveEventQueue)
+	{
+		return FRigUnit_BeginExecution::StaticStruct()->GetDisplayNameText();
+	}
+	if(EventQueue == BackwardsSolveEventQueue)
+	{
+		return FRigUnit_InverseExecution::StaticStruct()->GetDisplayNameText();
+	}
+	if(EventQueue == BackwardsAndForwardsSolveEventQueue)
+	{
+		return FText::FromString(FString::Printf(TEXT("%s and %s"),
+			*FRigUnit_InverseExecution::StaticStruct()->GetDisplayNameText().ToString(),
+			*FRigUnit_BeginExecution::StaticStruct()->GetDisplayNameText().ToString()));
+	}
+
+	if(EventQueue.Num() == 1)
+	{
+		FString EventName = EventQueue[0].ToString();
+		if(!EventName.EndsWith(TEXT("Event")))
 		{
-			return FText::FromString(FString::Printf(TEXT("%s and %s"),
-				*FRigUnit_InverseExecution::StaticStruct()->GetDisplayNameText().ToString(),
-				*FRigUnit_BeginExecution::StaticStruct()->GetDisplayNameText().ToString()));
+			EventName += TEXT(" Event");
 		}
+		return FText::FromString(EventName);
 	}
-	return StaticEnum<EControlRigEditorEventQueue>()->GetDisplayNameTextByValue((int64)EventQueue);
+	return LOCTEXT("CustomEventQueue", "Custom Event Queue");
 }
 
-FSlateIcon FControlRigEditor::GetEventQueueIcon(EControlRigEditorEventQueue InEventQueue)
+FSlateIcon FControlRigEditor::GetEventQueueIcon(const TArray<FName>& InEventQueue)
 {
-	switch (InEventQueue)
+	if(InEventQueue == SetupEventQueue)
 	{
-		case EControlRigEditorEventQueue::Setup:
-		{
-			return FSlateIcon(FControlRigEditorStyle::Get().GetStyleSetName(), "ControlRig.SetupMode");
-		}
-		case EControlRigEditorEventQueue::ForwardsSolve:
-		{
-			return FSlateIcon(FControlRigEditorStyle::Get().GetStyleSetName(), "ControlRig.ForwardsSolveEvent");
-		}
-		case EControlRigEditorEventQueue::BackwardsSolve:
-		{
-			return FSlateIcon(FControlRigEditorStyle::Get().GetStyleSetName(), "ControlRig.BackwardsSolveEvent");
-		}
-		case EControlRigEditorEventQueue::BackwardsAndForwardsSolve:
-		{
-			return FSlateIcon(FControlRigEditorStyle::Get().GetStyleSetName(), "ControlRig.BackwardsAndForwardsSolveEvent");
-		}
+		return FSlateIcon(FControlRigEditorStyle::Get().GetStyleSetName(), "ControlRig.SetupMode");
+	}
+	if(InEventQueue == ForwardsSolveEventQueue)
+	{
+		return FSlateIcon(FControlRigEditorStyle::Get().GetStyleSetName(), "ControlRig.ForwardsSolveEvent");
+	}
+	if(InEventQueue == BackwardsSolveEventQueue)
+	{
+		return FSlateIcon(FControlRigEditorStyle::Get().GetStyleSetName(), "ControlRig.BackwardsSolveEvent");
+	}
+	if(InEventQueue == BackwardsAndForwardsSolveEventQueue)
+	{
+		return FSlateIcon(FControlRigEditorStyle::Get().GetStyleSetName(), "ControlRig.BackwardsAndForwardsSolveEvent");
 	}
 
-	return FSlateIcon();
+	return FSlateIcon(FAppStyle::GetAppStyleSetName(), "GraphEditor.Event_16x");
 }
 
 FSlateIcon FControlRigEditor::GetEventQueueIcon() const
@@ -3691,20 +3711,18 @@ void FControlRigEditor::HandleViewportCreated(const TSharedRef<class IPersonaVie
 	auto GetBorderColorAndOpacity = [this]()
 	{
 		FLinearColor Color = FLinearColor::Transparent;
-		switch (GetEventQueue())
+		const TArray<FName> EventQueue = GetEventQueue();
+		if(EventQueue == SetupEventQueue)
 		{
-			case EControlRigEditorEventQueue::Setup:
-				Color = UControlRigEditorSettings::Get()->SetupEventBorderColor;
-				break;
-			case EControlRigEditorEventQueue::BackwardsSolve:
-				Color = UControlRigEditorSettings::Get()->BackwardsSolveBorderColor;
-				break;
-			case EControlRigEditorEventQueue::BackwardsAndForwardsSolve:
-				Color = UControlRigEditorSettings::Get()->BackwardsAndForwardsBorderColor;
-				break;
-			default:
-				Color = FLinearColor::Transparent;
-				break;
+			Color = UControlRigEditorSettings::Get()->SetupEventBorderColor;
+		}
+		if(EventQueue == BackwardsSolveEventQueue)
+		{
+			Color = UControlRigEditorSettings::Get()->BackwardsSolveBorderColor;
+		}
+		if(EventQueue == BackwardsAndForwardsSolveEventQueue)
+		{
+			Color = UControlRigEditorSettings::Get()->BackwardsAndForwardsBorderColor;
 		}
 		return Color;
 	};
@@ -3712,7 +3730,7 @@ void FControlRigEditor::HandleViewportCreated(const TSharedRef<class IPersonaVie
 	auto GetBorderVisibility = [this]()
 	{
 		EVisibility Visibility = EVisibility::Collapsed;
-		if (GetEventQueue() != EControlRigEditorEventQueue::ForwardsSolve)
+		if (GetEventQueueComboValue() != 0)
 		{
 			Visibility = EVisibility::HitTestInvisible;
 		}
diff --git a/Engine/Plugins/Animation/ControlRig/Source/ControlRigEditor/Private/Editor/ControlRigEditor.h b/Engine/Plugins/Animation/ControlRig/Source/ControlRigEditor/Private/Editor/ControlRigEditor.h
--- a/Engine/Plugins/Animation/ControlRig/Source/ControlRigEditor/Private/Editor/ControlRigEditor.h
+++ b/Engine/Plugins/Animation/ControlRig/Source/ControlRigEditor/Private/Editor/ControlRigEditor.h
@@ -33,25 +33,6 @@ class FControlRigEditor;
 
 DECLARE_MULTICAST_DELEGATE_TwoParams(FControlRigEditorClosed, const FControlRigEditor*, UControlRigBlueprint*);
 
-UENUM()
-enum class EControlRigEditorEventQueue : uint8
-{
-	/** Setup Event */
-	Setup,
-
-	/** Forwards Solve Event */
-	ForwardsSolve,
-
-	/** Backwards Solve Event */
-	BackwardsSolve,
-
-	/** Backwards Solve Event -> Forwards Solve Event */
-	BackwardsAndForwardsSolve,
-
-	/** MAX - invalid */
-	Max UMETA(Hidden),
-};
-
 struct FControlRigEditorModes
 {
 	// Mode constants
@@ -284,11 +265,11 @@ private:
 	/** Fill the toolbar with content */
 	void FillToolbar(FToolBarBuilder& ToolbarBuilder);
 
-	EControlRigEditorEventQueue GetEventQueue() const;
-	void SetEventQueue(EControlRigEditorEventQueue InEventQueue);
+	TArray<FName> GetEventQueue() const;
+	void SetEventQueue(TArray<FName> InEventQueue);
 	int32 GetEventQueueComboValue() const;
 	FText GetEventQueueLabel() const;
-	static FSlateIcon GetEventQueueIcon(EControlRigEditorEventQueue InEventQueue);
+	static FSlateIcon GetEventQueueIcon(const TArray<FName>& InEventQueue);
 	FSlateIcon GetEventQueueIcon() const;
 
 	enum EControlRigExecutionModeType
@@ -491,7 +472,7 @@ protected:
 
 	bool bAnyErrorsLeft;
 
-	EControlRigEditorEventQueue LastEventQueue;
+	TArray<FName> LastEventQueue;
 	EControlRigExecutionModeType ExecutionMode;
 	FString LastDebuggedRig;
 	int32 RigHierarchyTabCount;
@@ -505,7 +486,12 @@ protected:
 	TWeakObjectPtr<AStaticMeshActor> WeakGroundActorPtr;
 
 	FDelegateHandle PropertyChangedHandle;
-	
+
+	static const TArray<FName> ForwardsSolveEventQueue;
+	static const TArray<FName> BackwardsSolveEventQueue;
+	static const TArray<FName> SetupEventQueue;
+	static const TArray<FName> BackwardsAndForwardsSolveEventQueue;
+
 	friend class FControlRigEditorMode;
 	friend class SControlRigStackView;
 	friend class SRigHierarchy;
diff --git a/Engine/Plugins/Animation/ControlRig/Source/ControlRigEditor/Private/Editor/SRigHierarchy.cpp b/Engine/Plugins/Animation/ControlRig/Source/ControlRigEditor/Private/Editor/SRigHierarchy.cpp
--- a/Engine/Plugins/Animation/ControlRig/Source/ControlRigEditor/Private/Editor/SRigHierarchy.cpp
+++ b/Engine/Plugins/Animation/ControlRig/Source/ControlRigEditor/Private/Editor/SRigHierarchy.cpp
@@ -591,7 +591,7 @@ void SRigHierarchy::OnSelectionChanged(TSharedPtr<FRigTreeElement> Selection, ES
 		{
 			if (ControlRigEditor.IsValid())
 			{
-				if (ControlRigEditor.Pin()->GetEventQueue() == EControlRigEditorEventQueue::Setup)
+				if (ControlRigEditor.Pin()->GetEventQueueComboValue() == 1)
 				{
 					HandleControlBoneOrSpaceTransform();
 				}
@@ -740,7 +740,7 @@ void SRigHierarchy::OnHierarchyModified(ERigHierarchyNotification InNotif, URigH
 
 						if (ControlRigEditor.IsValid() && !GIsTransacting)
 						{
-							if (ControlRigEditor.Pin()->GetEventQueue() == EControlRigEditorEventQueue::Setup)
+							if (ControlRigEditor.Pin()->GetEventQueueComboValue() == 1)
 							{
 								TGuardValue<bool> GuardRigHierarchyChanges(bIsChangingRigHierarchy, true);
 								HandleControlBoneOrSpaceTransform();
diff --git a/Engine/Source/Developer/RigVMDeveloper/Private/RigVMModel/RigVMController.cpp b/Engine/Source/Developer/RigVMDeveloper/Private/RigVMModel/RigVMController.cpp
--- a/Engine/Source/Developer/RigVMDeveloper/Private/RigVMModel/RigVMController.cpp
+++ b/Engine/Source/Developer/RigVMDeveloper/Private/RigVMModel/RigVMController.cpp
@@ -13011,7 +13011,7 @@ bool URigVMController::CanAddEventNode(UScriptStruct* InScriptStruct, const bool
 	}
 
 	TObjectPtr<URigVMNode> EventNode = FindEventNode(InScriptStruct);
-	const bool bHasEventNode = EventNode != nullptr;
+	const bool bHasEventNode = (EventNode != nullptr) && EventNode->CanOnlyExistOnce();
 	if (bHasEventNode && bReportErrors)
 	{
 		const FString ErrorMessage = FString::Printf(TEXT("Rig Graph can only contain one single %s node."),
