Control Rig: Show events in Action Palette

diff --git a/Engine/Plugins/Animation/ControlRig/Source/ControlRigDeveloper/Private/Graph/ControlRigGraphSchema.cpp b/Engine/Plugins/Animation/ControlRig/Source/ControlRigDeveloper/Private/Graph/ControlRigGraphSchema.cpp
--- a/Engine/Plugins/Animation/ControlRig/Source/ControlRigDeveloper/Private/Graph/ControlRigGraphSchema.cpp
+++ b/Engine/Plugins/Animation/ControlRig/Source/ControlRigDeveloper/Private/Graph/ControlRigGraphSchema.cpp
@@ -386,6 +386,34 @@ UEdGraphNode* FControlRigGraphSchemaAction_PromoteToVariable::PerformAction(UEdG
 	return nullptr;
 }
 
+FControlRigGraphSchemaAction_Event::FControlRigGraphSchemaAction_Event(const FName& InEventName, const FString& InNodePath, const FText& InNodeCategory)
+: FEdGraphSchemaAction(	InNodeCategory, 
+						FText::FromName(InEventName),
+						FText(),
+						1)
+, NodePath(InNodePath)
+{
+}
+
+FReply FControlRigGraphSchemaAction_Event::OnDoubleClick(UBlueprint* InBlueprint)
+{
+	if (UControlRigBlueprint* Blueprint = Cast<UControlRigBlueprint>(InBlueprint))
+	{
+		if(const URigVMNode* ModelNode = Blueprint->GetRigVMClient()->FindNode(NodePath))
+		{
+			Blueprint->OnRequestJumpToHyperlink().Execute(ModelNode);
+			return FReply::Handled();
+		}
+	}
+	return FReply::Unhandled();;
+}
+
+FSlateBrush const* FControlRigGraphSchemaAction_Event::GetPaletteIcon() const
+{
+	static FSlateIcon EventIcon(FAppStyle::GetAppStyleSetName(), "GraphEditor.Event_16x");
+	return EventIcon.GetIcon();
+}
+
 FReply FControlRigFunctionDragDropAction::DroppedOnPanel(const TSharedRef< class SWidget >& Panel, FVector2D ScreenPosition, FVector2D GraphPosition, UEdGraph& Graph)
 {
 	// For local variables
@@ -1513,6 +1541,62 @@ bool UControlRigGraphSchema::TryRenameGraph(UEdGraph* GraphToRename, const FName
 	return false;
 }
 
+bool UControlRigGraphSchema::TryToGetChildEvents(const UEdGraph* Graph, const int32 SectionId,
+	TArray<TSharedPtr<FEdGraphSchemaAction>>& Actions, const FText& ParentCategory) const
+{
+	check(Graph);
+
+	if(const UControlRigGraph* RigGraph = Cast<UControlRigGraph>(Graph))
+	{
+		if(URigVMGraph* Model = RigGraph->GetModel())
+		{
+			TArray<FName> EventNames;
+			TMap<FName, const URigVMNode*> EventNameToNode;
+			for(const URigVMNode* Node : Model->GetNodes())
+			{
+				if(Node->IsEvent())
+				{
+					if(!EventNames.Contains(Node->GetEventName()))
+					{
+						EventNames.Add(Node->GetEventName());
+						EventNameToNode.Add(Node->GetEventName(), Node);
+					}
+				}
+			}
+
+			if(!EventNames.IsEmpty())
+			{
+				FGraphDisplayInfo EdGraphDisplayInfo;
+				GetGraphDisplayInformation(*Graph, EdGraphDisplayInfo);
+
+				FText EdGraphDisplayName = EdGraphDisplayInfo.DisplayName;
+				FText ActionCategory;
+				if (!ParentCategory.IsEmpty())
+				{
+					ActionCategory = FText::Format(FText::FromString(TEXT("{0}|{1}")), ParentCategory, EdGraphDisplayName);
+				}
+				else
+				{
+					ActionCategory = MoveTemp(EdGraphDisplayName);
+				}
+
+				for(const FName& EventName : EventNames)
+				{
+					const URigVMNode* Node = EventNameToNode.FindChecked(EventName);
+					TSharedPtr<FEdGraphSchemaAction> Action = MakeShareable(
+						new FControlRigGraphSchemaAction_Event(EventName, Node->GetNodePath(true), ActionCategory)
+					);
+					Action->SectionID = SectionId;
+					Actions.Add(Action);
+				}
+			}
+			return true;
+		}
+	}
+
+	return false;
+}
+
 UEdGraphPin* UControlRigGraphSchema::DropPinOnNode(UEdGraphNode* InTargetNode, const FName& InSourcePinName, const FEdGraphPinType& InSourcePinType, EEdGraphPinDirection InSourcePinDirection) const
 {
 	FString NewPinName;
diff --git a/Engine/Plugins/Animation/ControlRig/Source/ControlRigDeveloper/Public/Graph/ControlRigGraphSchema.h b/Engine/Plugins/Animation/ControlRig/Source/ControlRigDeveloper/Public/Graph/ControlRigGraphSchema.h
--- a/Engine/Plugins/Animation/ControlRig/Source/ControlRigDeveloper/Public/Graph/ControlRigGraphSchema.h
+++ b/Engine/Plugins/Animation/ControlRig/Source/ControlRigDeveloper/Public/Graph/ControlRigGraphSchema.h
@@ -141,6 +141,37 @@ private:
 	bool bLocalVariable;
 };
 
+USTRUCT()
+struct CONTROLRIGDEVELOPER_API FControlRigGraphSchemaAction_Event : public FEdGraphSchemaAction
+{
+	GENERATED_BODY()
+
+public:
+
+	// Simple type info
+	static FName StaticGetTypeId() {static FName Type("FControlRigGraphSchemaAction_Event"); return Type;}
+	virtual FName GetTypeId() const override { return StaticGetTypeId(); } 
+
+	FControlRigGraphSchemaAction_Event()
+		: FEdGraphSchemaAction()
+	{}
+
+	FControlRigGraphSchemaAction_Event(const FName& InEventName, const FString& InNodePath, const FText& InNodeCategory);
+
+	virtual bool IsA(const FName& InType) const override
+	{
+		return InType == GetTypeId() || Super::IsA(InType);
+	}
+	virtual bool IsParentable() const override { return true; }
+	virtual FReply OnDoubleClick(UBlueprint* InBlueprint) override;
+	virtual bool CanBeRenamed() const override { return false; }
+	virtual FSlateBrush const* GetPaletteIcon() const override;
+
+private:
+
+	FString NodePath;
+};
+
 /** DragDropAction class for drag and dropping an item from the My Blueprints tree (e.g., variable or function) */
 class CONTROLRIGDEVELOPER_API FControlRigFunctionDragDropAction : public FGraphSchemaActionDragDropAction
 {
@@ -231,6 +262,7 @@ public:
 	virtual FReply TrySetGraphCategory(const UEdGraph* InGraph, const FText& InCategory) override;
 	virtual bool TryDeleteGraph(UEdGraph* GraphToDelete) const override;
 	virtual bool TryRenameGraph(UEdGraph* GraphToRename, const FName& InNewName) const override;
+	virtual bool TryToGetChildEvents(const UEdGraph* Graph, const int32 SectionId, TArray<TSharedPtr<FEdGraphSchemaAction>>& Actions, const FText& ParentCategory) const override;
 	virtual bool CanDuplicateGraph(UEdGraph* InSourceGraph) const { return false; }
 	virtual UEdGraphPin* DropPinOnNode(UEdGraphNode* InTargetNode, const FName& InSourcePinName, const FEdGraphPinType& InSourcePinType, EEdGraphPinDirection InSourcePinDirection) const override;
 	virtual bool SupportsDropPinOnNode(UEdGraphNode* InTargetNode, const FEdGraphPinType& InSourcePinType, EEdGraphPinDirection InSourcePinDirection, FText& OutErrorMessage) const override;
diff --git a/Engine/Source/Editor/Kismet/Private/SBlueprintPalette.cpp b/Engine/Source/Editor/Kismet/Private/SBlueprintPalette.cpp
--- a/Engine/Source/Editor/Kismet/Private/SBlueprintPalette.cpp
+++ b/Engine/Source/Editor/Kismet/Private/SBlueprintPalette.cpp
@@ -413,6 +413,15 @@ static void GetPaletteItemIcon(TSharedPtr<FEdGraphSchemaAction> ActionIn, UBluep
 		BrushOut = FAppStyle::GetBrush(TEXT("GraphEditor.StructGlyph"));
 		ToolTipOut = LOCTEXT("Struct_Tooltip", "Struct Asset");
 	}
+	else
+	{
+		BrushOut = ActionIn->GetPaletteIcon();
+		const FText ActionToolTip = ActionIn->GetPaletteToolTip();
+		if(!ActionToolTip.IsEmpty())
+		{
+			ToolTipOut = ActionToolTip;
+		}
+	}
 }
 
 /**
diff --git a/Engine/Source/Editor/Kismet/Private/SMyBlueprint.cpp b/Engine/Source/Editor/Kismet/Private/SMyBlueprint.cpp
--- a/Engine/Source/Editor/Kismet/Private/SMyBlueprint.cpp
+++ b/Engine/Source/Editor/Kismet/Private/SMyBlueprint.cpp
@@ -1035,6 +1035,10 @@ bool SMyBlueprint::CanRequestRenameOnActionNode(TWeakPtr<FGraphActionNode> InSel
 	{
 		check( InSelectedNode.Pin()->Actions.Num() > 0 && InSelectedNode.Pin()->Actions[0].IsValid() );
 		bIsReadOnly = FBlueprintEditorUtils::IsPaletteActionReadOnly(InSelectedNode.Pin()->Actions[0], BlueprintEditorPtr.Pin());
+		if(!bIsReadOnly)
+		{
+			bIsReadOnly = !InSelectedNode.Pin()->Actions[0]->CanBeRenamed();
+		}
 	}
 
 	return IsEditingMode() && !bIsReadOnly;
@@ -1115,6 +1119,20 @@ void SMyBlueprint::GetChildEvents(UEdGraph const* InEdGraph, int32 const Section
 		return;
 	}
 
+	// ask the schema first to provide the child events
+	if (UEdGraphSchema const* Schema = InEdGraph->GetSchema())
+	{
+		TArray<TSharedPtr<FEdGraphSchemaAction>> ActionsFromSchema;
+		if(Schema->TryToGetChildEvents(InEdGraph, SectionId, ActionsFromSchema, ParentCategory))
+		{
+			for(TSharedPtr<FEdGraphSchemaAction> ActionFromSchema : ActionsFromSchema)
+			{
+				SortList.AddAction(ActionFromSchema);
+			}
+			return;
+		}
+	}
+
 	// grab the parent graph's name
 	FGraphDisplayInfo EdGraphDisplayInfo;
 	if (UEdGraphSchema const* Schema = InEdGraph->GetSchema())
@@ -2145,6 +2163,10 @@ void SMyBlueprint::ExecuteAction(TSharedPtr<FEdGraphSchemaAction> InAction)
 			FEdGraphSchemaAction_K2TargetNode* TargetNodeAction = (FEdGraphSchemaAction_K2TargetNode*)InAction.Get();
 			FKismetEditorUtilities::BringKismetToFocusAttentionOnObject(TargetNodeAction->NodeTemplate);
 		}
+		else
+		{
+			InAction->OnDoubleClick(Blueprint);
+		}
 	}
 }
 
@@ -3170,6 +3192,11 @@ bool SMyBlueprint::CanDeleteEntry() const
 			return true;
 		}
 	}
+	else if (FEdGraphSchemaAction* Action = SelectionAsType<FEdGraphSchemaAction>(GraphActionMenu))
+	{
+		return Action->CanBeDeleted();
+	}
+
 	return false;
 }
 
diff --git a/Engine/Source/Runtime/Engine/Classes/EdGraph/EdGraphSchema.h b/Engine/Source/Runtime/Engine/Classes/EdGraph/EdGraphSchema.h
--- a/Engine/Source/Runtime/Engine/Classes/EdGraph/EdGraphSchema.h
+++ b/Engine/Source/Runtime/Engine/Classes/EdGraph/EdGraphSchema.h
@@ -189,6 +189,9 @@ public:
 		return NewNode;
 	}
 
+	/** Performs a double click on the action */
+	virtual FReply OnDoubleClick(UBlueprint* InBlueprint) { return FReply::Unhandled(); }
+
 	// Updates the category of the *action* and refreshes the search text; does not change the persistent backing item
 	// (e.g., it will not actually move a user added variable or function to a new category)
 	void CosmeticUpdateCategory(FText NewCategory);
@@ -301,6 +304,18 @@ public:
 	// Returns true if the action refers to a member or local variable
 	virtual bool IsAVariable() const { return false; }
 
+	// Returns true if the action can be renamed
+	virtual bool CanBeRenamed() const { return true; }
+
+	// Returns true if the action can be deleted
+	virtual bool CanBeDeleted() const { return false; }
+
+	// Can be used to override the icon of the action in the palette
+	virtual FSlateBrush const* GetPaletteIcon() const { return nullptr; }
+
+	// Can be used to override the tooltip shown in the palette
+	virtual FText GetPaletteToolTip() const { return FText(); }
+
 private:
 	void UpdateSearchText();
 };
@@ -1031,6 +1046,11 @@ class ENGINE_API UEdGraphSchema : public UObject
 	 */
 	virtual bool TryRenameGraph(UEdGraph* GraphToRename, const FName& InNewName) const { return false; }
 
+	/*
+	 * Try to retrieve the event child actions for a given graph
+	 */
+	virtual bool TryToGetChildEvents(const UEdGraph* Graph, const int32 SectionId, TArray<TSharedPtr<FEdGraphSchemaAction>>& Actions, const FText& ParentCategory) const { return false; }
+
 	/**
 	 * Can TestNode be encapsulated into a child graph?
 	 */
