Control Rig: Baking Fix for renaming of inversion event. Checking old name also now

diff --git a/Engine/Plugins/Animation/ControlRig/Source/ControlRigEditor/Private/Sequencer/ControlRigParameterTrackEditor.cpp b/Engine/Plugins/Animation/ControlRig/Source/ControlRigEditor/Private/Sequencer/ControlRigParameterTrackEditor.cpp
--- a/Engine/Plugins/Animation/ControlRig/Source/ControlRigEditor/Private/Sequencer/ControlRigParameterTrackEditor.cpp
+++ b/Engine/Plugins/Animation/ControlRig/Source/ControlRigEditor/Private/Sequencer/ControlRigParameterTrackEditor.cpp
@@ -630,12 +630,13 @@ public:
 			if (Tag.IsSet())
 			{
 				FString EventString = FRigUnit_InverseExecution::EventName.ToString();
+				FString OldEventString = FString(TEXT("Inverse"));
 				TArray<FString> SupportedEventNames;
 				Tag.GetValue().ParseIntoArray(SupportedEventNames, TEXT(","), true);
 
 				for (const FString& Name : SupportedEventNames)
 				{
-					if (Name.Contains(EventString))
+					if (Name.Contains(EventString) || Name.Contains(OldEventString))
 					{
 						bHasInversion = true;
 						break;
