ControlRig: Harden against lookups in an empty event queue.

#rnx

diff --git a/Engine/Plugins/Animation/ControlRig/Source/ControlRig/Private/ControlRig.cpp b/Engine/Plugins/Animation/ControlRig/Source/ControlRig/Private/ControlRig.cpp
--- a/Engine/Plugins/Animation/ControlRig/Source/ControlRig/Private/ControlRig.cpp
+++ b/Engine/Plugins/Animation/ControlRig/Source/ControlRig/Private/ControlRig.cpp
@@ -614,6 +614,10 @@ void UControlRig::Execute(const EControlRigState InState, const FName& InEventNa
 	{
 		return;
 	}
+	
+	const bool bIsEventInQueue = GetEventQueue().Contains(InEventName);
+	const bool bIsEventFirstInQueue = !GetEventQueue().IsEmpty() && GetEventQueue()[0] == InEventName; 
+	const bool bIsEventLastInQueue = !GetEventQueue().IsEmpty() && GetEventQueue().Last() == InEventName; 
 
 	ensure(!HasAnyFlags(RF_ClassDefaultObject));
 	
@@ -707,7 +711,7 @@ void UControlRig::Execute(const EControlRigState InState, const FName& InEventNa
 
 	FRigUnitContext Context;
 
-	if(!GetEventQueue().Contains(InEventName) || (GetEventQueue()[0] == InEventName))
+	if(!bIsEventInQueue || bIsEventFirstInQueue)
 	{
 		DrawInterface.Reset();
 	}
@@ -1012,8 +1016,7 @@ void UControlRig::Execute(const EControlRigState InState, const FName& InEventNa
 	}
 	else if (InState == EControlRigState::Update)
 	{
-		if(InEventName == GetEventQueue().Last() ||
-			!GetEventQueue().Contains(InEventName))
+		if(!bIsEventInQueue || bIsEventLastInQueue) 
 		{
 			DeltaTime = 0.f;
 		}
@@ -1025,7 +1028,7 @@ void UControlRig::Execute(const EControlRigState InState, const FName& InEventNa
 		}
 	}
 
-	if (Context.DrawInterface && Context.DrawContainer && GetEventQueue().Last() == InEventName)
+	if (Context.DrawInterface && Context.DrawContainer && bIsEventLastInQueue) 
 	{
 		Context.DrawInterface->Instructions.Append(Context.DrawContainer->Instructions);
 
