Sequencer: Don't AddControlRig() if that track already exists.

This fixes an issue where dragging a control rig into Sequencer will create a control rig track for the actor and all of its skeletal mesh components.

diff --git a/Engine/Plugins/Animation/ControlRig/Source/ControlRigEditor/Private/Sequencer/ControlRigParameterTrackEditor.cpp b/Engine/Plugins/Animation/ControlRig/Source/ControlRigEditor/Private/Sequencer/ControlRigParameterTrackEditor.cpp
--- a/Engine/Plugins/Animation/ControlRig/Source/ControlRigEditor/Private/Sequencer/ControlRigParameterTrackEditor.cpp
+++ b/Engine/Plugins/Animation/ControlRig/Source/ControlRigEditor/Private/Sequencer/ControlRigParameterTrackEditor.cpp
@@ -1723,7 +1723,7 @@ void FControlRigParameterTrackEditor::OnChannelChanged(const FMovieSceneChannelM
 	}
 }
 
-void FControlRigParameterTrackEditor::AddTrackForComponent(USceneComponent* InComponent)
+void FControlRigParameterTrackEditor::AddTrackForComponent(USceneComponent* InComponent, FGuid InBinding)
 {
 	if (USkeletalMeshComponent* SkelMeshComp = Cast<USkeletalMeshComponent>(InComponent))
 	{
@@ -1732,7 +1732,7 @@ void FControlRigParameterTrackEditor::AddTrackForComponent(USceneComponent* InCo
 			UObject* Object = SkelMeshComp->GetDefaultAnimatingRig().LoadSynchronous();
 			if (Object != nullptr && (Object->IsA<UControlRigBlueprint>() || Object->IsA<UControlRigComponent>()))
 			{
-				FGuid Binding = GetSequencer()->GetHandleToObject(InComponent, true /*bCreateHandle*/);
+				FGuid Binding = InBinding.IsValid() ? InBinding : GetSequencer()->GetHandleToObject(InComponent, true /*bCreateHandle*/);
 				if (Binding.IsValid())
 				{
 					UMovieSceneSequence* OwnerSequence = GetSequencer()->GetFocusedMovieSceneSequence();
@@ -1759,7 +1759,7 @@ void FControlRigParameterTrackEditor::AddTrackForComponent(USceneComponent* InCo
 	InComponent->GetChildrenComponents(false, ChildComponents);
 	for (USceneComponent* ChildComponent : ChildComponents)
 	{
-		AddTrackForComponent(ChildComponent);
+		AddTrackForComponent(ChildComponent, FGuid());
 	}
 }
 
@@ -1772,11 +1772,18 @@ void FControlRigParameterTrackEditor::HandleActorAdded(AActor* Actor, FGuid Targ
 			AddControlRigFromComponent(TargetObjectGuid);
 			return;
 		}
+
+		if (USkeletalMeshComponent* SkeletalMeshComponent = Cast<USkeletalMeshComponent>(Actor->GetRootComponent()))
+		{
+			AddTrackForComponent(SkeletalMeshComponent, TargetObjectGuid);
+			return;
+		}
+
 		for (UActorComponent* Component : Actor->GetComponents())
 		{
 			if (USceneComponent* SceneComp = Cast<USceneComponent>(Component))
 			{
-				AddTrackForComponent(SceneComp);
+				AddTrackForComponent(SceneComp, FGuid());
 			}
 		}
 	}
@@ -3076,6 +3083,12 @@ bool FControlRigParameterTrackEditor::HandleAssetAdded(UObject* Asset, const FGu
 		return false;
 	}
 
+	UMovieScene* MovieScene = GetSequencer()->GetFocusedMovieSceneSequence()->GetMovieScene();
+	if (!MovieScene)
+	{
+		return false;
+	}
+
 	UControlRigBlueprint* ControlRigBlueprint = Cast<UControlRigBlueprint>(Asset);
 	UControlRigBlueprintGeneratedClass* RigClass = ControlRigBlueprint->GetControlRigBlueprintGeneratedClass();
 	if (!RigClass)
@@ -3111,13 +3124,17 @@ bool FControlRigParameterTrackEditor::HandleAssetAdded(UObject* Asset, const FGu
 		SkeletalMeshActor = Cast<ASkeletalMeshActor>(SpawnedSkeletalMeshActor);
 	}
 
-	FString NewName = MovieSceneHelpers::MakeUniqueSpawnableName(GetSequencer()->GetFocusedMovieSceneSequence()->GetMovieScene(), FName::NameToDisplayString(SkeletalMesh->GetName(), false));
+	FString NewName = MovieSceneHelpers::MakeUniqueSpawnableName(MovieScene, FName::NameToDisplayString(SkeletalMesh->GetName(), false));
 	SkeletalMeshActor->SetActorLabel(NewName, false);
 
-	UControlRig* CDO = Cast<UControlRig>(RigClass->GetDefaultObject(true /* create if needed */));
-	check(CDO);
+	UMovieSceneControlRigParameterTrack* Track = Cast<UMovieSceneControlRigParameterTrack>(MovieScene->FindTrack(UMovieSceneControlRigParameterTrack::StaticClass(), NewGuid, NAME_None));
+	if (Track == nullptr)
+	{
+		UControlRig* CDO = Cast<UControlRig>(RigClass->GetDefaultObject(true /* create if needed */));
+		check(CDO);
 
-	AddControlRig(CDO->GetClass(), SkeletalMeshActor->GetSkeletalMeshComponent(), NewGuid);
+		AddControlRig(CDO->GetClass(), SkeletalMeshActor->GetSkeletalMeshComponent(), NewGuid);
+	}
 
 	return true;
 }
diff --git a/Engine/Plugins/Animation/ControlRig/Source/ControlRigEditor/Private/Sequencer/ControlRigParameterTrackEditor.h b/Engine/Plugins/Animation/ControlRig/Source/ControlRigEditor/Private/Sequencer/ControlRigParameterTrackEditor.h
--- a/Engine/Plugins/Animation/ControlRig/Source/ControlRigEditor/Private/Sequencer/ControlRigParameterTrackEditor.h
+++ b/Engine/Plugins/Animation/ControlRig/Source/ControlRigEditor/Private/Sequencer/ControlRigParameterTrackEditor.h
@@ -119,7 +119,7 @@ private:
 	/** Actor Added Delegate*/
 	void HandleActorAdded(AActor* Actor, FGuid TargetObjectGuid);
 	/** Add Control Rig Tracks For Skelmesh Components*/
-	void AddTrackForComponent(USceneComponent* Component);
+	void AddTrackForComponent(USceneComponent* Component, FGuid Binding);
 
 	/** Control Rig Delegates*/
 	void HandleControlModified(UControlRig* Subject, FRigControlElement* ControlElement, const FRigControlModifiedContext& Context);
