Control Rig: Fix reroute not available when dragging execute context pin

diff --git a/Engine/Plugins/Animation/ControlRig/Source/ControlRigDeveloper/Private/Graph/ControlRigGraphSchema.cpp b/Engine/Plugins/Animation/ControlRig/Source/ControlRigDeveloper/Private/Graph/ControlRigGraphSchema.cpp
--- a/Engine/Plugins/Animation/ControlRig/Source/ControlRigDeveloper/Private/Graph/ControlRigGraphSchema.cpp
+++ b/Engine/Plugins/Animation/ControlRig/Source/ControlRigDeveloper/Private/Graph/ControlRigGraphSchema.cpp
@@ -1777,14 +1777,6 @@ bool UControlRigGraphSchema::ArePinsCompatible(const UEdGraphPin* PinA, const UE
 	{
 		auto IsPinCompatibleWithType = [](const UEdGraphPin* InPin, const FEdGraphPinType& InPinType) -> bool
 		{
-			if(const UScriptStruct* ScriptStruct = Cast<UScriptStruct>(InPinType.PinSubCategoryObject))
-			{
-				if(ScriptStruct->IsChildOf(FRigVMExecuteContext::StaticStruct()))
-				{
-					return false;
-				}
-			}
-
 			if(const UControlRigGraphNode* RigNode = Cast<UControlRigGraphNode>(InPin->GetOwningNode()))
 			{
 				if(const FRigVMTemplate* Template = RigNode->GetTemplate())
