Control Rig Component: Allow setting control rig class in Blueprint

diff --git a/Engine/Plugins/Animation/ControlRig/Source/ControlRig/Private/ControlRigComponent.cpp b/Engine/Plugins/Animation/ControlRig/Source/ControlRig/Private/ControlRigComponent.cpp
--- a/Engine/Plugins/Animation/ControlRig/Source/ControlRig/Private/ControlRigComponent.cpp
+++ b/Engine/Plugins/Animation/ControlRig/Source/ControlRig/Private/ControlRigComponent.cpp
@@ -7,6 +7,7 @@
 #include "SkeletalDebugRendering.h"
 #include "Components/InstancedStaticMeshComponent.h"
 #include "AnimCustomInstanceHelper.h"
+#include "ControlRigObjectBinding.h"
 
 // CVar to disable control rig execution within a component
 static TAutoConsoleVariable<int32> CVarControlRigDisableExecutionComponent(TEXT("ControlRig.DisableExecutionInComponent"), 0, TEXT("if nonzero we disable the execution of Control Rigs inside a ControlRigComponent."));
@@ -1192,10 +1193,35 @@ void UControlRigComponent::SetControlRig(UControlRig* InControlRig)
 	ControlRig->OnExecuted_AnyThread().AddUObject(this, &UControlRigComponent::HandleControlRigExecutedEvent);
 
 	ControlRig->GetDataSourceRegistry()->RegisterDataSource(UControlRig::OwnerComponent, this);
+	if(ObjectBinding.IsValid())
+	{
+		ControlRig->SetObjectBinding(ObjectBinding);
+	}
 
 	ControlRig->Initialize();
 }
 
+void UControlRigComponent::SetControlRigClass(TSubclassOf<UControlRig> InControlRigClass)
+{
+	ControlRig = nullptr;
+	ControlRigClass = InControlRigClass;
+	Initialize();
+}
+
+void UControlRigComponent::SetObjectBinding(UObject* InObjectToBind)
+{
+	if(!ObjectBinding.IsValid())
+	{
+		ObjectBinding = MakeShared<FControlRigObjectBinding>();
+	}
+	ObjectBinding->BindToObject(InObjectToBind);
+
+	if(UControlRig* CR = SetupControlRigIfRequired())
+	{
+		CR->SetObjectBinding(ObjectBinding);
+	}
+}
+
 void UControlRigComponent::ValidateMappingData()
 {
 	TMap<USkeletalMeshComponent*, FCachedSkeletalMeshComponentSettings> NewCachedSettings;
diff --git a/Engine/Plugins/Animation/ControlRig/Source/ControlRig/Public/ControlRigComponent.h b/Engine/Plugins/Animation/ControlRig/Source/ControlRig/Public/ControlRigComponent.h
--- a/Engine/Plugins/Animation/ControlRig/Source/ControlRig/Public/ControlRigComponent.h
+++ b/Engine/Plugins/Animation/ControlRig/Source/ControlRig/Public/ControlRigComponent.h
@@ -179,7 +179,7 @@ class CONTROLRIG_API UControlRigComponent : public UPrimitiveComponent
 public:
 
 	/** The class of control rig to instantiate */
-	UPROPERTY(EditAnywhere, Category = ControlRig, meta = (DisallowCreateNew))
+	UPROPERTY(BlueprintReadWrite, EditAnywhere, Category = ControlRig, meta = (DisallowCreateNew), BlueprintSetter = SetControlRigClass)
 	TSubclassOf<UControlRig> ControlRigClass;
 
 	/** Event fired just before this component's ControlRig is initialized */
@@ -597,6 +597,12 @@ public:
 		
 	void SetControlRig(UControlRig* ControlRig);
 
+	UFUNCTION(BlueprintCallable, Category = "ControlRig")
+	void SetControlRigClass(TSubclassOf<UControlRig> InControlRigClass);
+
+	UFUNCTION(BlueprintCallable, Category = "ControlRig")
+	void SetObjectBinding(UObject* InObjectToBind);
+
 private:
 
 	struct FCachedSkeletalMeshComponentSettings
@@ -666,6 +672,8 @@ private:
 	TArray<FTransform> InputTransforms;
 	TArray<FTransform> LastInputTransforms;
 
+	TSharedPtr<IControlRigObjectBinding> ObjectBinding;
+
 	friend class FControlRigSceneProxy;
 };
 
