Control Rig: Visibily of proxies affects visibility of other proxies sharing the same driven controls

diff --git a/Engine/Plugins/Animation/ControlRig/Source/ControlRig/Private/Rigs/RigHierarchy.cpp b/Engine/Plugins/Animation/ControlRig/Source/ControlRig/Private/Rigs/RigHierarchy.cpp
--- a/Engine/Plugins/Animation/ControlRig/Source/ControlRig/Private/Rigs/RigHierarchy.cpp
+++ b/Engine/Plugins/Animation/ControlRig/Source/ControlRig/Private/Rigs/RigHierarchy.cpp
@@ -4521,6 +4521,27 @@ void URigHierarchy::UpdateVisibilityOnProxyControls()
 		return;
 	}
 
+	// create a local map of visible things, starting with the selection
+	TSet<FRigElementKey> VisibleElements;
+	VisibleElements.Append(HierarchyForSelection->OrderedSelection);
+
+	// if the control is visible - we should consider the
+	// driven controls visible as well - so that other proxies
+	// assigned to the same driven control also show up
+	for(const FRigBaseElement* Element : Elements)
+	{
+		if(const FRigControlElement* ControlElement = Cast<FRigControlElement>(Element))
+		{
+			if(ControlElement->Settings.AnimationType == ERigControlAnimationType::ProxyControl)
+			{
+				if(VisibleElements.Contains(ControlElement->GetKey()))
+				{
+					VisibleElements.Append(ControlElement->Settings.DrivenControls);
+				}
+			}
+		}
+	}
+
 	for(FRigBaseElement* Element : Elements)
 	{
 		if(FRigControlElement* ControlElement = Cast<FRigControlElement>(Element))
@@ -4538,18 +4559,24 @@ void URigHierarchy::UpdateVisibilityOnProxyControls()
 					}
 					else
 					{
-						bool bVisible = HierarchyForSelection->IsSelected(ControlElement);
+						// a proxy control should be visible if itself or a driven control is selected / visible
+						bool bVisible = VisibleElements.Contains(ControlElement->GetKey());
 						if(!bVisible)
 						{
-							if(ControlElement->Settings.DrivenControls.FindByPredicate([HierarchyForSelection](const FRigElementKey& AffectedControl) -> bool
+							if(ControlElement->Settings.DrivenControls.FindByPredicate([VisibleElements](const FRigElementKey& AffectedControl) -> bool
 							{
-								return HierarchyForSelection->IsSelected(AffectedControl);
+								return VisibleElements.Contains(AffectedControl);
 							}) != nullptr)
 							{
 								bVisible = true;
 							}
 						}
 
+						if(bVisible)
+						{
+							VisibleElements.Add(ControlElement->GetKey());
+						}
+
 						if(ControlElement->Settings.SetVisible(bVisible, true))
 						{
 							Notify(ERigHierarchyNotification::ControlVisibilityChanged, ControlElement);
