Control Rig: Support undo bracketing for hierarchy events

diff --git a/Engine/Plugins/Animation/ControlRig/Source/ControlRig/Private/ControlRig.cpp b/Engine/Plugins/Animation/ControlRig/Source/ControlRig/Private/ControlRig.cpp
--- a/Engine/Plugins/Animation/ControlRig/Source/ControlRig/Private/ControlRig.cpp
+++ b/Engine/Plugins/Animation/ControlRig/Source/ControlRig/Private/ControlRig.cpp
@@ -84,6 +84,7 @@ UControlRig::UControlRig(const FObjectInitializer& ObjectInitializer)
 	, PostSetupBracket(0)
 	, InteractionBracket(0)
 	, InterRigSyncBracket(0)
+	, ControlUndoBracketIndex(0)
 	, InteractionType((uint8)EControlRigInteractionType::None)
 	, bInteractionJustBegan(false)
 #if WITH_EDITORONLY_DATA
@@ -1028,6 +1029,16 @@ void UControlRig::Execute(const EControlRigState InState, const FName& InEventNa
 		}
 	}
 
+	// close remaining undo brackets from hierarchy
+	while(ControlUndoBracketIndex > 0)
+	{
+		FRigEventContext EventContext;
+		EventContext.Event = ERigEvent::CloseUndoBracket;
+		EventContext.SourceEventName = InEventName;
+		EventContext.LocalTime = Context.AbsoluteTime;
+		HandleHierarchyEvent(GetHierarchy(), EventContext);
+	}
+
 	if (Context.DrawInterface && Context.DrawContainer && bIsEventLastInQueue) 
 	{
 		Context.DrawInterface->Instructions.Append(Context.DrawContainer->Instructions);
@@ -2416,6 +2427,15 @@ void UControlRig::HandleHierarchyEvent(URigHierarchy* InHierarchy, const FRigEve
 					ControlModified().Broadcast(this, ControlElement, Context);
 				}
 			}
+			break;
+		}
+		case ERigEvent::OpenUndoBracket:
+		case ERigEvent::CloseUndoBracket:
+		{
+			const bool bOpenUndoBracket = InEvent.Event == ERigEvent::OpenUndoBracket;
+			ControlUndoBracketIndex = FMath::Max<int32>(0, ControlUndoBracketIndex + (bOpenUndoBracket ? 1 : -1));
+			ControlUndoBracket().Broadcast(this, bOpenUndoBracket);
+			break;
 		}
 		default:
 		{
diff --git a/Engine/Plugins/Animation/ControlRig/Source/ControlRig/Public/ControlRig.h b/Engine/Plugins/Animation/ControlRig/Source/ControlRig/Public/ControlRig.h
--- a/Engine/Plugins/Animation/ControlRig/Source/ControlRig/Public/ControlRig.h
+++ b/Engine/Plugins/Animation/ControlRig/Source/ControlRig/Public/ControlRig.h
@@ -74,6 +74,10 @@ public:
 
 	/** Bindable event for external objects to be notified that a Control is Selected */
 	DECLARE_EVENT_ThreeParams(UControlRig, FControlSelectedEvent, UControlRig*, FRigControlElement*, bool);
+
+	/** Bindable event to manage undo / redo brackets in the client */
+	DECLARE_EVENT_TwoParams(UControlRig, FControlUndoBracketEvent, UControlRig*, bool /* bOpen */);
+	
 	// To support Blueprints/scripting, we need a different delegate type (a 'Dynamic' delegate) which supports looser style UFunction binding (using names).
 	DECLARE_DYNAMIC_MULTICAST_SPARSE_DELEGATE_ThreeParams(FOnControlSelectedBP, UControlRig, OnControlSelected_BP, UControlRig*, Rig, const FRigControlElement&, Control, bool, bSelected);
 
@@ -412,6 +416,10 @@ public:
 	// selection changes coming from the manipulated subject.
 	FControlSelectedEvent& ControlSelected() { return OnControlSelected; }
 
+	// Returns an event that can be used to subscribe to
+	// Undo Bracket requests such as Open and Close.
+	FControlUndoBracketEvent & ControlUndoBracket() { return OnControlUndoBracket; }
+
 	bool IsCurveControl(const FRigControlElement* InControlElement) const;
 
 	DECLARE_EVENT_ThreeParams(UControlRig, FControlRigExecuteEvent, class UControlRig*, const EControlRigState, const FName&);
@@ -710,6 +718,7 @@ protected:
 	int32 PostForwardsSolveBracket;
 	int32 InteractionBracket;
 	int32 InterRigSyncBracket;
+	int32 ControlUndoBracketIndex;
 	uint8 InteractionType;
 	TArray<FRigElementKey> ElementsBeingInteracted;
 	bool bInteractionJustBegan;
@@ -766,6 +775,7 @@ protected:
 	FFilterControlEvent OnFilterControl;
 	FControlModifiedEvent OnControlModified;
 	FControlSelectedEvent OnControlSelected;
+	FControlUndoBracketEvent OnControlUndoBracket;
 
 	UPROPERTY(BlueprintAssignable, Category = ControlRig, meta=(DisplayName="OnControlSelected"))
 	FOnControlSelectedBP OnControlSelected_BP;
diff --git a/Engine/Plugins/Animation/ControlRig/Source/ControlRig/Public/Rigs/RigHierarchyDefines.h b/Engine/Plugins/Animation/ControlRig/Source/ControlRig/Public/Rigs/RigHierarchyDefines.h
--- a/Engine/Plugins/Animation/ControlRig/Source/ControlRig/Public/Rigs/RigHierarchyDefines.h
+++ b/Engine/Plugins/Animation/ControlRig/Source/ControlRig/Public/Rigs/RigHierarchyDefines.h
@@ -71,6 +71,12 @@ enum class ERigEvent : uint8
 	/** Request to Auto-Key the Control in Sequencer */
 	RequestAutoKey,
 
+	/** Request to open an Undo bracket in the client */
+	OpenUndoBracket,
+
+	/** Request to close an Undo bracket in the client */
+	CloseUndoBracket,
+
 	/** MAX - invalid */
 	Max UMETA(Hidden),
 };
diff --git a/Engine/Plugins/Animation/ControlRig/Source/ControlRigEditor/Private/Sequencer/ControlRigParameterTrackEditor.cpp b/Engine/Plugins/Animation/ControlRig/Source/ControlRigEditor/Private/Sequencer/ControlRigParameterTrackEditor.cpp
--- a/Engine/Plugins/Animation/ControlRig/Source/ControlRigEditor/Private/Sequencer/ControlRigParameterTrackEditor.cpp
+++ b/Engine/Plugins/Animation/ControlRig/Source/ControlRigEditor/Private/Sequencer/ControlRigParameterTrackEditor.cpp
@@ -199,6 +199,7 @@ FControlRigParameterTrackEditor::FControlRigParameterTrackEditor(TSharedRef<ISeq
 	, bSkipNextSelectionFromTimer(false)
 	, bFilterAssetBySkeleton(true)
 	, bFilterAssetByAnimatableControls(true)
+	, ControlUndoBracket(0)
 {
 	FMovieSceneToolsModule::Get().RegisterAnimationBakeHelper(this);
 
@@ -399,6 +400,7 @@ void FControlRigParameterTrackEditor::BindControlRig(UControlRig* ControlRig)
 		ControlRig->ControlModified().AddRaw(this, &FControlRigParameterTrackEditor::HandleControlModified);
 		ControlRig->OnInitialized_AnyThread().AddRaw(this, &FControlRigParameterTrackEditor::HandleOnInitialized);
 		ControlRig->ControlSelected().AddRaw(this, &FControlRigParameterTrackEditor::HandleControlSelected);
+		ControlRig->ControlUndoBracket().AddRaw(this, &FControlRigParameterTrackEditor::HandleControlUndoBracket);
 		BoundControlRigs.Add(ControlRig);
 		UMovieSceneControlRigParameterTrack* Track = FindTrack(ControlRig);
 		if (Track)
@@ -2475,6 +2477,21 @@ void FControlRigParameterTrackEditor::HandleControlModified(UControlRig* Control
 	}
 }
 
+void FControlRigParameterTrackEditor::HandleControlUndoBracket(UControlRig* Subject, bool bOpenUndoBracket)
+{
+	if(bOpenUndoBracket && ControlUndoBracket == 0)
+	{
+		ControlUndoTransaction = MakeShareable(new FScopedTransaction(LOCTEXT("KeyMultipleControls", "Auto-Key multiple controls")));
+	}
+
+	ControlUndoBracket = FMath::Max<int32>(0, ControlUndoBracket + (bOpenUndoBracket ? 1 : -1));
+	
+	if(!bOpenUndoBracket && ControlUndoBracket == 0)
+	{
+		ControlUndoTransaction.Reset();
+	}
+}
+
 void FControlRigParameterTrackEditor::GetControlRigKeys(UControlRig* InControlRig, FName ParameterName, EControlRigContextChannelToKey ChannelsToKey, ESequencerKeyMode KeyMode, UMovieSceneControlRigParameterSection* SectionToKey, FGeneratedTrackKeys& OutGeneratedKeys)
 {
 	const TArray<bool>& ControlsMask = SectionToKey->GetControlsMask();
diff --git a/Engine/Plugins/Animation/ControlRig/Source/ControlRigEditor/Private/Sequencer/ControlRigParameterTrackEditor.h b/Engine/Plugins/Animation/ControlRig/Source/ControlRigEditor/Private/Sequencer/ControlRigParameterTrackEditor.h
--- a/Engine/Plugins/Animation/ControlRig/Source/ControlRigEditor/Private/Sequencer/ControlRigParameterTrackEditor.h
+++ b/Engine/Plugins/Animation/ControlRig/Source/ControlRigEditor/Private/Sequencer/ControlRigParameterTrackEditor.h
@@ -18,6 +18,7 @@
 #include "MovieSceneToolsModule.h"
 #include "Engine/EngineTypes.h"
 #include "EditorUndoClient.h"
+#include "ScopedTransaction.h"
 
 struct FAssetData;
 class FMenuBuilder;
@@ -124,6 +125,7 @@ private:
 	/** Control Rig Delegates*/
 	void HandleControlModified(UControlRig* Subject, FRigControlElement* ControlElement, const FRigControlModifiedContext& Context);
 	void HandleControlSelected(UControlRig* Subject, FRigControlElement* ControlElement, bool bSelected);
+	void HandleControlUndoBracket(UControlRig* Subject, bool bOpenUndoBracket);
 	void HandleOnInitialized(UControlRig* Subject, const EControlRigState InState, const FName& InEventName);
 
 	/** SpaceChannel Delegates*/
@@ -234,6 +236,12 @@ private:
 
 	/** Array of sections that are getting undone, we need to recreate any space channel add,move key delegates to them*/
 	mutable TArray<UMovieSceneControlRigParameterSection*> SectionsGettingUndone;
+
+	/** An index counter for the opened undo brackets */
+	int32 ControlUndoBracket;
+
+	/** A transaction used to group multiple key events */
+	TSharedPtr<FScopedTransaction> ControlUndoTransaction;
 };
 
 
