normalize the rotation when converting to a FTransform_Float

diff --git a/Engine/Plugins/Animation/ControlRig/Source/ControlRig/Public/Rigs/RigHierarchyDefines.h b/Engine/Plugins/Animation/ControlRig/Source/ControlRig/Public/Rigs/RigHierarchyDefines.h
--- a/Engine/Plugins/Animation/ControlRig/Source/ControlRig/Public/Rigs/RigHierarchyDefines.h
+++ b/Engine/Plugins/Animation/ControlRig/Source/ControlRig/Public/Rigs/RigHierarchyDefines.h
@@ -569,7 +569,9 @@ public:
 			}
 			case ERigControlType::Transform:
 			{
-				return Get<FTransform_Float>().ToTransform();
+				Transform = Get<FTransform_Float>().ToTransform();
+				Transform.NormalizeRotation();
+				break;
 			}
 			case ERigControlType::TransformNoScale:
 			{
