Sequencer: Reuse track and section if already added, otherwise a duplicate control rig track and section are added when dragging a control rig into the level.

diff --git a/Engine/Plugins/Animation/ControlRig/Source/ControlRigEditor/Private/ControlRigBlueprintActions.cpp b/Engine/Plugins/Animation/ControlRig/Source/ControlRigEditor/Private/ControlRigBlueprintActions.cpp
--- a/Engine/Plugins/Animation/ControlRig/Source/ControlRigEditor/Private/ControlRigBlueprintActions.cpp
+++ b/Engine/Plugins/Animation/ControlRig/Source/ControlRigEditor/Private/ControlRigBlueprintActions.cpp
@@ -368,16 +368,25 @@ void FControlRigBlueprintActions::OnSpawnedSkeletalMeshActorChanged(UObject* InO
 					MovieScene->RemoveSpawnable(CompGuid);
 				}
 			}
-			UMovieSceneControlRigParameterTrack* Track = MovieScene->AddTrack<UMovieSceneControlRigParameterTrack>(ActorTrackGuid);
-			if (Track)
+
+			UMovieSceneControlRigParameterTrack* Track = Cast<UMovieSceneControlRigParameterTrack>(MovieScene->FindTrack(UMovieSceneControlRigParameterTrack::StaticClass(), ActorTrackGuid));
+			if (!Track)
 			{
+				Track = MovieScene->AddTrack<UMovieSceneControlRigParameterTrack>(ActorTrackGuid);
+			}
+
+			UControlRig* ControlRig = Track->GetControlRig();
+
+			FString ObjectName = (ControlRigClass->GetName());
+
+			if (!ControlRig || ControlRig->GetClass() != ControlRigClass)
+			{		
 				USkeletalMesh* SkeletalMesh = MeshActor->GetSkeletalMeshComponent()->SkeletalMesh;
 				USkeleton* Skeleton = SkeletalMesh->GetSkeleton();
 				
-				FString ObjectName = (ControlRigClass->GetName());
 				ObjectName.RemoveFromEnd(TEXT("_C"));
 
-				UControlRig* ControlRig = NewObject<UControlRig>(Track, ControlRigClass, FName(*ObjectName), RF_Transactional);
+				ControlRig = NewObject<UControlRig>(Track, ControlRigClass, FName(*ObjectName), RF_Transactional);
 				ControlRig->SetObjectBinding(MakeShared<FControlRigObjectBinding>());
 				ControlRig->GetObjectBinding()->BindToObject(MeshActor->GetSkeletalMeshComponent());
 				ControlRig->GetDataSourceRegistry()->RegisterDataSource(UControlRig::OwnerComponent, ControlRig->GetObjectBinding()->GetBoundObject());
@@ -386,29 +395,34 @@ void FControlRigBlueprintActions::OnSpawnedSkeletalMeshActorChanged(UObject* InO
 				ControlRig->CreateRigControlsForCurveContainer();
 
 				WeakSequencer.Pin()->NotifyMovieSceneDataChanged(EMovieSceneDataChangeType::MovieSceneStructureItemsChanged);
+			}
 
+			UMovieSceneSection* Section = Track->GetAllSections().Num() ? Track->GetAllSections()[0] : nullptr;
+			if (!Section)
+			{
 				Track->Modify();
-				UMovieSceneSection* NewSection = Track->CreateControlRigSection(0, ControlRig, true);
+				Section = Track->CreateControlRigSection(0, ControlRig, true);
 				//mz todo need to have multiple rigs with same class
 				Track->SetTrackName(FName(*ObjectName));
 				Track->SetDisplayName(FText::FromString(ObjectName));
 
 				WeakSequencer.Pin()->EmptySelection();
-				WeakSequencer.Pin()->SelectSection(NewSection);
+				WeakSequencer.Pin()->SelectSection(Section);
 				WeakSequencer.Pin()->ThrobSectionSelection();
 				WeakSequencer.Pin()->ObjectImplicitlyAdded(ControlRig);
-				FText Name = LOCTEXT("SequenceTrackFilter_ControlRigControls", "Control Rig Controls");
-				WeakSequencer.Pin()->SetTrackFilterEnabled(Name, true);
-				WeakSequencer.Pin()->NotifyMovieSceneDataChanged(EMovieSceneDataChangeType::MovieSceneStructureItemAdded);
-				if (!ControlRigEditMode)
-				{
-					GLevelEditorModeTools().ActivateMode(FControlRigEditMode::ModeName);
-					ControlRigEditMode = static_cast<FControlRigEditMode*>(GLevelEditorModeTools().GetActiveMode(FControlRigEditMode::ModeName));
-				}
-				if (ControlRigEditMode)
-				{
-					ControlRigEditMode->AddControlRigObject(ControlRig, WeakSequencer.Pin());
-				}
+			}
+			
+			FText Name = LOCTEXT("SequenceTrackFilter_ControlRigControls", "Control Rig Controls");
+			WeakSequencer.Pin()->SetTrackFilterEnabled(Name, true);
+			WeakSequencer.Pin()->NotifyMovieSceneDataChanged(EMovieSceneDataChangeType::MovieSceneStructureItemAdded);
+			if (!ControlRigEditMode)
+			{
+				GLevelEditorModeTools().ActivateMode(FControlRigEditMode::ModeName);
+				ControlRigEditMode = static_cast<FControlRigEditMode*>(GLevelEditorModeTools().GetActiveMode(FControlRigEditMode::ModeName));
+			}
+			if (ControlRigEditMode)
+			{
+				ControlRigEditMode->AddControlRigObject(ControlRig, WeakSequencer.Pin());
 			}
 		}
 	}
