Control Rig: Fix crash when editing variable of trash blueprint

diff --git a/Engine/Plugins/Animation/ControlRig/Source/ControlRigEditor/Private/ControlRigGraphDetails.cpp b/Engine/Plugins/Animation/ControlRig/Source/ControlRigEditor/Private/ControlRigGraphDetails.cpp
--- a/Engine/Plugins/Animation/ControlRig/Source/ControlRigEditor/Private/ControlRigGraphDetails.cpp
+++ b/Engine/Plugins/Animation/ControlRig/Source/ControlRigEditor/Private/ControlRigGraphDetails.cpp
@@ -1790,7 +1790,6 @@ FControlRigGraphMathTypeDetails::FControlRigGraphMathTypeDetails()
 : ScriptStruct(nullptr)
 , BlueprintBeingCustomized(nullptr)
 , GraphBeingCustomized(nullptr)
-, bEnabled(true)
 {
 }
 
@@ -1802,7 +1801,7 @@ void FControlRigGraphMathTypeDetails::CustomizeHeader(TSharedRef<IPropertyHandle
 
 	for (UObject* Object : Objects)
 	{
-		ObjectsBeingCustomized.Add(MakeWeakObjectPtr(Object));
+		ObjectsBeingCustomized.Add(Object);
 
 		if(BlueprintBeingCustomized == nullptr)
 		{
@@ -1818,26 +1817,6 @@ void FControlRigGraphMathTypeDetails::CustomizeHeader(TSharedRef<IPropertyHandle
 	FProperty* Property = InPropertyHandle->GetProperty();
 	const FStructProperty* StructProperty = CastField<FStructProperty>(Property);
 	ScriptStruct = StructProperty->Struct;
-
-	TSharedPtr<class IPropertyHandle> ChainHandle = InPropertyHandle;
-	do
-	{
-		PropertyChain.AddHead(ChainHandle->GetProperty());
-		PropertyArrayIndices.Insert(ChainHandle->GetIndexInArray(), 0);
-		ChainHandle = ChainHandle->GetParentHandle();
-
-		if(ChainHandle.IsValid())
-		{
-			if(ChainHandle->GetProperty() == nullptr)
-			{
-				break;
-			}
-		}
-	}
-	while (ChainHandle.IsValid());
-	PropertyChain.SetActiveMemberPropertyNode(PropertyChain.GetTail()->GetValue());
-
-	bEnabled = !PropertyChain.GetHead()->GetValue()->HasAnyPropertyFlags(CPF_EditConst);
 }
 
 void FControlRigGraphMathTypeDetails::CustomizeChildren(TSharedRef<IPropertyHandle> InPropertyHandle,
@@ -1878,4 +1857,14 @@ void FControlRigGraphMathTypeDetails::CustomizeChildren(TSharedRef<IPropertyHand
 	}
 }
 
+void FControlRigGraphMathTypeDetails::AddReferencedObjects(FReferenceCollector& Collector)
+{
+	Collector.AddReferencedObjects(ObjectsBeingCustomized);
+}
+
+FString FControlRigGraphMathTypeDetails::GetReferencerName() const
+{
+	return TEXT("FControlRigGraphMathTypeDetails:") + ObjectsBeingCustomized[0]->GetPathName();
+}
+
 #undef LOCTEXT_NAMESPACE
diff --git a/Engine/Plugins/Animation/ControlRig/Source/ControlRigEditor/Private/ControlRigGraphDetails.h b/Engine/Plugins/Animation/ControlRig/Source/ControlRigEditor/Private/ControlRigGraphDetails.h
--- a/Engine/Plugins/Animation/ControlRig/Source/ControlRigEditor/Private/ControlRigGraphDetails.h
+++ b/Engine/Plugins/Animation/ControlRig/Source/ControlRigEditor/Private/ControlRigGraphDetails.h
@@ -245,7 +245,7 @@ public:
 };
 
 /** Customization for editing a Control Rig node */
-class FControlRigGraphMathTypeDetails : public IPropertyTypeCustomization
+class FControlRigGraphMathTypeDetails : public IPropertyTypeCustomization, public FGCObject
 {
 public:
 
@@ -260,13 +260,48 @@ public:
 	virtual void CustomizeHeader(TSharedRef<class IPropertyHandle> InPropertyHandle, class FDetailWidgetRow& HeaderRow, IPropertyTypeCustomizationUtils& StructCustomizationUtils) override;
 	virtual void CustomizeChildren(TSharedRef<class IPropertyHandle> InPropertyHandle, class IDetailChildrenBuilder& StructBuilder, IPropertyTypeCustomizationUtils& StructCustomizationUtils) override;
 
+	/** FGCObject interface */
+	virtual void AddReferencedObjects(FReferenceCollector& Collector) override;
+	virtual FString GetReferencerName() const override;
+
 protected:
 
+	FORCEINLINE_DEBUGGABLE void GetPropertyChain(TSharedRef<class IPropertyHandle> InPropertyHandle, FEditPropertyChain& OutPropertyChain, TArray<int32> &OutPropertyArrayIndices, bool& bOutEnabled)
+	{
+		OutPropertyChain.Empty();
+		OutPropertyArrayIndices.Reset();
+		bOutEnabled = false;
+		if (!ObjectsBeingCustomized.IsEmpty())
+		{
+			if (!ObjectsBeingCustomized[0].IsNull())
+			{
+				TSharedPtr<class IPropertyHandle> ChainHandle = InPropertyHandle;
+				do
+				{
+					OutPropertyChain.AddHead(ChainHandle->GetProperty());
+					OutPropertyArrayIndices.Insert(ChainHandle->GetIndexInArray(), 0);
+					ChainHandle = ChainHandle->GetParentHandle();
+
+					if(ChainHandle.IsValid())
+					{
+						if(ChainHandle->GetProperty() == nullptr)
+						{
+							break;
+						}
+					}
+				}
+				while (ChainHandle.IsValid());
+				OutPropertyChain.SetActiveMemberPropertyNode(OutPropertyChain.GetTail()->GetValue());
+				bOutEnabled = !OutPropertyChain.GetHead()->GetValue()->HasAnyPropertyFlags(CPF_EditConst);
+			}
+		}
+	}
+
 	// extracts the value for a nested property (for Example Settings.WorldTransform) from an outer owner
 	template<typename ValueType>
-	FORCEINLINE_DEBUGGABLE ValueType& ContainerUObjectToValueRef(UObject* InOwner, ValueType& InDefault) const
+	FORCEINLINE_DEBUGGABLE ValueType& ContainerUObjectToValueRef(UObject* InOwner, ValueType& InDefault, FEditPropertyChain& InPropertyChain, TArray<int32> &InPropertyArrayIndices) const
 	{
-		FEditPropertyChain::TDoubleLinkedListNode* PropertyNode = PropertyChain.GetHead();
+		FEditPropertyChain::TDoubleLinkedListNode* PropertyNode = InPropertyChain.GetHead();
 		uint8* MemoryPtr = (uint8*)InOwner;
 		int32 ChainIndex = 0;
 		do
@@ -277,9 +312,9 @@ protected:
 			PropertyNode = PropertyNode->GetNextNode();
 			ChainIndex++;
 			
-			if(PropertyArrayIndices.IsValidIndex(ChainIndex))
+			if(InPropertyArrayIndices.IsValidIndex(ChainIndex))
 			{
-				const int32 ArrayIndex = PropertyArrayIndices[ChainIndex];
+				const int32 ArrayIndex = InPropertyArrayIndices[ChainIndex];
 				if(ArrayIndex != INDEX_NONE)
 				{
 					const FArrayProperty* ArrayProperty = CastField<FArrayProperty>(Property->GetOwnerProperty());
@@ -312,13 +347,18 @@ protected:
 	template<typename VectorType, typename NumericType>
 	TOptional<NumericType> GetVectorComponent(TSharedRef<class IPropertyHandle> InPropertyHandle, int32 InComponent) 
 	{
+		FEditPropertyChain PropertyChain;
+		TArray<int32> PropertyArrayIndices;
+		bool bEnabled;
+		GetPropertyChain(InPropertyHandle, PropertyChain, PropertyArrayIndices, bEnabled);
+	
 		TOptional<NumericType> Result;
-		for(TWeakObjectPtr<UObject> Object : ObjectsBeingCustomized)
+		for(TObjectPtr<UObject> Object : ObjectsBeingCustomized)
 		{
-			if(Object.IsValid() && InPropertyHandle->IsValidHandle())
+			if(!Object.IsNull() && InPropertyHandle->IsValidHandle())
 			{
 				static VectorType ZeroVector = VectorType();
-				const VectorType& Vector = ContainerUObjectToValueRef<VectorType>(Object.Get(), ZeroVector);
+				const VectorType& Vector = ContainerUObjectToValueRef<VectorType>(Object.Get(), ZeroVector, PropertyChain, PropertyArrayIndices);
 				NumericType Component = Vector[InComponent];
 				if(Result.IsSet())
 				{
@@ -340,11 +380,16 @@ protected:
 	template<typename VectorType, typename NumericType>
 	void OnVectorComponentChanged(TSharedRef<class IPropertyHandle> InPropertyHandle, int32 InComponent, NumericType InValue, bool bIsCommit, ETextCommit::Type InCommitType = ETextCommit::Default)
 	{
+		FEditPropertyChain PropertyChain;
+		TArray<int32> PropertyArrayIndices;
+		bool bEnabled;
+		GetPropertyChain(InPropertyHandle, PropertyChain, PropertyArrayIndices, bEnabled);
+	
 		TArray<UObject*> ObjectsView;
 		for(int32 Index = 0; Index < ObjectsBeingCustomized.Num(); Index++)
 		{
-			TWeakObjectPtr<UObject> Object = ObjectsBeingCustomized[Index];
-			if (Object.IsValid())
+			TObjectPtr<UObject> Object = ObjectsBeingCustomized[Index];
+			if (!Object.IsNull())
 			{
 				ObjectsView.Add(Object.Get());
 			}
@@ -364,11 +409,11 @@ protected:
 
 		for(int32 Index = 0; Index < ObjectsBeingCustomized.Num(); Index++)
 		{
-			TWeakObjectPtr<UObject> Object = ObjectsBeingCustomized[Index];
-			if(Object.IsValid() && InPropertyHandle->IsValidHandle())
+			TObjectPtr<UObject> Object = ObjectsBeingCustomized[Index];
+			if(!Object.IsNull() && InPropertyHandle->IsValidHandle())
 			{
 				static VectorType ZeroVector = VectorType();
-				VectorType& Vector = ContainerUObjectToValueRef<VectorType>(Object.Get(), ZeroVector);
+				VectorType& Vector = ContainerUObjectToValueRef<VectorType>(Object.Get(), ZeroVector, PropertyChain, PropertyArrayIndices);
 				VectorType PreviousVector = Vector;
 				Vector[InComponent] = InValue;
 					
@@ -396,6 +441,11 @@ protected:
 		typedef typename VectorType::FReal NumericType;
 		typedef SNumericVectorInputBox<NumericType, VectorType, NumberOfComponents> SLocalVectorInputBox;
 
+		FEditPropertyChain PropertyChain;
+        TArray<int32> PropertyArrayIndices;
+        bool bEnabled;
+        GetPropertyChain(InPropertyHandle, PropertyChain, PropertyArrayIndices, bEnabled);
+
 		typename SLocalVectorInputBox::FArguments Args;
 		Args.Font(IDetailLayoutBuilder::GetDetailFont());
 		Args.IsEnabled(bEnabled);
@@ -448,13 +498,18 @@ protected:
 	template<typename RotationType>
 	TOptional<RotationType> GetRotation(TSharedRef<class IPropertyHandle> InPropertyHandle)
 	{
+		FEditPropertyChain PropertyChain;
+		TArray<int32> PropertyArrayIndices;
+		bool bEnabled;
+		GetPropertyChain(InPropertyHandle, PropertyChain, PropertyArrayIndices, bEnabled);
+		
 		TOptional<RotationType> Result;
-		for(TWeakObjectPtr<UObject> Object : ObjectsBeingCustomized)
+		for(TObjectPtr<UObject> Object : ObjectsBeingCustomized)
 		{
-			if(Object.IsValid() && InPropertyHandle->IsValidHandle())
+			if(!Object.IsNull() && InPropertyHandle->IsValidHandle())
 			{
 				static RotationType ZeroRotation = RotationType();
-				const RotationType& Rotation = ContainerUObjectToValueRef<RotationType>(Object.Get(), ZeroRotation);
+				const RotationType& Rotation = ContainerUObjectToValueRef<RotationType>(Object.Get(), ZeroRotation, PropertyChain, PropertyArrayIndices);
 				if(Result.IsSet())
 				{
 					if(!Rotation.Equals(Result.GetValue()))
@@ -475,11 +530,16 @@ protected:
 	template<typename RotationType>
 	void OnRotationChanged(TSharedRef<class IPropertyHandle> InPropertyHandle, RotationType InValue, bool bIsCommit, ETextCommit::Type InCommitType = ETextCommit::Default)
 	{
+		FEditPropertyChain PropertyChain;
+        TArray<int32> PropertyArrayIndices;
+        bool bEnabled;
+        GetPropertyChain(InPropertyHandle, PropertyChain, PropertyArrayIndices, bEnabled);
+	
 		TArray<UObject*> ObjectsView;
 		for(int32 Index = 0; Index < ObjectsBeingCustomized.Num(); Index++)
 		{
-			TWeakObjectPtr<UObject> Object = ObjectsBeingCustomized[Index];
-			if (Object.IsValid())
+			TObjectPtr<UObject> Object = ObjectsBeingCustomized[Index];
+			if (!Object.IsNull())
 			{
 				ObjectsView.Add(Object.Get());
 			}
@@ -499,11 +559,11 @@ protected:
 
 		for(int32 Index = 0; Index < ObjectsBeingCustomized.Num(); Index++)
 		{
-			TWeakObjectPtr<UObject> Object = ObjectsBeingCustomized[Index];
-			if(Object.IsValid() && InPropertyHandle->IsValidHandle())
+			TObjectPtr<UObject> Object = ObjectsBeingCustomized[Index];
+			if(!Object.IsNull() && InPropertyHandle->IsValidHandle())
 			{
 				static RotationType ZeroRotation = RotationType();
-				RotationType& Rotation = ContainerUObjectToValueRef<RotationType>(Object.Get(), ZeroRotation);
+				RotationType& Rotation = ContainerUObjectToValueRef<RotationType>(Object.Get(), ZeroRotation, PropertyChain, PropertyArrayIndices);
 				RotationType PreviousRotation = Rotation;
 				Rotation = InValue;
 					
@@ -529,6 +589,11 @@ protected:
 	template<typename RotationType>
 	FORCEINLINE void CustomizeRotation(TSharedRef<class IPropertyHandle> InPropertyHandle, class IDetailChildrenBuilder& StructBuilder, IPropertyTypeCustomizationUtils& StructCustomizationUtils)
 	{
+		FEditPropertyChain PropertyChain;
+		TArray<int32> PropertyArrayIndices;
+		bool bEnabled;
+		GetPropertyChain(InPropertyHandle, PropertyChain, PropertyArrayIndices, bEnabled);
+		
 		typedef typename RotationType::FReal NumericType;
 		typedef SAdvancedRotationInputBox<NumericType> SLocalRotationInputBox;
 		typename SLocalRotationInputBox::FArguments Args;
@@ -559,6 +624,11 @@ protected:
 	template<typename TransformType>
 	FORCEINLINE void CustomizeTransform(TSharedRef<class IPropertyHandle> InPropertyHandle, class IDetailChildrenBuilder& StructBuilder, IPropertyTypeCustomizationUtils& StructCustomizationUtils)
 	{
+		FEditPropertyChain PropertyChain;
+		TArray<int32> PropertyArrayIndices;
+		bool bEnabled;
+		GetPropertyChain(InPropertyHandle, PropertyChain, PropertyArrayIndices, bEnabled);
+		
 		typedef typename TransformType::FReal FReal;
 		typename SAdvancedTransformInputBox<TransformType>::FArguments WidgetArgs;
 		WidgetArgs.IsEnabled(bEnabled);
@@ -566,15 +636,20 @@ protected:
 		WidgetArgs.UseQuaternionForRotation(IsQuaternionBasedRotation<TransformType>());
 
 		static TransformType Identity = TransformType::Identity;
-		TransformType DefaultValue = ContainerUObjectToValueRef<TransformType>(ObjectsBeingCustomized[0]->GetClass()->GetDefaultObject(), Identity);
+		TransformType DefaultValue = ContainerUObjectToValueRef<TransformType>(ObjectsBeingCustomized[0]->GetClass()->GetDefaultObject(), Identity, PropertyChain, PropertyArrayIndices);
 
 		WidgetArgs.DiffersFromDefault_Lambda([this, InPropertyHandle, DefaultValue](ESlateTransformComponent::Type InTransformComponent) -> bool
 		{
-			for(TWeakObjectPtr<UObject> Object : ObjectsBeingCustomized)
+			FEditPropertyChain PropertyChain;
+			TArray<int32> PropertyArrayIndices;
+			bool bEnabled;
+			GetPropertyChain(InPropertyHandle, PropertyChain, PropertyArrayIndices, bEnabled);
+			
+			for(TObjectPtr<UObject> Object : ObjectsBeingCustomized)
 			{
-				if(Object.IsValid() && InPropertyHandle->IsValidHandle())
+				if(!Object.IsNull() && InPropertyHandle->IsValidHandle())
 				{
-					const TransformType& Transform = ContainerUObjectToValueRef<TransformType>(Object.Get(), Identity);
+					const TransformType& Transform = ContainerUObjectToValueRef<TransformType>(Object.Get(), Identity, PropertyChain, PropertyArrayIndices);
 
 					switch(InTransformComponent)
 					{
@@ -617,12 +692,17 @@ protected:
 			ESlateRotationRepresentation::Type InRotationRepresentation,
 			ESlateTransformSubComponent::Type InTransformSubComponent) -> TOptional<FReal>
 		{
+			FEditPropertyChain PropertyChain;
+			TArray<int32> PropertyArrayIndices;
+			bool bEnabled;
+			GetPropertyChain(InPropertyHandle, PropertyChain, PropertyArrayIndices, bEnabled);
+			
 			TOptional<FReal> Result;
-			for(TWeakObjectPtr<UObject> Object : ObjectsBeingCustomized)
+			for(TObjectPtr<UObject> Object : ObjectsBeingCustomized)
 			{
-				if(Object.IsValid() && InPropertyHandle->IsValidHandle())
+				if(!Object.IsNull() && InPropertyHandle->IsValidHandle())
 				{
-					const TransformType& Transform = ContainerUObjectToValueRef<TransformType>(Object.Get(), Identity);
+					const TransformType& Transform = ContainerUObjectToValueRef<TransformType>(Object.Get(), Identity, PropertyChain, PropertyArrayIndices);
 					
 					TOptional<FReal> Value = SAdvancedTransformInputBox<TransformType>::GetNumericValueFromTransform(
 						Transform,
@@ -659,11 +739,16 @@ protected:
 			bool bIsCommit,
 			ETextCommit::Type InCommitType = ETextCommit::Default)
 		{
+			FEditPropertyChain PropertyChain;
+			TArray<int32> PropertyArrayIndices;
+			bool bEnabled;
+			GetPropertyChain(InPropertyHandle, PropertyChain, PropertyArrayIndices, bEnabled);
+			
 			TArray<UObject*> ObjectsView;
 			for(int32 Index = 0; Index < ObjectsBeingCustomized.Num(); Index++)
 			{
-				TWeakObjectPtr<UObject> Object = ObjectsBeingCustomized[Index];
-				if (Object.IsValid())
+				TObjectPtr<UObject> Object = ObjectsBeingCustomized[Index];
+				if (!Object.IsNull())
 				{
 					ObjectsView.Add(Object.Get());
 				}
@@ -683,10 +768,10 @@ protected:
 
 			for(int32 Index = 0; Index < ObjectsBeingCustomized.Num(); Index++)
 			{
-				TWeakObjectPtr<UObject> Object = ObjectsBeingCustomized[Index];
-				if(Object.IsValid() && InPropertyHandle->IsValidHandle())
+				TObjectPtr<UObject> Object = ObjectsBeingCustomized[Index];
+				if(!Object.IsNull() && InPropertyHandle->IsValidHandle())
 				{
-					TransformType& Transform = ContainerUObjectToValueRef<TransformType>(Object.Get(), Identity);
+					TransformType& Transform = ContainerUObjectToValueRef<TransformType>(Object.Get(), Identity, PropertyChain, PropertyArrayIndices);
 					TransformType PreviousTransform = Transform;
 					
 					SAdvancedTransformInputBox<TransformType>::ApplyNumericValueChange(
@@ -731,6 +816,11 @@ protected:
 
 		WidgetArgs.OnResetToDefault_Lambda([this, DefaultValue, InPropertyHandle](ESlateTransformComponent::Type InTransformComponent)
 		{
+			FEditPropertyChain PropertyChain;
+			TArray<int32> PropertyArrayIndices;
+			bool bEnabled;
+			GetPropertyChain(InPropertyHandle, PropertyChain, PropertyArrayIndices, bEnabled);
+			
 			URigVMController* Controller = nullptr;
 			if(BlueprintBeingCustomized && GraphBeingCustomized)
 			{
@@ -744,8 +834,8 @@ protected:
 			TArray<UObject*> ObjectsView;
 			for(int32 Index = 0; Index < ObjectsBeingCustomized.Num(); Index++)
 			{
-				TWeakObjectPtr<UObject> Object = ObjectsBeingCustomized[Index];
-				if (Object.IsValid())
+				TObjectPtr<UObject> Object = ObjectsBeingCustomized[Index];
+				if (!Object.IsNull())
 				{
 					ObjectsView.Add(Object.Get());
 				}
@@ -755,10 +845,10 @@ protected:
 			
 			for(int32 Index = 0; Index < ObjectsBeingCustomized.Num(); Index++)
 			{
-				TWeakObjectPtr<UObject> Object = ObjectsBeingCustomized[Index];
-				if(Object.IsValid() && InPropertyHandle->IsValidHandle())
+				TObjectPtr<UObject> Object = ObjectsBeingCustomized[Index];
+				if(!Object.IsNull() && InPropertyHandle->IsValidHandle())
 				{
-					TransformType& Transform = ContainerUObjectToValueRef<TransformType>(Object.Get(), Identity);
+					TransformType& Transform = ContainerUObjectToValueRef<TransformType>(Object.Get(), Identity, PropertyChain, PropertyArrayIndices);
 					TransformType PreviousTransform = Transform;
 
 					switch(InTransformComponent)
@@ -811,10 +901,7 @@ protected:
 	UScriptStruct* ScriptStruct;
 	UControlRigBlueprint* BlueprintBeingCustomized;
 	URigVMGraph* GraphBeingCustomized;
-	TArray<TWeakObjectPtr<UObject>> ObjectsBeingCustomized; 
-	FEditPropertyChain PropertyChain;
-	TArray<int32> PropertyArrayIndices;
-	bool bEnabled;
+	TArray<TObjectPtr<UObject>> ObjectsBeingCustomized; 
 };
 
 template<>
