ControlRig: Mark GetWorldTime node as varying

diff --git a/Engine/Plugins/Animation/ControlRig/Source/ControlRig/Private/Units/Animation/RigUnit_GetWorldTime.h b/Engine/Plugins/Animation/ControlRig/Source/ControlRig/Private/Units/Animation/RigUnit_GetWorldTime.h
--- a/Engine/Plugins/Animation/ControlRig/Source/ControlRig/Private/Units/Animation/RigUnit_GetWorldTime.h
+++ b/Engine/Plugins/Animation/ControlRig/Source/ControlRig/Private/Units/Animation/RigUnit_GetWorldTime.h
@@ -8,7 +8,7 @@
 /**
  * Returns the current time (year, month, day, hour, minute)
  */
-USTRUCT(meta = (DisplayName = "Now", Keywords = "Time,Clock"))
+USTRUCT(meta = (DisplayName = "Now", Keywords = "Time,Clock", Varying))
 struct CONTROLRIG_API FRigUnit_GetWorldTime : public FRigUnit_AnimBase
 {
 	GENERATED_BODY()
