Control Rig: Fix crash when resetting control transform

diff --git a/Engine/Plugins/Animation/ControlRig/Source/ControlRigEditor/Private/EditMode/ControlRigEditMode.cpp b/Engine/Plugins/Animation/ControlRig/Source/ControlRigEditor/Private/EditMode/ControlRigEditMode.cpp
--- a/Engine/Plugins/Animation/ControlRig/Source/ControlRigEditor/Private/EditMode/ControlRigEditMode.cpp
+++ b/Engine/Plugins/Animation/ControlRig/Source/ControlRigEditor/Private/EditMode/ControlRigEditMode.cpp
@@ -43,6 +43,7 @@
 #include "ControlRigComponent.h"
 #include "EngineUtils.h"
 #include "ControlRig/Private/Units/Execution/RigUnit_BeginExecution.h"
+#include "ControlRig/Private/Units/Execution/RigUnit_InteractionExecution.h"
 #include "IPersonaPreviewScene.h"
 #include "PersonaSelectionProxies.h"
 #include "Framework/Application/SlateApplication.h"
@@ -2743,7 +2744,7 @@ void FControlRigEditMode::ResetTransforms(bool bSelectionOnly)
 			else
 			{
 				// we at least have to run the interaction event
-				TGuardValue<TArray<FName>> EventGuard(ControlRig->EventQueue, {});
+				TGuardValue<TArray<FName>> EventGuard(ControlRig->EventQueue, {FRigUnit_InteractionExecution::EventName});
 				ControlRig->Evaluate_AnyThread();
 				for (const FRigElementKey& ControlToReset : ControlsToReset)
 				{
