Control Rig: Animation Channels Refactoring

diff --git a/Engine/Plugins/Animation/ControlRig/Source/ControlRig/Private/Rigs/RigHierarchy.cpp b/Engine/Plugins/Animation/ControlRig/Source/ControlRig/Private/Rigs/RigHierarchy.cpp
--- a/Engine/Plugins/Animation/ControlRig/Source/ControlRig/Private/Rigs/RigHierarchy.cpp
+++ b/Engine/Plugins/Animation/ControlRig/Source/ControlRig/Private/Rigs/RigHierarchy.cpp
@@ -25,6 +25,7 @@ LLM_DEFINE_TAG(Animation_ControlRig);
 #include "Serialization/JsonReader.h"
 #include "Serialization/JsonSerializer.h"
 #include "ControlRig/Private/Units/Execution/RigUnit_BeginExecution.h"
+#include "Algo/Transform.h"
 
 static FCriticalSection GRigHierarchyStackTraceMutex;
 static char GRigHierarchyStackTrace[65536];
@@ -805,7 +806,7 @@ FName URigHierarchy::GetSanitizedName(const FString& InName)
 
 bool URigHierarchy::IsNameAvailable(const FString& InPotentialNewName, ERigElementType InType, FString* OutErrorMessage) const
 {
-	FString UnsanitizedName = InPotentialNewName;
+	const FString UnsanitizedName = InPotentialNewName;
 	if (UnsanitizedName.Len() > GetMaxNameLength())
 	{
 		if (OutErrorMessage)
@@ -855,6 +856,66 @@ bool URigHierarchy::IsNameAvailable(const FString& InPotentialNewName, ERigEleme
 	return true;
 }
 
+bool URigHierarchy::IsDisplayNameAvailable(const FRigElementKey& InParentElement,
+	const FString& InPotentialNewDisplayName, FString* OutErrorMessage) const
+{
+	const FString UnsanitizedName = InPotentialNewDisplayName;
+	if (UnsanitizedName.Len() > GetMaxNameLength())
+	{
+		if (OutErrorMessage)
+		{
+			*OutErrorMessage = TEXT("Name too long.");
+		}
+		return false;
+	}
+
+	if (UnsanitizedName == TEXT("None"))
+	{
+		if (OutErrorMessage)
+		{
+			*OutErrorMessage = TEXT("None is not a valid name.");
+		}
+		return false;
+	}
+
+	FString SanitizedName = UnsanitizedName;
+	SanitizeName(SanitizedName);
+
+	if (SanitizedName != UnsanitizedName)
+	{
+		if (OutErrorMessage)
+		{
+			*OutErrorMessage = TEXT("Name contains invalid characters.");
+		}
+		return false;
+	}
+
+	if(InParentElement.IsValid())
+	{
+		const TArray<FRigElementKey> ChildKeys = GetChildren(InParentElement);
+		if(ChildKeys.ContainsByPredicate([SanitizedName, this](const FRigElementKey& InChildKey) -> bool
+		{
+			if(const FRigBaseElement* BaseElement = Find(InChildKey))
+			{
+				if(BaseElement->GetDisplayName().ToString() == SanitizedName)
+				{
+					return true;
+				}
+			}
+			return false;
+		}))
+		{
+			if (OutErrorMessage)
+			{
+				*OutErrorMessage = TEXT("Name already used.");
+			}
+			return false;
+		}
+	}
+
+	return true;
+}
+
 FName URigHierarchy::GetSafeNewName(const FString& InPotentialNewName, ERigElementType InType) const
 {
 	FString SanitizedName = InPotentialNewName;
@@ -874,6 +935,46 @@ FName URigHierarchy::GetSafeNewName(const FString& InPotentialNewName, ERigEleme
 	return *Name;
 }
 
+FName URigHierarchy::GetSafeNewDisplayName(const FRigElementKey& InParentElement, const FString& InPotentialNewDisplayName) const
+{
+	if(InPotentialNewDisplayName.IsEmpty())
+	{
+		return NAME_None;
+	}
+	
+	if(InParentElement.IsValid())
+	{
+		FString SanitizedName = InPotentialNewDisplayName;
+		SanitizeName(SanitizedName);
+		FString Name = SanitizedName;
+
+		const TArray<FRigElementKey> ChildKeys = GetChildren(InParentElement);
+		TArray<FString> DisplayNames;
+		Algo::Transform(ChildKeys, DisplayNames, [this](const FRigElementKey& InKey) -> FString
+		{
+			if(const FRigBaseElement* BaseElement = Find(InKey))
+			{
+				return BaseElement->GetDisplayName().ToString();
+			}
+			return FString();
+		});
+
+		int32 Suffix = 1;
+		while (DisplayNames.Contains(Name))
+		{
+			FString BaseString = SanitizedName;
+			if (BaseString.Len() > GetMaxNameLength() - 4)
+			{
+				BaseString.LeftChopInline(BaseString.Len() - (GetMaxNameLength() - 4));
+			}
+			Name = *FString::Printf(TEXT("%s_%d"), *BaseString, ++Suffix);
+		}
+
+		return *Name;
+	}
+	return *InPotentialNewDisplayName;
+}
+
 FEdGraphPinType URigHierarchy::GetControlPinType(FRigControlElement* InControlElement) const
 {
 	check(InControlElement);
@@ -3949,15 +4050,25 @@ void URigHierarchy::PropagateDirtyFlags(FRigTransformElement* InTransformElement
 
 	const ERigTransformType::Type LocalType = bInitial ? ERigTransformType::InitialLocal : ERigTransformType::CurrentLocal;
 	const ERigTransformType::Type GlobalType = bInitial ? ERigTransformType::InitialGlobal : ERigTransformType::CurrentGlobal;
-	const ERigTransformType::Type TypeToCompute = bAffectChildren ? LocalType : GlobalType;
-	const ERigTransformType::Type TypeToDirty = SwapLocalAndGlobal(TypeToCompute);
 
 	if(bComputeOpposed)
 	{
 		for(const FRigTransformElement::FElementToDirty& ElementToDirty : InTransformElement->ElementsToDirty)
 		{
+			ERigTransformType::Type TypeToCompute = bAffectChildren ? LocalType : GlobalType;
+			ERigTransformType::Type TypeToDirty = SwapLocalAndGlobal(TypeToCompute);
+			
 			if(FRigControlElement* ControlElement = Cast<FRigControlElement>(ElementToDirty.Element))
 			{
+				// animation channels never dirty their local value
+				if(ControlElement->IsAnimationChannel())
+				{
+					if(ERigTransformType::IsLocal(TypeToDirty))
+					{
+						Swap(TypeToDirty, TypeToCompute);
+					}
+				}
+				
 				if(ERigTransformType::IsGlobal(TypeToDirty))
 				{
 					if(ControlElement->Offset.IsDirty(TypeToDirty) &&
@@ -4000,8 +4111,20 @@ void URigHierarchy::PropagateDirtyFlags(FRigTransformElement* InTransformElement
 	{
 		for(const FRigTransformElement::FElementToDirty& ElementToDirty : InTransformElement->ElementsToDirty)
 		{
+			ERigTransformType::Type TypeToCompute = bAffectChildren ? LocalType : GlobalType;
+			ERigTransformType::Type TypeToDirty = SwapLocalAndGlobal(TypeToCompute);
+
 			if(FRigControlElement* ControlElement = Cast<FRigControlElement>(ElementToDirty.Element))
 			{
+				// animation channels never dirty their local value
+				if(ControlElement->IsAnimationChannel())
+				{
+					if(ERigTransformType::IsLocal(TypeToDirty))
+					{
+						Swap(TypeToDirty, TypeToCompute);
+					}
+				}
+
 				if(ERigTransformType::IsGlobal(TypeToDirty))
 				{
 					if(ControlElement->Offset.IsDirty(TypeToDirty) &&
@@ -4839,7 +4962,7 @@ bool URigHierarchy::IsAnimatable(const FRigControlElement* InControlElement) con
 		}
 
 		// animation channels are dependent on the control they are under.
-		if(InControlElement->Settings.AnimationType == ERigControlAnimationType::AnimationChannel)
+		if(InControlElement->IsAnimationChannel())
 		{
 			if(const FRigControlElement* ParentControlElement = Cast<FRigControlElement>(GetFirstParent(InControlElement)))
 			{
diff --git a/Engine/Plugins/Animation/ControlRig/Source/ControlRig/Private/Rigs/RigHierarchyController.cpp b/Engine/Plugins/Animation/ControlRig/Source/ControlRig/Private/Rigs/RigHierarchyController.cpp
--- a/Engine/Plugins/Animation/ControlRig/Source/ControlRig/Private/Rigs/RigHierarchyController.cpp
+++ b/Engine/Plugins/Animation/ControlRig/Source/ControlRig/Private/Rigs/RigHierarchyController.cpp
@@ -329,6 +329,11 @@ FRigElementKey URigHierarchyController::AddControl(
 		{
 			NewElement->Settings.SetupLimitArrayForType();
 		}
+
+		if(!NewElement->Settings.DisplayName.IsNone())
+		{
+			NewElement->Settings.DisplayName = Hierarchy->GetSafeNewDisplayName(InParent, NewElement->Settings.DisplayName.ToString()); 
+		}
 		AddElement(NewElement, Hierarchy->Get(Hierarchy->GetIndex(InParent)), false);
 		
 		NewElement->Offset.Set(ERigTransformType::InitialLocal, InOffsetTransform);  
@@ -366,6 +371,26 @@ FRigElementKey URigHierarchyController::AddControl(
 	return NewElement->Key;
 }
 
+FRigElementKey URigHierarchyController::AddAnimationChannel(FName InName, FRigElementKey InParentControl,
+	FRigControlSettings InSettings, bool bSetupUndo, bool bPrintPythonCommand)
+{
+	if(!IsValid())
+	{
+		return FRigElementKey();
+	}
+
+	if(const FRigControlElement* ParentControl = Hierarchy->Find<FRigControlElement>(InParentControl))
+	{
+		InSettings.AnimationType = ERigControlAnimationType::AnimationChannel;
+		InSettings.bGroupWithParentControl = true;
+
+		return AddControl(InName, ParentControl->GetKey(), InSettings, InSettings.GetIdentityValue(),
+			FTransform::Identity, FTransform::Identity, bSetupUndo, bPrintPythonCommand);
+	}
+
+	return FRigElementKey();
+}
+
 FRigElementKey URigHierarchyController::AddCurve(FName InName, float InValue, bool bSetupUndo, bool bPrintPythonCommand)
 {
 	if(!IsValid())
@@ -1838,6 +1863,57 @@ FRigElementKey URigHierarchyController::RenameElement(FRigElementKey InElement,
 	return bRenamed ? Element->GetKey() : FRigElementKey();
 }
 
+FName URigHierarchyController::SetDisplayName(FRigElementKey InControl, FName InDisplayName, bool bRenameElement, bool bSetupUndo,
+	bool bPrintPythonCommand)
+{
+	if(!IsValid())
+	{
+		return NAME_None;
+	}
+
+	FRigControlElement* ControlElement = Hierarchy->Find<FRigControlElement>(InControl);
+	if(ControlElement == nullptr)
+	{
+		ReportWarningf(TEXT("Cannot Rename Control: '%s' not found."), *InControl.ToString());
+		return NAME_None;
+	}
+
+#if WITH_EDITOR
+	TSharedPtr<FScopedTransaction> TransactionPtr;
+	if(bSetupUndo)
+	{
+		TransactionPtr = MakeShared<FScopedTransaction>(NSLOCTEXT("RigHierarchyController", "Set Display Name on Control", "Set Display Name on Control"));
+		Hierarchy->Modify();
+	}
+#endif
+
+	const FName NewDisplayName = SetDisplayName(ControlElement, InDisplayName, bRenameElement);
+	const bool bDisplayNameChanged = !NewDisplayName.IsNone();
+
+#if WITH_EDITOR
+	if(!bDisplayNameChanged && TransactionPtr.IsValid())
+	{
+		TransactionPtr->Cancel();
+	}
+	TransactionPtr.Reset();
+
+	if (bDisplayNameChanged && bPrintPythonCommand && !bSuspendPythonPrinting)
+	{
+		UBlueprint* Blueprint = GetTypedOuter<UBlueprint>();
+		if (Blueprint)
+		{
+			RigVMPythonUtils::Print(Blueprint->GetFName().ToString(), 
+				FString::Printf(TEXT("hierarchy_controller.set_display_name(%s, '%s', %s)"),
+				*InControl.ToPythonString(),
+				*InDisplayName.ToString(),
+				(bRenameElement) ? TEXT("True") : TEXT("False")));
+		}
+	}
+#endif
+
+	return NewDisplayName;
+}
+
 bool URigHierarchyController::RenameElement(FRigBaseElement* InElement, const FName &InName, bool bClearSelection)
 {
 	if(InElement == nullptr)
@@ -1921,6 +1997,36 @@ bool URigHierarchyController::RenameElement(FRigBaseElement* InElement, const FN
 	return true;
 }
 
+FName URigHierarchyController::SetDisplayName(FRigControlElement* InControlElement, const FName& InDisplayName, bool bRenameElement)
+{
+	if(InControlElement == nullptr)
+	{
+		return NAME_None;
+	}
+
+	if (InControlElement->Settings.DisplayName.IsEqual(InDisplayName, ENameCase::CaseSensitive))
+	{
+		return NAME_None;
+	}
+
+	if(const FRigBaseElement* ParentElement = Hierarchy->GetFirstParent(InControlElement))
+	{
+		const FString DesiredDisplayName = InDisplayName.IsNone() ? FString() : InDisplayName.ToString();
+		const FName DisplayName = Hierarchy->GetSafeNewDisplayName(ParentElement->GetKey(), DesiredDisplayName);
+		InControlElement->Settings.DisplayName = DisplayName;
+	}
+	
+	Hierarchy->IncrementTopologyVersion();
+	Notify(ERigHierarchyNotification::ControlSettingChanged, InControlElement);
+
+	if(bRenameElement)
+	{
+		RenameElement(InControlElement, InControlElement->Settings.DisplayName, false);
+	}
+
+	return InControlElement->Settings.DisplayName;
+}
+
 bool URigHierarchyController::AddParent(FRigElementKey InChild, FRigElementKey InParent, float InWeight, bool bMaintainGlobalTransform, bool bSetupUndo)
 {
 	if(!IsValid())
@@ -1995,12 +2101,27 @@ bool URigHierarchyController::AddParent(FRigBaseElement* InChild, FRigBaseElemen
 	// we can only parent things to controls which are not animation channels (animation channels are not 3D things)
 	if(FRigControlElement* ParentControlElement = Cast<FRigControlElement>(InParent))
 	{
-		if(ParentControlElement->Settings.AnimationType == ERigControlAnimationType::AnimationChannel)
+		if(ParentControlElement->IsAnimationChannel())
 		{
 			return false;
 		}
 	}
 
+	// we can only reparent animation channels - we cannot add a parent to them
+	if(FRigControlElement* ChildControlElement = Cast<FRigControlElement>(InChild))
+	{
+		if(ChildControlElement->IsAnimationChannel())
+		{
+			if(!bRemoveAllParents)
+			{
+				ReportErrorf(TEXT("Cannot add multiple parents to animation channel '%s'."), *InChild->Key.ToString());
+				return false;
+			}
+
+			bMaintainGlobalTransform = false;
+		}
+	}
+
 	if(Hierarchy->IsParentedTo(InParent, InChild))
 	{
 		ReportErrorf(TEXT("Cannot parent '%s' to '%s' - would cause a cycle."), *InChild->Key.ToString(), *InParent->Key.ToString());
@@ -2071,6 +2192,17 @@ bool URigHierarchyController::AddParent(FRigBaseElement* InChild, FRigBaseElemen
 	}
 	else if(FRigMultiParentElement* MultiParentElement = Cast<FRigMultiParentElement>(InChild))
 	{
+		if(FRigControlElement* ControlElement = Cast<FRigControlElement>(InChild))
+		{
+			if(!ControlElement->Settings.DisplayName.IsNone())
+			{
+				ControlElement->Settings.DisplayName =
+					Hierarchy->GetSafeNewDisplayName(
+						InParent->GetKey(),
+						ControlElement->Settings.DisplayName.ToString());
+			}
+		}
+		
 		AddElementToDirty(Constraint.ParentElement, MultiParentElement);
 
 		const int32 ParentIndex = MultiParentElement->ParentConstraints.Add(Constraint);
diff --git a/Engine/Plugins/Animation/ControlRig/Source/ControlRig/Private/Units/Hierarchy/RigUnit_ControlChannel.cpp b/Engine/Plugins/Animation/ControlRig/Source/ControlRig/Private/Units/Hierarchy/RigUnit_ControlChannel.cpp
new file mode 100644
--- /dev/null
+++ b/Engine/Plugins/Animation/ControlRig/Source/ControlRig/Private/Units/Hierarchy/RigUnit_ControlChannel.cpp
@@ -0,0 +1,358 @@
+// Copyright Epic Games, Inc. All Rights Reserved.
+
+#include "RigUnit_ControlChannel.h"
+#include "Units/RigUnitContext.h"
+
+bool FRigUnit_GetAnimationChannelBase::UpdateCache(const URigHierarchy* InHierarchy, const FName& Control, const FName& Channel, FRigElementKey& Key, int32& Hash)
+{
+	if(!Key.IsValid())
+	{
+		Hash = INDEX_NONE;
+	}
+
+	const int32 ExpectedHash = (int32)HashCombine(GetTypeHash(InHierarchy->GetTopologyVersion()), HashCombine(GetTypeHash(Control), GetTypeHash(Channel)));
+	if(ExpectedHash != Hash)
+	{
+		if(const FRigControlElement* ControlElement = InHierarchy->Find<FRigControlElement>(FRigElementKey(Control, ERigElementType::Control)))
+		{
+			FRigBaseElementChildrenArray Children = InHierarchy->GetChildren(ControlElement);
+			for(const FRigBaseElement* Child : Children)
+			{
+				if(const FRigControlElement* ChildControl = Cast<FRigControlElement>(Child))
+				{
+					if(ChildControl->IsAnimationChannel())
+					{
+						if(ChildControl->GetDisplayName() == Channel)
+						{
+							Key = ChildControl->GetKey();
+							Hash = ExpectedHash;
+							return true;
+						}
+					}
+				}
+			}
+		}
+
+		return false;
+	}
+
+	return true;
+}
+
+FRigUnit_GetBoolAnimationChannel_Execute()
+{
+	Value = false;
+	
+	if(!UpdateCache(Context.Hierarchy, Control, Channel, CachedChannelKey, CachedChannelHash))
+	{
+		return;
+	}
+
+	if(const FRigControlElement* ChannelElement = Context.Hierarchy->Find<FRigControlElement>(CachedChannelKey))
+	{
+		if(ChannelElement->Settings.ControlType == ERigControlType::Bool)
+		{
+			const FRigControlValue StoredValue = Context.Hierarchy->GetControlValueByIndex(ChannelElement->GetIndex(), bInitial ? ERigControlValueType::Initial : ERigControlValueType::Current);
+			Value = StoredValue.Get<bool>();
+		}
+		else
+		{
+			UE_CONTROLRIG_RIGUNIT_REPORT_ERROR(TEXT("Animation Channel %s is not a Bool."), *CachedChannelKey.ToString());
+		}
+	}
+}
+
+FRigUnit_GetFloatAnimationChannel_Execute()
+{
+	Value = 0.f;
+	
+	if(!UpdateCache(Context.Hierarchy, Control, Channel, CachedChannelKey, CachedChannelHash))
+	{
+		return;
+	}
+
+	if(const FRigControlElement* ChannelElement = Context.Hierarchy->Find<FRigControlElement>(CachedChannelKey))
+	{
+		if(ChannelElement->Settings.ControlType == ERigControlType::Float)
+		{
+			const FRigControlValue StoredValue = Context.Hierarchy->GetControlValueByIndex(ChannelElement->GetIndex(), bInitial ? ERigControlValueType::Initial : ERigControlValueType::Current);
+			Value = StoredValue.Get<float>();
+		}
+		else
+		{
+			UE_CONTROLRIG_RIGUNIT_REPORT_ERROR(TEXT("Animation Channel %s is not a Float."), *CachedChannelKey.ToString());
+		}
+	}
+}
+
+FRigUnit_GetIntAnimationChannel_Execute()
+{
+	Value = 0;
+	
+	if(!UpdateCache(Context.Hierarchy, Control, Channel, CachedChannelKey, CachedChannelHash))
+	{
+		return;
+	}
+
+	if(const FRigControlElement* ChannelElement = Context.Hierarchy->Find<FRigControlElement>(CachedChannelKey))
+	{
+		if(ChannelElement->Settings.ControlType == ERigControlType::Integer)
+		{
+			const FRigControlValue StoredValue = Context.Hierarchy->GetControlValueByIndex(ChannelElement->GetIndex(), bInitial ? ERigControlValueType::Initial : ERigControlValueType::Current);
+			Value = StoredValue.Get<int32>();
+		}
+		else
+		{
+			UE_CONTROLRIG_RIGUNIT_REPORT_ERROR(TEXT("Animation Channel %s is not an Integer (or Enum)."), *CachedChannelKey.ToString());
+		}
+	}
+}
+
+FRigUnit_GetVector2DAnimationChannel_Execute()
+{
+	Value = FVector2D::ZeroVector;
+	
+	if(!UpdateCache(Context.Hierarchy, Control, Channel, CachedChannelKey, CachedChannelHash))
+	{
+		return;
+	}
+
+	if(const FRigControlElement* ChannelElement = Context.Hierarchy->Find<FRigControlElement>(CachedChannelKey))
+	{
+		if(ChannelElement->Settings.ControlType == ERigControlType::Vector2D)
+		{
+			const FRigControlValue StoredValue = Context.Hierarchy->GetControlValueByIndex(ChannelElement->GetIndex(), bInitial ? ERigControlValueType::Initial : ERigControlValueType::Current);
+			Value = FVector2D(StoredValue.Get<FVector2f>());
+		}
+		else
+		{
+			UE_CONTROLRIG_RIGUNIT_REPORT_ERROR(TEXT("Animation Channel %s is not a Vector2D."), *CachedChannelKey.ToString());
+		}
+	}
+}
+
+FRigUnit_GetVectorAnimationChannel_Execute()
+{
+	Value = FVector::ZeroVector;
+	
+	if(!UpdateCache(Context.Hierarchy, Control, Channel, CachedChannelKey, CachedChannelHash))
+	{
+		return;
+	}
+
+	if(const FRigControlElement* ChannelElement = Context.Hierarchy->Find<FRigControlElement>(CachedChannelKey))
+	{
+		if(ChannelElement->Settings.ControlType == ERigControlType::Position || ChannelElement->Settings.ControlType == ERigControlType::Scale)
+		{
+			const FRigControlValue StoredValue = Context.Hierarchy->GetControlValueByIndex(ChannelElement->GetIndex(), bInitial ? ERigControlValueType::Initial : ERigControlValueType::Current);
+			Value = FVector(StoredValue.Get<FVector3f>());
+		}
+		else
+		{
+			UE_CONTROLRIG_RIGUNIT_REPORT_ERROR(TEXT("Animation Channel %s is not a Vector (Position or Scale)."), *CachedChannelKey.ToString());
+		}
+	}
+}
+
+FRigUnit_GetRotatorAnimationChannel_Execute()
+{
+	Value = FRotator::ZeroRotator;
+	
+	if(!UpdateCache(Context.Hierarchy, Control, Channel, CachedChannelKey, CachedChannelHash))
+	{
+		return;
+	}
+
+	if(const FRigControlElement* ChannelElement = Context.Hierarchy->Find<FRigControlElement>(CachedChannelKey))
+	{
+		if(ChannelElement->Settings.ControlType == ERigControlType::Rotator)
+		{
+			const FRigControlValue StoredValue = Context.Hierarchy->GetControlValueByIndex(ChannelElement->GetIndex(), bInitial ? ERigControlValueType::Initial : ERigControlValueType::Current);
+			Value = FRotator::MakeFromEuler(FVector(StoredValue.Get<FVector3f>()));
+		}
+		else
+		{
+			UE_CONTROLRIG_RIGUNIT_REPORT_ERROR(TEXT("Animation Channel %s is not a Rotator."), *CachedChannelKey.ToString());
+		}
+	}
+}
+
+FRigUnit_GetTransformAnimationChannel_Execute()
+{
+	Value = FTransform::Identity;
+	
+	if(!UpdateCache(Context.Hierarchy, Control, Channel, CachedChannelKey, CachedChannelHash))
+	{
+		return;
+	}
+
+	if(const FRigControlElement* ChannelElement = Context.Hierarchy->Find<FRigControlElement>(CachedChannelKey))
+	{
+		if(ChannelElement->Settings.ControlType == ERigControlType::Transform)
+		{
+			const FRigControlValue StoredValue = Context.Hierarchy->GetControlValueByIndex(ChannelElement->GetIndex(), bInitial ? ERigControlValueType::Initial : ERigControlValueType::Current);
+			Value = StoredValue.Get<FRigControlValue::FTransform_Float>().ToTransform();
+		}
+		else if(ChannelElement->Settings.ControlType == ERigControlType::EulerTransform)
+		{
+			const FRigControlValue StoredValue = Context.Hierarchy->GetControlValueByIndex(ChannelElement->GetIndex(), bInitial ? ERigControlValueType::Initial : ERigControlValueType::Current);
+			Value = StoredValue.Get<FRigControlValue::FEulerTransform_Float>().ToTransform().ToFTransform();
+		}
+		else
+		{
+			UE_CONTROLRIG_RIGUNIT_REPORT_ERROR(TEXT("Animation Channel %s is not a EulerTransform / Transform."), *CachedChannelKey.ToString());
+		}
+	}
+}
+
+FRigUnit_SetBoolAnimationChannel_Execute()
+{
+	if(!UpdateCache(Context.Hierarchy, Control, Channel, CachedChannelKey, CachedChannelHash))
+	{
+		return;
+	}
+
+	if(const FRigControlElement* ChannelElement = Context.Hierarchy->Find<FRigControlElement>(CachedChannelKey))
+	{
+		if(ChannelElement->Settings.ControlType == ERigControlType::Bool)
+		{
+			const FRigControlValue ValueToStore = FRigControlValue::Make<bool>(Value);
+			Context.Hierarchy->SetControlValueByIndex(ChannelElement->GetIndex(), ValueToStore, bInitial ? ERigControlValueType::Initial : ERigControlValueType::Current);
+		}
+		else
+		{
+			UE_CONTROLRIG_RIGUNIT_REPORT_ERROR(TEXT("Animation Channel %s is not a Bool."), *CachedChannelKey.ToString());
+		}
+	}
+}
+
+FRigUnit_SetFloatAnimationChannel_Execute()
+{
+	if(!UpdateCache(Context.Hierarchy, Control, Channel, CachedChannelKey, CachedChannelHash))
+	{
+		return;
+	}
+
+	if(const FRigControlElement* ChannelElement = Context.Hierarchy->Find<FRigControlElement>(CachedChannelKey))
+	{
+		if(ChannelElement->Settings.ControlType == ERigControlType::Float)
+		{
+			const FRigControlValue ValueToStore = FRigControlValue::Make<float>(Value);
+			Context.Hierarchy->SetControlValueByIndex(ChannelElement->GetIndex(), ValueToStore, bInitial ? ERigControlValueType::Initial : ERigControlValueType::Current);
+		}
+		else
+		{
+			UE_CONTROLRIG_RIGUNIT_REPORT_ERROR(TEXT("Animation Channel %s is not a Float."), *CachedChannelKey.ToString());
+		}
+	}
+}
+
+FRigUnit_SetIntAnimationChannel_Execute()
+{
+	if(!UpdateCache(Context.Hierarchy, Control, Channel, CachedChannelKey, CachedChannelHash))
+	{
+		return;
+	}
+
+	if(const FRigControlElement* ChannelElement = Context.Hierarchy->Find<FRigControlElement>(CachedChannelKey))
+	{
+		if(ChannelElement->Settings.ControlType == ERigControlType::Integer)
+		{
+			const FRigControlValue ValueToStore = FRigControlValue::Make<int32>(Value);
+			Context.Hierarchy->SetControlValueByIndex(ChannelElement->GetIndex(), ValueToStore, bInitial ? ERigControlValueType::Initial : ERigControlValueType::Current);
+		}
+		else
+		{
+			UE_CONTROLRIG_RIGUNIT_REPORT_ERROR(TEXT("Animation Channel %s is not an Integer (or Enum)."), *CachedChannelKey.ToString());
+		}
+	}
+}
+
+FRigUnit_SetVector2DAnimationChannel_Execute()
+{
+	if(!UpdateCache(Context.Hierarchy, Control, Channel, CachedChannelKey, CachedChannelHash))
+	{
+		return;
+	}
+
+	if(const FRigControlElement* ChannelElement = Context.Hierarchy->Find<FRigControlElement>(CachedChannelKey))
+	{
+		if(ChannelElement->Settings.ControlType == ERigControlType::Vector2D)
+		{
+			const FRigControlValue ValueToStore = FRigControlValue::Make<FVector2f>(FVector2f(Value));
+			Context.Hierarchy->SetControlValueByIndex(ChannelElement->GetIndex(), ValueToStore, bInitial ? ERigControlValueType::Initial : ERigControlValueType::Current);
+		}
+		else
+		{
+			UE_CONTROLRIG_RIGUNIT_REPORT_ERROR(TEXT("Animation Channel %s is not a Vector2D."), *CachedChannelKey.ToString());
+		}
+	}
+}
+
+FRigUnit_SetVectorAnimationChannel_Execute()
+{
+	if(!UpdateCache(Context.Hierarchy, Control, Channel, CachedChannelKey, CachedChannelHash))
+	{
+		return;
+	}
+
+	if(const FRigControlElement* ChannelElement = Context.Hierarchy->Find<FRigControlElement>(CachedChannelKey))
+	{
+		if(ChannelElement->Settings.ControlType == ERigControlType::Position || ChannelElement->Settings.ControlType == ERigControlType::Scale)
+		{
+			const FRigControlValue ValueToStore = FRigControlValue::Make<FVector3f>(FVector3f(Value));
+			Context.Hierarchy->SetControlValueByIndex(ChannelElement->GetIndex(), ValueToStore, bInitial ? ERigControlValueType::Initial : ERigControlValueType::Current);
+		}
+		else
+		{
+			UE_CONTROLRIG_RIGUNIT_REPORT_ERROR(TEXT("Animation Channel %s is not a Vector (Position or Scale)."), *CachedChannelKey.ToString());
+		}
+	}
+}
+
+FRigUnit_SetRotatorAnimationChannel_Execute()
+{
+	if(!UpdateCache(Context.Hierarchy, Control, Channel, CachedChannelKey, CachedChannelHash))
+	{
+		return;
+	}
+
+	if(const FRigControlElement* ChannelElement = Context.Hierarchy->Find<FRigControlElement>(CachedChannelKey))
+	{
+		if(ChannelElement->Settings.ControlType == ERigControlType::Rotator)
+		{
+			const FRigControlValue ValueToStore = FRigControlValue::Make<FVector3f>(FVector3f(Value.Euler()));
+			Context.Hierarchy->SetControlValueByIndex(ChannelElement->GetIndex(), ValueToStore, bInitial ? ERigControlValueType::Initial : ERigControlValueType::Current);
+		}
+		else
+		{
+			UE_CONTROLRIG_RIGUNIT_REPORT_ERROR(TEXT("Animation Channel %s is not a Rotator."), *CachedChannelKey.ToString());
+		}
+	}
+}
+
+FRigUnit_SetTransformAnimationChannel_Execute()
+{
+	if(!UpdateCache(Context.Hierarchy, Control, Channel, CachedChannelKey, CachedChannelHash))
+	{
+		return;
+	}
+
+	if(const FRigControlElement* ChannelElement = Context.Hierarchy->Find<FRigControlElement>(CachedChannelKey))
+	{
+		if(ChannelElement->Settings.ControlType == ERigControlType::Transform)
+		{
+			const FRigControlValue ValueToStore = FRigControlValue::Make<FRigControlValue::FTransform_Float>(Value);
+			Context.Hierarchy->SetControlValueByIndex(ChannelElement->GetIndex(), ValueToStore, bInitial ? ERigControlValueType::Initial : ERigControlValueType::Current);
+		}
+		else if(ChannelElement->Settings.ControlType == ERigControlType::EulerTransform)
+		{
+			const FRigControlValue ValueToStore = FRigControlValue::Make<FRigControlValue::FEulerTransform_Float>(FEulerTransform(Value));
+			Context.Hierarchy->SetControlValueByIndex(ChannelElement->GetIndex(), ValueToStore, bInitial ? ERigControlValueType::Initial : ERigControlValueType::Current);
+		}
+		else
+		{
+			UE_CONTROLRIG_RIGUNIT_REPORT_ERROR(TEXT("Animation Channel %s is not a EulerTransform / Transform."), *CachedChannelKey.ToString());
+		}
+	}
+}
\ No newline at end of file
diff --git a/Engine/Plugins/Animation/ControlRig/Source/ControlRig/Private/Units/Hierarchy/RigUnit_ControlChannel.h b/Engine/Plugins/Animation/ControlRig/Source/ControlRig/Private/Units/Hierarchy/RigUnit_ControlChannel.h
new file mode 100644
--- /dev/null
+++ b/Engine/Plugins/Animation/ControlRig/Source/ControlRig/Private/Units/Hierarchy/RigUnit_ControlChannel.h
@@ -0,0 +1,362 @@
+// Copyright Epic Games, Inc. All Rights Reserved.
+
+#pragma once
+
+#include "Units/RigUnit.h"
+#include "RigUnit_ControlChannel.generated.h"
+
+/**
+ * Get Animation Channel is used to retrieve a control's animation channel value
+ */
+USTRUCT(meta=(Abstract, Category="Controls", DocumentationPolicy = "Strict", NodeColor="0.462745, 1,0, 0.329412",Varying))
+struct CONTROLRIG_API FRigUnit_GetAnimationChannelBase : public FRigUnit
+{
+	GENERATED_BODY()
+
+	FRigUnit_GetAnimationChannelBase()
+	{
+		Control = Channel = NAME_None;
+		bInitial = false;
+		CachedChannelKey = FRigElementKey();
+		CachedChannelHash = INDEX_NONE;
+	}
+
+	static bool UpdateCache(const URigHierarchy* InHierarchy, const FName& Control, const FName& Channel, FRigElementKey& Key, int32& Hash);
+
+	/**
+	 * The name of the Control to retrieve the value for.
+	 */
+	UPROPERTY(meta = (Input, CustomWidget = "ControlName"))
+	FName Control;
+
+	/**
+	 * The name of the animation channel to retrieve the value for.
+	 */
+	UPROPERTY(meta = (Input, CustomWidget = "AnimationChannelName"))
+	FName Channel;
+
+	/**
+	 * If set to true the initial value will be returned
+	 */
+	UPROPERTY(meta = (Input))
+	bool bInitial;
+
+	// Used to cache the internally used bone index
+	UPROPERTY()
+	FRigElementKey CachedChannelKey;
+
+	// A hash combining the control, channel and topology identifiers
+	UPROPERTY()
+	int32 CachedChannelHash;
+};
+
+/**
+ * Get Bool Channel is used to retrieve a control's animation channel value
+ */
+USTRUCT(meta=(DisplayName="Get Bool Channel", Keywords="Animation,Channel", TemplateName="GetAnimationChannel"))
+struct CONTROLRIG_API FRigUnit_GetBoolAnimationChannel : public FRigUnit_GetAnimationChannelBase
+{
+	GENERATED_BODY()
+
+	FRigUnit_GetBoolAnimationChannel()
+		: FRigUnit_GetAnimationChannelBase()
+		, Value(false)
+	{}
+
+	RIGVM_METHOD()
+	virtual void Execute(const FRigUnitContext& Context) override;
+
+	// The current value of the animation channel
+	UPROPERTY(meta=(Output))
+	bool Value;
+};
+
+/**
+ * Get Float Channel is used to retrieve a control's animation channel value
+ */
+USTRUCT(meta=(DisplayName="Get Float Channel", Keywords="Animation,Channel", TemplateName="GetAnimationChannel"))
+struct CONTROLRIG_API FRigUnit_GetFloatAnimationChannel : public FRigUnit_GetAnimationChannelBase
+{
+	GENERATED_BODY()
+
+	FRigUnit_GetFloatAnimationChannel()
+		: FRigUnit_GetAnimationChannelBase()
+		, Value(0.f)
+	{}
+
+	RIGVM_METHOD()
+	virtual void Execute(const FRigUnitContext& Context) override;
+
+	// The current value of the animation channel
+	UPROPERTY(meta=(Output, UIMin=0, UIMax=1))
+	float Value;
+};
+
+/**
+ * Get Int Channel is used to retrieve a control's animation channel value
+ */
+USTRUCT(meta=(DisplayName="Get Int Channel", Keywords="Animation,Channel", TemplateName="GetAnimationChannel"))
+struct CONTROLRIG_API FRigUnit_GetIntAnimationChannel : public FRigUnit_GetAnimationChannelBase
+{
+	GENERATED_BODY()
+
+	FRigUnit_GetIntAnimationChannel()
+		: FRigUnit_GetAnimationChannelBase()
+		, Value(0)
+	{}
+
+	RIGVM_METHOD()
+	virtual void Execute(const FRigUnitContext& Context) override;
+
+	// The current value of the animation channel
+	UPROPERTY(meta=(Output))
+	int32 Value;
+};
+
+/**
+ * Get Vector2D Channel is used to retrieve a control's animation channel value
+ */
+USTRUCT(meta=(DisplayName="Get Vector2D Channel", Keywords="Animation,Channel", TemplateName="GetAnimationChannel"))
+struct CONTROLRIG_API FRigUnit_GetVector2DAnimationChannel : public FRigUnit_GetAnimationChannelBase
+{
+	GENERATED_BODY()
+
+	FRigUnit_GetVector2DAnimationChannel()
+		: FRigUnit_GetAnimationChannelBase()
+		, Value(FVector2D::ZeroVector)
+	{}
+
+	RIGVM_METHOD()
+	virtual void Execute(const FRigUnitContext& Context) override;
+
+	// The current value of the animation channel
+	UPROPERTY(meta=(Output))
+	FVector2D Value;
+};
+
+/**
+ * Get Vector Channel is used to retrieve a control's animation channel value
+ */
+USTRUCT(meta=(DisplayName="Get Vector Channel", Keywords="Animation,Channel", TemplateName="GetAnimationChannel"))
+struct CONTROLRIG_API FRigUnit_GetVectorAnimationChannel : public FRigUnit_GetAnimationChannelBase
+{
+	GENERATED_BODY()
+
+	FRigUnit_GetVectorAnimationChannel()
+		: FRigUnit_GetAnimationChannelBase()
+		, Value(FVector::ZeroVector)
+	{}
+
+	RIGVM_METHOD()
+	virtual void Execute(const FRigUnitContext& Context) override;
+
+	// The current value of the animation channel
+	UPROPERTY(meta=(Output))
+	FVector Value;
+};
+
+/**
+ * Get Rotator Channel is used to retrieve a control's animation channel value
+ */
+USTRUCT(meta=(DisplayName="Get Rotator Channel", Keywords="Animation,Channel", TemplateName="GetAnimationChannel"))
+struct CONTROLRIG_API FRigUnit_GetRotatorAnimationChannel : public FRigUnit_GetAnimationChannelBase
+{
+	GENERATED_BODY()
+
+	FRigUnit_GetRotatorAnimationChannel()
+		: FRigUnit_GetAnimationChannelBase()
+		, Value(FRotator::ZeroRotator)
+	{}
+
+	RIGVM_METHOD()
+	virtual void Execute(const FRigUnitContext& Context) override;
+
+	// The current value of the animation channel
+	UPROPERTY(meta=(Output))
+	FRotator Value;
+};
+
+/**
+ * Get Transform Channel is used to retrieve a control's animation channel value
+ */
+USTRUCT(meta=(DisplayName="Get Transform Channel", Keywords="Animation,Channel", TemplateName="GetAnimationChannel"))
+struct CONTROLRIG_API FRigUnit_GetTransformAnimationChannel : public FRigUnit_GetAnimationChannelBase
+{
+	GENERATED_BODY()
+
+	FRigUnit_GetTransformAnimationChannel()
+		: FRigUnit_GetAnimationChannelBase()
+		, Value(FTransform::Identity)
+	{}
+
+	RIGVM_METHOD()
+	virtual void Execute(const FRigUnitContext& Context) override;
+
+	// The current value of the animation channel
+	UPROPERTY(meta=(Output))
+	FTransform Value;
+};
+
+/**
+ * Set Animation Channel is used to change a control's animation channel value
+ */
+USTRUCT(meta = (Abstract))
+struct CONTROLRIG_API FRigUnit_SetAnimationChannelBase : public FRigUnit_GetAnimationChannelBase
+{
+	GENERATED_BODY()
+
+	FRigUnit_SetAnimationChannelBase()
+		:FRigUnit_GetAnimationChannelBase()
+	{
+	}
+
+	UPROPERTY(DisplayName = "Execute", Transient, meta = (Input, Output))
+	FControlRigExecuteContext ExecuteContext;
+};
+
+/**
+ * Set Bool Channel is used to set a control's animation channel value
+ */
+USTRUCT(meta=(DisplayName="Set Bool Channel", Keywords="Animation,Channel", TemplateName="SetAnimationChannel"))
+struct CONTROLRIG_API FRigUnit_SetBoolAnimationChannel : public FRigUnit_SetAnimationChannelBase
+{
+	GENERATED_BODY()
+
+	FRigUnit_SetBoolAnimationChannel()
+		: FRigUnit_SetAnimationChannelBase()
+		, Value(false)
+	{}
+
+	RIGVM_METHOD()
+	virtual void Execute(const FRigUnitContext& Context) override;
+
+	// The new value of the animation channel
+	UPROPERTY(meta=(Input))
+	bool Value;
+};
+
+/**
+ * Set Float Channel is used to set a control's animation channel value
+ */
+USTRUCT(meta=(DisplayName="Set Float Channel", Keywords="Animation,Channel", TemplateName="SetAnimationChannel"))
+struct CONTROLRIG_API FRigUnit_SetFloatAnimationChannel : public FRigUnit_SetAnimationChannelBase
+{
+	GENERATED_BODY()
+
+	FRigUnit_SetFloatAnimationChannel()
+		: FRigUnit_SetAnimationChannelBase()
+		, Value(0.f)
+	{}
+
+	RIGVM_METHOD()
+	virtual void Execute(const FRigUnitContext& Context) override;
+
+	// The new value of the animation channel
+	UPROPERTY(meta=(Input, UIMin=0, UIMax=1))
+	float Value;
+};
+
+/**
+ * Set Int Channel is used to set a control's animation channel value
+ */
+USTRUCT(meta=(DisplayName="Set Int Channel", Keywords="Animation,Channel", TemplateName="SetAnimationChannel"))
+struct CONTROLRIG_API FRigUnit_SetIntAnimationChannel : public FRigUnit_SetAnimationChannelBase
+{
+	GENERATED_BODY()
+
+	FRigUnit_SetIntAnimationChannel()
+		: FRigUnit_SetAnimationChannelBase()
+		, Value(0)
+	{}
+
+	RIGVM_METHOD()
+	virtual void Execute(const FRigUnitContext& Context) override;
+
+	// The new value of the animation channel
+	UPROPERTY(meta=(Input))
+	int32 Value;
+};
+
+/**
+ * Set Vector2D Channel is used to set a control's animation channel value
+ */
+USTRUCT(meta=(DisplayName="Set Vector2D Channel", Keywords="Animation,Channel", TemplateName="SetAnimationChannel"))
+struct CONTROLRIG_API FRigUnit_SetVector2DAnimationChannel : public FRigUnit_SetAnimationChannelBase
+{
+	GENERATED_BODY()
+
+	FRigUnit_SetVector2DAnimationChannel()
+		: FRigUnit_SetAnimationChannelBase()
+		, Value(FVector2D::ZeroVector)
+	{}
+
+	RIGVM_METHOD()
+	virtual void Execute(const FRigUnitContext& Context) override;
+
+	// The new value of the animation channel
+	UPROPERTY(meta=(Input))
+	FVector2D Value;
+};
+
+/**
+ * Set Vector Channel is used to set a control's animation channel value
+ */
+USTRUCT(meta=(DisplayName="Set Vector Channel", Keywords="Animation,Channel", TemplateName="SetAnimationChannel"))
+struct CONTROLRIG_API FRigUnit_SetVectorAnimationChannel : public FRigUnit_SetAnimationChannelBase
+{
+	GENERATED_BODY()
+
+	FRigUnit_SetVectorAnimationChannel()
+		: FRigUnit_SetAnimationChannelBase()
+		, Value(FVector::ZeroVector)
+	{}
+
+	RIGVM_METHOD()
+	virtual void Execute(const FRigUnitContext& Context) override;
+
+	// The new value of the animation channel
+	UPROPERTY(meta=(Input))
+	FVector Value;
+};
+
+/**
+ * Set Rotator Channel is used to set a control's animation channel value
+ */
+USTRUCT(meta=(DisplayName="Set Rotator Channel", Keywords="Animation,Channel", TemplateName="SetAnimationChannel"))
+struct CONTROLRIG_API FRigUnit_SetRotatorAnimationChannel : public FRigUnit_SetAnimationChannelBase
+{
+	GENERATED_BODY()
+
+	FRigUnit_SetRotatorAnimationChannel()
+		: FRigUnit_SetAnimationChannelBase()
+		, Value(FRotator::ZeroRotator)
+	{}
+
+	RIGVM_METHOD()
+	virtual void Execute(const FRigUnitContext& Context) override;
+
+	// The new value of the animation channel
+	UPROPERTY(meta=(Input))
+	FRotator Value;
+};
+
+/**
+ * Set Transform Channel is used to set a control's animation channel value
+ */
+USTRUCT(meta=(DisplayName="Set Transform Channel", Keywords="Animation,Channel", TemplateName="SetAnimationChannel"))
+struct CONTROLRIG_API FRigUnit_SetTransformAnimationChannel : public FRigUnit_SetAnimationChannelBase
+{
+	GENERATED_BODY()
+
+	FRigUnit_SetTransformAnimationChannel()
+		: FRigUnit_SetAnimationChannelBase()
+		, Value(FTransform::Identity)
+	{}
+
+	RIGVM_METHOD()
+	virtual void Execute(const FRigUnitContext& Context) override;
+
+	// The new value of the animation channel
+	UPROPERTY(meta=(Input))
+	FTransform Value;
+};
diff --git a/Engine/Plugins/Animation/ControlRig/Source/ControlRig/Public/Rigs/RigHierarchy.h b/Engine/Plugins/Animation/ControlRig/Source/ControlRig/Public/Rigs/RigHierarchy.h
--- a/Engine/Plugins/Animation/ControlRig/Source/ControlRig/Public/Rigs/RigHierarchy.h
+++ b/Engine/Plugins/Animation/ControlRig/Source/ControlRig/Public/Rigs/RigHierarchy.h
@@ -869,6 +869,15 @@ public:
 	 */
 	bool IsNameAvailable(const FString& InPotentialNewName, ERigElementType InType, FString* OutErrorMessage = nullptr) const;
 
+	/**
+	 * Returns true if a given display name is available.
+	 * @param InParentElement The element to check the display name under
+	 * @param InPotentialNewDisplayName The name to test for availability
+	 * @param OutErrorMessage An optional pointer to return a potential error message 
+	 * @return Returns true if the name is available.
+	 */
+	bool IsDisplayNameAvailable(const FRigElementKey& InParentElement, const FString& InPotentialNewDisplayName, FString* OutErrorMessage = nullptr) const;
+
 	/**
 	 * Returns a valid new name for a to-be-added element.
 	 * @param InPotentialNewName The name to be sanitized and adjusted for availability
@@ -877,6 +886,14 @@ public:
 	 */
 	FName GetSafeNewName(const FString& InPotentialNewName, ERigElementType InType) const;
 
+	/**
+	 * Returns a valid new display name for a control
+	 * @param InParentElement The element to check the display name under
+	 * @param InPotentialNewDisplayName The name to be sanitized and adjusted for availability
+	 * @return Returns the name to use for the to-be-added element.
+	 */
+	FName GetSafeNewDisplayName(const FRigElementKey& InParentElement, const FString& InPotentialNewDisplayName) const;
+
 	/**
 	 * Returns the modified event, which can be used to 
 	 * subscribe to topological changes happening within the hierarchy.
diff --git a/Engine/Plugins/Animation/ControlRig/Source/ControlRig/Public/Rigs/RigHierarchyController.h b/Engine/Plugins/Animation/ControlRig/Source/ControlRig/Public/Rigs/RigHierarchyController.h
--- a/Engine/Plugins/Animation/ControlRig/Source/ControlRig/Public/Rigs/RigHierarchyController.h
+++ b/Engine/Plugins/Animation/ControlRig/Source/ControlRig/Public/Rigs/RigHierarchyController.h
@@ -146,6 +146,43 @@ public:
 		return AddControl(InName, InParent, InSettings, InValue, FTransform::Identity, FTransform::Identity, bSetupUndo, bPrintPythonCommand);
 	}
 
+		/**
+	 * Adds a control to the hierarchy
+	 * @param InName The suggested name of the new animation channel - will eventually be corrected by the namespace
+	 * @param InParentControl The parent of the new animation channel.
+	 * @param InSettings All of the animation channel's settings
+	 * @param bSetupUndo If set to true the stack will record the change for undo / redo
+	 * @param bPrintPythonCommand If set to true a python command equivalent to this call will be printed out
+	 * @return The key for the newly created animation channel.
+	 */
+    FRigElementKey AddAnimationChannel(
+    	FName InName,
+    	FRigElementKey InParentControl,
+    	FRigControlSettings InSettings,
+        bool bSetupUndo = true,
+        bool bPrintPythonCommand = false
+    );
+
+	/**
+	 * Adds a control to the hierarchy
+	 * @param InName The suggested name of the new animation channel - will eventually be corrected by the namespace
+	 * @param InParentControl The parent of the new animation channel.
+	 * @param InSettings All of the animation channel's settings
+	 * @param bSetupUndo If set to true the stack will record the change for undo / redo
+	 * @return The key for the newly created animation channel.
+	 */
+	UFUNCTION(BlueprintCallable, Category = URigHierarchyController, meta = (DisplayName = "Add Control", ScriptName = "AddAnimationChannel"))
+    FORCEINLINE FRigElementKey AddAnimationChannel_ForBlueprint(
+        FName InName,
+        FRigElementKey InParentControl,
+        FRigControlSettings InSettings,
+        bool bSetupUndo = true,
+        bool bPrintPythonCommand = false
+    )
+	{
+		return AddAnimationChannel(InName, InParentControl, InSettings, bSetupUndo, bPrintPythonCommand);
+	}
+
 	/**
 	 * Adds a curve to the hierarchy
 	 * @param InName The suggested name of the new curve - will eventually be corrected by the namespace
@@ -365,6 +402,18 @@ public:
 	UFUNCTION(BlueprintCallable, Category = URigHierarchyController)
     FRigElementKey RenameElement(FRigElementKey InElement, FName InName, bool bSetupUndo = false, bool bPrintPythonCommand = false, bool bClearSelection = true);
 
+	/**
+ 	 * Sets the display name on a control
+ 	 * @param InControl The key of the control to change the display name for
+ 	 * @param InDisplayName The new display name to set for the control
+ 	 * @param bRenameElement True if the control should also be renamed
+ 	 * @param bSetupUndo If set to true the stack will record the change for undo / redo
+ 	 * @param bPrintPythonCommand If set to true a python command equivalent to this call will be printed out
+	 * @return Returns the new display name used for the control
+	 */
+	UFUNCTION(BlueprintCallable, Category = URigHierarchyController)
+	FName SetDisplayName(FRigElementKey InControl, FName InDisplayName, bool bRenameElement = false, bool bSetupUndo = false, bool bPrintPythonCommand = false);
+
 	/**
 	 * Adds a new parent to an element. For elements that allow only one parent the parent will be replaced (Same as ::SetParent).
 	 * @param InChild The key of the element to add the parent for
@@ -560,6 +609,15 @@ private:
 	 */
     bool RenameElement(FRigBaseElement* InElement, const FName &InName, bool bClearSelection = true);
 
+	/**
+ 	 * Sets the display name on a control
+ 	 * @param InControlElement The element to change the display name for
+ 	 * @param InDisplayName The new display name to set for the control
+ 	 * @param bRenameElement True if the control should also be renamed
+ 	 * @return Returns true if successful.
+ 	 */
+	FName SetDisplayName(FRigControlElement* InControlElement, const FName &InDisplayName, bool bRenameElement = false);
+
 	/**
 	 * Removes an existing parent from an element in the hierarchy. For elements that allow only one parent the element will be unparented (same as ::RemoveAllParents)
 	 * @param InChild The element to remove the parent for
diff --git a/Engine/Plugins/Animation/ControlRig/Source/ControlRig/Public/Rigs/RigHierarchyElements.h b/Engine/Plugins/Animation/ControlRig/Source/ControlRig/Public/Rigs/RigHierarchyElements.h
--- a/Engine/Plugins/Animation/ControlRig/Source/ControlRig/Public/Rigs/RigHierarchyElements.h
+++ b/Engine/Plugins/Animation/ControlRig/Source/ControlRig/Public/Rigs/RigHierarchyElements.h
@@ -996,7 +996,7 @@ struct CONTROLRIG_API FRigControlSettings
 	TArray<FRigElementKey> PreviouslyDrivenControls;
 
 	/**
-	 * If set to true the control will be grouped with the parent control in sequencer
+	 * If set to true the animation channel will be grouped with the parent control in sequencer
 	 */
 	UPROPERTY(BlueprintReadOnly, EditAnywhere, Category = Animation)
 	bool bGroupWithParentControl;
@@ -1130,6 +1130,8 @@ struct CONTROLRIG_API FRigControlElement : public FRigMultiParentElement
 		return FRigMultiParentElement::GetDisplayName();
 	}
 
+	FORCEINLINE bool IsAnimationChannel() const { return Settings.AnimationType == ERigControlAnimationType::AnimationChannel; }
+
 	virtual void Save(FArchive& A, URigHierarchy* Hierarchy, ESerializationPhase SerializationPhase) override;
 	virtual void Load(FArchive& Ar, URigHierarchy* Hierarchy, ESerializationPhase SerializationPhase) override;
 
diff --git a/Engine/Plugins/Animation/ControlRig/Source/ControlRigDeveloper/Public/Graph/ControlRigGraph.h b/Engine/Plugins/Animation/ControlRig/Source/ControlRigDeveloper/Public/Graph/ControlRigGraph.h
--- a/Engine/Plugins/Animation/ControlRig/Source/ControlRigDeveloper/Public/Graph/ControlRigGraph.h
+++ b/Engine/Plugins/Animation/ControlRig/Source/ControlRigDeveloper/Public/Graph/ControlRigGraph.h
@@ -94,6 +94,12 @@ public:
 
 private:
 
+	template<class T>
+	static bool IncludeElementInNameList(const T* InElement)
+	{
+		return true;
+	}
+
 	template<class T>
 	void CacheNameListForHierarchy(URigHierarchy* InHierarchy, TArray<TSharedPtr<FString>>& OutNameList)
 	{
@@ -102,7 +108,10 @@ private:
 		{
 			if(Element->IsA<T>())
 			{
-				Names.Add(Element->GetName().ToString());
+				if(IncludeElementInNameList<T>(Cast<T>(Element)))
+				{
+					Names.Add(Element->GetName().ToString());
+				}
 			}
 		}
 		Names.Sort();
@@ -184,3 +193,8 @@ private:
 	friend class UControlRigBlueprint;
 };
 
+template<>
+inline bool UControlRigGraph::IncludeElementInNameList<FRigControlElement>(const FRigControlElement* InElement)
+{
+	return !InElement->IsAnimationChannel();
+}
diff --git a/Engine/Plugins/Animation/ControlRig/Source/ControlRigEditor/Private/ControlRigElementDetails.cpp b/Engine/Plugins/Animation/ControlRig/Source/ControlRigEditor/Private/ControlRigElementDetails.cpp
--- a/Engine/Plugins/Animation/ControlRig/Source/ControlRigEditor/Private/ControlRigElementDetails.cpp
+++ b/Engine/Plugins/Animation/ControlRig/Source/ControlRigEditor/Private/ControlRigElementDetails.cpp
@@ -5,6 +5,7 @@
 #include "IDetailChildrenBuilder.h"
 #include "Widgets/Text/STextBlock.h"
 #include "Widgets/Input/SEditableTextBox.h"
+#include "Widgets/Text/SInlineEditableTextBlock.h"
 #include "Widgets/Input/SVectorInputBox.h"
 #include "Widgets/Input/SCheckBox.h"
 #include "Widgets/Input/SButton.h"
@@ -18,6 +19,7 @@
 #include "Graph/SControlRigGraphPinVariableBinding.h"
 #include "HAL/PlatformApplicationMisc.h"
 #include "Styling/AppStyle.h"
+#include "SRigHierarchyTreeView.h"
 
 #define LOCTEXT_NAMESPACE "ControlRigElementDetails"
 
@@ -378,13 +380,13 @@ FText FRigElementKeyDetails::GetElementNameAsText() const
 	return FText::FromString(GetElementName());
 }
 
-FSlateColor FRigElementKeyDetails::OnGetWidgetForeground(const TSharedPtr<SButton> Button) const
+FSlateColor FRigElementKeyDetails::OnGetWidgetForeground(const TSharedPtr<SButton> Button)
 {
 	float Alpha = (Button.IsValid() && Button->IsHovered()) ? FRigElementKeyDetailsDefs::ActivePinForegroundAlpha : FRigElementKeyDetailsDefs::InactivePinForegroundAlpha;
 	return FSlateColor(FLinearColor(1.f, 1.f, 1.f, Alpha));
 }
 
-FSlateColor FRigElementKeyDetails::OnGetWidgetBackground(const TSharedPtr<SButton> Button) const
+FSlateColor FRigElementKeyDetails::OnGetWidgetBackground(const TSharedPtr<SButton> Button)
 {
 	float Alpha = (Button.IsValid() && Button->IsHovered()) ? FRigElementKeyDetailsDefs::ActivePinBackgroundAlpha : FRigElementKeyDetailsDefs::InactivePinBackgroundAlpha;
 	return FSlateColor(FLinearColor(1.f, 1.f, 1.f, Alpha));
@@ -553,22 +555,71 @@ void FRigBaseElementDetails::CustomizeDetails(IDetailLayoutBuilder& DetailBuilde
 
 	IDetailCategoryBuilder& GeneralCategory = DetailBuilder.EditCategory(TEXT("General"), LOCTEXT("General", "General"));
 
-	GeneralCategory.AddCustomRow(FText::FromString(TEXT("Name")))
-	.NameContent()
-	[
-		SNew(STextBlock)
-		.Text(FText::FromString(TEXT("Name")))
-		.Font(IDetailLayoutBuilder::GetDetailFont())
-	]
-	.ValueContent()
-	[
-		SNew(SEditableTextBox)
-		.Font(IDetailLayoutBuilder::GetDetailFont())
-		.Text(this, &FRigBaseElementDetails::GetName)
-		.OnTextCommitted(this, &FRigBaseElementDetails::SetName)
-		.OnVerifyTextChanged(this, &FRigBaseElementDetails::OnVerifyNameChanged)
-		.IsEnabled(ObjectsBeingCustomized.Num() == 1)
-	];
+	const bool bAllAnimationChannels = !IsAnyControlNotOfAnimationType(ERigControlAnimationType::AnimationChannel);
+	if(bAllAnimationChannels)
+	{
+		GeneralCategory.AddCustomRow(FText::FromString(TEXT("Parent Control")))
+		.NameContent()
+		[
+			SNew(SInlineEditableTextBlock)
+			.Text(FText::FromString(TEXT("Parent Control")))
+			.Font(IDetailLayoutBuilder::GetDetailFont())
+			.IsEnabled(false)
+		]
+		.ValueContent()
+		[
+			SNew(SHorizontalBox)
+			+SHorizontalBox::Slot()
+			.AutoWidth()
+			.HAlign(HAlign_Left)
+			.VAlign(VAlign_Center)
+			[
+				SNew(SEditableTextBox)
+				.Font(IDetailLayoutBuilder::GetDetailFont())
+				.Text(this, &FRigBaseElementDetails::GetParentElementName)
+				.IsEnabled(false)
+			]
+			+SHorizontalBox::Slot()
+			.AutoWidth()
+			.HAlign(HAlign_Left)
+			.VAlign(VAlign_Center)
+			[
+				SAssignNew(SelectParentElementButton, SButton)
+				.ButtonStyle( FAppStyle::Get(), "NoBorder" )
+				.ButtonColorAndOpacity_Lambda([this]() { return FRigElementKeyDetails::OnGetWidgetBackground(SelectParentElementButton); })
+				.OnClicked(this, &FRigBaseElementDetails::OnSelectParentElementInHierarchyClicked)
+				.ContentPadding(0)
+				.ToolTipText(NSLOCTEXT("ControlRigElementDetails", "SelectParentInHierarchyToolTip", "Select Parent in hierarchy"))
+				[
+					SNew(SImage)
+					.ColorAndOpacity_Lambda( [this]() { return FRigElementKeyDetails::OnGetWidgetForeground(SelectParentElementButton); })
+					.Image(FAppStyle::GetBrush("Icons.Search"))
+				]
+			]
+		];
+	}
+
+	DetailBuilder.HideCategory(TEXT("RigElement"));
+
+	if(!bAllAnimationChannels)
+	{
+		GeneralCategory.AddCustomRow(FText::FromString(TEXT("Name")))
+		.NameContent()
+		[
+			SNew(STextBlock)
+			.Text(FText::FromString(TEXT("Name")))
+			.Font(IDetailLayoutBuilder::GetDetailFont())
+		]
+		.ValueContent()
+		[
+			SNew(SInlineEditableTextBlock)
+			.Font(IDetailLayoutBuilder::GetDetailFont())
+			.Text(this, &FRigBaseElementDetails::GetName)
+			.OnTextCommitted(this, &FRigBaseElementDetails::SetName)
+			.OnVerifyTextChanged(this, &FRigBaseElementDetails::OnVerifyNameChanged)
+			.IsEnabled(ObjectsBeingCustomized.Num() == 1 && !bAllAnimationChannels)
+		];
+	}
 
 	DetailBuilder.HideCategory(TEXT("RigElement"));
 }
@@ -607,6 +658,15 @@ FText FRigBaseElementDetails::GetName() const
 	return FText::FromName(GetElementKey().Name);
 }
 
+FText FRigBaseElementDetails::GetParentElementName() const
+{
+	if(ObjectsBeingCustomized.Num() > 1 || HierarchyBeingCustomized == nullptr)
+	{
+		return ControlRigDetailsMultipleValues;
+	}
+	return FText::FromName(HierarchyBeingCustomized->GetFirstParent(GetElementKey()).Name);
+}
+
 void FRigBaseElementDetails::SetName(const FText& InNewText, ETextCommit::Type InCommitType)
 {
 	if(InCommitType == ETextCommit::OnCleared)
@@ -824,6 +884,35 @@ void FRigBaseElementDetails::RegisterSectionMappings(FPropertyEditorModule& Prop
 {
 }
 
+FReply FRigBaseElementDetails::OnSelectParentElementInHierarchyClicked()
+{
+	if (HierarchyBeingCustomized)
+	{
+		FRigElementKey Key = GetElementKey();
+		if (Key.IsValid())
+		{
+			const FRigElementKey ParentKey = HierarchyBeingCustomized->GetFirstParent(Key);
+			if(ParentKey.IsValid())
+			{
+				return OnSelectElementClicked(ParentKey);
+			}
+		}	
+	}
+	return FReply::Handled();
+}
+
+FReply FRigBaseElementDetails::OnSelectElementClicked(const FRigElementKey& InKey)
+{
+	if (HierarchyBeingCustomized)
+	{
+		if (InKey.IsValid())
+		{
+			HierarchyBeingCustomized->GetController(true)->SetSelection({InKey});
+		}	
+	}
+	return FReply::Handled();
+}
+
 TSharedPtr<TArray<ERigTransformElementDetailsTransform::Type>> FRigTransformElementDetails::PickedTransforms;
 
 void FRigTransformElementDetails::CustomizeDetails(IDetailLayoutBuilder& DetailBuilder)
@@ -854,7 +943,14 @@ void FRigTransformElementDetails::CustomizeTransform(IDetailLayoutBuilder& Detai
 
 	TArray<FRigElementKey> Keys = GetElementKeys();
 	Keys = HierarchyBeingDebugged->SortKeys(Keys);
-	const bool bAllControls = !IsAnyElementNotOfType(ERigElementType::Control) && !IsAnyControlOfValueType(ERigControlType::Bool);
+
+	const bool bAllControls = !IsAnyElementNotOfType(ERigElementType::Control) && !IsAnyControlOfValueType(ERigControlType::Bool);;
+	const bool bAllAnimationChannels = !IsAnyControlNotOfAnimationType(ERigControlAnimationType::AnimationChannel);
+	if(bAllAnimationChannels)
+	{
+		return;
+	}
+
 	bool bShowLimits = false;
 	TArray<ERigTransformElementDetailsTransform::Type> TransformTypes;
 	TArray<FText> ButtonLabels;
@@ -1065,9 +1161,60 @@ void FRigTransformElementDetails::AddChoiceWidgetRow(IDetailCategoryBuilder& InC
 	];
 }
 
+FDetailWidgetRow& FRigTransformElementDetails::CreateTransformComponentValueWidgetRow(
+	ERigControlType InControlType,
+	URigHierarchy* HierarchyBeingDebugged,
+	const TArray<FRigElementKey>& Keys,
+	SAdvancedTransformInputBox<FEulerTransform>::FArguments TransformWidgetArgs,
+	IDetailCategoryBuilder& CategoryBuilder,
+	const FText& Label,
+	const FText& Tooltip,
+	ERigTransformElementDetailsTransform::Type CurrentTransformType,
+	ERigControlValueType ValueType,
+	TSharedPtr<SWidget> NameContent)
+{
+	TransformWidgetArgs
+	.Font(IDetailLayoutBuilder::GetDetailFont())
+	.AllowEditRotationRepresentation(false);
+
+	if(TransformWidgetArgs._DisplayRelativeWorld &&
+		!TransformWidgetArgs._OnGetIsComponentRelative.IsBound() &&
+		!TransformWidgetArgs._OnIsComponentRelativeChanged.IsBound())
+	{
+		TSharedRef<UE::Math::TVector<float>> IsComponentRelative = MakeShareable(new UE::Math::TVector<float>(1.f, 1.f, 1.f));
+		
+		TransformWidgetArgs
+		.OnGetIsComponentRelative_Lambda(
+			[IsComponentRelative](ESlateTransformComponent::Type InComponent)
+			{
+				return IsComponentRelative->operator[]((int32)InComponent) > 0.f;
+			})
+		.OnIsComponentRelativeChanged_Lambda(
+			[IsComponentRelative](ESlateTransformComponent::Type InComponent, bool bIsRelative)
+			{
+				IsComponentRelative->operator[]((int32)InComponent) = bIsRelative ? 1.f : 0.f;
+			});
+	}
+
+	TransformWidgetArgs.ConstructLocation(InControlType == ERigControlType::Position);
+	TransformWidgetArgs.ConstructRotation(InControlType == ERigControlType::Rotator);
+	TransformWidgetArgs.ConstructScale(InControlType == ERigControlType::Scale);
+
+	return CreateEulerTransformValueWidgetRow(
+		GetHierarchyBeingDebugged(),
+		Keys,
+		TransformWidgetArgs,
+		CategoryBuilder,
+		Label,
+		Tooltip,
+		CurrentTransformType,
+		ValueType,
+		NameContent);
+}
+
 TSharedPtr<TArray<ERigControlValueType>> FRigControlElementDetails::PickedValueTypes;
 
-void FRigTransformElementDetails::CreateEulerTransformValueWidgetRow(
+FDetailWidgetRow& FRigTransformElementDetails::CreateEulerTransformValueWidgetRow(
 	URigHierarchy* HierarchyBeingDebugged,
 	const TArray<FRigElementKey>& Keys,
 	SAdvancedTransformInputBox<FEulerTransform>::FArguments TransformWidgetArgs,
@@ -1075,21 +1222,29 @@ void FRigTransformElementDetails::CreateEulerTransformValueWidgetRow(
 	const FText& Label,
 	const FText& Tooltip,
 	ERigTransformElementDetailsTransform::Type CurrentTransformType,
-	ERigControlValueType ValueType)
+	ERigControlValueType ValueType,
+	TSharedPtr<SWidget> NameContent)
 {
 	const FRigElementTransformWidgetSettings& Settings = FRigElementTransformWidgetSettings::FindOrAdd(ValueType, CurrentTransformType, TransformWidgetArgs);
-	TSharedRef<UE::Math::TVector<float>> IsComponentRelativeStorage = Settings.IsComponentRelative;
-	
-	TransformWidgetArgs.OnGetIsComponentRelative_Lambda(
-		[IsComponentRelativeStorage](ESlateTransformComponent::Type InComponent)
-		{
-			return IsComponentRelativeStorage->operator[]((int32)InComponent) > 0.f;
-		})
-	.OnIsComponentRelativeChanged_Lambda(
-		[IsComponentRelativeStorage](ESlateTransformComponent::Type InComponent, bool bIsRelative)
-		{
-			IsComponentRelativeStorage->operator[]((int32)InComponent) = bIsRelative ? 1.f : 0.f;
-		});
+
+	const bool bDisplayRelativeWorldOnCurrent = TransformWidgetArgs._DisplayRelativeWorld; 
+	if(bDisplayRelativeWorldOnCurrent &&
+		!TransformWidgetArgs._OnGetIsComponentRelative.IsBound() &&
+		!TransformWidgetArgs._OnIsComponentRelativeChanged.IsBound())
+	{
+		TSharedRef<UE::Math::TVector<float>> IsComponentRelativeStorage = Settings.IsComponentRelative;
+		
+		TransformWidgetArgs.OnGetIsComponentRelative_Lambda(
+			[IsComponentRelativeStorage](ESlateTransformComponent::Type InComponent)
+			{
+				return IsComponentRelativeStorage->operator[]((int32)InComponent) > 0.f;
+			})
+		.OnIsComponentRelativeChanged_Lambda(
+			[IsComponentRelativeStorage](ESlateTransformComponent::Type InComponent, bool bIsRelative)
+			{
+				IsComponentRelativeStorage->operator[]((int32)InComponent) = bIsRelative ? 1.f : 0.f;
+			});
+	}
 
 	const TSharedPtr<ESlateRotationRepresentation::Type> RotationRepresentationStorage = Settings.RotationRepresentation;
 	TransformWidgetArgs.RotationRepresentation(RotationRepresentationStorage);
@@ -1406,7 +1561,7 @@ void FRigTransformElementDetails::CreateEulerTransformValueWidgetRow(
 		default:
 		{
 			TransformWidgetArgs.AllowEditRotationRepresentation(true);
-			TransformWidgetArgs.DisplayRelativeWorld(true);
+			TransformWidgetArgs.DisplayRelativeWorld(bDisplayRelativeWorldOnCurrent);
 			TransformWidgetArgs.DisplayToggle(false);
 			TransformWidgetArgs._OnGetToggleChecked.Unbind();
 			TransformWidgetArgs._OnToggleChanged.Unbind();
@@ -2078,11 +2233,41 @@ void FRigTransformElementDetails::CreateEulerTransformValueWidgetRow(
 		}
 	});
 
-	SAdvancedTransformInputBox<FEulerTransform>::ConstructGroupedTransformRows(
+	return SAdvancedTransformInputBox<FEulerTransform>::ConstructGroupedTransformRows(
 		CategoryBuilder, 
 		Label, 
 		Tooltip, 
-		TransformWidgetArgs);
+		TransformWidgetArgs,
+		NameContent);
+}
+
+ERigTransformElementDetailsTransform::Type FRigTransformElementDetails::GetTransformTypeFromValueType(
+	ERigControlValueType InValueType)
+{
+	ERigTransformElementDetailsTransform::Type TransformType = ERigTransformElementDetailsTransform::Current;
+	switch(InValueType)
+	{
+		case ERigControlValueType::Initial:
+		{
+			TransformType = ERigTransformElementDetailsTransform::Initial;
+			break;
+		}
+		case ERigControlValueType::Minimum:
+		{
+			TransformType = ERigTransformElementDetailsTransform::Minimum;
+			break;
+		}
+		case ERigControlValueType::Maximum:
+		{
+			TransformType = ERigTransformElementDetailsTransform::Maximum;
+			break;
+		}
+		default:
+		{
+			break;
+		}
+	}
+	return TransformType;
 }
 
 void FRigBoneElementDetails::CustomizeDetails(IDetailLayoutBuilder& DetailBuilder)
@@ -2119,6 +2304,7 @@ void FRigControlElementDetails::CustomizeDetails(IDetailLayoutBuilder& DetailBui
 	CustomizeValue(DetailBuilder);
 	CustomizeTransform(DetailBuilder);
 	CustomizeShape(DetailBuilder);
+	CustomizeAnimationChannels(DetailBuilder);
 }
 
 void FRigControlElementDetails::CustomizeValue(IDetailLayoutBuilder& DetailBuilder)
@@ -2130,18 +2316,24 @@ void FRigControlElementDetails::CustomizeValue(IDetailLayoutBuilder& DetailBuild
 
 	// only show this section if all controls are the same type
 	const ERigControlType ControlType = ControlElements[0].Settings.ControlType;
+	bool bAllAnimationChannels = true;
 	for(const FRigControlElement& ControlElement : ControlElements)
 	{
 		if(ControlElement.Settings.ControlType != ControlType)
 		{
 			return;
 		}
+		if(ControlElement.Settings.AnimationType != ERigControlAnimationType::AnimationChannel)
+		{
+			bAllAnimationChannels = false;
+		}
 	}
 
 	// transforms don't show their value here - instead they are shown in the transform section
-	if(ControlType == ERigControlType::EulerTransform ||
+	if((ControlType == ERigControlType::EulerTransform ||
 		ControlType == ERigControlType::Transform ||
-		ControlType == ERigControlType::TransformNoScale)
+		ControlType == ERigControlType::TransformNoScale) &&
+		!bAllAnimationChannels)
 	{
 		return;
 	}
@@ -2203,8 +2395,8 @@ void FRigControlElementDetails::CustomizeValue(IDetailLayoutBuilder& DetailBuild
 
 	AddChoiceWidgetRow(ValueCategory, FText::FromString(TEXT("ValueType")), ValueTypeChoiceWidget.ToSharedRef());
 
-
-	TSharedRef<UE::Math::TVector<float>> IsComponentRelative = MakeShareable(new UE::Math::TVector<float>(1.f, 1.f, 1.f));
+	TArray<FRigElementKey> Keys = GetElementKeys();
+	Keys = GetHierarchyBeingDebugged()->SortKeys(Keys);
 
 	for(int32 Index=0; Index < ValueTypes.Num(); Index++)
 	{
@@ -2220,18 +2412,17 @@ void FRigControlElementDetails::CustomizeValue(IDetailLayoutBuilder& DetailBuild
 		{
 			case ERigControlType::Bool:
 			{
-				CreateBoolValueWidgetRow(ValueCategory, Labels[Index], Tooltips[Index], ValueType, VisibilityAttribute);
+				CreateBoolValueWidgetRow(Keys, ValueCategory, Labels[Index], Tooltips[Index], ValueType, VisibilityAttribute);
 				break;
 			}
 			case ERigControlType::Float:
 			{
-				CreateFloatValueWidgetRow(ValueCategory, Labels[Index], Tooltips[Index], ValueType, VisibilityAttribute);
+				CreateFloatValueWidgetRow(Keys, ValueCategory, Labels[Index], Tooltips[Index], ValueType, VisibilityAttribute);
 				break;
 			}
 			case ERigControlType::Integer:
 			{
 				bool bIsEnum = false;
-				const TArray<FRigElementKey> Keys = GetElementKeys();
 				for(const FRigElementKey& Key : Keys)
 				{
 					if(const FRigControlElement* ControlElement = HierarchyBeingCustomized->Find<FRigControlElement>(Key))
@@ -2246,73 +2437,37 @@ void FRigControlElementDetails::CustomizeValue(IDetailLayoutBuilder& DetailBuild
 
 				if(bIsEnum)
 				{
-					CreateEnumValueWidgetRow(ValueCategory, Labels[Index], Tooltips[Index], ValueType, VisibilityAttribute);
+					CreateEnumValueWidgetRow(Keys, ValueCategory, Labels[Index], Tooltips[Index], ValueType, VisibilityAttribute);
 				}
 				else
 				{
-					CreateIntegerValueWidgetRow(ValueCategory, Labels[Index], Tooltips[Index], ValueType, VisibilityAttribute);
+					CreateIntegerValueWidgetRow(Keys, ValueCategory, Labels[Index], Tooltips[Index], ValueType, VisibilityAttribute);
 				}
 				break;
 			}
 			case ERigControlType::Vector2D:
 			{
-				CreateVector2DValueWidgetRow(ValueCategory, Labels[Index], Tooltips[Index], ValueType, VisibilityAttribute);
+				CreateVector2DValueWidgetRow(Keys, ValueCategory, Labels[Index], Tooltips[Index], ValueType, VisibilityAttribute);
 				break;
 			}
 			case ERigControlType::Position:
 			case ERigControlType::Rotator:
 			case ERigControlType::Scale:
 			{
-				ERigTransformElementDetailsTransform::Type CurrentTransformType = ERigTransformElementDetailsTransform::Current;
-				switch(ValueType)
-				{
-					case ERigControlValueType::Initial:
-					{
-						CurrentTransformType = ERigTransformElementDetailsTransform::Initial;
-						break;
-					}
-					case ERigControlValueType::Minimum:
-					{
-						CurrentTransformType = ERigTransformElementDetailsTransform::Minimum;
-						break;
-					}
-					case ERigControlValueType::Maximum:
-					{
-						CurrentTransformType = ERigTransformElementDetailsTransform::Maximum;
-						break;
-					}
-				}
-
 				SAdvancedTransformInputBox<FEulerTransform>::FArguments TransformWidgetArgs = SAdvancedTransformInputBox<FEulerTransform>::FArguments()
 				.DisplayToggle(false)
 				.DisplayRelativeWorld(true)
-				.Font(IDetailLayoutBuilder::GetDetailFont())
-				.AllowEditRotationRepresentation(false)
-				.OnGetIsComponentRelative_Lambda(
-					[IsComponentRelative](ESlateTransformComponent::Type InComponent)
-					{
-						return IsComponentRelative->operator[]((int32)InComponent) > 0.f;
-					})
-				.OnIsComponentRelativeChanged_Lambda(
-					[IsComponentRelative](ESlateTransformComponent::Type InComponent, bool bIsRelative)
-					{
-						IsComponentRelative->operator[]((int32)InComponent) = bIsRelative ? 1.f : 0.f;
-					});
-
-				TransformWidgetArgs.Visibility(VisibilityAttribute);
-
-				TransformWidgetArgs.ConstructLocation(ControlType == ERigControlType::Position);
-				TransformWidgetArgs.ConstructRotation(ControlType == ERigControlType::Rotator);
-				TransformWidgetArgs.ConstructScale(ControlType == ERigControlType::Scale);
+				.Visibility(VisibilityAttribute);
 
-				CreateEulerTransformValueWidgetRow(
+				CreateTransformComponentValueWidgetRow(
+					ControlType,
 					GetHierarchyBeingDebugged(),
 					GetElementKeys(),
 					TransformWidgetArgs,
 					ValueCategory,
 					Labels[Index],
 					Tooltips[Index],
-					CurrentTransformType,
+					GetTransformTypeFromValueType(ValueType),
 					ValueType);
 				break;
 			}
@@ -2331,18 +2486,27 @@ void FRigControlElementDetails::CustomizeControl(IDetailLayoutBuilder& DetailBui
 
 	IDetailCategoryBuilder& ControlCategory = DetailBuilder.EditCategory(TEXT("Control"), LOCTEXT("Control", "Control"));
 
+	const bool bAllAnimationChannels = !IsAnyControlNotOfAnimationType(ERigControlAnimationType::AnimationChannel);
+	static const FText DisplayNameText = LOCTEXT("DisplayName", "Display Name");
+	static const FText ChannelNameText = LOCTEXT("ChannelName", "Channel Name");
+	const FText DisplayNameLabelText = bAllAnimationChannels ? ChannelNameText : DisplayNameText;
+
 	const TSharedPtr<IPropertyHandle> DisplayNameHandle = SettingsHandle->GetChildHandle(TEXT("DisplayName"));
-	ControlCategory.AddCustomRow(LOCTEXT("DisplayName", "Display Name"))
+	ControlCategory.AddCustomRow(DisplayNameLabelText)
 	.NameContent()
 	[
-		DisplayNameHandle->CreatePropertyNameWidget()
+		DisplayNameHandle->CreatePropertyNameWidget(DisplayNameLabelText)
 	]
 	.ValueContent()
 	[
-		SNew(SEditableTextBox)
+		SNew(SInlineEditableTextBlock)
 		.Font(IDetailLayoutBuilder::GetDetailFont())
 		.Text(this, &FRigControlElementDetails::GetDisplayName)
 		.OnTextCommitted(this, &FRigControlElementDetails::SetDisplayName)
+		.OnVerifyTextChanged_Lambda([this](const FText& InText, FText& OutErrorMessage) -> bool
+		{
+			return OnVerifyDisplayNameChanged(InText, OutErrorMessage, GetElementKey());
+		})
 		.IsEnabled(ObjectsBeingCustomized.Num() == 1)
 	];
 
@@ -2430,137 +2594,7 @@ void FRigControlElementDetails::CustomizeControl(IDetailLayoutBuilder& DetailBui
 		[this, PropertyUtilities]()
 		{
 			TArray<FRigControlElement> ControlElementsInView = GetElementsInDetailsView<FRigControlElement>();
-			TArray<FRigControlElement*> ControlElementsInHierarchy = GetElementsInHierarchy<FRigControlElement>();
-			check(ControlElementsInView.Num() == ControlElementsInHierarchy.Num());
-
-			if (this->HierarchyBeingCustomized && ControlElementsInHierarchy.Num() > 0)
-			{
-				this->HierarchyBeingCustomized->Modify();
-				
-				for(int32 ControlIndex = 0; ControlIndex< ControlElementsInView.Num(); ControlIndex++)
-				{
-					const FRigControlElement& ViewElement = ControlElementsInView[ControlIndex];
-					FRigControlElement* ControlElement = ControlElementsInHierarchy[ControlIndex];
-					
-					FRigControlValue ValueToSet;
-
-					ControlElement->Settings.ControlType = ViewElement.Settings.ControlType;
-					ControlElement->Settings.LimitEnabled.Reset();
-					ControlElement->Settings.bGroupWithParentControl = false;
-
-					switch (ControlElement->Settings.ControlType)
-					{
-						case ERigControlType::Bool:
-						{
-							ControlElement->Settings.AnimationType = ERigControlAnimationType::AnimationChannel;
-							ValueToSet = FRigControlValue::Make<bool>(false);
-							ControlElement->Settings.bGroupWithParentControl = ControlElement->Settings.IsAnimatable();
-							break;
-						}
-						case ERigControlType::Float:
-						{
-							ValueToSet = FRigControlValue::Make<float>(0.f);
-							ControlElement->Settings.SetupLimitArrayForType(true);
-							ControlElement->Settings.MinimumValue = FRigControlValue::Make<float>(0.f);
-							ControlElement->Settings.MaximumValue = FRigControlValue::Make<float>(100.f);
-							ControlElement->Settings.bGroupWithParentControl = ControlElement->Settings.IsAnimatable();
-							break;
-						}
-						case ERigControlType::Integer:
-						{
-							ValueToSet = FRigControlValue::Make<int32>(0);
-							ControlElement->Settings.SetupLimitArrayForType(true);
-							ControlElement->Settings.MinimumValue = FRigControlValue::Make<int32>(0);
-							ControlElement->Settings.MaximumValue = FRigControlValue::Make<int32>(100);
-							ControlElement->Settings.bGroupWithParentControl = ControlElement->Settings.IsAnimatable();
-							break;
-						}
-						case ERigControlType::Vector2D:
-						{
-							ValueToSet = FRigControlValue::Make<FVector2D>(FVector2D::ZeroVector);
-							ControlElement->Settings.SetupLimitArrayForType(true);
-							ControlElement->Settings.MinimumValue = FRigControlValue::Make<FVector2D>(FVector2D::ZeroVector);
-							ControlElement->Settings.MaximumValue = FRigControlValue::Make<FVector2D>(FVector2D(100.f, 100.f));
-							ControlElement->Settings.bGroupWithParentControl = ControlElement->Settings.IsAnimatable();
-							break;
-						}
-						case ERigControlType::Position:
-						{
-							ValueToSet = FRigControlValue::Make<FVector>(FVector::ZeroVector);
-							ControlElement->Settings.SetupLimitArrayForType(false);
-							ControlElement->Settings.MinimumValue = FRigControlValue::Make<FVector>(-FVector::OneVector);
-							ControlElement->Settings.MaximumValue = FRigControlValue::Make<FVector>(FVector::OneVector);
-							break;
-						}
-						case ERigControlType::Scale:
-						{
-							ValueToSet = FRigControlValue::Make<FVector>(FVector::OneVector);
-							ControlElement->Settings.SetupLimitArrayForType(false);
-							ControlElement->Settings.MinimumValue = FRigControlValue::Make<FVector>(FVector::ZeroVector);
-							ControlElement->Settings.MaximumValue = FRigControlValue::Make<FVector>(FVector::OneVector);
-							break;
-						}
-						case ERigControlType::Rotator:
-						{
-							ValueToSet = FRigControlValue::Make<FRotator>(FRotator::ZeroRotator);
-							ControlElement->Settings.SetupLimitArrayForType(false, false);
-							ControlElement->Settings.MinimumValue = FRigControlValue::Make<FRotator>(FRotator::ZeroRotator);
-							ControlElement->Settings.MaximumValue = FRigControlValue::Make<FRotator>(FRotator(180.f, 180.f, 180.f));
-							break;
-						}
-						case ERigControlType::Transform:
-						{
-							ValueToSet = FRigControlValue::Make<FTransform>(FTransform::Identity);
-							ControlElement->Settings.SetupLimitArrayForType(false, false, false);
-							ControlElement->Settings.MinimumValue = ValueToSet;
-							ControlElement->Settings.MaximumValue = ValueToSet;
-							break;
-						}
-						case ERigControlType::TransformNoScale:
-						{
-							FTransformNoScale Identity = FTransform::Identity;
-							ValueToSet = FRigControlValue::Make<FTransformNoScale>(Identity);
-							ControlElement->Settings.SetupLimitArrayForType(false, false, false);
-							ControlElement->Settings.MinimumValue = ValueToSet;
-							ControlElement->Settings.MaximumValue = ValueToSet;
-							break;
-						}
-						case ERigControlType::EulerTransform:
-						{
-							FEulerTransform Identity = FEulerTransform::Identity;
-							ValueToSet = FRigControlValue::Make<FEulerTransform>(Identity);
-							ControlElement->Settings.SetupLimitArrayForType(false, false, false);
-							ControlElement->Settings.MinimumValue = ValueToSet;
-							ControlElement->Settings.MaximumValue = ValueToSet;
-							break;
-						}
-						default:
-						{
-							ensure(false);
-							break;
-						}
-					}
-
-					this->HierarchyBeingCustomized->SetControlSettings(ControlElement, ControlElement->Settings, true, true, true);
-					this->HierarchyBeingCustomized->SetControlValue(ControlElement, ValueToSet, ERigControlValueType::Initial, true, false, true);
-					this->HierarchyBeingCustomized->SetControlValue(ControlElement, ValueToSet, ERigControlValueType::Current, true, false, true);
-
-					ObjectsBeingCustomized[ControlIndex]->SetContent<FRigControlElement>(*ControlElement);
-
-					if (this->HierarchyBeingCustomized != this->BlueprintBeingCustomized->Hierarchy)
-					{
-						if(FRigControlElement* OtherControlElement = this->BlueprintBeingCustomized->Hierarchy->Find<FRigControlElement>(ControlElement->GetKey()))
-						{
-							OtherControlElement->Settings = ControlElement->Settings;
-							this->BlueprintBeingCustomized->Hierarchy->SetControlSettings(OtherControlElement, OtherControlElement->Settings, true, true, true);
-							this->BlueprintBeingCustomized->Hierarchy->SetControlValue(OtherControlElement, ValueToSet, ERigControlValueType::Initial, true);
-							this->BlueprintBeingCustomized->Hierarchy->SetControlValue(OtherControlElement, ValueToSet, ERigControlValueType::Current, true);
-						}
-					}
-				}
-				
-				PropertyUtilities->ForceRefresh();
-			}
+			HandleControlTypeChanged(ControlElementsInView[0].Settings.ControlType, TArray<FRigElementKey>(), PropertyUtilities);
 		}
 	));
 
@@ -2714,36 +2748,509 @@ void FRigControlElementDetails::CustomizeControl(IDetailLayoutBuilder& DetailBui
 	}
 }
 
-void FRigControlElementDetails::CustomizeShape(IDetailLayoutBuilder& DetailBuilder)
+void FRigControlElementDetails::CustomizeAnimationChannels(IDetailLayoutBuilder& DetailBuilder)
 {
-	for(TWeakObjectPtr<UDetailsViewWrapperObject> ObjectBeingCustomized : ObjectsBeingCustomized)
+	// only show this for non-animation channels
+	if(!IsAnyControlNotOfAnimationType(ERigControlAnimationType::AnimationChannel))
 	{
-		if(ObjectBeingCustomized.IsValid())
-		{
-			if(ObjectBeingCustomized->IsChildOf<FRigControlElement>())
-			{
-				const FRigControlElement ControlElement = ObjectBeingCustomized->GetContent<FRigControlElement>();
-				if(!ControlElement.Settings.SupportsShape())
-				{
-					return;
-				}
-			}
-		}
+		return;
 	}
 
-	TSharedPtr<IPropertyHandle> ShapeHandle = DetailBuilder.GetProperty(TEXT("Shape"));
-	TSharedPtr<IPropertyHandle> InitialHandle = ShapeHandle->GetChildHandle(TEXT("Initial"));
-	TSharedPtr<IPropertyHandle> LocalHandle = InitialHandle->GetChildHandle(TEXT("Local"));
-	ShapeTransformHandle = LocalHandle->GetChildHandle(TEXT("Transform"));
-	
-	ShapeNameList.Reset();
-	if (BlueprintBeingCustomized)
+	// only show this if only one control is selected
+	if((ControlElements.Num() != 1) || HierarchyBeingCustomized == nullptr)
 	{
-		const bool bUseNameSpace = BlueprintBeingCustomized->ShapeLibraries.Num() > 1;
-		for(TSoftObjectPtr<UControlRigShapeLibrary>& ShapeLibrary : BlueprintBeingCustomized->ShapeLibraries)
-		{
-			if (!ShapeLibrary.IsValid())
-			{
+		return;
+	}
+
+	const FRigControlElement* ControlElement = HierarchyBeingCustomized->Find<FRigControlElement>(ControlElements[0].GetKey());
+	if(ControlElement == nullptr)
+	{
+		return;
+	}
+
+	IDetailCategoryBuilder& Category = DetailBuilder.EditCategory(TEXT("AnimationChannels"), LOCTEXT("AnimationChannels", "Animation Channels"));
+	
+	const TSharedRef<IPropertyUtilities> PropertyUtilities = DetailBuilder.GetPropertyUtilities();
+	
+	const TSharedRef<SHorizontalBox> HeaderContentWidget = SNew(SHorizontalBox);
+	HeaderContentWidget->AddSlot()
+	.HAlign(HAlign_Right)
+	[
+		SNew(SButton)
+		.ButtonStyle(FAppStyle::Get(), "SimpleButton")
+		.ContentPadding(FMargin(1, 0))
+		.OnClicked(this, &FRigControlElementDetails::OnAddAnimationChannelClicked)
+		.HAlign(HAlign_Right)
+		.ToolTipText(LOCTEXT("AddAnimationChannelToolTip", "Add a new animation channel"))
+		.VAlign(VAlign_Center)
+		[
+			SNew(SImage)
+			.Image(FAppStyle::Get().GetBrush("Icons.PlusCircle"))
+			.ColorAndOpacity(FSlateColor::UseForeground())
+		]
+	];
+	Category.HeaderContent(HeaderContentWidget);
+
+	bool bHasAnimationChannels = false;
+	const FRigBaseElementChildrenArray ChildElements = HierarchyBeingCustomized->GetChildren(ControlElement);
+	for(const FRigBaseElement* ChildElement : ChildElements)
+	{
+		if(const FRigControlElement* ChildControlElement = Cast<FRigControlElement>(ChildElement))
+		{
+			if(ChildControlElement->IsAnimationChannel())
+			{
+				bHasAnimationChannels = true;
+				const FRigElementKey ChildElementKey = ChildElement->GetKey();
+				
+				const TPair<const FSlateBrush*, FSlateColor> BrushAndColor = SRigHierarchyItem::GetBrushForElementType(HierarchyBeingCustomized, ChildElementKey);
+
+				static TArray<TSharedPtr<ERigControlType>> ControlValueTypes;
+				if(ControlValueTypes.IsEmpty())
+				{
+					const UEnum* ValueTypeEnum = StaticEnum<ERigControlType>();
+					for(int64 EnumValue = 0; EnumValue < ValueTypeEnum->GetMaxEnumValue(); EnumValue++)
+					{
+						if(ValueTypeEnum->HasMetaData(TEXT("Hidden"), (int32)EnumValue))
+						{
+							continue;
+						}
+						ControlValueTypes.Add(MakeShareable(new ERigControlType((ERigControlType)EnumValue)));
+					}
+				}
+
+				TSharedPtr<SButton> SelectAnimationChannelButton;
+				TSharedPtr<SWidget> NameContent;
+
+				SAssignNew(NameContent, SHorizontalBox)
+
+				+ SHorizontalBox::Slot()
+				.MaxWidth(32)
+				.HAlign(HAlign_Left)
+				.VAlign(VAlign_Center)
+				.Padding(FMargin(0.f, 0.f, 3.f, 0.f))
+				[
+					SNew(SComboButton)
+					.ContentPadding(0)
+					.HasDownArrow(false)
+					.ButtonContent()
+					[
+						SNew(SImage)
+						.Image(BrushAndColor.Key)
+						.ColorAndOpacity(BrushAndColor.Value)
+					]
+					.MenuContent()
+					[
+						SNew(SListView<TSharedPtr<ERigControlType>>)
+						.ListItemsSource( &ControlValueTypes )
+						.OnGenerateRow(this, &FRigControlElementDetails::HandleGenerateAnimationChannelTypeRow, ChildElementKey)
+						.OnSelectionChanged(this, &FRigControlElementDetails::HandleControlTypeChanged, ChildElementKey, PropertyUtilities)
+					]
+				]
+
+				+ SHorizontalBox::Slot()
+				.AutoWidth()
+				.HAlign(HAlign_Left)
+				.VAlign(VAlign_Center)
+				.Padding(FMargin(0.f, 0.f, 8.f, 0.f))
+				[
+					SNew(SInlineEditableTextBlock)
+					.Font(IDetailLayoutBuilder::GetDetailFont())
+					.Text(FText::FromName(ChildControlElement->GetDisplayName()))
+					.OnTextCommitted_Lambda([this, ChildElementKey](const FText& InNewText, ETextCommit::Type InCommitType)
+					{
+						SetDisplayNameForElement(InNewText, InCommitType, ChildElementKey);
+					})
+					.OnVerifyTextChanged_Lambda([this, ChildElementKey](const FText& InText, FText& OutErrorMessage) -> bool
+					{
+						return OnVerifyDisplayNameChanged(InText, OutErrorMessage, ChildElementKey);
+					})
+				]
+
+				+SHorizontalBox::Slot()
+				.AutoWidth()
+				.HAlign(HAlign_Right)
+				.VAlign(VAlign_Center)
+				.Padding(FMargin(0.f, 0.f, 0.f, 0.f))
+				[
+					SAssignNew(SelectAnimationChannelButton, SButton)
+					.ButtonStyle( FAppStyle::Get(), "NoBorder" )
+					.ButtonColorAndOpacity_Lambda([SelectAnimationChannelButton]() { return FRigElementKeyDetails::OnGetWidgetBackground(SelectAnimationChannelButton); })
+					.OnClicked_Lambda([this, ChildElementKey]() -> FReply
+					{
+						return OnSelectElementClicked(ChildElementKey);
+					})
+					.ContentPadding(0)
+					.ToolTipText(NSLOCTEXT("ControlRigElementDetails", "SelectParentInHierarchyToolTip", "Select Parent in hierarchy"))
+					[
+						SNew(SImage)
+						.ColorAndOpacity_Lambda( [SelectAnimationChannelButton]() { return FRigElementKeyDetails::OnGetWidgetForeground(SelectAnimationChannelButton); })
+						.Image(FAppStyle::GetBrush("Icons.Search"))
+					]
+				];
+
+				const FText Label = FText::FromString(FString::Printf(TEXT("Channel%s"), *ChildControlElement->GetDisplayName().ToString()));
+				const TArray<FRigElementKey> ChildElementKeys = {ChildElementKey};
+				TAttribute<EVisibility> Visibility = EVisibility::Visible;
+
+				FDetailWidgetRow* WidgetRow = nullptr; 
+				switch(ChildControlElement->Settings.ControlType)
+				{
+					case ERigControlType::Bool:
+					{
+						WidgetRow = &CreateBoolValueWidgetRow(ChildElementKeys, Category, Label, FText(), ERigControlValueType::Current, Visibility, NameContent);
+						break;
+					}
+					case ERigControlType::Float:
+					{
+						WidgetRow = &CreateFloatValueWidgetRow(ChildElementKeys, Category, Label, FText(), ERigControlValueType::Current, Visibility, NameContent);
+						break;
+					}
+					case ERigControlType::Integer:
+					{
+						if(ChildControlElement->Settings.ControlEnum)
+						{
+							WidgetRow = &CreateEnumValueWidgetRow(ChildElementKeys, Category, Label, FText(), ERigControlValueType::Current, Visibility, NameContent);
+						}
+						else
+						{
+							WidgetRow = &CreateIntegerValueWidgetRow(ChildElementKeys, Category, Label, FText(), ERigControlValueType::Current, Visibility, NameContent);
+						}
+						break;
+					}
+					case ERigControlType::Vector2D:
+					{
+						WidgetRow = &CreateVector2DValueWidgetRow(ChildElementKeys, Category, Label, FText(), ERigControlValueType::Current, Visibility, NameContent);
+						break;
+					}
+					case ERigControlType::Position:
+					case ERigControlType::Rotator:
+					case ERigControlType::Scale:
+					{
+						SAdvancedTransformInputBox<FEulerTransform>::FArguments TransformWidgetArgs =
+							SAdvancedTransformInputBox<FEulerTransform>::FArguments()
+						.DisplayToggle(false)
+						.DisplayRelativeWorld(false)
+						.Visibility(EVisibility::Visible);
+
+						WidgetRow = &CreateTransformComponentValueWidgetRow(
+							ChildControlElement->Settings.ControlType,
+							GetHierarchyBeingDebugged(),
+							ChildElementKeys,
+							TransformWidgetArgs,
+							Category,
+							Label,
+							FText(),
+							GetTransformTypeFromValueType(ERigControlValueType::Current),
+							ERigControlValueType::Current,
+							NameContent);
+						break;
+					}
+					case ERigControlType::Transform:
+					case ERigControlType::EulerTransform:
+					{
+						SAdvancedTransformInputBox<FEulerTransform>::FArguments TransformWidgetArgs =
+							SAdvancedTransformInputBox<FEulerTransform>::FArguments()
+						.DisplayToggle(false)
+						.DisplayRelativeWorld(false)
+						.Font(IDetailLayoutBuilder::GetDetailFont())
+						.Visibility(EVisibility::Visible);
+
+						WidgetRow = &CreateEulerTransformValueWidgetRow(
+							GetHierarchyBeingDebugged(),
+							ChildElementKeys,
+							TransformWidgetArgs,
+							Category,
+							Label,
+							FText(),
+							ERigTransformElementDetailsTransform::Current,
+							ERigControlValueType::Current,
+							NameContent);
+						break;
+					}
+					default:
+					{
+						WidgetRow = &Category.AddCustomRow(Label)
+						.NameContent()
+						[
+							NameContent.ToSharedRef()
+						];
+						break;
+					}
+				}
+
+				if(WidgetRow)
+				{
+					WidgetRow->AddCustomContextMenuAction(FUIAction(
+						FExecuteAction::CreateLambda([this, ChildElementKeys]()
+						{
+							if(BlueprintBeingCustomized)
+							{
+								if(URigHierarchyController* Controller = BlueprintBeingCustomized->GetHierarchyController())
+								{
+									FScopedTransaction(LOCTEXT("DeleteAnimationChannels", "Delete Animation Channels"));
+									BlueprintBeingCustomized->Hierarchy->Modify();
+									
+									for(const FRigElementKey& KeyToRemove : ChildElementKeys)
+									{
+										Controller->RemoveElement(KeyToRemove, true, true);
+									}
+								}
+							}
+						})),
+						LOCTEXT("DeleteAnimationChannel", "Delete"),
+						LOCTEXT("DeleteAnimationChannelTooltip", "Deletes this animation channel"),
+					FSlateIcon());
+				}
+			}
+		}
+	}
+
+	Category.InitiallyCollapsed(!bHasAnimationChannels);
+	if(!bHasAnimationChannels)
+	{
+		Category.AddCustomRow(FText()).WholeRowContent()
+		[
+			SNew(SHorizontalBox)
+	
+			+ SHorizontalBox::Slot()
+			.AutoWidth()
+			.HAlign(HAlign_Left)
+			.VAlign(VAlign_Center)
+			.Padding(FMargin(0.f, 0.f, 0.f, 0.f))
+			[
+				SNew(STextBlock)
+				.Font(IDetailLayoutBuilder::GetDetailFont())
+				.Text(LOCTEXT("NoAnimationChannels", "No animation channels"))
+			]
+		];
+	}
+}
+
+FReply FRigControlElementDetails::OnAddAnimationChannelClicked()
+{
+	static const FString ChannelName = TEXT("Channel");
+	FRigControlSettings Settings;
+	Settings.AnimationType = ERigControlAnimationType::AnimationChannel;
+	Settings.ControlType = ERigControlType::Float;
+	Settings.MinimumValue = FRigControlValue::Make<float>(0.f);
+	Settings.MaximumValue = FRigControlValue::Make<float>(1.f);
+	Settings.DisplayName = HierarchyBeingCustomized->GetSafeNewDisplayName(ControlElements[0].GetKey(), ChannelName);
+	HierarchyBeingCustomized->GetController(true)->AddAnimationChannel(*ChannelName, ControlElements[0].GetKey(), Settings, true, true);
+	return FReply::Handled();
+}
+
+TSharedRef<ITableRow> FRigControlElementDetails::HandleGenerateAnimationChannelTypeRow(TSharedPtr<ERigControlType> ControlType, const TSharedRef<STableViewBase>& OwnerTable, FRigElementKey ControlKey)
+{
+	TPair<const FSlateBrush*, FSlateColor> BrushAndColor = SRigHierarchyItem::GetBrushForElementType(HierarchyBeingCustomized, ControlKey);
+	BrushAndColor.Value = SRigHierarchyItem::GetColorForControlType(*ControlType.Get(), nullptr);
+
+	return SNew(STableRow<TSharedPtr<ERigControlType>>, OwnerTable)
+	.Content()
+	[
+		SNew(SHorizontalBox)
+	
+		+ SHorizontalBox::Slot()
+		.MaxWidth(18)
+		.FillWidth(1.0)
+		.HAlign(HAlign_Left)
+		.VAlign(VAlign_Center)
+		.Padding(FMargin(0.f, 0.f, 3.f, 0.f))
+		[
+			SNew(SImage)
+			.Image(BrushAndColor.Key)
+			.ColorAndOpacity(BrushAndColor.Value)
+		]
+
+		+ SHorizontalBox::Slot()
+		.AutoWidth()
+		.HAlign(HAlign_Left)
+		.VAlign(VAlign_Center)
+		.Padding(FMargin(0.f, 0.f, 0.f, 0.f))
+		[
+			SNew(STextBlock)
+			.Font(IDetailLayoutBuilder::GetDetailFont())
+			.Text(StaticEnum<ERigControlType>()->GetDisplayNameTextByValue((int64)*ControlType.Get()))
+		]
+	];
+}
+
+void FRigControlElementDetails::HandleControlTypeChanged(TSharedPtr<ERigControlType> ControlType, ESelectInfo::Type SelectInfo, FRigElementKey ControlKey, const TSharedRef<IPropertyUtilities> PropertyUtilities)
+{
+	HandleControlTypeChanged(*ControlType.Get(), {ControlKey}, PropertyUtilities);
+}
+
+void FRigControlElementDetails::HandleControlTypeChanged(ERigControlType ControlType, TArray<FRigElementKey> ControlKeys, const TSharedRef<IPropertyUtilities> PropertyUtilities)
+{
+	TArray<FRigControlElement*> ControlElementsInHierarchy = GetElementsInHierarchy<FRigControlElement>(ControlKeys);
+	if (this->HierarchyBeingCustomized && ControlElementsInHierarchy.Num() > 0)
+	{
+		this->HierarchyBeingCustomized->Modify();
+		
+		for(int32 ControlIndex = 0; ControlIndex< ControlElementsInHierarchy.Num(); ControlIndex++)
+		{
+			FRigControlElement* ControlElement = ControlElementsInHierarchy[ControlIndex];
+			
+			FRigControlValue ValueToSet;
+
+			ControlElement->Settings.ControlType = ControlType;
+			ControlElement->Settings.LimitEnabled.Reset();
+			ControlElement->Settings.bGroupWithParentControl = false;
+
+			switch (ControlElement->Settings.ControlType)
+			{
+				case ERigControlType::Bool:
+				{
+					ControlElement->Settings.AnimationType = ERigControlAnimationType::AnimationChannel;
+					ValueToSet = FRigControlValue::Make<bool>(false);
+					ControlElement->Settings.bGroupWithParentControl = ControlElement->Settings.IsAnimatable();
+					break;
+				}
+				case ERigControlType::Float:
+				{
+					ValueToSet = FRigControlValue::Make<float>(0.f);
+					ControlElement->Settings.SetupLimitArrayForType(true);
+					ControlElement->Settings.MinimumValue = FRigControlValue::Make<float>(0.f);
+					ControlElement->Settings.MaximumValue = FRigControlValue::Make<float>(100.f);
+					ControlElement->Settings.bGroupWithParentControl = ControlElement->Settings.IsAnimatable();
+					break;
+				}
+				case ERigControlType::Integer:
+				{
+					ValueToSet = FRigControlValue::Make<int32>(0);
+					ControlElement->Settings.SetupLimitArrayForType(true);
+					ControlElement->Settings.MinimumValue = FRigControlValue::Make<int32>(0);
+					ControlElement->Settings.MaximumValue = FRigControlValue::Make<int32>(100);
+					ControlElement->Settings.bGroupWithParentControl = ControlElement->Settings.IsAnimatable();
+					break;
+				}
+				case ERigControlType::Vector2D:
+				{
+					ValueToSet = FRigControlValue::Make<FVector2D>(FVector2D::ZeroVector);
+					ControlElement->Settings.SetupLimitArrayForType(true);
+					ControlElement->Settings.MinimumValue = FRigControlValue::Make<FVector2D>(FVector2D::ZeroVector);
+					ControlElement->Settings.MaximumValue = FRigControlValue::Make<FVector2D>(FVector2D(100.f, 100.f));
+					ControlElement->Settings.bGroupWithParentControl = ControlElement->Settings.IsAnimatable();
+					break;
+				}
+				case ERigControlType::Position:
+				{
+					ValueToSet = FRigControlValue::Make<FVector>(FVector::ZeroVector);
+					ControlElement->Settings.SetupLimitArrayForType(false);
+					ControlElement->Settings.MinimumValue = FRigControlValue::Make<FVector>(-FVector::OneVector);
+					ControlElement->Settings.MaximumValue = FRigControlValue::Make<FVector>(FVector::OneVector);
+					break;
+				}
+				case ERigControlType::Scale:
+				{
+					ValueToSet = FRigControlValue::Make<FVector>(FVector::OneVector);
+					ControlElement->Settings.SetupLimitArrayForType(false);
+					ControlElement->Settings.MinimumValue = FRigControlValue::Make<FVector>(FVector::ZeroVector);
+					ControlElement->Settings.MaximumValue = FRigControlValue::Make<FVector>(FVector::OneVector);
+					break;
+				}
+				case ERigControlType::Rotator:
+				{
+					ValueToSet = FRigControlValue::Make<FRotator>(FRotator::ZeroRotator);
+					ControlElement->Settings.SetupLimitArrayForType(false, false);
+					ControlElement->Settings.MinimumValue = FRigControlValue::Make<FRotator>(FRotator::ZeroRotator);
+					ControlElement->Settings.MaximumValue = FRigControlValue::Make<FRotator>(FRotator(180.f, 180.f, 180.f));
+					break;
+				}
+				case ERigControlType::Transform:
+				{
+					ValueToSet = FRigControlValue::Make<FTransform>(FTransform::Identity);
+					ControlElement->Settings.SetupLimitArrayForType(false, false, false);
+					ControlElement->Settings.MinimumValue = ValueToSet;
+					ControlElement->Settings.MaximumValue = ValueToSet;
+					break;
+				}
+				case ERigControlType::TransformNoScale:
+				{
+					FTransformNoScale Identity = FTransform::Identity;
+					ValueToSet = FRigControlValue::Make<FTransformNoScale>(Identity);
+					ControlElement->Settings.SetupLimitArrayForType(false, false, false);
+					ControlElement->Settings.MinimumValue = ValueToSet;
+					ControlElement->Settings.MaximumValue = ValueToSet;
+					break;
+				}
+				case ERigControlType::EulerTransform:
+				{
+					FEulerTransform Identity = FEulerTransform::Identity;
+					ValueToSet = FRigControlValue::Make<FEulerTransform>(Identity);
+					ControlElement->Settings.SetupLimitArrayForType(false, false, false);
+					ControlElement->Settings.MinimumValue = ValueToSet;
+					ControlElement->Settings.MaximumValue = ValueToSet;
+					break;
+				}
+				default:
+				{
+					ensure(false);
+					break;
+				}
+			}
+
+			this->HierarchyBeingCustomized->SetControlSettings(ControlElement, ControlElement->Settings, true, true, true);
+			this->HierarchyBeingCustomized->SetControlValue(ControlElement, ValueToSet, ERigControlValueType::Initial, true, false, true);
+			this->HierarchyBeingCustomized->SetControlValue(ControlElement, ValueToSet, ERigControlValueType::Current, true, false, true);
+
+			if(ObjectsBeingCustomized.IsValidIndex(ControlIndex))
+			{
+				if(ObjectsBeingCustomized[ControlIndex]->GetContent<FRigBaseElement>().GetKey() == ControlElement->GetKey())
+				{
+					ObjectsBeingCustomized[ControlIndex]->SetContent<FRigControlElement>(*ControlElement);
+				}
+			}
+
+			if (this->HierarchyBeingCustomized != this->BlueprintBeingCustomized->Hierarchy)
+			{
+				if(FRigControlElement* OtherControlElement = this->BlueprintBeingCustomized->Hierarchy->Find<FRigControlElement>(ControlElement->GetKey()))
+				{
+					OtherControlElement->Settings = ControlElement->Settings;
+					this->BlueprintBeingCustomized->Hierarchy->SetControlSettings(OtherControlElement, OtherControlElement->Settings, true, true, true);
+					this->BlueprintBeingCustomized->Hierarchy->SetControlValue(OtherControlElement, ValueToSet, ERigControlValueType::Initial, true);
+					this->BlueprintBeingCustomized->Hierarchy->SetControlValue(OtherControlElement, ValueToSet, ERigControlValueType::Current, true);
+				}
+			}
+			else
+			{
+				this->BlueprintBeingCustomized->PropagateHierarchyFromBPToInstances();
+			}
+		}
+		
+		PropertyUtilities->ForceRefresh();
+	}
+}
+
+void FRigControlElementDetails::CustomizeShape(IDetailLayoutBuilder& DetailBuilder)
+{
+	for(TWeakObjectPtr<UDetailsViewWrapperObject> ObjectBeingCustomized : ObjectsBeingCustomized)
+	{
+		if(ObjectBeingCustomized.IsValid())
+		{
+			if(ObjectBeingCustomized->IsChildOf<FRigControlElement>())
+			{
+				const FRigControlElement ControlElement = ObjectBeingCustomized->GetContent<FRigControlElement>();
+				if(!ControlElement.Settings.SupportsShape())
+				{
+					return;
+				}
+			}
+		}
+	}
+
+	TSharedPtr<IPropertyHandle> ShapeHandle = DetailBuilder.GetProperty(TEXT("Shape"));
+	TSharedPtr<IPropertyHandle> InitialHandle = ShapeHandle->GetChildHandle(TEXT("Initial"));
+	TSharedPtr<IPropertyHandle> LocalHandle = InitialHandle->GetChildHandle(TEXT("Local"));
+	ShapeTransformHandle = LocalHandle->GetChildHandle(TEXT("Transform"));
+	
+	ShapeNameList.Reset();
+	if (BlueprintBeingCustomized)
+	{
+		const bool bUseNameSpace = BlueprintBeingCustomized->ShapeLibraries.Num() > 1;
+		for(TSoftObjectPtr<UControlRigShapeLibrary>& ShapeLibrary : BlueprintBeingCustomized->ShapeLibraries)
+		{
+			if (!ShapeLibrary.IsValid())
+			{
 				ShapeLibrary.LoadSynchronous();
 			}
 			if (ShapeLibrary.IsValid())
@@ -3179,10 +3686,14 @@ void FRigControlElementDetails::RegisterSectionMappings(FPropertyEditorModule& P
 	ControlSection->AddCategory("General");
 	ControlSection->AddCategory("Control");
 	ControlSection->AddCategory("Value");
+	ControlSection->AddCategory("AnimationChannels");
 
 	TSharedRef<FPropertySection> ShapeSection = PropertyEditorModule.FindOrCreateSection(InClass->GetFName(), "Shape", LOCTEXT("Shape", "Shape"));
 	ShapeSection->AddCategory("General");
 	ShapeSection->AddCategory("Shape");
+
+	TSharedRef<FPropertySection> ChannelsSection = PropertyEditorModule.FindOrCreateSection(InClass->GetFName(), "Channels", LOCTEXT("Channels", "Channels"));
+	ShapeSection->AddCategory("AnimationChannels");
 }
 
 bool FRigControlElementDetails::IsShapeEnabled() const
@@ -3218,11 +3729,17 @@ FText FRigControlElementDetails::GetDisplayName() const
 			if(ObjectBeingCustomized->IsChildOf<FRigControlElement>())
 			{
 				const FRigControlElement ControlElement = ObjectBeingCustomized->GetContent<FRigControlElement>();
+
+ 				const FName ThisDisplayName =
+ 					(ControlElement.IsAnimationChannel()) ?
+ 					ControlElement.GetDisplayName() :
+					ControlElement.Settings.DisplayName;
+
 				if(ObjectIndex == 0)
 				{
-					DisplayName = ControlElement.Settings.DisplayName;
+					DisplayName = ThisDisplayName;
 				}
-				else if(DisplayName != ControlElement.Settings.DisplayName)
+				else if(DisplayName != ThisDisplayName)
 				{
 					return ControlRigDetailsMultipleValues;
 				}
@@ -3238,14 +3755,41 @@ FText FRigControlElementDetails::GetDisplayName() const
 }
 
 void FRigControlElementDetails::SetDisplayName(const FText& InNewText, ETextCommit::Type InCommitType)
+{
+	for(int32 ObjectIndex = 0; ObjectIndex < ObjectsBeingCustomized.Num(); ObjectIndex++)
+	{
+		TWeakObjectPtr<UDetailsViewWrapperObject> ObjectBeingCustomized = ObjectsBeingCustomized[ObjectIndex];
+		if(ObjectBeingCustomized.IsValid())
+		{
+			if(ObjectBeingCustomized->IsChildOf<FRigControlElement>())
+			{
+				FRigControlElement ControlElement = ObjectBeingCustomized->GetContent<FRigControlElement>();
+				SetDisplayNameForElement(InNewText, InCommitType, ControlElement.GetKey());
+			}
+		}
+	}
+}
+
+void FRigControlElementDetails::SetDisplayNameForElement(const FText& InNewText, ETextCommit::Type InCommitType, const FRigElementKey& InKeyToRename)
 {
 	if(InCommitType == ETextCommit::OnCleared)
 	{
 		return;
 	}
-	
+
 	const FName DisplayName = InNewText.IsEmpty() ? FName(NAME_None) : FName(*InNewText.ToString());
-	
+
+	if(HierarchyBeingCustomized)
+	{
+		if(const FRigControlElement* ControlElement = HierarchyBeingCustomized->Find<FRigControlElement>(InKeyToRename))
+		{
+			if(ControlElement->GetDisplayName() == DisplayName)
+			{
+				return;
+			}
+		}
+	}
+
 	for(int32 ObjectIndex = 0; ObjectIndex < ObjectsBeingCustomized.Num(); ObjectIndex++)
 	{
 		TWeakObjectPtr<UDetailsViewWrapperObject> ObjectBeingCustomized = ObjectsBeingCustomized[ObjectIndex];
@@ -3254,17 +3798,48 @@ void FRigControlElementDetails::SetDisplayName(const FText& InNewText, ETextComm
 			if(ObjectBeingCustomized->IsChildOf<FRigControlElement>())
 			{
 				FRigControlElement ControlElement = ObjectBeingCustomized->GetContent<FRigControlElement>();
-				ControlElement.Settings.DisplayName = DisplayName;
-				ObjectBeingCustomized->SetContent<FRigControlElement>(ControlElement);
+				if(ControlElement.GetKey() == InKeyToRename)
+				{
+					ControlElement.Settings.DisplayName = DisplayName;
+					ObjectBeingCustomized->SetContent<FRigControlElement>(ControlElement);
+				}
+			}
+		}
+	}
+
+	if(HierarchyBeingCustomized)
+	{
+		FScopedTransaction Transaction(LOCTEXT("SetDisplayName", "Set Display Name"));
+		HierarchyBeingCustomized->GetController(true)->SetDisplayName(InKeyToRename, DisplayName, true, true);
+	}
+}
+
+bool FRigControlElementDetails::OnVerifyDisplayNameChanged(const FText& InText, FText& OutErrorMessage, const FRigElementKey& InKeyToRename)
+{
+	const FString NewName = InText.ToString();
+	if (NewName.IsEmpty())
+	{
+		OutErrorMessage = FText::FromString(TEXT("Name is empty."));
+		return false;
+	}
 
-				if(HierarchyBeingCustomized)
+	// make sure there is no duplicate
+	if(URigHierarchy* Hierarchy = GetHierarchy())
+	{
+		if(const FRigControlElement* ControlElement = Hierarchy->Find<FRigControlElement>(InKeyToRename))
+		{
+			if(const FRigBaseElement* ParentElement = Hierarchy->GetFirstParent(ControlElement))
+			{
+				FString OutErrorString;
+				if (!Hierarchy->IsDisplayNameAvailable(ParentElement->GetKey(), NewName, &OutErrorString))
 				{
-					FScopedTransaction Transaction(LOCTEXT("SetDisplayName", "SetDisplayName"));
-					HierarchyBeingCustomized->SetControlSettings(ControlElement.GetKey(), ControlElement.Settings, true, true, true);
+					OutErrorMessage = FText::FromString(OutErrorString);
+					return false;
 				}
 			}
 		}
 	}
+	return true;
 }
 
 void FRigControlElementDetails::OnCopyShapeProperties()
@@ -3353,32 +3928,39 @@ void FRigControlElementDetails::OnPasteShapeProperties()
 	}
 }
 
-void FRigControlElementDetails::CreateBoolValueWidgetRow(
+FDetailWidgetRow& FRigControlElementDetails::CreateBoolValueWidgetRow(
+	const TArray<FRigElementKey>& Keys,
 	IDetailCategoryBuilder& CategoryBuilder,
 	const FText& Label,
 	const FText& Tooltip,
 	ERigControlValueType ValueType,
-	TAttribute<EVisibility> Visibility)
+	TAttribute<EVisibility> Visibility,
+	TSharedPtr<SWidget> NameContent)
 {
 	const bool bCurrent = ValueType == ERigControlValueType::Current;
 	const bool bInitial = ValueType == ERigControlValueType::Initial;
 
-	TArray<FRigElementKey> Keys = GetElementKeys();
 	URigHierarchy* HierarchyBeingDebugged = GetHierarchyBeingDebugged();
 	URigHierarchy* HierarchyToChange = bCurrent ? HierarchyBeingDebugged : HierarchyBeingCustomized;
-	Keys = HierarchyBeingDebugged->SortKeys(Keys);
 
 	const static TCHAR* TrueText = TEXT("True");
 	const static TCHAR* FalseText = TEXT("False");
 
-	CategoryBuilder.AddCustomRow(Label)
+	if(!NameContent.IsValid())
+	{
+		SAssignNew(NameContent, STextBlock)
+		.Text(Label)
+		.ToolTipText(Tooltip)
+		.Font(IDetailLayoutBuilder::GetDetailFont());
+	}
+
+	FDetailWidgetRow& WidgetRow = CategoryBuilder.AddCustomRow(Label)
 	.Visibility(Visibility)
 	.NameContent()
+	.MinDesiredWidth(200.f)
+	.MaxDesiredWidth(800.f)
 	[
-		SNew(STextBlock)
-		.Text(Label)
-		.ToolTipText(Tooltip)
-		.Font(IDetailLayoutBuilder::GetDetailFont())
+		NameContent.ToSharedRef()
 	]
 	.ValueContent()
 	[
@@ -3457,38 +4039,48 @@ void FRigControlElementDetails::CreateBoolValueWidgetRow(
 			}
 		})
 	));
+
+	return WidgetRow;
 }
 
-void FRigControlElementDetails::CreateFloatValueWidgetRow(
+FDetailWidgetRow& FRigControlElementDetails::CreateFloatValueWidgetRow(
+	const TArray<FRigElementKey>& Keys,
 	IDetailCategoryBuilder& CategoryBuilder,
 	const FText& Label,
 	const FText& Tooltip,
 	ERigControlValueType ValueType,
-	TAttribute<EVisibility> Visibility)
+	TAttribute<EVisibility> Visibility,
+	TSharedPtr<SWidget> NameContent)
 {
-	CreateNumericValueWidgetRow<float>(CategoryBuilder, Label, Tooltip, ValueType, Visibility);
+	return CreateNumericValueWidgetRow<float>(Keys, CategoryBuilder, Label, Tooltip, ValueType, Visibility, NameContent);
 }
 
-void FRigControlElementDetails::CreateIntegerValueWidgetRow(
+FDetailWidgetRow& FRigControlElementDetails::CreateIntegerValueWidgetRow(
+	const TArray<FRigElementKey>& Keys,
 	IDetailCategoryBuilder& CategoryBuilder,
 	const FText& Label,
 	const FText& Tooltip,
 	ERigControlValueType ValueType,
-	TAttribute<EVisibility> Visibility)
+	TAttribute<EVisibility> Visibility,
+	TSharedPtr<SWidget> NameContent)
 {
-	CreateNumericValueWidgetRow<int32>(CategoryBuilder, Label, Tooltip, ValueType, Visibility);
+	return CreateNumericValueWidgetRow<int32>(Keys, CategoryBuilder, Label, Tooltip, ValueType, Visibility, NameContent);
 }
 
-void FRigControlElementDetails::CreateEnumValueWidgetRow(IDetailCategoryBuilder& CategoryBuilder, const FText& Label,
-	const FText& Tooltip, ERigControlValueType ValueType, TAttribute<EVisibility> Visibility)
+FDetailWidgetRow& FRigControlElementDetails::CreateEnumValueWidgetRow(
+	const TArray<FRigElementKey>& Keys,
+	IDetailCategoryBuilder& CategoryBuilder,
+	const FText& Label,
+	const FText& Tooltip,
+	ERigControlValueType ValueType,
+	TAttribute<EVisibility> Visibility,
+	TSharedPtr<SWidget> NameContent)
 {
 	const bool bCurrent = ValueType == ERigControlValueType::Current;
 	const bool bInitial = ValueType == ERigControlValueType::Initial;
 	
-	TArray<FRigElementKey> Keys = GetElementKeys();
 	URigHierarchy* HierarchyBeingDebugged = GetHierarchyBeingDebugged();
 	URigHierarchy* HierarchyToChange = bCurrent ? HierarchyBeingDebugged : HierarchyBeingCustomized;
-	Keys = HierarchyBeingDebugged->SortKeys(Keys);
 
 	UEnum* Enum = nullptr;
 	for(const FRigElementKey& Key : Keys)
@@ -3505,14 +4097,21 @@ void FRigControlElementDetails::CreateEnumValueWidgetRow(IDetailCategoryBuilder&
 
 	check(Enum != nullptr);
 
-	CategoryBuilder.AddCustomRow(Label)
+	if(!NameContent.IsValid())
+	{
+		SAssignNew(NameContent, STextBlock)
+		.Text(Label)
+		.ToolTipText(Tooltip)
+		.Font(IDetailLayoutBuilder::GetDetailFont());
+	}
+
+	FDetailWidgetRow& WidgetRow = CategoryBuilder.AddCustomRow(Label)
 	.Visibility(Visibility)
 	.NameContent()
+	.MinDesiredWidth(200.f)
+	.MaxDesiredWidth(800.f)
 	[
-		SNew(STextBlock)
-		.Text(Label)
-		.ToolTipText(Tooltip)
-		.Font(IDetailLayoutBuilder::GetDetailFont())
+		NameContent.ToSharedRef()
 	]
 	.ValueContent()
 	[
@@ -3596,23 +4195,25 @@ void FRigControlElementDetails::CreateEnumValueWidgetRow(IDetailCategoryBuilder&
 			}
 		})
 	));
+
+	return WidgetRow;
 }
 
-void FRigControlElementDetails::CreateVector2DValueWidgetRow(
+FDetailWidgetRow& FRigControlElementDetails::CreateVector2DValueWidgetRow(
+	const TArray<FRigElementKey>& Keys,
 	IDetailCategoryBuilder& CategoryBuilder, 
 	const FText& Label, 
 	const FText& Tooltip, 
 	ERigControlValueType ValueType,
-	TAttribute<EVisibility> Visibility)
+	TAttribute<EVisibility> Visibility,
+	TSharedPtr<SWidget> NameContent)
 {
 	const bool bCurrent = ValueType == ERigControlValueType::Current;
 	const bool bInitial = ValueType == ERigControlValueType::Initial;
 	const bool bShowToggle = (ValueType == ERigControlValueType::Minimum) || (ValueType == ERigControlValueType::Maximum);
 	
-	TArray<FRigElementKey> Keys = GetElementKeys();
 	URigHierarchy* HierarchyBeingDebugged = GetHierarchyBeingDebugged();
 	URigHierarchy* HierarchyToChange = bCurrent ? HierarchyBeingDebugged : HierarchyBeingCustomized;
-	Keys = HierarchyBeingDebugged->SortKeys(Keys);
 
 	using SNumericVector2DInputBox = SNumericVectorInputBox<float, FVector2f, 2>;
 	TSharedPtr<SNumericVector2DInputBox> VectorInputBox;
@@ -3746,14 +4347,21 @@ void FRigControlElementDetails::CreateVector2DValueWidgetRow(
 			}
 	};
 
+	if(!NameContent.IsValid())
+	{
+		SAssignNew(NameContent, STextBlock)
+		.Text(Label)
+		.ToolTipText(Tooltip)
+		.Font(IDetailLayoutBuilder::GetDetailFont());
+	}
+
 	WidgetRow
 	.Visibility(Visibility)
 	.NameContent()
+	.MinDesiredWidth(200.f)
+	.MaxDesiredWidth(800.f)
 	[
-		SNew(STextBlock)
-		.Text(Label)
-		.ToolTipText(Tooltip)
-		.Font(IDetailLayoutBuilder::GetDetailFont())
+		NameContent.ToSharedRef()
 	]
 	.ValueContent()
 	[
@@ -3852,6 +4460,8 @@ void FRigControlElementDetails::CreateVector2DValueWidgetRow(
 			})
 		));
 	}
+
+	return WidgetRow;
 }
 
 void FRigNullElementDetails::CustomizeDetails(IDetailLayoutBuilder& DetailBuilder)
diff --git a/Engine/Plugins/Animation/ControlRig/Source/ControlRigEditor/Private/ControlRigElementDetails.h b/Engine/Plugins/Animation/ControlRig/Source/ControlRigEditor/Private/ControlRigElementDetails.h
--- a/Engine/Plugins/Animation/ControlRig/Source/ControlRigEditor/Private/ControlRigElementDetails.h
+++ b/Engine/Plugins/Animation/ControlRig/Source/ControlRigEditor/Private/ControlRigElementDetails.h
@@ -53,8 +53,14 @@ protected:
 	/** Helper buttons. */
 	TSharedPtr<SButton> UseSelectedButton;
 	TSharedPtr<SButton> SelectElementButton;
-	FSlateColor OnGetWidgetForeground(const TSharedPtr<SButton> Button) const;
-	FSlateColor OnGetWidgetBackground(const TSharedPtr<SButton> Button) const;
+
+public:
+	
+	static FSlateColor OnGetWidgetForeground(const TSharedPtr<SButton> Button);
+	static FSlateColor OnGetWidgetBackground(const TSharedPtr<SButton> Button);
+	
+protected:
+	
 	FReply OnGetSelectedClicked();
 	FReply OnSelectInHierarchyClicked();
 	
@@ -119,33 +125,51 @@ public:
 	void OnStructContentsChanged(FProperty* InProperty, const TSharedRef<IPropertyUtilities> PropertyUtilities);
 	bool IsSetupModeEnabled() const;
 
+	FText GetParentElementName() const;
+
 	TArray<FRigElementKey> GetElementKeys() const;
 
 	template<typename T>
-	TArray<T> GetElementsInDetailsView() const
+	TArray<T> GetElementsInDetailsView(const TArray<FRigElementKey>& InFilter = TArray<FRigElementKey>()) const
 	{
 		TArray<T> Elements;
 		for(TWeakObjectPtr<UDetailsViewWrapperObject> ObjectBeingCustomized : ObjectsBeingCustomized)
 		{
 			if(ObjectBeingCustomized.IsValid())
 			{
-				Elements.Add(ObjectBeingCustomized->GetContent<T>());
+				T Content = ObjectBeingCustomized->GetContent<T>();
+				if(!InFilter.IsEmpty() && !InFilter.Contains(Content.GetKey()))
+				{
+					continue;
+				}
+				Elements.Add(Content);
 			}
 		}
 		return Elements;
 	}
 
 	template<typename T>
-	TArray<T*> GetElementsInHierarchy() const
+	TArray<T*> GetElementsInHierarchy(const TArray<FRigElementKey>& InFilter = TArray<FRigElementKey>()) const
 	{
 		TArray<T*> Elements;
 		if(HierarchyBeingCustomized)
 		{
-			for(TWeakObjectPtr<UDetailsViewWrapperObject> ObjectBeingCustomized : ObjectsBeingCustomized)
+			if(InFilter.IsEmpty())
 			{
-				if(ObjectBeingCustomized.IsValid())
+				for(TWeakObjectPtr<UDetailsViewWrapperObject> ObjectBeingCustomized : ObjectsBeingCustomized)
 				{
-					Elements.Add(HierarchyBeingCustomized->Find<T>(ObjectBeingCustomized->GetContent<FRigBaseElement>().GetKey()));
+					if(ObjectBeingCustomized.IsValid())
+					{
+						const FRigElementKey Key = ObjectBeingCustomized->GetContent<FRigBaseElement>().GetKey();
+						Elements.Add(HierarchyBeingCustomized->Find<T>(Key));
+					}
+				}
+			}
+			else
+			{
+				for(const FRigElementKey& FilterKey : InFilter)
+				{
+					Elements.Add(HierarchyBeingCustomized->Find<T>(FilterKey));
 				}
 			}
 		}
@@ -162,11 +186,15 @@ public:
 	static void RegisterSectionMappings(FPropertyEditorModule& PropertyEditorModule);
 	virtual void RegisterSectionMappings(FPropertyEditorModule& PropertyEditorModule, UClass* InClass);
 
+	FReply OnSelectParentElementInHierarchyClicked();
+	FReply OnSelectElementClicked(const FRigElementKey& InKey);
+
 protected:
 
 	UControlRigBlueprint* BlueprintBeingCustomized;
 	URigHierarchy* HierarchyBeingCustomized;
 	TArray<TWeakObjectPtr<UDetailsViewWrapperObject>> ObjectsBeingCustomized;
+	TSharedPtr<SButton> SelectParentElementButton;
 };
 
 namespace ERigTransformElementDetailsTransform
@@ -202,7 +230,19 @@ protected:
 
 protected:
 
-	void CreateEulerTransformValueWidgetRow(
+	FDetailWidgetRow& CreateTransformComponentValueWidgetRow(
+		ERigControlType InControlType,
+		URigHierarchy* HierarchyBeingDebugged,
+		const TArray<FRigElementKey>& Keys,
+		SAdvancedTransformInputBox<FEulerTransform>::FArguments TransformWidgetArgs,
+		IDetailCategoryBuilder& CategoryBuilder, 
+		const FText& Label, 
+		const FText& Tooltip,
+		ERigTransformElementDetailsTransform::Type CurrentTransformType,
+		ERigControlValueType ValueType,
+		TSharedPtr<SWidget> NameContent = TSharedPtr<SWidget>());
+
+	FDetailWidgetRow& CreateEulerTransformValueWidgetRow(
 		URigHierarchy* HierarchyBeingDebugged,
 		const TArray<FRigElementKey>& Keys,
 		SAdvancedTransformInputBox<FEulerTransform>::FArguments TransformWidgetArgs,
@@ -210,7 +250,10 @@ protected:
 		const FText& Label, 
 		const FText& Tooltip,
 		ERigTransformElementDetailsTransform::Type CurrentTransformType,
-		ERigControlValueType ValueType);
+		ERigControlValueType ValueType,
+		TSharedPtr<SWidget> NameContent = TSharedPtr<SWidget>());
+
+	static ERigTransformElementDetailsTransform::Type GetTransformTypeFromValueType(ERigControlValueType InValueType);
 
 private:
 
@@ -249,6 +292,7 @@ public:
 	virtual void CustomizeDetails(IDetailLayoutBuilder& DetailBuilder) override;
 	void CustomizeValue(IDetailLayoutBuilder& DetailBuilder);
 	void CustomizeControl(IDetailLayoutBuilder& DetailBuilder);
+	void CustomizeAnimationChannels(IDetailLayoutBuilder& DetailBuilder);
 	void CustomizeShape(IDetailLayoutBuilder& DetailBuilder);
 	virtual void BeginDestroy() override;
 
@@ -261,44 +305,56 @@ public:
 
 	FText GetDisplayName() const;
 	void SetDisplayName(const FText& InNewText, ETextCommit::Type InCommitType);
+	void SetDisplayNameForElement(const FText& InNewText, ETextCommit::Type InCommitType, const FRigElementKey& InKeyToRename);
+	bool OnVerifyDisplayNameChanged(const FText& InText, FText& OutErrorMessage, const FRigElementKey& InKeyToRename);
 
 	void OnCopyShapeProperties();
 	void OnPasteShapeProperties();
 
-	void CreateBoolValueWidgetRow(
+	FDetailWidgetRow& CreateBoolValueWidgetRow(
+		const TArray<FRigElementKey>& Keys,
 		IDetailCategoryBuilder& CategoryBuilder, 
 		const FText& Label, 
 		const FText& Tooltip, 
 		ERigControlValueType ValueType,
-		TAttribute<EVisibility> Visibility = EVisibility::Visible);
+		TAttribute<EVisibility> Visibility = EVisibility::Visible,
+		TSharedPtr<SWidget> NameContent = TSharedPtr<SWidget>());
 
-	void CreateFloatValueWidgetRow(
+	FDetailWidgetRow& CreateFloatValueWidgetRow(
+		const TArray<FRigElementKey>& Keys,
 		IDetailCategoryBuilder& CategoryBuilder, 
 		const FText& Label, 
 		const FText& Tooltip, 
 		ERigControlValueType ValueType,
-		TAttribute<EVisibility> Visibility = EVisibility::Visible);
+		TAttribute<EVisibility> Visibility = EVisibility::Visible,
+		TSharedPtr<SWidget> NameContent = TSharedPtr<SWidget>());
 
-	void CreateIntegerValueWidgetRow(
+	FDetailWidgetRow& CreateIntegerValueWidgetRow(
+		const TArray<FRigElementKey>& Keys,
 		IDetailCategoryBuilder& CategoryBuilder, 
 		const FText& Label, 
 		const FText& Tooltip, 
 		ERigControlValueType ValueType,
-		TAttribute<EVisibility> Visibility = EVisibility::Visible);
+		TAttribute<EVisibility> Visibility = EVisibility::Visible,
+		TSharedPtr<SWidget> NameContent = TSharedPtr<SWidget>());
 
-	void CreateEnumValueWidgetRow(
+	FDetailWidgetRow& CreateEnumValueWidgetRow(
+		const TArray<FRigElementKey>& Keys,
 		IDetailCategoryBuilder& CategoryBuilder, 
 		const FText& Label, 
 		const FText& Tooltip, 
 		ERigControlValueType ValueType,
-		TAttribute<EVisibility> Visibility = EVisibility::Visible);
+		TAttribute<EVisibility> Visibility = EVisibility::Visible,
+		TSharedPtr<SWidget> NameContent = TSharedPtr<SWidget>());
 
-	void CreateVector2DValueWidgetRow(
+	FDetailWidgetRow& CreateVector2DValueWidgetRow(
+		const TArray<FRigElementKey>& Keys,
 		IDetailCategoryBuilder& CategoryBuilder, 
 		const FText& Label, 
 		const FText& Tooltip, 
 		ERigControlValueType ValueType,
-		TAttribute<EVisibility> Visibility = EVisibility::Visible);
+		TAttribute<EVisibility> Visibility = EVisibility::Visible,
+		TSharedPtr<SWidget> NameContent = TSharedPtr<SWidget>());
 
 
 	// this is a template since we specialize it further down for
@@ -312,21 +368,21 @@ public:
 private:
 
 	template<typename T>
-	void CreateNumericValueWidgetRow(
+	FDetailWidgetRow& CreateNumericValueWidgetRow(
+		const TArray<FRigElementKey>& Keys,
 		IDetailCategoryBuilder& CategoryBuilder, 
 		const FText& Label, 
 		const FText& Tooltip, 
 		ERigControlValueType ValueType,
-		TAttribute<EVisibility> Visibility = EVisibility::Visible)
+		TAttribute<EVisibility> Visibility = EVisibility::Visible,
+		TSharedPtr<SWidget> NameContent = TSharedPtr<SWidget>())
 	{
 		const bool bCurrent = ValueType == ERigControlValueType::Current;
 		const bool bInitial = ValueType == ERigControlValueType::Initial;
 		const bool bShowToggle = (ValueType == ERigControlValueType::Minimum) || (ValueType == ERigControlValueType::Maximum);
 		
-		TArray<FRigElementKey> Keys = GetElementKeys();
 		URigHierarchy* HierarchyBeingDebugged = GetHierarchyBeingDebugged();
 		URigHierarchy* HierarchyToChange = bCurrent ? HierarchyBeingDebugged : HierarchyBeingCustomized;
-		Keys = HierarchyBeingDebugged->SortKeys(Keys);
 
 		TSharedPtr<SNumericEntryBox<T>> NumericEntryBox;
 		
@@ -425,14 +481,21 @@ private:
 				}
 			};
 
+		if(!NameContent.IsValid())
+		{
+			SAssignNew(NameContent, STextBlock)
+			.Text(Label)
+			.ToolTipText(Tooltip)
+			.Font(IDetailLayoutBuilder::GetDetailFont());
+		}
+
 		WidgetRow
 		.Visibility(Visibility)
 		.NameContent()
+		.MinDesiredWidth(200.f)
+		.MaxDesiredWidth(800.f)
 		[
-			SNew(STextBlock)
-			.Text(Label)
-			.ToolTipText(Tooltip)
-			.Font(IDetailLayoutBuilder::GetDetailFont())
+			NameContent.ToSharedRef()
 		]
 		.ValueContent()
 		[
@@ -548,8 +611,16 @@ private:
 				})
 			));
 		}
+
+		return WidgetRow;
 	}
 
+	// animation channel related callbacks
+	FReply OnAddAnimationChannelClicked();
+	TSharedRef<ITableRow> HandleGenerateAnimationChannelTypeRow(TSharedPtr<ERigControlType> ControlType, const TSharedRef<STableViewBase>& OwnerTable, FRigElementKey ControlKey);
+	void HandleControlTypeChanged(TSharedPtr<ERigControlType> ControlType, ESelectInfo::Type SelectInfo, FRigElementKey ControlKey, const TSharedRef<IPropertyUtilities> PropertyUtilities);
+	void HandleControlTypeChanged(ERigControlType ControlType, TArray<FRigElementKey> ControlKeys, const TSharedRef<IPropertyUtilities> PropertyUtilities);
+
 	TArray<TSharedPtr<FString>> ShapeNameList;
 	TSharedPtr<FRigInfluenceEntryModifier> InfluenceModifier;
 	TSharedPtr<FStructOnScope> InfluenceModifierStruct;
diff --git a/Engine/Plugins/Animation/ControlRig/Source/ControlRigEditor/Private/ControlRigHierarchyCommands.cpp b/Engine/Plugins/Animation/ControlRig/Source/ControlRigEditor/Private/ControlRigHierarchyCommands.cpp
--- a/Engine/Plugins/Animation/ControlRig/Source/ControlRigEditor/Private/ControlRigHierarchyCommands.cpp
+++ b/Engine/Plugins/Animation/ControlRig/Source/ControlRigEditor/Private/ControlRigHierarchyCommands.cpp
@@ -8,6 +8,7 @@ void FControlRigHierarchyCommands::RegisterCommands()
 {
 	UI_COMMAND(AddBoneItem, "New Bone", "Add new bone at the origin (0, 0, 0) to the hierarchy.", EUserInterfaceActionType::Button, FInputChord());
 	UI_COMMAND(AddControlItem, "New Control", "Add new control at the origin (0, 0, 0) to the hierarchy.", EUserInterfaceActionType::Button, FInputChord(EKeys::N, EModifierKey::Control));
+	UI_COMMAND(AddAnimationChannelItem, "New Animation Channel", "Add new animation channel.", EUserInterfaceActionType::Button, FInputChord(EKeys::N, EModifierKey::Control));
 	UI_COMMAND(AddNullItem, "New Null", "Add new null at the origin (0, 0, 0) to the hierarchy.", EUserInterfaceActionType::Button, FInputChord());
 	UI_COMMAND(DuplicateItem, "Duplicate", "Duplicate the selected items in the hierarchy.", EUserInterfaceActionType::Button, FInputChord(EKeys::D, EModifierKey::Control));
 	UI_COMMAND(MirrorItem, "Mirror", "Mirror the selected items in the hierarchy.", EUserInterfaceActionType::Button, FInputChord());
diff --git a/Engine/Plugins/Animation/ControlRig/Source/ControlRigEditor/Private/ControlRigHierarchyCommands.h b/Engine/Plugins/Animation/ControlRig/Source/ControlRigEditor/Private/ControlRigHierarchyCommands.h
--- a/Engine/Plugins/Animation/ControlRig/Source/ControlRigEditor/Private/ControlRigHierarchyCommands.h
+++ b/Engine/Plugins/Animation/ControlRig/Source/ControlRigEditor/Private/ControlRigHierarchyCommands.h
@@ -24,6 +24,9 @@ public:
 	/** Add Item at origin */
 	TSharedPtr< FUICommandInfo > AddControlItem;
 
+	/** Add an animation channel */
+	TSharedPtr< FUICommandInfo > AddAnimationChannelItem;
+
 	/** Add Item at origin */
 	TSharedPtr< FUICommandInfo > AddNullItem;
 
diff --git a/Engine/Plugins/Animation/ControlRig/Source/ControlRigEditor/Private/EditMode/ControlRigControlsProxy.cpp b/Engine/Plugins/Animation/ControlRig/Source/ControlRigEditor/Private/EditMode/ControlRigControlsProxy.cpp
--- a/Engine/Plugins/Animation/ControlRig/Source/ControlRigEditor/Private/EditMode/ControlRigControlsProxy.cpp
+++ b/Engine/Plugins/Animation/ControlRig/Source/ControlRigEditor/Private/EditMode/ControlRigControlsProxy.cpp
@@ -801,7 +801,7 @@ void UControlRigDetailPanelControlProxies::AddProxy(UControlRig* ControlRig, con
 		if (Proxy)
 		{
 			Proxy->bIsIndividual =
-				(ControlElement->Settings.AnimationType == ERigControlAnimationType::AnimationChannel) ||
+				(ControlElement->IsAnimationChannel()) ||
 				(ControlElement->Settings.AnimationType == ERigControlAnimationType::ProxyControl) ;
 			Proxy->SetFlags(RF_Transactional);
 			Proxy->SetName(Name);
diff --git a/Engine/Plugins/Animation/ControlRig/Source/ControlRigEditor/Private/EditMode/ControlRigEditMode.cpp b/Engine/Plugins/Animation/ControlRig/Source/ControlRigEditor/Private/EditMode/ControlRigEditMode.cpp
--- a/Engine/Plugins/Animation/ControlRig/Source/ControlRigEditor/Private/EditMode/ControlRigEditMode.cpp
+++ b/Engine/Plugins/Animation/ControlRig/Source/ControlRigEditor/Private/EditMode/ControlRigEditMode.cpp
@@ -2635,7 +2635,7 @@ void FControlRigEditMode::ResetTransforms(bool bSelectionOnly)
 				{
 					ControlsToReset.Add(Control->GetKey());
 					if(Control->Settings.AnimationType == ERigControlAnimationType::AnimationControl ||
-						Control->Settings.AnimationType == ERigControlAnimationType::AnimationChannel)
+						Control->IsAnimationChannel())
 					{
 						ControlsInteracting.Add(Control->GetKey());
 					}
@@ -3342,7 +3342,7 @@ void FControlRigEditMode::OnHierarchyModified(ERigHierarchyNotification InNotif,
 									ControlRig->GetHierarchy()->ForEach<FRigControlElement>(
 										[this, ControlRig, ControlElement, DrivenKeys, bSelected](FRigControlElement* AnimationChannelControl) -> bool
 										{
-											if(AnimationChannelControl->Settings.AnimationType == ERigControlAnimationType::AnimationChannel)
+											if(AnimationChannelControl->IsAnimationChannel())
 											{
 												if(const FRigControlElement* ParentControlElement =
 													Cast<FRigControlElement>(ControlRig->GetHierarchy()->GetFirstParent(AnimationChannelControl)))
diff --git a/Engine/Plugins/Animation/ControlRig/Source/ControlRigEditor/Private/Editor/ControlRigEditor.cpp b/Engine/Plugins/Animation/ControlRig/Source/ControlRigEditor/Private/Editor/ControlRigEditor.cpp
--- a/Engine/Plugins/Animation/ControlRig/Source/ControlRigEditor/Private/Editor/ControlRigEditor.cpp
+++ b/Engine/Plugins/Animation/ControlRig/Source/ControlRigEditor/Private/Editor/ControlRigEditor.cpp
@@ -54,6 +54,7 @@
 #include "ControlRig/Private/Units/Execution/RigUnit_InverseExecution.h"
 #include "ControlRig/Private/Units/Hierarchy/RigUnit_GetControlTransform.h"
 #include "ControlRig/Private/Units/Hierarchy/RigUnit_SetControlTransform.h"
+#include "ControlRig/Private/Units/Hierarchy/RigUnit_ControlChannel.h"
 #include "ControlRig/Private/Units/Execution/RigUnit_Collection.h"
 #include "ControlRig/Private/Units/Highlevel/Hierarchy/RigUnit_TransformConstraint.h"
 #include "ControlRig/Private/Units/Hierarchy/RigUnit_GetControlTransform.h"
@@ -5415,12 +5416,31 @@ void FControlRigEditor::CreateRigHierarchyToGraphDragAndDropMenu() const
 					}
 			
 					uint8 DraggedTypes = 0;
+					uint8 DraggedAnimationTypes = 2;
 					for (const FRigElementKey& DraggedKey : DragDropContext.DraggedElementKeys)
 					{
 						if (DraggedKey.Type != LastType)
 						{
 							bMultipleTypeSelected = true;
 						}
+						else if(DraggedKey.Type == ERigElementType::Control)
+						{
+							if(const FRigControlElement* ControlElement = Blueprint->Hierarchy->Find<FRigControlElement>(DraggedKey))
+							{
+								const uint8 DraggedAnimationType = ControlElement->IsAnimationChannel() ? 1 : 0; 
+								if(DraggedAnimationTypes == 2)
+								{
+									DraggedAnimationTypes = DraggedAnimationType;
+								}
+								else
+								{
+									if(DraggedAnimationTypes != DraggedAnimationType)
+									{
+										bMultipleTypeSelected = true;
+									}
+								}
+							}
+						}
 				
 						DraggedTypes = DraggedTypes | (uint8)DraggedKey.Type;
 					}
@@ -5459,10 +5479,20 @@ void FControlRigEditor::CreateRigHierarchyToGraphDragAndDropMenu() const
 						}
 						else if ((DraggedTypes & (uint8)ERigElementType::Control) != 0)
 						{
-							GetterLabel = LOCTEXT("GetControl","Get Control");
-							GetterTooltip = LOCTEXT("GetControl_ToolTip", "Getter For Control");
-							SetterLabel = LOCTEXT("SetControl","Set Control");
-							SetterTooltip = LOCTEXT("SetControl_ToolTip", "Setter For Control");
+							if(DraggedAnimationTypes == 0)
+							{
+								GetterLabel = LOCTEXT("GetControl","Get Control");
+								GetterTooltip = LOCTEXT("GetControl_ToolTip", "Getter For Control");
+								SetterLabel = LOCTEXT("SetControl","Set Control");
+								SetterTooltip = LOCTEXT("SetControl_ToolTip", "Setter For Control");
+							}
+							else
+							{
+								GetterLabel = LOCTEXT("GetAnimationChannel","Get Animation Channel");
+								GetterTooltip = LOCTEXT("GetAnimationChannel_ToolTip", "Getter For Animation Channel");
+								SetterLabel = LOCTEXT("SetAnimationChannel","Set Animation Channel");
+								SetterTooltip = LOCTEXT("SetAnimationChannel_ToolTip", "Setter For Animation Channel");
+							}
 						}
 					}
 
@@ -5724,9 +5754,140 @@ void FControlRigEditor::HandleMakeElementGetterSetter(ERigElementGetterSetterTyp
 		TArray<FName> ItemPins;
 		ItemPins.Add(TEXT("Item"));
 
+		FName NameValue = Key.Name;
+		FName ChannelValue = Key.Name;
 		TArray<FName> NamePins;
+		TArray<FName> ChannelPins;
+		TMap<FName, FRigVMTemplateArgumentType> PinsToResolve; 
+
+		if(FRigControlElement* ControlElement = Blueprint->Hierarchy->Find<FRigControlElement>(Key))
+		{
+			if(ControlElement->IsAnimationChannel())
+			{
+				if(const FRigControlElement* ParentControlElement = Cast<FRigControlElement>(Blueprint->Hierarchy->GetFirstParent(ControlElement)))
+				{
+					NameValue = ParentControlElement->GetName();
+				}
+				else
+				{
+					NameValue = NAME_None;
+				}
+
+				ItemPins.Reset();
+				NamePins.Add(TEXT("Control"));
+				ChannelPins.Add(TEXT("Channel"));
+				static const FName ValueName = GET_MEMBER_NAME_CHECKED(FRigUnit_GetBoolAnimationChannel, Value);
+
+				switch (ControlElement->Settings.ControlType)
+				{
+					case ERigControlType::Bool:
+					{
+						if(bIsGetter)
+						{
+							StructTemplate = FRigUnit_GetBoolAnimationChannel::StaticStruct();
+						}
+						else
+						{
+							StructTemplate = FRigUnit_SetBoolAnimationChannel::StaticStruct();
+						}
+						PinsToResolve.Add(ValueName, FRigVMTemplateArgumentType(RigVMTypeUtils::BoolType));
+						break;
+					}
+					case ERigControlType::Float:
+					{
+						if(bIsGetter)
+						{
+							StructTemplate = FRigUnit_GetFloatAnimationChannel::StaticStruct();
+						}
+						else
+						{
+							StructTemplate = FRigUnit_SetFloatAnimationChannel::StaticStruct();
+						}
+						PinsToResolve.Add(ValueName, FRigVMTemplateArgumentType(RigVMTypeUtils::FloatType));
+						break;
+					}
+					case ERigControlType::Integer:
+					{
+						if(bIsGetter)
+						{
+							StructTemplate = FRigUnit_GetIntAnimationChannel::StaticStruct();
+						}
+						else
+						{
+							StructTemplate = FRigUnit_SetIntAnimationChannel::StaticStruct();
+						}
+						PinsToResolve.Add(ValueName, FRigVMTemplateArgumentType(RigVMTypeUtils::Int32Type));
+						break;
+					}
+					case ERigControlType::Vector2D:
+					{
+						if(bIsGetter)
+						{
+							StructTemplate = FRigUnit_GetVector2DAnimationChannel::StaticStruct();
+						}
+						else
+						{
+							StructTemplate = FRigUnit_SetVector2DAnimationChannel::StaticStruct();
+						}
 
-		if (bIsGetter)
+						UScriptStruct* ValueStruct = TBaseStructure<FVector2D>::Get();
+						PinsToResolve.Add(ValueName, FRigVMTemplateArgumentType(ValueStruct->GetStructCPPName(), ValueStruct));
+						break;
+					}
+					case ERigControlType::Position:
+					case ERigControlType::Scale:
+					{
+						if(bIsGetter)
+						{
+							StructTemplate = FRigUnit_GetVectorAnimationChannel::StaticStruct();
+						}
+						else
+						{
+							StructTemplate = FRigUnit_SetVectorAnimationChannel::StaticStruct();
+						}
+						UScriptStruct* ValueStruct = TBaseStructure<FVector>::Get();
+						PinsToResolve.Add(ValueName, FRigVMTemplateArgumentType(ValueStruct->GetStructCPPName(), ValueStruct));
+						break;
+					}
+					case ERigControlType::Rotator:
+					{
+						if(bIsGetter)
+						{
+							StructTemplate = FRigUnit_GetRotatorAnimationChannel::StaticStruct();
+						}
+						else
+						{
+							StructTemplate = FRigUnit_SetRotatorAnimationChannel::StaticStruct();
+						}
+						UScriptStruct* ValueStruct = TBaseStructure<FRotator>::Get();
+						PinsToResolve.Add(ValueName, FRigVMTemplateArgumentType(ValueStruct->GetStructCPPName(), ValueStruct));
+						break;
+					}
+					case ERigControlType::Transform:
+					case ERigControlType::TransformNoScale:
+					case ERigControlType::EulerTransform:
+					{
+						if(bIsGetter)
+						{
+							StructTemplate = FRigUnit_GetTransformAnimationChannel::StaticStruct();
+						}
+						else
+						{
+							StructTemplate = FRigUnit_SetTransformAnimationChannel::StaticStruct();
+						}
+						UScriptStruct* ValueStruct = TBaseStructure<FTransform>::Get();
+						PinsToResolve.Add(ValueName, FRigVMTemplateArgumentType(ValueStruct->GetStructCPPName(), ValueStruct));
+						break;
+					}
+					default:
+					{
+						break;
+					}
+				}
+			}
+		}
+
+		if (bIsGetter && StructTemplate == nullptr)
 		{
 			switch (Type)
 			{
@@ -5814,7 +5975,7 @@ void FControlRigEditor::HandleMakeElementGetterSetter(ERigElementGetterSetterTyp
 				}
 			}
 		}
-		else
+		else if(StructTemplate == nullptr)
 		{
 			switch (Type)
 			{
@@ -5956,7 +6117,15 @@ void FControlRigEditor::HandleMakeElementGetterSetter(ERigElementGetterSetterTyp
 			FString ItemTypeStr = StaticEnum<ERigElementType>()->GetDisplayNameTextByValue((int64)Key.Type).ToString();
 			NewNode.Name = ModelNode->GetFName();
 			NewNodes.Add(NewNode);
-			
+
+			for (const TPair<FName, FRigVMTemplateArgumentType>& PinToResolve : PinsToResolve)
+			{
+				if(URigVMPin* Pin = ModelNode->FindPin(PinToResolve.Key.ToString()))
+				{
+					GetFocusedController()->ResolveWildCardPin(Pin, PinToResolve.Value, true, true);
+				}
+			}
+
 			for (const FName& ItemPin : ItemPins)
 			{
 				GetFocusedController()->SetPinDefaultValue(FString::Printf(TEXT("%s.%s.Name"), *ModelNode->GetName(), *ItemPin.ToString()), Key.Name.ToString(), true, true, false, true);
@@ -5965,7 +6134,14 @@ void FControlRigEditor::HandleMakeElementGetterSetter(ERigElementGetterSetterTyp
 
 			for (const FName& NamePin : NamePins)
 			{
-				GetFocusedController()->SetPinDefaultValue(FString::Printf(TEXT("%s.%s"), *ModelNode->GetName(), *NamePin.ToString()), Key.Name.ToString(), true, true, false, true);
+				const FString PinPath = FString::Printf(TEXT("%s.%s"), *ModelNode->GetName(), *NamePin.ToString());
+				GetFocusedController()->SetPinDefaultValue(PinPath, NameValue.ToString(), true, true, false, true);
+			}
+
+			for (const FName& ChannelPin : ChannelPins)
+			{
+				const FString PinPath = FString::Printf(TEXT("%s.%s"), *ModelNode->GetName(), *ChannelPin.ToString());
+				GetFocusedController()->SetPinDefaultValue(PinPath, ChannelValue.ToString(), true, true, false, true);
 			}
 
 			if (!NewNode.ValuePinName.IsNone())
diff --git a/Engine/Plugins/Animation/ControlRig/Source/ControlRigEditor/Private/Editor/SRigHierarchy.cpp b/Engine/Plugins/Animation/ControlRig/Source/ControlRigEditor/Private/Editor/SRigHierarchy.cpp
--- a/Engine/Plugins/Animation/ControlRig/Source/ControlRigEditor/Private/Editor/SRigHierarchy.cpp
+++ b/Engine/Plugins/Animation/ControlRig/Source/ControlRigEditor/Private/Editor/SRigHierarchy.cpp
@@ -298,13 +298,17 @@ void SRigHierarchy::BindCommands()
 	const FControlRigHierarchyCommands& Commands = FControlRigHierarchyCommands::Get();
 
 	CommandList->MapAction(Commands.AddBoneItem,
-		FExecuteAction::CreateSP(this, &SRigHierarchy::HandleNewItem, ERigElementType::Bone));
+		FExecuteAction::CreateSP(this, &SRigHierarchy::HandleNewItem, ERigElementType::Bone, false));
 
 	CommandList->MapAction(Commands.AddControlItem,
-		FExecuteAction::CreateSP(this, &SRigHierarchy::HandleNewItem, ERigElementType::Control));
+		FExecuteAction::CreateSP(this, &SRigHierarchy::HandleNewItem, ERigElementType::Control, false));
+
+	CommandList->MapAction(Commands.AddAnimationChannelItem,
+		FExecuteAction::CreateSP(this, &SRigHierarchy::HandleNewItem, ERigElementType::Control, true),
+		FCanExecuteAction::CreateSP(this, &SRigHierarchy::IsControlSelected));
 
 	CommandList->MapAction(Commands.AddNullItem,
-		FExecuteAction::CreateSP(this, &SRigHierarchy::HandleNewItem, ERigElementType::Null));
+		FExecuteAction::CreateSP(this, &SRigHierarchy::HandleNewItem, ERigElementType::Null, false));
 
 	CommandList->MapAction(Commands.DuplicateItem,
 		FExecuteAction::CreateSP(this, &SRigHierarchy::HandleDuplicateItem),
@@ -761,11 +765,8 @@ void SRigHierarchy::OnHierarchyModified(ERigHierarchyNotification InNotif, URigH
 					TSharedPtr<FRigTreeElement> TreeElement = TreeView->FindElement(InElement->GetKey(), TreeView->RootElements[RootIndex]);
 					if (TreeElement.IsValid())
 					{
-						if(URigHierarchy* Hierarchy = GetDebuggedHierarchy())
-						{
-							const FRigTreeDisplaySettings& Settings = TreeView->GetRigTreeDelegates().GetDisplaySettings();
-							TreeElement->RefreshDisplaySettings(Hierarchy, Settings);
-						}
+						const FRigTreeDisplaySettings& Settings = TreeView->GetRigTreeDelegates().GetDisplaySettings();
+						TreeElement->RefreshDisplaySettings(InHierarchy, Settings);
 					}
 				}
 			}
@@ -1097,6 +1098,10 @@ void SRigHierarchy::CreateContextMenu() const
 								DefaultSection.AddMenuEntry(Commands.AddBoneItem);
 							}
 							DefaultSection.AddMenuEntry(Commands.AddControlItem);
+							if(SelectedKey.Type == ERigElementType::Control)
+							{
+								DefaultSection.AddMenuEntry(Commands.AddAnimationChannelItem);
+							}
 							DefaultSection.AddMenuEntry(Commands.AddNullItem);
 						})
 					);
@@ -1544,7 +1549,7 @@ bool SRigHierarchy::CanDeleteItem() const
 }
 
 /** Create Item */
-void SRigHierarchy::HandleNewItem(ERigElementType InElementType)
+void SRigHierarchy::HandleNewItem(ERigElementType InElementType, bool bIsAnimationChannel)
 {
 	if(!ControlRigEditor.IsValid())
 	{
@@ -1585,6 +1590,12 @@ void SRigHierarchy::HandleNewItem(ERigElementType InElementType)
 		else
 		{
 			NewNameTemplate = FString::Printf(TEXT("New%s"), *StaticEnum<ERigElementType>()->GetNameStringByValue((int64)InElementType));
+
+			if(bIsAnimationChannel)
+			{
+				static const FString NewAnimationChannel = TEXT("Channel");
+				NewNameTemplate = NewAnimationChannel;
+			}
 		}
 		
 		const FName NewElementName = CreateUniqueName(*NewNameTemplate, InElementType);
@@ -1600,14 +1611,27 @@ void SRigHierarchy::HandleNewItem(ERigElementType InElementType)
 				case ERigElementType::Control:
 				{
 					FRigControlSettings Settings;
-					Settings.ControlType = ERigControlType::EulerTransform;
 
-					FEulerTransform Identity = FEulerTransform::Identity;
-					FRigControlValue ValueToSet = FRigControlValue::Make<FEulerTransform>(Identity);
-					Settings.MinimumValue = ValueToSet;
-					Settings.MaximumValue = ValueToSet;
-						
-					NewItemKey = Controller->AddControl(NewElementName, ParentKey, Settings, Settings.GetIdentityValue(), FTransform::Identity, FTransform::Identity, true, true);
+					if(bIsAnimationChannel)
+					{
+						Settings.AnimationType = ERigControlAnimationType::AnimationChannel;
+						Settings.ControlType = ERigControlType::Float;
+						Settings.MinimumValue = FRigControlValue::Make<float>(0.f);
+						Settings.MaximumValue = FRigControlValue::Make<float>(1.f);
+						Settings.DisplayName = Hierarchy->GetSafeNewDisplayName(ParentKey, *NewNameTemplate);
+
+						NewItemKey = Controller->AddAnimationChannel(NewElementName, ParentKey, Settings, true, true);
+					}
+					else
+					{
+						Settings.ControlType = ERigControlType::EulerTransform;
+						FEulerTransform Identity = FEulerTransform::Identity;
+						FRigControlValue ValueToSet = FRigControlValue::Make<FEulerTransform>(Identity);
+						Settings.MinimumValue = ValueToSet;
+						Settings.MaximumValue = ValueToSet;
+
+						NewItemKey = Controller->AddControl(NewElementName, ParentKey, Settings, Settings.GetIdentityValue(), FTransform::Identity, FTransform::Identity, true, true);
+					}						
 					break;
 				}
 				case ERigElementType::Null:
@@ -2123,11 +2147,6 @@ FName SRigHierarchy::HandleRenameElement(const FRigElementKey& OldKey, const FSt
 {
 	ClearDetailPanel();
 
-	if (OldKey.Name.ToString().Equals(NewName, ESearchCase::CaseSensitive))	
-	{
-		return NAME_None;
-	}
-
 	// make sure there is no duplicate
 	if (ControlRigBlueprint.IsValid())
 	{
@@ -2142,7 +2161,23 @@ FName SRigHierarchy::HandleRenameElement(const FRigElementKey& OldKey, const FSt
 		const FName SanitizedName = *SanitizedNameStr;
 		FName ResultingName = NAME_None;
 
-		ResultingName = Controller->RenameElement(OldKey, SanitizedName, true, true, false).Name;
+		bool bUseDisplayName = false;
+		if(const FRigControlElement* ControlElement = Hierarchy->Find<FRigControlElement>(OldKey))
+		{
+			if(ControlElement->IsAnimationChannel())
+			{
+				bUseDisplayName = true;
+			}
+		}
+
+		if(bUseDisplayName)
+		{
+			ResultingName = Controller->SetDisplayName(OldKey, SanitizedName, true, true, true);
+		}
+		else
+		{
+			ResultingName = Controller->RenameElement(OldKey, SanitizedName, true, true, false).Name;
+		}
 		ControlRigBlueprint->PropagateHierarchyFromBPToInstances();
 		return ResultingName;
 	}
@@ -2152,9 +2187,30 @@ FName SRigHierarchy::HandleRenameElement(const FRigElementKey& OldKey, const FSt
 
 bool SRigHierarchy::HandleVerifyNameChanged(const FRigElementKey& OldKey, const FString& NewName, FText& OutErrorMessage)
 {
-	if (OldKey.Name.ToString() == NewName)
+	bool bIsAnimationChannel = false;
+	if (ControlRigBlueprint.IsValid())
 	{
-		return true;
+		URigHierarchy* Hierarchy = GetHierarchy();
+		if(const FRigControlElement* ControlElement = Hierarchy->Find<FRigControlElement>(OldKey))
+		{
+			if(ControlElement->IsAnimationChannel())
+			{
+				bIsAnimationChannel = true;
+
+				if(ControlElement->GetDisplayName().ToString() == NewName)
+				{
+					return true;
+				}
+			}
+		}
+	}
+
+	if(!bIsAnimationChannel)
+	{
+		if (OldKey.Name.ToString() == NewName)
+		{
+			return true;
+		}
 	}
 
 	if (NewName.IsEmpty())
@@ -2168,11 +2224,29 @@ bool SRigHierarchy::HandleVerifyNameChanged(const FRigElementKey& OldKey, const
 	{
 		URigHierarchy* Hierarchy = GetHierarchy();
 
-		FString OutErrorString;
-		if (!Hierarchy->IsNameAvailable(NewName, OldKey.Type, &OutErrorString))
+		if(bIsAnimationChannel)
 		{
-			OutErrorMessage = FText::FromString(OutErrorString);
-			return false;
+			if(const FRigControlElement* ControlElement = Hierarchy->Find<FRigControlElement>(OldKey))
+			{
+				if(const FRigBaseElement* ParentElement = Hierarchy->GetFirstParent(ControlElement))
+				{
+					FString OutErrorString;
+					if (!Hierarchy->IsDisplayNameAvailable(ParentElement->GetKey(), NewName, &OutErrorString))
+					{
+						OutErrorMessage = FText::FromString(OutErrorString);
+						return false;
+					}
+				}
+			}
+		}
+		else
+		{
+			FString OutErrorString;
+			if (!Hierarchy->IsNameAvailable(NewName, OldKey.Type, &OutErrorString))
+			{
+				OutErrorMessage = FText::FromString(OutErrorString);
+				return false;
+			}
 		}
 	}
 	return true;
diff --git a/Engine/Plugins/Animation/ControlRig/Source/ControlRigEditor/Private/Editor/SRigHierarchy.h b/Engine/Plugins/Animation/ControlRig/Source/ControlRigEditor/Private/Editor/SRigHierarchy.h
--- a/Engine/Plugins/Animation/ControlRig/Source/ControlRigEditor/Private/Editor/SRigHierarchy.h
+++ b/Engine/Plugins/Animation/ControlRig/Source/ControlRigEditor/Private/Editor/SRigHierarchy.h
@@ -104,7 +104,7 @@ private:
 	void HandleDeleteItem();
 
 	/** Create a new item */
-	void HandleNewItem(ERigElementType InElementType);
+	void HandleNewItem(ERigElementType InElementType, bool bIsAnimationChannel);
 
 	/** Check whether we can deleting the selected item(s) */
 	bool CanDuplicateItem() const;
diff --git a/Engine/Plugins/Animation/ControlRig/Source/ControlRigEditor/Private/Editor/SRigHierarchyTreeView.cpp b/Engine/Plugins/Animation/ControlRig/Source/ControlRigEditor/Private/Editor/SRigHierarchyTreeView.cpp
--- a/Engine/Plugins/Animation/ControlRig/Source/ControlRigEditor/Private/Editor/SRigHierarchyTreeView.cpp
+++ b/Engine/Plugins/Animation/ControlRig/Source/ControlRigEditor/Private/Editor/SRigHierarchyTreeView.cpp
@@ -32,7 +32,9 @@ FRigTreeDisplaySettings FRigTreeDelegates::DefaultDisplaySettings;
 FRigTreeElement::FRigTreeElement(const FRigElementKey& InKey, TWeakPtr<SRigHierarchyTreeView> InTreeView, bool InSupportsRename, ERigTreeFilterResult InFilterResult)
 {
 	Key = InKey;
+	ChannelName = NAME_None;
 	bIsTransient = false;
+	bIsAnimationChannel = false;
 	bSupportsRename = InSupportsRename;
 	FilterResult = InFilterResult;
 
@@ -43,6 +45,11 @@ FRigTreeElement::FRigTreeElement(const FRigElementKey& InKey, TWeakPtr<SRigHiera
 			if(const FRigControlElement* ControlElement = Hierarchy->Find<FRigControlElement>(Key))
 			{
 				bIsTransient = ControlElement->Settings.bIsTransientControl;
+				bIsAnimationChannel = ControlElement->IsAnimationChannel();
+				if(bIsAnimationChannel)
+				{
+					ChannelName = ControlElement->GetDisplayName();
+				}
 			}
 
 			const FRigTreeDisplaySettings& Settings = InTreeView.Pin()->GetRigTreeDelegates().GetDisplaySettings();
@@ -175,6 +182,10 @@ FText SRigHierarchyItem::GetName() const
 		static const FText TemporaryControl = FText::FromString(TEXT("Temporary Control"));
 		return TemporaryControl;
 	}
+	if(WeakRigTreeElement.Pin()->bIsAnimationChannel)
+	{
+		return FText::FromName(WeakRigTreeElement.Pin()->ChannelName);
+	}
 	return (FText::FromName(WeakRigTreeElement.Pin()->Key.Name));
 }
 
@@ -703,63 +714,7 @@ TPair<const FSlateBrush*, FSlateColor> SRigHierarchyItem::GetBrushForElementType
 				{
 					static FName TypeIcon(TEXT("Kismet.VariableList.TypeIcon"));
 					Brush = FAppStyle::GetBrush(TypeIcon);
-
-					FEdGraphPinType PinType;
-					switch(Control->Settings.ControlType)
-					{
-						case ERigControlType::Bool:
-						{
-							PinType = RigVMTypeUtils::PinTypeFromCPPType(RigVMTypeUtils::BoolType, nullptr);
-							break;
-						}
-						case ERigControlType::Float:
-						{
-							PinType = RigVMTypeUtils::PinTypeFromCPPType(RigVMTypeUtils::FloatType, nullptr);
-							break;
-						}
-						case ERigControlType::Integer:
-						{
-							if(Control->Settings.ControlEnum)
-							{
-								PinType = RigVMTypeUtils::PinTypeFromCPPType(FString(), Control->Settings.ControlEnum);
-							}
-							else
-							{
-								PinType = RigVMTypeUtils::PinTypeFromCPPType(RigVMTypeUtils::Int32Type, nullptr);
-							}
-							break;
-						}
-						case ERigControlType::Vector2D:
-						{
-							UScriptStruct* Struct = TBaseStructure<FVector2D>::Get(); 
-							PinType = RigVMTypeUtils::PinTypeFromCPPType(Struct->GetStructCPPName(), Struct);
-							break;
-						}
-						case ERigControlType::Position:
-						case ERigControlType::Scale:
-						{
-							UScriptStruct* Struct = TBaseStructure<FVector>::Get(); 
-							PinType = RigVMTypeUtils::PinTypeFromCPPType(Struct->GetStructCPPName(), Struct);
-							break;
-						}
-						case ERigControlType::Rotator:
-						{
-							UScriptStruct* Struct = TBaseStructure<FRotator>::Get(); 
-							PinType = RigVMTypeUtils::PinTypeFromCPPType(Struct->GetStructCPPName(), Struct);
-							break;
-						}
-						case ERigControlType::Transform:
-						case ERigControlType::TransformNoScale:
-						case ERigControlType::EulerTransform:
-						default:
-						{
-							UScriptStruct* Struct = TBaseStructure<FTransform>::Get(); 
-							PinType = RigVMTypeUtils::PinTypeFromCPPType(Struct->GetStructCPPName(), Struct);
-							break;
-						}
-					}
-					const UControlRigGraphSchema* Schema = GetDefault<UControlRigGraphSchema>();
-					ShapeColor = Schema->GetPinTypeColor(PinType);
+					ShapeColor = GetColorForControlType(Control->Settings.ControlType, Control->Settings.ControlEnum);
 				}
 				
 				// ensure the alpha is always visible
@@ -822,6 +777,66 @@ TPair<const FSlateBrush*, FSlateColor> SRigHierarchyItem::GetBrushForElementType
 	return TPair<const FSlateBrush*, FSlateColor>(Brush, Color);
 }
 
+FLinearColor SRigHierarchyItem::GetColorForControlType(ERigControlType InControlType, UEnum* InControlEnum)
+{
+	FEdGraphPinType PinType;
+	switch(InControlType)
+	{
+		case ERigControlType::Bool:
+		{
+			PinType = RigVMTypeUtils::PinTypeFromCPPType(RigVMTypeUtils::BoolType, nullptr);
+			break;
+		}
+		case ERigControlType::Float:
+		{
+			PinType = RigVMTypeUtils::PinTypeFromCPPType(RigVMTypeUtils::FloatType, nullptr);
+			break;
+		}
+		case ERigControlType::Integer:
+		{
+			if(InControlEnum)
+			{
+				PinType = RigVMTypeUtils::PinTypeFromCPPType(FString(), InControlEnum);
+			}
+			else
+			{
+				PinType = RigVMTypeUtils::PinTypeFromCPPType(RigVMTypeUtils::Int32Type, nullptr);
+			}
+			break;
+		}
+		case ERigControlType::Vector2D:
+		{
+			UScriptStruct* Struct = TBaseStructure<FVector2D>::Get(); 
+			PinType = RigVMTypeUtils::PinTypeFromCPPType(Struct->GetStructCPPName(), Struct);
+			break;
+		}
+		case ERigControlType::Position:
+		case ERigControlType::Scale:
+		{
+			UScriptStruct* Struct = TBaseStructure<FVector>::Get(); 
+			PinType = RigVMTypeUtils::PinTypeFromCPPType(Struct->GetStructCPPName(), Struct);
+			break;
+		}
+		case ERigControlType::Rotator:
+		{
+			UScriptStruct* Struct = TBaseStructure<FRotator>::Get(); 
+			PinType = RigVMTypeUtils::PinTypeFromCPPType(Struct->GetStructCPPName(), Struct);
+			break;
+		}
+		case ERigControlType::Transform:
+		case ERigControlType::TransformNoScale:
+		case ERigControlType::EulerTransform:
+		default:
+		{
+			UScriptStruct* Struct = TBaseStructure<FTransform>::Get(); 
+			PinType = RigVMTypeUtils::PinTypeFromCPPType(Struct->GetStructCPPName(), Struct);
+			break;
+		}
+	}
+	const UControlRigGraphSchema* Schema = GetDefault<UControlRigGraphSchema>();
+	return Schema->GetPinTypeColor(PinType);
+}
+
 void SRigHierarchyItem::OnNameCommitted(const FText& InText, ETextCommit::Type InCommitType) const
 {
 	// for now only allow enter
diff --git a/Engine/Plugins/Animation/ControlRig/Source/ControlRigEditor/Private/Editor/SRigHierarchyTreeView.h b/Engine/Plugins/Animation/ControlRig/Source/ControlRigEditor/Private/Editor/SRigHierarchyTreeView.h
--- a/Engine/Plugins/Animation/ControlRig/Source/ControlRigEditor/Private/Editor/SRigHierarchyTreeView.h
+++ b/Engine/Plugins/Animation/ControlRig/Source/ControlRigEditor/Private/Editor/SRigHierarchyTreeView.h
@@ -171,7 +171,9 @@ public:
 public:
 	/** Element Data to display */
 	FRigElementKey Key;
+	FName ChannelName;
 	bool bIsTransient;
+	bool bIsAnimationChannel;
 	bool bSupportsRename;
 	TArray<TSharedPtr<FRigTreeElement>> Children;
 
@@ -206,6 +208,7 @@ public:
  	void OnNameCommitted(const FText& InText, ETextCommit::Type InCommitType) const;
 	bool OnVerifyNameChanged(const FText& InText, FText& OutErrorMessage);
 	static TPair<const FSlateBrush*, FSlateColor> GetBrushForElementType(const URigHierarchy* InHierarchy, const FRigElementKey& InKey);
+	static FLinearColor GetColorForControlType(ERigControlType InControlType, UEnum* InControlEnum);
 
 private:
 	TWeakPtr<FRigTreeElement> WeakRigTreeElement;
diff --git a/Engine/Plugins/Animation/ControlRig/Source/ControlRigEditor/Private/Graph/ControlRigGraphPanelPinFactory.cpp b/Engine/Plugins/Animation/ControlRig/Source/ControlRigEditor/Private/Graph/ControlRigGraphPanelPinFactory.cpp
--- a/Engine/Plugins/Animation/ControlRig/Source/ControlRigEditor/Private/Graph/ControlRigGraphPanelPinFactory.cpp
+++ b/Engine/Plugins/Animation/ControlRig/Source/ControlRigEditor/Private/Graph/ControlRigGraphPanelPinFactory.cpp
@@ -100,6 +100,103 @@ TSharedPtr<SGraphPin> FControlRigGraphPanelPinFactory::CreatePin(UEdGraphPin* In
 				{
 					return SNew(SControlRigGraphPinVariableName, InPin);
 				}
+				else if (CustomWidgetName == TEXT("AnimationChannelName"))
+				{
+					struct FCachedAnimationChannelNames
+					{
+						int32 TopologyVersion;
+						TSharedPtr<TArray<TSharedPtr<FString>>> Names;
+						
+						FCachedAnimationChannelNames()
+						: TopologyVersion(INDEX_NONE)
+						{}
+					};
+					
+					return SNew(SControlRigGraphPinNameList, InPin)
+						.ModelPin(ModelPin)
+						.OnGetNameListContent_Lambda([RigGraph](const URigVMPin* InPin)
+						{
+							if (const UControlRigBlueprint* Blueprint = RigGraph->GetTypedOuter<UControlRigBlueprint>())
+							{
+								FRigElementKey ControlKey;
+
+								// find the pin that holds the control
+								for(URigVMPin* Pin : InPin->GetRootPin()->GetNode()->GetPins())
+								{
+									if(Pin->GetCPPType() == RigVMTypeUtils::FNameType && Pin->GetCustomWidgetName() == TEXT("ControlName"))
+									{
+										const FString DefaultValue = Pin->GetDefaultValue();
+										const FName ControlName = DefaultValue.IsEmpty() ? FName(NAME_None) : FName(*DefaultValue);
+										ControlKey = FRigElementKey(ControlName, ERigElementType::Control);
+										break;
+									}
+
+									if(Pin->GetCPPType() == FRigElementKey::StaticStruct()->GetStructCPPName())
+									{
+										const FString DefaultValue = Pin->GetDefaultValue();
+										if(!DefaultValue.IsEmpty())
+										{
+											FRigElementKey::StaticStruct()->ImportText(*DefaultValue, &ControlKey, nullptr, EPropertyPortFlags::PPF_None, nullptr, FRigElementKey::StaticStruct()->GetName(), true);
+										}
+										break;
+									}
+								}
+
+								if(Blueprint->Hierarchy->Find<FRigControlElement>(ControlKey) == nullptr)
+								{
+									ControlKey.Reset();
+								}
+
+								const FString MapHash = ControlKey.IsValid() ?
+									Blueprint->GetPathName() + TEXT("|") + ControlKey.Name.ToString() :
+									FName(NAME_None).ToString();
+
+								static TMap<FString, FCachedAnimationChannelNames> ChannelNameLists;
+								FCachedAnimationChannelNames& ChannelNames = ChannelNameLists.FindOrAdd(MapHash);
+
+								bool bRefreshList = !ChannelNames.Names.IsValid();
+
+								if(!bRefreshList)
+								{
+									const int32 TopologyVersion = Blueprint->Hierarchy->GetTopologyVersion();
+									if(ChannelNames.TopologyVersion != TopologyVersion)
+									{
+										bRefreshList = true;
+										ChannelNames.TopologyVersion = TopologyVersion;
+									}
+								}
+
+								if(bRefreshList)
+								{
+									if(!ChannelNames.Names.IsValid())
+									{
+										ChannelNames.Names = MakeShareable(new TArray<TSharedPtr<FString>>());
+									}
+									ChannelNames.Names->Reset();
+									ChannelNames.Names->Add(MakeShareable(new FString(FName(NAME_None).ToString())));
+
+									if(const FRigControlElement* ControlElement = Blueprint->Hierarchy->Find<FRigControlElement>(ControlKey))
+									{
+										FRigBaseElementChildrenArray Children = Blueprint->Hierarchy->GetChildren(ControlElement);
+										for(const FRigBaseElement* Child : Children)
+										{
+											if(const FRigControlElement* ChildControl = Cast<FRigControlElement>(Child))
+											{
+												if(ChildControl->IsAnimationChannel())
+												{
+													ChannelNames.Names->Add(MakeShareable(new FString(ChildControl->GetDisplayName().ToString())));
+												}
+											}
+										}
+									}
+								}
+								return ChannelNames.Names.Get();
+							}
+
+							static TArray<TSharedPtr<FString>> EmptyNameList;
+							return &EmptyNameList;
+						});
+				}
 			}
 
 			if (InPin->PinType.PinCategory == UEdGraphSchema_K2::PC_Struct)
diff --git a/Engine/Plugins/Animation/ControlRig/Source/ControlRigEditor/Private/Sequencer/ControlRigParameterTrackEditor.cpp b/Engine/Plugins/Animation/ControlRig/Source/ControlRigEditor/Private/Sequencer/ControlRigParameterTrackEditor.cpp
--- a/Engine/Plugins/Animation/ControlRig/Source/ControlRigEditor/Private/Sequencer/ControlRigParameterTrackEditor.cpp
+++ b/Engine/Plugins/Animation/ControlRig/Source/ControlRigEditor/Private/Sequencer/ControlRigParameterTrackEditor.cpp
@@ -1999,7 +1999,7 @@ void FControlRigParameterTrackEditor::SelectRigsAndControls(UControlRig* Control
 					{
 						if(const FRigControlElement* ControlElement = Hierarchy->Find<FRigControlElement>(FRigElementKey(ControlName, ERigElementType::Control)))
 						{
-							if(ControlElement->Settings.AnimationType == ERigControlAnimationType::AnimationChannel)
+							if(ControlElement->IsAnimationChannel())
 							{
 								if(const FRigControlElement* ParentControlElement = Cast<FRigControlElement>(Hierarchy->GetFirstParent(ControlElement)))
 								{
diff --git a/Engine/Source/Developer/AnimationWidgets/Public/SAdvancedTransformInputBox.h b/Engine/Source/Developer/AnimationWidgets/Public/SAdvancedTransformInputBox.h
--- a/Engine/Source/Developer/AnimationWidgets/Public/SAdvancedTransformInputBox.h
+++ b/Engine/Source/Developer/AnimationWidgets/Public/SAdvancedTransformInputBox.h
@@ -1106,11 +1106,12 @@ public:
 	}
 
 	template<typename BuilderType = IDetailCategoryBuilder>
-	static void ConstructGroupedTransformRows(
+	static FDetailWidgetRow& ConstructGroupedTransformRows(
 		BuilderType& InBuilder,
 		const FText& InLabel,
 		const FText& InTooltip,
-		typename SAdvancedTransformInputBox<TransformType>::FArguments WidgetArgs
+		typename SAdvancedTransformInputBox<TransformType>::FArguments WidgetArgs,
+		TSharedPtr<SWidget> NameContent = TSharedPtr<SWidget>()
 		)
 	{
 		auto ConstructComponentWidgetRow = [WidgetArgs](FDetailWidgetRow& WidgetRow, ESlateTransformComponent::Type InComponent)
@@ -1180,14 +1181,19 @@ public:
 		FDetailWidgetRow& HeaderRow = Group.HeaderRow();
 		ConstructComponentWidgetRow(HeaderRow, ESlateTransformComponent::Max);
 
+		if(!NameContent.IsValid())
+		{
+			SAssignNew(NameContent, STextBlock)
+			.Font(IDetailLayoutBuilder::GetDetailFont())
+			.Text(InLabel)
+			.ToolTipText(InTooltip);
+		}
+		
 		HeaderRow
 		.Visibility(WidgetArgs._Visibility)
 		.NameContent()
 		[
-			SNew(STextBlock)
-			.Font(IDetailLayoutBuilder::GetDetailFont())
-			.Text(InLabel)
-			.ToolTipText(InTooltip)
+			NameContent.ToSharedRef()
 		];
 
 		if(WidgetArgs._ConstructLocation)
@@ -1207,6 +1213,8 @@ public:
 			FDetailWidgetRow& WidgetRow = Group.AddWidgetRow();
 			ConstructComponentWidgetRow(WidgetRow, ESlateTransformComponent::Scale);
 		}
+
+		return HeaderRow;
 	}
 
 #endif
