Control Rig: Fix overlapping usage of Ctrl-N

diff --git a/Engine/Plugins/Animation/ControlRig/Source/ControlRigEditor/Private/ControlRigHierarchyCommands.cpp b/Engine/Plugins/Animation/ControlRig/Source/ControlRigEditor/Private/ControlRigHierarchyCommands.cpp
--- a/Engine/Plugins/Animation/ControlRig/Source/ControlRigEditor/Private/ControlRigHierarchyCommands.cpp
+++ b/Engine/Plugins/Animation/ControlRig/Source/ControlRigEditor/Private/ControlRigHierarchyCommands.cpp
@@ -8,7 +8,7 @@ void FControlRigHierarchyCommands::RegisterCommands()
 {
 	UI_COMMAND(AddBoneItem, "New Bone", "Add new bone at the origin (0, 0, 0) to the hierarchy.", EUserInterfaceActionType::Button, FInputChord());
 	UI_COMMAND(AddControlItem, "New Control", "Add new control at the origin (0, 0, 0) to the hierarchy.", EUserInterfaceActionType::Button, FInputChord(EKeys::N, EModifierKey::Control));
-	UI_COMMAND(AddAnimationChannelItem, "New Animation Channel", "Add new animation channel.", EUserInterfaceActionType::Button, FInputChord(EKeys::N, EModifierKey::Control));
+	UI_COMMAND(AddAnimationChannelItem, "New Animation Channel", "Add new animation channel.", EUserInterfaceActionType::Button, FInputChord());
 	UI_COMMAND(AddNullItem, "New Null", "Add new null at the origin (0, 0, 0) to the hierarchy.", EUserInterfaceActionType::Button, FInputChord());
 	UI_COMMAND(DuplicateItem, "Duplicate", "Duplicate the selected items in the hierarchy.", EUserInterfaceActionType::Button, FInputChord(EKeys::D, EModifierKey::Control));
 	UI_COMMAND(MirrorItem, "Mirror", "Mirror the selected items in the hierarchy.", EUserInterfaceActionType::Button, FInputChord());
