Fixed text conflict

diff --git a/Engine/Plugins/Animation/ControlRig/Source/ControlRigEditor/Private/ControlRigElementDetails.cpp b/Engine/Plugins/Animation/ControlRig/Source/ControlRigEditor/Private/ControlRigElementDetails.cpp
--- a/Engine/Plugins/Animation/ControlRig/Source/ControlRigEditor/Private/ControlRigElementDetails.cpp
+++ b/Engine/Plugins/Animation/ControlRig/Source/ControlRigEditor/Private/ControlRigElementDetails.cpp
@@ -3809,7 +3809,7 @@ void FRigControlElementDetails::SetDisplayNameForElement(const FText& InNewText,
 
 	if(HierarchyBeingCustomized)
 	{
-		FScopedTransaction Transaction(LOCTEXT("SetDisplayName", "Set Display Name"));
+		FScopedTransaction Transaction(LOCTEXT("SetElementDisplayName", "Set Element Display Name"));
 		HierarchyBeingCustomized->GetController(true)->SetDisplayName(InKeyToRename, DisplayName, true, true);
 	}
 }
