fix tangent update when switching space
(the last change, due to static analysis failing, changed the intended behaviour. cf. UE-144801)

diff --git a/Engine/Plugins/Animation/ControlRig/Source/ControlRigEditor/Private/Sequencer/ControlRigSpaceChannelEditors.cpp b/Engine/Plugins/Animation/ControlRig/Source/ControlRigEditor/Private/Sequencer/ControlRigSpaceChannelEditors.cpp
--- a/Engine/Plugins/Animation/ControlRig/Source/ControlRigEditor/Private/Sequencer/ControlRigSpaceChannelEditors.cpp
+++ b/Engine/Plugins/Animation/ControlRig/Source/ControlRigEditor/Private/Sequencer/ControlRigSpaceChannelEditors.cpp
@@ -1526,10 +1526,11 @@ void FControlRigSpaceChannelHelpers::SetTangentsAtThisTime(	UControlRig* Control
 	for (int32 Index = 0; Index < NumChannels; ++Index, ++ChannelIndex)
 	{
 		TMovieSceneChannelData<FMovieSceneFloatValue> ChannelInterface = FloatChannels[ChannelIndex]->GetData();
+		TArrayView<FMovieSceneFloatValue> Values = ChannelInterface.GetValues();
 		const int32 KeyIndex = ChannelInterface.FindKey(Time);
 		if (KeyIndex != INDEX_NONE)
 		{
-			FMovieSceneFloatValue Value = ChannelInterface.GetValues()[KeyIndex];
+			FMovieSceneFloatValue& Value = Values[KeyIndex];
 			Value.Tangent.ArriveTangent = InTangents[Index].ArriveTangent;
 			Value.Tangent.LeaveTangent = InTangents[Index].LeaveTangent;
 			Value.TangentMode = RCTM_Break;	
