Control Rig: Change offset transform when dragdropping Controls

diff --git a/Engine/Plugins/Animation/ControlRig/Source/ControlRigEditor/Private/Editor/SRigHierarchy.cpp b/Engine/Plugins/Animation/ControlRig/Source/ControlRigEditor/Private/Editor/SRigHierarchy.cpp
--- a/Engine/Plugins/Animation/ControlRig/Source/ControlRigEditor/Private/Editor/SRigHierarchy.cpp
+++ b/Engine/Plugins/Animation/ControlRig/Source/ControlRigEditor/Private/Editor/SRigHierarchy.cpp
@@ -2655,8 +2655,11 @@ FReply SRigHierarchy::ReparentOrMatchTransform(const TArray<FRigElementKey>& Dra
 
 			FRigElementKey ParentKey = TargetKey;
 
-			FTransform InitialTransform = DebuggedHierarchy->GetInitialGlobalTransform(DraggedKey);
-			FTransform GlobalTransform = DebuggedHierarchy->GetGlobalTransform(DraggedKey);
+			const FTransform InitialGlobalTransform = DebuggedHierarchy->GetInitialGlobalTransform(DraggedKey);
+			const FTransform CurrentGlobalTransform = DebuggedHierarchy->GetGlobalTransform(DraggedKey);
+			const FTransform InitialLocalTransform = DebuggedHierarchy->GetInitialLocalTransform(DraggedKey);
+			const FTransform CurrentLocalTransform = DebuggedHierarchy->GetLocalTransform(DraggedKey);
+			const FTransform CurrentGlobalOffsetTransform = DebuggedHierarchy->GetGlobalControlOffsetTransform(DraggedKey, false);
 
 			if(ParentKey.IsValid())
 			{
@@ -2667,10 +2670,33 @@ FReply SRigHierarchy::ReparentOrMatchTransform(const TArray<FRigElementKey>& Dra
 				Controller->RemoveAllParents(DraggedKey, true, true, true);
 			}
 
-			DebuggedHierarchy->SetInitialGlobalTransform(DraggedKey, InitialTransform, true, true);
-			DebuggedHierarchy->SetGlobalTransform(DraggedKey, GlobalTransform, false, true, true);
-			Hierarchy->SetInitialGlobalTransform(DraggedKey, InitialTransform, true, true);
-			Hierarchy->SetGlobalTransform(DraggedKey, GlobalTransform, false, true, true);
+			if (DraggedKey.Type == ERigElementType::Control)
+			{
+				int32 ControlIndex = DebuggedHierarchy->GetIndex(DraggedKey);
+				if (ControlIndex == INDEX_NONE)
+				{
+					continue;
+				}
+
+				const FTransform GlobalParentTransform = DebuggedHierarchy->GetGlobalTransform(ParentKey, false);
+				const FTransform LocalOffsetTransform = CurrentGlobalOffsetTransform.GetRelativeTransform(GlobalParentTransform);
+
+				Hierarchy->SetControlOffsetTransformByIndex(ControlIndex, LocalOffsetTransform, ERigTransformType::InitialLocal, true, true, true);
+				Hierarchy->SetControlOffsetTransformByIndex(ControlIndex, LocalOffsetTransform, ERigTransformType::CurrentLocal, true, true, true);
+				Hierarchy->SetLocalTransform(DraggedKey, CurrentLocalTransform, true, true, true, true);
+				Hierarchy->SetInitialLocalTransform(DraggedKey, InitialLocalTransform, true, true, true);
+				DebuggedHierarchy->SetControlOffsetTransformByIndex(ControlIndex, LocalOffsetTransform, ERigTransformType::InitialLocal, true, true);
+				DebuggedHierarchy->SetControlOffsetTransformByIndex(ControlIndex, LocalOffsetTransform, ERigTransformType::CurrentLocal, true, true);
+				DebuggedHierarchy->SetLocalTransform(DraggedKey, CurrentLocalTransform, true, true, true);
+				DebuggedHierarchy->SetInitialLocalTransform(DraggedKey, InitialLocalTransform, true, true);
+			}
+			else
+			{
+				DebuggedHierarchy->SetInitialGlobalTransform(DraggedKey, InitialGlobalTransform, true, true);
+				DebuggedHierarchy->SetGlobalTransform(DraggedKey, CurrentGlobalTransform, false, true, true);
+				Hierarchy->SetInitialGlobalTransform(DraggedKey, InitialGlobalTransform, true, true);
+				Hierarchy->SetGlobalTransform(DraggedKey, CurrentGlobalTransform, false, true, true);
+			}
 		}
 	}
 
