Control Rig: Restructuring execute - getting rid of bSetupModeEnabled

diff --git a/Engine/Plugins/Animation/ControlRig/Source/ControlRig/Private/ControlRig.cpp b/Engine/Plugins/Animation/ControlRig/Source/ControlRig/Private/ControlRig.cpp
--- a/Engine/Plugins/Animation/ControlRig/Source/ControlRig/Private/ControlRig.cpp
+++ b/Engine/Plugins/Animation/ControlRig/Source/ControlRig/Private/ControlRig.cpp
@@ -74,7 +74,6 @@ UControlRig::UControlRig(const FObjectInitializer& ObjectInitializer)
 #endif
 	, bRequiresInitExecution(false)
 	, bRequiresSetupEvent(false)
-	, bSetupModeEnabled(false)
 	, bCopyHierarchyBeforeSetup(true)
 	, bResetInitialTransformsBeforeSetup(true)
 	, bManipulationEnabled(false)
@@ -279,7 +278,7 @@ void UControlRig::Evaluate_AnyThread()
 	// create a copy since we need to change it here temporarily,
 	// and UI / the rig may change the event queue while it is running
 	TArray<FName> EventQueueToRun = EventQueue;
-	
+
 	if(InteractionType != (uint8)EControlRigInteractionType::None)
 	{
 		if(EventQueueToRun.IsEmpty())
@@ -294,6 +293,14 @@ void UControlRig::Evaluate_AnyThread()
 		}
 	}
 
+	// execute the setup event prior to everything else
+	if(bRequiresSetupEvent)
+	{
+		if(!EventQueueToRun.Contains(FRigUnit_PrepareForExecution::EventName))
+		{
+			EventQueueToRun.Insert(FRigUnit_PrepareForExecution::EventName, 0);
+		}
+	}
 	
 	for (const FName& EventName : EventQueueToRun)
 	{
@@ -308,7 +315,6 @@ void UControlRig::Evaluate_AnyThread()
 			}
 		}
 #endif
-			
 	}
 }
 
@@ -609,16 +615,21 @@ void UControlRig::CopyExternalVariableDefaultValuesFromCDO()
 	}
 }
 
-void UControlRig::Execute(const EControlRigState InState, const FName& InEventName)
+bool UControlRig::Execute(const EControlRigState InState, const FName& InEventName)
 {
 	if(!CanExecute())
 	{
-		return;
+		return false;
 	}
 	
 	const bool bIsEventInQueue = GetEventQueue().Contains(InEventName);
 	const bool bIsEventFirstInQueue = !GetEventQueue().IsEmpty() && GetEventQueue()[0] == InEventName; 
-	const bool bIsEventLastInQueue = !GetEventQueue().IsEmpty() && GetEventQueue().Last() == InEventName; 
+	const bool bIsEventLastInQueue = !GetEventQueue().IsEmpty() && GetEventQueue().Last() == InEventName;
+	const bool bIsInitializingMemory = InState == EControlRigState::Init;
+	const bool bIsExecutingInstructions = InState == EControlRigState::Update;
+	const bool bIsSetupEvent = InEventName == FRigUnit_PrepareForExecution::EventName;
+	const bool bIsForwardSolve = InEventName == FRigUnit_BeginExecution::EventName;
+	const bool bIsInteractionEvent = InEventName == FRigUnit_InteractionExecution::EventName;
 
 	ensure(!HasAnyFlags(RF_ClassDefaultObject));
 	
@@ -634,7 +645,7 @@ void UControlRig::Execute(const EControlRigState InState, const FName& InEventNa
 			InstantiateVMFromCDO();
 		}
 
-		if (InState == EControlRigState::Init)
+		if (bIsInitializingMemory)
 		{
 			VM->ClearExternalVariables();
 
@@ -703,21 +714,27 @@ void UControlRig::Execute(const EControlRigState InState, const FName& InEventNa
 	{
 		bRequiresInitExecution = false;
 
-		if (InState != EControlRigState::Init)
+		if (!bIsInitializingMemory)
 		{
-			Execute(EControlRigState::Init, InEventName);
+			// if init is required we'll run init on the whole bytecode
+			if(!Execute(EControlRigState::Init, NAME_None))
+			{
+				return false;
+			}
 			bJustRanInit = true;
 		}
 	}
 
 	FRigUnitContext Context;
 
+	// setup the draw interface for debug drawing
 	if(!bIsEventInQueue || bIsEventFirstInQueue)
 	{
 		DrawInterface.Reset();
 	}
 	Context.DrawInterface = &DrawInterface;
 
+	// setup the animation attribute container
 	Context.AnimAttributeContainer = ExternalAnimAttributeContainer;
 
 	// draw container contains persistent draw instructions, 
@@ -727,29 +744,34 @@ void UControlRig::Execute(const EControlRigState InState, const FName& InEventNa
 	{
 		DrawContainer = CDO->DrawContainer;
 	}
-
 	Context.DrawContainer = &DrawContainer;
+
+	// setup the data source registry
 	Context.DataSourceRegistry = GetDataSourceRegistry();
 
-	if (InState == EControlRigState::Init)
+	// reset the time and caches during init
+	if (bIsInitializingMemory)
 	{
 		AbsoluteTime = DeltaTime = 0.f;
 		NameCache.Reset();
 	}
 
+	// setup the context with further fields
 	Context.DeltaTime = DeltaTime;
 	Context.AbsoluteTime = AbsoluteTime;
 	Context.FramesPerSecond = GetCurrentFramesPerSecond();
 	Context.InteractionType = InteractionType;
 	Context.ElementsBeingInteracted = ElementsBeingInteracted;
 	Context.State = InState;
-	
+
+	// allow access to the hierarchy
 	Context.Hierarchy = GetHierarchy();
 	check(Context.Hierarchy);
 
 	// for now we'll disable any controller access during the run of a control rig
 	FRigHierarchyEnableControllerBracket DisableHierarchyController(Context.Hierarchy, false);
 
+	// setup the context with further fields
 	Context.ToWorldSpaceTransform = FTransform::Identity;
 	Context.OwningComponent = nullptr;
 	Context.OwningActor = nullptr;
@@ -781,6 +803,8 @@ void UControlRig::Execute(const EControlRigState InState, const FName& InEventNa
 		}
 	}
 
+	// given the outer scene component configure
+	// the transform lookups to map transforms from rig space to world space
 	if (OuterSceneComponent.IsValid())
 	{
 		Context.ToWorldSpaceTransform = OuterSceneComponent->GetComponentToWorld();
@@ -813,13 +837,14 @@ void UControlRig::Execute(const EControlRigState InState, const FName& InEventNa
 		}
 	}
 
+	// if we have any referenced elements dirty them
 	if(GetHierarchy())
 	{
-		// if we have any referenced elements dirty them
 		GetHierarchy()->UpdateReferences(&Context);
 	}
 
 #if WITH_EDITOR
+	// setup the log and VM settings
 	Context.Log = ControlRigLog;
 	if (ControlRigLog != nullptr)
 	{
@@ -828,75 +853,64 @@ void UControlRig::Execute(const EControlRigState InState, const FName& InEventNa
 	}
 #endif
 
-	// execute units
-	if (bRequiresSetupEvent && InState != EControlRigState::Init)
+	// guard against recursion
+	if(IsExecuting())
 	{
-		if(!IsRunningPreSetup() && !IsRunningPostSetup())
+		UE_LOG(LogControlRig, Warning, TEXT("%s: Execute is being called recursively."), *GetPathName());
+		return false;
+	}
+	if(bIsSetupEvent)
+	{
+		if(IsRunningPreSetup() || IsRunningPostSetup())
 		{
-			bRequiresSetupEvent = bSetupModeEnabled;
+			UE_LOG(LogControlRig, Warning, TEXT("%s: Setup is being called recursively."), *GetPathName());
+			return false;
+		}
+	}
+
+	bool bSuccess = true;
+
+	// we'll special case the setup event here
+	if (bIsSetupEvent && !bIsInitializingMemory)
+	{
+		// construction mode means that we are running the construction event
+		// constantly for testing purposes.
+		const bool bSetupModeEnabled = IsSetupModeEnabled();
+		{
+			// save the current state of all pose elements to preserve user intention, since setup event can
+			// run in between forward events
+			// the saved pose is reapplied to the rig after setup event as the pose scope goes out of scope
+			TUniquePtr<UControlRig::FPoseScope> PoseScope;
+			if (!bSetupModeEnabled)
 			{
-				// save the current state of all pose elements to preserve user intention, since setup event can
-				// run in between forward events
-				// the saved pose is reapplied to the rig after setup event as the pose scope goes out of scope
-				TUniquePtr<UControlRig::FPoseScope> PoseScope;
-				if (!bSetupModeEnabled)
-				{
-					// only do this in non-setup mode because 
-					// when setup mode is enabled, the control values are cleared before reaching here (too late to save them)
-					PoseScope = MakeUnique<UControlRig::FPoseScope>(this, ERigElementType::ToResetAfterSetupEvent);
-				}
-				
-				if (UControlRig* CDO = Cast<UControlRig>(GetClass()->GetDefaultObject()))
+				// only do this in non-setup mode because 
+				// when setup mode is enabled, the control values are cleared before reaching here (too late to save them)
+				PoseScope = MakeUnique<UControlRig::FPoseScope>(this, ERigElementType::ToResetAfterSetupEvent);
+			}
+
+			// Copy the hierarchy from the default object onto this one
+			if (UControlRig* CDO = Cast<UControlRig>(GetClass()->GetDefaultObject()))
+			{
+				if(bCopyHierarchyBeforeSetup && !bSetupModeEnabled)
 				{
-					if(bCopyHierarchyBeforeSetup && !bSetupModeEnabled)
+					if(CDO->GetHierarchy()->GetTopologyVersion()!= GetHierarchy()->GetTopologyVersion())
 					{
-						if(CDO->GetHierarchy()->GetTopologyVersion()!= GetHierarchy()->GetTopologyVersion())
-						{
 #if WITH_EDITOR
-							FTransientControlScope TransientControlScope(GetHierarchy());
+						FTransientControlScope TransientControlScope(GetHierarchy());
 #endif
-							GetHierarchy()->CopyHierarchy(CDO->GetHierarchy());
-						}
-					}
-					
-					if (bResetInitialTransformsBeforeSetup && !bSetupModeEnabled)
-					{
-						// when copying the pose from the default we don't copy the weights over,
-						// since we don't want to reset the topology in case it has diverged from the
-						// CDO.
-						GetHierarchy()->CopyPose(CDO->GetHierarchy(), false, true, false);
+						GetHierarchy()->CopyHierarchy(CDO->GetHierarchy());
 					}
 				}
-
-				{
-#if WITH_EDITOR
-					TUniquePtr<FTransientControlPoseScope> TransientControlPoseScope;
-					if (bSetupModeEnabled)
-					{
-						// save the transient control value, it should not be constantly reset in setup mode
-						TransientControlPoseScope = MakeUnique<FTransientControlPoseScope>(this);
-					}
-#endif
-					// reset the pose to initial such that setup event can run from a deterministic initial state
-					GetHierarchy()->ResetPoseToInitial(ERigElementType::All);
-				}
-
-				if (PreSetupEvent.IsBound())
-				{
-					FControlRigBracketScope BracketScope(PreSetupBracket);
-					PreSetupEvent.Broadcast(this, EControlRigState::Update, FRigUnit_PrepareForExecution::EventName);
-				}
-
-				ExecuteUnits(Context, FRigUnit_PrepareForExecution::EventName);
-
-				if (PostSetupEvent.IsBound())
+				
+				if (bResetInitialTransformsBeforeSetup && !bSetupModeEnabled)
 				{
-					FControlRigBracketScope BracketScope(PostSetupBracket);
-					PostSetupEvent.Broadcast(this, EControlRigState::Update, FRigUnit_PrepareForExecution::EventName);
+					// when copying the pose from the default we don't copy the weights over,
+					// since we don't want to reset the topology in case it has diverged from the
+					// CDO.
+					GetHierarchy()->CopyPose(CDO->GetHierarchy(), false, true, false);
 				}
 			}
 
-			if (bSetupModeEnabled)
 			{
 #if WITH_EDITOR
 				TUniquePtr<FTransientControlPoseScope> TransientControlPoseScope;
@@ -906,108 +920,124 @@ void UControlRig::Execute(const EControlRigState InState, const FName& InEventNa
 					TransientControlPoseScope = MakeUnique<FTransientControlPoseScope>(this);
 				}
 #endif
-				GetHierarchy()->ResetPoseToInitial(ERigElementType::Bone);
-			}			
-		}
-		else
-		{
-			UE_LOG(LogControlRig, Warning, TEXT("%s: Setup is being called recursively."), *GetPathName());
-		}
-	}
+				// reset the pose to initial such that setup event can run from a deterministic initial state
+				GetHierarchy()->ResetPoseToInitial(ERigElementType::All);
+			}
 
-	if (!bSetupModeEnabled)
-	{
-		if(!IsExecuting())
-		{ 
+			if (PreSetupEvent.IsBound())
+			{
+				FControlRigBracketScope BracketScope(PreSetupBracket);
+				PreSetupEvent.Broadcast(this, EControlRigState::Update, FRigUnit_PrepareForExecution::EventName);
+			}
 
-#if WITH_EDITOR
-			// only set a valid first entry event when execution
-			// has passed the initialization stage and there are multiple events present
-			if (EventQueue.Num() >= 2 && VM && InState != EControlRigState::Init)
+			bSuccess = ExecuteUnits(Context, FRigUnit_PrepareForExecution::EventName);
+
+			if (PostSetupEvent.IsBound())
 			{
-				VM->SetFirstEntryEventInEventQueue(EventQueue[0]);
+				FControlRigBracketScope BracketScope(PostSetupBracket);
+				PostSetupEvent.Broadcast(this, EControlRigState::Update, FRigUnit_PrepareForExecution::EventName);
 			}
+		}
 
-			// Transform Overrride is generated using a Transient Control 
-			ApplyTransformOverrideForUserCreatedBones();
-#endif
-			
+		if (bSetupModeEnabled)
+		{
 #if WITH_EDITOR
-			if (bEnableAnimAttributeTrace && ExternalAnimAttributeContainer != nullptr)
+			TUniquePtr<FTransientControlPoseScope> TransientControlPoseScope;
+			if (bSetupModeEnabled)
 			{
-				InputAnimAttributeSnapshot.CopyFrom(*ExternalAnimAttributeContainer);
+				// save the transient control value, it should not be constantly reset in setup mode
+				TransientControlPoseScope = MakeUnique<FTransientControlPoseScope>(this);
 			}
 #endif
-			
-			if (InState == EControlRigState::Update && InEventName == FRigUnit_BeginExecution::EventName)
+			GetHierarchy()->ResetPoseToInitial(ERigElementType::Bone);
+		}			
+	}
+	else
+	{
+#if WITH_EDITOR
+		// only set a valid first entry event when execution
+		// has passed the initialization stage and there are multiple events present
+		if (EventQueue.Num() >= 2 && VM && !bIsInitializingMemory)
+		{
+			VM->SetFirstEntryEventInEventQueue(EventQueue[0]);
+		}
+
+		// Transform Overrride is generated using a Transient Control 
+		ApplyTransformOverrideForUserCreatedBones();
+
+		if (bEnableAnimAttributeTrace && ExternalAnimAttributeContainer != nullptr)
+		{
+			InputAnimAttributeSnapshot.CopyFrom(*ExternalAnimAttributeContainer);
+		}
+#endif
+		
+		if (bIsExecutingInstructions && bIsForwardSolve)
+		{
+			if (PreForwardsSolveEvent.IsBound())
 			{
-				if (PreForwardsSolveEvent.IsBound())
-				{
-					FControlRigBracketScope BracketScope(PreForwardsSolveBracket);
-					PreForwardsSolveEvent.Broadcast(this, EControlRigState::Update, FRigUnit_BeginExecution::EventName);
-				}
+				FControlRigBracketScope BracketScope(PreForwardsSolveBracket);
+				PreForwardsSolveEvent.Broadcast(this, EControlRigState::Update, FRigUnit_BeginExecution::EventName);
 			}
+		}
 
-			ExecuteUnits(Context, InEventName);
+		bSuccess = ExecuteUnits(Context, InEventName);
 
 #if WITH_EDITOR
-			if (bEnableAnimAttributeTrace && ExternalAnimAttributeContainer != nullptr)
-			{
-				OutputAnimAttributeSnapshot.CopyFrom(*ExternalAnimAttributeContainer);
-			}
+		if (bEnableAnimAttributeTrace && ExternalAnimAttributeContainer != nullptr)
+		{
+			OutputAnimAttributeSnapshot.CopyFrom(*ExternalAnimAttributeContainer);
+		}
 #endif
 
-			if (InState == EControlRigState::Update && InEventName == FRigUnit_BeginExecution::EventName)
-			{
-				if (PostForwardsSolveEvent.IsBound())
-				{
-					FControlRigBracketScope BracketScope(PostForwardsSolveBracket);
-					PostForwardsSolveEvent.Broadcast(this, EControlRigState::Update, FRigUnit_BeginExecution::EventName);
-				}
-			}
-			
-			if (InState == EControlRigState::Init && InEventName != FRigUnit_BeginExecution::EventName)
+		if (bIsExecutingInstructions && bIsForwardSolve)
+		{
+			if (PostForwardsSolveEvent.IsBound())
 			{
-				ExecuteUnits(Context, FRigUnit_BeginExecution::EventName);
+				FControlRigBracketScope BracketScope(PostForwardsSolveBracket);
+				PostForwardsSolveEvent.Broadcast(this, EControlRigState::Update, FRigUnit_BeginExecution::EventName);
 			}
 		}
-		else
+		
+		if (bIsInitializingMemory && !bIsForwardSolve)
 		{
-			UE_LOG(LogControlRig, Warning, TEXT("%s: Update is being called recursively."), *GetPathName());
+			bSuccess = ExecuteUnits(Context, FRigUnit_BeginExecution::EventName);
 		}
 	}
 
 #if WITH_EDITOR
-	if (ControlRigLog != nullptr && bEnableControlRigLogging && InState != EControlRigState::Init && !bJustRanInit)
+
+	// for the last event in the queue - clear the log message queue
+	if (ControlRigLog != nullptr && bEnableControlRigLogging && !bIsInitializingMemory)
 	{
-		for (const FControlRigLog::FLogEntry& Entry : ControlRigLog->Entries)
+		if (bJustRanInit)
+		{
+			ControlRigLog->KnownMessages.Reset();
+			LoggedMessages.Reset();
+		}
+		else if(bIsEventLastInQueue)
 		{
-			if (Entry.FunctionName == NAME_None || Entry.InstructionIndex == INDEX_NONE || Entry.Message.IsEmpty())
+			for (const FControlRigLog::FLogEntry& Entry : ControlRigLog->Entries)
 			{
-				continue;
-			}
+				if (Entry.FunctionName == NAME_None || Entry.InstructionIndex == INDEX_NONE || Entry.Message.IsEmpty())
+				{
+					continue;
+				}
 
-			FString PerInstructionMessage = 
-				FString::Printf(
-					TEXT("Instruction[%d] '%s': '%s'"),
-					Entry.InstructionIndex,
-					*Entry.FunctionName.ToString(),
-					*Entry.Message
-				);
+				FString PerInstructionMessage = 
+					FString::Printf(
+						TEXT("Instruction[%d] '%s': '%s'"),
+						Entry.InstructionIndex,
+						*Entry.FunctionName.ToString(),
+						*Entry.Message
+					);
 
-			LogOnce(Entry.Severity, Entry.InstructionIndex, PerInstructionMessage);
+				LogOnce(Entry.Severity, Entry.InstructionIndex, PerInstructionMessage);
+			}
 		}
 	}
-
-	if (bJustRanInit && ControlRigLog != nullptr)
-	{
-		ControlRigLog->KnownMessages.Reset();
-		LoggedMessages.Reset();
-	}
 #endif
 
-
-	if (InState == EControlRigState::Init)
+	if (bIsInitializingMemory)
 	{
 		if (InitializedEvent.IsBound())
 		{
@@ -1015,7 +1045,7 @@ void UControlRig::Execute(const EControlRigState InState, const FName& InEventNa
 			InitializedEvent.Broadcast(this, EControlRigState::Init, InEventName);
 		}
 	}
-	else if (InState == EControlRigState::Update)
+	else if (bIsExecutingInstructions)
 	{
 		if(!bIsEventInQueue || bIsEventLastInQueue) 
 		{
@@ -1039,7 +1069,7 @@ void UControlRig::Execute(const EControlRigState InState, const FName& InEventNa
 		HandleHierarchyEvent(GetHierarchy(), EventContext);
 	}
 
-	if (Context.DrawInterface && Context.DrawContainer && bIsEventLastInQueue) 
+	if (Context.DrawInterface && Context.DrawContainer && bIsEventLastInQueue && bIsExecutingInstructions) 
 	{
 		Context.DrawInterface->Instructions.Append(Context.DrawContainer->Instructions);
 
@@ -1268,13 +1298,23 @@ void UControlRig::Execute(const EControlRigState InState, const FName& InEventNa
 		});
 	}
 
-	if(InEventName == FRigUnit_InteractionExecution::EventName)
+	if(bIsInteractionEvent)
 	{
 		bInteractionJustBegan = false;
 	}
+
+	if(bSuccess)
+	{
+		if(bIsSetupEvent)
+		{
+			bRequiresSetupEvent = false;
+		}
+	}
+
+	return bSuccess;
 }
 
-void UControlRig::ExecuteUnits(FRigUnitContext& InOutContext, const FName& InEventName)
+bool UControlRig::ExecuteUnits(FRigUnitContext& InOutContext, const FName& InEventName)
 {
 	if (VM)
 	{
@@ -1285,7 +1325,7 @@ void UControlRig::ExecuteUnits(FRigUnitContext& InOutContext, const FName& InEve
 				!VM->IsValidLowLevel())
 			{
 				UE_LOG(LogControlRig, Warning, InvalidatedVMFormat, *GetClass()->GetName());
-				return;
+				return false;
 			}
 		}
 		else
@@ -1297,7 +1337,7 @@ void UControlRig::ExecuteUnits(FRigUnitContext& InOutContext, const FName& InEve
 				!VM->GetWorkMemory()->IsValidLowLevel())
 			{
 				UE_LOG(LogControlRig, Warning, InvalidatedVMFormat, *GetClass()->GetName());
-				return;
+				return false;
 			}
 		}
 		
@@ -1315,6 +1355,8 @@ void UControlRig::ExecuteUnits(FRigUnitContext& InOutContext, const FName& InEve
 		TArray<void*> AdditionalArguments;
 		AdditionalArguments.Add(&InOutContext);
 
+		bool bSuccess = true;
+
 		if (InOutContext.State == EControlRigState::Init)
 		{
 			if(IsInGameThread() && bUseSnapshots)
@@ -1361,7 +1403,7 @@ void UControlRig::ExecuteUnits(FRigUnitContext& InOutContext, const FName& InEve
 			}
 			else
 			{
-				VM->Initialize(LocalMemory, AdditionalArguments);
+				bSuccess = VM->Initialize(LocalMemory, AdditionalArguments);
 			}
 		}
 		else
@@ -1405,7 +1447,7 @@ void UControlRig::ExecuteUnits(FRigUnitContext& InOutContext, const FName& InEve
 #endif
 			FRigHierarchyExecuteContextBracket HierarchyContextGuard(Hierarchy, &VM->GetContext());
 
-			VM->Execute(LocalMemory, AdditionalArguments, InEventName);
+			bSuccess = VM->Execute(LocalMemory, AdditionalArguments, InEventName);
 		}
 
 #if UE_CONTROLRIG_PROFILE_EXECUTE_UNITS_NUM
@@ -1420,7 +1462,9 @@ void UControlRig::ExecuteUnits(FRigUnitContext& InOutContext, const FName& InEve
 		}
 #endif
 
+		return bSuccess;
 	}
+	return false;
 }
 
 bool UControlRig::ContainsEvent(const FName& InEventName) const
@@ -1701,6 +1745,11 @@ FRigControlElement* UControlRig::FindControl(const FName& InControlName) const
 	return DynamicHierarchy->Find<FRigControlElement>(FRigElementKey(InControlName, ERigElementType::Control));
 }
 
+bool UControlRig::IsSetupModeEnabled() const
+{
+	return EventQueue.Contains(FRigUnit_PrepareForExecution::EventName);
+}
+
 FTransform UControlRig::SetupControlFromGlobalTransform(const FName& InControlName, const FTransform& InGlobalTransform)
 {
 	if (IsSetupModeEnabled())
diff --git a/Engine/Plugins/Animation/ControlRig/Source/ControlRig/Private/Rigs/AdditiveControlRig.cpp b/Engine/Plugins/Animation/ControlRig/Source/ControlRig/Private/Rigs/AdditiveControlRig.cpp
--- a/Engine/Plugins/Animation/ControlRig/Source/ControlRig/Private/Rigs/AdditiveControlRig.cpp
+++ b/Engine/Plugins/Animation/ControlRig/Source/ControlRig/Private/Rigs/AdditiveControlRig.cpp
@@ -39,11 +39,13 @@ FName UAdditiveControlRig::GetNullName(const FName& InBoneName)
 	return NAME_None;
 }
 
-void UAdditiveControlRig::ExecuteUnits(FRigUnitContext& InOutContext, const FName& InEventName)
+bool UAdditiveControlRig::ExecuteUnits(FRigUnitContext& InOutContext, const FName& InEventName)
 {
 	if (InEventName != FRigUnit_BeginExecution::EventName)
 	{
-		return;
+		// still return true - we want the evaluation code
+		// to simply carry on for other events.
+		return true;
 	}
 
 	for (FRigUnit_AddBoneTransform& Unit : AddBoneRigUnits)
@@ -55,6 +57,8 @@ void UAdditiveControlRig::ExecuteUnits(FRigUnitContext& InOutContext, const FNam
 		Unit.ExecuteContext.EventName = InEventName;
 		Unit.Execute(InOutContext);
 	}
+
+	return true;
 }
 
 void UAdditiveControlRig::Initialize(bool bInitRigUnits /*= true*/)
diff --git a/Engine/Plugins/Animation/ControlRig/Source/ControlRig/Private/Rigs/FKControlRig.cpp b/Engine/Plugins/Animation/ControlRig/Source/ControlRig/Private/Rigs/FKControlRig.cpp
--- a/Engine/Plugins/Animation/ControlRig/Source/ControlRig/Private/Rigs/FKControlRig.cpp
+++ b/Engine/Plugins/Animation/ControlRig/Source/ControlRig/Private/Rigs/FKControlRig.cpp
@@ -68,11 +68,13 @@ FName UFKControlRig::GetControlTargetName(const FName& InName, const ERigElement
 	return NAME_None;
 }
 
-void UFKControlRig::ExecuteUnits(FRigUnitContext& InOutContext, const FName& InEventName)
+bool UFKControlRig::ExecuteUnits(FRigUnitContext& InOutContext, const FName& InEventName)
 {
 	if (InOutContext.State != EControlRigState::Update)
 	{
-		return;
+		// we don't need any initialization so
+		// it's fine to consider this done here and return true.
+		return true;
 	}
 
 #if WITH_EDITOR
@@ -224,7 +226,9 @@ void UFKControlRig::ExecuteUnits(FRigUnitContext& InOutContext, const FName& InE
 			Hierarchy->DumpTransformStackToFile();
 		}
 	}
-#endif	
+#endif
+	
+	return true;
 }
 
 void UFKControlRig::SetBoneInitialTransformsFromSkeletalMeshComponent(USkeletalMeshComponent* InSkelMeshComp,
diff --git a/Engine/Plugins/Animation/ControlRig/Source/ControlRig/Public/ControlRig.h b/Engine/Plugins/Animation/ControlRig/Source/ControlRig/Public/ControlRig.h
--- a/Engine/Plugins/Animation/ControlRig/Source/ControlRig/Public/ControlRig.h
+++ b/Engine/Plugins/Animation/ControlRig/Source/ControlRig/Public/ControlRig.h
@@ -281,10 +281,10 @@ public:
 
 	/** Execute */
 	UFUNCTION(BlueprintCallable, Category = "Control Rig")
-	void Execute(const EControlRigState State, const FName& InEventName);
+	bool Execute(const EControlRigState State, const FName& InEventName);
 
 	/** ExecuteUnits */
-	virtual void ExecuteUnits(FRigUnitContext& InOutContext, const FName& InEventName);
+	virtual bool ExecuteUnits(FRigUnitContext& InOutContext, const FName& InEventName);
 
 	/** Returns true if this rig contains a given event */
 	UFUNCTION(BlueprintPure, Category = "Control Rig")
@@ -326,8 +326,8 @@ public:
 
 	virtual TArray<FRigControlElement*> AvailableControls() const;
 	virtual FRigControlElement* FindControl(const FName& InControlName) const;
-	virtual bool ShouldApplyLimits() const { return !bSetupModeEnabled; }
-	virtual bool IsSetupModeEnabled() const { return bSetupModeEnabled; }
+	virtual bool ShouldApplyLimits() const { return !IsSetupModeEnabled(); }
+	virtual bool IsSetupModeEnabled() const;
 	virtual FTransform SetupControlFromGlobalTransform(const FName& InControlName, const FTransform& InGlobalTransform);
 	virtual FTransform GetControlGlobalTransform(const FName& InControlName) const;
 
@@ -705,7 +705,6 @@ private:
 protected:
 	bool bRequiresInitExecution;
 	bool bRequiresSetupEvent;
-	bool bSetupModeEnabled;
 	bool bCopyHierarchyBeforeSetup;
 	bool bResetInitialTransformsBeforeSetup;
 	bool bManipulationEnabled;
diff --git a/Engine/Plugins/Animation/ControlRig/Source/ControlRig/Public/Rigs/AdditiveControlRig.h b/Engine/Plugins/Animation/ControlRig/Source/ControlRig/Public/Rigs/AdditiveControlRig.h
--- a/Engine/Plugins/Animation/ControlRig/Source/ControlRig/Public/Rigs/AdditiveControlRig.h
+++ b/Engine/Plugins/Animation/ControlRig/Source/ControlRig/Public/Rigs/AdditiveControlRig.h
@@ -18,7 +18,7 @@ class CONTROLRIG_API UAdditiveControlRig : public UControlRig
 public: 
 	// BEGIN ControlRig
 	virtual void Initialize(bool bInitRigUnits = true) override;
-	virtual void ExecuteUnits(FRigUnitContext& InOutContext, const FName& InEventName) override;
+	virtual bool ExecuteUnits(FRigUnitContext& InOutContext, const FName& InEventName) override;
 	// END ControlRig
 
 	// utility function to 
diff --git a/Engine/Plugins/Animation/ControlRig/Source/ControlRig/Public/Rigs/FKControlRig.h b/Engine/Plugins/Animation/ControlRig/Source/ControlRig/Public/Rigs/FKControlRig.h
--- a/Engine/Plugins/Animation/ControlRig/Source/ControlRig/Public/Rigs/FKControlRig.h
+++ b/Engine/Plugins/Animation/ControlRig/Source/ControlRig/Public/Rigs/FKControlRig.h
@@ -44,7 +44,7 @@ public:
 
 	// BEGIN ControlRig
 	virtual void Initialize(bool bInitRigUnits = true) override;
-	virtual void ExecuteUnits(FRigUnitContext& InOutContext, const FName& InEventName) override;
+	virtual bool ExecuteUnits(FRigUnitContext& InOutContext, const FName& InEventName) override;
 	virtual void SetBoneInitialTransformsFromSkeletalMeshComponent(USkeletalMeshComponent* InSkelMeshComp, bool bUseAnimInstance = false) override;
 	// END ControlRig
 
diff --git a/Engine/Plugins/Animation/ControlRig/Source/ControlRigEditor/Private/Editor/ControlRigEditor.cpp b/Engine/Plugins/Animation/ControlRig/Source/ControlRigEditor/Private/Editor/ControlRigEditor.cpp
--- a/Engine/Plugins/Animation/ControlRig/Source/ControlRigEditor/Private/Editor/ControlRigEditor.cpp
+++ b/Engine/Plugins/Animation/ControlRig/Source/ControlRigEditor/Private/Editor/ControlRigEditor.cpp
@@ -216,7 +216,6 @@ FControlRigEditor::FControlRigEditor()
 	, bControlRigEditorInitialized(false)
 	, bIsSettingObjectBeingDebugged(false)
 	, bExecutionControlRig(true)
-	, bSetupModeEnabled(false)
 	, bIsCompilingThroughUI(false)
 	, bAnyErrorsLeft(false)
 	, LastEventQueue(SetupEventQueue)
@@ -954,11 +953,6 @@ TArray<FName> FControlRigEditor::GetEventQueue() const
 {
 	if (ControlRig)
 	{
-		if (bSetupModeEnabled)
-		{
-			return SetupEventQueue;
-		}
-
 		return ControlRig->GetEventQueue();
 	}
 
@@ -966,6 +960,11 @@ TArray<FName> FControlRigEditor::GetEventQueue() const
 }
 
 void FControlRigEditor::SetEventQueue(TArray<FName> InEventQueue)
+{
+	return SetEventQueue(InEventQueue, false);
+}
+
+void FControlRigEditor::SetEventQueue(TArray<FName> InEventQueue, bool bCompile)
 {
 	if (GetEventQueue() == InEventQueue)
 	{
@@ -976,19 +975,34 @@ void FControlRigEditor::SetEventQueue(TArray<FName> InEventQueue)
 
 	SetHaltedNode(nullptr);
 
-	if (ControlRig)
+	TArray<FRigElementKey> PreviousSelection;
+	if (UControlRigBlueprint* RigBlueprint = Cast<UControlRigBlueprint>(GetBlueprintObj()))
 	{
-		if (InEventQueue.Num() > 0)
+		if(bCompile)
 		{
-			if(InEventQueue == SetupEventQueue)
+			if (RigBlueprint->bAutoRecompileVM)
 			{
-				if (!bSetupModeEnabled)
-				{
-					ToggleSetupMode();
-				}
-				return;
+				RigBlueprint->RequestAutoVMRecompilation();
 			}
-			
+			RigBlueprint->Validator->SetControlRig(ControlRig);
+		}
+		
+		// need to clear selection before remove transient control
+		// because active selection will trigger transient control recreation after removal	
+		PreviousSelection = RigBlueprint->Hierarchy->GetSelectedKeys();
+		RigBlueprint->GetHierarchyController()->ClearSelection();
+		
+		// need to copy here since the removal changes the iterator
+		if (ControlRig)
+		{
+			RigBlueprint->ClearTransientControls();
+		}
+	}
+
+	if (ControlRig)
+	{
+		if (InEventQueue.Num() > 0)
+		{
 			ControlRig->SetEventQueue(InEventQueue);
 
 			if (UControlRigBlueprint* RigBlueprint = Cast<UControlRigBlueprint>(GetBlueprintObj()))
@@ -997,11 +1011,6 @@ void FControlRigEditor::SetEventQueue(TArray<FName> InEventQueue)
 			}
 		}
 
-		if (bSetupModeEnabled)
-		{
-			ToggleSetupMode();
-		}
-
 		// Reset transforms only for setup and forward solve to not inturrupt any animation that might be playing
 		if (InEventQueue.Contains(FRigUnit_PrepareForExecution::EventName) ||
 			InEventQueue.Contains(FRigUnit_BeginExecution::EventName))
@@ -1009,6 +1018,26 @@ void FControlRigEditor::SetEventQueue(TArray<FName> InEventQueue)
 			ControlRig->GetHierarchy()->ResetPoseToInitial(ERigElementType::All);
 		}
 	}
+
+	if (FControlRigEditMode* EditMode = GetEditMode())
+	{
+		if (UControlRigBlueprint* RigBlueprint = Cast<UControlRigBlueprint>(GetBlueprintObj()))
+		{
+			EditMode->RecreateControlShapeActors(RigBlueprint->Hierarchy->GetSelectedKeys());
+		}
+
+		UControlRigEditModeSettings* Settings = GetMutableDefault<UControlRigEditModeSettings>();
+		Settings->bDisplayNulls = IsSetupModeEnabled();
+	}
+
+	if (PreviousSelection.Num() > 0)
+	{
+		if (UControlRigBlueprint* RigBlueprint = Cast<UControlRigBlueprint>(GetBlueprintObj()))
+		{
+			RigBlueprint->GetHierarchyController()->SetSelection(PreviousSelection);
+		}
+		SetDetailViewForRigElements();
+	}
 }
 
 int32 FControlRigEditor::GetEventQueueComboValue() const
@@ -1144,63 +1173,6 @@ FSlateIcon FControlRigEditor::GetExecutionModeIcon() const
 	return GetExecutionModeIcon(ExecutionMode);
 }
 
-void FControlRigEditor::ToggleSetupMode()
-{
-	bSetupModeEnabled = !bSetupModeEnabled;
-
-	TArray<FRigElementKey> PreviousSelection;
-	if (UControlRigBlueprint* RigBlueprint = Cast<UControlRigBlueprint>(GetBlueprintObj()))
-	{
-		if (RigBlueprint->bAutoRecompileVM)
-		{
-			RigBlueprint->RequestAutoVMRecompilation();
-		}
-
-		RigBlueprint->Validator->SetControlRig(ControlRig);
-		
-		// need to clear selection before remove transient control
-		// because active selection will trigger transient control recreation after removal	
-		PreviousSelection = RigBlueprint->Hierarchy->GetSelectedKeys();
-		RigBlueprint->GetHierarchyController()->ClearSelection();
-		
-		// need to copy here since the removal changes the iterator
-		if (ControlRig)
-		{
-			RigBlueprint->ClearTransientControls();
-		}
-	}
-
-	if (ControlRig)
-	{
-		ControlRig->bSetupModeEnabled = bSetupModeEnabled;
-		if (bSetupModeEnabled)
-		{
-			ControlRig->Initialize(true);
-			ControlRig->RequestSetup();
-		}
-	}
-
-	if (FControlRigEditMode* EditMode = GetEditMode())
-	{
-		if (UControlRigBlueprint* RigBlueprint = Cast<UControlRigBlueprint>(GetBlueprintObj()))
-		{
-			EditMode->RecreateControlShapeActors(RigBlueprint->Hierarchy->GetSelectedKeys());
-		}
-
-		UControlRigEditModeSettings* Settings = GetMutableDefault<UControlRigEditModeSettings>();
-		Settings->bDisplayNulls = bSetupModeEnabled;
-	}
-
-	if (PreviousSelection.Num() > 0)
-	{
-		if (UControlRigBlueprint* RigBlueprint = Cast<UControlRigBlueprint>(GetBlueprintObj()))
-		{
-			RigBlueprint->GetHierarchyController()->SetSelection(PreviousSelection);
-		}
-		SetDetailViewForRigElements();
-	}
-}
-
 void FControlRigEditor::GetCustomDebugObjects(TArray<FCustomDebugObject>& DebugList) const
 {
 	UControlRigBlueprint* RigBlueprint = Cast<UControlRigBlueprint>(GetBlueprintObj());
@@ -1606,7 +1578,7 @@ void FControlRigEditor::SetDetailViewForRigElements()
 	UControlRigBlueprint* RigBlueprint = Cast<UControlRigBlueprint>(GetBlueprintObj());
 	URigHierarchy* Hierarchy = RigBlueprint->Hierarchy;
 	
-	if (!bSetupModeEnabled)
+	if (!IsSetupModeEnabled())
 	{
 		if (UControlRig* DebuggedControlRig = Cast<UControlRig>(RigBlueprint->GetObjectBeingDebugged()))
 		{
@@ -1944,9 +1916,9 @@ void FControlRigEditor::Compile()
 			}
 		}
 
-		if (bSetupModeEnabled)
+		if(IsSetupModeEnabled())
 		{
-			bSetupModeEnabled = false;
+			SetEventQueue(ForwardsSolveEventQueue, false);
 		}
 
 		// clear transient controls such that we don't leave
@@ -2944,7 +2916,7 @@ void FControlRigEditor::HandleControlRigExecutedEvent(UControlRig* InControlRig,
 		}
 
 		URigHierarchy* Hierarchy = ControlRigBP->Hierarchy; 
-		if(!bSetupModeEnabled)
+		if(!IsSetupModeEnabled())
 		{
 			if(DebuggedControlRig)
 			{
@@ -3948,6 +3920,11 @@ void FControlRigEditor::OnPinControlNameListComboBox(const TArray<TSharedPtr<FSt
 	PinControlNameList->SetSelectedItem(CurrentlySelected);
 }
 
+bool FControlRigEditor::IsSetupModeEnabled() const
+{
+	return LastEventQueue.Contains(FRigUnit_PrepareForExecution::EventName);
+}
+
 void FControlRigEditor::HandlePreviewSceneCreated(const TSharedRef<IPersonaPreviewScene>& InPersonaPreviewScene)
 {
 	DECLARE_SCOPE_HIERARCHICAL_COUNTER_FUNC()
@@ -4048,7 +4025,6 @@ void FControlRigEditor::UpdateControlRig()
  			}
 
 			ControlRig->PreviewInstance = PreviewInstance;
-			ControlRig->bSetupModeEnabled = bSetupModeEnabled;
 
 #if WITH_EDITOR
 			ControlRig->bIsInDebugMode = ExecutionMode == EControlRigExecutionModeType_Debug;
@@ -4550,7 +4526,7 @@ void FControlRigEditor::OnWrappedPropertyChangedChainEvent(UDetailsViewWrapperOb
 				const ERigTransformType::Type TransformType = Local::GetTransformTypeFromPath(PropertyPath);
 				bIsInitial = bIsInitial || ERigTransformType::IsInitial(TransformType);
 				
-				if(ERigTransformType::IsInitial(TransformType) || bSetupModeEnabled)
+				if(ERigTransformType::IsInitial(TransformType) || IsSetupModeEnabled())
 				{
 					Hierarchy = ControlRigBP->Hierarchy;
 				}
@@ -4570,7 +4546,7 @@ void FControlRigEditor::OnWrappedPropertyChangedChainEvent(UDetailsViewWrapperOb
 					FRigControlValue Value;
 					Value.SetFromTransform(Transform, ControlElement->Settings.ControlType, ControlElement->Settings.PrimaryAxis);
 							
-					if(ERigTransformType::IsInitial(TransformType) || bSetupModeEnabled)
+					if(ERigTransformType::IsInitial(TransformType) || IsSetupModeEnabled())
 					{
 						Hierarchy->SetControlValue(ControlElement, Value, ERigControlValueType::Initial, true, true, true);
 					}
@@ -4624,7 +4600,7 @@ void FControlRigEditor::OnWrappedPropertyChangedChainEvent(UDetailsViewWrapperOb
 				ControlRigBP->Hierarchy->SetControlSettings(ControlElement, Settings, true, false, true);
 			}
 
-			if(bSetupModeEnabled || bIsInitial)
+			if(IsSetupModeEnabled() || bIsInitial)
 			{
 				ControlRigBP->PropagatePoseFromBPToInstances();
 				ControlRigBP->Modify();
@@ -6275,7 +6251,7 @@ void FControlRigEditor::HandleOnControlModified(UControlRig* Subject, FRigContro
 				Blueprint->Hierarchy->SetLocalTransform(ElementKey, Transform);
 				Hierarchy->SetLocalTransform(ElementKey, Transform);
 
-				if (bSetupModeEnabled)
+				if (IsSetupModeEnabled())
 				{
 					Blueprint->Hierarchy->SetInitialLocalTransform(ElementKey, Transform);
 					Hierarchy->SetInitialLocalTransform(ElementKey, Transform);
@@ -6295,7 +6271,7 @@ void FControlRigEditor::HandleOnControlModified(UControlRig* Subject, FRigContro
 			}
 		}
 	}
-	else if (bSetupModeEnabled)
+	else if (IsSetupModeEnabled())
 	{
 		FRigControlElement* SourceControlElement = Hierarchy->Find<FRigControlElement>(ControlElement->GetKey());
 		FRigControlElement* TargetControlElement = Blueprint->Hierarchy->Find<FRigControlElement>(ControlElement->GetKey());
diff --git a/Engine/Plugins/Animation/ControlRig/Source/ControlRigEditor/Private/Editor/ControlRigEditor.h b/Engine/Plugins/Animation/ControlRig/Source/ControlRigEditor/Private/Editor/ControlRigEditor.h
--- a/Engine/Plugins/Animation/ControlRig/Source/ControlRigEditor/Private/Editor/ControlRigEditor.h
+++ b/Engine/Plugins/Animation/ControlRig/Source/ControlRigEditor/Private/Editor/ControlRigEditor.h
@@ -267,6 +267,7 @@ private:
 
 	TArray<FName> GetEventQueue() const;
 	void SetEventQueue(TArray<FName> InEventQueue);
+	void SetEventQueue(TArray<FName> InEventQueue, bool bCompile);
 	int32 GetEventQueueComboValue() const;
 	FText GetEventQueueLabel() const;
 	static FSlateIcon GetEventQueueIcon(const TArray<FName>& InEventQueue);
@@ -450,9 +451,6 @@ protected:
 
 	void OnAnimInitialized();
 
-	/** Are we currently in setup mode */
-	bool bSetupModeEnabled;
-
 	/** Are we currently compiling through the user interface */
 	bool bIsCompilingThroughUI;
 
@@ -466,7 +464,7 @@ protected:
 	void SetPinControlNameListText(const FText& NewTypeInValue, ETextCommit::Type /*CommitInfo*/);
 	void OnPinControlNameListChanged(TSharedPtr<FString> NewSelection, ESelectInfo::Type SelectInfo);
 	void OnPinControlNameListComboBox(const TArray<TSharedPtr<FString>>* InNameList);
-	void ToggleSetupMode();
+	bool IsSetupModeEnabled() const;
 
 	bool bAnyErrorsLeft;
 
