Control Rig: Rename setup event to construction event

diff --git a/Engine/Plugins/Animation/ControlRig/Config/BaseControlRig.ini b/Engine/Plugins/Animation/ControlRig/Config/BaseControlRig.ini
--- a/Engine/Plugins/Animation/ControlRig/Config/BaseControlRig.ini
+++ b/Engine/Plugins/Animation/ControlRig/Config/BaseControlRig.ini
@@ -43,3 +43,8 @@
 +PropertyRedirects=(OldName="/Script/ControlRig.RigControlSettings.GizmoColor",NewName="ShapeColor")
 +PropertyRedirects=(OldName="/Script/ControlRig.RigControlElement.Gizmo",NewName="Shape")
 +PropertyRedirects=(OldName="/Script/ControlRigEditor.ControlRigEditModeSettings.GizmoScale",NewName="ShapeScale")
+
+; renaming setup event to construction event
++PropertyRedirects=(OldName="/Script/ControlRig.ControlRigComponent.OnPreSetupDelegate",NewName="OnPreConstructionDelegate")
++PropertyRedirects=(OldName="/Script/ControlRig.ControlRigComponent.OnPostSetupDelegate",NewName="OnPostConstructionDelegate")
++PropertyRedirects=(OldName="/Script/ControlRig.UControlRigSettings.SetupEventBorderColor",NewName="ConstructionEventBorderColor")
diff --git a/Engine/Plugins/Animation/ControlRig/Source/ControlRig/Private/ControlRig.cpp b/Engine/Plugins/Animation/ControlRig/Source/ControlRig/Private/ControlRig.cpp
--- a/Engine/Plugins/Animation/ControlRig/Source/ControlRig/Private/ControlRig.cpp
+++ b/Engine/Plugins/Animation/ControlRig/Source/ControlRig/Private/ControlRig.cpp
@@ -73,14 +73,14 @@ UControlRig::UControlRig(const FObjectInitializer& ObjectInitializer)
 	, PreviewInstance(nullptr)
 #endif
 	, bRequiresInitExecution(false)
-	, bRequiresSetupEvent(false)
-	, bCopyHierarchyBeforeSetup(true)
-	, bResetInitialTransformsBeforeSetup(true)
+	, bRequiresConstructionEvent(false)
+	, bCopyHierarchyBeforeConstruction(true)
+	, bResetInitialTransformsBeforeConstruction(true)
 	, bManipulationEnabled(false)
 	, InitBracket(0)
 	, UpdateBracket(0)
-	, PreSetupBracket(0)
-	, PostSetupBracket(0)
+	, PreConstructionBracket(0)
+	, PostConstructionBracket(0)
 	, InteractionBracket(0)
 	, InterRigSyncBracket(0)
 	, ControlUndoBracketIndex(0)
@@ -103,8 +103,8 @@ void UControlRig::BeginDestroy()
 {
 	Super::BeginDestroy();
 	InitializedEvent.Clear();
-	PreSetupEvent.Clear();
-	PostSetupEvent.Clear();
+	PreConstructionEvent.Clear();
+	PostConstructionEvent.Clear();
 	PreForwardsSolveEvent.Clear();
 	PostForwardsSolveEvent.Clear();
 	ExecutedEvent.Clear();
@@ -216,7 +216,7 @@ void UControlRig::Initialize(bool bInitRigUnits)
 	GetHierarchy()->GetController(true);
 	
 	// should refresh mapping 
-	RequestSetup();
+	RequestConstruction();
 
 	if (bInitRigUnits)
 	{
@@ -293,8 +293,8 @@ void UControlRig::Evaluate_AnyThread()
 		}
 	}
 
-	// execute the setup event prior to everything else
-	if(bRequiresSetupEvent)
+	// execute the construction event prior to everything else
+	if(bRequiresConstructionEvent)
 	{
 		if(!EventQueueToRun.Contains(FRigUnit_PrepareForExecution::EventName))
 		{
@@ -627,7 +627,7 @@ bool UControlRig::Execute(const EControlRigState InState, const FName& InEventNa
 	const bool bIsEventLastInQueue = !GetEventQueue().IsEmpty() && GetEventQueue().Last() == InEventName;
 	const bool bIsInitializingMemory = InState == EControlRigState::Init;
 	const bool bIsExecutingInstructions = InState == EControlRigState::Update;
-	const bool bIsSetupEvent = InEventName == FRigUnit_PrepareForExecution::EventName;
+	const bool bIsConstructionEvent = InEventName == FRigUnit_PrepareForExecution::EventName;
 	const bool bIsForwardSolve = InEventName == FRigUnit_BeginExecution::EventName;
 	const bool bIsInteractionEvent = InEventName == FRigUnit_InteractionExecution::EventName;
 
@@ -859,39 +859,39 @@ bool UControlRig::Execute(const EControlRigState InState, const FName& InEventNa
 		UE_LOG(LogControlRig, Warning, TEXT("%s: Execute is being called recursively."), *GetPathName());
 		return false;
 	}
-	if(bIsSetupEvent)
+	if(bIsConstructionEvent)
 	{
-		if(IsRunningPreSetup() || IsRunningPostSetup())
+		if(IsRunningPreConstruction() || IsRunningPostConstruction())
 		{
-			UE_LOG(LogControlRig, Warning, TEXT("%s: Setup is being called recursively."), *GetPathName());
+			UE_LOG(LogControlRig, Warning, TEXT("%s: Construction is being called recursively."), *GetPathName());
 			return false;
 		}
 	}
 
 	bool bSuccess = true;
 
-	// we'll special case the setup event here
-	if (bIsSetupEvent && !bIsInitializingMemory)
+	// we'll special case the construction event here
+	if (bIsConstructionEvent && !bIsInitializingMemory)
 	{
 		// construction mode means that we are running the construction event
 		// constantly for testing purposes.
-		const bool bSetupModeEnabled = IsSetupModeEnabled();
+		const bool bConstructionModeEnabled = IsConstructionModeEnabled();
 		{
-			// save the current state of all pose elements to preserve user intention, since setup event can
+			// save the current state of all pose elements to preserve user intention, since construction event can
 			// run in between forward events
-			// the saved pose is reapplied to the rig after setup event as the pose scope goes out of scope
+			// the saved pose is reapplied to the rig after construction event as the pose scope goes out of scope
 			TUniquePtr<UControlRig::FPoseScope> PoseScope;
-			if (!bSetupModeEnabled)
+			if (!bConstructionModeEnabled)
 			{
-				// only do this in non-setup mode because 
-				// when setup mode is enabled, the control values are cleared before reaching here (too late to save them)
-				PoseScope = MakeUnique<UControlRig::FPoseScope>(this, ERigElementType::ToResetAfterSetupEvent);
+				// only do this in non-construction mode because 
+				// when construction mode is enabled, the control values are cleared before reaching here (too late to save them)
+				PoseScope = MakeUnique<UControlRig::FPoseScope>(this, ERigElementType::ToResetAfterConstructionEvent);
 			}
 
 			// Copy the hierarchy from the default object onto this one
 			if (UControlRig* CDO = Cast<UControlRig>(GetClass()->GetDefaultObject()))
 			{
-				if(bCopyHierarchyBeforeSetup && !bSetupModeEnabled)
+				if(bCopyHierarchyBeforeConstruction && !bConstructionModeEnabled)
 				{
 					if(CDO->GetHierarchy()->GetTopologyVersion()!= GetHierarchy()->GetTopologyVersion())
 					{
@@ -902,7 +902,7 @@ bool UControlRig::Execute(const EControlRigState InState, const FName& InEventNa
 					}
 				}
 				
-				if (bResetInitialTransformsBeforeSetup && !bSetupModeEnabled)
+				if (bResetInitialTransformsBeforeConstruction && !bConstructionModeEnabled)
 				{
 					// when copying the pose from the default we don't copy the weights over,
 					// since we don't want to reset the topology in case it has diverged from the
@@ -914,38 +914,38 @@ bool UControlRig::Execute(const EControlRigState InState, const FName& InEventNa
 			{
 #if WITH_EDITOR
 				TUniquePtr<FTransientControlPoseScope> TransientControlPoseScope;
-				if (bSetupModeEnabled)
+				if (bConstructionModeEnabled)
 				{
-					// save the transient control value, it should not be constantly reset in setup mode
+					// save the transient control value, it should not be constantly reset in construction mode
 					TransientControlPoseScope = MakeUnique<FTransientControlPoseScope>(this);
 				}
 #endif
-				// reset the pose to initial such that setup event can run from a deterministic initial state
+				// reset the pose to initial such that construction event can run from a deterministic initial state
 				GetHierarchy()->ResetPoseToInitial(ERigElementType::All);
 			}
 
-			if (PreSetupEvent.IsBound())
+			if (PreConstructionEvent.IsBound())
 			{
-				FControlRigBracketScope BracketScope(PreSetupBracket);
-				PreSetupEvent.Broadcast(this, EControlRigState::Update, FRigUnit_PrepareForExecution::EventName);
+				FControlRigBracketScope BracketScope(PreConstructionBracket);
+				PreConstructionEvent.Broadcast(this, EControlRigState::Update, FRigUnit_PrepareForExecution::EventName);
 			}
 
 			bSuccess = ExecuteUnits(Context, FRigUnit_PrepareForExecution::EventName);
 
-			if (PostSetupEvent.IsBound())
+			if (PostConstructionEvent.IsBound())
 			{
-				FControlRigBracketScope BracketScope(PostSetupBracket);
-				PostSetupEvent.Broadcast(this, EControlRigState::Update, FRigUnit_PrepareForExecution::EventName);
+				FControlRigBracketScope BracketScope(PostConstructionBracket);
+				PostConstructionEvent.Broadcast(this, EControlRigState::Update, FRigUnit_PrepareForExecution::EventName);
 			}
 		}
 
-		if (bSetupModeEnabled)
+		if (bConstructionModeEnabled)
 		{
 #if WITH_EDITOR
 			TUniquePtr<FTransientControlPoseScope> TransientControlPoseScope;
-			if (bSetupModeEnabled)
+			if (bConstructionModeEnabled)
 			{
-				// save the transient control value, it should not be constantly reset in setup mode
+				// save the transient control value, it should not be constantly reset in construction mode
 				TransientControlPoseScope = MakeUnique<FTransientControlPoseScope>(this);
 			}
 #endif
@@ -1305,9 +1305,9 @@ bool UControlRig::Execute(const EControlRigState InState, const FName& InEventNa
 
 	if(bSuccess)
 	{
-		if(bIsSetupEvent)
+		if(bIsConstructionEvent)
 		{
-			bRequiresSetupEvent = false;
+			bRequiresConstructionEvent = false;
 		}
 	}
 
@@ -1499,12 +1499,12 @@ bool UControlRig::ExecuteEvent(const FName& InEventName)
 void UControlRig::RequestInit()
 {
 	bRequiresInitExecution = true;
-	RequestSetup();
+	RequestConstruction();
 }
 
-void UControlRig::RequestSetup()
+void UControlRig::RequestConstruction()
 {
-	bRequiresSetupEvent = true;
+	bRequiresConstructionEvent = true;
 }
 
 void UControlRig::SetEventQueue(const TArray<FName>& InEventNames)
@@ -1745,14 +1745,14 @@ FRigControlElement* UControlRig::FindControl(const FName& InControlName) const
 	return DynamicHierarchy->Find<FRigControlElement>(FRigElementKey(InControlName, ERigElementType::Control));
 }
 
-bool UControlRig::IsSetupModeEnabled() const
+bool UControlRig::IsConstructionModeEnabled() const
 {
 	return EventQueue.Contains(FRigUnit_PrepareForExecution::EventName);
 }
 
 FTransform UControlRig::SetupControlFromGlobalTransform(const FName& InControlName, const FTransform& InGlobalTransform)
 {
-	if (IsSetupModeEnabled())
+	if (IsConstructionModeEnabled())
 	{
 		FRigControlElement* ControlElement = FindControl(InControlName);
 		if (ControlElement && !ControlElement->Settings.bIsTransientControl)
@@ -1807,7 +1807,7 @@ void UControlRig::HandleOnControlModified(UControlRig* Subject, FRigControlEleme
 {
 	if (Control->Settings.bIsCurve && DynamicHierarchy)
 	{
-		const FRigControlValue Value = DynamicHierarchy->GetControlValue(Control, IsSetupModeEnabled() ? ERigControlValueType::Initial : ERigControlValueType::Current);
+		const FRigControlValue Value = DynamicHierarchy->GetControlValue(Control, IsConstructionModeEnabled() ? ERigControlValueType::Initial : ERigControlValueType::Current);
 		DynamicHierarchy->SetCurveValue(FRigElementKey(Control->GetName(), ERigElementType::Curve), Value.Get<float>());
 	}	
 }
@@ -1849,7 +1849,7 @@ FTransform UControlRig::GetControlGlobalTransform(const FName& InControlName) co
 bool UControlRig::SetControlGlobalTransform(const FName& InControlName, const FTransform& InGlobalTransform, bool bNotify, const FRigControlModifiedContext& Context, bool bSetupUndo, bool bPrintPythonCommands, bool bFixEulerFlips)
 {
 	FTransform GlobalTransform = InGlobalTransform;
-	if (IsSetupModeEnabled())
+	if (IsConstructionModeEnabled())
 	{
 		GlobalTransform = SetupControlFromGlobalTransform(InControlName, GlobalTransform);
 	}
@@ -2220,8 +2220,8 @@ FName UControlRig::AddTransientControl(const FRigElementKey& InElement)
 	if (InElement.Type == ERigElementType::Bone)
 	{
 		// don't allow transient control to modify forward mode poses when we
-		// already switched to the setup mode
-		if (!IsSetupModeEnabled())
+		// already switched to the construction mode
+		if (!IsConstructionModeEnabled())
 		{
 			if(FRigBoneElement* BoneElement = DynamicHierarchy->Find<FRigBoneElement>(InElement))
 			{
@@ -2271,9 +2271,9 @@ bool UControlRig::SetTransientControlValue(const FRigElementKey& InElement)
 	{
 		if (InElement.Type == ERigElementType::Bone)
 		{
-			if (IsSetupModeEnabled())
+			if (IsConstructionModeEnabled())
 			{
-				// need to get initial because that is what setup mode uses
+				// need to get initial because that is what construction mode uses
 				// specifically, when user change the initial from the details panel
 				// this will allow the transient control to react to that change
 				const FTransform InitialLocalTransform = DynamicHierarchy->GetInitialLocalTransform(InElement);
@@ -2540,7 +2540,7 @@ void UControlRig::SetInteractionRig(UControlRig* InInteractionRig)
 
 		InteractionRig->Initialize(true);
 		InteractionRig->CopyPoseFromOtherRig(this);
-		InteractionRig->RequestSetup();
+		InteractionRig->RequestConstruction();
 		InteractionRig->Execute(EControlRigState::Update, FRigUnit_BeginExecution::EventName);
 
 		InteractionRig->ControlModified().AddUObject(this, &UControlRig::HandleInteractionRigControlModified);
@@ -3002,8 +3002,8 @@ void UControlRig::SetBoneInitialTransformsFromRefSkeleton(const FReferenceSkelet
 		}
 		return true;
 	});
-	bResetInitialTransformsBeforeSetup = false;
-	RequestSetup();
+	bResetInitialTransformsBeforeConstruction = false;
+	RequestConstruction();
 }
 
 void UControlRig::SetBoneInitialTransformsFromCompactPose(FCompactPose* InCompactPose)
@@ -3039,8 +3039,8 @@ void UControlRig::SetBoneInitialTransformsFromCompactPose(FCompactPose* InCompac
 			return true;
 		});
 
-	bResetInitialTransformsBeforeSetup = false;
-	RequestSetup();
+	bResetInitialTransformsBeforeConstruction = false;
+	RequestConstruction();
 }
 
 const FRigControlElementCustomization* UControlRig::GetControlCustomization(const FRigElementKey& InControl) const
diff --git a/Engine/Plugins/Animation/ControlRig/Source/ControlRig/Private/ControlRigComponent.cpp b/Engine/Plugins/Animation/ControlRig/Source/ControlRig/Private/ControlRigComponent.cpp
--- a/Engine/Plugins/Animation/ControlRig/Source/ControlRig/Private/ControlRigComponent.cpp
+++ b/Engine/Plugins/Animation/ControlRig/Source/ControlRig/Private/ControlRigComponent.cpp
@@ -53,7 +53,7 @@ UControlRigComponent::UControlRigComponent(const FObjectInitializer& ObjectIniti
 
 	ControlRig = nullptr;
 	bResetTransformBeforeTick = true;
-	bResetInitialsBeforeSetup = true;
+	bResetInitialsBeforeConstruction = true;
 	bUpdateRigOnTick = true;
 	bUpdateInEditor = true;
 	bDrawBones = true;
@@ -283,14 +283,14 @@ void UControlRigComponent::OnPostInitialize_Implementation(UControlRigComponent*
 	OnPostInitializeDelegate.Broadcast(Component);
 }
 
-void UControlRigComponent::OnPreSetup_Implementation(UControlRigComponent* Component)
+void UControlRigComponent::OnPreConstruction_Implementation(UControlRigComponent* Component)
 {
-	OnPreSetupDelegate.Broadcast(Component);
+	OnPreConstructionDelegate.Broadcast(Component);
 }
 
-void UControlRigComponent::OnPostSetup_Implementation(UControlRigComponent* Component)
+void UControlRigComponent::OnPostConstruction_Implementation(UControlRigComponent* Component)
 {
-	OnPostSetupDelegate.Broadcast(Component);
+	OnPostConstructionDelegate.Broadcast(Component);
 }
 
 void UControlRigComponent::OnPreForwardsSolve_Implementation(UControlRigComponent* Component)
@@ -371,7 +371,7 @@ void UControlRigComponent::Update(float DeltaTime)
 		else
 		{
 			CR->SetDeltaTime(DeltaTime);
-			CR->bResetInitialTransformsBeforeSetup = bResetInitialsBeforeSetup;
+			CR->bResetInitialTransformsBeforeConstruction = bResetInitialsBeforeConstruction;
 
 			// todo: set log
 			// todo: set external data providers
@@ -676,7 +676,7 @@ void UControlRigComponent::SetBoneInitialTransformsFromSkeletalMesh(USkeletalMes
 		if (UControlRig* CR = SetupControlRigIfRequired())
 		{
 			CR->SetBoneInitialTransformsFromSkeletalMesh(InSkeletalMesh);
-			bResetInitialsBeforeSetup = false;
+			bResetInitialsBeforeConstruction = false;
 		}
 	}
 }
@@ -778,9 +778,9 @@ void UControlRigComponent::SetInitialBoneTransform(FName BoneName, FTransform In
 		const int32 BoneIndex = CR->GetHierarchy()->GetIndex(FRigElementKey(BoneName, ERigElementType::Bone));
 		if (BoneIndex != INDEX_NONE)
 		{
-			if(!CR->IsRunningPreSetup() && !CR->IsRunningPostSetup())
+			if(!CR->IsRunningPreConstruction() && !CR->IsRunningPostConstruction())
 			{
-				ReportError(TEXT("SetInitialBoneTransform should only be called during OnPreSetup / OnPostSetup."));
+				ReportError(TEXT("SetInitialBoneTransform should only be called during OnPreConstruction / OnPostConstruction."));
 				return;
 			}
 
@@ -1144,8 +1144,8 @@ UControlRig* UControlRigComponent::SetupControlRigIfRequired()
 		if (ControlRig->GetClass() != ControlRigClass)
 		{
 			ControlRig->OnInitialized_AnyThread().RemoveAll(this);
-			ControlRig->OnPreSetup_AnyThread().RemoveAll(this);
-			ControlRig->OnPostSetup_AnyThread().RemoveAll(this);
+			ControlRig->OnPreConstruction_AnyThread().RemoveAll(this);
+			ControlRig->OnPostConstruction_AnyThread().RemoveAll(this);
 			ControlRig->OnPreForwardsSolve_AnyThread().RemoveAll(this);
 			ControlRig->OnPostForwardsSolve_AnyThread().RemoveAll(this);
 			ControlRig->OnExecuted_AnyThread().RemoveAll(this);
@@ -1178,16 +1178,16 @@ void UControlRigComponent::SetControlRig(UControlRig* InControlRig)
 	if (ControlRig)
 	{
 		ControlRig->OnInitialized_AnyThread().RemoveAll(this);
-		ControlRig->OnPreSetup_AnyThread().RemoveAll(this);
-		ControlRig->OnPostSetup_AnyThread().RemoveAll(this);
+		ControlRig->OnPreConstruction_AnyThread().RemoveAll(this);
+		ControlRig->OnPostConstruction_AnyThread().RemoveAll(this);
 		ControlRig->OnPreForwardsSolve_AnyThread().RemoveAll(this);
 		ControlRig->OnPostForwardsSolve_AnyThread().RemoveAll(this);
 		ControlRig->OnExecuted_AnyThread().RemoveAll(this);
 	}
 	ControlRig = InControlRig;
 	ControlRig->OnInitialized_AnyThread().AddUObject(this, &UControlRigComponent::HandleControlRigInitializedEvent);
-	ControlRig->OnPreSetup_AnyThread().AddUObject(this, &UControlRigComponent::HandleControlRigPreSetupEvent);
-	ControlRig->OnPostSetup_AnyThread().AddUObject(this, &UControlRigComponent::HandleControlRigPostSetupEvent);
+	ControlRig->OnPreConstruction_AnyThread().AddUObject(this, &UControlRigComponent::HandleControlRigPreConstructionEvent);
+	ControlRig->OnPostConstruction_AnyThread().AddUObject(this, &UControlRigComponent::HandleControlRigPostConstructionEvent);
 	ControlRig->OnPreForwardsSolve_AnyThread().AddUObject(this, &UControlRigComponent::HandleControlRigPreForwardsSolveEvent);
 	ControlRig->OnPostForwardsSolve_AnyThread().AddUObject(this, &UControlRigComponent::HandleControlRigPostForwardsSolveEvent);
 	ControlRig->OnExecuted_AnyThread().AddUObject(this, &UControlRigComponent::HandleControlRigExecutedEvent);
@@ -1617,7 +1617,7 @@ void UControlRigComponent::HandleControlRigInitializedEvent(UControlRig* InContr
 	}
 }
 
-void UControlRigComponent::HandleControlRigPreSetupEvent(UControlRig* InControlRig, const EControlRigState InState, const FName& InEventName)
+void UControlRigComponent::HandleControlRigPreConstructionEvent(UControlRig* InControlRig, const EControlRigState InState, const FName& InEventName)
 {
 	TArray<USkeletalMeshComponent*> ComponentsToTick;
 
@@ -1658,27 +1658,27 @@ void UControlRigComponent::HandleControlRigPreSetupEvent(UControlRig* InControlR
 	if (bUpdateInEditor)
 	{
 		FEditorScriptExecutionGuard AllowScripts;
-		OnPreSetup(this);
+		OnPreConstruction(this);
 	}
 	else
 #endif
 	{
-		OnPreSetup(this);
+		OnPreConstruction(this);
 	}
 }
 
-void UControlRigComponent::HandleControlRigPostSetupEvent(UControlRig* InControlRig, const EControlRigState InState, const FName& InEventName)
+void UControlRigComponent::HandleControlRigPostConstructionEvent(UControlRig* InControlRig, const EControlRigState InState, const FName& InEventName)
 {
 #if WITH_EDITOR
 	if (bUpdateInEditor)
 	{
 		FEditorScriptExecutionGuard AllowScripts;
-		OnPostSetup(this);
+		OnPostConstruction(this);
 	}
 	else
 #endif
 	{
-		OnPostSetup(this);
+		OnPostConstruction(this);
 	}
 }
 
@@ -1771,30 +1771,30 @@ bool UControlRigComponent::EnsureCalledOutsideOfBracket(const TCHAR* InCallingFu
 {
 	if (UControlRig* CR = SetupControlRigIfRequired())
 	{
-		if (CR->IsRunningPreSetup())
+		if (CR->IsRunningPreConstruction())
 		{
 			if (InCallingFunctionName)
 			{
-				ReportError(FString::Printf(TEXT("%s cannot be called during the PreSetupEvent - use ConstructionScript instead."), InCallingFunctionName));
+				ReportError(FString::Printf(TEXT("%s cannot be called during the PreConstructionEvent - use ConstructionScript instead."), InCallingFunctionName));
 				return false;
 			}
 			else
 			{
-				ReportError(FString::Printf(TEXT("Cannot be called during the PreSetupEvent - use ConstructionScript instead."), InCallingFunctionName));
+				ReportError(FString::Printf(TEXT("Cannot be called during the PreConstructionEvent - use ConstructionScript instead."), InCallingFunctionName));
 				return false;
 			}
 		}
 
-		if (CR->IsRunningPostSetup())
+		if (CR->IsRunningPostConstruction())
 		{
 			if (InCallingFunctionName)
 			{
-				ReportError(FString::Printf(TEXT("%s cannot be called during the PostSetupEvent - use ConstructionScript instead."), InCallingFunctionName));
+				ReportError(FString::Printf(TEXT("%s cannot be called during the PostConstructionEvent - use ConstructionScript instead."), InCallingFunctionName));
 				return false;
 			}
 			else
 			{
-				ReportError(FString::Printf(TEXT("Cannot be called during the PostSetupEvent - use ConstructionScript instead."), InCallingFunctionName));
+				ReportError(FString::Printf(TEXT("Cannot be called during the PostConstructionEvent - use ConstructionScript instead."), InCallingFunctionName));
 				return false;
 			}
 		}
diff --git a/Engine/Plugins/Animation/ControlRig/Source/ControlRig/Private/Rigs/AdditiveControlRig.cpp b/Engine/Plugins/Animation/ControlRig/Source/ControlRig/Private/Rigs/AdditiveControlRig.cpp
--- a/Engine/Plugins/Animation/ControlRig/Source/ControlRig/Private/Rigs/AdditiveControlRig.cpp
+++ b/Engine/Plugins/Animation/ControlRig/Source/ControlRig/Private/Rigs/AdditiveControlRig.cpp
@@ -13,8 +13,8 @@
 UAdditiveControlRig::UAdditiveControlRig(const FObjectInitializer& ObjectInitializer)
 	: Super(ObjectInitializer)
 {
-	bCopyHierarchyBeforeSetup = false;
-	bResetInitialTransformsBeforeSetup = false;
+	bCopyHierarchyBeforeConstruction = false;
+	bResetInitialTransformsBeforeConstruction = false;
 }
 
 FName UAdditiveControlRig::GetControlName(const FName& InBoneName)
diff --git a/Engine/Plugins/Animation/ControlRig/Source/ControlRig/Private/Rigs/FKControlRig.cpp b/Engine/Plugins/Animation/ControlRig/Source/ControlRig/Private/Rigs/FKControlRig.cpp
--- a/Engine/Plugins/Animation/ControlRig/Source/ControlRig/Private/Rigs/FKControlRig.cpp
+++ b/Engine/Plugins/Animation/ControlRig/Source/ControlRig/Private/Rigs/FKControlRig.cpp
@@ -17,8 +17,8 @@ UFKControlRig::UFKControlRig(const FObjectInitializer& ObjectInitializer)
 	, ApplyMode(EControlRigFKRigExecuteMode::Replace)
 	, CachedToggleApplyMode(EControlRigFKRigExecuteMode::Replace)
 {
-	bCopyHierarchyBeforeSetup = false;
-	bResetInitialTransformsBeforeSetup = false;
+	bCopyHierarchyBeforeConstruction = false;
+	bResetInitialTransformsBeforeConstruction = false;
 }
 
 FName UFKControlRig::GetControlName(const FName& InName, const ERigElementType& InType)
diff --git a/Engine/Plugins/Animation/ControlRig/Source/ControlRig/Private/Sequencer/MovieSceneControlRigParameterSection.cpp b/Engine/Plugins/Animation/ControlRig/Source/ControlRig/Private/Sequencer/MovieSceneControlRigParameterSection.cpp
--- a/Engine/Plugins/Animation/ControlRig/Source/ControlRig/Private/Sequencer/MovieSceneControlRigParameterSection.cpp
+++ b/Engine/Plugins/Animation/ControlRig/Source/ControlRig/Private/Sequencer/MovieSceneControlRigParameterSection.cpp
@@ -2526,7 +2526,7 @@ bool UMovieSceneControlRigParameterSection::LoadAnimSequenceIntoThisSection(UAni
 	FScopedSlowTask Progress(NumberOfKeys + ExtraProgress, LOCTEXT("BakingToControlRig_SlowTask", "Baking To Control Rig..."));	
 	Progress.MakeDialog(true);
 
-	//Make sure we are reset and run setup event  before evaluating
+	//Make sure we are reset and run construction event  before evaluating
 	/*
 	TArray<FRigElementKey>ControlsToReset = ControlRig->GetHierarchy()->GetAllKeys(true, ERigElementType::Control);
 	for (const FRigElementKey& ControlToReset : ControlsToReset)
@@ -2568,7 +2568,7 @@ bool UMovieSceneControlRigParameterSection::LoadAnimSequenceIntoThisSection(UAni
 	{
 		ControlRig->SetBoneInitialTransformsFromRefSkeleton(Skeleton->GetReferenceSkeleton());
 	}
-	ControlRig->RequestSetup();
+	ControlRig->RequestConstruction();
 	ControlRig->Evaluate_AnyThread();
 
 	for (int32 Index = 0; Index < NumberOfKeys; ++Index)
diff --git a/Engine/Plugins/Animation/ControlRig/Source/ControlRig/Private/Settings/ControlRigSettings.cpp b/Engine/Plugins/Animation/ControlRig/Source/ControlRig/Private/Settings/ControlRigSettings.cpp
--- a/Engine/Plugins/Animation/ControlRig/Source/ControlRig/Private/Settings/ControlRigSettings.cpp
+++ b/Engine/Plugins/Animation/ControlRig/Source/ControlRig/Private/Settings/ControlRigSettings.cpp
@@ -27,7 +27,7 @@ UControlRigEditorSettings::UControlRigEditorSettings(const FObjectInitializer& O
 	bResetControlsOnPinValueInteraction = false;
 	bEnableUndoForPoseInteraction = true;
 
-	SetupEventBorderColor = FLinearColor::Red;
+	ConstructionEventBorderColor = FLinearColor::Red;
 	BackwardsSolveBorderColor = FLinearColor::Yellow;
 	BackwardsAndForwardsBorderColor = FLinearColor::Blue;
 	DynamicHierarchyBorderColor = FLinearColor::Blue;
diff --git a/Engine/Plugins/Animation/ControlRig/Source/ControlRig/Private/Units/Execution/RigUnit_DynamicHierarchy.cpp b/Engine/Plugins/Animation/ControlRig/Source/ControlRig/Private/Units/Execution/RigUnit_DynamicHierarchy.cpp
--- a/Engine/Plugins/Animation/ControlRig/Source/ControlRig/Private/Units/Execution/RigUnit_DynamicHierarchy.cpp
+++ b/Engine/Plugins/Animation/ControlRig/Source/ControlRig/Private/Units/Execution/RigUnit_DynamicHierarchy.cpp
@@ -107,7 +107,7 @@ FRigUnit_SwitchParent_Execute()
 				return;
 			}
 
-			// during setup event also change the initial weights
+			// during construction event also change the initial weights
 			if(ExecuteContext.GetEventName() == FRigUnit_PrepareForExecution::EventName)
 			{
 				if(!ExecuteContext.Hierarchy->SwitchToParent(ChildElement, ParentElement, true, true, EmptyDependencyMap, &FailureReason))
@@ -181,7 +181,7 @@ FRigUnit_HierarchySetParentWeights_Execute()
 
 	ExecuteContext.Hierarchy->SetParentWeightArray(ChildElement, Weights, false, true);
 
-	// during setup event also change the initial weights
+	// during construction event also change the initial weights
 	if(ExecuteContext.GetEventName() == FRigUnit_PrepareForExecution::EventName)
 	{
 		ExecuteContext.Hierarchy->SetParentWeightArray(ChildElement, Weights, true, true);
diff --git a/Engine/Plugins/Animation/ControlRig/Source/ControlRig/Private/Units/Execution/RigUnit_PrepareForExecution.cpp b/Engine/Plugins/Animation/ControlRig/Source/ControlRig/Private/Units/Execution/RigUnit_PrepareForExecution.cpp
--- a/Engine/Plugins/Animation/ControlRig/Source/ControlRig/Private/Units/Execution/RigUnit_PrepareForExecution.cpp
+++ b/Engine/Plugins/Animation/ControlRig/Source/ControlRig/Private/Units/Execution/RigUnit_PrepareForExecution.cpp
@@ -3,7 +3,7 @@
 #include "RigUnit_PrepareForExecution.h"
 #include "Units/RigUnitContext.h"
 
-FName FRigUnit_PrepareForExecution::EventName = TEXT("Setup");
+FName FRigUnit_PrepareForExecution::EventName = TEXT("Construction");
 
 FRigUnit_PrepareForExecution_Execute()
 {
diff --git a/Engine/Plugins/Animation/ControlRig/Source/ControlRig/Private/Units/Execution/RigUnit_PrepareForExecution.h b/Engine/Plugins/Animation/ControlRig/Source/ControlRig/Private/Units/Execution/RigUnit_PrepareForExecution.h
--- a/Engine/Plugins/Animation/ControlRig/Source/ControlRig/Private/Units/Execution/RigUnit_PrepareForExecution.h
+++ b/Engine/Plugins/Animation/ControlRig/Source/ControlRig/Private/Units/Execution/RigUnit_PrepareForExecution.h
@@ -7,9 +7,9 @@
 #include "RigUnit_PrepareForExecution.generated.h"
 
 /**
- * Event to setup elements before all solves
+ * Event to create / configure elements before any other event
  */
-USTRUCT(meta=(DisplayName="Setup Event", Category="Events", NodeColor="1, 0, 0", Keywords="Setup,Init,Fit"))
+USTRUCT(meta=(DisplayName="Construction Event", Category="Events", NodeColor="0.6, 0, 1", Keywords="Create,Build,Spawn,Setup,Init,Fit"))
 struct CONTROLRIG_API FRigUnit_PrepareForExecution : public FRigUnit
 {
 	GENERATED_BODY()
diff --git a/Engine/Plugins/Animation/ControlRig/Source/ControlRig/Private/Units/Hierarchy/RigUnit_SetControlOffset.h b/Engine/Plugins/Animation/ControlRig/Source/ControlRig/Private/Units/Hierarchy/RigUnit_SetControlOffset.h
--- a/Engine/Plugins/Animation/ControlRig/Source/ControlRig/Private/Units/Hierarchy/RigUnit_SetControlOffset.h
+++ b/Engine/Plugins/Animation/ControlRig/Source/ControlRig/Private/Units/Hierarchy/RigUnit_SetControlOffset.h
@@ -7,7 +7,7 @@
 
 /**
  * SetControlOffset is used to perform a change in the hierarchy by setting a single control's transform.
- * This is typically only used during the Setup Event.
+ * This is typically only used during the Construction Event.
  */
 USTRUCT(meta=(DisplayName="Set Control Offset", Category="Controls", DocumentationPolicy="Strict", Keywords = "SetControlOffset,Initial,InitialTransform,SetInitialTransform,SetInitialControlTransform", NodeColor="0, 0.364706, 1.0"))
 struct CONTROLRIG_API FRigUnit_SetControlOffset : public FRigUnitMutable
@@ -49,7 +49,7 @@ struct CONTROLRIG_API FRigUnit_SetControlOffset : public FRigUnitMutable
 
 /**
  * GetShapeTransform is used to retrieve single control's shape transform.
- * This is typically only used during the Setup Event.
+ * This is typically only used during the Construction Event.
  */
 USTRUCT(meta=(DisplayName="Get Shape Transform", Category="Controls", DocumentationPolicy="Strict", Keywords = "GetControlShapeTransform,Gizmo,GizmoTransform,MeshTransform", NodeColor="0, 0.364706, 1.0"))
 struct CONTROLRIG_API FRigUnit_GetShapeTransform : public FRigUnit
@@ -84,7 +84,7 @@ struct CONTROLRIG_API FRigUnit_GetShapeTransform : public FRigUnit
 
 /**
  * SetShapeTransform is used to perform a change in the hierarchy by setting a single control's shape transform.
- * This is typically only used during the Setup Event.
+ * This is typically only used during the Construction Event.
  */
 USTRUCT(meta=(DisplayName="Set Shape Transform", Category="Controls", DocumentationPolicy="Strict", Keywords = "SetControlShapeTransform,Gizmo,GizmoTransform,MeshTransform", NodeColor="0, 0.364706, 1.0"))
 struct CONTROLRIG_API FRigUnit_SetShapeTransform : public FRigUnitMutable
diff --git a/Engine/Plugins/Animation/ControlRig/Source/ControlRig/Public/ControlRig.h b/Engine/Plugins/Animation/ControlRig/Source/ControlRig/Public/ControlRig.h
--- a/Engine/Plugins/Animation/ControlRig/Source/ControlRig/Public/ControlRig.h
+++ b/Engine/Plugins/Animation/ControlRig/Source/ControlRig/Public/ControlRig.h
@@ -302,9 +302,9 @@ public:
 	UFUNCTION(BlueprintCallable, Category = "Control Rig")
 	void RequestInit();
 
-	/** Requests to perform a setup during the next execution */
+	/** Requests to perform construction during the next execution */
 	UFUNCTION(BlueprintCallable, Category = "Control Rig")
-	void RequestSetup();
+	void RequestConstruction();
 
 	/** Returns the queue of events to run */
 	const TArray<FName>& GetEventQueue() const { return EventQueue; }
@@ -326,8 +326,8 @@ public:
 
 	virtual TArray<FRigControlElement*> AvailableControls() const;
 	virtual FRigControlElement* FindControl(const FName& InControlName) const;
-	virtual bool ShouldApplyLimits() const { return !IsSetupModeEnabled(); }
-	virtual bool IsSetupModeEnabled() const;
+	virtual bool ShouldApplyLimits() const { return !IsConstructionModeEnabled(); }
+	virtual bool IsConstructionModeEnabled() const;
 	virtual FTransform SetupControlFromGlobalTransform(const FName& InControlName, const FTransform& InGlobalTransform);
 	virtual FTransform GetControlGlobalTransform(const FName& InControlName) const;
 
@@ -424,8 +424,8 @@ public:
 
 	DECLARE_EVENT_ThreeParams(UControlRig, FControlRigExecuteEvent, class UControlRig*, const EControlRigState, const FName&);
 	FControlRigExecuteEvent& OnInitialized_AnyThread() { return InitializedEvent; }
-	FControlRigExecuteEvent& OnPreSetup_AnyThread() { return PreSetupEvent; }
-	FControlRigExecuteEvent& OnPostSetup_AnyThread() { return PostSetupEvent; }
+	FControlRigExecuteEvent& OnPreConstruction_AnyThread() { return PreConstructionEvent; }
+	FControlRigExecuteEvent& OnPostConstruction_AnyThread() { return PostConstructionEvent; }
 	FControlRigExecuteEvent& OnPreForwardsSolve_AnyThread() { return PreForwardsSolveEvent; }
 	FControlRigExecuteEvent& OnPostForwardsSolve_AnyThread() { return PostForwardsSolveEvent; }
 	FControlRigExecuteEvent& OnExecuted_AnyThread() { return ExecutedEvent; }
@@ -586,10 +586,10 @@ private:
 	FControlRigExecuteEvent InitializedEvent;
 
 	/** Broadcasts a notification just before the controlrig is setup. */
-	FControlRigExecuteEvent PreSetupEvent;
+	FControlRigExecuteEvent PreConstructionEvent;
 
 	/** Broadcasts a notification whenever the controlrig has been setup. */
-	FControlRigExecuteEvent PostSetupEvent;
+	FControlRigExecuteEvent PostConstructionEvent;
 
 	/** Broadcasts a notification before a forward solve has been initiated */
 	FControlRigExecuteEvent PreForwardsSolveEvent;
@@ -704,15 +704,15 @@ private:
 
 protected:
 	bool bRequiresInitExecution;
-	bool bRequiresSetupEvent;
-	bool bCopyHierarchyBeforeSetup;
-	bool bResetInitialTransformsBeforeSetup;
+	bool bRequiresConstructionEvent;
+	bool bCopyHierarchyBeforeConstruction;
+	bool bResetInitialTransformsBeforeConstruction;
 	bool bManipulationEnabled;
 
 	int32 InitBracket;
 	int32 UpdateBracket;
-	int32 PreSetupBracket;
-	int32 PostSetupBracket;
+	int32 PreConstructionBracket;
+	int32 PostConstructionBracket;
 	int32 PreForwardsSolveBracket;
 	int32 PostForwardsSolveBracket;
 	int32 InteractionBracket;
@@ -734,14 +734,14 @@ protected:
 		return UpdateBracket > 0;
 	}
 
-	FORCEINLINE bool IsRunningPreSetup() const
+	FORCEINLINE bool IsRunningPreConstruction() const
 	{
-		return PreSetupBracket > 0;
+		return PreConstructionBracket > 0;
 	}
 
-	FORCEINLINE bool IsRunningPostSetup() const
+	FORCEINLINE bool IsRunningPostConstruction() const
 	{
-		return PostSetupBracket > 0;
+		return PostConstructionBracket > 0;
 	}
 
 	FORCEINLINE bool IsInteracting() const
@@ -875,7 +875,7 @@ private:
 public:
 	// Class used to temporarily cache current pose of transient controls
 	// restore them after a ResetPoseToInitial call,
-	// which allows user to move bones in setup mode
+	// which allows user to move bones in construction mode
 	class FTransientControlPoseScope
 	{
 	public:
diff --git a/Engine/Plugins/Animation/ControlRig/Source/ControlRig/Public/ControlRigComponent.h b/Engine/Plugins/Animation/ControlRig/Source/ControlRig/Public/ControlRigComponent.h
--- a/Engine/Plugins/Animation/ControlRig/Source/ControlRig/Public/ControlRigComponent.h
+++ b/Engine/Plugins/Animation/ControlRig/Source/ControlRig/Public/ControlRigComponent.h
@@ -191,12 +191,12 @@ public:
 	FControlRigComponentDelegate OnPostInitializeDelegate;
 
 	/** Event fired before this component's ControlRig is setup */
-	UPROPERTY(BlueprintAssignable, Category = "ControlRig", meta = (DisplayName = "On Pre Setup"))
-	FControlRigComponentDelegate OnPreSetupDelegate;
+	UPROPERTY(BlueprintAssignable, Category = "ControlRig", meta = (DisplayName = "On Pre Construction"))
+	FControlRigComponentDelegate OnPreConstructionDelegate;
 
 	/** Event fired after this component's ControlRig is setup */
-	UPROPERTY(BlueprintAssignable, Category = "ControlRig", meta = (DisplayName = "On Post Setup"))
-	FControlRigComponentDelegate OnPostSetupDelegate;
+	UPROPERTY(BlueprintAssignable, Category = "ControlRig", meta = (DisplayName = "On Post Construction"))
+	FControlRigComponentDelegate OnPostConstructionDelegate;
 
 	/** Event fired before this component's ControlRig's forwards solve */
 	UPROPERTY(BlueprintAssignable, Category = "ControlRig", meta = (DisplayName = "On Pre Forwards Solve"))
@@ -247,10 +247,10 @@ public:
 	void OnPostInitialize(UControlRigComponent* Component);
 
 	UFUNCTION(BlueprintNativeEvent, Category = "ControlRig", CallInEditor)
-	void OnPreSetup(UControlRigComponent* Component);
+	void OnPreConstruction(UControlRigComponent* Component);
 
 	UFUNCTION(BlueprintNativeEvent, Category = "ControlRig", CallInEditor)
-	void OnPostSetup(UControlRigComponent* Component);
+	void OnPostConstruction(UControlRigComponent* Component);
 
 	UFUNCTION(BlueprintNativeEvent, Category = "ControlRig", CallInEditor, meta = (DisplayName = "On Pre Forwards Solve"))
 	void OnPreForwardsSolve(UControlRigComponent* Component);
@@ -258,7 +258,7 @@ public:
 	UFUNCTION(BlueprintNativeEvent, Category = "ControlRig", CallInEditor, meta = (DisplayName = "On Post Forwards Solve"))
 	void OnPostForwardsSolve(UControlRigComponent* Component);
 
-	/** Initializes the rig's memory and calls the setup event */
+	/** Initializes the rig's memory and calls the construction event */
 	UFUNCTION(BlueprintCallable, Category = "ControlRig")
 	void Initialize();
 
@@ -320,9 +320,9 @@ public:
 	UPROPERTY(EditAnywhere, Category = "Animation")
 	bool bResetTransformBeforeTick;
 
-	/** When checked the initial transforms on bones, nulls and controls are reset prior to a setup event */
+	/** When checked the initial transforms on bones, nulls and controls are reset prior to a construction event */
 	UPROPERTY(EditAnywhere, Category = "Animation")
-	bool bResetInitialsBeforeSetup;
+	bool bResetInitialsBeforeConstruction;
 
 	/** When checked this ensures to run the rig's update on the component's tick automatically */
 	UPROPERTY(EditAnywhere, Category = "Animation")
@@ -643,8 +643,8 @@ private:
 	static FName GetComponentNameWithinActor(UActorComponent* InComponent);
 
 	void HandleControlRigInitializedEvent(UControlRig* InControlRig, const EControlRigState InState, const FName& InEventName);
-	void HandleControlRigPreSetupEvent(UControlRig* InControlRig, const EControlRigState InState, const FName& InEventName);
-	void HandleControlRigPostSetupEvent(UControlRig* InControlRig, const EControlRigState InState, const FName& InEventName);
+	void HandleControlRigPreConstructionEvent(UControlRig* InControlRig, const EControlRigState InState, const FName& InEventName);
+	void HandleControlRigPostConstructionEvent(UControlRig* InControlRig, const EControlRigState InState, const FName& InEventName);
 	void HandleControlRigPreForwardsSolveEvent(UControlRig* InControlRig, const EControlRigState InState, const FName& InEventName);
 	void HandleControlRigPostForwardsSolveEvent(UControlRig* InControlRig, const EControlRigState InState, const FName& InEventName);
 	void HandleControlRigExecutedEvent(UControlRig* InControlRig, const EControlRigState InState, const FName& InEventName);
diff --git a/Engine/Plugins/Animation/ControlRig/Source/ControlRig/Public/Rigs/RigHierarchyDefines.h b/Engine/Plugins/Animation/ControlRig/Source/ControlRig/Public/Rigs/RigHierarchyDefines.h
--- a/Engine/Plugins/Animation/ControlRig/Source/ControlRig/Public/Rigs/RigHierarchyDefines.h
+++ b/Engine/Plugins/Animation/ControlRig/Source/ControlRig/Public/Rigs/RigHierarchyDefines.h
@@ -30,7 +30,7 @@ enum class ERigElementType : uint8
 	Reference = 0x020,
 	Last = 0x040 UMETA(Hidden),
 	All = Bone | Null | Control | Curve | RigidBody | Reference,
-	ToResetAfterSetupEvent = Bone | Control | Curve UMETA(Hidden),
+	ToResetAfterConstructionEvent = Bone | Control | Curve UMETA(Hidden),
 };
 
 UENUM(BlueprintType)
diff --git a/Engine/Plugins/Animation/ControlRig/Source/ControlRig/Public/Settings/ControlRigSettings.h b/Engine/Plugins/Animation/ControlRig/Source/ControlRig/Public/Settings/ControlRigSettings.h
--- a/Engine/Plugins/Animation/ControlRig/Source/ControlRig/Public/Settings/ControlRigSettings.h
+++ b/Engine/Plugins/Animation/ControlRig/Source/ControlRig/Public/Settings/ControlRigSettings.h
@@ -93,10 +93,10 @@ public:
 	TMap<FString, FControlRigSettingsPerPinBool> RigUnitPinExpansion;
 	
 	/**
-	 * The border color of the viewport when entering "Setup Event" mode
+	 * The border color of the viewport when entering "Construction Event" mode
 	 */
 	UPROPERTY(EditAnywhere, config, Category = Viewport)
-	FLinearColor SetupEventBorderColor;
+	FLinearColor ConstructionEventBorderColor;
 	
 	/**
 	 * The border color of the viewport when entering "Backwards Solve" mode
diff --git a/Engine/Plugins/Animation/ControlRig/Source/ControlRigEditor/Private/ControlRigBlueprintCommands.cpp b/Engine/Plugins/Animation/ControlRig/Source/ControlRigEditor/Private/ControlRigBlueprintCommands.cpp
--- a/Engine/Plugins/Animation/ControlRig/Source/ControlRigEditor/Private/ControlRigBlueprintCommands.cpp
+++ b/Engine/Plugins/Animation/ControlRig/Source/ControlRigEditor/Private/ControlRigBlueprintCommands.cpp
@@ -10,7 +10,7 @@ void FControlRigBlueprintCommands::RegisterCommands()
 	UI_COMMAND(ExecuteGraph, "Execute", "Execute the rig graph if On.", EUserInterfaceActionType::ToggleButton, FInputChord());
 	UI_COMMAND(AutoCompileGraph, "Auto Compile", "Auto-compile the rig graph if On.", EUserInterfaceActionType::ToggleButton, FInputChord());
 	UI_COMMAND(ToggleEventQueue, "Toggle Event", "Toggle between the current and last running event", EUserInterfaceActionType::Button, FInputChord());
-	UI_COMMAND(SetupEvent, "Setup Event", "Enable the setup mode for the rig", EUserInterfaceActionType::Button, FInputChord());
+	UI_COMMAND(ConstructionEvent, "Construction Event", "Enable the construction mode for the rig", EUserInterfaceActionType::Button, FInputChord());
 	UI_COMMAND(ForwardsSolveEvent, "Forwards Solve", "Run the forwards solve graph", EUserInterfaceActionType::Button, FInputChord());
 	UI_COMMAND(BackwardsSolveEvent, "Backwards Solve", "Run the backwards solve graph", EUserInterfaceActionType::Button, FInputChord());
 	UI_COMMAND(BackwardsAndForwardsSolveEvent, "Backwards and Forwards", "Run backwards solve followed by forwards solve", EUserInterfaceActionType::Button, FInputChord());
diff --git a/Engine/Plugins/Animation/ControlRig/Source/ControlRigEditor/Private/ControlRigBlueprintCommands.h b/Engine/Plugins/Animation/ControlRig/Source/ControlRigEditor/Private/ControlRigBlueprintCommands.h
--- a/Engine/Plugins/Animation/ControlRig/Source/ControlRigEditor/Private/ControlRigBlueprintCommands.h
+++ b/Engine/Plugins/Animation/ControlRig/Source/ControlRigEditor/Private/ControlRigBlueprintCommands.h
@@ -30,8 +30,8 @@ public:
 	/** Toggle between this and the last event queue */
 	TSharedPtr< FUICommandInfo > ToggleEventQueue;
 
-	/** Enable the setup mode for the rig */
-	TSharedPtr< FUICommandInfo > SetupEvent;
+	/** Enable the construction mode for the rig */
+	TSharedPtr< FUICommandInfo > ConstructionEvent;
 
 	/** Run the forwards solve graph */
 	TSharedPtr< FUICommandInfo > ForwardsSolveEvent;
diff --git a/Engine/Plugins/Animation/ControlRig/Source/ControlRigEditor/Private/ControlRigEditorStyle.h b/Engine/Plugins/Animation/ControlRig/Source/ControlRigEditor/Private/ControlRigEditorStyle.h
--- a/Engine/Plugins/Animation/ControlRig/Source/ControlRigEditor/Private/ControlRigEditorStyle.h
+++ b/Engine/Plugins/Animation/ControlRig/Source/ControlRigEditor/Private/ControlRigEditorStyle.h
@@ -104,8 +104,8 @@ public:
 
 			{
 				FContentRootBracket Bracket(this, EngineEditorSlateDir);
-				Set("ControlRig.SetupMode", new IMAGE_BRUSH_SVG("Starship/Common/Adjust", Icon40x40));
-				Set("ControlRig.SetupMode.Small", new IMAGE_BRUSH_SVG("Starship/Common/Adjust", Icon20x20));
+				Set("ControlRig.ConstructionMode", new IMAGE_BRUSH_SVG("Starship/Common/Adjust", Icon40x40));
+				Set("ControlRig.ConstructionMode.Small", new IMAGE_BRUSH_SVG("Starship/Common/Adjust", Icon20x20));
 				Set("ControlRig.ForwardsSolveEvent", new IMAGE_BRUSH("Icons/diff_next_40x", Icon40x40));
 				Set("ControlRig.BackwardsSolveEvent", new IMAGE_BRUSH("Icons/diff_prev_40x", Icon40x40));
 				Set("ControlRig.BackwardsAndForwardsSolveEvent", new IMAGE_BRUSH("Icons/Loop_40x", Icon40x40));
diff --git a/Engine/Plugins/Animation/ControlRig/Source/ControlRigEditor/Private/ControlRigElementDetails.cpp b/Engine/Plugins/Animation/ControlRig/Source/ControlRigEditor/Private/ControlRigElementDetails.cpp
--- a/Engine/Plugins/Animation/ControlRig/Source/ControlRigEditor/Private/ControlRigElementDetails.cpp
+++ b/Engine/Plugins/Animation/ControlRig/Source/ControlRigEditor/Private/ControlRigElementDetails.cpp
@@ -630,7 +630,7 @@ URigHierarchy* FRigBaseElementDetails::GetHierarchyBeingDebugged() const
 		
 	if (UControlRig* DebuggedRig = Cast<UControlRig>(BlueprintBeingCustomized->GetObjectBeingDebugged()))
 	{
-		if(!DebuggedRig->IsSetupModeEnabled())
+		if(!DebuggedRig->IsConstructionModeEnabled())
 		{
 			HierarchyBeingDebugged = DebuggedRig->GetHierarchy();
 		}
@@ -742,13 +742,13 @@ void FRigBaseElementDetails::OnStructContentsChanged(FProperty* InProperty, cons
 	PropertyUtilities->NotifyFinishedChangingProperties(ChangeEvent);
 }
 
-bool FRigBaseElementDetails::IsSetupModeEnabled() const
+bool FRigBaseElementDetails::IsConstructionModeEnabled() const
 {
 	if(BlueprintBeingCustomized)
 	{
 		if (UControlRig* DebuggedRig = Cast<UControlRig>(BlueprintBeingCustomized->GetObjectBeingDebugged()))
 		{
-			return DebuggedRig->IsSetupModeEnabled();
+			return DebuggedRig->IsConstructionModeEnabled();
 		}
 	}
 	return false;
@@ -935,7 +935,7 @@ void FRigTransformElementDetails::CustomizeTransform(IDetailLayoutBuilder& Detai
 	URigHierarchy* HierarchyBeingDebugged = HierarchyBeingCustomized;
 	if (UControlRig* DebuggedRig = Cast<UControlRig>(BlueprintBeingCustomized->GetObjectBeingDebugged()))
 	{
- 		if(!DebuggedRig->IsSetupModeEnabled())
+ 		if(!DebuggedRig->IsConstructionModeEnabled())
 		{
 			HierarchyBeingDebugged = DebuggedRig->GetHierarchy();
 		}
@@ -1741,15 +1741,15 @@ FDetailWidgetRow& FRigTransformElementDetails::CreateEulerTransformValueWidgetRo
 		const bool bCurrent = CurrentTransformType == ERigTransformElementDetailsTransform::Current; 
 		const bool bInitial = CurrentTransformType == ERigTransformElementDetailsTransform::Initial; 
 
-		bool bSetupModeEnabled = false;
+		bool bConstructionModeEnabled = false;
 		if (UControlRig* DebuggedRig = Cast<UControlRig>(BlueprintBeingCustomized->GetObjectBeingDebugged()))
 		{
-			bSetupModeEnabled = DebuggedRig->IsSetupModeEnabled();
+			bConstructionModeEnabled = DebuggedRig->IsConstructionModeEnabled();
 		}
 
 		TArray<URigHierarchy*> HierarchiesToUpdate;
 		HierarchiesToUpdate.Add(HierarchyBeingDebugged);
-		if(!bCurrent || bSetupModeEnabled)
+		if(!bCurrent || bConstructionModeEnabled)
 		{
 			HierarchiesToUpdate.Add(HierarchyBeingCustomized);
 		}
diff --git a/Engine/Plugins/Animation/ControlRig/Source/ControlRigEditor/Private/ControlRigElementDetails.h b/Engine/Plugins/Animation/ControlRig/Source/ControlRigEditor/Private/ControlRigElementDetails.h
--- a/Engine/Plugins/Animation/ControlRig/Source/ControlRigEditor/Private/ControlRigElementDetails.h
+++ b/Engine/Plugins/Animation/ControlRig/Source/ControlRigEditor/Private/ControlRigElementDetails.h
@@ -123,7 +123,7 @@ public:
 	bool OnVerifyNameChanged(const FText& InText, FText& OutErrorMessage);
 
 	void OnStructContentsChanged(FProperty* InProperty, const TSharedRef<IPropertyUtilities> PropertyUtilities);
-	bool IsSetupModeEnabled() const;
+	bool IsConstructionModeEnabled() const;
 
 	FText GetParentElementName() const;
 
diff --git a/Engine/Plugins/Animation/ControlRig/Source/ControlRigEditor/Private/EditMode/ControlRigEditMode.cpp b/Engine/Plugins/Animation/ControlRig/Source/ControlRigEditor/Private/EditMode/ControlRigEditMode.cpp
--- a/Engine/Plugins/Animation/ControlRig/Source/ControlRigEditor/Private/EditMode/ControlRigEditMode.cpp
+++ b/Engine/Plugins/Animation/ControlRig/Source/ControlRigEditor/Private/EditMode/ControlRigEditMode.cpp
@@ -833,7 +833,7 @@ void FControlRigEditMode::Render(const FSceneView* View, FViewport* Viewport, FP
 						});
 				}
 
-				if (Settings->bDisplayNulls || ControlRig->IsSetupModeEnabled())
+				if (Settings->bDisplayNulls || ControlRig->IsConstructionModeEnabled())
 				{
 					TArray<FTransform> SpaceTransforms;
 					TArray<FTransform> SelectedSpaceTransforms;
diff --git a/Engine/Plugins/Animation/ControlRig/Source/ControlRigEditor/Private/Editor/ControlRigEditor.cpp b/Engine/Plugins/Animation/ControlRig/Source/ControlRigEditor/Private/Editor/ControlRigEditor.cpp
--- a/Engine/Plugins/Animation/ControlRig/Source/ControlRigEditor/Private/Editor/ControlRigEditor.cpp
+++ b/Engine/Plugins/Animation/ControlRig/Source/ControlRigEditor/Private/Editor/ControlRigEditor.cpp
@@ -207,7 +207,7 @@ struct FControlRigZoomLevelsContainer : public FZoomLevelsContainer
 
 const TArray<FName> FControlRigEditor::ForwardsSolveEventQueue = {FRigUnit_BeginExecution::EventName};
 const TArray<FName> FControlRigEditor::BackwardsSolveEventQueue = {FRigUnit_InverseExecution::EventName};
-const TArray<FName> FControlRigEditor::SetupEventQueue = {FRigUnit_PrepareForExecution::EventName};
+const TArray<FName> FControlRigEditor::ConstructionEventQueue = {FRigUnit_PrepareForExecution::EventName};
 const TArray<FName> FControlRigEditor::BackwardsAndForwardsSolveEventQueue = {FRigUnit_InverseExecution::EventName, FRigUnit_BeginExecution::EventName};
 
 FControlRigEditor::FControlRigEditor()
@@ -218,7 +218,7 @@ FControlRigEditor::FControlRigEditor()
 	, bExecutionControlRig(true)
 	, bIsCompilingThroughUI(false)
 	, bAnyErrorsLeft(false)
-	, LastEventQueue(SetupEventQueue)
+	, LastEventQueue(ConstructionEventQueue)
 	, ExecutionMode(EControlRigExecutionModeType::EControlRigExecutionModeType_Release)
 	, LastDebuggedRig()
 	, RigHierarchyTabCount(0)
@@ -660,8 +660,8 @@ void FControlRigEditor::BindCommands()
 		FCanExecuteAction());
 
 	GetToolkitCommands()->MapAction(
-		FControlRigBlueprintCommands::Get().SetupEvent,
-		FExecuteAction::CreateSP(this, &FControlRigEditor::SetEventQueue, TArray<FName>(SetupEventQueue)),
+		FControlRigBlueprintCommands::Get().ConstructionEvent,
+		FExecuteAction::CreateSP(this, &FControlRigEditor::SetEventQueue, TArray<FName>(ConstructionEventQueue)),
 		FCanExecuteAction());
 
 	GetToolkitCommands()->MapAction(
@@ -781,7 +781,7 @@ TSharedRef<SWidget> FControlRigEditor::GenerateEventQueueMenuContent()
 	FMenuBuilder MenuBuilder(true, GetToolkitCommands());
 
 	MenuBuilder.BeginSection(TEXT("Events"));
-	MenuBuilder.AddMenuEntry(FControlRigBlueprintCommands::Get().SetupEvent, TEXT("Setup"), TAttribute<FText>(), TAttribute<FText>(), GetEventQueueIcon(SetupEventQueue));
+	MenuBuilder.AddMenuEntry(FControlRigBlueprintCommands::Get().ConstructionEvent, TEXT("Setup"), TAttribute<FText>(), TAttribute<FText>(), GetEventQueueIcon(ConstructionEventQueue));
 	MenuBuilder.AddMenuEntry(FControlRigBlueprintCommands::Get().ForwardsSolveEvent, TEXT("Forwards Solve"), TAttribute<FText>(), TAttribute<FText>(), GetEventQueueIcon(ForwardsSolveEventQueue));
 	MenuBuilder.AddMenuEntry(FControlRigBlueprintCommands::Get().BackwardsSolveEvent, TEXT("Backwards Solve"), TAttribute<FText>(), TAttribute<FText>(), GetEventQueueIcon(BackwardsSolveEventQueue));
 	MenuBuilder.EndSection();
@@ -1011,7 +1011,7 @@ void FControlRigEditor::SetEventQueue(TArray<FName> InEventQueue, bool bCompile)
 			}
 		}
 
-		// Reset transforms only for setup and forward solve to not inturrupt any animation that might be playing
+		// Reset transforms only for construction and forward solve to not inturrupt any animation that might be playing
 		if (InEventQueue.Contains(FRigUnit_PrepareForExecution::EventName) ||
 			InEventQueue.Contains(FRigUnit_BeginExecution::EventName))
 		{
@@ -1027,7 +1027,7 @@ void FControlRigEditor::SetEventQueue(TArray<FName> InEventQueue, bool bCompile)
 		}
 
 		UControlRigEditModeSettings* Settings = GetMutableDefault<UControlRigEditModeSettings>();
-		Settings->bDisplayNulls = IsSetupModeEnabled();
+		Settings->bDisplayNulls = IsConstructionModeEnabled();
 	}
 
 	if (PreviousSelection.Num() > 0)
@@ -1047,7 +1047,7 @@ int32 FControlRigEditor::GetEventQueueComboValue() const
 	{
 		return 0;
 	}
-	if(EventQueue == SetupEventQueue)
+	if(EventQueue == ConstructionEventQueue)
 	{
 		return 1;
 	}
@@ -1066,7 +1066,7 @@ FText FControlRigEditor::GetEventQueueLabel() const
 {
 	TArray<FName> EventQueue = GetEventQueue();
 
-	if(EventQueue == SetupEventQueue)
+	if(EventQueue == ConstructionEventQueue)
 	{
 		return FRigUnit_PrepareForExecution::StaticStruct()->GetDisplayNameText();
 	}
@@ -1099,9 +1099,9 @@ FText FControlRigEditor::GetEventQueueLabel() const
 
 FSlateIcon FControlRigEditor::GetEventQueueIcon(const TArray<FName>& InEventQueue)
 {
-	if(InEventQueue == SetupEventQueue)
+	if(InEventQueue == ConstructionEventQueue)
 	{
-		return FSlateIcon(FControlRigEditorStyle::Get().GetStyleSetName(), "ControlRig.SetupMode");
+		return FSlateIcon(FControlRigEditorStyle::Get().GetStyleSetName(), "ControlRig.ConstructionMode");
 	}
 	if(InEventQueue == ForwardsSolveEventQueue)
 	{
@@ -1578,7 +1578,7 @@ void FControlRigEditor::SetDetailViewForRigElements()
 	UControlRigBlueprint* RigBlueprint = Cast<UControlRigBlueprint>(GetBlueprintObj());
 	URigHierarchy* Hierarchy = RigBlueprint->Hierarchy;
 	
-	if (!IsSetupModeEnabled())
+	if (!IsConstructionModeEnabled())
 	{
 		if (UControlRig* DebuggedControlRig = Cast<UControlRig>(RigBlueprint->GetObjectBeingDebugged()))
 		{
@@ -1916,7 +1916,7 @@ void FControlRigEditor::Compile()
 			}
 		}
 
-		if(IsSetupModeEnabled())
+		if(IsConstructionModeEnabled())
 		{
 			SetEventQueue(ForwardsSolveEventQueue, false);
 		}
@@ -2916,7 +2916,7 @@ void FControlRigEditor::HandleControlRigExecutedEvent(UControlRig* InControlRig,
 		}
 
 		URigHierarchy* Hierarchy = ControlRigBP->Hierarchy; 
-		if(!IsSetupModeEnabled())
+		if(!IsConstructionModeEnabled())
 		{
 			if(DebuggedControlRig)
 			{
@@ -3657,9 +3657,9 @@ void FControlRigEditor::HandleViewportCreated(const TSharedRef<class IPersonaVie
 	{
 		FLinearColor Color = FLinearColor::Transparent;
 		const TArray<FName> EventQueue = GetEventQueue();
-		if(EventQueue == SetupEventQueue)
+		if(EventQueue == ConstructionEventQueue)
 		{
-			Color = UControlRigEditorSettings::Get()->SetupEventBorderColor;
+			Color = UControlRigEditorSettings::Get()->ConstructionEventBorderColor;
 		}
 		if(EventQueue == BackwardsSolveEventQueue)
 		{
@@ -3762,7 +3762,7 @@ bool FControlRigEditor::IsToolbarDrawNullsEnabled() const
 {
 	if (ControlRig)
 	{
-		if (!ControlRig->IsSetupModeEnabled())
+		if (!ControlRig->IsConstructionModeEnabled())
 		{
 			return true;
 		}
@@ -3920,7 +3920,7 @@ void FControlRigEditor::OnPinControlNameListComboBox(const TArray<TSharedPtr<FSt
 	PinControlNameList->SetSelectedItem(CurrentlySelected);
 }
 
-bool FControlRigEditor::IsSetupModeEnabled() const
+bool FControlRigEditor::IsConstructionModeEnabled() const
 {
 	return LastEventQueue.Contains(FRigUnit_PrepareForExecution::EventName);
 }
@@ -4526,7 +4526,7 @@ void FControlRigEditor::OnWrappedPropertyChangedChainEvent(UDetailsViewWrapperOb
 				const ERigTransformType::Type TransformType = Local::GetTransformTypeFromPath(PropertyPath);
 				bIsInitial = bIsInitial || ERigTransformType::IsInitial(TransformType);
 				
-				if(ERigTransformType::IsInitial(TransformType) || IsSetupModeEnabled())
+				if(ERigTransformType::IsInitial(TransformType) || IsConstructionModeEnabled())
 				{
 					Hierarchy = ControlRigBP->Hierarchy;
 				}
@@ -4546,7 +4546,7 @@ void FControlRigEditor::OnWrappedPropertyChangedChainEvent(UDetailsViewWrapperOb
 					FRigControlValue Value;
 					Value.SetFromTransform(Transform, ControlElement->Settings.ControlType, ControlElement->Settings.PrimaryAxis);
 							
-					if(ERigTransformType::IsInitial(TransformType) || IsSetupModeEnabled())
+					if(ERigTransformType::IsInitial(TransformType) || IsConstructionModeEnabled())
 					{
 						Hierarchy->SetControlValue(ControlElement, Value, ERigControlValueType::Initial, true, true, true);
 					}
@@ -4600,7 +4600,7 @@ void FControlRigEditor::OnWrappedPropertyChangedChainEvent(UDetailsViewWrapperOb
 				ControlRigBP->Hierarchy->SetControlSettings(ControlElement, Settings, true, false, true);
 			}
 
-			if(IsSetupModeEnabled() || bIsInitial)
+			if(IsConstructionModeEnabled() || bIsInitial)
 			{
 				ControlRigBP->PropagatePoseFromBPToInstances();
 				ControlRigBP->Modify();
@@ -6251,7 +6251,7 @@ void FControlRigEditor::HandleOnControlModified(UControlRig* Subject, FRigContro
 				Blueprint->Hierarchy->SetLocalTransform(ElementKey, Transform);
 				Hierarchy->SetLocalTransform(ElementKey, Transform);
 
-				if (IsSetupModeEnabled())
+				if (IsConstructionModeEnabled())
 				{
 					Blueprint->Hierarchy->SetInitialLocalTransform(ElementKey, Transform);
 					Hierarchy->SetInitialLocalTransform(ElementKey, Transform);
@@ -6271,7 +6271,7 @@ void FControlRigEditor::HandleOnControlModified(UControlRig* Subject, FRigContro
 			}
 		}
 	}
-	else if (IsSetupModeEnabled())
+	else if (IsConstructionModeEnabled())
 	{
 		FRigControlElement* SourceControlElement = Hierarchy->Find<FRigControlElement>(ControlElement->GetKey());
 		FRigControlElement* TargetControlElement = Blueprint->Hierarchy->Find<FRigControlElement>(ControlElement->GetKey());
@@ -6284,7 +6284,7 @@ void FControlRigEditor::HandleOnControlModified(UControlRig* Subject, FRigContro
 			Blueprint->Hierarchy->OnModified().Broadcast(ERigHierarchyNotification::ControlSettingChanged, Blueprint->Hierarchy, TargetControlElement);
 		}
 
-		// we copy the pose including the weights since we want the topology to align during setup mode.
+		// we copy the pose including the weights since we want the topology to align during construction mode.
 		// i.e. dynamic reparenting should be reset here.
 		TargetControlElement->CopyPose(SourceControlElement, true, true, true);
 	}
diff --git a/Engine/Plugins/Animation/ControlRig/Source/ControlRigEditor/Private/Editor/ControlRigEditor.h b/Engine/Plugins/Animation/ControlRig/Source/ControlRigEditor/Private/Editor/ControlRigEditor.h
--- a/Engine/Plugins/Animation/ControlRig/Source/ControlRigEditor/Private/Editor/ControlRigEditor.h
+++ b/Engine/Plugins/Animation/ControlRig/Source/ControlRigEditor/Private/Editor/ControlRigEditor.h
@@ -464,7 +464,7 @@ protected:
 	void SetPinControlNameListText(const FText& NewTypeInValue, ETextCommit::Type /*CommitInfo*/);
 	void OnPinControlNameListChanged(TSharedPtr<FString> NewSelection, ESelectInfo::Type SelectInfo);
 	void OnPinControlNameListComboBox(const TArray<TSharedPtr<FString>>* InNameList);
-	bool IsSetupModeEnabled() const;
+	bool IsConstructionModeEnabled() const;
 
 	bool bAnyErrorsLeft;
 
@@ -485,7 +485,7 @@ protected:
 
 	static const TArray<FName> ForwardsSolveEventQueue;
 	static const TArray<FName> BackwardsSolveEventQueue;
-	static const TArray<FName> SetupEventQueue;
+	static const TArray<FName> ConstructionEventQueue;
 	static const TArray<FName> BackwardsAndForwardsSolveEventQueue;
 
 	friend class FControlRigEditorMode;
diff --git a/Engine/Plugins/Animation/ControlRig/Source/ControlRigEditor/Private/Editor/SRigHierarchy.cpp b/Engine/Plugins/Animation/ControlRig/Source/ControlRigEditor/Private/Editor/SRigHierarchy.cpp
--- a/Engine/Plugins/Animation/ControlRig/Source/ControlRigEditor/Private/Editor/SRigHierarchy.cpp
+++ b/Engine/Plugins/Animation/ControlRig/Source/ControlRigEditor/Private/Editor/SRigHierarchy.cpp
@@ -1264,13 +1264,13 @@ void SRigHierarchy::RefreshHierarchy(const FAssetData& InAssetData)
 
 		FScopedTransaction Transaction(LOCTEXT("HierarchyRefresh", "Refresh Transform"));
 
-		// don't select bone if we are in setup mode.
+		// don't select bone if we are in construction mode.
 		// we do this to avoid the editmode / viewport shapes to refresh recursively,
 		// which can add an extreme slowdown depending on the number of bones (n^(n-1))
 		bool bSelectBones = true;
 		if (UControlRig* CurrentRig = ControlRigEditor.Pin()->ControlRig)
 		{
-			bSelectBones = !CurrentRig->IsSetupModeEnabled();
+			bSelectBones = !CurrentRig->IsConstructionModeEnabled();
 		}
 
 		const FReferenceSkeleton& RefSkeleton = Mesh->GetRefSkeleton();
@@ -1366,13 +1366,13 @@ void SRigHierarchy::ImportHierarchy(const FAssetData& InAssetData)
 
 		FScopedTransaction Transaction(LOCTEXT("HierarchyImport", "Import Hierarchy"));
 
-		// don't select bone if we are in setup mode.
+		// don't select bone if we are in construction mode.
 		// we do this to avoid the editmode / viewport shapes to refresh recursively,
 		// which can add an extreme slowdown depending on the number of bones (n^(n-1))
 		bool bSelectBones = true;
 		if (UControlRig* CurrentRig = ControlRigEditor.Pin()->ControlRig)
 		{
-			bSelectBones = !CurrentRig->IsSetupModeEnabled();
+			bSelectBones = !CurrentRig->IsConstructionModeEnabled();
 		}
 
 		URigHierarchyController* Controller = Hierarchy->GetController(true);
