ControlRig: Rename Attribute nodes to Get/Set Anim Attribute

diff --git a/Engine/Plugins/Animation/ControlRig/Source/ControlRig/Private/Units/Animation/RigUnit_AnimAttribute.h b/Engine/Plugins/Animation/ControlRig/Source/ControlRig/Private/Units/Animation/RigUnit_AnimAttribute.h
--- a/Engine/Plugins/Animation/ControlRig/Source/ControlRig/Private/Units/Animation/RigUnit_AnimAttribute.h
+++ b/Engine/Plugins/Animation/ControlRig/Source/ControlRig/Private/Units/Animation/RigUnit_AnimAttribute.h
@@ -11,14 +11,14 @@
  * Animation Attributes allow dynamically added data to flow from
  * one Anim Node to other Anim Nodes downstream in the Anim Graph.
  */
-USTRUCT(meta=(Abstract, Category="Animation Attribute", TemplateName="Set Attribute", NodeColor = "0.0 0.36470600962638855 1.0", Varying))
+USTRUCT(meta=(Abstract, Category="Animation Attribute", TemplateName="Set Anim Attribute", NodeColor = "0.0 0.36470600962638855 1.0", Varying))
 struct FRigUnit_SetAnimAttributeBase : public FRigUnitMutable
 {
 	GENERATED_BODY()
 };
 
 
-USTRUCT(meta=(DisplayName="Set Attribute - Integer"))
+USTRUCT(meta=(DisplayName="Set Anim Attribute - Integer"))
 struct CONTROLRIG_API FRigUnit_SetAnimAttribute_Integer: public FRigUnit_SetAnimAttributeBase
 {
 	GENERATED_BODY()
@@ -52,7 +52,7 @@ struct CONTROLRIG_API FRigUnit_SetAnimAttribute_Integer: public FRigUnit_SetAnim
 };
 
 
-USTRUCT(meta=(DisplayName="Set Attribute - Float"))
+USTRUCT(meta=(DisplayName="Set Anim Attribute - Float"))
 struct CONTROLRIG_API FRigUnit_SetAnimAttribute_Float: public FRigUnit_SetAnimAttributeBase
 {
 	GENERATED_BODY()
@@ -86,7 +86,7 @@ struct CONTROLRIG_API FRigUnit_SetAnimAttribute_Float: public FRigUnit_SetAnimAt
 };
 
 
-USTRUCT(meta=(DisplayName="Set Attribute - Transform"))
+USTRUCT(meta=(DisplayName="Set Anim Attribute - Transform"))
 struct CONTROLRIG_API FRigUnit_SetAnimAttribute_Transform: public FRigUnit_SetAnimAttributeBase
 {
 	GENERATED_BODY()
@@ -119,7 +119,7 @@ struct CONTROLRIG_API FRigUnit_SetAnimAttribute_Transform: public FRigUnit_SetAn
 	int32 CachedBoneIndex;
 };
 
-USTRUCT(meta=(DisplayName="Set Attribute - Vector"))
+USTRUCT(meta=(DisplayName="Set Anim Attribute - Vector"))
 struct CONTROLRIG_API FRigUnit_SetAnimAttribute_Vector: public FRigUnit_SetAnimAttributeBase
 {
 	GENERATED_BODY()
@@ -152,7 +152,7 @@ struct CONTROLRIG_API FRigUnit_SetAnimAttribute_Vector: public FRigUnit_SetAnimA
 	int32 CachedBoneIndex;
 };
 
-USTRUCT(meta=(DisplayName="Set Attribute - Quaternion"))
+USTRUCT(meta=(DisplayName="Set Anim Attribute - Quaternion"))
 struct CONTROLRIG_API FRigUnit_SetAnimAttribute_Quaternion: public FRigUnit_SetAnimAttributeBase
 {
 	GENERATED_BODY()
@@ -192,14 +192,14 @@ struct CONTROLRIG_API FRigUnit_SetAnimAttribute_Quaternion: public FRigUnit_SetA
  * Animation Attributes allow dynamically added data to flow from
  * one Anim Node to other Anim Nodes downstream in the Anim Graph.
  */
-USTRUCT(meta=(Abstract, Category="Animation Attribute", TemplateName="Get Attribute", NodeColor="0.462745, 1,0, 0.329412", Varying))
+USTRUCT(meta=(Abstract, Category="Animation Attribute", TemplateName="Get Anim Attribute", NodeColor="0.462745, 1,0, 0.329412", Varying))
 struct FRigUnit_GetAnimAttributeBase : public FRigUnit
 {
 	GENERATED_BODY()
 };
 
 
-USTRUCT(meta=(DisplayName="Get Attribute - Integer"))
+USTRUCT(meta=(DisplayName="Get Anim Attribute - Integer"))
 struct CONTROLRIG_API FRigUnit_GetAnimAttribute_Integer: public FRigUnit_GetAnimAttributeBase
 {
 	GENERATED_BODY()
@@ -241,7 +241,7 @@ struct CONTROLRIG_API FRigUnit_GetAnimAttribute_Integer: public FRigUnit_GetAnim
 };
 
 
-USTRUCT(meta=(DisplayName="Get Attribute - Float"))
+USTRUCT(meta=(DisplayName="Get Anim Attribute - Float"))
 struct CONTROLRIG_API FRigUnit_GetAnimAttribute_Float: public FRigUnit_GetAnimAttributeBase
 {
 	GENERATED_BODY()
@@ -283,7 +283,7 @@ struct CONTROLRIG_API FRigUnit_GetAnimAttribute_Float: public FRigUnit_GetAnimAt
 };
 
 
-USTRUCT(meta=(DisplayName="Get Attribute - Transform"))
+USTRUCT(meta=(DisplayName="Get Anim Attribute - Transform"))
 struct CONTROLRIG_API FRigUnit_GetAnimAttribute_Transform: public FRigUnit_GetAnimAttributeBase
 {
 	GENERATED_BODY()
@@ -325,7 +325,7 @@ struct CONTROLRIG_API FRigUnit_GetAnimAttribute_Transform: public FRigUnit_GetAn
 };
 
 
-USTRUCT(meta=(DisplayName="Get Attribute - Vector"))
+USTRUCT(meta=(DisplayName="Get Anim Attribute - Vector"))
 struct CONTROLRIG_API FRigUnit_GetAnimAttribute_Vector: public FRigUnit_GetAnimAttributeBase
 {
 	GENERATED_BODY()
@@ -367,7 +367,7 @@ struct CONTROLRIG_API FRigUnit_GetAnimAttribute_Vector: public FRigUnit_GetAnimA
 };
 
 
-USTRUCT(meta=(DisplayName="Get Attribute - Quaternion"))
+USTRUCT(meta=(DisplayName="Get Anim Attribute - Quaternion"))
 struct CONTROLRIG_API FRigUnit_GetAnimAttribute_Quaternion: public FRigUnit_GetAnimAttributeBase
 {
 	GENERATED_BODY()
