Sequencer: Anim Tools: New Slider Tools.
Built model framework to create system to blend between keys ore objects/controls based upon seelction. Created view to handle pick which slider to use then blend it.

diff --git a/Engine/Plugins/Animation/ControlRig/Source/ControlRigEditor/Private/EditMode/SControlRigTweenWidget.cpp b/Engine/Plugins/Animation/ControlRig/Source/ControlRigEditor/Private/EditMode/SControlRigTweenWidget.cpp
--- a/Engine/Plugins/Animation/ControlRig/Source/ControlRigEditor/Private/EditMode/SControlRigTweenWidget.cpp
+++ b/Engine/Plugins/Animation/ControlRig/Source/ControlRigEditor/Private/EditMode/SControlRigTweenWidget.cpp
@@ -12,9 +12,7 @@
 #include "ScopedTransaction.h"
 #include "ControlRig.h"
 #include "UnrealEdGlobals.h"
-#include "ControlRigEditMode.h"
 #include "Tools/ControlRigPose.h"
-#include "EditorModeManager.h"
 #include "ISequencer.h"
 #include "LevelSequence.h"
 #include "LevelSequenceEditorBlueprintLibrary.h"
@@ -24,93 +22,89 @@
 
 #define LOCTEXT_NAMESPACE "ControlRigTweenWidget"
 
-void SControlRigTweenWidget::Construct(const FArguments& InArgs)
+void SControlRigTweenSlider::Construct(const FArguments& InArgs)
 {
+
 	PoseBlendValue = 0.0f;
 	bIsBlending = false;
 	bSliderStartedTransaction = false;
-	OwningToolkit = InArgs._InOwningToolkit;
+	AnimSlider = InArgs._InAnimSlider;
+	FText Text = AnimSlider->GetText();
+	FText TooltipText = AnimSlider->GetTooltipText();
 	ChildSlot
-	[
-		SNew(SBorder)
-		.BorderImage(FAppStyle::Get().GetBrush("EditorViewport.OverlayBrush"))
-		.Padding(20.f)
+[
+		SNew(SVerticalBox)
+		+ SVerticalBox::Slot()
+		.AutoHeight()
+		.HAlign(HAlign_Center)
+		.Padding(0.0f, 0.0f, 0.0f, 0.0f)
 		[
-			SNew(SVerticalBox)
-			+ SVerticalBox::Slot()
-			.AutoHeight()
-			.HAlign(HAlign_Center)
-			.Padding(0.0f, 0.0f, 0.0f, 5.0f)
-			[
-				SNew(STextBlock)
-				.TextStyle(FAppStyle::Get(), "NormalText.Important")
-				.Text(LOCTEXT("TweenController", "Tween Controller"))
+			SNew(SHorizontalBox)
+			+ SHorizontalBox::Slot()
+				.AutoWidth()
+				.VAlign(VAlign_Center)
+				[
+					SNew(SSpinBox<double>)
+					.PreventThrottling(true)
+					.Value(this, &SControlRigTweenSlider::OnGetPoseBlendValue)
+					.ToolTipText(LOCTEXT("TweenTooltip", "Key at current frame between previous(-1.0) and next(1.0) poses. Use Ctrl drag for under and over shoot."))
+					.MinValue(-2.0)
+					.MaxValue(2.0)
+					.MinSliderValue(-1.0)
+					.MaxSliderValue(1.0)
+					.SliderExponent(1)
+					.Delta(0.005)
+					.MinDesiredWidth(100.0)
+					.SupportDynamicSliderMinValue(true)
+					.SupportDynamicSliderMaxValue(true)
+					.ClearKeyboardFocusOnCommit(true)
+					.OnValueChanged(this, &SControlRigTweenSlider::OnPoseBlendChanged)
+					.OnValueCommitted(this, &SControlRigTweenSlider::OnPoseBlendCommited)
+					.OnBeginSliderMovement(this, &SControlRigTweenSlider::OnBeginSliderMovement)
+					.OnEndSliderMovement(this, &SControlRigTweenSlider::OnEndSliderMovement)
+				]
 			]
-			+ SVerticalBox::Slot()
-			.AutoHeight()
-			.HAlign(HAlign_Center)
-			[
-				SNew(SSpinBox<float>)
-				.PreventThrottling(true)
-				.Value(this, &SControlRigTweenWidget::OnGetPoseBlendValueFloat)
-				.ToolTipText(LOCTEXT("TweenTooltip", "Key at current frame between previous(-1.0) and next(1.0) poses. Use Ctrl drag for under and over shoot."))
-				.MinValue(-2.0f)
-				.MaxValue(2.0f)
-				.MinSliderValue(-1.0f)
-				.MaxSliderValue(1.0f)
-				.SliderExponent(1)
-				.Delta(0.005f)
-				.MinDesiredWidth(100.0f)
-				.SupportDynamicSliderMinValue(true)
-				.SupportDynamicSliderMaxValue(true)
-				.OnValueChanged(this, &SControlRigTweenWidget::OnPoseBlendChanged)
-				.OnValueCommitted(this, &SControlRigTweenWidget::OnPoseBlendCommited)
-				.OnBeginSliderMovement(this, &SControlRigTweenWidget::OnBeginSliderMovement)
-				.OnEndSliderMovement(this, &SControlRigTweenWidget::OnEndSliderMovement)
-			]
-		]
+		
 	];	
 }
 
-void SControlRigTweenWidget::OnPoseBlendChanged(float ChangedVal)
+void SControlRigTweenSlider::OnPoseBlendChanged(double ChangedVal)
 {
-	TArray<UControlRig*> ControlRigs = GetControlRigs();
-	if (ControlRigs.Num() > 0 && WeakSequencer.IsValid() && bIsBlending)
+	if (WeakSequencer.IsValid() && bIsBlending)
 	{
 		PoseBlendValue = ChangedVal;
-		ControlsToTween.Blend(WeakSequencer, ChangedVal);
+		AnimSlider->Blend(WeakSequencer, ChangedVal);
 		WeakSequencer.Pin()->NotifyMovieSceneDataChanged(EMovieSceneDataChangeType::TrackValueChanged);
 	}
 
 }
 
-void SControlRigTweenWidget::OnBeginSliderMovement()
+void SControlRigTweenSlider::OnBeginSliderMovement()
 {
 	if (bSliderStartedTransaction == false)
 	{
-		bIsBlending = true;
-		bSliderStartedTransaction = true;
-		GEditor->BeginTransaction(LOCTEXT("TweenTransaction", "Tween"));
-		SetupControls();
+		bIsBlending = bSliderStartedTransaction = Setup();
+		if (bIsBlending)
+		{
+			GEditor->BeginTransaction(LOCTEXT("AnimSliderBlend", "AnimSlider Blend"));
+
+		}
 	}
 }
 
-void SControlRigTweenWidget::SetupControls()
+bool SControlRigTweenSlider::Setup()
 {
 	//if getting sequencer from level sequence need to use the current(master), not the focused
 	ULevelSequence* LevelSequence = ULevelSequenceEditorBlueprintLibrary::GetCurrentLevelSequence();
 	IAssetEditorInstance* AssetEditor = GEditor->GetEditorSubsystem<UAssetEditorSubsystem>()->FindEditorForAsset(LevelSequence, false);
 	ILevelSequenceEditorToolkit* LevelSequenceEditor = static_cast<ILevelSequenceEditorToolkit*>(AssetEditor);
 	WeakSequencer = LevelSequenceEditor ? LevelSequenceEditor->GetSequencer() : nullptr;
-	TArray<UControlRig*> ControlRigs = GetControlRigs();
-	if (ControlRigs.Num() > 0)
-	{
-		WeakSequencer.Pin()->GetFocusedMovieSceneSequence()->GetMovieScene()->Modify();
-		ControlsToTween.Setup(ControlRigs, WeakSequencer);
-	}
+
+	return AnimSlider->Setup(WeakSequencer);
+	
 }
 
-void SControlRigTweenWidget::OnEndSliderMovement(float NewValue)
+void SControlRigTweenSlider::OnEndSliderMovement(double NewValue)
 {
 	if (bSliderStartedTransaction)
 	{
@@ -135,13 +129,12 @@ FReply SControlRigTweenWidget::OnDragDetected(const FGeometry& MyGeometry, const
 			GetDesiredSize(),
 			OnUIDropped
 		);
+
 	if (OwningToolkit.IsValid())
 	{
 		OwningToolkit.Pin()->TryRemoveTweenOverlay();
 	}
 	return FReply::Handled().BeginDragDrop(DragDropOperation);
-
-	return FReply::Unhandled();
 }
 
 void SControlRigTweenWidget::FinishDraggingWidget(const FVector2D InLocation)
@@ -153,17 +146,16 @@ void SControlRigTweenWidget::FinishDraggingWidget(const FVector2D InLocation)
 	}
 }
 
-void SControlRigTweenWidget::OnPoseBlendCommited(float ChangedVal, ETextCommit::Type Type)
+void SControlRigTweenSlider::OnPoseBlendCommited(double ChangedVal, ETextCommit::Type Type)
 {
-	TArray<UControlRig*> ControlRigs = GetControlRigs();
-	if (ControlRigs.Num() > 0)
+	if (bIsBlending == false)
+	{
+		bIsBlending = Setup();
+
+	}
+	if(bIsBlending)
 	{
 		FScopedTransaction ScopedTransaction(LOCTEXT("TweenTransaction", "Tween"));
-		if (bIsBlending == false)
-		{
-			SetupControls();
-			bIsBlending = true;
-		}
 		PoseBlendValue = ChangedVal;
 		OnPoseBlendChanged(ChangedVal);
 		bIsBlending = false;
@@ -171,19 +163,130 @@ void SControlRigTweenWidget::OnPoseBlendCommited(float ChangedVal, ETextCommit::
 	}
 }
 
-TArray<UControlRig*>SControlRigTweenWidget::GetControlRigs()
+void SControlRigTweenWidget::OnSelectSliderTool(int32 Index)
 {
-	TArray<UControlRig*> ControlRigs;
-	if (FControlRigEditMode* EditMode = static_cast<FControlRigEditMode*>(GLevelEditorModeTools().GetActiveMode(FControlRigEditMode::ModeName)))
+	ActiveSlider = Index;
+	if (ActiveSlider >= 0 && ActiveSlider < AnimBlendTools.GetAnimSliders().Num())
 	{
-		TMap<UControlRig*, TArray<FRigElementKey>> SelectedControls;
-		EditMode->GetAllSelectedControls(SelectedControls);
-		for (TPair<UControlRig*, TArray<FRigElementKey>>& Selected: SelectedControls)
-		{
-			ControlRigs.Add(Selected.Key);
-		}
+		SliderWidget->SetAnimSlider(AnimBlendTools.GetAnimSliders()[ActiveSlider]);
 	}
-	return ControlRigs;
+	
+}
+FText SControlRigTweenWidget::GetActiveSliderName() const
+{
+	if (ActiveSlider >= 0 && ActiveSlider < AnimBlendTools.GetAnimSliders().Num())
+	{
+		return AnimBlendTools.GetAnimSliders()[ActiveSlider].Get()->GetText();
+	}
+	return FText();
+}
+
+FText SControlRigTweenWidget::GetActiveSliderTooltip() const
+{
+	if (ActiveSlider >= 0 && ActiveSlider < AnimBlendTools.GetAnimSliders().Num())
+	{
+		return AnimBlendTools.GetAnimSliders()[ActiveSlider].Get()->GetTooltipText();
+	}
+	return FText();
+}
+
+int32 SControlRigTweenWidget::ActiveSlider = 0;
+
+
+void SControlRigTweenWidget::Construct(const FArguments& InArgs)
+{
+	TSharedPtr<FBaseAnimSlider> TweenPtr = MakeShareable(new FControlsToTween());
+	AnimBlendTools.RegisterAnimSlider(TweenPtr);
+	TSharedPtr<FBaseAnimSlider> AmplifyPtr = MakeShareable(new FAmplifySlider());
+	AnimBlendTools.RegisterAnimSlider(AmplifyPtr);
+	TSharedPtr<FBaseAnimSlider> BlendNeighborPtr = MakeShareable(new FBlendToPreviousNextSlider());
+	AnimBlendTools.RegisterAnimSlider(BlendNeighborPtr);
+
+	OwningToolkit = InArgs._InOwningToolkit;
+
+
+	// Combo Button to swap sliders 
+	TSharedRef<SComboButton> SliderComoboBtn = SNew(SComboButton)
+		.OnGetMenuContent_Lambda([this]()
+			{
+
+				FMenuBuilder MenuBuilder(true, NULL); //maybe todo look at settting these up with commands
+
+				MenuBuilder.BeginSection("Anim Sliders");
+
+				int32 Index = 0;
+				for (const TSharedPtr<FBaseAnimSlider>& SliderPtr : AnimBlendTools.GetAnimSliders())
+				{
+					FUIAction ItemAction(FExecuteAction::CreateSP(this, &SControlRigTweenWidget::OnSelectSliderTool, Index));
+					MenuBuilder.AddMenuEntry(SliderPtr.Get()->GetText(), TAttribute<FText>(), FSlateIcon(), ItemAction);
+					++Index;
+				}
+
+				MenuBuilder.EndSection();
+
+				return MenuBuilder.MakeWidget();
+			})
+			.ButtonContent()
+				[
+					SNew(SHorizontalBox)
+					/*  todo add an icon
+					+ SHorizontalBox::Slot()
+					.AutoWidth()
+					.Padding(FMargin(0.f, 0.f, 6.f, 0.f))
+					[
+						SNew(SBox)
+						.WidthOverride(16.f)
+						.HeightOverride(16.f)
+						[
+							SNew(SImage)
+							.Image_Static(&FLevelEditorToolBar::GetActiveModeIcon, LevelEditorPtr)
+						]
+
+					]
+					*/
+				+SHorizontalBox::Slot()
+				[
+					SNew(STextBlock)
+					.Text_Lambda([this]()
+						{
+							return GetActiveSliderName();
+						})
+					.ToolTipText_Lambda([this]()
+					{
+						return GetActiveSliderTooltip();
+					})
+				]
+
+			];
+
+		TSharedRef<SHorizontalBox> MainBox = SNew(SHorizontalBox);
+		TSharedPtr<FBaseAnimSlider> SliderPtr = AnimBlendTools.GetAnimSliders()[ActiveSlider];
+		MainBox->AddSlot()
+		.AutoWidth()
+		.HAlign(HAlign_Center)
+		.VAlign(VAlign_Center)
+		.Padding(2.0f, 2.0f, 2.0f, 2.0f)
+		[
+			SliderComoboBtn
+		];
+		MainBox->AddSlot()
+		.AutoWidth()
+		.HAlign(HAlign_Center)
+		.VAlign(VAlign_Center)
+		.Padding(2.0f, 2.0f, 2.0f, 2.0f)
+		[
+			SAssignNew(SliderWidget,SControlRigTweenSlider).InAnimSlider(SliderPtr)
+		];
+	
+	ChildSlot
+	[
+		SNew(SBorder)
+		.BorderImage(FAppStyle::Get().GetBrush("EditorViewport.OverlayBrush"))
+		[
+			MainBox
+		]
+	];
+
 }
 
 #undef LOCTEXT_NAMESPACE
diff --git a/Engine/Plugins/Animation/ControlRig/Source/ControlRigEditor/Private/EditMode/SControlRigTweenWidget.h b/Engine/Plugins/Animation/ControlRig/Source/ControlRigEditor/Private/EditMode/SControlRigTweenWidget.h
--- a/Engine/Plugins/Animation/ControlRig/Source/ControlRigEditor/Private/EditMode/SControlRigTweenWidget.h
+++ b/Engine/Plugins/Animation/ControlRig/Source/ControlRigEditor/Private/EditMode/SControlRigTweenWidget.h
@@ -20,30 +20,60 @@ class UControlRig;
 class ISequencer;
 class FControlRigEditModeToolkit;
 
-class SControlRigTweenWidget : public SCompoundWidget
+class SControlRigTweenSlider : public SCompoundWidget
 {
-	SLATE_BEGIN_ARGS(SControlRigTweenWidget) {}
-	SLATE_ARGUMENT(UControlRigPoseAsset*, PoseAsset)
-    SLATE_ARGUMENT(TSharedPtr<FControlRigEditModeToolkit>, InOwningToolkit)
+	SLATE_BEGIN_ARGS(SControlRigTweenSlider) {}
+	SLATE_ARGUMENT(TSharedPtr<FBaseAnimSlider>, InAnimSlider)
 	SLATE_END_ARGS()
-	~SControlRigTweenWidget()
+	~SControlRigTweenSlider()
 	{
 	}
 
 	void Construct(const FArguments& InArgs);
+	void SetAnimSlider(TSharedPtr<FBaseAnimSlider>& InAnimSlider) { AnimSlider = InAnimSlider; }
 
 private:
 
-	TArray<UControlRig*> GetControlRigs();
 
 	/*
 	* Delegates and Helpers
 	*/
-	void OnPoseBlendChanged(float ChangedVal);
-	void OnPoseBlendCommited(float ChangedVal, ETextCommit::Type Type);
+	void OnPoseBlendChanged(double ChangedVal);
+	void OnPoseBlendCommited(double ChangedVal, ETextCommit::Type Type);
 	void OnBeginSliderMovement();
-	void OnEndSliderMovement(float NewValue);
-	float OnGetPoseBlendValueFloat() const { return PoseBlendValue; }
+	void OnEndSliderMovement(double NewValue);
+	double OnGetPoseBlendValue() const { return PoseBlendValue; }
+
+	bool Setup();
+	double PoseBlendValue;
+	bool bIsBlending;
+	bool bSliderStartedTransaction;
+
+	
+	TWeakPtr<ISequencer> WeakSequencer;
+	TSharedPtr<FBaseAnimSlider> AnimSlider;
+
+};
+
+
+class SControlRigTweenWidget : public SCompoundWidget
+{
+	SLATE_BEGIN_ARGS(SControlRigTweenWidget) {}
+	SLATE_ARGUMENT(TSharedPtr<FControlRigEditModeToolkit>, InOwningToolkit)
+	SLATE_END_ARGS()
+		~SControlRigTweenWidget()
+	{
+	}
+
+	void Construct(const FArguments& InArgs);
+
+private:
+
+	void OnSelectSliderTool(int32 Index);
+	FText GetActiveSliderName() const;
+	FText GetActiveSliderTooltip() const;
+
+
 
 	FReply OnMouseButtonDown(const FGeometry& MyGeometry, const FPointerEvent& MouseEvent) override
 	{
@@ -52,14 +82,11 @@ private:
 	FReply OnDragDetected(const FGeometry& MyGeometry, const FPointerEvent& MouseEvent);
 	void FinishDraggingWidget(const FVector2D InLocation);
 
-	void SetupControls();
-	float PoseBlendValue;
-	bool bIsBlending;
-	bool bSliderStartedTransaction;
-
-	FControlsToTween  ControlsToTween;
-
 	TWeakPtr<ISequencer> WeakSequencer;
 	TWeakPtr<FControlRigEditModeToolkit> OwningToolkit;
-};
+	FAnimBlendTooLManager  AnimBlendTools;
+
+	TSharedPtr<SControlRigTweenSlider> SliderWidget;
+	static int32 ActiveSlider;
 
+};
diff --git a/Engine/Plugins/Animation/ControlRig/Source/ControlRigEditor/Private/Tools/ControlRigTweener.cpp b/Engine/Plugins/Animation/ControlRig/Source/ControlRigEditor/Private/Tools/ControlRigTweener.cpp
--- a/Engine/Plugins/Animation/ControlRig/Source/ControlRigEditor/Private/Tools/ControlRigTweener.cpp
+++ b/Engine/Plugins/Animation/ControlRig/Source/ControlRigEditor/Private/Tools/ControlRigTweener.cpp
@@ -2,6 +2,7 @@
 
 #include "Tools/ControlRigTweener.h"
 #include "Channels/MovieSceneFloatChannel.h"
+#include "Channels/MovieSceneDoubleChannel.h"
 #include "ControlRig.h"
 #include "ISequencer.h"
 #include "MovieSceneSequence.h"
@@ -11,24 +12,344 @@
 #include "Sequencer/MovieSceneControlRigParameterSection.h"
 #include "IKeyArea.h"
 #include "Channels/MovieSceneChannelProxy.h"
+#include "CurveEditor.h"
+#include "CurveModel.h"
+#include "Channels/FloatChannelCurveModel.h"
+#include "MovieSceneSection.h"
+#include "EditorModeManager.h"
+#include "ControlRigEditMode.h"
 
-//category not used so far, todo
-void FControlsToTween::Setup(const TArray<UControlRig*>& SelectedControlRigs, TWeakPtr<ISequencer>& InSequencer)
+#define LOCTEXT_NAMESPACE "ControlRigTweenWidget"
+
+/**
+*
+* FAnimSliderKeySelection
+*
+*/
+
+FAnimSliderKeySelection::FContiguousKeys::FContiguousKeys(const TArray<FKeyHandle>& AllKeyHandles, const TArray<FKeyPosition>& AllKeyPositions, const TArray<int32>& InContiguousKeyIndices)
+{
+	Indices.Reset();
+	if (AllKeyHandles.Num() > 0 && AllKeyPositions.Num() == AllKeyHandles.Num() && InContiguousKeyIndices.Num() > 0)
+	{
+		Indices = InContiguousKeyIndices;
+		//Previous
+		PreviousIndex = Indices[0] - 1;
+		if (PreviousIndex < 0)
+		{
+			PreviousIndex = INDEX_NONE;
+		}		
+		//Next
+		NextIndex = Indices[Indices.Num() -1] + 1;
+		if (NextIndex >= AllKeyHandles.Num())
+		{
+			NextIndex = INDEX_NONE;
+		}
+	}
+}
+
+FAnimSliderKeySelection::FContiguousKeysArray::FContiguousKeysArray(const TArray<FKeyHandle>& InAllKeyHandles, const TArray<FKeyPosition>& InAllKeyPositions)
+{
+	AllKeyHandles = InAllKeyHandles;
+	AllKeyPositions = InAllKeyPositions;
+	KeysArray.Reset();	
+}
+
+void FAnimSliderKeySelection::FContiguousKeysArray::Add(const TArray<int32>& InContiguousKeyIndices)
+{
+	if (AllKeyHandles.Num() > 0 && AllKeyPositions.Num() == AllKeyHandles.Num() && InContiguousKeyIndices.Num() > 0)
+	{
+		KeysArray.Add(FContiguousKeys(AllKeyHandles, AllKeyPositions, InContiguousKeyIndices));
+	}
+}
+
+static int32 GetIndex(const TArray<FKeyPosition>& Keys, double Time)
+{
+	int32 Index = Algo::LowerBoundBy(Keys, Time, [](const FKeyPosition& Value) { return Value.InputValue; });
+
+	// don't trust precision issues so will double check to make sure the index is correct
+	if (Index != INDEX_NONE)
+	{
+		if (FMath::IsNearlyEqual(Keys[Index].InputValue, Time))
+		{
+			return Index;
+		}
+		else if (((Index - 1) >= 0) && FMath::IsNearlyEqual(Keys[Index - 1].InputValue, Time))
+		{
+			return Index - 1;
+		}
+		else if (((Index + 1) < Keys.Num()) && FMath::IsNearlyEqual(Keys[Index + 1].InputValue, Time))
+		{
+			return Index + 1;
+		}
+	}
+	return Index;
+}
+
+bool FAnimSliderKeySelection::Setup(TWeakPtr<ISequencer>& InSequencer)
+{
+	KeyMap.Reset();
+	if (InSequencer.IsValid())
+	{
+		CurveEditor = InSequencer.Pin()->GetCurveEditor();
+		GetMapOfContiguousKeys();
+	}
+	return (KeyMap.Num() > 0);
+}
+
+bool FAnimSliderObjectSelection::Setup(TWeakPtr<ISequencer>& InSequencer)
 {
+	ChannelsArray.Reset();
+	TArray<UControlRig*> ControlRigs = GetControlRigs();
+	return Setup(ControlRigs, InSequencer);
+}
+
+void FAnimSliderKeySelection::GetMapOfContiguousKeys()
+{
+	TArray<FKeyHandle> AllKeyHandles;
+	TArray<FKeyPosition> AllKeyPositions;
+
+	TArray<FKeyHandle> SelectedKeyHandles;
+	TArray<FKeyPosition> SelectedKeyPositions;
+
+	const TMap<FCurveModelID, FKeyHandleSet>& SelectionKeyMap = CurveEditor->Selection.GetAll();
+
+	TArray<int32> ContiguousKeyIndices;
+
+	for (const TPair<FCurveModelID, FKeyHandleSet>& Pair : SelectionKeyMap)
+	{
+		FCurveModel* Curve = CurveEditor->FindCurve(Pair.Key);
+
+		if (Curve)
+		{
+			AllKeyHandles.Reset();
+			Curve->GetKeys(*CurveEditor, TNumericLimits<double>::Lowest(), TNumericLimits<double>::Max(), TNumericLimits<double>::Lowest(), TNumericLimits<double>::Max(), AllKeyHandles);
+			AllKeyPositions.SetNum(AllKeyHandles.Num());
+			Curve->GetKeyPositions(AllKeyHandles, AllKeyPositions);
+			SelectedKeyHandles.Reset(Pair.Value.Num());
+			SelectedKeyHandles.Append(Pair.Value.AsArray().GetData(), Pair.Value.Num());
+
+			// Get all the selected keys
+			SelectedKeyPositions.SetNum(SelectedKeyHandles.Num());
+			Curve->GetKeyPositions(SelectedKeyHandles, SelectedKeyPositions);
+
+			ContiguousKeyIndices.Reset();
+			FContiguousKeysArray& KeyArray = KeyMap.Add(Pair.Key, FContiguousKeysArray(AllKeyHandles, AllKeyPositions));
+
+			TArray<int32> SelectedIndices;
+			for (int32 Index = 0; Index < SelectedKeyPositions.Num(); ++Index)
+			{
+				int32 AllKeysIndex = GetIndex(AllKeyPositions, SelectedKeyPositions[Index].InputValue);
+				SelectedIndices.Add(AllKeysIndex);
+			}
+			SelectedIndices.Sort();
+			for(int32 AllKeysIndex: SelectedIndices)
+			{
+				if (ContiguousKeyIndices.Num() > 0) //see if this key is next to the previous one
+				{
+					if (ContiguousKeyIndices[ContiguousKeyIndices.Num() - 1] + 1 == AllKeysIndex)
+					{
+						ContiguousKeyIndices.Add(AllKeysIndex);
+					}
+					else //not contiguous, so need to add to the OutKeyMap
+					{
+						//now start new set with the new index
+						KeyArray.Add(ContiguousKeyIndices);
+						ContiguousKeyIndices.Reset();
+						ContiguousKeyIndices.Add(AllKeysIndex);
+					}
+				}
+				else//first key in this set so just add
+				{
+					ContiguousKeyIndices.Add(AllKeysIndex);
+				}
+			}
+			if (ContiguousKeyIndices.Num() > 0)
+			{
+				KeyArray.Add(ContiguousKeyIndices);
+				ContiguousKeyIndices.Reset();
+			}
+		}
+	}
+}
+
+/**
+*
+* FAnimSliderObjectSelection
+*
+*/
+TArray<UControlRig*> FAnimSliderObjectSelection::GetControlRigs()
+{
+	TArray<UControlRig*> ControlRigs;
+	if (FControlRigEditMode* EditMode = static_cast<FControlRigEditMode*>(GLevelEditorModeTools().GetActiveMode(FControlRigEditMode::ModeName)))
+	{
+		TMap<UControlRig*, TArray<FRigElementKey>> SelectedControls;
+		EditMode->GetAllSelectedControls(SelectedControls);
+		for (TPair<UControlRig*, TArray<FRigElementKey>>& Selected : SelectedControls)
+		{
+			ControlRigs.Add(Selected.Key);
+		}
+	}
+	return ControlRigs;
+}
+
+void FAnimSliderObjectSelection::SetupChannel(FFrameNumber CurrentFrame, TArray<FFrameNumber>& KeyTimes, TArray<FKeyHandle>& Handles, 
+	FMovieSceneFloatChannel* FloatChannel, FMovieSceneDoubleChannel* DoubleChannel,
+	FChannelKeyBounds& KeyBounds)
+{
+	KeyBounds.FloatChannel = FloatChannel;
+	KeyBounds.DoubleChannel = DoubleChannel;
+	KeyBounds.CurrentFrame = CurrentFrame;
+	KeyBounds.PreviousIndex = INDEX_NONE;
+	KeyBounds.NextIndex = INDEX_NONE;
+	KeyTimes.SetNum(0);
+	Handles.SetNum(0);
+	if (FloatChannel)
+	{
+		FloatChannel->GetKeys(TRange<FFrameNumber>(), &KeyTimes, &Handles);
+		float OutValue;
+		FloatChannel->Evaluate(FFrameTime(CurrentFrame), OutValue);
+		KeyBounds.CurrentValue = OutValue;
+	}
+	else if (DoubleChannel)
+	{
+		DoubleChannel->GetKeys(TRange<FFrameNumber>(), &KeyTimes, &Handles);
+		double OutValue;
+		DoubleChannel->Evaluate(FFrameTime(CurrentFrame), OutValue);
+		KeyBounds.CurrentValue = OutValue;
+	}
+
+	if (KeyTimes.Num() > 0)
+	{
+		TArrayView<const FMovieSceneFloatValue> FloatValues;
+		TArrayView<const FMovieSceneDoubleValue> DoubleValues;
+		if (FloatChannel)
+		{
+			FloatValues = FloatChannel->GetValues();
+		}
+		else if (DoubleChannel)
+		{
+			DoubleValues = DoubleChannel->GetValues();
+
+		}
+		for (int32 Index = 0; Index < KeyTimes.Num(); Index++)
+		{
+			FFrameNumber FrameNumber = KeyTimes[Index];
+			if (FrameNumber < CurrentFrame)
+			{
+				KeyBounds.PreviousIndex = Index;
+				KeyBounds.PreviousFrame = FrameNumber;
+				if (FloatChannel)
+				{
+					KeyBounds.PreviousValue = FloatValues[Index].Value;
+				}
+				else if (DoubleChannel)
+				{
+					KeyBounds.PreviousValue = DoubleValues[Index].Value;
+				}
+			}
+			else if (FrameNumber > CurrentFrame)
+			{
+				KeyBounds.NextIndex = Index;
+				KeyBounds.NextFrame = FrameNumber;
+				if (FloatChannel)
+				{
+					KeyBounds.NextValue = FloatValues[Index].Value;
+				}
+				else if (DoubleChannel)
+				{
+					KeyBounds.NextValue = DoubleValues[Index].Value;
+				}
+				break;
+			}
+		}
+	}
+	KeyBounds.bValid = (KeyBounds.PreviousIndex != INDEX_NONE && KeyBounds.NextIndex != INDEX_NONE
+		&& KeyBounds.PreviousIndex != KeyBounds.NextIndex) ? true : false;
+}
+
+bool FAnimSliderObjectSelection::Setup(const TArray<UControlRig*>& SelectedControlRigs, TWeakPtr<ISequencer>& InSequencer)
+{
+	ChannelsArray.Reset();
 	if (InSequencer.IsValid() == false)
 	{
-		return;
+		return false;
 	}
 	ISequencer* Sequencer = InSequencer.Pin().Get();
-	ControlRigChannelsMap.Reset();
 
 	FQualifiedFrameTime CurrentTime = Sequencer->GetLocalTime();
 
 	UMovieScene* MovieScene = Sequencer->GetFocusedMovieSceneSequence()->GetMovieScene();
 
+	//get selected controls and objects
+	TArray<FGuid> SelectedGuids;
+	Sequencer->GetSelectedObjects(SelectedGuids);
+
 	const TArray<FMovieSceneBinding>& Bindings = MovieScene->GetBindings();
 	for (const FMovieSceneBinding& Binding : Bindings)
 	{
+		if (SelectedGuids.Num() > 0 && SelectedGuids.Contains(Binding.GetObjectGuid()))
+		{
+			const TArray<UMovieSceneTrack*>& Tracks = Binding.GetTracks();
+			for (UMovieSceneTrack* Track : Tracks)
+			{
+				if (Track && Track->IsA<UMovieSceneControlRigParameterTrack>() == false)
+				{
+					const TArray<UMovieSceneSection*>& Sections = Track->GetAllSections();
+					for (UMovieSceneSection* Section : Sections)
+					{
+						if (Section && Section->IsActive())
+						{
+							FMovieSceneChannelProxy& Proxy = Section->GetChannelProxy();
+
+							const FName FloatChannelTypeName = FMovieSceneFloatChannel::StaticStruct()->GetFName();
+
+							FMovieSceneChannelProxy& ChannelProxy = Section->GetChannelProxy();
+							TArrayView<FMovieSceneFloatChannel*> FloatChannels = ChannelProxy.GetChannels<FMovieSceneFloatChannel>();
+							TArrayView<FMovieSceneDoubleChannel*> DoubleChannels = ChannelProxy.GetChannels<FMovieSceneDoubleChannel>();
+
+							//reuse these arrays
+							TArray<FFrameNumber> KeyTimes;
+							TArray<FKeyHandle> Handles;
+							FAnimSliderObjectSelection::FObjectChannels ObjectChannels;
+							ObjectChannels.Section = Section;
+							KeyTimes.SetNum(0);
+							Handles.SetNum(0);
+							const int32 NumFloatChannels = FloatChannels.Num();
+							const int32 NumDoubleChannels = DoubleChannels.Num();
+
+							ObjectChannels.KeyBounds.SetNum(NumFloatChannels + NumDoubleChannels);
+							int32 NumValidChannels = 0;
+							int ChannelIdx = 0;
+							for (ChannelIdx = 0;ChannelIdx < NumFloatChannels; ++ChannelIdx)
+							{
+								FMovieSceneFloatChannel* FloatChannel = FloatChannels[ChannelIdx];
+								SetupChannel(CurrentTime.Time.GetFrame(), KeyTimes, Handles, FloatChannel, nullptr,
+									ObjectChannels.KeyBounds[ChannelIdx]);
+								if (ObjectChannels.KeyBounds[ChannelIdx].bValid)
+								{
+									++NumValidChannels;
+								}
+							}
+							for (ChannelIdx = 0; ChannelIdx < NumDoubleChannels; ++ChannelIdx)
+							{
+								FMovieSceneDoubleChannel* DoubleChannel = DoubleChannels[ChannelIdx];
+								SetupChannel(CurrentTime.Time.GetFrame(), KeyTimes, Handles, nullptr, DoubleChannel,
+									ObjectChannels.KeyBounds[ChannelIdx]);
+								if (ObjectChannels.KeyBounds[ChannelIdx].bValid)
+								{
+									++NumValidChannels;
+								}
+							}
+							if (NumValidChannels > 0)
+							{
+								ChannelsArray.Add(ObjectChannels);
+							}
+						}
+					}
+				}
+			}
+		}
 		UMovieSceneControlRigParameterTrack* Track = Cast<UMovieSceneControlRigParameterTrack>(MovieScene->FindTrack(UMovieSceneControlRigParameterTrack::StaticClass(), Binding.GetObjectGuid(), NAME_None));
 		if (Track && Track->GetControlRig() && SelectedControlRigs.Contains(Track->GetControlRig()))
 		{
@@ -37,7 +358,6 @@ void FControlsToTween::Setup(const TArray<UControlRig*>& SelectedControlRigs, TW
 				UMovieSceneControlRigParameterSection* Section = Cast<UMovieSceneControlRigParameterSection>(MovieSection);
 				if (Section && Section->IsActive() && Section->GetRange().Contains(CurrentTime.Time.GetFrame()))
 				{
-					Section->Modify();
 					UControlRig* ControlRig = Track->GetControlRig();
 					TArray<FRigControlElement*> CurrentControls;
 					ControlRig->GetControlsInOrder(CurrentControls);
@@ -55,9 +375,8 @@ void FControlsToTween::Setup(const TArray<UControlRig*>& SelectedControlRigs, TW
 					{
 						if (ControlRig->GetHierarchy()->IsAnimatable(ControlElement) &&  ControlRig->IsControlSelected(ControlElement->GetName()))
 						{
-							FControlRigChannels ControlRigChannels;
-							ControlRigChannels.NumChannels = 0;
-							ControlRigChannels.Section = Section;
+							FAnimSliderObjectSelection::FObjectChannels ObjectChannels;
+							ObjectChannels.Section = Section;
 							KeyTimes.SetNum(0);
 							Handles.SetNum(0);
 							FChannelMapInfo* pChannelIndex = Section->ControlChannelMap.Find(ControlElement->GetName());
@@ -80,8 +399,12 @@ void FControlsToTween::Setup(const TArray<UControlRig*>& SelectedControlRigs, TW
 									NumChannels = 3;
 									break;
 								}
-								case ERigControlType::Transform:
 								case ERigControlType::TransformNoScale:
+								{
+									NumChannels = 6;
+									break;
+								}
+								case ERigControlType::Transform:
 								case ERigControlType::EulerTransform:
 								{
 									NumChannels = 9;
@@ -89,23 +412,22 @@ void FControlsToTween::Setup(const TArray<UControlRig*>& SelectedControlRigs, TW
 								}
 							}
 							int32 BoundIndex = 0;
+							int32 NumValidChannels = 0;
+							ObjectChannels.KeyBounds.SetNum(NumChannels);
 							for (int ChannelIdx = pChannelIndex->ChannelIndex; ChannelIdx < (pChannelIndex->ChannelIndex + NumChannels); ++ChannelIdx)
 							{
 								FMovieSceneFloatChannel* Channel = Channels[ChannelIdx];
-								SetupControlRigChannel(CurrentTime.Time.GetFrame(), KeyTimes, Handles, Channel,
-									ControlRigChannels.KeyBounds[BoundIndex]);
-								if (ControlRigChannels.KeyBounds[BoundIndex].bValid)
+								SetupChannel(CurrentTime.Time.GetFrame(), KeyTimes, Handles, Channel, nullptr,
+									ObjectChannels.KeyBounds[BoundIndex]);
+								if (ObjectChannels.KeyBounds[BoundIndex].bValid)
 								{
-									++ControlRigChannels.NumChannels;
+									++NumValidChannels;
 								}
 								++BoundIndex;
 							}
-							if (ControlRigChannels.NumChannels > 0)
+							if (NumValidChannels > 0)
 							{
-								TPair<UControlRig*, FName> Pair;
-								Pair.Key = ControlRig;
-								Pair.Value = ControlElement->GetName();
-								ControlRigChannelsMap.Add(Pair, ControlRigChannels);
+								ChannelsArray.Add(ObjectChannels);
 							}
 						}
 					}
@@ -113,84 +435,226 @@ void FControlsToTween::Setup(const TArray<UControlRig*>& SelectedControlRigs, TW
 			}
 		}
 	}
+	return ChannelsArray.Num() > 0;
 }
 
-void FControlsToTween::Blend(TWeakPtr<ISequencer>& InSequencer, float BlendValue)
+/*
+*
+*  FBaseAnimSlider
+* 
+*/
+
+bool FBaseAnimSlider::Setup(TWeakPtr<ISequencer>& InSequencer)
+{
+	KeySelection.KeyMap.Reset();
+	ObjectSelection.ChannelsArray.Reset();
+	if (KeySelection.Setup(InSequencer) == false)
+	{
+		return ObjectSelection.Setup(InSequencer);
+	}
+	else
+	{
+		return true;
+	}
+	return false;
+}
+
+/*
+*
+*  FBasicBlendSlider
+*
+*/
+
+bool FBasicBlendSlider::Blend(TWeakPtr<ISequencer>& InSequencer, const double BlendValue)
 {
 	if (InSequencer.IsValid() == false)
 	{
-		return;
+		return false;
 	}
 	ISequencer* Sequencer = InSequencer.Pin().Get();
-	FFrameTime  FrameTime = Sequencer->GetLocalTime().Time;
-	float FrameAsFloat = (float)FrameTime.GetFrame().Value; 
-	Sequencer->GetFocusedMovieSceneSequence()->GetMovieScene()->Modify();
-	//BlendValue of 0.0 = half, 1.0f = Second, -1.0f = FirtValue;
-	//NormalizedBlendValue goes from 0 to 1.0f
-	for (const TPair<TPair<UControlRig*,FName>, FControlRigChannels>& Pair : ControlRigChannelsMap)
+
+	bool bDidBlend = false;
+
+	if (KeySelection.KeyMap.Num() > 0)
 	{
-		for (int Index = 0; Index < 9; ++Index)
+		TArray<FKeyHandle> KeyHandles;
+		TArray<FKeyPosition> KeyPositions;
+		for (TPair<FCurveModelID, FAnimSliderKeySelection::FContiguousKeysArray>& KeysArray : KeySelection.KeyMap)
 		{
-			float NormalizedBlendValue = (BlendValue + 1.0f) * 0.5f;
-
-			if (Pair.Value.KeyBounds[Index].bValid)
+			FCurveModel* Curve = KeySelection.CurveEditor->FindCurve(KeysArray.Key);
+			if (Curve)
 			{
-				Pair.Value.Section->Modify();
-				if (Pair.Value.KeyBounds[Index].FirstFrame.Value != Pair.Value.KeyBounds[Index].SecondFrame.Value)
+				Curve->Modify();
+				for (FAnimSliderKeySelection::FContiguousKeys& Keys : KeysArray.Value.KeysArray)
 				{
-					float FirstTime = (float)Pair.Value.KeyBounds[Index].FirstFrame.Value;
-					float SecondTime = (float)Pair.Value.KeyBounds[Index].SecondFrame.Value;
-					float TimeLocation = (FrameAsFloat - FirstTime) / (SecondTime - FirstTime);
-					//use TimeLocation to modify the NormalizedBlendValue 
-					if (NormalizedBlendValue > 0.5f)
-					{
-						NormalizedBlendValue = ((NormalizedBlendValue - .5f) / .5f) * (1.0f - TimeLocation) + TimeLocation;
-					}
-					else
+					const int32 PreviousIndex = Keys.PreviousIndex != INDEX_NONE ? Keys.PreviousIndex : Keys.Indices[0];
+					const double PreviousTime = KeysArray.Value.AllKeyPositions[PreviousIndex].InputValue;
+					const double  PreviousValue = KeysArray.Value.AllKeyPositions[PreviousIndex].OutputValue;
+					const int32 NextIndex = Keys.NextIndex != INDEX_NONE ? Keys.NextIndex : Keys.Indices[Keys.Indices.Num() - 1];
+					const double NextTime = KeysArray.Value.AllKeyPositions[NextIndex].InputValue;
+					const double  NextValue = KeysArray.Value.AllKeyPositions[NextIndex].OutputValue;
+
+					KeyHandles.Reset();
+					KeyPositions.Reset();
+					for (int32 Index : Keys.Indices)
 					{
-						NormalizedBlendValue = (NormalizedBlendValue / 0.5f) * TimeLocation;
+						KeyHandles.Add(KeysArray.Value.AllKeyHandles[Index]);
+						const double CurrentValue = KeysArray.Value.AllKeyPositions[Index].OutputValue;
+						const double CurrentTime = KeysArray.Value.AllKeyPositions[Index].InputValue;
+
+						double NewValue = DoBlend(PreviousTime,PreviousValue, CurrentTime, CurrentValue,
+							NextTime,  NextValue, BlendValue);
+
+						 KeyPositions.Add(FKeyPosition(KeysArray.Value.AllKeyPositions[Index].InputValue, NewValue));
 					}
+					Curve->SetKeyPositions(KeyHandles, KeyPositions);
+					bDidBlend = true;
 				}
-				float FirstValue = Pair.Value.KeyBounds[Index].FirstValue;
-				float SecondValue = Pair.Value.KeyBounds[Index].SecondValue;
-				float Value = FirstValue + (SecondValue - FirstValue) * (NormalizedBlendValue);
-				AddKeyToChannel(Pair.Value.KeyBounds[Index].Channel, FrameTime.GetFrame(), Value, Sequencer->GetKeyInterpolation());
 			}
 		}
+		return bDidBlend;
 	}
-}
 
-void FControlsToTween::SetupControlRigChannel(FFrameNumber CurrentFrame, TArray<FFrameNumber>& KeyTimes, TArray<FKeyHandle>& Handles, FMovieSceneFloatChannel* FloatChannel,
-	FChannelKeyBounds& KeyBounds)
-{
-	KeyBounds.Channel = FloatChannel;
-	KeyBounds.FirstIndex = INDEX_NONE;
-	KeyBounds.SecondIndex = INDEX_NONE;
-	KeyTimes.SetNum(0);
-	Handles.SetNum(0);
-	FloatChannel->GetKeys(TRange<FFrameNumber>(), &KeyTimes, &Handles);
-	if (KeyTimes.Num() > 0)
-	{
-		TArrayView<const FMovieSceneFloatValue> Values = FloatChannel->GetValues();
+	FFrameTime  FrameTime = Sequencer->GetLocalTime().Time;
+	FFrameRate TickResoultion = Sequencer->GetFocusedTickResolution();
 
-		for (int32 Index = 0; Index < KeyTimes.Num(); Index++)
+	if(ObjectSelection.ChannelsArray.Num() > 0)
+	{ 
+		for (const FAnimSliderObjectSelection::FObjectChannels& ObjectChannels : ObjectSelection.ChannelsArray)
 		{
-			FFrameNumber FrameNumber = KeyTimes[Index];
-			if (FrameNumber < CurrentFrame)
+			if (ObjectChannels.Section)
 			{
-				KeyBounds.FirstIndex = Index;
-				KeyBounds.FirstFrame = FrameNumber;
-				KeyBounds.FirstValue = Values[Index].Value;
+				ObjectChannels.Section->Modify();
 			}
-			else if (FrameNumber > CurrentFrame)
+			for (int Index = 0; Index < ObjectChannels.KeyBounds.Num(); ++Index)
 			{
-				KeyBounds.SecondIndex = Index;
-				KeyBounds.SecondFrame = FrameNumber;
-				KeyBounds.SecondValue = Values[Index].Value;
-				break;
+				if (ObjectChannels.KeyBounds[Index].bValid)
+				{
+					double PreviousValue = ObjectChannels.KeyBounds[Index].PreviousValue;
+					double PreviousTime = TickResoultion.AsSeconds(FFrameTime(ObjectChannels.KeyBounds[Index].PreviousFrame));
+					double NextValue = ObjectChannels.KeyBounds[Index].NextValue;
+					double NextTime = TickResoultion.AsSeconds(FFrameTime(ObjectChannels.KeyBounds[Index].NextFrame));
+					double CurrentValue = ObjectChannels.KeyBounds[Index].CurrentValue;
+					double CurrentTime = TickResoultion.AsSeconds(FFrameTime(ObjectChannels.KeyBounds[Index].CurrentFrame));
+					double NewValue = DoBlend(PreviousTime, PreviousValue, CurrentTime, CurrentValue,
+						NextTime, NextValue, BlendValue);
+					if (ObjectChannels.KeyBounds[Index].FloatChannel)
+					{
+						AddKeyToChannel(ObjectChannels.KeyBounds[Index].FloatChannel, FrameTime.GetFrame(), (float)NewValue, Sequencer->GetKeyInterpolation());
+					}
+					else if (ObjectChannels.KeyBounds[Index].DoubleChannel)
+					{
+						AddKeyToChannel(ObjectChannels.KeyBounds[Index].DoubleChannel, FrameTime.GetFrame(), NewValue, Sequencer->GetKeyInterpolation());
+					}
+					bDidBlend = true;
+				}
 			}
 		}
 	}
-	KeyBounds.bValid = (KeyBounds.FirstIndex != INDEX_NONE && KeyBounds.SecondIndex != INDEX_NONE
-		&& KeyBounds.FirstIndex != KeyBounds.SecondIndex) ? true : false;
-}
\ No newline at end of file
+	return bDidBlend;
+}
+
+/*
+*
+*  FControlsToTween
+*
+*/
+
+double FControlsToTween::DoBlend(const double PreviousTime, const double PreviousValue, const double CurrentTime, const double CurrentValue,
+	const double NextTime, const double NextValue, const double BlendValue)
+{
+	//clasic tween will move all to same location, not based on current time at all just blend and values
+	double NormalizedBlendValue = (BlendValue + 1.0f) * 0.5f;
+	double Value = PreviousValue + (NextValue - PreviousValue) * (NormalizedBlendValue);
+	return Value;
+}
+
+FText FControlsToTween::GetText()
+{
+	return LOCTEXT("TC", "TC");
+}
+
+FText FControlsToTween::GetTooltipText()
+{
+	return LOCTEXT("TweenControllerTooltip", "Tween between the next and previous keys");
+}
+
+bool FControlsToTween::Setup(const TArray<UControlRig*>& SelectedControlRigs, TWeakPtr<ISequencer>& InSequencer)
+{
+	return ObjectSelection.Setup(SelectedControlRigs, InSequencer);
+}
+
+bool FControlsToTween::Setup(TWeakPtr<ISequencer>& InSequencer)
+{
+	return FBaseAnimSlider::Setup(InSequencer);
+}
+
+/*
+*
+*  Amplify Slider
+*
+*/
+
+
+double FAmplifySlider::DoBlend(const double PreviousTime, const double PreviousValue, const double CurrentTime, const double CurrentValue,
+	const double NextTime, const double NextValue, const double BlendValue)
+{
+	const double T = (CurrentTime - PreviousTime) / (NextTime - PreviousTime);
+	const double ValueAtT = PreviousValue + T * (NextValue - PreviousValue);
+	double NewValue = 0.0;
+	if (BlendValue < 0.0)
+	{
+		NewValue = CurrentValue + (-1.0 * BlendValue) * (ValueAtT - CurrentValue);
+	}
+	else
+	{
+		const double AmplifyValueAtT = CurrentValue + (CurrentValue - ValueAtT);
+		NewValue = CurrentValue + BlendValue * (AmplifyValueAtT - CurrentValue);
+	}
+	return NewValue;
+}
+
+FText FAmplifySlider::GetText()
+{
+	return LOCTEXT("AA", "AA");
+}
+
+FText FAmplifySlider::GetTooltipText()
+{
+	return LOCTEXT("AmplifySliderTooltip", "Amplify or diminish between the next and previous keys");
+}
+
+
+/*
+*
+*  Blend To Neighbors Slider
+*
+*/
+
+double FBlendToPreviousNextSlider::DoBlend(const double PreviousTime, const double PreviousValue, const double CurrentTime, const double CurrentValue,
+	const double NextTime, const double NextValue, const double BlendValue)
+{
+	double NewValue;
+	if (BlendValue < 0.0)
+	{
+		NewValue = CurrentValue + (-1.0 * BlendValue) * (PreviousValue - CurrentValue);
+	}
+	else
+	{
+		NewValue = CurrentValue + BlendValue * (NextValue - CurrentValue);
+	}
+	return NewValue;
+}
+
+FText FBlendToPreviousNextSlider::GetText()
+{
+	return LOCTEXT("PN", "PN");
+}
+
+FText FBlendToPreviousNextSlider::GetTooltipText()
+{
+	return LOCTEXT("BlendToNeighborsSliderTooltip", "Blend to the next or previous values for selected keys or objects");
+}
+
+#undef LOCTEXT_NAMESPACE
+
diff --git a/Engine/Plugins/Animation/ControlRig/Source/ControlRigEditor/Public/Tools/ControlRigTweener.h b/Engine/Plugins/Animation/ControlRig/Source/ControlRigEditor/Public/Tools/ControlRigTweener.h
--- a/Engine/Plugins/Animation/ControlRig/Source/ControlRigEditor/Public/Tools/ControlRigTweener.h
+++ b/Engine/Plugins/Animation/ControlRig/Source/ControlRigEditor/Public/Tools/ControlRigTweener.h
@@ -5,60 +5,224 @@
 #include "UObject/ObjectMacros.h"
 #include "Misc/FrameNumber.h"
 #include "Curves/KeyHandle.h"
+#include "CurveDataAbstraction.h"
+#include "CurveModel.h"
 
 struct FMovieSceneFloatChannel;
+struct FMovieSceneDoubleChannel;
 class UControlRig;
 class ISequencer;
 struct FRigControl;
 class UMovieSceneSection;
+class FCurveEditor;
 
+/*
+* Contains the selection state for a set of keys to blend with the anim slider
+*/
+struct FAnimSliderKeySelection
+{
+	//list of keys in a row
+	struct FContiguousKeys
+	{
+		int32 PreviousIndex;
+		TArray <int32> Indices;
+		int32 NextIndex;
+		FContiguousKeys(const TArray<FKeyHandle>& InAllKeyHandles, const TArray<FKeyPosition>& InAllKeyPositions, const TArray<int32>& ContiguousKeyIndices);
+	};
+
+	//all the selectedv keys in a row for a particular curve, plust pointers to all of the keys
+	struct FContiguousKeysArray
+	{
+		TArray<FKeyHandle> AllKeyHandles;
+		TArray<FKeyPosition> AllKeyPositions;
+
+		TArray<FContiguousKeys> KeysArray;
+
+		FContiguousKeysArray(const TArray<FKeyHandle>& InAllKeyHandle, const TArray<FKeyPosition>& InAllKeyPositions);
+		void Add(const TArray<int32>& ContiguousKeyIndices);
+	};
+
+	bool Setup(TWeakPtr<ISequencer>& InSequencer);
 
-//first and second key to blend between for each channel
+	TSharedPtr<FCurveEditor> CurveEditor;
+	TMap<FCurveModelID, FContiguousKeysArray> KeyMap;
+private:
+	void GetMapOfContiguousKeys();
+
+};
+
+/*
+*  Key bounds and their values used to blend with Controls and Actor float channels
+*/
 struct FChannelKeyBounds
 {
-	FChannelKeyBounds() : bValid(false), Channel(nullptr), FirstIndex(INDEX_NONE), SecondIndex(INDEX_NONE), FirstFrame(0), SecondFrame(0), FirstValue(0.0f), SecondValue(0.0f) {}
+	FChannelKeyBounds() : bValid(false), FloatChannel(nullptr), DoubleChannel(nullptr), PreviousIndex(INDEX_NONE), NextIndex(INDEX_NONE), PreviousFrame(0), NextFrame(0),
+		CurrentFrame(0), PreviousValue(0.0), NextValue(0.0), CurrentValue(0.0) {}
 	bool bValid;
-	FMovieSceneFloatChannel* Channel;
-	int32 FirstIndex;
-	int32 SecondIndex;
-	FFrameNumber FirstFrame;
-	FFrameNumber SecondFrame;
-	float FirstValue;
-	float SecondValue;
+	FMovieSceneFloatChannel* FloatChannel;
+	FMovieSceneDoubleChannel* DoubleChannel;
+	int32 PreviousIndex;
+	int32 NextIndex;
+	FFrameNumber PreviousFrame;
+	FFrameNumber NextFrame;
+	FFrameNumber CurrentFrame;
+	double PreviousValue;
+	double NextValue;
+	double CurrentValue;
 };
 
-struct FControlRigChannels
+/*
+* Contains the selection state for a set of Control Rig Controls to blend with the anim slider
+*/
+struct FAnimSliderObjectSelection
 {
-	FControlRigChannels() : Section(nullptr), NumChannels(0) {};
-	FChannelKeyBounds  KeyBounds[9];
-	UMovieSceneSection* Section;
-	int NumChannels; //may not use all 9 channels.
+	//set of possible float channels
+	struct FObjectChannels
+	{
+		FObjectChannels() : Section(nullptr) {};
+		TArray<FChannelKeyBounds>  KeyBounds;
+		UMovieSceneSection* Section;
+	};
+
+	bool Setup(TWeakPtr<ISequencer>& InSequencer);
+	bool Setup(const TArray<UControlRig*>& SelectedControlRigs, TWeakPtr<ISequencer>& InSequencer);
+
+	TArray<FObjectChannels>  ChannelsArray;
+
+private:
+
+	//float or double will be non-null
+	void SetupChannel(FFrameNumber CurrentFrame, TArray<FFrameNumber>& KeyTimes, TArray<FKeyHandle>& Handles, FMovieSceneFloatChannel* FloatChannel,
+		FMovieSceneDoubleChannel* DoubleChannel, FChannelKeyBounds& KeyBounds);
+
+	TArray<UControlRig*> GetControlRigs();
 };
 
-struct FControlsToTween
+/*
+* Base struct that performs common actions for slider animation tools
+*/
+struct FBaseAnimSlider
 {
+	FBaseAnimSlider(){};
+	virtual ~FBaseAnimSlider() {};
 	/**
-	 Setup up controls to tween
-	* @param InControlRig Control Rig to use
-	* @param InSequencer Sequencer to get keys from
+	 Setup up what's needed for the slider to blend
+	* @param InSequencer Sequencer to get keys, controls, etc. from
+	* @param returns true if we have something that we can blend
 	*/
-	void Setup(const TArray<UControlRig*>& SelectedControlRigs, TWeakPtr<ISequencer>& InSequencer);
+	virtual bool Setup(TWeakPtr<ISequencer>& InSequencer);
 
 	/**
 	* @param InSequencer Sequencer to blend at current time
 	* @param InSequencer BlendValue where in time to Blend in at Current Frame, 0.0 will be at current frame, -1.0 will value at previous key, 1.0 value at next key,
 	* other values interpolated between.
 	*/
-	void Blend(TWeakPtr<ISequencer>& InSequencer,  float BlendValue);
+	virtual bool Blend(TWeakPtr<ISequencer>& InSequencer, const double BlendValue) = 0;
 
+	/**
+	* @return the name of the slider
+	*/
+	virtual FText GetText() = 0;
+
+	/**
+	* @return the tooltip text for the slider
+	*/
+	virtual FText GetTooltipText() = 0;
+
+protected:
+
+	FAnimSliderKeySelection  KeySelection;
+	FAnimSliderObjectSelection ObjectSelection;
+};
+
+/*
+* Simple manager that holds all of the registered blend tools
+* This will then be used by whatever the view wants to do to display these tool
+*/
+class FAnimBlendTooLManager
+{
+public:
+	FAnimBlendTooLManager() {};
+	void RegisterAnimSlider(TSharedPtr<FBaseAnimSlider> AnimSlider)
+	{
+		AnimSliders.Add(AnimSlider);
+	}
+	const TArray<TSharedPtr<FBaseAnimSlider>>& GetAnimSliders() const 
+	{ 
+		return AnimSliders; 
+	}
+	TArray<TSharedPtr<FBaseAnimSlider>>& GetAnimSliders() 
+	{
+		return AnimSliders;
+	}
 private:
-	void SetupControlRigChannel(FFrameNumber CurrentFrame, TArray<FFrameNumber>& KeyTimes, TArray<FKeyHandle>& Handles, FMovieSceneFloatChannel* FloatChannel,
-		FChannelKeyBounds& KeyBounds);
+	TArray<TSharedPtr<FBaseAnimSlider>> AnimSliders;
+};
 
-	TMap<TPair<UControlRig*,FName>, FControlRigChannels>  ControlRigChannelsMap;
+/**
+*
+*  Slider that used the next and previous values to do the blend
+* 
+**/
+struct FBasicBlendSlider : public FBaseAnimSlider
+{
+	virtual double DoBlend(const double PreviousTime, const double PreviousValue, const double CurrentTime, const double CurrentValue,
+		const double NextTime, const double NextValue, const double BlendValue) = 0;
+	//BaseAnimSlider overrides
+	virtual bool Blend(TWeakPtr<ISequencer>& InSequencer, const double BlendValue) override;
 };
 
+/**
+*
+*    Classic tween slider, will move to the blend between the next and previous keys
+*
+**/
+struct FControlsToTween :public FBasicBlendSlider
+{
+	//BaseAnimSlider overrides
+	virtual FText GetText() override;
+	virtual FText GetTooltipText() override;
+	virtual bool Setup(TWeakPtr<ISequencer>& InSequencer) override;
+
+	//BasicBlendSlider overrides
+	virtual double DoBlend(const double PreviousTime, const double PreviousValue, const double CurrentTime, const double CurrentValue,
+		const double NextTime, const double NextValue, const double BlendValue) override;
+
+	bool Setup(const TArray<UControlRig*>& SelectedControlRigs, TWeakPtr<ISequencer>& InSequencer);
 
+};
+
+/**
+*
+*   Slider that will amplify or reduce the differences from the slope between the next or previous
+*
+**/
+struct FAmplifySlider :public FBasicBlendSlider
+{
+	//BaseAnimSlider overrides
+	virtual FText GetText();
+	virtual FText GetTooltipText();
+
+	//BasicBlendSlider overrides
+	virtual double DoBlend(const double PreviousTime, const double PreviousValue, const double CurrentTime, const double CurrentValue,
+		const double NextTime, const double NextValue, const double BlendValue) override;
+};
+
+/**
+*
+*   Similar to tween but will blend off of the current position when blending to the next or previous keys
+*
+**/
+struct FBlendToPreviousNextSlider : public FBasicBlendSlider
+{
+	//BaseAnimSlider overrides
+	virtual FText GetText();
+	virtual FText GetTooltipText();
+
+	//BasicBlendSlider overrides
+	virtual double DoBlend(const double PreviousTime, const double PreviousValue, const double CurrentTime, const double CurrentValue,
+		const double NextTime, const double NextValue, const double BlendValue) override;
+};
 
 
 
diff --git a/Engine/Source/Editor/Sequencer/Private/Sequencer.h b/Engine/Source/Editor/Sequencer/Private/Sequencer.h
--- a/Engine/Source/Editor/Sequencer/Private/Sequencer.h
+++ b/Engine/Source/Editor/Sequencer/Private/Sequencer.h
@@ -617,7 +617,7 @@ public:
 
 	void OnSelectedNodesOnlyChanged();
 
-	TSharedPtr<FCurveEditor> GetCurveEditor() const
+	virtual TSharedPtr<FCurveEditor> GetCurveEditor() const override
 	{
 		return CurveEditorModel;
 	}
diff --git a/Engine/Source/Editor/Sequencer/Public/ISequencer.h b/Engine/Source/Editor/Sequencer/Public/ISequencer.h
--- a/Engine/Source/Editor/Sequencer/Public/ISequencer.h
+++ b/Engine/Source/Editor/Sequencer/Public/ISequencer.h
@@ -210,6 +210,9 @@
 	/** @return a multicast delegate which is executed when sequencer closes. */
 	virtual FOnCloseEvent& OnCloseEvent() = 0;
 
+	/** @return associated curve editor*/
+	virtual TSharedPtr<FCurveEditor> GetCurveEditor() const = 0;
+
 	/** @return Widget used to display the sequencer */
 	virtual TSharedRef<SWidget> GetSequencerWidget() const = 0;
 
