Constraints in Sequencer:

- handled notifications for child/parent deletion, control rig replaced, control element renamed, etc.
- added FMovieSceneConstraintChannel to manage the constraint's active state animation + drawing overrides
- added dynamic mode for parent constraints when animating constraints to treat control animation as constraint's space local animation. This is experimetal and will have to be done for other types once stable.

diff --git a/Engine/Plugins/Animation/ControlRig/Source/ControlRig/Private/Constraints/ControlRigTransformableHandle.cpp b/Engine/Plugins/Animation/ControlRig/Source/ControlRig/Private/Constraints/ControlRigTransformableHandle.cpp
--- a/Engine/Plugins/Animation/ControlRig/Source/ControlRig/Private/Constraints/ControlRigTransformableHandle.cpp
+++ b/Engine/Plugins/Animation/ControlRig/Source/ControlRig/Private/Constraints/ControlRigTransformableHandle.cpp
@@ -12,7 +12,9 @@
  */
 
 UTransformableControlHandle::~UTransformableControlHandle()
-{}
+{
+	UnregisterDelegates();
+}
 
 bool UTransformableControlHandle::IsValid() const
 {
@@ -134,6 +136,11 @@
 	return 0;
 }
 
+TWeakObjectPtr<UObject> UTransformableControlHandle::GetTarget() const
+{
+	return GetSkeletalMesh();
+}
+
 USkeletalMeshComponent* UTransformableControlHandle::GetSkeletalMesh() const
 {
 	const TSharedPtr<IControlRigObjectBinding> ObjectBinding = ControlRig.IsValid() ? ControlRig->GetObjectBinding() : nullptr;
@@ -150,7 +157,78 @@
 	return ControlRig->FindControl(ControlName);
 }
 
+void UTransformableControlHandle::UnregisterDelegates() const
+{
 #if WITH_EDITOR
+	FCoreUObjectDelegates::OnObjectsReplaced.RemoveAll(this);
+#endif
+	
+	if (ControlRig.IsValid())
+	{
+		if (URigHierarchy* Hierarchy = ControlRig->GetHierarchy())
+		{
+			Hierarchy->OnModified().RemoveAll(this);
+		}
+	}
+}
+
+void UTransformableControlHandle::RegisterDelegates()
+{
+	UnregisterDelegates();
+
+#if WITH_EDITOR
+	FCoreUObjectDelegates::OnObjectsReplaced.AddUObject(this, &UTransformableControlHandle::OnObjectsReplaced);
+#endif
+
+	if (ControlRig.IsValid())
+	{
+		if (URigHierarchy* Hierarchy = ControlRig->GetHierarchy())
+		{
+			Hierarchy->OnModified().AddUObject(this, &UTransformableControlHandle::OnHierarchyModified);
+		}
+	}
+}
+
+void UTransformableControlHandle::OnHierarchyModified(
+	ERigHierarchyNotification InNotif,
+	URigHierarchy* InHierarchy,
+	const FRigBaseElement* InElement)
+{
+	if (!ControlRig.IsValid())
+	{
+	 	return;
+	}
+
+	const URigHierarchy* Hierarchy = ControlRig->GetHierarchy();
+	if (!Hierarchy || InHierarchy != Hierarchy)
+	{
+		return;
+	}
+
+	switch (InNotif)
+	{
+		case ERigHierarchyNotification::ElementRemoved:
+		{
+			// FIXME this leaves the constraint invalid as the element won't exist anymore
+			// find a way to remove this from the constraints list 
+			break;
+		}
+		case ERigHierarchyNotification::ElementRenamed:
+		{
+			const FName OldName = Hierarchy->GetPreviousName(InElement->GetKey());
+			if (OldName == ControlName)
+			{
+				ControlName = InElement->GetName();
+			}
+			break;
+		}
+		default:
+			break;
+	}
+}
+
+#if WITH_EDITOR
+
 FName UTransformableControlHandle::GetName() const
 {
 	const USkeletalMeshComponent* SkeletalMesh = GetSkeletalMesh();
@@ -162,4 +240,26 @@
 
 	return FName(*FullName);
 }
+
+void UTransformableControlHandle::OnObjectsReplaced(const TMap<UObject*, UObject*>& InOldToNewInstances)
+{
+	if (UObject* NewObject = InOldToNewInstances.FindRef(ControlRig.Get()))
+	{
+		if (UControlRig* NewControlRig = Cast<UControlRig>(NewObject))
+		{
+			if (URigHierarchy* Hierarchy = ControlRig->GetHierarchy())
+			{
+				Hierarchy->OnModified().RemoveAll(this);
+			}
+			
+			ControlRig = NewControlRig;
+			
+			if (URigHierarchy* Hierarchy = ControlRig->GetHierarchy())
+			{
+				Hierarchy->OnModified().AddUObject(this, &UTransformableControlHandle::OnHierarchyModified);
+			}
+		}
+	}
+}
+
 #endif
diff --git a/Engine/Plugins/Animation/ControlRig/Source/ControlRig/Private/ControlRigModule.cpp b/Engine/Plugins/Animation/ControlRig/Source/ControlRig/Private/ControlRigModule.cpp
--- a/Engine/Plugins/Animation/ControlRig/Source/ControlRig/Private/ControlRigModule.cpp
+++ b/Engine/Plugins/Animation/ControlRig/Source/ControlRig/Private/ControlRigModule.cpp
@@ -46,6 +46,7 @@ void FControlRigModule::RegisterTransformableCustomization() const
 			UTransformableControlHandle* CtrlHandle = NewObject<UTransformableControlHandle>(Outer);
 			CtrlHandle->ControlRig = ControlActor->ControlRig;
 			CtrlHandle->ControlName = ControlActor->ControlName;
+			CtrlHandle->RegisterDelegates();
 			return CtrlHandle;
 		}
 		return nullptr;
diff --git a/Engine/Plugins/Animation/ControlRig/Source/ControlRig/Private/Sequencer/MovieSceneControlRigParameterSection.cpp b/Engine/Plugins/Animation/ControlRig/Source/ControlRig/Private/Sequencer/MovieSceneControlRigParameterSection.cpp
--- a/Engine/Plugins/Animation/ControlRig/Source/ControlRig/Private/Sequencer/MovieSceneControlRigParameterSection.cpp
+++ b/Engine/Plugins/Animation/ControlRig/Source/ControlRig/Private/Sequencer/MovieSceneControlRigParameterSection.cpp
@@ -1,6 +1,8 @@
 // Copyright Epic Games, Inc. All Rights Reserved.
 
 #include "Sequencer/MovieSceneControlRigParameterSection.h"
+
+#include "ConstraintsManager.h"
 #include "Animation/AnimSequence.h"
 #include "Logging/MessageLog.h"
 #include "Compilation/MovieSceneTemplateInterrogation.h"
@@ -1274,6 +1276,66 @@ void UMovieSceneControlRigParameterSection::AddSpaceChannel(FName InControlName,
 	}
 }
 
+bool UMovieSceneControlRigParameterSection::HasConstraintChannel(const FName& InConstraintName) const
+{
+	return ConstraintsChannels.ContainsByPredicate( [InConstraintName](const FConstraintAndActiveChannel& InChannel)
+	{
+		return InChannel.Constraint.IsValid() ? InChannel.Constraint->GetFName() == InConstraintName : false;
+	});
+}
+
+FConstraintAndActiveChannel* UMovieSceneControlRigParameterSection::GetConstraintChannel(const FName& InConstraintName)
+{
+	const int32 Index = ConstraintsChannels.IndexOfByPredicate([InConstraintName](const FConstraintAndActiveChannel& InChannel)
+	{
+		return InChannel.Constraint.IsValid() ? InChannel.Constraint->GetFName() == InConstraintName : false;
+	});
+	return (Index != INDEX_NONE) ? &ConstraintsChannels[Index] : nullptr;	
+}
+
+void UMovieSceneControlRigParameterSection::AddConstraintChannel(UTickableConstraint* InConstraint, bool bReconstructChannel)
+{
+	if (!HasConstraintChannel(InConstraint->GetFName()))
+	{
+		Modify();
+		
+		const int32 NewIndex = ConstraintsChannels.Add(FConstraintAndActiveChannel(InConstraint));
+
+		FMovieSceneConstraintChannel* ExistingChannel = &ConstraintsChannels[NewIndex].ActiveChannel;
+		ExistingChannel->SetDefault(false);
+		
+		if (bReconstructChannel)
+		{
+			ReconstructChannelProxy();
+		}
+
+		if (!OnConstraintRemovedHandle.IsValid())
+		{
+			if (ControlRig)
+			{
+				FConstraintsManagerController& Controller = FConstraintsManagerController::Get(ControlRig->GetWorld());
+				OnConstraintRemovedHandle = Controller.OnConstraintRemoved().AddLambda([this](FName InConstraintName)
+				{
+					Modify();
+					const int32 NumRemoved = ConstraintsChannels.RemoveAll( [InConstraintName](const FConstraintAndActiveChannel& InChannel)
+					{
+						return InChannel.Constraint.IsValid() ? InChannel.Constraint->GetFName() == InConstraintName : false;
+					});
+					if (NumRemoved)
+					{
+						ReconstructChannelProxy();
+					}
+				});
+			}
+		}
+	}
+}
+
+const TArray<FConstraintAndActiveChannel>& UMovieSceneControlRigParameterSection::GetConstraintsChannels() const
+{
+	return ConstraintsChannels;
+}
+
 TArray<FSpaceControlNameAndChannel>& UMovieSceneControlRigParameterSection::GetSpaceChannels()
 {
 	return SpaceChannels;
@@ -1312,7 +1374,7 @@ void UMovieSceneControlRigParameterSection::ReconstructChannelProxy()
 			OnArray.Init(true, ControlRig->AvailableControls().Num());
 			SetControlsMask(OnArray);
 		}
-			
+
 		int32 ControlIndex = 0; 
 		int32 MaskIndex = 0;
 		int32 TotalIndex = 0; 
@@ -1322,11 +1384,52 @@ void UMovieSceneControlRigParameterSection::ReconstructChannelProxy()
 		int32 IntegerChannelIndex = 0;
 		int32 SpaceChannelIndex = 0;
 		int32 CategoryIndex = 0;
+		int32 ConstraintsChannelIndex = 0;
+		
 		const FName BoolChannelTypeName = FMovieSceneBoolChannel::StaticStruct()->GetFName();
 		const FName EnumChannelTypeName = FMovieSceneByteChannel::StaticStruct()->GetFName();
 		const FName IntegerChannelTypeName = FMovieSceneIntegerChannel::StaticStruct()->GetFName();
 		const FName SpaceName = FName(TEXT("Space"));
 
+		// begin constraints
+		const FConstraintsManagerController& Controller = FConstraintsManagerController::Get(ControlRig->GetWorld());
+		auto GetConstraints = [&Controller, this](const FName& InControlName)
+		{
+			static constexpr bool bSorted = true;
+			const uint32 ControlHash = HashCombine(GetTypeHash(ControlRig.Get()), GetTypeHash(InControlName));
+			return Controller.GetParentConstraints(ControlHash, bSorted);
+		};
+
+
+		auto AddConstrainChannels = [this, GetConstraints, &ConstraintsChannelIndex, &TotalIndex, &Channels](const FName& InControlName, const FText& InGroup, const bool bEnabled)
+		{
+			TArray<TObjectPtr<UTickableConstraint>> Constraints = GetConstraints(InControlName);
+			for (const TObjectPtr<UTickableConstraint>& Constraint: Constraints)
+			{
+				const FName& ConstraintName = Constraint->GetFName();
+				if(FConstraintAndActiveChannel* ConstraintChannel = GetConstraintChannel(ConstraintName))
+				{
+					if (FChannelMapInfo* ChannelInfo = ControlChannelMap.Find(InControlName))
+					{
+						ChannelInfo->ConstraintsIndex.Add(ConstraintsChannelIndex);
+					}
+
+#if WITH_EDITOR
+					FMovieSceneChannelMetaData MetaData(ConstraintName,  FText::FromName(ConstraintName), InGroup, bEnabled);
+					ConstraintsChannelIndex += 1;
+					MetaData.SortOrder = TotalIndex++;
+					MetaData.bCanCollapseToTrack = false;
+
+					// Channels.Add(ConstraintChannel->ActiveCurve.ParameterCurve, MetaData, TMovieSceneExternalValue<bool>());
+					Channels.Add(ConstraintChannel->ActiveChannel, MetaData, TMovieSceneExternalValue<bool>());
+#else
+					Channels.Add(ConstraintChannel->ActiveChannel);
+#endif
+				}
+			}
+		};
+		// end constraints
+		
 		for (FRigControlElement* ControlElement : SortedControls)
 		{
 			if (!ControlRig->GetHierarchy()->IsAnimatable(ControlElement))
@@ -1591,6 +1694,7 @@ void UMovieSceneControlRigParameterSection::ReconstructChannelProxy()
 					{
 						if (ControlElement->GetName() == Transform.ParameterName)
 						{
+							const FName ControlName = ControlElement->GetName();
 							if(Group.IsEmpty())
 							{
 								ControlChannelMap.Add(Transform.ParameterName, FChannelMapInfo(ControlIndex, TotalIndex, FloatChannelIndex, INDEX_NONE, NAME_None, MaskIndex, CategoryIndex));
@@ -1607,9 +1711,12 @@ void UMovieSceneControlRigParameterSection::ReconstructChannelProxy()
 								ControlChannelMap.Add(Transform.ParameterName, FChannelMapInfo(ControlIndex, TotalIndex, FloatChannelIndex, ParentControlIndex, NAME_None, MaskIndex, CategoryIndex));
 							}
 
+							// constraints
+							AddConstrainChannels(ControlName, Group, bEnabled);
+
+							// spaces
 							if (FSpaceControlNameAndChannel* SpaceChannel = GetSpaceChannel(Transform.ParameterName))
 							{
-
 								FChannelMapInfo* pChannelIndex = ControlChannelMap.Find(Transform.ParameterName);
 								if (pChannelIndex)
 								{
@@ -2759,6 +2866,7 @@ void UMovieSceneControlRigParameterSection::ClearAllParameters()
 	EnumParameterNamesAndCurves.SetNum(0);
 	IntegerParameterNamesAndCurves.SetNum(0);
 	SpaceChannels.SetNum(0);
+	ConstraintsChannels.SetNum(0);
 }
 void UMovieSceneControlRigParameterSection::RemoveAllKeys(bool bIncludeSpaceKeys)
 {
diff --git a/Engine/Plugins/Animation/ControlRig/Source/ControlRig/Private/Sequencer/MovieSceneControlRigParameterTemplate.cpp b/Engine/Plugins/Animation/ControlRig/Source/ControlRig/Private/Sequencer/MovieSceneControlRigParameterTemplate.cpp
--- a/Engine/Plugins/Animation/ControlRig/Source/ControlRig/Private/Sequencer/MovieSceneControlRigParameterTemplate.cpp
+++ b/Engine/Plugins/Animation/ControlRig/Source/ControlRig/Private/Sequencer/MovieSceneControlRigParameterTemplate.cpp
@@ -170,6 +170,15 @@
 	}
 };
 
+struct FConstraintAndActiveValue
+{
+	FConstraintAndActiveValue(TWeakObjectPtr<UTickableConstraint> InConstraint, bool InValue)
+		: Constraint(InConstraint)
+		, Value(InValue)
+	{}
+	TWeakObjectPtr<UTickableConstraint> Constraint;
+	bool Value;
+};
 
 struct FControlRigAnimTypeIDs;
 
@@ -373,6 +382,8 @@
 	TArray<FColorParameterStringAndValue, TInlineAllocator<2>> ColorValues;
 	/** Array of evaluated transform values */
 	TArray<FEulerTransformParameterStringAndValue, TInlineAllocator<2>> TransformValues;
+	/** Array of evaluated constraint values */
+	TArray<FConstraintAndActiveValue, TInlineAllocator<2>> ConstraintsValues;
 };
 
 /** Token for control rig control parameters */
@@ -1082,6 +1093,7 @@
 		BoolValues = Values.BoolValues;
 		IntegerValues = Values.IntegerValues;
 		SpaceValues = Values.SpaceValues;
+		ConstraintsValues = Values.ConstraintsValues;
 	}
 	FControlRigParameterExecutionToken(FControlRigParameterExecutionToken&&) = default;
 	FControlRigParameterExecutionToken& operator=(FControlRigParameterExecutionToken&&) = default;
@@ -1237,6 +1249,14 @@
 						}
 					}
 				}
+
+				for (const FConstraintAndActiveValue& ConstraintValue : ConstraintsValues)
+				{
+					if (ConstraintValue.Constraint.IsValid())
+					{
+						ConstraintValue.Constraint->SetActive(ConstraintValue.Value);
+					}
+				}
 			}
 			Section->SetDoNotKey(bWasDoNotKey);
 		}
@@ -1250,16 +1270,20 @@
 	TArray<FIntegerParameterStringAndValue, TInlineAllocator<2>> IntegerValues;
 	/** Array of Space Values*/
 	TArray<FControlSpaceAndValue, TInlineAllocator<2>> SpaceValues;
-
+	/** Array of evaluated bool values */
+	TArray<FConstraintAndActiveValue, TInlineAllocator<2>> ConstraintsValues;
 };
 
 
-FMovieSceneControlRigParameterTemplate::FMovieSceneControlRigParameterTemplate(const UMovieSceneControlRigParameterSection& Section, const UMovieSceneControlRigParameterTrack& Track)
-	: FMovieSceneParameterSectionTemplate(Section) , Enums(Section.GetEnumParameterNamesAndCurves()), Integers(Section.GetIntegerParameterNamesAndCurves()), Spaces(Section.GetSpaceChannels())
-
-{
-
-}
+FMovieSceneControlRigParameterTemplate::FMovieSceneControlRigParameterTemplate(
+	const UMovieSceneControlRigParameterSection& Section,
+	const UMovieSceneControlRigParameterTrack& Track)
+	: FMovieSceneParameterSectionTemplate(Section)
+	, Enums(Section.GetEnumParameterNamesAndCurves())
+	, Integers(Section.GetIntegerParameterNamesAndCurves())
+	, Spaces(Section.GetSpaceChannels())
+	, Constraints(Section.GetConstraintsChannels())
+{}
 
 
 
@@ -1725,6 +1749,7 @@
 		Values.VectorValues.Reserve(Vectors.Num());
 		Values.ColorValues.Reserve(Colors.Num());
 		Values.TransformValues.Reserve(Transforms.Num());
+		Values.ConstraintsValues.Reserve(Constraints.Num());
 
 		// Populate each of the output arrays in turn
 		for (int32 Index = 0; Index < Scalars.Num(); ++Index)
@@ -1755,6 +1780,15 @@
 			Values.SpaceValues.Emplace(Space.ControlName, Value);
 		}
 
+		for (int32 Index = 0; Index < Constraints.Num(); ++Index)
+		{
+			bool Value = false;
+			const FConstraintAndActiveChannel& ConstraintAndActiveChannel = Constraints[Index];
+			ConstraintAndActiveChannel.ActiveChannel.Evaluate(Time, Value);
+			
+			Values.ConstraintsValues.Emplace(ConstraintAndActiveChannel.Constraint, Value);
+		}
+		
 		for (int32 Index = 0; Index < Bools.Num(); ++Index)
 		{
 			bool Value = false;
diff --git a/Engine/Plugins/Animation/ControlRig/Source/ControlRig/Public/Constraints/ControlRigTransformableHandle.h b/Engine/Plugins/Animation/ControlRig/Source/ControlRig/Public/Constraints/ControlRigTransformableHandle.h
--- a/Engine/Plugins/Animation/ControlRig/Source/ControlRig/Public/Constraints/ControlRigTransformableHandle.h
+++ b/Engine/Plugins/Animation/ControlRig/Source/ControlRig/Public/Constraints/ControlRigTransformableHandle.h
@@ -3,13 +3,17 @@
 #pragma once
 
 #include "TransformableHandle.h"
+#include "Rigs/RigHierarchyDefines.h"
 
 #include "ControlRigTransformableHandle.generated.h"
 
 class UControlRig;
 class USkeletalMeshComponent;
+struct FRigBaseElement;
+class URigHierarchy;
 struct FRigControlElement;
 
+
 /**
  * UTransformableControlHandle
  */
@@ -20,9 +24,8 @@
 	GENERATED_BODY()
 	
 public:
-	
 	virtual ~UTransformableControlHandle();
-
+	
 	/** @todo document */
 	virtual bool IsValid() const override;
 
@@ -42,6 +45,14 @@
 
 	/** @todo document */
 	virtual uint32 GetHash() const override;
+	/** @todo document */
+	virtual TWeakObjectPtr<UObject> GetTarget() const override;
+	/** @todo document */
+	USkeletalMeshComponent* GetSkeletalMesh() const;
+
+	/** @todo document */
+	void UnregisterDelegates() const;
+	void RegisterDelegates();
 
 #if WITH_EDITOR
 	/** @todo document */
@@ -51,7 +62,7 @@
 	/** @todo document */
 	UPROPERTY()
 	TWeakObjectPtr<UControlRig> ControlRig;
-	
+
 	/** @todo document */
 	UPROPERTY()
 	FName ControlName;
@@ -59,7 +70,15 @@
 private:
 
 	/** @todo document */
-	USkeletalMeshComponent* GetSkeletalMesh() const;
+	void OnHierarchyModified(
+		ERigHierarchyNotification InNotif,
+		URigHierarchy* InHierarchy,
+		const FRigBaseElement* InElement);
+
+#if WITH_EDITOR
+	/** @todo document */
+	void OnObjectsReplaced(const TMap<UObject*, UObject*>& InOldToNewInstances);
+#endif
 
 	/** @todo document */
 	FRigControlElement* GetControlElement() const;
diff --git a/Engine/Plugins/Animation/ControlRig/Source/ControlRig/Public/Sequencer/MovieSceneControlRigParameterSection.h b/Engine/Plugins/Animation/ControlRig/Source/ControlRig/Public/Sequencer/MovieSceneControlRigParameterSection.h
--- a/Engine/Plugins/Animation/ControlRig/Source/ControlRig/Public/Sequencer/MovieSceneControlRigParameterSection.h
+++ b/Engine/Plugins/Animation/ControlRig/Source/ControlRig/Public/Sequencer/MovieSceneControlRigParameterSection.h
@@ -3,6 +3,7 @@
 #pragma once
 
 #include "CoreMinimal.h"
+#include "ConstraintsManager.h"
 #include "Sections/MovieSceneParameterSection.h"
 #include "UObject/ObjectMacros.h"
 #include "Channels/MovieSceneFloatChannel.h"
@@ -14,6 +15,8 @@
 #include "Compilation/MovieSceneTemplateInterrogation.h"
 #include "Channels/MovieSceneIntegerChannel.h"
 #include "Channels/MovieSceneByteChannel.h"
+#include "Channels/MovieSceneBoolChannel.h"
+#include "Channels/MovieSceneConstraintChannel.h"
 #include "Sequencer/MovieSceneControlRigSpaceChannel.h"
 
 
@@ -103,8 +106,25 @@ struct CONTROLRIG_API FSpaceControlNameAndChannel
 	FMovieSceneControlRigSpaceChannel SpaceCurve;
 };
 
+USTRUCT()
+struct CONTROLRIG_API FConstraintAndActiveChannel
+{
+	GENERATED_USTRUCT_BODY()
+
+	FConstraintAndActiveChannel(){}
+	FConstraintAndActiveChannel(const TObjectPtr<UTickableConstraint>& InConstraint)
+		: Constraint(InConstraint)
+	{};
+
+	UPROPERTY()
+	TWeakObjectPtr<UTickableConstraint> Constraint;
+
+	UPROPERTY()
+	FMovieSceneConstraintChannel ActiveChannel;
+};
+
 /**
-*  Data that's queried during an interrogtion
+*  Data that's queried during an interrogation
 */
 struct FFloatInterrogationData
 {
@@ -161,7 +181,8 @@ struct CONTROLRIG_API FChannelMapInfo
 
 	int32 GeneratedKeyIndex = -1; //temp index set by the ControlRigParameterTrack, not saved
 
-
+	UPROPERTY()
+	TArray<uint32> ConstraintsIndex; //constraints data
 };
 
 
@@ -200,6 +221,8 @@ public:
 	
 	FSpaceChannelAddedEvent& SpaceChannelAdded() { return OnSpaceChannelAdded; }
 
+	const TArray<FConstraintAndActiveChannel>& GetConstraintsChannels() const;
+	
 	bool RenameParameterName(const FName& OldParameterName, const FName& NewParameterName);
 private:
 
@@ -244,6 +267,10 @@ protected:
 	UPROPERTY()
 	TArray<FSpaceControlNameAndChannel>  SpaceChannels;
 
+	/** Space Channels*/
+	UPROPERTY()
+	TArray<FConstraintAndActiveChannel> ConstraintsChannels;
+
 public:
 
 	UMovieSceneControlRigParameterSection();
@@ -388,6 +415,12 @@ public:
 	/** Add Space Parameter for a specified Control, no Default since that is Parent space*/
 	void AddSpaceChannel(FName InControlName, bool bReconstructChannel);
 
+	/** todo */
+	FDelegateHandle OnConstraintRemovedHandle;
+	bool HasConstraintChannel(const FName& InConstraintName) const;
+	FConstraintAndActiveChannel* GetConstraintChannel(const FName& InConstraintName);
+	void AddConstraintChannel(UTickableConstraint* InConstraint, bool bReconstructChannel);
+
 	/** Clear Everything Out*/
 	void ClearAllParameters();
 
diff --git a/Engine/Plugins/Animation/ControlRig/Source/ControlRig/Public/Sequencer/MovieSceneControlRigParameterTemplate.h b/Engine/Plugins/Animation/ControlRig/Source/ControlRig/Public/Sequencer/MovieSceneControlRigParameterTemplate.h
--- a/Engine/Plugins/Animation/ControlRig/Source/ControlRig/Public/Sequencer/MovieSceneControlRigParameterTemplate.h
+++ b/Engine/Plugins/Animation/ControlRig/Source/ControlRig/Public/Sequencer/MovieSceneControlRigParameterTemplate.h
@@ -38,5 +38,8 @@ protected:
 	/** Controls and their space keys*/
 	UPROPERTY()
 	TArray<FSpaceControlNameAndChannel> Spaces;
+	
+	UPROPERTY()
+	TArray<FConstraintAndActiveChannel> Constraints;
 };
 
diff --git a/Engine/Plugins/Animation/ControlRig/Source/ControlRigEditor/Private/Editor/Constraints/SConstraintsWidget.cpp b/Engine/Plugins/Animation/ControlRig/Source/ControlRigEditor/Private/Editor/Constraints/SConstraintsWidget.cpp
--- a/Engine/Plugins/Animation/ControlRig/Source/ControlRigEditor/Private/Editor/Constraints/SConstraintsWidget.cpp
+++ b/Engine/Plugins/Animation/ControlRig/Source/ControlRigEditor/Private/Editor/Constraints/SConstraintsWidget.cpp
@@ -3,6 +3,7 @@
 #include "SConstraintsWidget.h"
 
 #include "ActorPickerMode.h"
+#include "ConstraintChannelHelper.h"
 #include "ConstraintsActor.h"
 #include "ControlRigEditorStyle.h"
 #include "DetailLayoutBuilder.h"
@@ -202,7 +203,8 @@
 	if (MouseEvent.GetEffectingButton() == EKeys::LeftMouseButton)
 	{
 		bIsPressed = true;
-		return FReply::Handled().DetectDrag( SharedThis( this ), MouseEvent.GetEffectingButton() );
+		// return FReply::Handled().DetectDrag( SharedThis( this ), MouseEvent.GetEffectingButton() );
+		return CreateSelectionPicker();
 	}
 
 	return FReply::Unhandled();
@@ -221,30 +223,31 @@
 FReply SDroppableConstraintItem::OnDragDetected(const FGeometry& MyGeometry, const FPointerEvent& MouseEvent)
 {
 	bIsPressed = false;
+	
+	// if (MouseEvent.IsMouseButtonDown( EKeys::LeftMouseButton ))
+	// {
+	// 	return CreateSelectionPicker();
+	// }
+	
+	return FReply::Handled();
+}
 
+FReply SDroppableConstraintItem::CreateSelectionPicker() const
+{
 	// FIXME temp approach for selecting the parent
-	auto OnConstraintDragged = [this]()
-	{
-		static const FActorPickerModeModule& ActorPickerMode = FModuleManager::Get().GetModuleChecked<FActorPickerModeModule>("ActorPickerMode");
+	
+	static const FActorPickerModeModule& ActorPickerMode = FModuleManager::Get().GetModuleChecked<FActorPickerModeModule>("ActorPickerMode");
 
-		TSharedPtr<SConstraintsCreationWidget> ConstraintsCreationWidget = this->ConstraintsWidget.Pin();
-		ETransformConstraintType ConstraintTypeCopy = this->ConstraintType;
-		ActorPickerMode.BeginActorPickingMode(
-			FOnGetAllowedClasses(), 
-			FOnShouldFilterActor(), 
-			FOnActorSelected::CreateLambda([ConstraintsCreationWidget, ConstraintTypeCopy](AActor* InActor)
-			{
-				SDroppableConstraintItem::CreateConstraint(InActor, ConstraintsCreationWidget, ConstraintTypeCopy);
-			}) );
+	TSharedPtr<SConstraintsCreationWidget> ConstraintsCreationWidget = this->ConstraintsWidget.Pin();
+	ETransformConstraintType ConstraintTypeCopy = this->ConstraintType;
+	ActorPickerMode.BeginActorPickingMode(
+		FOnGetAllowedClasses(), 
+		FOnShouldFilterActor(), 
+		FOnActorSelected::CreateLambda([ConstraintsCreationWidget, ConstraintTypeCopy](AActor* InActor)
+		{
+			SDroppableConstraintItem::CreateConstraint(InActor, ConstraintsCreationWidget, ConstraintTypeCopy);
+		}) );
 
-		
-		return FReply::Handled();
-	};
-	
-	if (MouseEvent.IsMouseButtonDown( EKeys::LeftMouseButton ))
-	{
-		return OnConstraintDragged();
-	}
 	
 	return FReply::Handled();
 }
@@ -355,17 +358,11 @@
 		}
 	});
 
-	auto GetFont = [this]()
+	auto GetConstraint = [this]()
 	{
 		UWorld* World = GCurrentLevelEditingViewportClient->GetWorld();
 		const FConstraintsManagerController& Controller = FConstraintsManagerController::Get(World);
-		UTickableConstraint* Constraint = Controller.GetConstraint(ConstraintItem->Name);
-		if (!Constraint)
-		{
-			return IDetailLayoutBuilder::GetDetailFont();
-		}
-
-		return Constraint->Active ? IDetailLayoutBuilder::GetDetailFont() : IDetailLayoutBuilder::GetDetailFontItalic();
+		return Controller.GetConstraint(ConstraintItem->Name);
 	};
 	
 	ChildSlot
@@ -382,7 +379,14 @@
 			SNew(SBorder)
 			.Padding(FMargin(5.0, 2.0, 5.0, 2.0))
 			.BorderImage(RoundedBoxBrush)
-			.BorderBackgroundColor(FStyleColors::Transparent)
+			.BorderBackgroundColor_Lambda([GetConstraint]()
+			{
+				if (const UTickableConstraint* Constraint = GetConstraint())
+				{
+					return Constraint->Active ? FStyleColors::Select : FStyleColors::Transparent;
+				}
+				return FStyleColors::Transparent;
+			})
 			.Content()
 			[
 				SNew(SHorizontalBox)
@@ -417,7 +421,16 @@
 						const FString TextStr = FString::Printf(TEXT("%s (%s)"), *InItem->Name.ToString(), *InItem->ParentName.ToString());
 						return FText::FromString(TextStr);
 					})
-					.Font_Lambda(GetFont)
+					.Font_Lambda([GetConstraint]()
+					{
+						const UTickableConstraint* Constraint = GetConstraint();
+						if (!Constraint)
+						{
+							return IDetailLayoutBuilder::GetDetailFont();
+						}
+
+						return Constraint->Active ? IDetailLayoutBuilder::GetDetailFont() : IDetailLayoutBuilder::GetDetailFontItalic();
+					})
 				]
 			]
 		]
@@ -428,6 +441,39 @@
 			SNew(SSpacer)
 		]
 
+		// add key
+		+ SHorizontalBox::Slot()
+		.AutoWidth()
+		.VAlign(VAlign_Center)
+		.HAlign(HAlign_Left)
+		.Padding(0)
+		[
+			SNew(SButton)
+			.ButtonStyle(FAppStyle::Get(), TEXT("SimpleButton"))
+			.ContentPadding(0)
+			.OnClicked_Lambda(	[GetConstraint]()
+			{
+				UTickableConstraint* Constraint = GetConstraint();
+				if (UTickableTransformConstraint* TransformConstraint = Cast<UTickableTransformConstraint>(Constraint))
+				{
+					// FConstraintChannelHelper::AddConstraintKey(TransformConstraint);
+					FConstraintChannelHelper::SmartConstraintKey(TransformConstraint);
+				}
+				return FReply::Handled();
+			})
+			.IsEnabled_Lambda([this]()
+			{
+				// TODO check if we have a focussed sequencer
+				return true;
+			})
+			.ToolTipText(LOCTEXT("KeyConstraintToolTip", "Add an active keyframe for that constraint."))
+			[
+				SNew(SImage)
+				.Image(FAppStyle::GetBrush("Sequencer.AddKey.Details"))
+				.ColorAndOpacity(FSlateColor::UseForeground())
+			]
+		]
+		
 		// move up
 		+ SHorizontalBox::Slot()
 		.AutoWidth()
@@ -769,7 +815,7 @@
 	MenuBuilder.BeginSection("BakeConstraint", LOCTEXT("BakeConstraintHeader", "Bake Constraint"));
 	{
 		MenuBuilder.AddMenuEntry(
-		LOCTEXT("BakeConstraintDoIt", "Do It"),
+		LOCTEXT("BakeConstraintDoIt", "Bake"),
 		FText::Format(LOCTEXT("BakeConstraintDoItTooltip", "Bake {0} transforms."), FText::FromName(ListItems[Index]->Name)),
 		FSlateIcon(),
 		FUIAction(FExecuteAction::CreateLambda([Constraint]()
@@ -784,6 +830,25 @@
 	}
 	MenuBuilder.EndSection();
 
+	// test section
+	MenuBuilder.BeginSection("TestConstraint", LOCTEXT("TestConstraintHeader", "Test"));
+	{
+		MenuBuilder.AddMenuEntry(
+		LOCTEXT("AddConstraintKey", "Add Key"),
+		FText::Format(LOCTEXT("AddKeyTooltip", "Add active key for {0}."), FText::FromName(ListItems[Index]->Name)),
+		FSlateIcon(),
+		FUIAction(FExecuteAction::CreateLambda([Constraint]()
+		{
+			if (UTickableTransformConstraint* TransformConstraint = Cast<UTickableTransformConstraint>(Constraint))
+			{
+				FConstraintChannelHelper::AddConstraintKey(TransformConstraint);
+			}
+		})),
+		NAME_None,
+		EUserInterfaceActionType::Button);
+	}
+	MenuBuilder.EndSection();
+
 	return MenuBuilder.MakeWidget();
 }
 
diff --git a/Engine/Plugins/Animation/ControlRig/Source/ControlRigEditor/Private/Editor/Constraints/SConstraintsWidget.h b/Engine/Plugins/Animation/ControlRig/Source/ControlRigEditor/Private/Editor/Constraints/SConstraintsWidget.h
--- a/Engine/Plugins/Animation/ControlRig/Source/ControlRigEditor/Private/Editor/Constraints/SConstraintsWidget.h
+++ b/Engine/Plugins/Animation/ControlRig/Source/ControlRigEditor/Private/Editor/Constraints/SConstraintsWidget.h
@@ -77,6 +77,10 @@
 	// End of SWidget interface
 	
 private:
+
+	/** todo documentation */
+	FReply CreateSelectionPicker() const;
+	
 	/** Creates the constraint between the current selection and the picked actor. */
 	static void CreateConstraint(
 		AActor* InParent,
diff --git a/Engine/Plugins/Animation/ControlRig/Source/ControlRigEditor/Private/Sequencer/ConstraintChannelHelper.cpp b/Engine/Plugins/Animation/ControlRig/Source/ControlRigEditor/Private/Sequencer/ConstraintChannelHelper.cpp
new file mode 100644
--- /dev/null
+++ b/Engine/Plugins/Animation/ControlRig/Source/ControlRigEditor/Private/Sequencer/ConstraintChannelHelper.cpp
@@ -0,0 +1,682 @@
+﻿// Copyright Epic Games, Inc. All Rights Reserved.
+
+#include "ConstraintChannelHelper.h"
+
+#include "ActorForWorldTransforms.h"
+#include "ControlRigSpaceChannelEditors.h"
+#include "ISequencer.h"
+
+#include "Constraints/ControlRigTransformableHandle.h"
+#include "Sequencer/MovieSceneControlRigParameterTrack.h"
+#include "LevelSequence.h"
+#include "MovieSceneToolHelpers.h"
+
+#include "TransformConstraint.h"
+#include "Algo/Unique.h"
+#include "Channels/MovieSceneConstraintChannel.h"
+
+#include "Tools/BakingHelper.h"
+#include "Tools/ConstraintBaker.h"
+
+namespace
+{
+	bool CanAddKey(const FConstraintAndActiveChannel* Channel, const FFrameNumber& InTime, bool& ActiveValue)
+	{
+		const TMovieSceneChannelData<const bool> ChannelData = Channel->ActiveChannel.GetData();
+		const TArrayView<const FFrameNumber> Times = ChannelData.GetTimes();
+		if (Times.IsEmpty())
+		{
+			ActiveValue = true;
+			return true;
+		}
+
+		const TArrayView<const bool> Values = ChannelData.GetValues();
+		if (InTime < Times[0])
+		{
+			if (!Values[0])
+			{
+				ActiveValue = true;
+				return true;
+			}
+			return false;
+		}
+		
+		if (InTime > Times.Last())
+		{
+			ActiveValue = !Values.Last();
+			return true;
+		}
+
+		return false;
+	}
+}
+
+void FConstraintChannelHelper::SmartConstraintKey(UTickableTransformConstraint* InConstraint)
+{
+	const TWeakPtr<ISequencer> WeakSequencer = FBakingHelper::GetSequencer();
+	if (!WeakSequencer.IsValid() || !WeakSequencer.Pin()->GetFocusedMovieSceneSequence())
+	{
+		return;
+	}
+
+	if (const UTransformableComponentHandle* ComponentHandle = Cast<UTransformableComponentHandle>(InConstraint->ChildTRSHandle))
+	{
+		return SmartComponentConstraintKey(InConstraint, WeakSequencer.Pin());
+	}
+	
+	if (const UTransformableControlHandle* ControlHandle = Cast<UTransformableControlHandle>(InConstraint->ChildTRSHandle))
+	{
+		return SmartControlConstraintKey(InConstraint, WeakSequencer.Pin());
+	}
+}
+
+void FConstraintChannelHelper::AddChildTransformKey(
+	const UTransformableHandle* InHandle,
+	const FFrameNumber& InTime,
+	const TSharedPtr<ISequencer>& InSequencer)
+{
+	if (const UTransformableComponentHandle* ComponentHandle = Cast<UTransformableComponentHandle>(InHandle))
+	{
+		AActor* Actor = ComponentHandle->Component->GetOwner();
+		if (!Actor)
+		{
+			return;
+		}
+
+		const FTransform LocalTransform = ComponentHandle->GetLocalTransform();
+		const FGuid Guid = InSequencer->GetHandleToObject(Actor,true);
+		if (!Guid.IsValid())
+		{
+			return;
+		}
+
+		const UMovieScene3DTransformSection* TransformSection = FBakingHelper::GetTransformSection(InSequencer.Get(), Guid, LocalTransform);
+		if (!TransformSection)
+		{
+			return;
+		}
+
+		FBakingHelper::AddTransformKeys(TransformSection, {InTime}, {LocalTransform}, EMovieSceneTransformChannel::AllTransform);
+		return; 
+	}
+	
+	const UTransformableControlHandle* ControlHandle = Cast<UTransformableControlHandle>(InHandle);
+	if (ControlHandle && ControlHandle->IsValid())
+	{
+		const FTransform LocalTransform = ControlHandle->GetLocalTransform();
+		
+		static constexpr bool bNotify = true;
+		static constexpr bool bUndo = false;
+		static constexpr bool bFixEuler = true;
+	
+		FRigControlModifiedContext KeyframeContext;
+		KeyframeContext.SetKey = EControlRigSetKey::Always;
+		KeyframeContext.KeyMask = static_cast<uint32>(EControlRigContextChannelToKey::AllTransform);
+
+		const FFrameRate TickResolution = InSequencer->GetFocusedTickResolution();
+		KeyframeContext.LocalTime = TickResolution.AsSeconds(FFrameTime(InTime));
+	
+		ControlHandle->ControlRig->SetControlLocalTransform(ControlHandle->ControlName, LocalTransform, bNotify, KeyframeContext, bUndo, bFixEuler);
+	}
+}
+
+void FConstraintChannelHelper::AddConstraintKey(UTickableTransformConstraint* InConstraint)
+{
+	const TWeakPtr<ISequencer> WeakSequencer = FBakingHelper::GetSequencer();
+	if (!WeakSequencer.IsValid() || !WeakSequencer.Pin()->GetFocusedMovieSceneSequence())
+	{
+		return;
+	}
+	const TSharedPtr<ISequencer> Sequencer = WeakSequencer.Pin();
+	const UMovieScene* MovieScene = Sequencer->GetFocusedMovieSceneSequence()->GetMovieScene();
+	if (!MovieScene)
+	{
+		return;
+	}
+
+	const UTransformableControlHandle* ControlHandle = Cast<UTransformableControlHandle>(InConstraint->ChildTRSHandle);
+	if (!ControlHandle)
+	{
+		return;
+	}
+
+	if (UMovieSceneControlRigParameterSection* ParamSection = GetControlSection(ControlHandle, Sequencer))
+	{
+		ParamSection->AddConstraintChannel(InConstraint, true);
+
+		if (FConstraintAndActiveChannel* Channel = ParamSection->GetConstraintChannel(InConstraint->GetFName()))
+		{
+			ParamSection->Modify();
+
+			const FFrameRate TickResolution = Sequencer->GetFocusedTickResolution();
+			const FFrameTime FrameTime = Sequencer->GetLocalTime().ConvertTo(TickResolution);
+			const FFrameNumber Time = FrameTime.GetFrame();
+
+			bool Value = false;
+			CanAddKey(Channel, Time, Value);
+			
+			TMovieSceneChannelData<bool> ChannelData = Channel->ActiveChannel.GetData();
+			ChannelData.AddKey(Time, InConstraint->Active);
+
+			// bool PrevActiveValue = false; Channel->ActiveChannel.Evaluate(Time, PrevActiveValue);
+			// bool NextActiveValue = false; Channel->ActiveChannel.Evaluate(Time, PrevActiveValue);
+
+					
+			// int32 ExistingIndex = ChannelData.FindKey(Time);
+			// if (ExistingIndex != INDEX_NONE)
+			// {
+			// 	Handle = ChannelInterface.GetHandle(ExistingIndex);
+			// 	using namespace UE::MovieScene;
+			// 	AssignValue(Channel, Handle, Forward<FMovieSceneControlRigSpaceBaseKey>(Value));
+			// }
+			// else
+			// {
+			// 	ExistingIndex = ChannelInterface.AddKey(Time, Forward<FMovieSceneControlRigSpaceBaseKey>(Value));
+			// 	Handle = ChannelInterface.GetHandle(ExistingIndex);
+			// }
+		}
+	}
+}
+
+UMovieSceneControlRigParameterSection* FConstraintChannelHelper::GetControlSection(
+	const UTransformableControlHandle* InHandle,
+	const TSharedPtr<ISequencer>& InSequencer)
+{
+	const UMovieScene* MovieScene = InSequencer->GetFocusedMovieSceneSequence()->GetMovieScene();
+	if (!MovieScene)
+	{
+		return nullptr;
+	}
+
+	const TWeakObjectPtr<UControlRig> ControlRig = InHandle->ControlRig;
+	
+	const TArray<FMovieSceneBinding>& Bindings = MovieScene->GetBindings();
+	for (const FMovieSceneBinding& Binding : Bindings)
+	{
+		UMovieSceneTrack* Track = MovieScene->FindTrack(UMovieSceneControlRigParameterTrack::StaticClass(), Binding.GetObjectGuid());
+		UMovieSceneControlRigParameterTrack* ControlRigTrack = Cast<UMovieSceneControlRigParameterTrack>(Track);
+		if (ControlRigTrack && ControlRigTrack->GetControlRig() == ControlRig)
+		{
+			return Cast<UMovieSceneControlRigParameterSection>(ControlRigTrack->FindSection(0));
+		}
+	}
+
+	return nullptr;
+}
+
+UMovieScene3DTransformSection* FConstraintChannelHelper::GetTransformSection(
+	const UTransformableComponentHandle* InHandle,
+	const TSharedPtr<ISequencer>& InSequencer)
+{
+	AActor* Actor = InHandle->Component->GetOwner();
+	if (!Actor)
+	{
+		return nullptr;
+	}
+	
+	const FGuid Guid = InSequencer->GetHandleToObject(Actor,true);
+	if (!Guid.IsValid())
+	{
+		return nullptr;
+	}
+	
+	return FBakingHelper::GetTransformSection(InSequencer.Get(), Guid);
+}
+
+TArrayView<FMovieSceneFloatChannel*> FConstraintChannelHelper::GetTransformFloatChannels(
+	const UTransformableHandle* InHandle,
+	const TSharedPtr<ISequencer>& InSequencer)
+{
+	// InHandle and InSequencer are assumed to be valid at this stage so no need to check
+	
+	static const TArrayView<FMovieSceneFloatChannel*> EmptyChannelsView;
+	
+	if (const UTransformableControlHandle* ControlHandle = Cast<UTransformableControlHandle>(InHandle))
+	{
+		UMovieSceneControlRigParameterSection* Section = GetControlSection(ControlHandle, InSequencer);
+		FChannelMapInfo* pChannelIndex = nullptr;
+		FRigControlElement* ControlElement = nullptr;
+		Tie(ControlElement, pChannelIndex) = FControlRigSpaceChannelHelpers::GetControlAndChannelInfo(
+			ControlHandle->ControlRig.Get(),
+			Section,
+			ControlHandle->ControlName);
+    
+		if (!pChannelIndex || !ControlElement)
+		{
+			return EmptyChannelsView;
+		}
+
+		// get the number of float channels to treat
+		const int32 NumChannels = FControlRigSpaceChannelHelpers::GetNumFloatChannels(ControlElement->Settings.ControlType);
+		if (!NumChannels)
+		{
+			return EmptyChannelsView;
+		}
+
+		// return a sub view that just represents the control's channels
+		const TArrayView<FMovieSceneFloatChannel*> FloatChannels = Section->GetChannelProxy().GetChannels<FMovieSceneFloatChannel>();
+		const int32 ChannelStartIndex = pChannelIndex->ChannelIndex;
+		return FloatChannels.Slice(ChannelStartIndex, NumChannels);
+	}
+	
+	if (const UTransformableComponentHandle* ComponentHandle = Cast<UTransformableComponentHandle>(InHandle))
+	{
+		const UMovieScene3DTransformSection* Section = GetTransformSection(ComponentHandle, InSequencer);
+		return Section ? Section->GetChannelProxy().GetChannels<FMovieSceneFloatChannel>() : EmptyChannelsView;
+	}
+	
+	return EmptyChannelsView;
+}
+
+namespace
+{
+	
+void GetComponentGlobalTransforms(
+	const UTransformableControlHandle* InControlHandle,
+	const TSharedPtr<ISequencer>& InSequencer,
+	TArray<FFrameNumber> InFrames,
+	TArray<FTransform>& OutTransforms)
+{
+	const USkeletalMeshComponent* SkeletalMeshComponent = InControlHandle->GetSkeletalMesh();
+	AActor* Actor = SkeletalMeshComponent ? SkeletalMeshComponent->GetTypedOuter< AActor >() : nullptr;
+	if (Actor)
+	{
+		FActorForWorldTransforms ControlRigActorSelection;
+		ControlRigActorSelection.Actor = Actor;
+		MovieSceneToolHelpers::GetActorWorldTransforms(InSequencer.Get(), ControlRigActorSelection, InFrames, OutTransforms);		
+	}
+	else
+	{
+		OutTransforms.SetNum(InFrames.Num());
+		for (FTransform& Transform : OutTransforms)
+		{
+			Transform = FTransform::Identity;
+		}
+	}
+}
+
+struct FTransformEvaluator
+{
+public:
+	TArray<FTransform> ChildTransforms;
+	TArray<FTransform> ParentGlobalTransforms;
+	TArray<FTransform> ComponentGlobalTransforms;
+
+	FTransformEvaluator(UTickableTransformConstraint* InConstraint)
+		: Constraint(InConstraint)
+	{}
+	
+	void ComputeFrames(const TSharedPtr<ISequencer>& InSequencer, const TArray<FFrameNumber>& InFrames)
+	{
+		if (!Constraint)
+		{
+			return;
+		}
+		
+		const UTransformableControlHandle* ControlHandle = Cast<UTransformableControlHandle>(Constraint->ChildTRSHandle);
+		if (!ControlHandle)
+		{
+			return;
+		}
+
+		// todo: compute child, parent and component in one single pass for faster computation
+		FConstraintBaker::GetHandleTransforms(
+		   ControlHandle->ControlRig->GetWorld(),
+		   InSequencer,
+		   ControlHandle,
+		   InFrames,
+		   bLocal,
+		   ChildTransforms);
+	
+		FConstraintBaker::GetHandleTransforms(
+			ControlHandle->ControlRig->GetWorld(),
+			InSequencer,
+			Constraint->ParentTRSHandle,
+			InFrames,
+			bLocal,
+			ParentGlobalTransforms);
+
+		GetComponentGlobalTransforms(
+			ControlHandle,
+			InSequencer,
+			InFrames,
+			ComponentGlobalTransforms);
+	}
+	
+private:
+	static constexpr bool bLocal = false; // store global transforms
+	UTickableTransformConstraint* Constraint = nullptr;
+};
+
+}
+
+void FConstraintChannelHelper::SmartControlConstraintKey(
+	UTickableTransformConstraint* InConstraint,
+	const TSharedPtr<ISequencer>& InSequencer)
+{
+	const UTransformableControlHandle* ControlHandle = Cast<UTransformableControlHandle>(InConstraint->ChildTRSHandle);
+	if (!ControlHandle)
+	{
+		return;
+	}
+	
+	if (UMovieSceneControlRigParameterSection* Section = GetControlSection(ControlHandle, InSequencer))
+	{
+		// add the channel
+		Section->AddConstraintChannel(InConstraint, true);
+
+		// add key if needed
+		if (FConstraintAndActiveChannel* Channel = Section->GetConstraintChannel(InConstraint->GetFName()))
+		{
+			const FFrameRate TickResolution = InSequencer->GetFocusedTickResolution();
+			const FFrameTime FrameTime = InSequencer->GetLocalTime().ConvertTo(TickResolution);
+			const FFrameNumber Time = FrameTime.GetFrame();
+
+			bool ActiveValueToBeSet = false;
+			if (CanAddKey(Channel, Time, ActiveValueToBeSet))
+			{
+				// set constraint as dynamic
+				InConstraint->bDynamicOffset = true;
+				
+				UControlRig* ControlRig = ControlHandle->ControlRig.Get();
+				const FName& ControlName = ControlHandle->ControlName;
+				
+				// store the frames to compensate
+				const TArrayView<FMovieSceneFloatChannel*> Channels = GetTransformFloatChannels(ControlHandle, InSequencer);
+				TArray<FFrameNumber> FramesToCompensate;
+				GetFramesToCompensate(Channel->ActiveChannel, ActiveValueToBeSet, Time, Channels, FramesToCompensate);
+
+				// store child, parent and component global transforms for these frames
+				FTransformEvaluator Evaluator(InConstraint);
+				Evaluator.ComputeFrames(InSequencer, FramesToCompensate);
+				TArray<FTransform>& ChildTransforms = Evaluator.ChildTransforms;
+				const TArray<FTransform>& ParentTransforms = Evaluator.ParentGlobalTransforms;
+				const TArray<FTransform>& ComponentTransforms = Evaluator.ComponentGlobalTransforms;
+
+				// store tangents at this time
+				TArray<FMovieSceneTangentData> Tangents;
+				FControlRigSpaceChannelHelpers::EvaluateTangentAtThisTime(ControlRig, Section, ControlName, Time, Tangents);
+
+				const EMovieSceneTransformChannel ChannelsToKey = FConstraintBaker::GetChannelsToKey(InConstraint);
+				
+				// add child's transform key at Time-1 to keep animation
+				{
+					const FFrameNumber TimeMinusOne(Time -1);
+
+					const bool bPrevAsLocal = !ActiveValueToBeSet;
+					const FTransform& SpaceTransformTime = bPrevAsLocal ? ParentTransforms[0] : // constraint space  
+																		  ComponentTransforms[0]; // component space
+					const FTransform ChildTransformAtMinusOne = ChildTransforms[0].GetRelativeTransform(SpaceTransformTime);
+					FBakingHelper::AddTransformKeys(ControlRig, ControlName, {TimeMinusOne},
+						{ChildTransformAtMinusOne}, ChannelsToKey, TickResolution, bPrevAsLocal);
+				
+					// set tangents at Time-1
+					FControlRigSpaceChannelHelpers::SetTangentsAtThisTime(ControlRig, Section, ControlName, TimeMinusOne, Tangents);
+				}
+
+				// add active key
+				{
+					Section->Modify();
+					TMovieSceneChannelData<bool> ChannelData = Channel->ActiveChannel.GetData();
+					ChannelData.AddKey(Time, ActiveValueToBeSet);
+				}
+
+				// compensate
+				{
+					const bool bNextAsLocal = ActiveValueToBeSet;
+					for (int32 Index = 0; Index < ChildTransforms.Num(); ++Index)
+					{
+						const FTransform& SpaceTransform = bNextAsLocal ? ParentTransforms[Index] : // constraint space 
+																		  ComponentTransforms[Index]; // component space
+						FTransform& ChildTransform = ChildTransforms[Index];
+						ChildTransform = ChildTransform.GetRelativeTransform(SpaceTransform);
+					}
+
+					// add keys
+					FBakingHelper::AddTransformKeys( ControlRig, ControlHandle->ControlName, FramesToCompensate,
+						ChildTransforms, ChannelsToKey, TickResolution, bNextAsLocal);
+				
+					// set tangents at Time
+					FControlRigSpaceChannelHelpers::SetTangentsAtThisTime(ControlRig, Section, ControlName, Time, Tangents);
+				}
+			}
+		}
+	}
+}
+
+void FConstraintChannelHelper::SmartComponentConstraintKey(
+	UTickableTransformConstraint* InConstraint,
+	const TSharedPtr<ISequencer>& InSequencer)
+{
+	const UTransformableComponentHandle* ComponentHandle = Cast<UTransformableComponentHandle>(InConstraint->ChildTRSHandle);
+	if (!ComponentHandle)
+	{
+		return;
+	}
+}
+
+void FConstraintChannelHelper::GetFramesToCompensate(
+	const FMovieSceneConstraintChannel& InActiveChannel,
+	const bool InActiveValueToBeSet,
+	const FFrameNumber& InTime,
+	const TArrayView<FMovieSceneFloatChannel*>& InChannels,
+	TArray<FFrameNumber>& OutFramesAfter)
+{
+	const bool bHasKeys = (InActiveChannel.GetNumKeys() > 0);
+	
+	OutFramesAfter.Reset();
+
+	// add the current frame
+	OutFramesAfter.Add(InTime);
+
+	// add the next frames that need transform compensation 
+	for (const FMovieSceneFloatChannel* InChannel: InChannels)
+	{
+		const TMovieSceneChannelData<const FMovieSceneFloatValue> ChannelData = InChannel->GetData();
+		const TArrayView<const FFrameNumber> Times = ChannelData.GetTimes();
+		if (!Times.IsEmpty())
+		{
+			// look for the first next key frame for this channel 
+			const int32 NextTimeIndex = Algo::UpperBound(Times, InTime);
+			if (Times.IsValidIndex(NextTimeIndex))
+			{
+				// store the time while the state is different
+				for (int32 Index = NextTimeIndex; Index < Times.Num(); ++Index)
+				{
+					if (!bHasKeys)
+					{
+						OutFramesAfter.Add(Times[Index]);
+					}
+					else
+					{
+						bool NextValue = false; InActiveChannel.Evaluate(Times[Index], NextValue);
+						if (NextValue == InActiveValueToBeSet)
+						{
+							break;
+						}
+						OutFramesAfter.Add(Times[Index]);
+					}
+				}
+			}
+		}
+	}
+
+	// uniqueness
+	OutFramesAfter.Sort();
+	OutFramesAfter.SetNum(Algo::Unique(OutFramesAfter));
+}
+
+// namespace
+// {
+// 	using IndexRange = ::TRange<int32>;
+// 	IndexRange GetRange(const TArrayView<const bool>& Values, const int32 Offset) 
+// 	{
+// 		IndexRange Range;
+// 			
+// 		int32 FirstActive = INDEX_NONE;
+// 		for (int32 Index = Offset; Index < Values.Num(); ++Index)
+// 		{
+// 			if (Values[Index])
+// 			{
+// 				FirstActive = Index;
+// 				break;
+// 			}
+// 		}
+//
+// 		if (FirstActive == INDEX_NONE)
+// 		{
+// 			return Range;
+// 		}
+//
+// 		Range.SetLowerBound(TRangeBound<int32>(FirstActive));
+// 		Range.SetUpperBound(TRangeBound<int32>(FirstActive));
+//
+// 		for (int32 NextInactive = FirstActive+1; NextInactive < Values.Num(); ++NextInactive)
+// 		{
+// 			if (Values[NextInactive] == false)
+// 			{
+// 				Range.SetUpperBound(TRangeBound<int32>(NextInactive));
+// 				return Range;
+// 			}
+// 		}
+// 			
+// 		return Range;
+// 	};
+//
+// 	TArray<IndexRange> GetIndexRanges(const FMovieSceneConstraintChannel* Channel)
+// 	{
+// 		TArray<IndexRange> Ranges;
+//
+// 		const TMovieSceneChannelData<const bool> ChannelData = Channel->GetData();
+// 		const TArrayView<const bool> Values = ChannelData.GetValues();
+// 		if (Values.Num() == 1)
+// 		{
+// 			if (Values[0] == true)
+// 			{
+// 				Ranges.Emplace(0);
+// 			}
+// 			return Ranges;
+// 		}
+// 		
+// 		int32 Offset = 0;
+// 		while (Values.IsValidIndex(Offset))
+// 		{
+// 			TRange<int32> Range = GetRange(Values, Offset);
+// 			if (!Range.IsEmpty())
+// 			{
+// 				Ranges.Emplace(Range);
+// 				Offset = Range.GetUpperBoundValue()+1;
+// 			}
+// 			else
+// 			{
+// 				Offset = INDEX_NONE;
+// 			}
+// 		}
+// 			
+// 		return Ranges;
+// 	}
+// }
+//
+// void DrawKeys(
+// 	FMovieSceneConstraintChannel* Channel,
+// 	TArrayView<const FKeyHandle> InKeyHandles,
+// 	const UMovieSceneSection* InOwner,
+// 	TArrayView<FKeyDrawParams> OutKeyDrawParams)
+// {
+// 	const UMovieSceneControlRigParameterSection* Section = Cast<UMovieSceneControlRigParameterSection>(InOwner);
+// 	if (!Section)
+// 	{
+// 		return;
+// 	}
+// 	
+// 	static const FName SquareKeyBrushName("Sequencer.KeySquare");
+// 	static const FName FilledBorderBrushName("FilledBorder");
+//
+// 	static FKeyDrawParams Params;
+// 	Params.FillBrush = FAppStyle::GetBrush(FilledBorderBrushName);
+// 	Params.BorderBrush = FAppStyle::GetBrush(SquareKeyBrushName);
+// 	
+// 	for (FKeyDrawParams& Param : OutKeyDrawParams)
+// 	{
+// 		Param = Params;
+// 	}
+// }
+//
+// void DrawExtra(FMovieSceneConstraintChannel* Channel, const UMovieSceneSection* Owner,const FGeometry& AllottedGeometry, FSequencerSectionPainter& Painter)
+// {
+// 	const UMovieSceneControlRigParameterSection* Section = Cast<UMovieSceneControlRigParameterSection>(Owner);
+// 	if (!Section || !Channel)
+// 	{
+// 		return;
+// 	}
+// 	
+// 	// using namespace UE::Sequencer;
+//
+// 	// get index range
+// 	TArray<IndexRange> IndexRanges = GetIndexRanges(Channel);
+// 	if (IndexRanges.IsEmpty())
+// 	{
+// 		return;
+// 	}
+// 	
+// 	// convert to bar range
+// 	TArray<FKeyBarCurveModel::FBarRange> Ranges;
+// 	const TMovieSceneChannelData<const bool> ChannelData = Channel->GetData();
+// 	const TArrayView<const FFrameNumber> Times = ChannelData.GetTimes();
+//
+// 	// const bool IsFirstValueTrue = IndexRanges[0].GetLowerBoundValue() == 0;
+// 	
+// 	for (const IndexRange& ActiveRange: IndexRanges)
+// 	{
+// 		FKeyBarCurveModel::FBarRange BarRange;
+// 		FFrameRate TickResolution = Section->GetTypedOuter<UMovieScene>()->GetTickResolution();
+// 		double LowerValue = Times[ActiveRange.GetLowerBoundValue()] / TickResolution;
+// 		double UpperValue = Times[ActiveRange.GetUpperBoundValue()] / TickResolution;
+// 	
+// 		BarRange.Range.SetLowerBound(TRangeBound<double>(LowerValue));
+// 		BarRange.Range.SetUpperBound(TRangeBound<double>(UpperValue));
+// 	
+// 		BarRange.Name = "Constraint";
+// 		BarRange.Color = FLinearColor(.2, .5, .1);
+// 		static FLinearColor ZebraTint = FLinearColor::White.CopyWithNewOpacity(0.01f);
+// 		BarRange.Color = BarRange.Color * (1.f - ZebraTint.A) + ZebraTint * ZebraTint.A;
+// 		
+// 		Ranges.Add(BarRange);
+// 	}
+//
+// 	// draw bars
+// 	const FSlateBrush* WhiteBrush = FAppStyle::GetBrush("WhiteBrush");
+// 	static constexpr ESlateDrawEffect DrawEffects = ESlateDrawEffect::None;
+//
+// 	const FSlateFontInfo FontInfo = FCoreStyle::Get().GetFontStyle("ToolTip.LargerFont");
+// 	const TSharedRef<FSlateFontMeasure> FontMeasure = FSlateApplication::Get().GetRenderer()->GetFontMeasureService();
+//
+// 	const FVector2D& LocalSize = AllottedGeometry.GetLocalSize();
+// 	static constexpr float LaneTop = 0;
+//
+// 	const FTimeToPixel& TimeToPixel = Painter.GetTimeConverter();
+// 	const double InputMin = TimeToPixel.PixelToSeconds(0.f);
+// 	const double InputMax = TimeToPixel.PixelToSeconds(Painter.SectionGeometry.GetLocalSize().X);
+//
+// 	for (int32 Index = 0; Index < Ranges.Num(); ++Index)
+// 	{
+// 		const FKeyBarCurveModel::FBarRange& Range = Ranges[Index];
+//
+// 		double LowerSeconds = /*(Index == 0 && IsFirstValueTrue) ? InputMin : */Range.Range.GetLowerBoundValue();
+// 		double UpperSeconds = Range.Range.GetUpperBoundValue();
+// 		if (UpperSeconds == Range.Range.GetLowerBoundValue())
+// 		{
+// 			UpperSeconds = InputMax;
+// 		}
+//
+// 		const double BoxStart = TimeToPixel.SecondsToPixel(LowerSeconds);
+// 		const double BoxEnd = TimeToPixel.SecondsToPixel(UpperSeconds);
+// 		const double BoxSize = BoxEnd - BoxStart;
+//
+// 		const FVector2D Size = FVector2D(BoxSize, LocalSize.Y);
+// 		const FVector2D Translation = FVector2D(BoxStart, LaneTop);
+// 		const FPaintGeometry BoxGeometry = AllottedGeometry.ToPaintGeometry(Size, FSlateLayoutTransform(Translation));
+// 		
+// 		FSlateDrawElement::MakeBox(Painter.DrawElements, Painter.LayerId, BoxGeometry, WhiteBrush, DrawEffects, Range.Color);
+// 	}
+// }
diff --git a/Engine/Plugins/Animation/ControlRig/Source/ControlRigEditor/Private/Sequencer/ControlRigParameterTrackEditor.cpp b/Engine/Plugins/Animation/ControlRig/Source/ControlRigEditor/Private/Sequencer/ControlRigParameterTrackEditor.cpp
--- a/Engine/Plugins/Animation/ControlRig/Source/ControlRigEditor/Private/Sequencer/ControlRigParameterTrackEditor.cpp
+++ b/Engine/Plugins/Animation/ControlRig/Source/ControlRigEditor/Private/Sequencer/ControlRigParameterTrackEditor.cpp
@@ -419,7 +419,7 @@ void FControlRigParameterTrackEditor::BindControlRig(UControlRig* ControlRig)
 					}
 				}
 			}
-			Track->SpaceChannelAdded().AddRaw(this, &FControlRigParameterTrackEditor::HandleOnSpaceAdded);
+			// Track->SpaceChannelAdded().AddRaw(this, &FControlRigParameterTrackEditor::HandleOnSpaceAdded);
 		}
 	}
 }
@@ -2492,7 +2492,13 @@ void FControlRigParameterTrackEditor::HandleControlUndoBracket(UControlRig* Subj
 	}
 }
 
-void FControlRigParameterTrackEditor::GetControlRigKeys(UControlRig* InControlRig, FName ParameterName, EControlRigContextChannelToKey ChannelsToKey, ESequencerKeyMode KeyMode, UMovieSceneControlRigParameterSection* SectionToKey, FGeneratedTrackKeys& OutGeneratedKeys)
+void FControlRigParameterTrackEditor::GetControlRigKeys(
+	UControlRig* InControlRig,
+	FName ParameterName,
+	EControlRigContextChannelToKey ChannelsToKey,
+	ESequencerKeyMode KeyMode,
+	UMovieSceneControlRigParameterSection* SectionToKey,
+	FGeneratedTrackKeys& OutGeneratedKeys)
 {
 	const TArray<bool>& ControlsMask = SectionToKey->GetControlsMask();
 	EMovieSceneTransformChannel TransformMask = SectionToKey->GetTransformMask().GetChannels();
@@ -2628,6 +2634,10 @@ void FControlRigParameterTrackEditor::GetControlRigKeys(UControlRig* InControlRi
 					Translation = Val.GetTranslation();
 					Scale = Val.GetScale3D();
 				}
+
+				// transform these values in the constraint space if needed
+				// NOTE TRS are local to their space 
+					
 				FVector3f CurrentVector = (FVector3f)Translation;
 				bool bKeyX = bSetKey && EnumHasAnyFlags(ChannelsToKey, EControlRigContextChannelToKey::TranslationX);
 				bool bKeyY = bSetKey && EnumHasAnyFlags(ChannelsToKey, EControlRigContextChannelToKey::TranslationY);
diff --git a/Engine/Plugins/Animation/ControlRig/Source/ControlRigEditor/Private/Sequencer/ControlRigSpaceChannelEditors.cpp b/Engine/Plugins/Animation/ControlRig/Source/ControlRigEditor/Private/Sequencer/ControlRigSpaceChannelEditors.cpp
--- a/Engine/Plugins/Animation/ControlRig/Source/ControlRigEditor/Private/Sequencer/ControlRigSpaceChannelEditors.cpp
+++ b/Engine/Plugins/Animation/ControlRig/Source/ControlRigEditor/Private/Sequencer/ControlRigSpaceChannelEditors.cpp
@@ -756,9 +756,14 @@ void FControlRigSpaceChannelHelpers::DeleteTransformKeysAtThisTime(UControlRig*
 	}
 }
 
-void FControlRigSpaceChannelHelpers::GetFramesInThisSpaceAfterThisTime(UControlRig* ControlRig, FName ControlName, FMovieSceneControlRigSpaceBaseKey CurrentValue,
-	FMovieSceneControlRigSpaceChannel* Channel, UMovieSceneSection* SectionToKey,
-	FFrameNumber Time, TSortedMap<FFrameNumber,FFrameNumber>& OutMoreFrames)
+void FControlRigSpaceChannelHelpers::GetFramesInThisSpaceAfterThisTime(
+	UControlRig* ControlRig,
+	FName ControlName,
+	FMovieSceneControlRigSpaceBaseKey CurrentValue,
+	FMovieSceneControlRigSpaceChannel* Channel,
+	UMovieSceneSection* SectionToKey,
+	FFrameNumber Time,
+	TSortedMap<FFrameNumber,FFrameNumber>& OutMoreFrames)
 {
 	OutMoreFrames.Reset();
 	if (ControlRig && Channel && SectionToKey)
@@ -1124,80 +1129,93 @@ void FControlRigSpaceChannelHelpers::HandleSpaceKeyTimeChanged(UControlRig* Cont
 }
 
 
-void FControlRigSpaceChannelHelpers::CompensateIfNeeded(UControlRig* ControlRig, ISequencer* Sequencer, UMovieSceneControlRigParameterSection* Section, FName ControlName, TOptional<FFrameNumber>& OptionalTime)
+void FControlRigSpaceChannelHelpers::CompensateIfNeeded(
+	UControlRig* ControlRig,
+	ISequencer* Sequencer,
+	UMovieSceneControlRigParameterSection* Section,
+	FName ControlName,
+	TOptional<FFrameNumber>& OptionalTime)
 {
 	if (bDoNotCompensate == true)
 	{
 		return;
 	}
+
+	// Frames to compensate
+	TArray<FFrameNumber> OptionalTimeArray;
+	if(OptionalTime.IsSet())
+	{
+		OptionalTimeArray.Add(OptionalTime.GetValue());
+	}
+	
+	auto GetSpaceTimesToCompensate = [&OptionalTimeArray](const FSpaceControlNameAndChannel* Channel)->TArrayView<const FFrameNumber>
+	{
+		if (OptionalTimeArray.IsEmpty())
+		{
+			return Channel->SpaceCurve.GetData().GetTimes();
+		}
+		return OptionalTimeArray;
+	};
+
+	// keyframe context
+	FRigControlModifiedContext KeyframeContext;
+	KeyframeContext.SetKey = EControlRigSetKey::Always;
+	const FFrameRate TickResolution = Sequencer->GetFocusedTickResolution();
 	
 	//we need to check all controls for 1) space and 2) previous frame and if so we automatically compensate.
-	TArray<FRigControlElement*> Controls = ControlRig->GetHierarchy()->GetControls();
 	bool bDidIt = false;
-	for(FRigControlElement* Control: Controls)
+
+	// compensate spaces
+	URigHierarchy* RigHierarchy = ControlRig->GetHierarchy();
+	const TArray<FRigControlElement*> Controls = RigHierarchy->GetControls();
+	for (const FRigControlElement* Control: Controls)
 	{ 
 		if(Control)// ac && Control->GetName() != ControlName)
 		{ 
 			//only if we have a channel
 			if (FSpaceControlNameAndChannel* Channel = Section->GetSpaceChannel(Control->GetName()))
 			{
-				TArray<FFrameNumber> AllFrames; 
-				if (OptionalTime.IsSet())
+				const TArrayView<const FFrameNumber> FramesToCompensate = GetSpaceTimesToCompensate(Channel);
+				if (FramesToCompensate.Num() > 0)
 				{
-					FFrameNumber Time = OptionalTime.GetValue();
-					AllFrames.Add(Time);
-				}
-				else
-				{
-					AllFrames = Channel->SpaceCurve.GetData().GetTimes();
-				}
-				if (AllFrames.Num() > 0)
-				{
-					for (FFrameNumber& Time : AllFrames)
+					for (const FFrameNumber& Time : FramesToCompensate)
 					{
 						FMovieSceneControlRigSpaceBaseKey ExistingValue, PreviousValue;
 						using namespace UE::MovieScene;
 						EvaluateChannel(&(Channel->SpaceCurve), Time - 1, PreviousValue);
 						EvaluateChannel(&(Channel->SpaceCurve), Time, ExistingValue);
+
 						if (ExistingValue != PreviousValue) //if they are the same no need to do anything
 						{
-							//find global value at curren time
-							TArray<FFrameNumber> Frames;
-							Frames.Add(Time);
-							TArray<FTransform> ControlRigParentWorldTransforms;
-							ControlRigParentWorldTransforms.SetNum(Frames.Num());
-							for (FTransform& WorldTransform : ControlRigParentWorldTransforms)
-							{
-								WorldTransform = FTransform::Identity;
-							}
+							//find global value at current time
+							TArray<FTransform> ControlRigParentWorldTransforms({FTransform::Identity});
 							TArray<FTransform> ControlWorldTransforms;
 							FControlRigSnapper Snapper;
-							Snapper.GetControlRigControlTransforms(Sequencer, ControlRig, Control->GetName(), Frames, ControlRigParentWorldTransforms, ControlWorldTransforms);
-							FRigElementKey ControlKey;
-							ControlKey.Name = Control->GetName();
-							ControlKey.Type = ERigElementType::Control;
+							Snapper.GetControlRigControlTransforms(
+								Sequencer, ControlRig, Control->GetName(),
+								{Time},
+								ControlRigParentWorldTransforms, ControlWorldTransforms);
+
 							//set space to previous space value that's different.
-							URigHierarchy* RigHierarchy = ControlRig->GetHierarchy();
 							switch (PreviousValue.SpaceType)
 							{
-							case EMovieSceneControlRigSpaceType::Parent:
-								RigHierarchy->SwitchToDefaultParent(ControlKey);
-								break;
-							case EMovieSceneControlRigSpaceType::World:
-								RigHierarchy->SwitchToWorldSpace(ControlKey);
-								break;
-							case EMovieSceneControlRigSpaceType::ControlRig:
-								URigHierarchy::TElementDependencyMap Dependencies = RigHierarchy->GetDependenciesForVM(ControlRig->GetVM());
-								RigHierarchy->SwitchToParent(ControlKey, PreviousValue.ControlRigElement, false, true, Dependencies, nullptr);
-								break;
+								case EMovieSceneControlRigSpaceType::Parent:
+									RigHierarchy->SwitchToDefaultParent(Control->GetKey());
+									break;
+								case EMovieSceneControlRigSpaceType::World:
+									RigHierarchy->SwitchToWorldSpace(Control->GetKey());
+									break;
+								case EMovieSceneControlRigSpaceType::ControlRig:
+									URigHierarchy::TElementDependencyMap Dependencies = RigHierarchy->GetDependenciesForVM(ControlRig->GetVM());
+									RigHierarchy->SwitchToParent(Control->GetKey(), PreviousValue.ControlRigElement, false, true, Dependencies, nullptr);
+									break;
 							}
+							
 							//now set time -1 frame value
-							FRigControlModifiedContext Context;
-							Context.SetKey = EControlRigSetKey::Always;
-							FFrameRate TickResolution = Sequencer->GetFocusedTickResolution();
 							ControlRig->Evaluate_AnyThread();
-							Context.LocalTime = TickResolution.AsSeconds(FFrameTime(Time - 1));
-							ControlRig->SetControlGlobalTransform(ControlKey.Name, ControlWorldTransforms[0], true, Context);
+							KeyframeContext.LocalTime = TickResolution.AsSeconds(FFrameTime(Time - 1));
+							ControlRig->SetControlGlobalTransform(Control->GetName(), ControlWorldTransforms[0], true, KeyframeContext);
+							
 							bDidIt = true;
 						}
 					}
@@ -1205,7 +1223,8 @@ void FControlRigSpaceChannelHelpers::CompensateIfNeeded(UControlRig* ControlRig,
 			}
 		}
 	}
-	if (bDidIt == true)
+	
+	if (bDidIt)
 	{
 		Sequencer->ForceEvaluate();
 	}
diff --git a/Engine/Plugins/Animation/ControlRig/Source/ControlRigEditor/Private/Tools/BakingHelper.cpp b/Engine/Plugins/Animation/ControlRig/Source/ControlRigEditor/Private/Tools/BakingHelper.cpp
--- a/Engine/Plugins/Animation/ControlRig/Source/ControlRigEditor/Private/Tools/BakingHelper.cpp
+++ b/Engine/Plugins/Animation/ControlRig/Source/ControlRigEditor/Private/Tools/BakingHelper.cpp
@@ -225,31 +225,39 @@ bool FBakingHelper::AddTransformKeys(
 bool FBakingHelper::AddTransformKeys(
 	UControlRig* InControlRig, const FName& InControlName,
 	const TArray<FFrameNumber>& InFrames,
-	const TArray<FTransform>& InLocalTransforms,
+	const TArray<FTransform>& InTransforms,
 	const EMovieSceneTransformChannel& InChannels,
-	const FFrameRate& InTickResolution)
+	const FFrameRate& InTickResolution,
+	const bool bLocal)
 {
-	if (InFrames.IsEmpty() || InFrames.Num() != InLocalTransforms.Num())
+	if (InFrames.IsEmpty() || InFrames.Num() != InTransforms.Num())
 	{
 		return false;
 	}
 
-	static constexpr bool bNotify = true;
-	static constexpr bool bUndo = false;
-	static constexpr bool bFixEuler = true;
-	
+	auto KeyframeFunc = [InControlRig, InControlName, bLocal](const FTransform& InTransform, const FRigControlModifiedContext& InKeyframeContext)
+	{
+		static constexpr bool bNotify = true;
+		static constexpr bool bUndo = false;
+		static constexpr bool bFixEuler = true;
+
+		if (bLocal)
+		{
+			return InControlRig->SetControlLocalTransform(InControlName, InTransform, bNotify, InKeyframeContext, bUndo, bFixEuler);
+		}
+		InControlRig->SetControlGlobalTransform(InControlName, InTransform, bNotify, InKeyframeContext, bUndo, bFixEuler);
+	};
+
 	FRigControlModifiedContext KeyframeContext;
 	KeyframeContext.SetKey = EControlRigSetKey::Always;
 	KeyframeContext.KeyMask = static_cast<uint32>(InChannels);
 	
 	for (int32 Index = 0; Index < InFrames.Num(); ++Index)
 	{
-		const FTransform& LocalTransform = InLocalTransforms[Index];
-
 		const FFrameNumber& Frame = InFrames[Index];
 		KeyframeContext.LocalTime = InTickResolution.AsSeconds(FFrameTime(Frame));
 		
-		InControlRig->SetControlLocalTransform(InControlName, LocalTransform, bNotify, KeyframeContext, bUndo, bFixEuler);
+		KeyframeFunc(InTransforms[Index], KeyframeContext);
 	}
 
 	return true;
diff --git a/Engine/Plugins/Animation/ControlRig/Source/ControlRigEditor/Private/Tools/ConstraintBaker.cpp b/Engine/Plugins/Animation/ControlRig/Source/ControlRigEditor/Private/Tools/ConstraintBaker.cpp
--- a/Engine/Plugins/Animation/ControlRig/Source/ControlRigEditor/Private/Tools/ConstraintBaker.cpp
+++ b/Engine/Plugins/Animation/ControlRig/Source/ControlRigEditor/Private/Tools/ConstraintBaker.cpp
@@ -90,25 +90,6 @@ void BakeControl(
 		MovieScene->GetTickResolution());
 }
 
-EMovieSceneTransformChannel GetChannelsToKey(const UTickableTransformConstraint* InConstraint)
-{
-	static const TMap< ETransformConstraintType, EMovieSceneTransformChannel > ConstraintToChannels({
-	{ETransformConstraintType::Translation, EMovieSceneTransformChannel::Translation},
-	{ETransformConstraintType::Rotation, EMovieSceneTransformChannel::Rotation},
-	{ETransformConstraintType::Scale, EMovieSceneTransformChannel::Scale},
-	{ETransformConstraintType::Parent, EMovieSceneTransformChannel::AllTransform},
-	{ETransformConstraintType::LookAt, EMovieSceneTransformChannel::Rotation}
-	});
-
-	const ETransformConstraintType ConstType = static_cast<ETransformConstraintType>(InConstraint->GetType());
-	if (const EMovieSceneTransformChannel* Channel = ConstraintToChannels.Find(ConstType))
-	{
-		return *Channel;; 
-	}
-	
-	return EMovieSceneTransformChannel::AllTransform;
-}
-	
 }
 
 void FConstraintBaker::DoIt(UTickableTransformConstraint* InConstraint)
@@ -138,7 +119,7 @@ void FConstraintBaker::DoIt(UTickableTransformConstraint* InConstraint)
 
 	// compute transforms
 	TArray<FTransform> Transforms;
-	GetChildTransforms(Sequencer, *InConstraint, Frames, true, Transforms);
+	GetHandleTransforms(Sequencer, InConstraint->ChildTRSHandle, {InConstraint}, Frames, true, Transforms);
 	
 	if (Frames.Num() != Transforms.Num())
 	{
@@ -150,20 +131,49 @@ void FConstraintBaker::DoIt(UTickableTransformConstraint* InConstraint)
 	BakeChild(Sequencer, InConstraint->ChildTRSHandle, Frames, Transforms, Channels);
 
 	// disable constraint
-	InConstraint->Active = false;
-	InConstraint->ConstraintTick.SetTickFunctionEnable(false);
+	InConstraint->SetActive(false);
 	
 	// notify
 	Sequencer->NotifyMovieSceneDataChanged(EMovieSceneDataChangeType::MovieSceneStructureItemsChanged);
 }
 
-void FConstraintBaker::GetChildTransforms(
+void FConstraintBaker::GetHandleTransforms(
+	UWorld* InWorld,
+	const TSharedPtr<ISequencer>& InSequencer,
+	const UTransformableHandle* InHandle,
+	const TArray<FFrameNumber>& InFrames,
+	const bool bLocal,
+	TArray<FTransform>& OutTransforms)
+{
+	const FConstraintsManagerController& Controller = FConstraintsManagerController::Get(InWorld);
+	static constexpr bool bSorted = true;
+	TArray<TObjectPtr<UTickableConstraint>> Constraints = Controller.GetParentConstraints(InHandle->GetHash(), bSorted);
+
+	TArray<UTickableTransformConstraint*> TransformConstraints;
+	for (const TObjectPtr<UTickableConstraint>& Constraint: Constraints)
+	{
+		if (UTickableTransformConstraint* TransformConstraint = Cast<UTickableTransformConstraint>(Constraint.Get()))
+		{
+			TransformConstraints.Add(TransformConstraint);
+		}
+	}
+	
+	return GetHandleTransforms(InSequencer, InHandle, TransformConstraints, InFrames, bLocal, OutTransforms);
+}
+
+void FConstraintBaker::GetHandleTransforms(
 	const TSharedPtr<ISequencer>& InSequencer,
-	const UTickableTransformConstraint& InConstraint,
+	const UTransformableHandle* InHandle,
+	const TArray<UTickableTransformConstraint*>& InConstraintsToEvaluate,
 	const TArray<FFrameNumber>& InFrames,
 	const bool bLocal,
 	TArray<FTransform>& OutTransforms)
 {
+	// if (InConstraintsToEvaluate.IsEmpty())
+	// {
+	// 	return;
+	// }
+	
 	const UMovieScene* MovieScene = InSequencer->GetFocusedMovieSceneSequence()->GetMovieScene();
 	
 	const FFrameRate TickResolution = MovieScene->GetTickResolution();
@@ -181,12 +191,15 @@ void FConstraintBaker::GetChildTransforms(
 		InSequencer->GetEvaluationTemplate().Evaluate(Context, *InSequencer);
 	
 		// evaluate constraints
-		InConstraint.Evaluate();
+		for (const UTickableTransformConstraint* Constraint: InConstraintsToEvaluate)
+		{
+			Constraint->Evaluate();
+		}
 	
 		// evaluate ControlRig?
 		// ControlRig->Evaluate_AnyThread();
 		
-		OutTransforms[Index] = bLocal ? InConstraint.GetChildLocalTransform() : InConstraint.GetChildGlobalTransform();
+		OutTransforms[Index] = bLocal ? InHandle->GetLocalTransform() : InHandle->GetGlobalTransform();
 	}
 }
 
@@ -208,6 +221,25 @@ void FConstraintBaker::BakeChild(
 	}
 }
 
+EMovieSceneTransformChannel FConstraintBaker::GetChannelsToKey(const UTickableTransformConstraint* InConstraint)
+{
+	static const TMap< ETransformConstraintType, EMovieSceneTransformChannel > ConstraintToChannels({
+	{ETransformConstraintType::Translation, EMovieSceneTransformChannel::Translation},
+	{ETransformConstraintType::Rotation, EMovieSceneTransformChannel::Rotation},
+	{ETransformConstraintType::Scale, EMovieSceneTransformChannel::Scale},
+	{ETransformConstraintType::Parent, EMovieSceneTransformChannel::AllTransform},
+	{ETransformConstraintType::LookAt, EMovieSceneTransformChannel::Rotation}
+	});
+
+	const ETransformConstraintType ConstType = static_cast<ETransformConstraintType>(InConstraint->GetType());
+	if (const EMovieSceneTransformChannel* Channel = ConstraintToChannels.Find(ConstType))
+	{
+		return *Channel;; 
+	}
+	
+	return EMovieSceneTransformChannel::AllTransform;
+}
+
 #undef LOCTEXT_NAMESPACE
 
 
diff --git a/Engine/Plugins/Animation/ControlRig/Source/ControlRigEditor/Public/ConstraintChannelHelper.h b/Engine/Plugins/Animation/ControlRig/Source/ControlRigEditor/Public/ConstraintChannelHelper.h
new file mode 100644
--- /dev/null
+++ b/Engine/Plugins/Animation/ControlRig/Source/ControlRigEditor/Public/ConstraintChannelHelper.h
@@ -0,0 +1,90 @@
+﻿// Copyright Epic Games, Inc. All Rights Reserved.
+
+#pragma once
+
+#include "Templates/SharedPointer.h"
+
+class ISequencer;
+class UMovieScene3DTransformSection;
+class UTransformableHandle;
+class UTickableTransformConstraint;
+class UTransformableControlHandle;
+class UTransformableComponentHandle;
+struct FMovieSceneFloatChannel;
+struct FMovieSceneConstraintChannel;
+class UMovieSceneControlRigParameterSection;
+struct FFrameNumber;
+
+/**
+ * 
+ */
+
+struct FConstraintChannelHelper
+{
+public:
+	/** @todo documentation. */
+	static void AddConstraintKey(UTickableTransformConstraint* InConstraint);
+	static void SmartConstraintKey(UTickableTransformConstraint* InConstraint);
+
+	/** @todo documentation. */
+	static void AddChildTransformKey(
+		const UTransformableHandle* InHandle,
+		const FFrameNumber& InTime,
+		const TSharedPtr<ISequencer>& InSequencer);
+
+private:
+
+	/** BEGIN CONTROL RIG SECTION */
+	
+	/** @todo documentation. */
+	static void SmartControlConstraintKey(
+		UTickableTransformConstraint* InConstraint,
+		const TSharedPtr<ISequencer>& InSequencer);
+
+	/** @todo documentation. */
+	static UMovieSceneControlRigParameterSection* GetControlSection(
+		const UTransformableControlHandle* InHandle,
+		const TSharedPtr<ISequencer>& InSequencer);
+	
+	/** END CONTROL RIG SECTION */
+
+	/** BEGIN COMPONENT SECTION */
+
+	/** @todo documentation. */
+	static void SmartComponentConstraintKey(
+		UTickableTransformConstraint* InConstraint,
+		const TSharedPtr<ISequencer>& InSequencer);
+
+	/** @todo documentation. */
+	static UMovieScene3DTransformSection* GetTransformSection(
+		const UTransformableComponentHandle* InHandle,
+		const TSharedPtr<ISequencer>& InSequencer);
+
+	/** END COMPONENT SECTION */
+
+	/** @todo documentation. */
+	static void GetFramesToCompensate(
+		const FMovieSceneConstraintChannel& InActiveChannel,
+		const bool InValue,
+		const FFrameNumber& InTime,
+		const TArrayView<FMovieSceneFloatChannel*>& InChannels,
+		TArray<FFrameNumber>& OutFramesAfter);
+
+	/** @todo documentation. */
+	static TArrayView<FMovieSceneFloatChannel*> GetTransformFloatChannels(
+		const UTransformableHandle* InHandle,
+		const TSharedPtr<ISequencer>& InSequencer);
+};
+
+/** Key drawing overrides */
+// void DrawKeys(
+// 	FMovieSceneConstraintChannel* Channel,
+// 	TArrayView<const FKeyHandle> InKeyHandles,
+// 	const UMovieSceneSection* InOwner,
+// 	TArrayView<FKeyDrawParams> OutKeyDrawParams);
+//
+// void DrawExtra(
+// 	FMovieSceneConstraintChannel* Channel,
+// 	const UMovieSceneSection* Owner,
+// 	const FGeometry& AllottedGeometry,
+// 	FSequencerSectionPainter& Painter);
diff --git a/Engine/Plugins/Animation/ControlRig/Source/ControlRigEditor/Public/Tools/BakingHelper.h b/Engine/Plugins/Animation/ControlRig/Source/ControlRigEditor/Public/Tools/BakingHelper.h
--- a/Engine/Plugins/Animation/ControlRig/Source/ControlRigEditor/Public/Tools/BakingHelper.h
+++ b/Engine/Plugins/Animation/ControlRig/Source/ControlRigEditor/Public/Tools/BakingHelper.h
@@ -49,7 +49,8 @@ struct FBakingHelper
 		UControlRig* InControlRig,
 		const FName& InControlName,
 		const TArray<FFrameNumber>& InFrames,
-		const TArray<FTransform>& InLocalTransforms,
+		const TArray<FTransform>& InTransforms,
 		const EMovieSceneTransformChannel& InChannels,
-		const FFrameRate& InTickResolution);
+		const FFrameRate& InTickResolution,
+		const bool bLocal = true);
 };
diff --git a/Engine/Plugins/Animation/ControlRig/Source/ControlRigEditor/Public/Tools/ConstraintBaker.h b/Engine/Plugins/Animation/ControlRig/Source/ControlRigEditor/Public/Tools/ConstraintBaker.h
--- a/Engine/Plugins/Animation/ControlRig/Source/ControlRigEditor/Public/Tools/ConstraintBaker.h
+++ b/Engine/Plugins/Animation/ControlRig/Source/ControlRigEditor/Public/Tools/ConstraintBaker.h
@@ -15,13 +15,25 @@
 public:
 	/** @todo documentation. (subject to changes to handle start / stop as parameters) */
 	static void DoIt(UTickableTransformConstraint* InConstraint);
+
+	/** @todo documentation. */
+	static void GetHandleTransforms(
+		UWorld* InWorld,
+		const TSharedPtr<ISequencer>& InSequencer,
+		const UTransformableHandle* InHandle,
+		const TArray<FFrameNumber>& InFrames,
+		const bool bLocal,
+		TArray<FTransform>& OutTransforms); 
+
+	static EMovieSceneTransformChannel GetChannelsToKey(const UTickableTransformConstraint* InConstraint);
 	
 private:
 
 	/** Evaluates the constraint at each frames and stores the resulting child transforms. */
-	static void GetChildTransforms(
+	static void GetHandleTransforms(
 		const TSharedPtr<ISequencer>& InSequencer,
-		const UTickableTransformConstraint& InConstraint,
+		const UTransformableHandle* InHandle,
+		const TArray<UTickableTransformConstraint*>& InConstraintsToEvaluate,
 		const TArray<FFrameNumber>& InFrames,
 		const bool bLocal,
 		TArray<FTransform>& OutTransforms);
diff --git a/Engine/Source/Editor/MovieSceneTools/Private/Channels/ConstraintChannelEditor.cpp b/Engine/Source/Editor/MovieSceneTools/Private/Channels/ConstraintChannelEditor.cpp
new file mode 100644
--- /dev/null
+++ b/Engine/Source/Editor/MovieSceneTools/Private/Channels/ConstraintChannelEditor.cpp
@@ -0,0 +1,189 @@
+﻿// Copyright Epic Games, Inc. All Rights Reserved.
+
+#include "ConstraintChannelEditor.h"
+
+#include "Channels/MovieSceneConstraintChannel.h"
+
+#include "KeyBarCurveModel.h"
+#include "SequencerSectionPainter.h"
+#include "TimeToPixel.h"
+#include "Styling/AppStyle.h"
+#include "Styling/CoreStyle.h"
+#include "Framework/Application/SlateApplication.h"
+
+namespace
+{
+	using IndexRange = ::TRange<int32>;
+	IndexRange GetRange(const TArrayView<const bool>& Values, const int32 Offset) 
+	{
+		IndexRange Range;
+			
+		int32 FirstActive = INDEX_NONE;
+		for (int32 Index = Offset; Index < Values.Num(); ++Index)
+		{
+			if (Values[Index])
+			{
+				FirstActive = Index;
+				break;
+			}
+		}
+
+		if (FirstActive == INDEX_NONE)
+		{
+			return Range;
+		}
+
+		Range.SetLowerBound(TRangeBound<int32>(FirstActive));
+		Range.SetUpperBound(TRangeBound<int32>(FirstActive));
+
+		for (int32 NextInactive = FirstActive+1; NextInactive < Values.Num(); ++NextInactive)
+		{
+			if (Values[NextInactive] == false)
+			{
+				Range.SetUpperBound(TRangeBound<int32>(NextInactive));
+				return Range;
+			}
+		}
+			
+		return Range;
+	};
+
+	TArray<IndexRange> GetIndexRanges(const FMovieSceneConstraintChannel* Channel)
+	{
+		TArray<IndexRange> Ranges;
+
+		const TMovieSceneChannelData<const bool> ChannelData = Channel->GetData();
+		const TArrayView<const bool> Values = ChannelData.GetValues();
+		if (Values.Num() == 1)
+		{
+			if (Values[0] == true)
+			{
+				Ranges.Emplace(0);
+			}
+			return Ranges;
+		}
+		
+		int32 Offset = 0;
+		while (Values.IsValidIndex(Offset))
+		{
+			TRange<int32> Range = GetRange(Values, Offset);
+			if (!Range.IsEmpty())
+			{
+				Ranges.Emplace(Range);
+				Offset = Range.GetUpperBoundValue()+1;
+			}
+			else
+			{
+				Offset = INDEX_NONE;
+			}
+		}
+			
+		return Ranges;
+	}
+}
+
+void DrawKeys(
+	FMovieSceneConstraintChannel* Channel,
+	TArrayView<const FKeyHandle> InKeyHandles,
+	const UMovieSceneSection* InOwner,
+	TArrayView<FKeyDrawParams> OutKeyDrawParams)
+{
+	if (!InOwner)
+	{
+		return;
+	}
+	
+	static const FName SquareKeyBrushName("Sequencer.KeySquare");
+	static const FName FilledBorderBrushName("FilledBorder");
+
+	static FKeyDrawParams Params;
+	Params.FillBrush = FAppStyle::GetBrush(FilledBorderBrushName);
+	Params.BorderBrush = FAppStyle::GetBrush(SquareKeyBrushName);
+	
+	for (FKeyDrawParams& Param : OutKeyDrawParams)
+	{
+		Param = Params;
+	}
+}
+
+void DrawExtra(
+	FMovieSceneConstraintChannel* Channel,
+	const UMovieSceneSection* Owner,
+	const FGeometry& AllottedGeometry,
+	FSequencerSectionPainter& Painter)
+{
+	if (!Owner || !Channel)
+	{
+		return;
+	}
+	
+	// using namespace UE::Sequencer;
+
+	// get index range
+	TArray<IndexRange> IndexRanges = GetIndexRanges(Channel);
+	if (IndexRanges.IsEmpty())
+	{
+		return;
+	}
+	
+	// convert to bar range
+	TArray<FKeyBarCurveModel::FBarRange> Ranges;
+	const TMovieSceneChannelData<const bool> ChannelData = Channel->GetData();
+	const TArrayView<const FFrameNumber> Times = ChannelData.GetTimes();
+
+	// const bool IsFirstValueTrue = IndexRanges[0].GetLowerBoundValue() == 0;
+	
+	for (const IndexRange& ActiveRange: IndexRanges)
+	{
+		FKeyBarCurveModel::FBarRange BarRange;
+		FFrameRate TickResolution = Owner->GetTypedOuter<UMovieScene>()->GetTickResolution();
+		double LowerValue = Times[ActiveRange.GetLowerBoundValue()] / TickResolution;
+		double UpperValue = Times[ActiveRange.GetUpperBoundValue()] / TickResolution;
+	
+		BarRange.Range.SetLowerBound(TRangeBound<double>(LowerValue));
+		BarRange.Range.SetUpperBound(TRangeBound<double>(UpperValue));
+	
+		BarRange.Name = "Constraint";
+		BarRange.Color = FLinearColor(.2, .5, .1);
+		static FLinearColor ZebraTint = FLinearColor::White.CopyWithNewOpacity(0.01f);
+		BarRange.Color = BarRange.Color * (1.f - ZebraTint.A) + ZebraTint * ZebraTint.A;
+		
+		Ranges.Add(BarRange);
+	}
+
+	// draw bars
+	const FSlateBrush* WhiteBrush = FAppStyle::GetBrush("WhiteBrush");
+	static constexpr ESlateDrawEffect DrawEffects = ESlateDrawEffect::None;
+
+	const FSlateFontInfo FontInfo = FCoreStyle::Get().GetFontStyle("ToolTip.LargerFont");
+	const TSharedRef<FSlateFontMeasure> FontMeasure = FSlateApplication::Get().GetRenderer()->GetFontMeasureService();
+
+	const FVector2D& LocalSize = AllottedGeometry.GetLocalSize();
+	static constexpr float LaneTop = 0;
+
+	const FTimeToPixel& TimeToPixel = Painter.GetTimeConverter();
+	const double InputMin = TimeToPixel.PixelToSeconds(0.f);
+	const double InputMax = TimeToPixel.PixelToSeconds(Painter.SectionGeometry.GetLocalSize().X);
+
+	for (int32 Index = 0; Index < Ranges.Num(); ++Index)
+	{
+		const FKeyBarCurveModel::FBarRange& Range = Ranges[Index];
+
+		double LowerSeconds = /*(Index == 0 && IsFirstValueTrue) ? InputMin : */Range.Range.GetLowerBoundValue();
+		double UpperSeconds = Range.Range.GetUpperBoundValue();
+		if (UpperSeconds == Range.Range.GetLowerBoundValue())
+		{
+			UpperSeconds = InputMax;
+		}
+
+		const double BoxStart = TimeToPixel.SecondsToPixel(LowerSeconds);
+		const double BoxEnd = TimeToPixel.SecondsToPixel(UpperSeconds);
+		const double BoxSize = BoxEnd - BoxStart;
+
+		const FVector2D Size = FVector2D(BoxSize, LocalSize.Y);
+		const FVector2D Translation = FVector2D(BoxStart, LaneTop);
+		const FPaintGeometry BoxGeometry = AllottedGeometry.ToPaintGeometry(Size, FSlateLayoutTransform(Translation));
+		
+		FSlateDrawElement::MakeBox(Painter.DrawElements, Painter.LayerId, BoxGeometry, WhiteBrush, DrawEffects, Range.Color);
+	}
+}
diff --git a/Engine/Source/Editor/MovieSceneTools/Private/Channels/ConstraintChannelEditor.h b/Engine/Source/Editor/MovieSceneTools/Private/Channels/ConstraintChannelEditor.h
new file mode 100644
--- /dev/null
+++ b/Engine/Source/Editor/MovieSceneTools/Private/Channels/ConstraintChannelEditor.h
@@ -0,0 +1,21 @@
+﻿// Copyright Epic Games, Inc. All Rights Reserved.
+
+#pragma once
+
+#include "SequencerChannelTraits.h"
+#include "SequencerChannelInterface.h"
+
+struct FMovieSceneConstraintChannel;
+
+/** Key drawing overrides */
+void DrawKeys(
+	FMovieSceneConstraintChannel* Channel,
+	TArrayView<const FKeyHandle> InKeyHandles,
+	const UMovieSceneSection* InOwner,
+	TArrayView<FKeyDrawParams> OutKeyDrawParams);
+
+void DrawExtra(
+	FMovieSceneConstraintChannel* Channel,
+	const UMovieSceneSection* Owner,
+	const FGeometry& AllottedGeometry,
+	FSequencerSectionPainter& Painter);
diff --git a/Engine/Source/Runtime/Experimental/Animation/Constraints/Private/ConstraintsManager.cpp b/Engine/Source/Runtime/Experimental/Animation/Constraints/Private/ConstraintsManager.cpp
--- a/Engine/Source/Runtime/Experimental/Animation/Constraints/Private/ConstraintsManager.cpp
+++ b/Engine/Source/Runtime/Experimental/Animation/Constraints/Private/ConstraintsManager.cpp
@@ -65,6 +65,12 @@
  * UTickableConstraint
  **/
 
+void UTickableConstraint::SetActive(const bool bIsActive)
+{
+	Active = bIsActive;
+	ConstraintTick.SetTickFunctionEnable(bIsActive);
+}
+
 void UTickableConstraint::Evaluate() const
 {
 	ConstraintTick.EvaluateFunctions();
@@ -100,17 +106,55 @@
  **/
 
 UConstraintsManager::UConstraintsManager()
-	: Level(nullptr)
 {}
 
 UConstraintsManager::~UConstraintsManager()
-{}
-
-void UConstraintsManager::Init(ULevel* InLevel)
 {
-	Level = InLevel;
+	UnregisterDelegates();
 }
 
+void UConstraintsManager::OnActorDestroyed(AActor* InActor)
+{
+	if (USceneComponent* SceneComponent = InActor->GetRootComponent())
+	{
+		Constraints.RemoveAll([SceneComponent](const TObjectPtr<UTickableConstraint>& Constraint)
+		{
+			return IsValid(Constraint) ? Constraint->ReferencesObject(SceneComponent) : false;
+		} );
+	}
+}
+
+void UConstraintsManager::RegisterDelegates()
+{
+	if (!OnActorDestroyedHandle.IsValid())
+	{
+		FOnActorDestroyed::FDelegate ActorDestroyedDelegate =
+				FOnActorDestroyed::FDelegate::CreateUObject(this, &UConstraintsManager::OnActorDestroyed);
+		OnActorDestroyedHandle = World->AddOnActorDestroyedHandler(ActorDestroyedDelegate);
+	}
+}
+
+void UConstraintsManager::UnregisterDelegates()
+{
+	if (World)
+	{
+		World->RemoveOnActorSpawnedHandler(OnActorDestroyedHandle);
+	}
+	OnActorDestroyedHandle.Reset();
+}
+
+void UConstraintsManager::Init(UWorld* InWorld)
+{
+	if (InWorld)
+	{
+		UnregisterDelegates();
+
+		World = InWorld;
+
+		RegisterDelegates();
+	}
+}
+
 UConstraintsManager* UConstraintsManager::Get(UWorld* InWorld)
 {
 	// look for ConstraintsActor and return its manager
@@ -126,9 +170,9 @@
 #endif // WITH_EDITOR
 	ConstraintsActor->ConstraintsManager = NewObject<UConstraintsManager>(ConstraintsActor);
 	
-	ULevel* Level = InWorld->GetCurrentLevel();
-	ConstraintsActor->ConstraintsManager->Init(Level);
-	ConstraintsActor->ConstraintsManager->Level = Level;
+	// ULevel* Level = InWorld->GetCurrentLevel();
+	ConstraintsActor->ConstraintsManager->Init(InWorld);
+	// ConstraintsActor->ConstraintsManager->Level = Level;
 
 	return ConstraintsActor->ConstraintsManager;
 }
@@ -233,9 +277,9 @@
 	ConstraintsActor->SetActorLabel(TEXT("Constraints Manager"));
 #endif // WITH_EDITOR
 	ConstraintsActor->ConstraintsManager = NewObject<UConstraintsManager>(ConstraintsActor);
-	
-	ULevel* Level = World->GetCurrentLevel();
-	ConstraintsActor->ConstraintsManager->Level = Level;
+	ConstraintsActor->ConstraintsManager->Init(World);
+	// ULevel* Level = World->GetCurrentLevel();
+	// ConstraintsActor->ConstraintsManager->Level = Level;
 
 	return ConstraintsActor->ConstraintsManager;
 }
@@ -348,10 +392,14 @@
 	}
 
 	// TODO handle transaction
+	const FName ConstraintName = Manager->Constraints[InConstraintIndex]->GetFName(); 
 	
 	Manager->Constraints[InConstraintIndex]->ConstraintTick.UnRegisterTickFunction();
 	Manager->Constraints.RemoveAt(InConstraintIndex);
 
+	// notify deletion
+	ConstraintRemoved.Broadcast(ConstraintName);
+	
 	// destroy constraints actor if no constraints left
 	if (Manager->Constraints.IsEmpty())
 	{
diff --git a/Engine/Source/Runtime/Experimental/Animation/Constraints/Private/TransformConstraint.cpp b/Engine/Source/Runtime/Experimental/Animation/Constraints/Private/TransformConstraint.cpp
--- a/Engine/Source/Runtime/Experimental/Animation/Constraints/Private/TransformConstraint.cpp
+++ b/Engine/Source/Runtime/Experimental/Animation/Constraints/Private/TransformConstraint.cpp
@@ -96,6 +96,23 @@ uint32 UTickableTransformConstraint::GetTargetHash() const
 	return ChildTRSHandle->IsValid() ? ChildTRSHandle->GetHash() : 0;
 }
 
+bool UTickableTransformConstraint::ReferencesObject(TWeakObjectPtr<UObject> InObject) const
+{
+	const TWeakObjectPtr<UObject> ChildTarget = ChildTRSHandle->IsValid() ? ChildTRSHandle->GetTarget() : nullptr;
+	if (ChildTarget == InObject)
+	{
+		return true;	
+	}
+
+	const TWeakObjectPtr<UObject> ParentTarget = ParentTRSHandle->IsValid() ? ParentTRSHandle->GetTarget() : nullptr;
+	if (ParentTarget == InObject)
+	{
+		return true;	
+	}
+	
+	return false;
+}
+
 void UTickableTransformConstraint::SetChildGlobalTransform(const FTransform& InGlobal) const
 {
 	if(ChildTRSHandle->IsValid())
@@ -308,6 +325,21 @@ void UTickableParentConstraint::ComputeOffset()
 	}
 }
 
+uint32 UTickableParentConstraint::CalculateDependenciesHash() const
+{
+	uint32 Hash = 0;
+	
+	Hash = HashCombine(Hash, GetTypeHash(Active));
+	Hash = HashCombine(Hash, GetTypeHash(Weight));
+
+	const FTransform ParentGlobal = GetParentGlobalTransform();
+	Hash = HashCombine(Hash, GetTypeHash(ParentGlobal.GetTranslation()));
+	Hash = HashCombine(Hash, GetTypeHash(ParentGlobal.GetRotation().Euler() ));
+	Hash = HashCombine(Hash, GetTypeHash(ParentGlobal.GetScale3D()));
+	
+	return Hash;
+}
+
 FConstraintTickFunction::ConstraintFunction UTickableParentConstraint::GetFunction() const
 {
 	return [this]()
@@ -322,21 +354,62 @@ FConstraintTickFunction::ConstraintFunction UTickableParentConstraint::GetFuncti
 		{
 			return;
 		}
-		
-		const FTransform ParentTransform = ParentTRSHandle->GetGlobalTransform();
 
-		FTransform TargetTransform = bMaintainOffset ? OffsetTransform * ParentTransform : ParentTransform;
-		if (ClampedWeight < 1.0f - KINDA_SMALL_NUMBER)
+		auto LerpTransform = [ClampedWeight](const FTransform& InTransform, FTransform& InTransformToBeSet)
 		{
-			const  FTransform Transform = GetChildGlobalTransform();
-			TargetTransform.SetLocation(
-				FMath::Lerp<FVector>(Transform.GetLocation(), TargetTransform.GetLocation(), ClampedWeight));
-			TargetTransform.SetRotation(
-				FQuat::Slerp(Transform.GetRotation(), TargetTransform.GetRotation(), ClampedWeight));
-			TargetTransform.SetScale3D(
-				FMath::Lerp<FVector>(Transform.GetScale3D(), TargetTransform.GetScale3D(), ClampedWeight));
-		}
+			if (ClampedWeight < 1.0f - KINDA_SMALL_NUMBER)
+			{
+				InTransformToBeSet.SetLocation(
+					FMath::Lerp<FVector>(InTransform.GetLocation(), InTransformToBeSet.GetLocation(), ClampedWeight));
+				InTransformToBeSet.SetRotation(
+					FQuat::Slerp(InTransform.GetRotation(), InTransformToBeSet.GetRotation(), ClampedWeight));
+				InTransformToBeSet.SetScale3D(
+					FMath::Lerp<FVector>(InTransform.GetScale3D(), InTransformToBeSet.GetScale3D(), ClampedWeight));
+			}
+		};
+
+		const FTransform ParentTransform = GetParentGlobalTransform();
 		
+		// if bDynamicOffset is on, means we assume that the incoming local child transform is coming from the child
+		// TRS animation (for instance) which is local to the constraint's parent
+		if (bDynamicOffset)
+		{
+			const uint32 DependenciesHash = CalculateDependenciesHash();
+
+			const FTransform ChildGlobalTransform = GetChildGlobalTransform();
+			const bool bIsGlobalDirty = !ChildGlobalTransform.Equals(Cache.LastGlobalSet);
+			const bool bIsInputHashDirty = DependenciesHash != Cache.CachedInputHash;
+
+			// we avoid computation if nothing has changed
+			if (!bIsGlobalDirty && !bIsInputHashDirty)
+			{
+				return;
+			}
+
+			FTransform LocalTransform = Cache.LastLocalSet;
+			if (bIsGlobalDirty)
+			{
+				LocalTransform = GetChildLocalTransform();
+			}
+			
+			FTransform TargetTransform = LocalTransform * ParentTransform;
+			//apply weight if needed
+			LerpTransform(ChildGlobalTransform, TargetTransform);
+			
+			SetChildGlobalTransform(TargetTransform);
+
+			// update cache
+			Cache.CachedInputHash = DependenciesHash;
+			Cache.LastGlobalSet = TargetTransform;
+			Cache.LastLocalSet = LocalTransform;
+			
+			return;
+		}
+
+		// static behavior
+		FTransform TargetTransform = bMaintainOffset ? OffsetTransform * ParentTransform : ParentTransform;
+		//apply weight if needed
+		LerpTransform(GetChildGlobalTransform(), TargetTransform);
 		SetChildGlobalTransform(TargetTransform);
 	};
 }
diff --git a/Engine/Source/Runtime/Experimental/Animation/Constraints/Private/TransformableHandle.cpp b/Engine/Source/Runtime/Experimental/Animation/Constraints/Private/TransformableHandle.cpp
--- a/Engine/Source/Runtime/Experimental/Animation/Constraints/Private/TransformableHandle.cpp
+++ b/Engine/Source/Runtime/Experimental/Animation/Constraints/Private/TransformableHandle.cpp
@@ -71,6 +71,11 @@
 	return Component.IsValid() ? GetTypeHash(Component.Get()) : 0;
 }
 
+TWeakObjectPtr<UObject> UTransformableComponentHandle::GetTarget() const
+{
+	return Component;
+}
+
 #if WITH_EDITOR
 FName UTransformableComponentHandle::GetName() const
 {
diff --git a/Engine/Source/Runtime/Experimental/Animation/Constraints/Public/ConstraintsManager.h b/Engine/Source/Runtime/Experimental/Animation/Constraints/Public/ConstraintsManager.h
--- a/Engine/Source/Runtime/Experimental/Animation/Constraints/Public/ConstraintsManager.h
+++ b/Engine/Source/Runtime/Experimental/Animation/Constraints/Public/ConstraintsManager.h
@@ -81,12 +81,17 @@
 	UPROPERTY()
 	FConstraintTickFunction ConstraintTick;
 
+	/** Sets the Active value and enable/disable the tick function. */
+	void SetActive(const bool bIsActive);
+	
 	/** Evaluates the constraint in a context where it's mot done thru the ConstraintTick's tick function. */
 	void Evaluate() const;
 	
 	/** @todo document */
 	virtual uint32 GetTargetHash() const PURE_VIRTUAL(GetTargetHash, return 0;);
-	
+	/** @todo document */
+	virtual bool ReferencesObject(TWeakObjectPtr<UObject> InObject) const PURE_VIRTUAL(ReferencesObject, return false;);
+
 #if WITH_EDITOR
 	/** @todo document */
 	virtual FName GetLabel() const;
@@ -94,7 +99,7 @@
 	// UObject interface
 	virtual void PostEditChangeProperty(FPropertyChangedEvent& PropertyChangedEvent) override;
 	// End of UObject interface
-	#endif
+#endif
 
 	/** @todo documentation. */
 	UPROPERTY(EditAnywhere, DisplayName="Active State", Category="Constraint")
@@ -124,7 +129,7 @@
 	static UConstraintsManager* Find(const UWorld* InWorld);
 
 	/** @todo document */
-	void Init(ULevel* InLevel);
+	void Init(UWorld* InWorld);
 	
 	/* Set tick dependencies between two constraints. */
 	void SetConstraintDependencies(
@@ -137,11 +142,19 @@
 private:
 
 	/** @todo document */
+	FDelegateHandle OnActorDestroyedHandle;
+	
+	void OnActorDestroyed(AActor* InActor);
+
+	void RegisterDelegates();
+	void UnregisterDelegates();
+	
+	/** @todo document */
 	void Dump() const;
 
 	/** @todo document */
 	UPROPERTY()
-	TObjectPtr<ULevel> Level = nullptr;
+	TObjectPtr<UWorld> World = nullptr;
 
 	/** @todo document */
 	UPROPERTY()
@@ -213,6 +226,10 @@
 		const FName& InNameToTickAfter) const;
 	
 private:
+
+	/** @todo document */
+	DECLARE_MULTICAST_DELEGATE_OneParam(FOnConstraintRemoved, FName /*InConstraintName*/);
+	FOnConstraintRemoved ConstraintRemoved;
 	
 	/** Find the existing Constraints Manager in World or create a new one. */
 	UConstraintsManager* GetManager() const;
@@ -225,4 +242,9 @@
 
 	/** The World that holds the ConstraintsManagerActor. */
 	UWorld* World = nullptr;
+
+public:
+
+	/** @todo document */
+	FOnConstraintRemoved& OnConstraintRemoved(){ return ConstraintRemoved; };
 };
diff --git a/Engine/Source/Runtime/Experimental/Animation/Constraints/Public/TransformConstraint.h b/Engine/Source/Runtime/Experimental/Animation/Constraints/Public/TransformConstraint.h
--- a/Engine/Source/Runtime/Experimental/Animation/Constraints/Public/TransformConstraint.h
+++ b/Engine/Source/Runtime/Experimental/Animation/Constraints/Public/TransformConstraint.h
@@ -32,6 +32,8 @@ public:
 
 	/** @todo document */
 	virtual uint32 GetTargetHash() const override;
+	/** @todo document */
+	virtual bool ReferencesObject(TWeakObjectPtr<UObject> InObject) const override;
 	
 	/** @todo document */
 	UPROPERTY()
@@ -48,7 +50,10 @@ public:
 	// UPROPERTY(EditAnywhere, Category="Weight", meta = (Input, ClampMin = "0.0", ClampMax = "1.0", UIMin = "0.0", UIMax = "1.0"))
 	UPROPERTY(meta = (Input, ClampMin = "0.0", ClampMax = "1.0", UIMin = "0.0", UIMax = "1.0"))
 	float Weight = 1.f;
-	
+
+	UPROPERTY(EditAnywhere, Category="Offset")
+	bool bDynamicOffset = false;
+
 	/** @todo document. */
 	int64 GetType() const;
 
@@ -183,6 +188,18 @@ protected:
 	/** @todo document */
 	UPROPERTY(EditAnywhere, Category="Offset")
 	FTransform OffsetTransform = FTransform::Identity;
+
+	/** @todo document */
+	struct FDynamicCache
+	{
+		FTransform LastGlobalSet = FTransform::Identity;
+		FTransform LastLocalSet = FTransform::Identity;
+		uint32 CachedInputHash = 0;
+	};
+	mutable FDynamicCache Cache;
+
+	/** @todo document */
+	uint32 CalculateDependenciesHash() const;
 };
 
 /** 
diff --git a/Engine/Source/Runtime/Experimental/Animation/Constraints/Public/TransformableHandle.h b/Engine/Source/Runtime/Experimental/Animation/Constraints/Public/TransformableHandle.h
--- a/Engine/Source/Runtime/Experimental/Animation/Constraints/Public/TransformableHandle.h
+++ b/Engine/Source/Runtime/Experimental/Animation/Constraints/Public/TransformableHandle.h
@@ -50,7 +50,10 @@
 	 * Generates a hash value of the underlying transformable object.
 	**/
 	virtual uint32 GetHash() const PURE_VIRTUAL(GetHash, return 0;);
-	
+
+	/** @todo document */
+	virtual TWeakObjectPtr<UObject> GetTarget() const PURE_VIRTUAL(GetTarget, return nullptr;);
+
 #if WITH_EDITOR
 	virtual FName GetName() const PURE_VIRTUAL(GetName, return NAME_None;);
 #endif
@@ -89,6 +92,9 @@
 	/** @todo document */
 	virtual uint32 GetHash() const override;
 
+	/** @todo document */
+	virtual TWeakObjectPtr<UObject> GetTarget() const override;
+
 #if WITH_EDITOR
 	/** @todo document */
 	virtual FName GetName() const override;
diff --git a/Engine/Source/Runtime/MovieScene/Private/Channels/MovieSceneConstraintChannel.cpp b/Engine/Source/Runtime/MovieScene/Private/Channels/MovieSceneConstraintChannel.cpp
new file mode 100644
--- /dev/null
+++ b/Engine/Source/Runtime/MovieScene/Private/Channels/MovieSceneConstraintChannel.cpp
@@ -0,0 +1,19 @@
+﻿// Copyright Epic Games, Inc. All Rights Reserved.
+
+#include "Channels/MovieSceneConstraintChannel.h"
+
+bool FMovieSceneConstraintChannel::Evaluate(FFrameTime InTime, bool& OutValue) const
+{
+	const TArrayView<const FFrameNumber> TimeArray = GetTimes();
+	if (TimeArray.IsEmpty())
+	{
+		return false;
+	}
+	
+	if (InTime.FrameNumber < TimeArray[0])
+	{
+		return false;
+	}	
+	
+	return FMovieSceneBoolChannel::Evaluate(InTime, OutValue);
+}
diff --git a/Engine/Source/Runtime/MovieScene/Public/Channels/MovieSceneConstraintChannel.h b/Engine/Source/Runtime/MovieScene/Public/Channels/MovieSceneConstraintChannel.h
new file mode 100644
--- /dev/null
+++ b/Engine/Source/Runtime/MovieScene/Public/Channels/MovieSceneConstraintChannel.h
@@ -0,0 +1,50 @@
+﻿// Copyright Epic Games, Inc. All Rights Reserved.
+
+#pragma once
+
+#include "Channels/MovieSceneBoolChannel.h"
+#include "MovieSceneClipboard.h"
+
+#include "MovieSceneConstraintChannel.generated.h"
+
+USTRUCT()
+struct MOVIESCENE_API FMovieSceneConstraintChannel : public FMovieSceneBoolChannel
+{
+	GENERATED_BODY()
+
+	FMovieSceneConstraintChannel() {};
+
+	// @todo 
+	bool Evaluate(FFrameTime InTime, bool& OutValue) const;
+};
+
+template<>
+struct TStructOpsTypeTraits<FMovieSceneConstraintChannel> : public TStructOpsTypeTraitsBase2<FMovieSceneConstraintChannel>
+{
+	enum { WithStructuredSerializeFromMismatchedTag = true };
+};
+
+
+template<>
+struct TMovieSceneChannelTraits<FMovieSceneConstraintChannel> : TMovieSceneChannelTraitsBase<FMovieSceneConstraintChannel>
+{
+	enum { SupportsDefaults = false };
+
+#if WITH_EDITOR
+
+	/** Byte channels can have external values (ie, they can get their values from external objects for UI purposes) */
+	typedef TMovieSceneExternalValue<bool> ExtendedEditorDataType;
+
+#endif
+};
+
+// #if WITH_EDITOR
+// namespace MovieSceneClipboard
+// {
+// 	template<> inline FName GetKeyTypeName<bool>()
+// 	{
+// 		static FName Name("Bool");
+// 		return Name;
+// 	}
+// }
+// #endif
