Control Rig: Fix crash when changing  vector variable from single to array

diff --git a/Engine/Plugins/Animation/ControlRig/Source/ControlRigEditor/Private/ControlRigGraphDetails.h b/Engine/Plugins/Animation/ControlRig/Source/ControlRigEditor/Private/ControlRigGraphDetails.h
--- a/Engine/Plugins/Animation/ControlRig/Source/ControlRigEditor/Private/ControlRigGraphDetails.h
+++ b/Engine/Plugins/Animation/ControlRig/Source/ControlRigEditor/Private/ControlRigGraphDetails.h
@@ -266,7 +266,7 @@ public:
 
 protected:
 
-	FORCEINLINE_DEBUGGABLE void GetPropertyChain(TSharedRef<class IPropertyHandle> InPropertyHandle, FEditPropertyChain& OutPropertyChain, TArray<int32> &OutPropertyArrayIndices, bool& bOutEnabled)
+	FORCEINLINE_DEBUGGABLE bool GetPropertyChain(TSharedRef<class IPropertyHandle> InPropertyHandle, FEditPropertyChain& OutPropertyChain, TArray<int32> &OutPropertyArrayIndices, bool& bOutEnabled)
 	{
 		OutPropertyChain.Empty();
 		OutPropertyArrayIndices.Reset();
@@ -276,31 +276,33 @@ protected:
 			if (!ObjectsBeingCustomized[0].IsNull())
 			{
 				TSharedPtr<class IPropertyHandle> ChainHandle = InPropertyHandle;
-				do
+				while (ChainHandle.IsValid() && ChainHandle->GetProperty() != nullptr)
 				{
 					OutPropertyChain.AddHead(ChainHandle->GetProperty());
 					OutPropertyArrayIndices.Insert(ChainHandle->GetIndexInArray(), 0);
-					ChainHandle = ChainHandle->GetParentHandle();
+					ChainHandle = ChainHandle->GetParentHandle();					
+				}
 
-					if(ChainHandle.IsValid())
-					{
-						if(ChainHandle->GetProperty() == nullptr)
-						{
-							break;
-						}
-					}
+				if (OutPropertyChain.GetHead() != nullptr)
+				{
+					OutPropertyChain.SetActiveMemberPropertyNode(OutPropertyChain.GetTail()->GetValue());
+					bOutEnabled = !OutPropertyChain.GetHead()->GetValue()->HasAnyPropertyFlags(CPF_EditConst);
+					return true;
 				}
-				while (ChainHandle.IsValid());
-				OutPropertyChain.SetActiveMemberPropertyNode(OutPropertyChain.GetTail()->GetValue());
-				bOutEnabled = !OutPropertyChain.GetHead()->GetValue()->HasAnyPropertyFlags(CPF_EditConst);
 			}
 		}
+		return false;
 	}
 
 	// extracts the value for a nested property (for Example Settings.WorldTransform) from an outer owner
 	template<typename ValueType>
 	FORCEINLINE_DEBUGGABLE ValueType& ContainerUObjectToValueRef(UObject* InOwner, ValueType& InDefault, FEditPropertyChain& InPropertyChain, TArray<int32> &InPropertyArrayIndices) const
 	{
+		if (InPropertyChain.GetHead() == nullptr)
+		{
+			return InDefault;
+		}
+		
 		FEditPropertyChain::TDoubleLinkedListNode* PropertyNode = InPropertyChain.GetHead();
 		uint8* MemoryPtr = (uint8*)InOwner;
 		int32 ChainIndex = 0;
@@ -347,12 +349,15 @@ protected:
 	template<typename VectorType, typename NumericType>
 	TOptional<NumericType> GetVectorComponent(TSharedRef<class IPropertyHandle> InPropertyHandle, int32 InComponent) 
 	{
+		TOptional<NumericType> Result;
 		FEditPropertyChain PropertyChain;
 		TArray<int32> PropertyArrayIndices;
 		bool bEnabled;
-		GetPropertyChain(InPropertyHandle, PropertyChain, PropertyArrayIndices, bEnabled);
+		if (!GetPropertyChain(InPropertyHandle, PropertyChain, PropertyArrayIndices, bEnabled))
+		{
+			return Result;
+		}
 	
-		TOptional<NumericType> Result;
 		for(TObjectPtr<UObject> Object : ObjectsBeingCustomized)
 		{
 			if(!Object.IsNull() && InPropertyHandle->IsValidHandle())
@@ -383,7 +388,10 @@ protected:
 		FEditPropertyChain PropertyChain;
 		TArray<int32> PropertyArrayIndices;
 		bool bEnabled;
-		GetPropertyChain(InPropertyHandle, PropertyChain, PropertyArrayIndices, bEnabled);
+		if (!GetPropertyChain(InPropertyHandle, PropertyChain, PropertyArrayIndices, bEnabled))
+		{
+			return;
+		}
 	
 		TArray<UObject*> ObjectsView;
 		for(int32 Index = 0; Index < ObjectsBeingCustomized.Num(); Index++)
@@ -444,7 +452,10 @@ protected:
 		FEditPropertyChain PropertyChain;
         TArray<int32> PropertyArrayIndices;
         bool bEnabled;
-        GetPropertyChain(InPropertyHandle, PropertyChain, PropertyArrayIndices, bEnabled);
+        if (!GetPropertyChain(InPropertyHandle, PropertyChain, PropertyArrayIndices, bEnabled))
+        {
+        	return;
+        }
 
 		typename SLocalVectorInputBox::FArguments Args;
 		Args.Font(IDetailLayoutBuilder::GetDetailFont());
@@ -498,12 +509,15 @@ protected:
 	template<typename RotationType>
 	TOptional<RotationType> GetRotation(TSharedRef<class IPropertyHandle> InPropertyHandle)
 	{
+		TOptional<RotationType> Result;
 		FEditPropertyChain PropertyChain;
 		TArray<int32> PropertyArrayIndices;
 		bool bEnabled;
-		GetPropertyChain(InPropertyHandle, PropertyChain, PropertyArrayIndices, bEnabled);
+		if (!GetPropertyChain(InPropertyHandle, PropertyChain, PropertyArrayIndices, bEnabled))
+		{
+			return Result;
+		}
 		
-		TOptional<RotationType> Result;
 		for(TObjectPtr<UObject> Object : ObjectsBeingCustomized)
 		{
 			if(!Object.IsNull() && InPropertyHandle->IsValidHandle())
@@ -533,7 +547,10 @@ protected:
 		FEditPropertyChain PropertyChain;
         TArray<int32> PropertyArrayIndices;
         bool bEnabled;
-        GetPropertyChain(InPropertyHandle, PropertyChain, PropertyArrayIndices, bEnabled);
+        if (!GetPropertyChain(InPropertyHandle, PropertyChain, PropertyArrayIndices, bEnabled))
+        {
+        	return;
+        }
 	
 		TArray<UObject*> ObjectsView;
 		for(int32 Index = 0; Index < ObjectsBeingCustomized.Num(); Index++)
@@ -592,7 +609,10 @@ protected:
 		FEditPropertyChain PropertyChain;
 		TArray<int32> PropertyArrayIndices;
 		bool bEnabled;
-		GetPropertyChain(InPropertyHandle, PropertyChain, PropertyArrayIndices, bEnabled);
+		if (!GetPropertyChain(InPropertyHandle, PropertyChain, PropertyArrayIndices, bEnabled))
+		{
+			return;
+		}
 		
 		typedef typename RotationType::FReal NumericType;
 		typedef SAdvancedRotationInputBox<NumericType> SLocalRotationInputBox;
@@ -627,7 +647,10 @@ protected:
 		FEditPropertyChain PropertyChain;
 		TArray<int32> PropertyArrayIndices;
 		bool bEnabled;
-		GetPropertyChain(InPropertyHandle, PropertyChain, PropertyArrayIndices, bEnabled);
+		if (!GetPropertyChain(InPropertyHandle, PropertyChain, PropertyArrayIndices, bEnabled))
+		{
+			return;
+		}
 		
 		typedef typename TransformType::FReal FReal;
 		typename SAdvancedTransformInputBox<TransformType>::FArguments WidgetArgs;
@@ -643,7 +666,10 @@ protected:
 			FEditPropertyChain PropertyChain;
 			TArray<int32> PropertyArrayIndices;
 			bool bEnabled;
-			GetPropertyChain(InPropertyHandle, PropertyChain, PropertyArrayIndices, bEnabled);
+			if (!GetPropertyChain(InPropertyHandle, PropertyChain, PropertyArrayIndices, bEnabled))
+			{
+				return false;
+			}
 			
 			for(TObjectPtr<UObject> Object : ObjectsBeingCustomized)
 			{
@@ -692,12 +718,15 @@ protected:
 			ESlateRotationRepresentation::Type InRotationRepresentation,
 			ESlateTransformSubComponent::Type InTransformSubComponent) -> TOptional<FReal>
 		{
+			TOptional<FReal> Result;
 			FEditPropertyChain PropertyChain;
 			TArray<int32> PropertyArrayIndices;
 			bool bEnabled;
-			GetPropertyChain(InPropertyHandle, PropertyChain, PropertyArrayIndices, bEnabled);
+			if (!GetPropertyChain(InPropertyHandle, PropertyChain, PropertyArrayIndices, bEnabled))
+			{
+				return Result;
+			}
 			
-			TOptional<FReal> Result;
 			for(TObjectPtr<UObject> Object : ObjectsBeingCustomized)
 			{
 				if(!Object.IsNull() && InPropertyHandle->IsValidHandle())
@@ -742,7 +771,10 @@ protected:
 			FEditPropertyChain PropertyChain;
 			TArray<int32> PropertyArrayIndices;
 			bool bEnabled;
-			GetPropertyChain(InPropertyHandle, PropertyChain, PropertyArrayIndices, bEnabled);
+			if (!GetPropertyChain(InPropertyHandle, PropertyChain, PropertyArrayIndices, bEnabled))
+			{
+				return;
+			}
 			
 			TArray<UObject*> ObjectsView;
 			for(int32 Index = 0; Index < ObjectsBeingCustomized.Num(); Index++)
@@ -819,7 +851,10 @@ protected:
 			FEditPropertyChain PropertyChain;
 			TArray<int32> PropertyArrayIndices;
 			bool bEnabled;
-			GetPropertyChain(InPropertyHandle, PropertyChain, PropertyArrayIndices, bEnabled);
+			if (!GetPropertyChain(InPropertyHandle, PropertyChain, PropertyArrayIndices, bEnabled))
+			{
+				return;
+			}
 			
 			URigVMController* Controller = nullptr;
 			if(BlueprintBeingCustomized && GraphBeingCustomized)
