fixed localization warning

diff --git a/Engine/Plugins/Animation/ControlRig/Source/ControlRigEditor/Private/Editor/Constraints/SConstraintsWidget.cpp b/Engine/Plugins/Animation/ControlRig/Source/ControlRigEditor/Private/Editor/Constraints/SConstraintsWidget.cpp
--- a/Engine/Plugins/Animation/ControlRig/Source/ControlRigEditor/Private/Editor/Constraints/SConstraintsWidget.cpp
+++ b/Engine/Plugins/Animation/ControlRig/Source/ControlRigEditor/Private/Editor/Constraints/SConstraintsWidget.cpp
@@ -815,7 +815,7 @@ TSharedPtr<SWidget> SConstraintsEditionWidget::CreateContextMenu()
 	MenuBuilder.BeginSection("BakeConstraint", LOCTEXT("BakeConstraintHeader", "Bake Constraint"));
 	{
 		MenuBuilder.AddMenuEntry(
-		LOCTEXT("BakeConstraintDoIt", "Bake"),
+		LOCTEXT("BakeConstraintLabel", "Bake"),
 		FText::Format(LOCTEXT("BakeConstraintDoItTooltip", "Bake {0} transforms."), FText::FromName(ListItems[Index]->Name)),
 		FSlateIcon(),
 		FUIAction(FExecuteAction::CreateLambda([Constraint]()
