fix missing HandleOnSpaceAdded delegate

diff --git a/Engine/Plugins/Animation/ControlRig/Source/ControlRigEditor/Private/Sequencer/ControlRigParameterTrackEditor.cpp b/Engine/Plugins/Animation/ControlRig/Source/ControlRigEditor/Private/Sequencer/ControlRigParameterTrackEditor.cpp
--- a/Engine/Plugins/Animation/ControlRig/Source/ControlRigEditor/Private/Sequencer/ControlRigParameterTrackEditor.cpp
+++ b/Engine/Plugins/Animation/ControlRig/Source/ControlRigEditor/Private/Sequencer/ControlRigParameterTrackEditor.cpp
@@ -419,7 +419,7 @@ void FControlRigParameterTrackEditor::BindControlRig(UControlRig* ControlRig)
 					}
 				}
 			}
-			// Track->SpaceChannelAdded().AddRaw(this, &FControlRigParameterTrackEditor::HandleOnSpaceAdded);
+			Track->SpaceChannelAdded().AddRaw(this, &FControlRigParameterTrackEditor::HandleOnSpaceAdded);
 		}
 	}
 }
@@ -2635,9 +2635,6 @@ void FControlRigParameterTrackEditor::GetControlRigKeys(
 					Scale = Val.GetScale3D();
 				}
 
-				// transform these values in the constraint space if needed
-				// NOTE TRS are local to their space 
-					
 				FVector3f CurrentVector = (FVector3f)Translation;
 				bool bKeyX = bSetKey && EnumHasAnyFlags(ChannelsToKey, EControlRigContextChannelToKey::TranslationX);
 				bool bKeyY = bSetKey && EnumHasAnyFlags(ChannelsToKey, EControlRigContextChannelToKey::TranslationY);
