follow up to CL19360468 and CL19387103: update bPropagateToChildren default value for some nodes saved in previous versions

the "Space" element type is changed to "Null" after we apply the PropagateToChildren patch thus leaving old tranform nodes using a bPropagateToChildren to false.

diff --git a/Engine/Plugins/Animation/ControlRig/Source/ControlRigDeveloper/Private/ControlRigBlueprint.cpp b/Engine/Plugins/Animation/ControlRig/Source/ControlRigDeveloper/Private/ControlRigBlueprint.cpp
--- a/Engine/Plugins/Animation/ControlRig/Source/ControlRigDeveloper/Private/ControlRigBlueprint.cpp
+++ b/Engine/Plugins/Animation/ControlRig/Source/ControlRigDeveloper/Private/ControlRigBlueprint.cpp
@@ -4098,7 +4098,7 @@ void UControlRigBlueprint::PatchPropagateToChildren()
 		if (const URigVMPin* TypePin = InPin->FindSubPin(TEXT("Type")))
 		{
 			const FString& TypeValue = TypePin->GetDefaultValue();
-			return TypeValue == TEXT("Null") || TypeValue == TEXT("Control");
+			return TypeValue == TEXT("Null") || TypeValue == TEXT("Space") || TypeValue == TEXT("Control");
 		}
 		
 		return false;
