Sequencer: Control Rig: Need to invalidate hit proxies after gizmo changes. Sequencer does this also with evaluation but it appears with an edit mode that you also need to do it again. To handle cases where it can also change without Sequencer evalaution we check for for transform of the controls changing which will mean the widget will change also.

diff --git a/Engine/Plugins/Animation/ControlRig/Source/ControlRigEditor/Private/EditMode/ControlRigEditMode.cpp b/Engine/Plugins/Animation/ControlRig/Source/ControlRigEditor/Private/EditMode/ControlRigEditMode.cpp
--- a/Engine/Plugins/Animation/ControlRig/Source/ControlRigEditor/Private/EditMode/ControlRigEditMode.cpp
+++ b/Engine/Plugins/Animation/ControlRig/Source/ControlRigEditor/Private/EditMode/ControlRigEditMode.cpp
@@ -2152,6 +2152,56 @@ void FControlRigEditMode::RecalcPivotTransform()
 			}
 		}
 	}
+
+
+	//If in level editor and the transforms changed we need to force hit proxy invalidate so widget hit testing 
+	//doesn't work off of it's last transform.  Similar to what sequencer does on re-evaluation but do to how edit modes and widget ticks happen
+	//it doesn't work for control rig gizmo's
+	if (IsInLevelEditor())
+	{
+		if (HasPivotTransformsChanged())
+		{
+			for (FEditorViewportClient* LevelVC : GEditor->GetLevelViewportClients())
+			{
+				if (LevelVC)
+				{
+					if (!LevelVC->IsRealtime())
+					{
+						LevelVC->RequestRealTimeFrames(1);
+					}
+
+					if (LevelVC->Viewport)
+					{
+						LevelVC->Viewport->InvalidateHitProxy();
+					}
+				}
+			}
+		}
+		LastPivotTransforms = PivotTransforms;
+	}
+}
+
+bool FControlRigEditMode::HasPivotTransformsChanged() const
+{
+	if (PivotTransforms.Num() != LastPivotTransforms.Num())
+	{
+		return true;
+	}
+	for (const TPair<UControlRig*, FTransform>& Transform : PivotTransforms)
+	{
+		if (const FTransform* LastTransform = LastPivotTransforms.Find(Transform.Key))
+		{
+			if (Transform.Value.Equals(*LastTransform, 1e-4f) == false)
+			{
+				return true;
+			}
+		}
+		else
+		{
+			return true;
+		}
+	}
+	return false;
 }
 
 void FControlRigEditMode::HandleSelectionChanged()
diff --git a/Engine/Plugins/Animation/ControlRig/Source/ControlRigEditor/Private/EditMode/ControlRigEditMode.h b/Engine/Plugins/Animation/ControlRig/Source/ControlRigEditor/Private/EditMode/ControlRigEditMode.h
--- a/Engine/Plugins/Animation/ControlRig/Source/ControlRigEditor/Private/EditMode/ControlRigEditMode.h
+++ b/Engine/Plugins/Animation/ControlRig/Source/ControlRigEditor/Private/EditMode/ControlRigEditMode.h
@@ -277,6 +277,9 @@ protected:
 	/** Cached transform of pivot point for selected objects for each Control Rig */
 	TMap<UControlRig*,FTransform> PivotTransforms;
 
+	/** Previous cached transforms, need this to check on tick if any transform changed, gizmo may have changed*/
+	TMap<UControlRig*, FTransform> LastPivotTransforms;
+
 	/** Command bindings for keyboard shortcuts */
 	TSharedPtr<FUICommandList> CommandBindings;
 
@@ -351,7 +354,8 @@ public:
 	/** Suspend Rig Hierarchy Notifies*/
 	void SuspendHierarchyNotifs(bool bVal) { bSuspendHierarchyNotifs = bVal; }
 private:
-
+	/** Whether or not Pivot Transforms have changed, in which case we need to redraw viewport*/
+	bool HasPivotTransformsChanged() const;
 	/** Set a RigElement's selection state */
 	void SetRigElementSelectionInternal(UControlRig* ControlRig, ERigElementType Type, const FName& InRigElementName, bool bSelected);
 	
