Control Rig: More fixes for inversion  name change missing invertible control rigs.

diff --git a/Engine/Plugins/Animation/ControlRig/Source/ControlRigEditor/Private/ControlRigEditorModule.cpp b/Engine/Plugins/Animation/ControlRig/Source/ControlRigEditor/Private/ControlRigEditorModule.cpp
--- a/Engine/Plugins/Animation/ControlRig/Source/ControlRigEditor/Private/ControlRigEditorModule.cpp
+++ b/Engine/Plugins/Animation/ControlRig/Source/ControlRigEditor/Private/ControlRigEditorModule.cpp
@@ -477,12 +477,13 @@ TSharedRef< SWidget > FControlRigEditorModule::GenerateAnimationMenu(TWeakPtr<IA
 											if (Tag.IsSet())
 											{
 												FString EventString = FRigUnit_InverseExecution::EventName.ToString();
+												FString OldEventString = FString(TEXT("Inverse"));
 												TArray<FString> SupportedEventNames;
 												Tag.GetValue().ParseIntoArray(SupportedEventNames, TEXT(","), true);
 
 												for (const FString& Name : SupportedEventNames)
 												{
-													if (Name.Contains(EventString))
+													if (Name.Contains(EventString) || Name.Contains(OldEventString))
 													{
 														bHasInversion = true;
 														break;
diff --git a/Engine/Plugins/Animation/ControlRig/Source/ControlRigEditor/Private/ControlRigSequencerEditorLibrary.cpp b/Engine/Plugins/Animation/ControlRig/Source/ControlRigEditor/Private/ControlRigSequencerEditorLibrary.cpp
--- a/Engine/Plugins/Animation/ControlRig/Source/ControlRigEditor/Private/ControlRigSequencerEditorLibrary.cpp
+++ b/Engine/Plugins/Animation/ControlRig/Source/ControlRigEditor/Private/ControlRigSequencerEditorLibrary.cpp
@@ -713,7 +713,9 @@ bool UControlRigSequencerEditorLibrary::BakeToControlRig(UWorld* World, ULevelSe
 					FString ObjectName = InClass->GetName();
 					ObjectName.RemoveFromEnd(TEXT("_C"));
 					UControlRig* ControlRig = NewObject<UControlRig>(Track, InClass, FName(*ObjectName), RF_Transactional);
-					if (InClass != UFKControlRig::StaticClass() && !ControlRig->SupportsEvent(FRigUnit_InverseExecution::EventName))
+					FName OldEventString = FName(FString(TEXT("Inverse")));
+
+					if (InClass != UFKControlRig::StaticClass() && !(ControlRig->SupportsEvent(FRigUnit_InverseExecution::EventName) || ControlRig->SupportsEvent(OldEventString)))
 					{
 						TempAnimSequence->MarkAsGarbage();
 						MovieScene->RemoveTrack(*Track);
