Control Rig: Construction Event

Changes include
* renamed setup event to construction event
* new nodes to spawn bones, nulls, controls etc
* class settings to limit how many elements can be created procedurally
* lookup from procedural element to spawner node (who created this?)

diff --git a/Engine/Plugins/Animation/ControlRig/Source/ControlRig/Private/AnimNode_ControlRigBase.cpp b/Engine/Plugins/Animation/ControlRig/Source/ControlRig/Private/AnimNode_ControlRigBase.cpp
--- a/Engine/Plugins/Animation/ControlRig/Source/ControlRig/Private/AnimNode_ControlRigBase.cpp
+++ b/Engine/Plugins/Animation/ControlRig/Source/ControlRig/Private/AnimNode_ControlRigBase.cpp
@@ -7,6 +7,7 @@
 #include "Animation/NodeMappingContainer.h"
 #include "AnimationRuntime.h"
 #include "Units/Execution/RigUnit_BeginExecution.h"
+#include "Units/Execution/RigUnit_PrepareForExecution.h"
 #include "Algo/Transform.h"
 
 #if ENABLE_ANIM_DEBUG
@@ -595,6 +596,12 @@ void FAnimNode_ControlRigBase::CacheBones_AnyThread(const FAnimationCacheBonesCo
 
 	if (UControlRig* ControlRig = GetControlRig())
 	{
+		// the construction event may create a set of bones that we can map to. let's run construction now.
+		if(ControlRig->IsConstructionModeEnabled() || ControlRig->IsConstructionRequired())
+		{
+			ControlRig->Execute(EControlRigState::Update, FRigUnit_PrepareForExecution::EventName);
+		}
+			
 		// fill up node names
 		FBoneContainer& RequiredBones = Context.AnimInstanceProxy->GetRequiredBones();
 
diff --git a/Engine/Plugins/Animation/ControlRig/Source/ControlRig/Private/ControlRig.cpp b/Engine/Plugins/Animation/ControlRig/Source/ControlRig/Private/ControlRig.cpp
--- a/Engine/Plugins/Animation/ControlRig/Source/ControlRig/Private/ControlRig.cpp
+++ b/Engine/Plugins/Animation/ControlRig/Source/ControlRig/Private/ControlRig.cpp
@@ -90,6 +90,10 @@ UControlRig::UControlRig(const FObjectInitializer& ObjectInitializer)
 	, VMSnapshotBeforeExecution(nullptr)
 #endif
 	, DebugBoneRadiusMultiplier(1.f)
+#if WITH_EDITOR
+	, bRecordSelectionPoseForConstructionMode(true)
+	, bIsClearingTransientControls(false)
+#endif
 #if UE_CONTROLRIG_PROFILE_EXECUTE_UNITS_NUM
 	, ProfilingRunsLeft(0)
 	, AccumulatedCycles(0)
@@ -265,6 +269,9 @@ void UControlRig::InitializeFromCDO()
 		// copy draw container
 		DrawContainer = CDO->DrawContainer;
 
+		// copy hierarchy settings
+		HierarchySettings = CDO->HierarchySettings;
+
 		// copy vm settings
 		VMRuntimeSettings = CDO->VMRuntimeSettings;
 	}
@@ -766,10 +773,23 @@ bool UControlRig::Execute(const EControlRigState InState, const FName& InEventNa
 
 	// allow access to the hierarchy
 	Context.Hierarchy = GetHierarchy();
+	Context.HierarchySettings = HierarchySettings;
 	check(Context.Hierarchy);
 
-	// for now we'll disable any controller access during the run of a control rig
-	FRigHierarchyEnableControllerBracket DisableHierarchyController(Context.Hierarchy, false);
+	// allow access to the default hierarchy to allow to reset
+	if(!HasAnyFlags(RF_ClassDefaultObject))
+	{
+		if (UControlRig* CDO = Cast<UControlRig>(GetClass()->GetDefaultObject()))
+		{
+			if(URigHierarchy* DefaultHierarchy = CDO->GetHierarchy())
+			{
+				Context.Hierarchy->DefaultHierarchyPtr = DefaultHierarchy;
+			}
+		}
+	}
+
+	// disable any controller access outside of the construction event
+	FRigHierarchyEnableControllerBracket DisableHierarchyController(Context.Hierarchy, bIsConstructionEvent);
 
 	// setup the context with further fields
 	Context.ToWorldSpaceTransform = FTransform::Identity;
@@ -873,71 +893,84 @@ bool UControlRig::Execute(const EControlRigState InState, const FName& InEventNa
 	// we'll special case the construction event here
 	if (bIsConstructionEvent && !bIsInitializingMemory)
 	{
+		// remember the previous selection
+		const TArray<FRigElementKey> PreviousSelection = GetHierarchy()->GetSelectedKeys();
+
 		// construction mode means that we are running the construction event
 		// constantly for testing purposes.
 		const bool bConstructionModeEnabled = IsConstructionModeEnabled();
 		{
-			// save the current state of all pose elements to preserve user intention, since construction event can
-			// run in between forward events
-			// the saved pose is reapplied to the rig after construction event as the pose scope goes out of scope
-			TUniquePtr<UControlRig::FPoseScope> PoseScope;
-			if (!bConstructionModeEnabled)
+			if (PreConstructionForUIEvent.IsBound())
 			{
-				// only do this in non-construction mode because 
-				// when construction mode is enabled, the control values are cleared before reaching here (too late to save them)
-				PoseScope = MakeUnique<UControlRig::FPoseScope>(this, ERigElementType::ToResetAfterConstructionEvent);
+				FControlRigBracketScope BracketScope(PreConstructionBracket);
+				PreConstructionForUIEvent.Broadcast(this, EControlRigState::Update, FRigUnit_PrepareForExecution::EventName);
 			}
 
-			// Copy the hierarchy from the default object onto this one
-			if (UControlRig* CDO = Cast<UControlRig>(GetClass()->GetDefaultObject()))
-			{
-				if(bCopyHierarchyBeforeConstruction && !bConstructionModeEnabled)
-				{
-					if(CDO->GetHierarchy()->GetTopologyVersion()!= GetHierarchy()->GetTopologyVersion())
-					{
 #if WITH_EDITOR
-						FTransientControlScope TransientControlScope(GetHierarchy());
+			// apply the selection pose for the construction mode
+			// we are doing this here to make sure the brackets below have the right data
+			if(bConstructionModeEnabled)
+			{
+				ApplySelectionPoseForConstructionMode(InEventName);
+			}
 #endif
-						GetHierarchy()->CopyHierarchy(CDO->GetHierarchy());
-					}
+
+			// disable selection notifications from the hierarchy
+			TGuardValue<bool> DisableSelectionNotifications(GetHierarchy()->GetController(true)->bSuspendSelectionNotifications, true);
+			{				
+				// save the current state of all pose elements to preserve user intention, since construction event can
+				// run in between forward events
+				// the saved pose is reapplied to the rig after construction event as the pose scope goes out of scope
+				TUniquePtr<FPoseScope> PoseScope;
+				if (!bConstructionModeEnabled)
+				{
+					// only do this in non-construction mode because 
+					// when construction mode is enabled, the control values are cleared before reaching here (too late to save them)
+					PoseScope = MakeUnique<FPoseScope>(this, ERigElementType::ToResetAfterConstructionEvent);
 				}
 				
-				if (bResetInitialTransformsBeforeConstruction && !bConstructionModeEnabled)
 				{
-					// when copying the pose from the default we don't copy the weights over,
-					// since we don't want to reset the topology in case it has diverged from the
-					// CDO.
-					GetHierarchy()->CopyPose(CDO->GetHierarchy(), false, true, false);
-				}
-			}
-
-			{
+					// Copy the hierarchy from the default object onto this one
 #if WITH_EDITOR
-				TUniquePtr<FTransientControlPoseScope> TransientControlPoseScope;
-				if (bConstructionModeEnabled)
-				{
-					// save the transient control value, it should not be constantly reset in construction mode
-					TransientControlPoseScope = MakeUnique<FTransientControlPoseScope>(this);
-				}
-#endif
-				// reset the pose to initial such that construction event can run from a deterministic initial state
-				GetHierarchy()->ResetPoseToInitial(ERigElementType::All);
-			}
+					FTransientControlScope TransientControlScope(GetHierarchy());
+	#endif
+					GetHierarchy()->ResetToDefault();
 
-			if (PreConstructionEvent.IsBound())
-			{
-				FControlRigBracketScope BracketScope(PreConstructionBracket);
-				PreConstructionEvent.Broadcast(this, EControlRigState::Update, FRigUnit_PrepareForExecution::EventName);
-			}
+					{
+	#if WITH_EDITOR
+						TUniquePtr<FTransientControlPoseScope> TransientControlPoseScope;
+						if (bConstructionModeEnabled)
+						{
+							// save the transient control value, it should not be constantly reset in construction mode
+							TransientControlPoseScope = MakeUnique<FTransientControlPoseScope>(this);
+						}
+	#endif
+						// reset the pose to initial such that construction event can run from a deterministic initial state
+						GetHierarchy()->ResetPoseToInitial(ERigElementType::All);
+					}
 
-			bSuccess = ExecuteUnits(Context, FRigUnit_PrepareForExecution::EventName);
+					if (PreConstructionEvent.IsBound())
+					{
+						FControlRigBracketScope BracketScope(PreConstructionBracket);
+						PreConstructionEvent.Broadcast(this, EControlRigState::Update, FRigUnit_PrepareForExecution::EventName);
+					}
 
-			if (PostConstructionEvent.IsBound())
-			{
-				FControlRigBracketScope BracketScope(PostConstructionBracket);
-				PostConstructionEvent.Broadcast(this, EControlRigState::Update, FRigUnit_PrepareForExecution::EventName);
+					bSuccess = ExecuteUnits(Context, FRigUnit_PrepareForExecution::EventName);
+					
+				} // destroy FTransientControlScope
+				
+				if (PostConstructionEvent.IsBound())
+				{
+					FControlRigBracketScope BracketScope(PostConstructionBracket);
+					PostConstructionEvent.Broadcast(this, EControlRigState::Update, FRigUnit_PrepareForExecution::EventName);
+				}
 			}
-		}
+			
+			// set it here to reestablish the selection. the notifications
+			// will be eaten since we still have the bSuspend flag on in the controller.
+			GetHierarchy()->GetController()->SetSelection(PreviousSelection);
+			
+		} // destroy DisableSelectionNotifications
 
 		if (bConstructionModeEnabled)
 		{
@@ -950,7 +983,27 @@ bool UControlRig::Execute(const EControlRigState InState, const FName& InEventNa
 			}
 #endif
 			GetHierarchy()->ResetPoseToInitial(ERigElementType::Bone);
-		}			
+		}
+
+		// synchronize the selection now with the new hierarchy after running construction
+		const TArray<const FRigBaseElement*> CurrentSelection = GetHierarchy()->GetSelectedElements();
+		for(const FRigBaseElement* SelectedElement : CurrentSelection)
+		{
+			if(!PreviousSelection.Contains(SelectedElement->GetKey()))
+			{
+				GetHierarchy()->Notify(ERigHierarchyNotification::ElementSelected, SelectedElement);
+			}
+		}
+		for(const FRigElementKey& PreviouslySelectedKey : PreviousSelection)
+		{
+			if(const FRigBaseElement* PreviouslySelectedElement = GetHierarchy()->Find(PreviouslySelectedKey))
+			{
+				if(!CurrentSelection.Contains(PreviouslySelectedElement))
+				{
+					GetHierarchy()->Notify(ERigHierarchyNotification::ElementDeselected, PreviouslySelectedElement);
+				}
+			}
+		}
 	}
 	else
 	{
@@ -1759,8 +1812,17 @@ FTransform UControlRig::SetupControlFromGlobalTransform(const FName& InControlNa
 		{
 			const FTransform ParentTransform = GetHierarchy()->GetParentTransform(ControlElement, ERigTransformType::CurrentGlobal);
 			const FTransform OffsetTransform = InGlobalTransform.GetRelativeTransform(ParentTransform);
-			GetHierarchy()->SetControlOffsetTransform(ControlElement, OffsetTransform, ERigTransformType::InitialLocal, true, true);
-			GetHierarchy()->SetControlOffsetTransform(ControlElement, OffsetTransform, ERigTransformType::CurrentLocal, true, true);
+			GetHierarchy()->SetControlOffsetTransform(ControlElement, OffsetTransform, ERigTransformType::InitialLocal, true, false);
+			GetHierarchy()->SetControlOffsetTransform(ControlElement, OffsetTransform, ERigTransformType::CurrentLocal, true, false);
+
+			if(URigHierarchy* DefaultHierarchy = GetHierarchy()->GetDefaultHierarchy())
+			{
+				if(FRigControlElement* DefaultControlElement = DefaultHierarchy->Find<FRigControlElement>(ControlElement->GetKey()))
+				{
+					DefaultHierarchy->SetControlOffsetTransform(DefaultControlElement, OffsetTransform, ERigTransformType::InitialLocal, true, true);
+					DefaultHierarchy->SetControlOffsetTransform(DefaultControlElement, OffsetTransform, ERigTransformType::CurrentLocal, true, true);
+				}
+			}
 		}
 	}
 	return InGlobalTransform;
@@ -1849,12 +1911,20 @@ FTransform UControlRig::GetControlGlobalTransform(const FName& InControlName) co
 bool UControlRig::SetControlGlobalTransform(const FName& InControlName, const FTransform& InGlobalTransform, bool bNotify, const FRigControlModifiedContext& Context, bool bSetupUndo, bool bPrintPythonCommands, bool bFixEulerFlips)
 {
 	FTransform GlobalTransform = InGlobalTransform;
+	ERigTransformType::Type TransformType = ERigTransformType::CurrentGlobal;
 	if (IsConstructionModeEnabled())
 	{
+#if WITH_EDITOR
+		if(bRecordSelectionPoseForConstructionMode)
+		{
+			SelectionPoseForConstructionMode.FindOrAdd(FRigElementKey(InControlName, ERigElementType::Control)) = GlobalTransform;
+		}
+#endif
+		TransformType = ERigTransformType::InitialGlobal;
 		GlobalTransform = SetupControlFromGlobalTransform(InControlName, GlobalTransform);
 	}
 
-	FRigControlValue Value = GetControlValueFromGlobalTransform(InControlName, GlobalTransform);
+	FRigControlValue Value = GetControlValueFromGlobalTransform(InControlName, GlobalTransform, TransformType);
 	if (OnFilterControl.IsBound())
 	{
 		FRigControlElement* Control = FindControl(InControlName);
@@ -1868,7 +1938,7 @@ bool UControlRig::SetControlGlobalTransform(const FName& InControlName, const FT
 	return true;
 }
 
-FRigControlValue UControlRig::GetControlValueFromGlobalTransform(const FName& InControlName, const FTransform& InGlobalTransform)
+FRigControlValue UControlRig::GetControlValueFromGlobalTransform(const FName& InControlName, const FTransform& InGlobalTransform, ERigTransformType::Type InTransformType)
 {
 	FRigControlValue Value;
 
@@ -1876,7 +1946,7 @@ FRigControlValue UControlRig::GetControlValueFromGlobalTransform(const FName& In
 	{
 		if(DynamicHierarchy)
 		{
-			FTransform Transform = DynamicHierarchy->ComputeLocalControlValue(ControlElement, InGlobalTransform, ERigTransformType::CurrentGlobal);
+			FTransform Transform = DynamicHierarchy->ComputeLocalControlValue(ControlElement, InGlobalTransform, InTransformType);
 			Value.SetFromTransform(Transform, ControlElement->Settings.ControlType, ControlElement->Settings.PrimaryAxis);
 
 			if (ShouldApplyLimits())
@@ -2003,6 +2073,28 @@ void UControlRig::HandleHierarchyModified(ERigHierarchyNotification InNotificati
 
 				OnControlSelected_BP.Broadcast(this, *ControlElement, bSelected);
 			}
+
+#if WITH_EDITOR
+			if(IsConstructionModeEnabled())
+			{
+				if(InElement->GetType() == ERigElementType::Control)
+				{
+					if(InNotification == ERigHierarchyNotification::ElementSelected)
+					{
+						SelectionPoseForConstructionMode.FindOrAdd(InElement->GetKey()) = GetHierarchy()->GetGlobalTransform(InElement->GetKey());
+					}
+					else
+					{
+						SelectionPoseForConstructionMode.Remove(InElement->GetKey());
+					}
+				}
+
+				if(InNotification == ERigHierarchyNotification::ElementDeselected)
+				{
+					ClearTransientControls();
+				}
+			}
+#endif
 			break;
 		}
 		case ERigHierarchyNotification::ControlSettingChanged:
@@ -2076,11 +2168,6 @@ FName UControlRig::AddTransientControl(URigVMPin* InPin, FRigElementKey SpaceKey
 
 	SetTransientControlValue(InPin);
 
-	if(const FRigBaseElement* Element = DynamicHierarchy->Find(ControlKey))
-	{
-		DynamicHierarchy->Notify(ERigHierarchyNotification::ElementSelected, Element);
-	}
-
 	return ControlName;
 }
 
@@ -2138,6 +2225,9 @@ FName UControlRig::RemoveTransientControl(URigVMPin* InPin)
 		DynamicHierarchy->Notify(ERigHierarchyNotification::ElementDeselected, ControlElement);
 		if(Controller->RemoveElement(ControlElement))
 		{
+#if WITH_EDITOR
+			SelectionPoseForConstructionMode.Remove(FRigElementKey(ControlName, ERigElementType::Control));
+#endif
 			return ControlName;
 		}
 	}
@@ -2203,12 +2293,6 @@ FName UControlRig::AddTransientControl(const FRigElementKey& InElement)
 		}
 	}
 
-	TArray<FRigElementKey> SelectedControls = DynamicHierarchy->GetSelectedKeys(ERigElementType::Control);
-	for(const FRigElementKey& SelectedControl : SelectedControls)
-	{
-		Controller->DeselectElement(SelectedControl);
-	}
-
 	const FRigElementKey ControlKey = Controller->AddControl(
         ControlName,
         Parent,
@@ -2246,11 +2330,6 @@ FName UControlRig::AddTransientControl(const FRigElementKey& InElement)
 
 	SetTransientControlValue(InElement);
 
-	if(const FRigBaseElement* Element = DynamicHierarchy->Find(ControlKey))
-	{
-		DynamicHierarchy->Notify(ERigHierarchyNotification::ElementSelected, Element);
-	}
-
 	return ControlName;
 }
 
@@ -2347,6 +2426,9 @@ FName UControlRig::RemoveTransientControl(const FRigElementKey& InElement)
 		DynamicHierarchy->Notify(ERigHierarchyNotification::ElementDeselected, ControlElement);
 		if(Controller->RemoveElement(ControlElement))
 		{
+#if WITH_EDITOR
+			SelectionPoseForConstructionMode.Remove(FRigElementKey(ControlName, ERigElementType::Control));
+#endif
 			return ControlName;
 		}
 	}
@@ -2430,11 +2512,22 @@ void UControlRig::ClearTransientControls()
 		return;
 	}
 
+	if(bIsClearingTransientControls)
+	{
+		return;
+	}
+	TGuardValue<bool> ReEntryGuard(bIsClearingTransientControls, true);
+
 	const TArray<FRigControlElement*> ControlsToRemove = DynamicHierarchy->GetTransientControls();
 	for (FRigControlElement* ControlToRemove : ControlsToRemove)
 	{
-		DynamicHierarchy->Notify(ERigHierarchyNotification::ElementDeselected, ControlToRemove);
-		Controller->RemoveElement(ControlToRemove);
+		const FRigElementKey KeyToRemove = ControlToRemove->GetKey();
+		if(Controller->RemoveElement(ControlToRemove))
+		{
+#if WITH_EDITOR
+			SelectionPoseForConstructionMode.Remove(KeyToRemove);
+#endif
+		}
 	}
 }
 
@@ -2451,6 +2544,20 @@ void UControlRig::ApplyTransformOverrideForUserCreatedBones()
 	}
 }
 
+void UControlRig::ApplySelectionPoseForConstructionMode(const FName& InEventName)
+{
+	FRigControlModifiedContext ControlValueContext;
+	ControlValueContext.EventName = InEventName;
+
+	TGuardValue<bool> DisableRecording(bRecordSelectionPoseForConstructionMode, false);
+	for(const TPair<FRigElementKey, FTransform>& Pair : SelectionPoseForConstructionMode)
+	{
+		const FName ControlName = Pair.Key.Name;
+		const FRigControlValue Value = GetControlValueFromGlobalTransform(ControlName, Pair.Value, ERigTransformType::InitialGlobal);
+		SetControlValue(ControlName, Value, true, ControlValueContext, false, false, false);
+	}
+}
+
 #endif
 
 void UControlRig::HandleHierarchyEvent(URigHierarchy* InHierarchy, const FRigEventContext& InEvent)
@@ -3251,12 +3358,13 @@ void UControlRig::OnHierarchyTransformUndoRedo(URigHierarchy* InHierarchy, const
 	}
 }
 
-UControlRig::FPoseScope::FPoseScope(UControlRig* InControlRig, ERigElementType InFilter)
+UControlRig::FPoseScope::FPoseScope(UControlRig* InControlRig, ERigElementType InFilter, const TArray<FRigElementKey>& InElements)
 : ControlRig(InControlRig)
 , Filter(InFilter)
 {
 	check(InControlRig);
-	CachedPose = InControlRig->GetHierarchy()->GetPose(false, InFilter, FRigElementKeyCollection());
+	const TArrayView<const FRigElementKey> ElementView(InElements.GetData(), InElements.Num());
+	CachedPose = InControlRig->GetHierarchy()->GetPose(false, InFilter, ElementView);
 }
 
 UControlRig::FPoseScope::~FPoseScope()
diff --git a/Engine/Plugins/Animation/ControlRig/Source/ControlRig/Private/Rigs/RigHierarchy.cpp b/Engine/Plugins/Animation/ControlRig/Source/ControlRig/Private/Rigs/RigHierarchy.cpp
--- a/Engine/Plugins/Animation/ControlRig/Source/ControlRig/Private/Rigs/RigHierarchy.cpp
+++ b/Engine/Plugins/Animation/ControlRig/Source/ControlRig/Private/Rigs/RigHierarchy.cpp
@@ -324,6 +324,19 @@ void URigHierarchy::Reset()
 	}
 }
 
+void URigHierarchy::ResetToDefault()
+{
+	if(DefaultHierarchyPtr.IsValid())
+	{
+		if(URigHierarchy* DefaultHierarchy = DefaultHierarchyPtr.Get())
+		{
+			CopyHierarchy(DefaultHierarchy);
+			return;
+		}
+	}
+	Reset();
+}
+
 void URigHierarchy::CopyHierarchy(URigHierarchy* InHierarchy)
 {
 	check(InHierarchy);
@@ -416,6 +429,55 @@ uint32 URigHierarchy::GetNameHash() const
 	return Hash;
 }
 
+uint32 URigHierarchy::GetTopologyHash(bool bIncludeTopologyVersion, bool bIncludeTransientControls) const
+{
+	uint32 Hash = bIncludeTopologyVersion ? TopologyVersion : 0;
+
+	for (int32 ElementIndex = 0; ElementIndex < Elements.Num(); ElementIndex++)
+	{
+		const FRigBaseElement* Element = Elements[ElementIndex];
+		
+		// skip transient controls
+		if(!bIncludeTransientControls)
+		{
+			if(const FRigControlElement* ControlElement = Cast<FRigControlElement>(Element))
+			{
+				if(ControlElement->Settings.bIsTransientControl)
+				{
+					continue;
+				}
+			}
+		}
+		
+		Hash = HashCombine(Hash, GetTypeHash(Element->GetKey()));
+
+		if(const FRigSingleParentElement* SingleParentElement = Cast<FRigSingleParentElement>(Element))
+		{
+			if(SingleParentElement->ParentElement)
+			{
+				Hash = HashCombine(Hash, GetTypeHash(SingleParentElement->ParentElement->GetKey()));
+			}
+		}
+		if(const FRigMultiParentElement* MultiParentElement = Cast<FRigMultiParentElement>(Element))
+		{
+			for(const FRigElementParentConstraint& ParentConstraint : MultiParentElement->ParentConstraints)
+			{
+				Hash = HashCombine(Hash, GetTypeHash(ParentConstraint.ParentElement->GetKey()));
+			}
+		}
+		if(const FRigBoneElement* BoneElement = Cast<FRigBoneElement>(Element))
+		{
+			Hash = HashCombine(Hash, GetTypeHash(BoneElement->BoneType));
+		}
+		if(const FRigControlElement* ControlElement = Cast<FRigControlElement>(Element))
+		{
+			Hash = HashCombine(Hash, GetTypeHash(ControlElement->Settings));
+		}
+	}
+
+	return Hash;
+}
+
 #if WITH_EDITOR
 void URigHierarchy::RegisterListeningHierarchy(URigHierarchy* InHierarchy)
 {
@@ -714,6 +776,20 @@ int32 URigHierarchy::Num(ERigElementType InElementType) const
 	return ElementsPerType[RigElementTypeToFlatIndex(InElementType)].Num();
 }
 
+bool URigHierarchy::IsProcedural(const FRigElementKey& InKey) const
+{
+	return IsProcedural(Find(InKey));
+}
+
+bool URigHierarchy::IsProcedural(const FRigBaseElement* InElement) const
+{
+	if(InElement == nullptr)
+	{
+		return false;
+	}
+	return InElement->IsProcedural();
+}
+
 TArray<const FRigBaseElement*> URigHierarchy::GetSelectedElements(ERigElementType InTypeFilter) const
 {
 	LLM_SCOPE_BYNAME(TEXT("Animation/ControlRig"));
diff --git a/Engine/Plugins/Animation/ControlRig/Source/ControlRig/Private/Rigs/RigHierarchyController.cpp b/Engine/Plugins/Animation/ControlRig/Source/ControlRig/Private/Rigs/RigHierarchyController.cpp
--- a/Engine/Plugins/Animation/ControlRig/Source/ControlRig/Private/Rigs/RigHierarchyController.cpp
+++ b/Engine/Plugins/Animation/ControlRig/Source/ControlRig/Private/Rigs/RigHierarchyController.cpp
@@ -192,7 +192,7 @@ FRigElementKey URigHierarchyController::AddBone(FName InName, FRigElementKey InP
 	}
 #endif
 
-	FRigBoneElement* NewElement = URigHierarchy::NewElement<FRigBoneElement>();
+	FRigBoneElement* NewElement = MakeElement<FRigBoneElement>();
 	{
 		TGuardValue<bool> DisableCacheValidityChecks(Hierarchy->bEnableCacheValidityCheck, false);
 		NewElement->Key.Type = ERigElementType::Bone;
@@ -254,7 +254,7 @@ FRigElementKey URigHierarchyController::AddNull(FName InName, FRigElementKey InP
 	}
 #endif
 
-	FRigNullElement* NewElement = URigHierarchy::NewElement<FRigNullElement>();
+	FRigNullElement* NewElement = MakeElement<FRigNullElement>();
 	{
 		TGuardValue<bool> DisableCacheValidityChecks(Hierarchy->bEnableCacheValidityCheck, false);		
 		NewElement->Key.Type = ERigElementType::Null;
@@ -319,7 +319,7 @@ FRigElementKey URigHierarchyController::AddControl(
 	}
 #endif
 
-	FRigControlElement* NewElement = URigHierarchy::NewElement<FRigControlElement>();
+	FRigControlElement* NewElement = MakeElement<FRigControlElement>();
 	{
 		TGuardValue<bool> DisableCacheValidityChecks(Hierarchy->bEnableCacheValidityCheck, false);		
 		NewElement->Key.Type = ERigElementType::Control;
@@ -407,7 +407,7 @@ FRigElementKey URigHierarchyController::AddCurve(FName InName, float InValue, bo
 	}
 #endif
 
-	FRigCurveElement* NewElement = URigHierarchy::NewElement<FRigCurveElement>();
+	FRigCurveElement* NewElement = MakeElement<FRigCurveElement>();
 	{
 		TGuardValue<bool> DisableCacheValidityChecks(Hierarchy->bEnableCacheValidityCheck, false);		
 		NewElement->Key.Type = ERigElementType::Curve;
@@ -456,7 +456,7 @@ FRigElementKey URigHierarchyController::AddRigidBody(FName InName, FRigElementKe
 	}
 #endif
 
-	FRigRigidBodyElement* NewElement = URigHierarchy::NewElement<FRigRigidBodyElement>();
+	FRigRigidBodyElement* NewElement = MakeElement<FRigRigidBodyElement>();
 	{
 		TGuardValue<bool> DisableCacheValidityChecks(Hierarchy->bEnableCacheValidityCheck, false);		
 		NewElement->Key.Type = ERigElementType::RigidBody;
@@ -508,7 +508,7 @@ FRigElementKey URigHierarchyController::AddReference(FName InName, FRigElementKe
 	}
 #endif
 
-	FRigReferenceElement* NewElement = URigHierarchy::NewElement<FRigReferenceElement>();
+	FRigReferenceElement* NewElement = MakeElement<FRigReferenceElement>();
 	{
 		TGuardValue<bool> DisableCacheValidityChecks(Hierarchy->bEnableCacheValidityCheck, false);		
 		NewElement->Key.Type = ERigElementType::Reference;
@@ -1093,38 +1093,38 @@ TArray<FRigElementKey> URigHierarchyController::ImportFromText(FString InContent
 		{
 			case ERigElementType::Bone:
 			{
-				NewElement = URigHierarchy::NewElement<FRigBoneElement>();
+				NewElement = MakeElement<FRigBoneElement>();
 				FRigBoneElement::StaticStruct()->ImportText(*PerElementData.Content, NewElement, nullptr, EPropertyPortFlags::PPF_None, &ErrorPipe, FRigBoneElement::StaticStruct()->GetName(), true);
 				CastChecked<FRigBoneElement>(NewElement)->BoneType = ERigBoneType::User;
 				break;
 			}
 			case ERigElementType::Null:
 			{
-				NewElement = URigHierarchy::NewElement<FRigNullElement>();
+				NewElement = MakeElement<FRigNullElement>();
 				FRigNullElement::StaticStruct()->ImportText(*PerElementData.Content, NewElement, nullptr, EPropertyPortFlags::PPF_None, &ErrorPipe, FRigNullElement::StaticStruct()->GetName(), true);
 				break;
 			}
 			case ERigElementType::Control:
 			{
-				NewElement = URigHierarchy::NewElement<FRigControlElement>();
+				NewElement = MakeElement<FRigControlElement>();
 				FRigControlElement::StaticStruct()->ImportText(*PerElementData.Content, NewElement, nullptr, EPropertyPortFlags::PPF_None, &ErrorPipe, FRigControlElement::StaticStruct()->GetName(), true);
 				break;
 			}
 			case ERigElementType::Curve:
 			{
-				NewElement = URigHierarchy::NewElement<FRigCurveElement>();
+				NewElement = MakeElement<FRigCurveElement>();
 				FRigCurveElement::StaticStruct()->ImportText(*PerElementData.Content, NewElement, nullptr, EPropertyPortFlags::PPF_None, &ErrorPipe, FRigCurveElement::StaticStruct()->GetName(), true);
 				break;
 			}
 			case ERigElementType::RigidBody:
 			{
-				NewElement = URigHierarchy::NewElement<FRigRigidBodyElement>();
+				NewElement = MakeElement<FRigRigidBodyElement>();
 				FRigRigidBodyElement::StaticStruct()->ImportText(*PerElementData.Content, NewElement, nullptr, EPropertyPortFlags::PPF_None, &ErrorPipe, FRigRigidBodyElement::StaticStruct()->GetName(), true);
 				break;
 			}
 			case ERigElementType::Reference:
 			{
-				NewElement = URigHierarchy::NewElement<FRigReferenceElement>();
+				NewElement = MakeElement<FRigReferenceElement>();
 				FRigReferenceElement::StaticStruct()->ImportText(*PerElementData.Content, NewElement, nullptr, EPropertyPortFlags::PPF_None, &ErrorPipe, FRigReferenceElement::StaticStruct()->GetName(), true);
 				break;
 			}
@@ -1598,16 +1598,24 @@ void URigHierarchyController::Notify(ERigHierarchyNotification InNotifType, cons
 	{
 		return;
 	}
-	if(bSuspendNotifications)
+	if(bSuspendAllNotifications)
 	{
 		return;
 	}
+	if(bSuspendSelectionNotifications)
+	{
+		if(InNotifType == ERigHierarchyNotification::ElementSelected ||
+			InNotifType == ERigHierarchyNotification::ElementDeselected)
+		{
+			return;
+		}
+	}	
 	Hierarchy->Notify(InNotifType, InElement);
 }
 
 void URigHierarchyController::HandleHierarchyModified(ERigHierarchyNotification InNotifType, URigHierarchy* InHierarchy, const FRigBaseElement* InElement) const
 {
-	if(bSuspendNotifications)
+	if(bSuspendAllNotifications)
 	{
 		return;
 	}
@@ -1629,7 +1637,7 @@ int32 URigHierarchyController::AddElement(FRigBaseElement* InElementToAdd, FRigB
 	Hierarchy->IncrementTopologyVersion();
 
 	{
-		const TGuardValue<bool> Guard(bSuspendNotifications, true);
+		const TGuardValue<bool> Guard(bSuspendAllNotifications, true);
 		SetParent(InElementToAdd, InFirstParent, bMaintainGlobalTransform);
 	}
 
diff --git a/Engine/Plugins/Animation/ControlRig/Source/ControlRig/Private/Rigs/RigHierarchyElements.cpp b/Engine/Plugins/Animation/ControlRig/Source/ControlRig/Private/Rigs/RigHierarchyElements.cpp
--- a/Engine/Plugins/Animation/ControlRig/Source/ControlRig/Private/Rigs/RigHierarchyElements.cpp
+++ b/Engine/Plugins/Animation/ControlRig/Source/ControlRig/Private/Rigs/RigHierarchyElements.cpp
@@ -215,6 +215,40 @@ void FRigPreferredEulerAngles::SetAngles(const FVector& InValue, bool bInitial,
 	Get(bInitial) = Value;
 }
 
+////////////////////////////////////////////////////////////////////////////////
+// FRigElementHandle
+////////////////////////////////////////////////////////////////////////////////
+
+FRigElementHandle::FRigElementHandle(URigHierarchy* InHierarchy, const FRigElementKey& InKey)
+: Hierarchy(InHierarchy)
+, Key(InKey)
+{
+}
+
+FRigElementHandle::FRigElementHandle(URigHierarchy* InHierarchy, const FRigBaseElement* InElement)
+: Hierarchy(InHierarchy)
+, Key(InElement->GetKey())
+{
+}
+
+const FRigBaseElement* FRigElementHandle::Get() const
+{
+	if(Hierarchy.IsValid())
+	{
+		return Hierarchy->Find(Key);
+	}
+	return nullptr;
+}
+
+FRigBaseElement* FRigElementHandle::Get()
+{
+	if(Hierarchy.IsValid())
+	{
+		return Hierarchy->Find(Key);
+	}
+	return nullptr;
+}
+
 ////////////////////////////////////////////////////////////////////////////////
 // FRigTransformElement
 ////////////////////////////////////////////////////////////////////////////////
@@ -704,6 +738,24 @@ void FRigControlSettings::Load(FArchive& Ar)
 	}
 }
 
+uint32 GetTypeHash(const FRigControlSettings& Settings)
+{
+	uint32 Hash = GetTypeHash(Settings.ControlType);
+	Hash = HashCombine(Hash, GetTypeHash(Settings.AnimationType));
+	Hash = HashCombine(Hash, GetTypeHash(Settings.DisplayName));
+	Hash = HashCombine(Hash, GetTypeHash(Settings.PrimaryAxis));
+	Hash = HashCombine(Hash, GetTypeHash(Settings.bIsCurve));
+	Hash = HashCombine(Hash, GetTypeHash(Settings.bDrawLimits));
+	Hash = HashCombine(Hash, GetTypeHash(Settings.bShapeVisible));
+	Hash = HashCombine(Hash, GetTypeHash(Settings.ShapeVisibility));
+	Hash = HashCombine(Hash, GetTypeHash(Settings.ShapeName));
+	Hash = HashCombine(Hash, GetTypeHash(Settings.ShapeColor));
+	Hash = HashCombine(Hash, GetTypeHash(Settings.ControlEnum));
+	Hash = HashCombine(Hash, GetTypeHash(Settings.DrivenControls));
+	Hash = HashCombine(Hash, GetTypeHash(Settings.bGroupWithParentControl));
+	return Hash;
+}
+
 ////////////////////////////////////////////////////////////////////////////////
 // FRigControlElement
 ////////////////////////////////////////////////////////////////////////////////
diff --git a/Engine/Plugins/Animation/ControlRig/Source/ControlRig/Private/Settings/ControlRigSettings.cpp b/Engine/Plugins/Animation/ControlRig/Source/ControlRig/Private/Settings/ControlRigSettings.cpp
--- a/Engine/Plugins/Animation/ControlRig/Source/ControlRig/Private/Settings/ControlRigSettings.cpp
+++ b/Engine/Plugins/Animation/ControlRig/Source/ControlRig/Private/Settings/ControlRigSettings.cpp
@@ -30,7 +30,6 @@ UControlRigEditorSettings::UControlRigEditorSettings(const FObjectInitializer& O
 	ConstructionEventBorderColor = FLinearColor::Red;
 	BackwardsSolveBorderColor = FLinearColor::Yellow;
 	BackwardsAndForwardsBorderColor = FLinearColor::Blue;
-	DynamicHierarchyBorderColor = FLinearColor::Blue;
 	bShowStackedHierarchy = false;
 	MaxStackSize = 16;
 #endif
diff --git a/Engine/Plugins/Animation/ControlRig/Source/ControlRig/Private/Units/Execution/RigUnit_Collection.cpp b/Engine/Plugins/Animation/ControlRig/Source/ControlRig/Private/Units/Execution/RigUnit_Collection.cpp
--- a/Engine/Plugins/Animation/ControlRig/Source/ControlRig/Private/Units/Execution/RigUnit_Collection.cpp
+++ b/Engine/Plugins/Animation/ControlRig/Source/ControlRig/Private/Units/Execution/RigUnit_Collection.cpp
@@ -19,6 +19,7 @@
 #define FRigUnit_CollectionReverse_Hash 9
 #define FRigUnit_CollectionCount_Hash 10
 #define FRigUnit_CollectionItemAtIndex_Hash 11
+#define FRigUnit_CollectionGetAll_Hash 12
 
 FRigUnit_CollectionChain_Execute()
 {
@@ -158,6 +159,41 @@ FRigUnit_CollectionChildrenArray_Execute()
 	Items = Collection.Keys;
 }
 
+FRigUnit_CollectionGetAll_Execute()
+{
+	DECLARE_SCOPE_HIERARCHICAL_COUNTER_RIGUNIT()
+
+	uint32 Hash = FRigUnit_CollectionGetAll_Hash + Context.Hierarchy->GetTopologyVersion() * 17;
+	Hash = HashCombine(Hash, (int32)TypeToSearch * 8);
+
+	FRigElementKeyCollection Collection;
+	if(const FRigElementKeyCollection* Cache = Context.Hierarchy->FindCachedCollection(Hash))
+	{
+		Collection = *Cache;
+	}
+	else
+	{
+		Context.Hierarchy->Traverse([&Collection, TypeToSearch](FRigBaseElement* InElement, bool &bContinue)
+		{
+			bContinue = true;
+			
+			const FRigElementKey Key = InElement->GetKey();
+			if(((uint8)TypeToSearch & (uint8)Key.Type) == (uint8)Key.Type)
+			{
+				Collection.AddUnique(Key);
+			}
+		});
+		
+		if (Collection.IsEmpty())
+		{
+			UE_CONTROLRIG_RIGUNIT_REPORT_WARNING(TEXT("%s"), TEXT("No elements found for given filter."));
+		}
+		Context.Hierarchy->AddCachedCollection(Hash, Collection);
+	}
+
+	Items = Collection.Keys;
+}
+
 #if WITH_EDITOR
 
 IMPLEMENT_RIGUNIT_AUTOMATION_TEST(FRigUnit_CollectionChildren)
diff --git a/Engine/Plugins/Animation/ControlRig/Source/ControlRig/Private/Units/Execution/RigUnit_Collection.h b/Engine/Plugins/Animation/ControlRig/Source/ControlRig/Private/Units/Execution/RigUnit_Collection.h
--- a/Engine/Plugins/Animation/ControlRig/Source/ControlRig/Private/Units/Execution/RigUnit_Collection.h
+++ b/Engine/Plugins/Animation/ControlRig/Source/ControlRig/Private/Units/Execution/RigUnit_Collection.h
@@ -217,6 +217,29 @@ struct CONTROLRIG_API FRigUnit_CollectionChildrenArray : public FRigUnit_Collect
 	TArray<FRigElementKey> Items;
 };
 
+/**
+* Creates an item array for all elements given the filter.
+*/
+USTRUCT(meta = (DisplayName = "Get All", Category = "Hierarchy", Keywords = "Bone,Joint,Collection,Filter,Parent", Varying))
+struct CONTROLRIG_API FRigUnit_CollectionGetAll : public FRigUnit_CollectionBase
+{
+	GENERATED_BODY()
+
+	FRigUnit_CollectionGetAll()
+	{
+		TypeToSearch = ERigElementType::All;
+	}
+
+	RIGVM_METHOD()
+	virtual void Execute(const FRigUnitContext& Context) override;
+
+	UPROPERTY(meta = (Input))
+	ERigElementType TypeToSearch;
+
+	UPROPERTY(meta = (Output))
+	TArray<FRigElementKey> Items;
+};
+
 /**
  * Replaces all names within the collection
  */
diff --git a/Engine/Plugins/Animation/ControlRig/Source/ControlRig/Private/Units/Execution/RigUnit_DynamicHierarchy.cpp b/Engine/Plugins/Animation/ControlRig/Source/ControlRig/Private/Units/Execution/RigUnit_DynamicHierarchy.cpp
--- a/Engine/Plugins/Animation/ControlRig/Source/ControlRig/Private/Units/Execution/RigUnit_DynamicHierarchy.cpp
+++ b/Engine/Plugins/Animation/ControlRig/Source/ControlRig/Private/Units/Execution/RigUnit_DynamicHierarchy.cpp
@@ -3,15 +3,56 @@
 #include "RigUnit_DynamicHierarchy.h"
 #include "Rigs/RigHierarchyController.h"
 #include "Units/RigUnitContext.h"
+#include "ControlRig.h"
+#include "Components/SkeletalMeshComponent.h"
 #include "ControlRig/Private/Units/Execution/RigUnit_PrepareForExecution.h"
 
+bool FRigUnit_DynamicHierarchyBase::IsValidToRunInContext(
+	const FRigUnitContext& InContext,
+	const FControlRigExecuteContext& InExecuteContext,
+	bool bAllowOnlyConstructionEvent,
+	FString* OutErrorMessage)
+{
+	if((InContext.State != EControlRigState::Update) ||
+		(InContext.Hierarchy == nullptr) ||
+		(InExecuteContext.Hierarchy == nullptr))
+	{
+		return false;
+	}
+
+	if(bAllowOnlyConstructionEvent)
+	{
+		if(InExecuteContext.GetEventName() != FRigUnit_PrepareForExecution::EventName)
+		{
+			if(OutErrorMessage)
+			{
+				static constexpr TCHAR ErrorMessageFormat[] = TEXT("Node can only run in %s Event");
+				*OutErrorMessage = FString::Printf(ErrorMessageFormat, *FRigUnit_PrepareForExecution::EventName.ToString());
+			}
+			return false;
+		}
+	}
+
+	if(InContext.Hierarchy->Num() >= InContext.HierarchySettings.ProceduralElementLimit)
+	{
+		if(OutErrorMessage)
+		{
+			static constexpr TCHAR ErrorMessageFormat[] = TEXT("Node has hit the Procedural Element Limit. Check the Class Settings under Hierarchy.");
+			*OutErrorMessage = FString::Printf(ErrorMessageFormat, *FRigUnit_PrepareForExecution::EventName.ToString());
+		}
+		return false;
+	}
+
+	return true;
+}
+
 FRigUnit_AddParent_Execute()
 {
-	if((Context.State != EControlRigState::Update) || (ExecuteContext.Hierarchy == nullptr))
+	if(!FRigUnit_DynamicHierarchyBase::IsValidToRunInContext(Context, ExecuteContext, true))
 	{
 		return;
 	}
-
+	
 	FRigTransformElement* ChildElement = ExecuteContext.Hierarchy->Find<FRigTransformElement>(Child);
 	if(ChildElement == nullptr)
 	{
@@ -35,7 +76,7 @@ FRigUnit_AddParent_Execute()
 
 FRigUnit_SwitchParent_Execute()
 {
-	if((Context.State != EControlRigState::Update) || (ExecuteContext.Hierarchy == nullptr))
+	if(!FRigUnit_DynamicHierarchyBase::IsValidToRunInContext(Context, ExecuteContext, false))
 	{
 		return;
 	}
@@ -161,7 +202,7 @@ FRigUnit_HierarchyGetParentWeightsArray_Execute()
 
 FRigUnit_HierarchySetParentWeights_Execute()
 {
-	if((Context.State != EControlRigState::Update) || (ExecuteContext.Hierarchy == nullptr))
+	if(!FRigUnit_DynamicHierarchyBase::IsValidToRunInContext(Context, ExecuteContext, false))
 	{
 		return;
 	}
@@ -187,3 +228,606 @@ FRigUnit_HierarchySetParentWeights_Execute()
 		ExecuteContext.Hierarchy->SetParentWeightArray(ChildElement, Weights, true, true);
 	}
 }
+
+FRigUnit_HierarchyReset_Execute()
+{
+	FString ErrorMessage;
+	if(!FRigUnit_DynamicHierarchyBase::IsValidToRunInContext(Context, ExecuteContext, true, &ErrorMessage))
+	{
+		if(!ErrorMessage.IsEmpty())
+		{
+			UE_CONTROLRIG_RIGUNIT_REPORT_ERROR(TEXT("%s"), *ErrorMessage);
+		}
+		return;
+	}
+	ExecuteContext.Hierarchy->ResetToDefault();
+}
+
+FRigUnit_HierarchyImportFromSkeleton_Execute()
+{
+	FString ErrorMessage;
+	if(!FRigUnit_DynamicHierarchyBase::IsValidToRunInContext(Context, ExecuteContext, true, &ErrorMessage))
+	{
+		if(!ErrorMessage.IsEmpty())
+		{
+			UE_CONTROLRIG_RIGUNIT_REPORT_ERROR(TEXT("%s"), *ErrorMessage);
+		}
+		return;
+	}
+
+	Items.Reset();
+
+	if(URigHierarchyController* Controller = ExecuteContext.Hierarchy->GetController(true))
+	{
+		FRigHierarchyControllerInstructionBracket InstructionBracket(Controller, RigVMExecuteContext.GetInstructionIndex());
+		
+		if(const USkeletalMeshComponent* SkelMeshComponent = Context.DataSourceRegistry->RequestSource<USkeletalMeshComponent>(UControlRig::OwnerComponent))
+		{
+			if(SkelMeshComponent->SkeletalMesh)
+			{
+				const FReferenceSkeleton& ReferenceSkeleton = SkelMeshComponent->SkeletalMesh->GetRefSkeleton();
+				Items = Controller->ImportBones(ReferenceSkeleton, NameSpace, false, false, false, false);
+				if(bIncludeCurves)
+				{
+					if(USkeleton* Skeleton = SkelMeshComponent->SkeletalMesh->GetSkeleton())
+					{
+						Items.Append(Controller->ImportCurves(Skeleton, NameSpace, false, false, false));
+					}
+				}
+			}
+		}
+	}
+}
+
+FRigUnit_HierarchyRemoveElement_Execute()
+{
+	FString ErrorMessage;
+	if(!FRigUnit_DynamicHierarchyBase::IsValidToRunInContext(Context, ExecuteContext, true, &ErrorMessage))
+	{
+		if(!ErrorMessage.IsEmpty())
+		{
+			UE_CONTROLRIG_RIGUNIT_REPORT_ERROR(TEXT("%s"), *ErrorMessage);
+		}
+		return;
+	}
+	
+	bSuccess = false;
+
+	if(URigHierarchyController* Controller = ExecuteContext.Hierarchy->GetController(true))
+	{
+		FRigHierarchyControllerInstructionBracket InstructionBracket(Controller, RigVMExecuteContext.GetInstructionIndex());
+		bSuccess = Controller->RemoveElement(Item);
+	}
+}
+
+FRigUnit_HierarchyAddBone_Execute()
+{
+	FString ErrorMessage;
+	if(!FRigUnit_DynamicHierarchyBase::IsValidToRunInContext(Context, ExecuteContext, true, &ErrorMessage))
+	{
+		if(!ErrorMessage.IsEmpty())
+		{
+			UE_CONTROLRIG_RIGUNIT_REPORT_ERROR(TEXT("%s"), *ErrorMessage);
+		}
+		return;
+	}
+
+	Item.Reset();
+
+	if(URigHierarchyController* Controller = ExecuteContext.Hierarchy->GetController(true))
+	{
+		FRigHierarchyControllerInstructionBracket InstructionBracket(Controller, RigVMExecuteContext.GetInstructionIndex());
+		Item = Controller->AddBone(Name, Parent, Transform, Space == EBoneGetterSetterMode::GlobalSpace, ERigBoneType::Imported, false, false);
+	}
+}
+
+FRigUnit_HierarchyAddNull_Execute()
+{
+	FString ErrorMessage;
+	if(!FRigUnit_DynamicHierarchyBase::IsValidToRunInContext(Context, ExecuteContext, true, &ErrorMessage))
+	{
+		if(!ErrorMessage.IsEmpty())
+		{
+			UE_CONTROLRIG_RIGUNIT_REPORT_ERROR(TEXT("%s"), *ErrorMessage);
+		}
+		return;
+	}
+
+	Item.Reset();
+
+	if(URigHierarchyController* Controller = ExecuteContext.Hierarchy->GetController(true))
+	{
+		FRigHierarchyControllerInstructionBracket InstructionBracket(Controller, RigVMExecuteContext.GetInstructionIndex());
+		Item = Controller->AddNull(Name, Parent, Transform, Space == EBoneGetterSetterMode::GlobalSpace, false, false);
+	}
+}
+
+void FRigUnit_HierarchyAddControl_Settings::Configure(FRigControlSettings& OutSettings) const
+{
+	OutSettings.DisplayName = DisplayName;
+}
+
+void FRigUnit_HierarchyAddControl_ShapeSettings::Configure(FRigControlSettings& OutSettings) const
+{
+	OutSettings.bShapeVisible = bVisible;
+	OutSettings.ShapeName = Name;
+	OutSettings.ShapeColor = Color;
+}
+
+void FRigUnit_HierarchyAddControl_ProxySettings::Configure(FRigControlSettings& OutSettings) const
+{
+	OutSettings.AnimationType = bIsProxy ? ERigControlAnimationType::ProxyControl : ERigControlAnimationType::AnimationControl;
+	OutSettings.DrivenControls = DrivenControls;
+	OutSettings.ShapeVisibility = ShapeVisibility;
+}
+
+void FRigUnit_HierarchyAddControlFloat_LimitSettings::Configure(FRigControlSettings& OutSettings) const
+{
+	OutSettings.SetupLimitArrayForType(false, false, false);
+	OutSettings.LimitEnabled[0] = Limit;
+	OutSettings.MinimumValue = FRigControlValue::Make<float>(MinValue);
+	OutSettings.MaximumValue = FRigControlValue::Make<float>(MaxValue);
+	OutSettings.bDrawLimits = bDrawLimits;
+}
+
+void FRigUnit_HierarchyAddControlFloat_Settings::Configure(FRigControlSettings& OutSettings) const
+{
+	Super::Configure(OutSettings);
+	
+	OutSettings.ControlType = ERigControlType::Float;
+	OutSettings.PrimaryAxis = PrimaryAxis;
+
+	Proxy.Configure(OutSettings);
+	Limits.Configure(OutSettings);
+	Shape.Configure(OutSettings);
+}
+
+FRigUnit_HierarchyAddControlFloat_Execute()
+{
+	FString ErrorMessage;
+	if(!FRigUnit_DynamicHierarchyBase::IsValidToRunInContext(Context, ExecuteContext, true, &ErrorMessage))
+	{
+		if(!ErrorMessage.IsEmpty())
+		{
+			UE_CONTROLRIG_RIGUNIT_REPORT_ERROR(TEXT("%s"), *ErrorMessage);
+		}
+		return;
+	}
+
+	Item.Reset();
+
+	if(URigHierarchyController* Controller = ExecuteContext.Hierarchy->GetController(true))
+	{
+		FRigControlSettings ControlSettings;
+		Settings.Configure(ControlSettings);
+
+		const FRigControlValue Value = FRigControlValue::Make<float>(InitialValue);
+		const FTransform ShapeTransform = Settings.Shape.Transform;
+		
+		FRigHierarchyControllerInstructionBracket InstructionBracket(Controller, RigVMExecuteContext.GetInstructionIndex());
+		Item = Controller->AddControl(Name, Parent, ControlSettings, Value, OffsetTransform, ShapeTransform, false, false);
+	}
+}
+
+void FRigUnit_HierarchyAddControlInteger_LimitSettings::Configure(FRigControlSettings& OutSettings) const
+{
+	OutSettings.SetupLimitArrayForType(false, false, false);
+	OutSettings.LimitEnabled[0] = Limit;
+	OutSettings.MinimumValue = FRigControlValue::Make<int32>(MinValue);
+	OutSettings.MaximumValue = FRigControlValue::Make<int32>(MaxValue);
+	OutSettings.bDrawLimits = bDrawLimits;
+}
+
+void FRigUnit_HierarchyAddControlInteger_Settings::Configure(FRigControlSettings& OutSettings) const
+{
+	Super::Configure(OutSettings);
+	
+	OutSettings.ControlType = ERigControlType::Integer;
+	OutSettings.PrimaryAxis = PrimaryAxis;
+
+	Proxy.Configure(OutSettings);
+	Limits.Configure(OutSettings);
+	Shape.Configure(OutSettings);
+}
+
+FRigUnit_HierarchyAddControlInteger_Execute()
+{
+	FString ErrorMessage;
+	if(!FRigUnit_DynamicHierarchyBase::IsValidToRunInContext(Context, ExecuteContext, true, &ErrorMessage))
+	{
+		if(!ErrorMessage.IsEmpty())
+		{
+			UE_CONTROLRIG_RIGUNIT_REPORT_ERROR(TEXT("%s"), *ErrorMessage);
+		}
+		return;
+	}
+
+	Item.Reset();
+
+	if(URigHierarchyController* Controller = ExecuteContext.Hierarchy->GetController(true))
+	{
+		FRigControlSettings ControlSettings;
+		Settings.Configure(ControlSettings);
+
+		const FRigControlValue Value = FRigControlValue::Make<int32>(InitialValue);
+		const FTransform ShapeTransform = Settings.Shape.Transform;
+		
+		FRigHierarchyControllerInstructionBracket InstructionBracket(Controller, RigVMExecuteContext.GetInstructionIndex());
+		Item = Controller->AddControl(Name, Parent, ControlSettings, Value, OffsetTransform, ShapeTransform, false, false);
+	}
+}
+
+void FRigUnit_HierarchyAddControlVector2D_LimitSettings::Configure(FRigControlSettings& OutSettings) const
+{
+	OutSettings.SetupLimitArrayForType(false, false, false);
+	OutSettings.LimitEnabled[0] = LimitX;
+	OutSettings.LimitEnabled[1] = LimitY;
+	OutSettings.MinimumValue = FRigControlValue::Make<FVector3f>(FVector3f(MinValue.X, MinValue.Y, 0.f));
+	OutSettings.MaximumValue = FRigControlValue::Make<FVector3f>(FVector3f(MaxValue.X, MaxValue.Y, 0.f));
+	OutSettings.bDrawLimits = bDrawLimits;
+}
+
+void FRigUnit_HierarchyAddControlVector2D_Settings::Configure(FRigControlSettings& OutSettings) const
+{
+	Super::Configure(OutSettings);
+	
+	OutSettings.ControlType = ERigControlType::Vector2D;
+	OutSettings.PrimaryAxis = PrimaryAxis;
+
+	Proxy.Configure(OutSettings);
+	Limits.Configure(OutSettings);
+	Shape.Configure(OutSettings);
+}
+
+FRigUnit_HierarchyAddControlVector2D_Execute()
+{
+	FString ErrorMessage;
+	if(!FRigUnit_DynamicHierarchyBase::IsValidToRunInContext(Context, ExecuteContext, true, &ErrorMessage))
+	{
+		if(!ErrorMessage.IsEmpty())
+		{
+			UE_CONTROLRIG_RIGUNIT_REPORT_ERROR(TEXT("%s"), *ErrorMessage);
+		}
+		return;
+	}
+
+	Item.Reset();
+
+	if(URigHierarchyController* Controller = ExecuteContext.Hierarchy->GetController(true))
+	{
+		FRigControlSettings ControlSettings;
+		Settings.Configure(ControlSettings);
+
+		const FRigControlValue Value = FRigControlValue::Make<FVector3f>(FVector3f(InitialValue.X, InitialValue.Y, 0.f));
+		const FTransform ShapeTransform = Settings.Shape.Transform;
+		
+		FRigHierarchyControllerInstructionBracket InstructionBracket(Controller, RigVMExecuteContext.GetInstructionIndex());
+		Item = Controller->AddControl(Name, Parent, ControlSettings, Value, OffsetTransform, ShapeTransform, false, false);
+	}
+}
+
+void FRigUnit_HierarchyAddControlVector_LimitSettings::Configure(FRigControlSettings& OutSettings) const
+{
+	OutSettings.SetupLimitArrayForType(false, false, false);
+	OutSettings.LimitEnabled[0] = LimitX;
+	OutSettings.LimitEnabled[1] = LimitY;
+	OutSettings.LimitEnabled[2] = LimitZ;
+	OutSettings.MinimumValue = FRigControlValue::Make<FVector3f>(FVector3f(MinValue));
+	OutSettings.MaximumValue = FRigControlValue::Make<FVector3f>(FVector3f(MaxValue));
+	OutSettings.bDrawLimits = bDrawLimits;
+}
+
+void FRigUnit_HierarchyAddControlVector_Settings::Configure(FRigControlSettings& OutSettings) const
+{
+	Super::Configure(OutSettings);
+	
+	OutSettings.ControlType = bIsPosition ? ERigControlType::Position : ERigControlType::Scale;
+
+	Proxy.Configure(OutSettings);
+	Limits.Configure(OutSettings);
+	Shape.Configure(OutSettings);
+}
+
+FRigUnit_HierarchyAddControlVector_Execute()
+{
+	FString ErrorMessage;
+	if(!FRigUnit_DynamicHierarchyBase::IsValidToRunInContext(Context, ExecuteContext, true, &ErrorMessage))
+	{
+		if(!ErrorMessage.IsEmpty())
+		{
+			UE_CONTROLRIG_RIGUNIT_REPORT_ERROR(TEXT("%s"), *ErrorMessage);
+		}
+		return;
+	}
+
+	Item.Reset();
+
+	if(URigHierarchyController* Controller = ExecuteContext.Hierarchy->GetController(true))
+	{
+		FRigControlSettings ControlSettings;
+		Settings.Configure(ControlSettings);
+
+		const FRigControlValue Value = FRigControlValue::Make<FVector3f>(FVector3f(InitialValue));
+		const FTransform ShapeTransform = Settings.Shape.Transform;
+		
+		FRigHierarchyControllerInstructionBracket InstructionBracket(Controller, RigVMExecuteContext.GetInstructionIndex());
+		Item = Controller->AddControl(Name, Parent, ControlSettings, Value, OffsetTransform, ShapeTransform, false, false);
+	}
+}
+
+void FRigUnit_HierarchyAddControlRotator_LimitSettings::Configure(FRigControlSettings& OutSettings) const
+{
+	OutSettings.SetupLimitArrayForType(false, false, false);
+	OutSettings.LimitEnabled[0] = LimitPitch;
+	OutSettings.LimitEnabled[1] = LimitYaw;
+	OutSettings.LimitEnabled[2] = LimitRoll;
+	OutSettings.MinimumValue = FRigControlValue::Make<FVector3f>(FVector3f(MinValue.Euler()));
+	OutSettings.MaximumValue = FRigControlValue::Make<FVector3f>(FVector3f(MaxValue.Euler()));
+	OutSettings.bDrawLimits = bDrawLimits;
+}
+
+void FRigUnit_HierarchyAddControlRotator_Settings::Configure(FRigControlSettings& OutSettings) const
+{
+	Super::Configure(OutSettings);
+	
+	OutSettings.ControlType = ERigControlType::Rotator;
+
+	Proxy.Configure(OutSettings);
+	Limits.Configure(OutSettings);
+	Shape.Configure(OutSettings);
+}
+
+FRigUnit_HierarchyAddControlRotator_Execute()
+{
+	FString ErrorMessage;
+	if(!FRigUnit_DynamicHierarchyBase::IsValidToRunInContext(Context, ExecuteContext, true, &ErrorMessage))
+	{
+		if(!ErrorMessage.IsEmpty())
+		{
+			UE_CONTROLRIG_RIGUNIT_REPORT_ERROR(TEXT("%s"), *ErrorMessage);
+		}
+		return;
+	}
+
+	Item.Reset();
+
+	if(URigHierarchyController* Controller = ExecuteContext.Hierarchy->GetController(true))
+	{
+		FRigControlSettings ControlSettings;
+		Settings.Configure(ControlSettings);
+
+		const FRigControlValue Value = FRigControlValue::Make<FVector3f>(FVector3f(InitialValue.Euler()));
+		const FTransform ShapeTransform = Settings.Shape.Transform;
+		
+		FRigHierarchyControllerInstructionBracket InstructionBracket(Controller, RigVMExecuteContext.GetInstructionIndex());
+		Item = Controller->AddControl(Name, Parent, ControlSettings, Value, OffsetTransform, ShapeTransform, false, false);
+	}
+}
+
+void FRigUnit_HierarchyAddControlTransform_Settings::Configure(FRigControlSettings& OutSettings) const
+{
+	Super::Configure(OutSettings);
+	
+	OutSettings.ControlType = ERigControlType::EulerTransform;
+
+	Proxy.Configure(OutSettings);
+	Shape.Configure(OutSettings);
+}
+
+FRigUnit_HierarchyAddControlTransform_Execute()
+{
+	FString ErrorMessage;
+	if(!FRigUnit_DynamicHierarchyBase::IsValidToRunInContext(Context, ExecuteContext, true, &ErrorMessage))
+	{
+		if(!ErrorMessage.IsEmpty())
+		{
+			UE_CONTROLRIG_RIGUNIT_REPORT_ERROR(TEXT("%s"), *ErrorMessage);
+		}
+		return;
+	}
+
+	Item.Reset();
+
+	if(URigHierarchyController* Controller = ExecuteContext.Hierarchy->GetController(true))
+	{
+		FRigControlSettings ControlSettings;
+		Settings.Configure(ControlSettings);
+
+		const FEulerTransform EulerTransform(InitialValue);
+		const FRigControlValue Value = FRigControlValue::Make<FRigControlValue::FEulerTransform_Float>(EulerTransform);
+		const FTransform ShapeTransform = Settings.Shape.Transform;
+		
+		FRigHierarchyControllerInstructionBracket InstructionBracket(Controller, RigVMExecuteContext.GetInstructionIndex());
+		Item = Controller->AddControl(Name, Parent, ControlSettings, Value, OffsetTransform, ShapeTransform, false, false);
+	}
+}
+
+FRigUnit_HierarchyAddAnimationChannelBool_Execute()
+{
+	FString ErrorMessage;
+	if(!FRigUnit_DynamicHierarchyBase::IsValidToRunInContext(Context, ExecuteContext, true, &ErrorMessage))
+	{
+		if(!ErrorMessage.IsEmpty())
+		{
+			UE_CONTROLRIG_RIGUNIT_REPORT_ERROR(TEXT("%s"), *ErrorMessage);
+		}
+		return;
+	}
+
+	Item.Reset();
+
+	if(URigHierarchyController* Controller = ExecuteContext.Hierarchy->GetController(true))
+	{
+		FRigControlSettings ControlSettings;
+		ControlSettings.ControlType = ERigControlType::Bool;
+		ControlSettings.SetupLimitArrayForType(true, true, true);
+		ControlSettings.MinimumValue = FRigControlValue::Make<float>(MinimumValue ? 1.f : 0.f);
+		ControlSettings.MaximumValue = FRigControlValue::Make<float>(MaximumValue ? 1.f : 0.f);
+		const FRigControlValue Value = FRigControlValue::Make<float>(InitialValue ? 1.f : 0.f);
+
+		FRigHierarchyControllerInstructionBracket InstructionBracket(Controller, RigVMExecuteContext.GetInstructionIndex());
+		Item = Controller->AddAnimationChannel(Name, Parent, ControlSettings, false, false);
+
+		if(Item.IsValid())
+		{
+			ExecuteContext.Hierarchy->SetControlValue(Item, Value, ERigControlValueType::Initial, false, false);
+			ExecuteContext.Hierarchy->SetControlValue(Item, Value, ERigControlValueType::Current, false, false);
+		}
+		else
+		{
+			UE_CONTROLRIG_RIGUNIT_REPORT_ERROR(TEXT("Animation Channel could not be added. Is Parent valid?"));
+		}
+	}
+}
+
+FRigUnit_HierarchyAddAnimationChannelFloat_Execute()
+{
+	FString ErrorMessage;
+	if(!FRigUnit_DynamicHierarchyBase::IsValidToRunInContext(Context, ExecuteContext, true, &ErrorMessage))
+	{
+		if(!ErrorMessage.IsEmpty())
+		{
+			UE_CONTROLRIG_RIGUNIT_REPORT_ERROR(TEXT("%s"), *ErrorMessage);
+		}
+		return;
+	}
+
+	Item.Reset();
+
+	if(URigHierarchyController* Controller = ExecuteContext.Hierarchy->GetController(true))
+	{
+		FRigControlSettings ControlSettings;
+		ControlSettings.ControlType = ERigControlType::Float;
+		ControlSettings.SetupLimitArrayForType(true, true, true);
+		ControlSettings.MinimumValue = FRigControlValue::Make<float>(MinimumValue);
+		ControlSettings.MaximumValue = FRigControlValue::Make<float>(MaximumValue);
+		const FRigControlValue Value = FRigControlValue::Make<float>(InitialValue);
+		
+		FRigHierarchyControllerInstructionBracket InstructionBracket(Controller, RigVMExecuteContext.GetInstructionIndex());
+		Item = Controller->AddAnimationChannel(Name, Parent, ControlSettings, false, false);
+
+		if(Item.IsValid())
+		{
+			ExecuteContext.Hierarchy->SetControlValue(Item, Value, ERigControlValueType::Initial, false, false);
+			ExecuteContext.Hierarchy->SetControlValue(Item, Value, ERigControlValueType::Current, false, false);
+		}
+		else
+		{
+			UE_CONTROLRIG_RIGUNIT_REPORT_ERROR(TEXT("Animation Channel could not be added. Is Parent valid?"));
+		}
+	}
+}
+
+FRigUnit_HierarchyAddAnimationChannelVector2D_Execute()
+{
+	FString ErrorMessage;
+	if(!FRigUnit_DynamicHierarchyBase::IsValidToRunInContext(Context, ExecuteContext, true, &ErrorMessage))
+	{
+		if(!ErrorMessage.IsEmpty())
+		{
+			UE_CONTROLRIG_RIGUNIT_REPORT_ERROR(TEXT("%s"), *ErrorMessage);
+		}
+		return;
+	}
+
+	Item.Reset();
+
+	if(URigHierarchyController* Controller = ExecuteContext.Hierarchy->GetController(true))
+	{
+		FRigControlSettings ControlSettings;
+		ControlSettings.ControlType = ERigControlType::Vector2D;
+		ControlSettings.SetupLimitArrayForType(true, true, true);
+		ControlSettings.MinimumValue = FRigControlValue::Make<FVector3f>(FVector3f(MinimumValue.X, MinimumValue.Y, 0.f));
+		ControlSettings.MaximumValue = FRigControlValue::Make<FVector3f>(FVector3f(MaximumValue.X, MaximumValue.Y, 0.f));
+		const FRigControlValue Value = FRigControlValue::Make<FVector3f>(FVector3f(InitialValue.X, InitialValue.Y, 0.f));
+		
+		FRigHierarchyControllerInstructionBracket InstructionBracket(Controller, RigVMExecuteContext.GetInstructionIndex());
+		Item = Controller->AddAnimationChannel(Name, Parent, ControlSettings, false, false);
+
+		if(Item.IsValid())
+		{
+			ExecuteContext.Hierarchy->SetControlValue(Item, Value, ERigControlValueType::Initial, false, false);
+			ExecuteContext.Hierarchy->SetControlValue(Item, Value, ERigControlValueType::Current, false, false);
+		}
+		else
+		{
+			UE_CONTROLRIG_RIGUNIT_REPORT_ERROR(TEXT("Animation Channel could not be added. Is Parent valid?"));
+		}
+	}
+}
+
+FRigUnit_HierarchyAddAnimationChannelVector_Execute()
+{
+	FString ErrorMessage;
+	if(!FRigUnit_DynamicHierarchyBase::IsValidToRunInContext(Context, ExecuteContext, true, &ErrorMessage))
+	{
+		if(!ErrorMessage.IsEmpty())
+		{
+			UE_CONTROLRIG_RIGUNIT_REPORT_ERROR(TEXT("%s"), *ErrorMessage);
+		}
+		return;
+	}
+
+	Item.Reset();
+
+	if(URigHierarchyController* Controller = ExecuteContext.Hierarchy->GetController(true))
+	{
+		FRigControlSettings ControlSettings;
+		ControlSettings.ControlType = ERigControlType::Position;
+		ControlSettings.SetupLimitArrayForType(true, true, true);
+		ControlSettings.MinimumValue = FRigControlValue::Make<FVector3f>(FVector3f(MinimumValue));
+		ControlSettings.MaximumValue = FRigControlValue::Make<FVector3f>(FVector3f(MaximumValue));
+		const FRigControlValue Value = FRigControlValue::Make<FVector3f>(FVector3f(InitialValue));
+		
+		FRigHierarchyControllerInstructionBracket InstructionBracket(Controller, RigVMExecuteContext.GetInstructionIndex());
+		Item = Controller->AddAnimationChannel(Name, Parent, ControlSettings, false, false);
+
+		if(Item.IsValid())
+		{
+			ExecuteContext.Hierarchy->SetControlValue(Item, Value, ERigControlValueType::Initial, false, false);
+			ExecuteContext.Hierarchy->SetControlValue(Item, Value, ERigControlValueType::Current, false, false);
+		}
+		else
+		{
+			UE_CONTROLRIG_RIGUNIT_REPORT_ERROR(TEXT("Animation Channel could not be added. Is Parent valid?"));
+		}
+	}
+}
+
+FRigUnit_HierarchyAddAnimationChannelRotator_Execute()
+{
+	FString ErrorMessage;
+	if(!FRigUnit_DynamicHierarchyBase::IsValidToRunInContext(Context, ExecuteContext, true, &ErrorMessage))
+	{
+		if(!ErrorMessage.IsEmpty())
+		{
+			UE_CONTROLRIG_RIGUNIT_REPORT_ERROR(TEXT("%s"), *ErrorMessage);
+		}
+		return;
+	}
+
+	Item.Reset();
+
+	if(URigHierarchyController* Controller = ExecuteContext.Hierarchy->GetController(true))
+	{
+		FRigControlSettings ControlSettings;
+		ControlSettings.ControlType = ERigControlType::Rotator;
+		ControlSettings.SetupLimitArrayForType(true, true, true);
+		ControlSettings.MinimumValue = FRigControlValue::Make<FVector3f>(FVector3f(MinimumValue.Euler()));
+		ControlSettings.MaximumValue = FRigControlValue::Make<FVector3f>(FVector3f(MaximumValue.Euler()));
+		const FRigControlValue Value = FRigControlValue::Make<FVector3f>(FVector3f(InitialValue.Euler()));
+		
+		FRigHierarchyControllerInstructionBracket InstructionBracket(Controller, RigVMExecuteContext.GetInstructionIndex());
+		Item = Controller->AddAnimationChannel(Name, Parent, ControlSettings, false, false);
+
+		if(Item.IsValid())
+		{
+			ExecuteContext.Hierarchy->SetControlValue(Item, Value, ERigControlValueType::Initial, false, false);
+			ExecuteContext.Hierarchy->SetControlValue(Item, Value, ERigControlValueType::Current, false, false);
+		}
+		else
+		{
+			UE_CONTROLRIG_RIGUNIT_REPORT_ERROR(TEXT("Animation Channel could not be added. Is Parent valid?"));
+		}
+	}
+}
diff --git a/Engine/Plugins/Animation/ControlRig/Source/ControlRig/Private/Units/Execution/RigUnit_DynamicHierarchy.h b/Engine/Plugins/Animation/ControlRig/Source/ControlRig/Private/Units/Execution/RigUnit_DynamicHierarchy.h
--- a/Engine/Plugins/Animation/ControlRig/Source/ControlRig/Private/Units/Execution/RigUnit_DynamicHierarchy.h
+++ b/Engine/Plugins/Animation/ControlRig/Source/ControlRig/Private/Units/Execution/RigUnit_DynamicHierarchy.h
@@ -10,6 +10,12 @@ USTRUCT(meta = (Abstract, NodeColor="0.262745, 0.8, 0, 0.229412", Category = "Dy
 struct CONTROLRIG_API FRigUnit_DynamicHierarchyBase : public FRigUnit
 {
 	GENERATED_BODY()
+
+	static bool IsValidToRunInContext(
+		const FRigUnitContext& InContext,
+		const FControlRigExecuteContext& InExecuteContext,
+		bool bAllowOnlyConstructionEvent,
+		FString* OutErrorMessage = nullptr);
 };
 
 USTRUCT(meta = (Abstract, NodeColor="0.262745, 0.8, 0, 0.229412", Category = "DynamicHierarchy"))
@@ -201,4 +207,962 @@ struct CONTROLRIG_API FRigUnit_HierarchySetParentWeights : public FRigUnit_Dynam
 	 */
 	UPROPERTY(meta = (Input))
 	TArray<FRigElementWeight> Weights;
+};
+
+/**
+ * Removes all elements from the hierarchy
+ * Note: This node only runs as part of the construction event.
+ */
+USTRUCT(meta=(DisplayName="Reset Hierarchy", Keywords="Delete,Erase,Clear,Empty,RemoveAll", Varying))
+struct FRigUnit_HierarchyReset : public FRigUnit_DynamicHierarchyBaseMutable
+{
+	GENERATED_BODY()
+
+	FRigUnit_HierarchyReset()
+	{
+	}
+
+	RIGVM_METHOD()
+	virtual void Execute(const FRigUnitContext& Context) override;
+};
+
+/**
+ * Imports all bones (and curves) from the currently assigned skeleton.
+ * Note: This node only runs as part of the construction event.
+ */
+USTRUCT(meta=(DisplayName="Import Skeleton", Keywords="Skeleton,Mesh,AddBone,AddCurve,Spawn", Varying))
+struct FRigUnit_HierarchyImportFromSkeleton : public FRigUnit_DynamicHierarchyBaseMutable
+{
+	GENERATED_BODY()
+
+	FRigUnit_HierarchyImportFromSkeleton()
+	{
+		NameSpace = NAME_None;
+		bIncludeCurves = false;
+	}
+
+	RIGVM_METHOD()
+	virtual void Execute(const FRigUnitContext& Context) override;
+
+	UPROPERTY(meta = (Input))
+	FName NameSpace;
+
+	UPROPERTY(meta = (Input))
+	bool bIncludeCurves;
+
+	UPROPERTY(meta = (Output))
+	TArray<FRigElementKey> Items;
+};
+
+/**
+ * Removes an element from the hierarchy
+ * Note: This node only runs as part of the construction event.
+ */
+USTRUCT(meta=(DisplayName="Remove Item", Keywords="Delete,Erase,Joint,Skeleton", Varying))
+struct CONTROLRIG_API FRigUnit_HierarchyRemoveElement : public FRigUnit_DynamicHierarchyBaseMutable
+{
+	GENERATED_BODY()
+
+	FRigUnit_HierarchyRemoveElement()
+	{
+		Item = FRigElementKey(NAME_None, ERigElementType::Bone);
+		bSuccess = false;
+	}
+
+	RIGVM_METHOD()
+	virtual void Execute(const FRigUnitContext& Context) override;
+
+	/*
+	 * The item to remove
+	 */
+	UPROPERTY(meta = (Input, ExpandByDefault))
+	FRigElementKey Item;
+
+	/*
+	 * True if the element has been removed successfuly 
+	 */
+	UPROPERTY(meta = (Output, ExpandByDefault))
+	bool bSuccess;
+};
+
+USTRUCT(meta=(Abstract))
+struct CONTROLRIG_API FRigUnit_HierarchyAddElement : public FRigUnit_DynamicHierarchyBaseMutable
+{
+	GENERATED_BODY()
+
+	FRigUnit_HierarchyAddElement()
+	{
+		Parent = Item = FRigElementKey(NAME_None, ERigElementType::Bone);
+		Name = NAME_None;
+	}
+
+	/*
+	 * The parent of the new element to add
+	 */
+	UPROPERTY(meta = (Input, ExpandByDefault))
+	FRigElementKey Parent;
+
+	/*
+	 * The name of the new element to add
+	 */
+	UPROPERTY(meta = (Input))
+	FName Name;
+
+	/*
+	 * The resulting item
+	 */
+	UPROPERTY(meta = (Output))
+	FRigElementKey Item;
+};
+
+/**
+ * Adds a new bone to the hierarchy
+ * Note: This node only runs as part of the construction event.
+ */
+USTRUCT(meta=(DisplayName="Spawn Bone", Keywords="Construction,Create,New,Joint,Skeleton", Varying))
+struct CONTROLRIG_API FRigUnit_HierarchyAddBone : public FRigUnit_HierarchyAddElement
+{
+	GENERATED_BODY()
+
+	FRigUnit_HierarchyAddBone()
+	{
+		Name = TEXT("NewBone");
+		Transform = FTransform::Identity;
+		Space = EBoneGetterSetterMode::LocalSpace;
+	}
+
+	/*
+	 * The initial transform of the new element
+	 */
+	UPROPERTY(meta = (Input))
+	FTransform Transform;
+
+	/**
+	 * Defines if the transform should be interpreted in local or global space
+	 */ 
+	UPROPERTY(meta = (Input))
+	EBoneGetterSetterMode Space;
+
+	RIGVM_METHOD()
+	virtual void Execute(const FRigUnitContext& Context) override;
+};
+
+/**
+ * Adds a new null to the hierarchy
+ * Note: This node only runs as part of the construction event.
+ */
+USTRUCT(meta=(DisplayName="Spawn Null", Keywords="Construction,Create,New,Locator,Group", Varying))
+struct CONTROLRIG_API FRigUnit_HierarchyAddNull : public FRigUnit_HierarchyAddElement
+{
+	GENERATED_BODY()
+
+	FRigUnit_HierarchyAddNull()
+	{
+		Name = TEXT("NewNull");
+		Transform = FTransform::Identity;
+		Space = EBoneGetterSetterMode::LocalSpace;
+	}
+
+	/*
+	 * The initial transform of the new element
+	 */
+	UPROPERTY(meta = (Input))
+	FTransform Transform;
+
+	/**
+	 * Defines if the transform should be interpreted in local or global space
+	 */ 
+	UPROPERTY(meta = (Input))
+	EBoneGetterSetterMode Space;
+
+	RIGVM_METHOD()
+	virtual void Execute(const FRigUnitContext& Context) override;
+};
+
+USTRUCT(BlueprintType)
+struct CONTROLRIG_API FRigUnit_HierarchyAddControl_Settings
+{
+	GENERATED_BODY()
+
+	FRigUnit_HierarchyAddControl_Settings()
+		: DisplayName(NAME_None)
+	{}
+	virtual ~FRigUnit_HierarchyAddControl_Settings(){}
+
+	void Configure(FRigControlSettings& OutSettings) const;
+
+	UPROPERTY(EditAnywhere, BlueprintReadWrite, Category = "Settings")
+	FName DisplayName;
+};
+
+USTRUCT(BlueprintType)
+struct CONTROLRIG_API FRigUnit_HierarchyAddControl_ShapeSettings
+{
+	GENERATED_BODY()
+
+	FRigUnit_HierarchyAddControl_ShapeSettings()
+		: bVisible(true)
+		, Name(TEXT("Default"))
+		, Color(FLinearColor::Red)
+		, Transform(FTransform::Identity)
+	{}
+
+	void Configure(FRigControlSettings& OutSettings) const;
+
+	UPROPERTY(EditAnywhere, BlueprintReadWrite, Category = "Settings")
+	bool bVisible;
+
+	UPROPERTY(EditAnywhere, BlueprintReadWrite, Category = "Settings", meta=(CustomWidget = "ShapeName"))
+	FName Name;
+
+	UPROPERTY(EditAnywhere, BlueprintReadWrite, Category = "Settings")
+	FLinearColor Color;
+
+	UPROPERTY(EditAnywhere, BlueprintReadWrite, Category = "Settings")
+	FTransform Transform;
+};
+
+USTRUCT(BlueprintType)
+struct CONTROLRIG_API FRigUnit_HierarchyAddControl_ProxySettings
+{
+	GENERATED_BODY()
+
+	FRigUnit_HierarchyAddControl_ProxySettings()
+		: bIsProxy(false)
+		, ShapeVisibility(ERigControlVisibility::BasedOnSelection)
+	{}
+
+	void Configure(FRigControlSettings& OutSettings) const;
+
+	UPROPERTY(EditAnywhere, BlueprintReadWrite, Category = "Settings")
+	bool bIsProxy;
+
+	UPROPERTY(EditAnywhere, BlueprintReadWrite, Category = "Settings")
+	TArray<FRigElementKey> DrivenControls;
+	
+	UPROPERTY(EditAnywhere, BlueprintReadWrite, Category = "Settings")
+	ERigControlVisibility ShapeVisibility;
+};
+
+USTRUCT(BlueprintType)
+struct CONTROLRIG_API FRigUnit_HierarchyAddControlFloat_LimitSettings
+{
+	GENERATED_BODY();
+	
+	FRigUnit_HierarchyAddControlFloat_LimitSettings()
+		: Limit(true, true)
+		, MinValue(0.f)
+		, MaxValue(100.f)
+		, bDrawLimits(true)
+	{}
+
+	void Configure(FRigControlSettings& OutSettings) const;
+
+	UPROPERTY(EditAnywhere, BlueprintReadWrite, Category = "Settings")
+	FRigControlLimitEnabled Limit;
+
+	UPROPERTY(EditAnywhere, BlueprintReadWrite, Category = "Settings")
+	float MinValue;
+
+	UPROPERTY(EditAnywhere, BlueprintReadWrite, Category = "Settings")
+	float MaxValue;
+
+	UPROPERTY(EditAnywhere, BlueprintReadWrite, Category = "Settings")
+	bool bDrawLimits;
+};
+
+USTRUCT(BlueprintType)
+struct CONTROLRIG_API FRigUnit_HierarchyAddControlFloat_Settings : public FRigUnit_HierarchyAddControl_Settings
+{
+	GENERATED_BODY()
+	
+	FRigUnit_HierarchyAddControlFloat_Settings()
+		: FRigUnit_HierarchyAddControl_Settings()
+		, PrimaryAxis(ERigControlAxis::X)
+	{}
+	virtual ~FRigUnit_HierarchyAddControlFloat_Settings() override {}
+
+	void Configure(FRigControlSettings& OutSettings) const;
+
+	UPROPERTY(EditAnywhere, BlueprintReadWrite, Category = "Settings")
+	ERigControlAxis PrimaryAxis;
+	
+	UPROPERTY(EditAnywhere, BlueprintReadWrite, Category = "Settings")
+	FRigUnit_HierarchyAddControlFloat_LimitSettings Limits;;
+
+	UPROPERTY(EditAnywhere, BlueprintReadWrite, Category = "Settings")
+	FRigUnit_HierarchyAddControl_ShapeSettings Shape;
+
+	UPROPERTY(EditAnywhere, BlueprintReadWrite, Category = "Settings")
+	FRigUnit_HierarchyAddControl_ProxySettings Proxy;
+};
+
+/**
+ * Adds a new control to the hierarchy
+ * Note: This node only runs as part of the construction event.
+ */
+USTRUCT(meta=(DisplayName="Spawn Float Control", TemplateName="SpawnControl", Keywords="Construction,Create,New,AddControl,NewControl,CreateControl", Varying))
+struct CONTROLRIG_API FRigUnit_HierarchyAddControlFloat : public FRigUnit_HierarchyAddElement
+{
+	GENERATED_BODY()
+
+	FRigUnit_HierarchyAddControlFloat()
+	{
+		Name = TEXT("NewControl");
+		OffsetTransform = FTransform::Identity;
+		InitialValue = 0.f;
+	}
+
+	/*
+	 * The offset transform of the new control
+	 */
+	UPROPERTY(meta = (Input))
+	FTransform OffsetTransform;
+
+	/*
+	 * The initial value of the new control
+	 */
+	UPROPERTY(meta = (Input))
+	float InitialValue;
+	
+	/*
+	 * The settings for the control
+	 */
+	UPROPERTY(meta = (Input))
+	FRigUnit_HierarchyAddControlFloat_Settings Settings;
+
+	RIGVM_METHOD()
+	virtual void Execute(const FRigUnitContext& Context) override;
+};
+
+USTRUCT(BlueprintType)
+struct CONTROLRIG_API FRigUnit_HierarchyAddControlInteger_LimitSettings
+{
+	GENERATED_BODY();
+	
+	FRigUnit_HierarchyAddControlInteger_LimitSettings()
+		: Limit(true, true)
+		, MinValue(0)
+		, MaxValue(100)
+		, bDrawLimits(true)
+	{}
+
+	void Configure(FRigControlSettings& OutSettings) const;
+
+	UPROPERTY(EditAnywhere, BlueprintReadWrite, Category = "Settings")
+	FRigControlLimitEnabled Limit;
+
+	UPROPERTY(EditAnywhere, BlueprintReadWrite, Category = "Settings")
+	int32 MinValue;
+
+	UPROPERTY(EditAnywhere, BlueprintReadWrite, Category = "Settings")
+	int32 MaxValue;
+
+	UPROPERTY(EditAnywhere, BlueprintReadWrite, Category = "Settings")
+	bool bDrawLimits;
+};
+
+USTRUCT(BlueprintType)
+struct CONTROLRIG_API FRigUnit_HierarchyAddControlInteger_Settings : public FRigUnit_HierarchyAddControl_Settings
+{
+	GENERATED_BODY()
+	
+	FRigUnit_HierarchyAddControlInteger_Settings()
+		: FRigUnit_HierarchyAddControl_Settings()
+		, PrimaryAxis(ERigControlAxis::X)
+	{}
+	virtual ~FRigUnit_HierarchyAddControlInteger_Settings() override {}
+
+	void Configure(FRigControlSettings& OutSettings) const;
+
+	UPROPERTY(EditAnywhere, BlueprintReadWrite, Category = "Settings")
+	ERigControlAxis PrimaryAxis;
+	
+	UPROPERTY(EditAnywhere, BlueprintReadWrite, Category = "Settings")
+	FRigUnit_HierarchyAddControlInteger_LimitSettings Limits;;
+
+	UPROPERTY(EditAnywhere, BlueprintReadWrite, Category = "Settings")
+	FRigUnit_HierarchyAddControl_ShapeSettings Shape;
+
+	UPROPERTY(EditAnywhere, BlueprintReadWrite, Category = "Settings")
+	FRigUnit_HierarchyAddControl_ProxySettings Proxy;
+};
+
+/**
+ * Adds a new control to the hierarchy
+ * Note: This node only runs as part of the construction event.
+ */
+USTRUCT(meta=(DisplayName="Spawn Integer Control", TemplateName="SpawnControl", Keywords="Construction,Create,New,AddControl,NewControl,CreateControl", Varying))
+struct CONTROLRIG_API FRigUnit_HierarchyAddControlInteger : public FRigUnit_HierarchyAddElement
+{
+	GENERATED_BODY()
+
+	FRigUnit_HierarchyAddControlInteger()
+	{
+		Name = TEXT("NewControl");
+		OffsetTransform = FTransform::Identity;
+		InitialValue = 0;
+	}
+
+	/*
+	 * The offset transform of the new control
+	 */
+	UPROPERTY(meta = (Input))
+	FTransform OffsetTransform;
+
+	/*
+	 * The initial value of the new control
+	 */
+	UPROPERTY(meta = (Input))
+	int32 InitialValue;
+	
+	/*
+	 * The settings for the control
+	 */
+	UPROPERTY(meta = (Input))
+	FRigUnit_HierarchyAddControlInteger_Settings Settings;
+
+	RIGVM_METHOD()
+	virtual void Execute(const FRigUnitContext& Context) override;
+};
+
+USTRUCT(BlueprintType)
+struct CONTROLRIG_API FRigUnit_HierarchyAddControlVector2D_LimitSettings
+{
+	GENERATED_BODY();
+	
+	FRigUnit_HierarchyAddControlVector2D_LimitSettings()
+		: LimitX(true, true)
+		, LimitY(true, true)
+		, MinValue(FVector2D::ZeroVector)
+		, MaxValue(FVector2D(1.f, 1.f))
+		, bDrawLimits(true)
+	{}
+
+	void Configure(FRigControlSettings& OutSettings) const;
+
+	UPROPERTY(EditAnywhere, BlueprintReadWrite, Category = "Settings")
+	FRigControlLimitEnabled LimitX;
+
+	UPROPERTY(EditAnywhere, BlueprintReadWrite, Category = "Settings")
+	FRigControlLimitEnabled LimitY;
+
+	UPROPERTY(EditAnywhere, BlueprintReadWrite, Category = "Settings")
+	FVector2D MinValue;
+
+	UPROPERTY(EditAnywhere, BlueprintReadWrite, Category = "Settings")
+	FVector2D MaxValue;
+
+	UPROPERTY(EditAnywhere, BlueprintReadWrite, Category = "Settings")
+	bool bDrawLimits;
+};
+
+USTRUCT(BlueprintType)
+struct CONTROLRIG_API FRigUnit_HierarchyAddControlVector2D_Settings : public FRigUnit_HierarchyAddControl_Settings
+{
+	GENERATED_BODY()
+	
+	FRigUnit_HierarchyAddControlVector2D_Settings()
+		: FRigUnit_HierarchyAddControl_Settings()
+		, PrimaryAxis(ERigControlAxis::X)
+	{}
+	virtual ~FRigUnit_HierarchyAddControlVector2D_Settings() override {}
+
+	void Configure(FRigControlSettings& OutSettings) const;
+
+	UPROPERTY(EditAnywhere, BlueprintReadWrite, Category = "Settings")
+	ERigControlAxis PrimaryAxis;
+	
+	UPROPERTY(EditAnywhere, BlueprintReadWrite, Category = "Settings")
+	FRigUnit_HierarchyAddControlVector2D_LimitSettings Limits;;
+
+	UPROPERTY(EditAnywhere, BlueprintReadWrite, Category = "Settings")
+	FRigUnit_HierarchyAddControl_ShapeSettings Shape;
+
+	UPROPERTY(EditAnywhere, BlueprintReadWrite, Category = "Settings")
+	FRigUnit_HierarchyAddControl_ProxySettings Proxy;
+};
+
+/**
+ * Adds a new control to the hierarchy
+ * Note: This node only runs as part of the construction event.
+ */
+USTRUCT(meta=(DisplayName="Spawn Vector2D Control", TemplateName="SpawnControl", Keywords="Construction,Create,New,AddControl,NewControl,CreateControl", Varying))
+struct CONTROLRIG_API FRigUnit_HierarchyAddControlVector2D : public FRigUnit_HierarchyAddElement
+{
+	GENERATED_BODY()
+
+	FRigUnit_HierarchyAddControlVector2D()
+	{
+		Name = TEXT("NewControl");
+		OffsetTransform = FTransform::Identity;
+		InitialValue = FVector2D::ZeroVector;
+	}
+
+	/*
+	 * The offset transform of the new control
+	 */
+	UPROPERTY(meta = (Input))
+	FTransform OffsetTransform;
+
+	/*
+	 * The initial value of the new control
+	 */
+	UPROPERTY(meta = (Input))
+	FVector2D InitialValue;
+	
+	/*
+	 * The settings for the control
+	 */
+	UPROPERTY(meta = (Input))
+	FRigUnit_HierarchyAddControlVector2D_Settings Settings;
+
+	RIGVM_METHOD()
+	virtual void Execute(const FRigUnitContext& Context) override;
+};
+
+USTRUCT(BlueprintType)
+struct CONTROLRIG_API FRigUnit_HierarchyAddControlVector_LimitSettings
+{
+	GENERATED_BODY();
+	
+	FRigUnit_HierarchyAddControlVector_LimitSettings()
+		: LimitX(false, false)
+		, LimitY(false, false)
+		, LimitZ(false, false)
+		, MinValue(FVector::ZeroVector)
+		, MaxValue(FVector::OneVector)
+		, bDrawLimits(true)
+	{}
+
+	void Configure(FRigControlSettings& OutSettings) const;
+
+	UPROPERTY(EditAnywhere, BlueprintReadWrite, Category = "Settings")
+	FRigControlLimitEnabled LimitX;
+
+	UPROPERTY(EditAnywhere, BlueprintReadWrite, Category = "Settings")
+	FRigControlLimitEnabled LimitY;
+
+	UPROPERTY(EditAnywhere, BlueprintReadWrite, Category = "Settings")
+	FRigControlLimitEnabled LimitZ;
+
+	UPROPERTY(EditAnywhere, BlueprintReadWrite, Category = "Settings")
+	FVector MinValue;
+
+	UPROPERTY(EditAnywhere, BlueprintReadWrite, Category = "Settings")
+	FVector MaxValue;
+
+	UPROPERTY(EditAnywhere, BlueprintReadWrite, Category = "Settings")
+	bool bDrawLimits;
+};
+
+USTRUCT(BlueprintType)
+struct CONTROLRIG_API FRigUnit_HierarchyAddControlVector_Settings : public FRigUnit_HierarchyAddControl_Settings
+{
+	GENERATED_BODY()
+	
+	FRigUnit_HierarchyAddControlVector_Settings()
+		: FRigUnit_HierarchyAddControl_Settings()
+		, bIsPosition(true)
+	{}
+	virtual ~FRigUnit_HierarchyAddControlVector_Settings() override {}
+
+	void Configure(FRigControlSettings& OutSettings) const;
+
+	UPROPERTY(EditAnywhere, BlueprintReadWrite, Category = "Settings")
+	bool bIsPosition;
+
+	UPROPERTY(EditAnywhere, BlueprintReadWrite, Category = "Settings")
+	FRigUnit_HierarchyAddControlVector_LimitSettings Limits;;
+
+	UPROPERTY(EditAnywhere, BlueprintReadWrite, Category = "Settings")
+	FRigUnit_HierarchyAddControl_ShapeSettings Shape;
+
+	UPROPERTY(EditAnywhere, BlueprintReadWrite, Category = "Settings")
+	FRigUnit_HierarchyAddControl_ProxySettings Proxy;
+};
+
+/**
+ * Adds a new control to the hierarchy
+ * Note: This node only runs as part of the construction event.
+ */
+USTRUCT(meta=(DisplayName="Spawn Vector Control", TemplateName="SpawnControl", Keywords="Construction,Create,New,AddControl,NewControl,CreateControl", Varying))
+struct CONTROLRIG_API FRigUnit_HierarchyAddControlVector : public FRigUnit_HierarchyAddElement
+{
+	GENERATED_BODY()
+
+	FRigUnit_HierarchyAddControlVector()
+	{
+		Name = TEXT("NewControl");
+		OffsetTransform = FTransform::Identity;
+		InitialValue = FVector::ZeroVector;
+	}
+
+	/*
+	 * The offset transform of the new control
+	 */
+	UPROPERTY(meta = (Input))
+	FTransform OffsetTransform;
+
+	/*
+	 * The initial value of the new control
+	 */
+	UPROPERTY(meta = (Input))
+	FVector InitialValue;
+
+	/*
+	 * The settings for the control
+	 */
+	UPROPERTY(meta = (Input))
+	FRigUnit_HierarchyAddControlVector_Settings Settings;
+
+	RIGVM_METHOD()
+	virtual void Execute(const FRigUnitContext& Context) override;
+};
+
+USTRUCT(BlueprintType)
+struct CONTROLRIG_API FRigUnit_HierarchyAddControlRotator_LimitSettings
+{
+	GENERATED_BODY();
+	
+	FRigUnit_HierarchyAddControlRotator_LimitSettings()
+		: LimitPitch(false, false)
+		, LimitYaw(false, false)
+		, LimitRoll(false, false)
+		, MinValue(FRotator(-180.f, -180.f, -180.f))
+		, MaxValue(FRotator(180.f, 180.f, 180.f))
+		, bDrawLimits(true)
+	{}
+
+	void Configure(FRigControlSettings& OutSettings) const;
+
+	UPROPERTY(EditAnywhere, BlueprintReadWrite, Category = "Settings")
+	FRigControlLimitEnabled LimitPitch;
+
+	UPROPERTY(EditAnywhere, BlueprintReadWrite, Category = "Settings")
+	FRigControlLimitEnabled LimitYaw;
+
+	UPROPERTY(EditAnywhere, BlueprintReadWrite, Category = "Settings")
+	FRigControlLimitEnabled LimitRoll;
+
+	UPROPERTY(EditAnywhere, BlueprintReadWrite, Category = "Settings")
+	FRotator MinValue;
+
+	UPROPERTY(EditAnywhere, BlueprintReadWrite, Category = "Settings")
+	FRotator MaxValue;
+
+	UPROPERTY(EditAnywhere, BlueprintReadWrite, Category = "Settings")
+	bool bDrawLimits;
+};
+
+USTRUCT(BlueprintType)
+struct CONTROLRIG_API FRigUnit_HierarchyAddControlRotator_Settings : public FRigUnit_HierarchyAddControl_Settings
+{
+	GENERATED_BODY()
+	
+	FRigUnit_HierarchyAddControlRotator_Settings()
+		: FRigUnit_HierarchyAddControl_Settings()
+	{}
+	virtual ~FRigUnit_HierarchyAddControlRotator_Settings() override {}
+
+	void Configure(FRigControlSettings& OutSettings) const;
+
+	UPROPERTY(EditAnywhere, BlueprintReadWrite, Category = "Settings")
+	FRigUnit_HierarchyAddControlRotator_LimitSettings Limits;;
+
+	UPROPERTY(EditAnywhere, BlueprintReadWrite, Category = "Settings")
+	FRigUnit_HierarchyAddControl_ShapeSettings Shape;
+
+	UPROPERTY(EditAnywhere, BlueprintReadWrite, Category = "Settings")
+	FRigUnit_HierarchyAddControl_ProxySettings Proxy;
+};
+
+/**
+ * Adds a new control to the hierarchy
+ * Note: This node only runs as part of the construction event.
+ */
+USTRUCT(meta=(DisplayName="Spawn Rotator Control", TemplateName="SpawnControl", Keywords="Construction,Create,New,AddControl,NewControl,CreateControl,Rotation", Varying))
+struct CONTROLRIG_API FRigUnit_HierarchyAddControlRotator : public FRigUnit_HierarchyAddElement
+{
+	GENERATED_BODY()
+
+	FRigUnit_HierarchyAddControlRotator()
+	{
+		Name = TEXT("NewControl");
+		OffsetTransform = FTransform::Identity;
+		InitialValue = FRotator::ZeroRotator;
+	}
+
+	/*
+	 * The offset transform of the new control
+	 */
+	UPROPERTY(meta = (Input))
+	FTransform OffsetTransform;
+
+	/*
+	 * The initial value of the new control
+	 */
+	UPROPERTY(meta = (Input))
+	FRotator InitialValue;
+
+	/*
+	 * The settings for the control
+	 */
+	UPROPERTY(meta = (Input))
+	FRigUnit_HierarchyAddControlRotator_Settings Settings;
+
+	RIGVM_METHOD()
+	virtual void Execute(const FRigUnitContext& Context) override;
+};
+
+USTRUCT(BlueprintType)
+struct CONTROLRIG_API FRigUnit_HierarchyAddControlTransform_Settings : public FRigUnit_HierarchyAddControl_Settings
+{
+	GENERATED_BODY()
+	
+	FRigUnit_HierarchyAddControlTransform_Settings()
+		: FRigUnit_HierarchyAddControl_Settings()
+	{}
+	virtual ~FRigUnit_HierarchyAddControlTransform_Settings() override {}
+
+	void Configure(FRigControlSettings& OutSettings) const;
+
+	UPROPERTY(EditAnywhere, BlueprintReadWrite, Category = "Settings")
+	FRigUnit_HierarchyAddControl_ShapeSettings Shape;
+
+	UPROPERTY(EditAnywhere, BlueprintReadWrite, Category = "Settings")
+	FRigUnit_HierarchyAddControl_ProxySettings Proxy;
+};
+
+/**
+ * Adds a new control to the hierarchy
+ * Note: This node only runs as part of the construction event.
+ */
+USTRUCT(meta=(DisplayName="Spawn Transform Control", TemplateName="SpawnControl", Keywords="Construction,Create,New,AddControl,NewControl,CreateControl", Varying))
+struct CONTROLRIG_API FRigUnit_HierarchyAddControlTransform : public FRigUnit_HierarchyAddElement
+{
+	GENERATED_BODY()
+
+	FRigUnit_HierarchyAddControlTransform()
+	{
+		Name = TEXT("NewControl");
+		OffsetTransform = FTransform::Identity;
+		InitialValue = FTransform::Identity;
+	}
+
+	/*
+	 * The offset transform of the new control
+	 */
+	UPROPERTY(meta = (Input))
+	FTransform OffsetTransform;
+
+	/*
+	 * The initial value of the new control
+	 */
+	UPROPERTY(meta = (Input))
+	FTransform InitialValue;
+	
+	/*
+	 * The settings for the control
+	 */
+	UPROPERTY(meta = (Input))
+	FRigUnit_HierarchyAddControlTransform_Settings Settings;
+
+	RIGVM_METHOD()
+	virtual void Execute(const FRigUnitContext& Context) override;
+};
+
+/**
+ * Adds a new animation channel to the hierarchy
+ * Note: This node only runs as part of the construction event.
+ */
+USTRUCT(meta=(DisplayName="Spawn Bool Animation Channel", TemplateName="SpawnAnimationChannel", Keywords="Construction,Create,New,AddAnimationChannel,NewAnimationChannel,CreateAnimationChannel,AddChannel,NewChannel,CreateChannel,SpawnChannel", Varying))
+struct CONTROLRIG_API FRigUnit_HierarchyAddAnimationChannelBool : public FRigUnit_HierarchyAddElement
+{
+	GENERATED_BODY()
+
+	FRigUnit_HierarchyAddAnimationChannelBool()
+	{
+		Name = TEXT("NewChannel");
+		InitialValue = false;
+		MinimumValue = false;
+		MaximumValue = true;
+	}
+
+	/*
+	 * The initial value of the new animation channel
+	 */
+	UPROPERTY(meta = (Input))
+	bool InitialValue;
+
+	/*
+	 * The initial value of the new animation channel
+	 */
+	UPROPERTY(meta = (Input))
+	bool MinimumValue;
+
+	/*
+	 * The maximum value for the animation channel
+	 */
+	UPROPERTY(meta = (Input))
+	bool MaximumValue;
+
+	RIGVM_METHOD()
+	virtual void Execute(const FRigUnitContext& Context) override;
+};
+
+/**
+ * Adds a new animation channel to the hierarchy
+ * Note: This node only runs as part of the construction event.
+ */
+USTRUCT(meta=(DisplayName="Spawn Float Animation Channel", TemplateName="SpawnAnimationChannel", Keywords="Construction,Create,New,AddAnimationChannel,NewAnimationChannel,CreateAnimationChannel,AddChannel,NewChannel,CreateChannel,SpawnChannel", Varying))
+struct CONTROLRIG_API FRigUnit_HierarchyAddAnimationChannelFloat : public FRigUnit_HierarchyAddElement
+{
+	GENERATED_BODY()
+
+	FRigUnit_HierarchyAddAnimationChannelFloat()
+	{
+		Name = TEXT("NewChannel");
+		InitialValue = 0.f;
+		MinimumValue = 0.f;
+		MaximumValue = 1.f;
+	}
+
+	/*
+	 * The initial value of the new animation channel
+	 */
+	UPROPERTY(meta = (Input))
+	float InitialValue;
+
+	/*
+	 * The initial value of the new animation channel
+	 */
+	UPROPERTY(meta = (Input))
+	float MinimumValue;
+
+	/*
+	 * The maximum value for the animation channel
+	 */
+	UPROPERTY(meta = (Input))
+	float MaximumValue;
+
+	RIGVM_METHOD()
+	virtual void Execute(const FRigUnitContext& Context) override;
+};
+
+/**
+ * Adds a new animation channel to the hierarchy
+ * Note: This node only runs as part of the construction event.
+ */
+USTRUCT(meta=(DisplayName="Spawn Vector2D Animation Channel", TemplateName="SpawnAnimationChannel", Keywords="Construction,Create,New,AddAnimationChannel,NewAnimationChannel,CreateAnimationChannel,AddChannel,NewChannel,CreateChannel,SpawnChannel", Varying))
+struct CONTROLRIG_API FRigUnit_HierarchyAddAnimationChannelVector2D : public FRigUnit_HierarchyAddElement
+{
+	GENERATED_BODY()
+
+	FRigUnit_HierarchyAddAnimationChannelVector2D()
+	{
+		Name = TEXT("NewChannel");
+		InitialValue = FVector2D::ZeroVector;
+		MinimumValue = FVector2D::ZeroVector;
+		MaximumValue = FVector2D(1.f, 1.f);
+	}
+
+	/*
+	 * The initial value of the new animation channel
+	 */
+	UPROPERTY(meta = (Input))
+	FVector2D InitialValue;
+
+	/*
+	 * The initial value of the new animation channel
+	 */
+	UPROPERTY(meta = (Input))
+	FVector2D MinimumValue;
+
+	/*
+	 * The maximum value for the animation channel
+	 */
+	UPROPERTY(meta = (Input))
+	FVector2D MaximumValue;
+
+	RIGVM_METHOD()
+	virtual void Execute(const FRigUnitContext& Context) override;
+};
+
+/**
+ * Adds a new animation channel to the hierarchy
+ * Note: This node only runs as part of the construction event.
+ */
+USTRUCT(meta=(DisplayName="Spawn Vector Animation Channel", TemplateName="SpawnAnimationChannel", Keywords="Construction,Create,New,AddAnimationChannel,NewAnimationChannel,CreateAnimationChannel,AddChannel,NewChannel,CreateChannel,SpawnChannel", Varying))
+struct CONTROLRIG_API FRigUnit_HierarchyAddAnimationChannelVector : public FRigUnit_HierarchyAddElement
+{
+	GENERATED_BODY()
+
+	FRigUnit_HierarchyAddAnimationChannelVector()
+	{
+		Name = TEXT("NewChannel");
+		InitialValue = FVector::ZeroVector;
+		MinimumValue = FVector::ZeroVector;
+		MaximumValue = FVector(1.f, 1.f, 1.f);
+	}
+
+	/*
+	 * The initial value of the new animation channel
+	 */
+	UPROPERTY(meta = (Input))
+	FVector InitialValue;
+
+	/*
+	 * The initial value of the new animation channel
+	 */
+	UPROPERTY(meta = (Input))
+	FVector MinimumValue;
+
+	/*
+	 * The maximum value for the animation channel
+	 */
+	UPROPERTY(meta = (Input))
+	FVector MaximumValue;
+
+	RIGVM_METHOD()
+	virtual void Execute(const FRigUnitContext& Context) override;
+};
+
+/**
+ * Adds a new animation channel to the hierarchy
+ * Note: This node only runs as part of the construction event.
+ */
+USTRUCT(meta=(DisplayName="Spawn Rotator Animation Channel", TemplateName="SpawnAnimationChannel", Keywords="Construction,Create,New,AddAnimationChannel,NewAnimationChannel,CreateAnimationChannel,AddChannel,NewChannel,CreateChannel,SpawnChannel", Varying))
+struct CONTROLRIG_API FRigUnit_HierarchyAddAnimationChannelRotator : public FRigUnit_HierarchyAddElement
+{
+	GENERATED_BODY()
+
+	FRigUnit_HierarchyAddAnimationChannelRotator()
+	{
+		Name = TEXT("NewChannel");
+		InitialValue = FRotator::ZeroRotator;
+		MinimumValue = FRotator(-180.f, -180.f, -180.f);
+		MaximumValue = FRotator(180.f, 180.f, 180.f);
+	}
+
+	/*
+	 * The initial value of the new animation channel
+	 */
+	UPROPERTY(meta = (Input))
+	FRotator InitialValue;
+
+	/*
+	 * The initial value of the new animation channel
+	 */
+	UPROPERTY(meta = (Input))
+	FRotator MinimumValue;
+
+	/*
+	 * The maximum value for the animation channel
+	 */
+	UPROPERTY(meta = (Input))
+	FRotator MaximumValue;
+
+	RIGVM_METHOD()
+	virtual void Execute(const FRigUnitContext& Context) override;
 };
\ No newline at end of file
diff --git a/Engine/Plugins/Animation/ControlRig/Source/ControlRig/Private/Units/Execution/RigUnit_UserDefinedEvent.cpp b/Engine/Plugins/Animation/ControlRig/Source/ControlRig/Private/Units/Execution/RigUnit_UserDefinedEvent.cpp
--- a/Engine/Plugins/Animation/ControlRig/Source/ControlRig/Private/Units/Execution/RigUnit_UserDefinedEvent.cpp
+++ b/Engine/Plugins/Animation/ControlRig/Source/ControlRig/Private/Units/Execution/RigUnit_UserDefinedEvent.cpp
@@ -18,8 +18,11 @@ FRigUnit_UserDefinedEvent_Execute()
     DECLARE_SCOPE_HIERARCHICAL_COUNTER_RIGUNIT()
 	ExecuteContext.CopyFrom(RigVMExecuteContext);
 	ExecuteContext.Hierarchy = Context.Hierarchy;
-	ExecuteContext.EventName = EventName;
 
+	if(EventName.IsNone())
+	{
+		ExecuteContext.EventName = EventName;
+	}
 	if(EventName.IsNone())
 	{
 		UE_CONTROLRIG_RIGUNIT_REPORT_ERROR(TEXT("Unable to run event - EventName is None."));
diff --git a/Engine/Plugins/Animation/ControlRig/Source/ControlRig/Public/ControlRig.h b/Engine/Plugins/Animation/ControlRig/Source/ControlRig/Public/ControlRig.h
--- a/Engine/Plugins/Animation/ControlRig/Source/ControlRig/Public/ControlRig.h
+++ b/Engine/Plugins/Animation/ControlRig/Source/ControlRig/Public/ControlRig.h
@@ -276,6 +276,9 @@ public:
 	UPROPERTY(transient)
 	ERigExecutionType ExecutionType;
 
+	UPROPERTY()
+	FRigHierarchySettings HierarchySettings;
+
 	UPROPERTY()
 	FRigVMRuntimeSettings VMRuntimeSettings;
 
@@ -306,6 +309,8 @@ public:
 	UFUNCTION(BlueprintCallable, Category = "Control Rig")
 	void RequestConstruction();
 
+	bool IsConstructionRequired() const { return bRequiresConstructionEvent; }
+
 	/** Returns the queue of events to run */
 	const TArray<FName>& GetEventQueue() const { return EventQueue; }
 
@@ -368,7 +373,7 @@ public:
 
 	bool SetControlGlobalTransform(const FName& InControlName, const FTransform& InGlobalTransform, bool bNotify = true, const FRigControlModifiedContext& Context = FRigControlModifiedContext(), bool bSetupUndo = true, bool bPrintPythonCommands = false, bool bFixEulerFlips = false);
 
-	virtual FRigControlValue GetControlValueFromGlobalTransform(const FName& InControlName, const FTransform& InGlobalTransform);
+	virtual FRigControlValue GetControlValueFromGlobalTransform(const FName& InControlName, const FTransform& InGlobalTransform, ERigTransformType::Type InTransformType);
 
 	virtual void SetControlLocalTransform(const FName& InControlName, const FTransform& InLocalTransform, bool bNotify = true, const FRigControlModifiedContext& Context = FRigControlModifiedContext(), bool bSetupUndo = true, bool bFixEulerFlips = false);
 	virtual FTransform GetControlLocalTransform(const FName& InControlName) ;
@@ -424,6 +429,9 @@ public:
 
 	DECLARE_EVENT_ThreeParams(UControlRig, FControlRigExecuteEvent, class UControlRig*, const EControlRigState, const FName&);
 	FControlRigExecuteEvent& OnInitialized_AnyThread() { return InitializedEvent; }
+#if WITH_EDITOR
+	FControlRigExecuteEvent& OnPreConstructionForUI_AnyThread() { return PreConstructionForUIEvent; }
+#endif
 	FControlRigExecuteEvent& OnPreConstruction_AnyThread() { return PreConstructionEvent; }
 	FControlRigExecuteEvent& OnPostConstruction_AnyThread() { return PostConstructionEvent; }
 	FControlRigExecuteEvent& OnPreForwardsSolve_AnyThread() { return PreForwardsSolveEvent; }
@@ -585,6 +593,9 @@ private:
 	/** Broadcasts a notification whenever the controlrig's memory is initialized. */
 	FControlRigExecuteEvent InitializedEvent;
 
+	/** Broadcasts a notification when launching the construction event */
+	FControlRigExecuteEvent PreConstructionForUIEvent;
+
 	/** Broadcasts a notification just before the controlrig is setup. */
 	FControlRigExecuteEvent PreConstructionEvent;
 
@@ -637,6 +648,7 @@ private:
 public:
 	
 	void ApplyTransformOverrideForUserCreatedBones();
+	void ApplySelectionPoseForConstructionMode(const FName& InEventName);
 	
 #endif
 
@@ -835,7 +847,8 @@ private:
 	class FPoseScope
 	{
 	public:
-		FPoseScope(UControlRig* InControlRig, ERigElementType InFilter = ERigElementType::All);
+		FPoseScope(UControlRig* InControlRig, ERigElementType InFilter = ERigElementType::All,
+			const TArray<FRigElementKey>& InElements = TArray<FRigElementKey>());
 		~FPoseScope();
 
 	private:
@@ -911,6 +924,10 @@ public:
 		FRigPose CachedPose;	
 	};	
 
+	bool bRecordSelectionPoseForConstructionMode;
+	TMap<FRigElementKey, FTransform> SelectionPoseForConstructionMode;
+	bool bIsClearingTransientControls;
+	
 #endif
 	
 private:
diff --git a/Engine/Plugins/Animation/ControlRig/Source/ControlRig/Public/Rigs/RigHierarchy.h b/Engine/Plugins/Animation/ControlRig/Source/ControlRig/Public/Rigs/RigHierarchy.h
--- a/Engine/Plugins/Animation/ControlRig/Source/ControlRig/Public/Rigs/RigHierarchy.h
+++ b/Engine/Plugins/Animation/ControlRig/Source/ControlRig/Public/Rigs/RigHierarchy.h
@@ -114,6 +114,13 @@ public:
 	UFUNCTION(BlueprintCallable, Category = URigHierarchy)
 	void Reset();
 
+	/**
+	 * Resets the hierarchy to the state of its default. This refers to the
+	 * hierarchy on the default object.
+	 */
+	UFUNCTION(BlueprintCallable, Category = URigHierarchy)
+	void ResetToDefault();
+
 	/**
 	 * Copies the contents of a hierarchy onto this one
 	 */
@@ -126,6 +133,11 @@ public:
 	 */
 	uint32 GetNameHash() const;
 
+	/**
+	 * Returns a hash representing the topological state of the hierarchy
+	 */
+	uint32 GetTopologyHash(bool bIncludeTopologyVersion = true, bool bIncludeTransientControls = false) const;
+
 #if WITH_EDITOR
 	/**
 	* Add dependent hierarchies that listens to changes made to this hierarchy
@@ -141,6 +153,11 @@ public:
 	void ClearListeningHierarchy();
 #endif
 
+	/**
+	 * Returns the default hierarchy for this hierarchy (or nullptr)
+	 */
+	URigHierarchy* GetDefaultHierarchy() { return DefaultHierarchyPtr.Get(); }
+
 public:
 	/**
 	 * Copies the contents of a hierarchy onto this one
@@ -260,6 +277,32 @@ public:
 		return GetIndex(InKey) != INDEX_NONE;
 	}
 
+	/**
+	 * Returns true if the provided element key is valid as a certain typename
+	 * @param InKey The key to validate
+	 * @return Returns true if the provided element key is valid
+	 */
+	template<typename T>
+	FORCEINLINE bool Contains(const FRigElementKey& InKey) const
+	{
+		return Find<T>(InKey) != nullptr;
+	}
+
+	/**
+	 * Returns true if the provided element is procedural.
+	 * @param InKey The key to validate
+	 * @return Returns true if the element is procedural
+	 */
+	UFUNCTION(BlueprintCallable, Category = URigHierarchy)
+	bool IsProcedural(const FRigElementKey& InKey) const;
+
+	/**
+	 * Returns true if the provided element is procedural.
+	 * @param InElement The element to check
+	 * @return Returns true if the element is procedural
+	 */
+	bool IsProcedural(const FRigBaseElement* InElement) const;
+
 	/**
 	 * Returns the index of an element given its key
 	 * @param InKey The key of the element to retrieve the index for
@@ -411,6 +454,20 @@ public:
 		return CastChecked<T>(Get(InIndex));
 	}
 
+	/**
+     * Returns a handle to an existing element
+     * @param InKey The key of the handle to retrieve.
+     * @return The retrieved handle (may be invalid)
+     */
+	FORCEINLINE FRigElementHandle GetHandle(const FRigElementKey& InKey) const
+	{
+		if(Contains(InKey))
+		{
+			return FRigElementHandle((URigHierarchy*)this, InKey);
+		}
+		return FRigElementHandle();
+	}
+
 	/**
 	 * Returns an element for a given key or nullptr.
 	 * @param InKey The key of the element to retrieve.
@@ -2828,7 +2885,7 @@ private:
 		ElementType* Elements = (ElementType*)FMemory::Malloc(sizeof(ElementType) * Num);
 		for(int32 Index=0;Index<Num;Index++)
 		{
-			new(&Elements[Index]) ElementType(); 
+			new(&Elements[Index]) ElementType();
 		}
 		return Elements;
 	}
@@ -2975,6 +3032,7 @@ private:
 	FRigHierarchyUndoRedoTransformEvent UndoRedoEvent;
 
 	TWeakObjectPtr<URigHierarchy> HierarchyForSelectionPtr;
+	TWeakObjectPtr<URigHierarchy> DefaultHierarchyPtr;
 	TArray<FRigElementKey> OrderedSelection;
 
 	UPROPERTY(Transient)
diff --git a/Engine/Plugins/Animation/ControlRig/Source/ControlRig/Public/Rigs/RigHierarchyController.h b/Engine/Plugins/Animation/ControlRig/Source/ControlRig/Public/Rigs/RigHierarchyController.h
--- a/Engine/Plugins/Animation/ControlRig/Source/ControlRig/Public/Rigs/RigHierarchyController.h
+++ b/Engine/Plugins/Animation/ControlRig/Source/ControlRig/Public/Rigs/RigHierarchyController.h
@@ -19,8 +19,10 @@ public:
 	URigHierarchyController()
 	: bReportWarningsAndErrors(true)
 	, Hierarchy(nullptr)
-	, bSuspendNotifications(false)
+	, bSuspendAllNotifications(false)
+	, bSuspendSelectionNotifications(false)
 	, bSuspendPythonPrinting(false)
+	, CurrentInstructionIndex(INDEX_NONE)
 	{}
 
 	virtual ~URigHierarchyController();
@@ -544,7 +546,7 @@ public:
 	/**
 	 * Returns a reference to the suspend notifications flag
 	 */
-	FORCEINLINE bool& GetSuspendNotificationsFlag() { return bSuspendNotifications; }
+	FORCEINLINE bool& GetSuspendNotificationsFlag() { return bSuspendAllNotifications; }
 
 #if WITH_EDITOR
 	UFUNCTION(BlueprintCallable, Category = URigHierarchyController)
@@ -666,20 +668,49 @@ private:
 #endif
 
 	/** 
-	 * If set to true all notifs coming from this hierarchy will be suspended
+	 * If set to true all notifications coming from this hierarchy will be suspended
 	 */
-	bool bSuspendNotifications;
+	bool bSuspendAllNotifications;
+
+	/** 
+	 * If set to true selection related notifications coming from this hierarchy will be suspended
+	 */
+	bool bSuspendSelectionNotifications;
 
 	/** 
 	* If set to true all python printing can be disabled.  
 	*/
 	bool bSuspendPythonPrinting;
 
+	/**
+	 * If set the controller will mark new items as procedural and created at the current instruction
+	 */
+	int32 CurrentInstructionIndex;
+
 	/**
 	 * This function can be used to override the controller's logging mechanism
 	 */
 	TFunction<void(EMessageSeverity::Type,const FString&)> LogFunction = nullptr;
 
+	template<typename T>
+	T* MakeElement()
+	{
+		T* Element = URigHierarchy::NewElement<T>();
+		Element->CreatedAtInstructionIndex = CurrentInstructionIndex;
+		return Element;
+	}
+	
 	friend class UControlRig;
 	friend class URigHierarchy;
+	friend class FRigHierarchyControllerInstructionBracket;
+};
+
+class CONTROLRIG_API FRigHierarchyControllerInstructionBracket : TGuardValue<int32>
+{
+public:
+	
+	FRigHierarchyControllerInstructionBracket(URigHierarchyController* InController, int32 InInstructionIndex)
+		: TGuardValue<int32>(InController->CurrentInstructionIndex, InInstructionIndex)
+	{
+	}
 };
