Control Rig: Construction Event

diff --git a/Engine/Plugins/Animation/ControlRig/Source/ControlRig/Public/Rigs/RigHierarchyElements.h b/Engine/Plugins/Animation/ControlRig/Source/ControlRig/Public/Rigs/RigHierarchyElements.h
--- a/Engine/Plugins/Animation/ControlRig/Source/ControlRig/Public/Rigs/RigHierarchyElements.h
+++ b/Engine/Plugins/Animation/ControlRig/Source/ControlRig/Public/Rigs/RigHierarchyElements.h
@@ -445,6 +445,58 @@ typedef TArray<FRigBaseElement*, TInlineAllocator<3>> FRigBaseElementChildrenArr
 //typedef TArray<FRigBaseElement*> FRigBaseElementParentArray;
 typedef TArray<FRigBaseElement*, TInlineAllocator<1>> FRigBaseElementParentArray;
 
+struct CONTROLRIG_API FRigElementHandle
+{
+public:
+
+	FRigElementHandle()
+		: Hierarchy(nullptr)
+		, Key()
+	{}
+
+	FRigElementHandle(URigHierarchy* InHierarchy, const FRigElementKey& InKey);
+	FRigElementHandle(URigHierarchy* InHierarchy, const FRigBaseElement* InElement);
+
+	bool IsValid() const { return Get() != nullptr; }
+	operator bool() const { return IsValid(); }
+	
+	const URigHierarchy* GetHierarchy() const { return Hierarchy.Get(); }
+	URigHierarchy* GetHierarchy() { return Hierarchy.Get(); }
+	const FRigElementKey& GetKey() const { return Key; }
+
+	const FRigBaseElement* Get() const;
+	FRigBaseElement* Get();
+
+	template<typename T>
+	T* Get()
+	{
+		return Cast<T>(Get());
+	}
+
+	template<typename T>
+	const T* Get() const
+	{
+		return Cast<T>(Get());
+	}
+
+	template<typename T>
+	T* GetChecked()
+	{
+		return CastChecked<T>(Get());
+	}
+
+	template<typename T>
+	const T* GetChecked() const
+	{
+		return CastChecked<T>(Get());
+	}
+
+private:
+
+	TWeakObjectPtr<URigHierarchy> Hierarchy;
+	FRigElementKey Key;
+};
+
 USTRUCT(BlueprintType)
 struct CONTROLRIG_API FRigBaseElement
 {
@@ -458,6 +510,7 @@ public:
     , Index(INDEX_NONE)
 	, SubIndex(INDEX_NONE)
 	, bSelected(false)
+	, CreatedAtInstructionIndex(INDEX_NONE)
 	, TopologyVersion(0)
 	, OwnedInstances(0)
 	{}
@@ -487,6 +540,9 @@ protected:
 	UPROPERTY(BlueprintReadOnly, Transient, Category = RigElement, meta = (AllowPrivateAccess = "true"))
 	bool bSelected;
 
+	UPROPERTY(BlueprintReadOnly, Transient, Category = RigElement, meta = (AllowPrivateAccess = "true"))
+	int32 CreatedAtInstructionIndex;
+
 	FORCEINLINE static bool IsClassOf(const FRigBaseElement* InElement)
 	{
 		return true;
@@ -507,6 +563,8 @@ public:
 	FORCEINLINE int32 GetIndex() const { return Index; }
 	FORCEINLINE int32 GetSubIndex() const { return SubIndex; }
 	FORCEINLINE bool IsSelected() const { return bSelected; }
+	FORCEINLINE int32 GetCreatedAtInstructionIndex() const { return CreatedAtInstructionIndex; }
+	FORCEINLINE bool IsProcedural() const { return CreatedAtInstructionIndex != INDEX_NONE; }
 
 	template<typename T>
 	FORCEINLINE bool IsA() const { return T::IsClassOf(this); }
@@ -918,6 +976,8 @@ struct CONTROLRIG_API FRigControlSettings
 	void Save(FArchive& Ar);
 	void Load(FArchive& Ar);
 
+	friend uint32 GetTypeHash(const FRigControlSettings& Settings);
+
 	UPROPERTY(BlueprintReadWrite, EditAnywhere, Category = Control)
 	ERigControlAnimationType AnimationType;
 
diff --git a/Engine/Plugins/Animation/ControlRig/Source/ControlRig/Public/Settings/ControlRigSettings.h b/Engine/Plugins/Animation/ControlRig/Source/ControlRig/Public/Settings/ControlRigSettings.h
--- a/Engine/Plugins/Animation/ControlRig/Source/ControlRig/Public/Settings/ControlRigSettings.h
+++ b/Engine/Plugins/Animation/ControlRig/Source/ControlRig/Public/Settings/ControlRigSettings.h
@@ -110,12 +110,6 @@ public:
 	UPROPERTY(EditAnywhere, config, Category = Viewport)
 	FLinearColor BackwardsAndForwardsBorderColor;
 
-	/**
-	 * The border color of the hierarchy tree in dynamic hierarchy mode
-	 */
-	UPROPERTY(EditAnywhere, config, Category = Hierarchy)
-	FLinearColor DynamicHierarchyBorderColor;
-
 	/**
 	 * Option to toggle displaying the stacked hierarchy items.
 	 * Note that changing this option potentially requires to re-open the editors in question. 
diff --git a/Engine/Plugins/Animation/ControlRig/Source/ControlRig/Public/Units/RigUnitContext.h b/Engine/Plugins/Animation/ControlRig/Source/ControlRig/Public/Units/RigUnitContext.h
--- a/Engine/Plugins/Animation/ControlRig/Source/ControlRig/Public/Units/RigUnitContext.h
+++ b/Engine/Plugins/Animation/ControlRig/Source/ControlRig/Public/Units/RigUnitContext.h
@@ -12,6 +12,7 @@
 #include "Drawing/ControlRigDrawInterface.h"
 #include "GameFramework/Actor.h"
 #include "Components/SceneComponent.h"
+#include "RigUnitContext.generated.h"
 
 /** Current state of rig
 *	What  state Control Rig currently is
@@ -37,6 +38,21 @@ enum class EControlRigInteractionType : uint8
 	All = Translate | Rotate | Scale
 };
 
+USTRUCT()
+struct CONTROLRIG_API FRigHierarchySettings
+{
+	GENERATED_BODY();
+
+	FRigHierarchySettings()
+		: ProceduralElementLimit(128)
+	{
+	}
+
+	// Sets the limit for the number of elements to create procedurally
+	UPROPERTY(EditAnywhere, Category = "Hierarchy Settings")
+	int32 ProceduralElementLimit;
+};
+
 /** Execution context that rig units use */
 struct FRigUnitContext
 {
@@ -90,6 +106,9 @@ struct FRigUnitContext
 	/** The current hierarchy being executed */
 	URigHierarchy* Hierarchy;
 
+	/** The current hierarchy settings */
+	FRigHierarchySettings HierarchySettings;
+
 	/** The type of interaction currently happening on the rig (0 == None) */
 	uint8 InteractionType;
 
diff --git a/Engine/Plugins/Animation/ControlRig/Source/ControlRigDeveloper/Private/ControlRigBlueprint.cpp b/Engine/Plugins/Animation/ControlRig/Source/ControlRigDeveloper/Private/ControlRigBlueprint.cpp
--- a/Engine/Plugins/Animation/ControlRig/Source/ControlRigDeveloper/Private/ControlRigBlueprint.cpp
+++ b/Engine/Plugins/Animation/ControlRig/Source/ControlRigDeveloper/Private/ControlRigBlueprint.cpp
@@ -1236,6 +1236,28 @@ void UControlRigBlueprint::RecompileVM()
 			}
 			return;
 		}
+		else
+		{
+			TArray<UEdGraph*> EdGraphs;
+			GetAllGraphs(EdGraphs);
+
+			for (UEdGraph* Graph : EdGraphs)
+			{
+				UControlRigGraph* RigGraph = Cast<UControlRigGraph>(Graph);
+				if (RigGraph == nullptr)
+				{
+					continue;
+				}
+
+				for (UEdGraphNode* GraphNode : Graph->Nodes)
+				{
+					if (UControlRigGraphNode* ControlRigGraphNode = Cast<UControlRigGraphNode>(GraphNode))
+					{
+						ControlRigGraphNode->bHasCompilerMessage = false;
+					}
+				}
+			}
+		}
 
 		TArray<UObject*> ArchetypeInstances;
 		CDO->GetArchetypeInstances(ArchetypeInstances);
@@ -2354,7 +2376,7 @@ void UControlRigBlueprint::PostTransacted(const FTransactionObjectEvent& Transac
 				{
 					continue;
 				}
-				RigGraph->CacheNameLists(Hierarchy, &DrawContainer);
+				RigGraph->CacheNameLists(Hierarchy, &DrawContainer, ShapeLibraries);
 			}
 
 			RequestAutoVMRecompilation();
@@ -2826,24 +2848,6 @@ void UControlRigBlueprint::ClearTransientControls()
 	}
 }
 
-void UControlRigBlueprint::SetTransientControlValue(const FRigElementKey& InElement)
-{
-	UControlRigBlueprintGeneratedClass* RigClass = GetControlRigBlueprintGeneratedClass();
-	UControlRig* CDO = Cast<UControlRig>(RigClass->GetDefaultObject(true /* create if needed */));
-
-	TArray<FRigControl> PreviousControls;
-	TArray<UObject*> ArchetypeInstances;
-	CDO->GetArchetypeInstances(ArchetypeInstances);
-	for (UObject* ArchetypeInstance : ArchetypeInstances)
-	{
-		UControlRig* InstancedControlRig = Cast<UControlRig>(ArchetypeInstance);
-		if (InstancedControlRig)
-		{
-			InstancedControlRig->SetTransientControlValue(InElement);
-		}
-	}
-}
-
 #endif
 
 void UControlRigBlueprint::PopulateModelFromGraphForBackwardsCompatibility(UControlRigGraph* InGraph)
@@ -3192,7 +3196,7 @@ void UControlRigBlueprint::RebuildGraphFromModel()
 	{
 		if (UControlRigGraph* RigGraph = Cast<UControlRigGraph>(Graph))
 		{
-			RigGraph->CacheNameLists(Hierarchy, &DrawContainer);
+			RigGraph->CacheNameLists(Hierarchy, &DrawContainer, ShapeLibraries);
 		}
 	}
 
@@ -4300,6 +4304,7 @@ void UControlRigBlueprint::PropagateHierarchyFromBPToInstances()
 		{
 			DefaultObject->PostInitInstanceIfRequired();
 			DefaultObject->GetHierarchy()->CopyHierarchy(Hierarchy);
+			DefaultObject->HierarchySettings = HierarchySettings;
 			DefaultObject->Initialize(true);
 
 			TArray<UObject*> ArchetypeInstances;
@@ -4310,6 +4315,7 @@ void UControlRigBlueprint::PropagateHierarchyFromBPToInstances()
 				{
 					InstanceRig->PostInitInstanceIfRequired();
 					InstanceRig->GetHierarchy()->CopyHierarchy(Hierarchy);
+					InstanceRig->HierarchySettings = HierarchySettings;
 					InstanceRig->Initialize(true);
 				}
 			}
@@ -4347,7 +4353,7 @@ void UControlRigBlueprint::PropagateDrawInstructionsFromBPToInstances()
 		{
 			continue;
 		}
-		RigGraph->CacheNameLists(Hierarchy, &DrawContainer);
+		RigGraph->CacheNameLists(Hierarchy, &DrawContainer, ShapeLibraries);
 	}
 }
 
diff --git a/Engine/Plugins/Animation/ControlRig/Source/ControlRigDeveloper/Private/Graph/ControlRigGraph.cpp b/Engine/Plugins/Animation/ControlRig/Source/ControlRigDeveloper/Private/Graph/ControlRigGraph.cpp
--- a/Engine/Plugins/Animation/ControlRig/Source/ControlRigDeveloper/Private/Graph/ControlRigGraph.cpp
+++ b/Engine/Plugins/Animation/ControlRig/Source/ControlRigDeveloper/Private/Graph/ControlRigGraph.cpp
@@ -118,7 +118,7 @@ void UControlRigGraph::Initialize(UControlRigBlueprint* InBlueprint)
 
 	if(Hierarchy)
 	{
-		CacheNameLists(Hierarchy, &InBlueprint->DrawContainer);
+		CacheNameLists(Hierarchy, &InBlueprint->DrawContainer, InBlueprint->ShapeLibraries);
 	}
 }
 
@@ -148,7 +148,7 @@ void UControlRigGraph::Serialize(FArchive& Ar)
 
 TArray<TSharedPtr<FString>> UControlRigGraph::EmptyElementNameList;
 
-void UControlRigGraph::CacheNameLists(URigHierarchy* InHierarchy, const FControlRigDrawContainer* DrawContainer)
+void UControlRigGraph::CacheNameLists(URigHierarchy* InHierarchy, const FControlRigDrawContainer* DrawContainer, TArray<TSoftObjectPtr<UControlRigShapeLibrary>> ShapeLibraries)
 {
 	if (UControlRigGraph* OuterGraph = Cast<UControlRigGraph>(GetOuter()))
 	{
@@ -199,6 +199,25 @@ void UControlRigGraph::CacheNameLists(URigHierarchy* InHierarchy, const FControl
 			EntryNameList.Add(MakeShared<FString>(EntryName.ToString()));
 		}
 	}
+
+	ShapeNameList.Reset();
+	ShapeNameList.Add(MakeShared<FString>(FName(NAME_None).ToString()));
+
+	for(const TSoftObjectPtr<UControlRigShapeLibrary>& ShapeLibrary : ShapeLibraries)
+	{
+		if(ShapeLibrary.IsNull())
+		{
+			continue;
+		}
+
+		const bool bUseNameSpace = ShapeLibraries.Num() > 1;
+		const FString NameSpace = bUseNameSpace ? ShapeLibrary->GetName() + TEXT(".") : FString();
+		ShapeNameList.Add(MakeShared<FString>(NameSpace + ShapeLibrary->DefaultShape.ShapeName.ToString()));
+		for (const FControlRigShapeDefinition& Shape : ShapeLibrary->Shapes)
+		{
+			ShapeNameList.Add(MakeShared<FString>(NameSpace + Shape.ShapeName.ToString()));
+		}
+	}
 }
 
 const TArray<TSharedPtr<FString>>* UControlRigGraph::GetElementNameList(ERigElementType InElementType) const
@@ -222,7 +241,13 @@ const TArray<TSharedPtr<FString>>* UControlRigGraph::GetElementNameList(ERigElem
 		}
 
 		UControlRigGraph* MutableThis = (UControlRigGraph*)this;
-		MutableThis->CacheNameLists(Blueprint->Hierarchy, &Blueprint->DrawContainer);
+		URigHierarchy* Hierarchy = Blueprint->Hierarchy;
+		if(UControlRig* ControlRig = Cast<UControlRig>(Blueprint->GetObjectBeingDebugged()))
+		{
+			Hierarchy = ControlRig->GetHierarchy();
+		}	
+			
+		MutableThis->CacheNameLists(Hierarchy, &Blueprint->DrawContainer, Blueprint->ShapeLibraries);
 	}
 	return &ElementNameLists.FindChecked(InElementType);
 }
@@ -294,6 +319,16 @@ const TArray<TSharedPtr<FString>>* UControlRigGraph::GetEntryNameList(URigVMPin*
 	return &EntryNameList;
 }
 
+const TArray<TSharedPtr<FString>>* UControlRigGraph::GetShapeNameList(URigVMPin* InPin) const
+{
+	if (UControlRigGraph* OuterGraph = Cast<UControlRigGraph>(GetOuter()))
+	{
+		return OuterGraph->GetShapeNameList(InPin);
+	}
+	return &ShapeNameList;
+}
+
+
 void UControlRigGraph::HandleModifiedEvent(ERigVMGraphNotifType InNotifType, URigVMGraph* InGraph, UObject* InSubject)
 {
 	DECLARE_SCOPE_HIERARCHICAL_COUNTER_FUNC()
diff --git a/Engine/Plugins/Animation/ControlRig/Source/ControlRigDeveloper/Public/ControlRigBlueprint.h b/Engine/Plugins/Animation/ControlRig/Source/ControlRigDeveloper/Public/ControlRigBlueprint.h
--- a/Engine/Plugins/Animation/ControlRig/Source/ControlRigDeveloper/Public/ControlRigBlueprint.h
+++ b/Engine/Plugins/Animation/ControlRig/Source/ControlRigDeveloper/Public/ControlRigBlueprint.h
@@ -365,6 +365,9 @@ public:
 	UPROPERTY(EditAnywhere, Category = "User Interface")
 	FRigGraphDisplaySettings RigGraphDisplaySettings;
 
+	UPROPERTY(EditAnywhere, Category = "Hierarchy")
+	FRigHierarchySettings HierarchySettings;
+
 	UPROPERTY(EditAnywhere, Category = "VM")
 	FRigVMRuntimeSettings VMRuntimeSettings;
 
@@ -480,9 +483,6 @@ public:
 	/** Removes all  transient / temporary control used to interact with pins */
 	void ClearTransientControls();
 
-	/** update the value of the transient / temporary control to the current value of the rig element*/
-	void SetTransientControlValue(const FRigElementKey& InElement);
-
 #endif
 
 	UPROPERTY(EditAnywhere, Category = "Influence Map")
diff --git a/Engine/Plugins/Animation/ControlRig/Source/ControlRigDeveloper/Public/Graph/ControlRigGraph.h b/Engine/Plugins/Animation/ControlRig/Source/ControlRigDeveloper/Public/Graph/ControlRigGraph.h
--- a/Engine/Plugins/Animation/ControlRig/Source/ControlRigDeveloper/Public/Graph/ControlRigGraph.h
+++ b/Engine/Plugins/Animation/ControlRig/Source/ControlRigDeveloper/Public/Graph/ControlRigGraph.h
@@ -62,12 +62,13 @@ public:
 		return GetElementNameList(ERigElementType::Curve);
 	}
 
-	void CacheNameLists(URigHierarchy* InHierarchy, const FControlRigDrawContainer* DrawContainer);
+	void CacheNameLists(URigHierarchy* InHierarchy, const FControlRigDrawContainer* DrawContainer, TArray<TSoftObjectPtr<UControlRigShapeLibrary>> ShapeLibraries);
 	const TArray<TSharedPtr<FString>>* GetElementNameList(ERigElementType InElementType = ERigElementType::Bone) const;
 	const TArray<TSharedPtr<FString>>* GetElementNameList(URigVMPin* InPin = nullptr) const;
 	const TArray<TSharedPtr<FString>> GetSelectedElementsNameList() const;
 	const TArray<TSharedPtr<FString>>* GetDrawingNameList(URigVMPin* InPin = nullptr) const;
 	const TArray<TSharedPtr<FString>>* GetEntryNameList(URigVMPin* InPin = nullptr) const;
+	const TArray<TSharedPtr<FString>>* GetShapeNameList(URigVMPin* InPin = nullptr) const;
 
 	bool bSuspendModelNotifications;
 	bool bIsTemporaryGraphForCopyPaste;
@@ -145,6 +146,7 @@ private:
 	TMap<ERigElementType, TArray<TSharedPtr<FString>>> ElementNameLists;
 	TArray<TSharedPtr<FString>> DrawingNameList;
 	TArray<TSharedPtr<FString>> EntryNameList;
+	TArray<TSharedPtr<FString>> ShapeNameList;
 	int32 LastHierarchyTopologyVersion;
 
 	bool bIsSelecting;
