Control Rig: Construction Event

diff --git a/Engine/Plugins/Animation/ControlRig/Source/ControlRigEditor/Private/ControlRigElementDetails.cpp b/Engine/Plugins/Animation/ControlRig/Source/ControlRigEditor/Private/ControlRigElementDetails.cpp
--- a/Engine/Plugins/Animation/ControlRig/Source/ControlRigEditor/Private/ControlRigElementDetails.cpp
+++ b/Engine/Plugins/Animation/ControlRig/Source/ControlRigEditor/Private/ControlRigElementDetails.cpp
@@ -14,8 +14,6 @@
 #include "PropertyCustomizationHelpers.h"
 #include "SEnumCombo.h"
 #include "ControlRig/Private/Units/Execution/RigUnit_BeginExecution.h"
-#include "RigVMModel/RigVMGraph.h"
-#include "RigVMModel/RigVMNode.h"
 #include "Graph/SControlRigGraphPinVariableBinding.h"
 #include "HAL/PlatformApplicationMisc.h"
 #include "Styling/AppStyle.h"
@@ -71,7 +69,7 @@ namespace FRigElementKeyDetailsDefs
 	static const float InactivePinBackgroundAlpha = 0.4f;
 };
 
-void RigElementDetails_GetCustomizedInfo(TSharedRef<IPropertyHandle> InStructPropertyHandle, UControlRigBlueprint*& OutBlueprint)
+void RigElementKeyDetails_GetCustomizedInfo(TSharedRef<IPropertyHandle> InStructPropertyHandle, UControlRigBlueprint*& OutBlueprint)
 {
 	TArray<UObject*> Objects;
 	InStructPropertyHandle->GetOuterObjects(Objects);
@@ -149,7 +147,7 @@ UControlRigBlueprint* RigElementDetails_GetBlueprintFromHierarchy(URigHierarchy*
 void FRigElementKeyDetails::CustomizeHeader(TSharedRef<IPropertyHandle> InStructPropertyHandle, FDetailWidgetRow& HeaderRow, IPropertyTypeCustomizationUtils& StructCustomizationUtils)
 {
 	BlueprintBeingCustomized = nullptr;
-	RigElementDetails_GetCustomizedInfo(InStructPropertyHandle, BlueprintBeingCustomized);
+	RigElementKeyDetails_GetCustomizedInfo(InStructPropertyHandle, BlueprintBeingCustomized);
 
 	UControlRigGraph* RigGraph = nullptr;
 	if(BlueprintBeingCustomized)
@@ -438,7 +436,7 @@ FReply FRigElementKeyDetails::OnSelectInHierarchyClicked()
 void FRigComputedTransformDetails::CustomizeHeader(TSharedRef<IPropertyHandle> InStructPropertyHandle, FDetailWidgetRow& HeaderRow, IPropertyTypeCustomizationUtils& StructCustomizationUtils)
 {
 	BlueprintBeingCustomized = nullptr;
-	RigElementDetails_GetCustomizedInfo(InStructPropertyHandle, BlueprintBeingCustomized);
+	RigElementKeyDetails_GetCustomizedInfo(InStructPropertyHandle, BlueprintBeingCustomized);
 }
 
 void FRigComputedTransformDetails::CustomizeChildren(TSharedRef<IPropertyHandle> InStructPropertyHandle, IDetailChildrenBuilder& StructBuilder, IPropertyTypeCustomizationUtils& StructCustomizationUtils)
@@ -518,9 +516,7 @@ void FRigComputedTransformDetails::OnTransformChanged(FEditPropertyChain* InProp
 
 void FRigBaseElementDetails::CustomizeDetails(IDetailLayoutBuilder& DetailBuilder)
 {
-	BlueprintBeingCustomized = nullptr;
-	HierarchyBeingCustomized = nullptr;
-	ObjectsBeingCustomized.Reset();
+	PerElementInfos.Reset();
 
 	TArray<TWeakObjectPtr<UObject>> DetailObjects;
 	DetailBuilder.GetObjectsBeingCustomized(DetailObjects);
@@ -528,35 +524,42 @@ void FRigBaseElementDetails::CustomizeDetails(IDetailLayoutBuilder& DetailBuilde
 	{
 		UDetailsViewWrapperObject* WrapperObject = CastChecked<UDetailsViewWrapperObject>(DetailObject.Get());
 
-		if(HierarchyBeingCustomized == nullptr)
-		{
-			HierarchyBeingCustomized = Cast<URigHierarchy>(WrapperObject->GetOuter());
-		}
+		const FRigElementKey Key = WrapperObject->GetContent<FRigBaseElement>().GetKey();
 
-		ObjectsBeingCustomized.Add(WrapperObject);
-	}
+		FPerElementInfo Info;
+		Info.WrapperObject = WrapperObject;
+		Info.Element = Cast<URigHierarchy>(WrapperObject->GetOuter())->GetHandle(Key);
 
-	if(HierarchyBeingCustomized)
-	{
-		BlueprintBeingCustomized = HierarchyBeingCustomized->GetTypedOuter<UControlRigBlueprint>();
-		if(BlueprintBeingCustomized == nullptr)
+		if(!Info.Element.IsValid())
 		{
-			if(UControlRig* ControlRig = HierarchyBeingCustomized->GetTypedOuter<UControlRig>())
-			{
-				BlueprintBeingCustomized = Cast<UControlRigBlueprint>(ControlRig->GetClass()->ClassGeneratedBy);
-			}
+			return;
+		}
+		if(const UControlRigBlueprint* Blueprint = Info.GetBlueprint())
+		{
+			Info.DefaultElement = Blueprint->Hierarchy->GetHandle(Key);
 		}
-	}
 
-	if(BlueprintBeingCustomized == nullptr)
-	{	
-		RigElementDetails_GetCustomizedInfo(DetailBuilder.GetProperty(TEXT("Key")), BlueprintBeingCustomized);
+		PerElementInfos.Add(Info);
 	}
 
 	IDetailCategoryBuilder& GeneralCategory = DetailBuilder.EditCategory(TEXT("General"), LOCTEXT("General", "General"));
 
+	const bool bIsProcedural = IsAnyElementProcedural();
+	if(bIsProcedural)
+	{
+		GeneralCategory.AddCustomRow(LOCTEXT("ProceduralElement", "ProceduralElement")).WholeRowContent()
+		[
+			SNew(STextBlock)
+			.Text(LOCTEXT("ProceduralElementNote", "This item has been created procedurally."))
+			.ToolTipText(LOCTEXT("ProceduralElementTooltip", "You cannot edit the values of the item here.\nPlease change the settings on the node\nthat created the item."))
+			.Font(IDetailLayoutBuilder::GetDetailFont())
+			.ColorAndOpacity(FLinearColor::Red)
+		];
+	}
+
+	const bool bAllControls = !IsAnyElementNotOfType(ERigElementType::Control);
 	const bool bAllAnimationChannels = !IsAnyControlNotOfAnimationType(ERigControlAnimationType::AnimationChannel);
-	if(bAllAnimationChannels)
+	if(bAllControls && bAllAnimationChannels)
 	{
 		GeneralCategory.AddCustomRow(FText::FromString(TEXT("Parent Control")))
 		.NameContent()
@@ -609,6 +612,7 @@ void FRigBaseElementDetails::CustomizeDetails(IDetailLayoutBuilder& DetailBuilde
 			SNew(STextBlock)
 			.Text(FText::FromString(TEXT("Name")))
 			.Font(IDetailLayoutBuilder::GetDetailFont())
+			.IsEnabled(!bIsProcedural)
 		]
 		.ValueContent()
 		[
@@ -617,41 +621,22 @@ void FRigBaseElementDetails::CustomizeDetails(IDetailLayoutBuilder& DetailBuilde
 			.Text(this, &FRigBaseElementDetails::GetName)
 			.OnTextCommitted(this, &FRigBaseElementDetails::SetName)
 			.OnVerifyTextChanged(this, &FRigBaseElementDetails::OnVerifyNameChanged)
-			.IsEnabled(ObjectsBeingCustomized.Num() == 1 && !bAllAnimationChannels)
+			.IsEnabled(!bIsProcedural && PerElementInfos.Num() == 1 && !bAllAnimationChannels)
 		];
 	}
 
 	DetailBuilder.HideCategory(TEXT("RigElement"));
 }
 
-URigHierarchy* FRigBaseElementDetails::GetHierarchyBeingDebugged() const
-{
-	URigHierarchy* HierarchyBeingDebugged = HierarchyBeingCustomized;
-		
-	if (UControlRig* DebuggedRig = Cast<UControlRig>(BlueprintBeingCustomized->GetObjectBeingDebugged()))
-	{
-		if(!DebuggedRig->IsConstructionModeEnabled())
-		{
-			HierarchyBeingDebugged = DebuggedRig->GetHierarchy();
-		}
-	}
-
-	return HierarchyBeingDebugged;
-}
-
 FRigElementKey FRigBaseElementDetails::GetElementKey() const
 {
-	check(ObjectsBeingCustomized.Num() == 1);
-	if(ObjectsBeingCustomized[0].IsValid())
-	{
-		return ObjectsBeingCustomized[0]->GetContent<FRigBaseElement>().GetKey(); 
-	}
-	return FRigElementKey();
+	check(PerElementInfos.Num() == 1);
+	return PerElementInfos[0].GetElement()->GetKey();
 }
 
 FText FRigBaseElementDetails::GetName() const
 {
-	if(ObjectsBeingCustomized.Num() > 1)
+	if(PerElementInfos.Num() > 1)
 	{
 		return ControlRigDetailsMultipleValues;
 	}
@@ -660,11 +645,11 @@ FText FRigBaseElementDetails::GetName() const
 
 FText FRigBaseElementDetails::GetParentElementName() const
 {
-	if(ObjectsBeingCustomized.Num() > 1 || HierarchyBeingCustomized == nullptr)
+	if(PerElementInfos.Num() > 1)
 	{
 		return ControlRigDetailsMultipleValues;
 	}
-	return FText::FromName(HierarchyBeingCustomized->GetFirstParent(GetElementKey()).Name);
+	return FText::FromName(PerElementInfos[0].GetHierarchy()->GetFirstParent(GetElementKey()).Name);
 }
 
 void FRigBaseElementDetails::SetName(const FText& InNewText, ETextCommit::Type InCommitType)
@@ -674,22 +659,17 @@ void FRigBaseElementDetails::SetName(const FText& InNewText, ETextCommit::Type I
 		return;
 	}
 	
-	if(ObjectsBeingCustomized.Num() > 1)
+	if(PerElementInfos.Num() > 1)
 	{
 		return;
 	}
-	
-	URigHierarchy* Hierarchy = nullptr;
-	if (BlueprintBeingCustomized)
-	{
-		Hierarchy = BlueprintBeingCustomized->Hierarchy;
-	}
-	else
+
+	if(PerElementInfos[0].IsProcedural())
 	{
-		Hierarchy = GetHierarchy();
+		return;
 	}
-
-	if (Hierarchy)
+	
+	if (URigHierarchy* Hierarchy = PerElementInfos[0].GetDefaultHierarchy())
 	{
 		BeginDestroy();
 		
@@ -701,22 +681,18 @@ void FRigBaseElementDetails::SetName(const FText& InNewText, ETextCommit::Type I
 
 bool FRigBaseElementDetails::OnVerifyNameChanged(const FText& InText, FText& OutErrorMessage)
 {
-	if(ObjectsBeingCustomized.Num() > 1)
+	if(PerElementInfos.Num() > 1)
 	{
 		return false;
 	}
 
-	URigHierarchy* Hierarchy = nullptr;
-	if (BlueprintBeingCustomized)
+	if(PerElementInfos[0].IsProcedural())
 	{
-		Hierarchy = BlueprintBeingCustomized->Hierarchy;
-	}
-	else
-	{
-		Hierarchy = GetHierarchy();
+		return false;
 	}
 
-	if (!Hierarchy)
+	const URigHierarchy* Hierarchy = PerElementInfos[0].GetDefaultHierarchy();
+	if (Hierarchy == nullptr)
 	{
 		return false;
 	}
@@ -744,9 +720,14 @@ void FRigBaseElementDetails::OnStructContentsChanged(FProperty* InProperty, cons
 
 bool FRigBaseElementDetails::IsConstructionModeEnabled() const
 {
-	if(BlueprintBeingCustomized)
+	if(PerElementInfos.IsEmpty())
+	{
+		return false;
+	}
+	
+	if(const UControlRigBlueprint* Blueprint = PerElementInfos[0].GetBlueprint())
 	{
-		if (UControlRig* DebuggedRig = Cast<UControlRig>(BlueprintBeingCustomized->GetObjectBeingDebugged()))
+		if (const UControlRig* DebuggedRig = Cast<UControlRig>(Blueprint ->GetObjectBeingDebugged()))
 		{
 			return DebuggedRig->IsConstructionModeEnabled();
 		}
@@ -757,120 +738,109 @@ bool FRigBaseElementDetails::IsConstructionModeEnabled() const
 TArray<FRigElementKey> FRigBaseElementDetails::GetElementKeys() const
 {
 	TArray<FRigElementKey> Keys;
-	for(TWeakObjectPtr<UDetailsViewWrapperObject> ObjectBeingCustomized : ObjectsBeingCustomized)
+	Algo::Transform(PerElementInfos, Keys, [](const FPerElementInfo& Info)
 	{
-		if(ObjectBeingCustomized.IsValid())
-		{
-			Keys.Add(ObjectBeingCustomized->GetContent<FRigBaseElement>().GetKey());
-		}
-	}
+		return Info.GetElement()->GetKey();
+	});
 	return Keys;
 }
 
-bool FRigBaseElementDetails::IsAnyElementOfType(ERigElementType InType) const
+const FRigBaseElementDetails::FPerElementInfo& FRigBaseElementDetails::FindElement(const FRigElementKey& InKey) const
 {
-	for(TWeakObjectPtr<UDetailsViewWrapperObject> ObjectBeingCustomized : ObjectsBeingCustomized)
+	const FPerElementInfo* Info = FindElementByPredicate([InKey](const FPerElementInfo& Info)
 	{
-		if(ObjectBeingCustomized.IsValid())
-		{
-			if(ObjectBeingCustomized->GetContent<FRigBaseElement>().GetKey().Type == InType)
-			{
-				return true;
-			}
-		}
+		return Info.GetElement()->GetKey() == InKey;
+	});
+
+	if(Info)
+	{
+		return *Info;
 	}
-	return false;
+
+	static const FPerElementInfo EmptyInfo;
+	return EmptyInfo;
+}
+
+bool FRigBaseElementDetails::IsAnyElementOfType(ERigElementType InType) const
+{
+	return ContainsElementByPredicate([InType](const FPerElementInfo& Info)
+	{
+		return Info.GetElement()->GetType() == InType;
+	});
 }
 
 bool FRigBaseElementDetails::IsAnyElementNotOfType(ERigElementType InType) const
 {
-	for(TWeakObjectPtr<UDetailsViewWrapperObject> ObjectBeingCustomized : ObjectsBeingCustomized)
+	return ContainsElementByPredicate([InType](const FPerElementInfo& Info)
 	{
-		if(ObjectBeingCustomized.IsValid())
-		{
-			if(ObjectBeingCustomized->GetContent<FRigBaseElement>().GetKey().Type != InType)
-			{
-				return true;
-			}
-		}
-	}
-	return false;
+		return Info.GetElement()->GetType() != InType;
+	});
 }
 
 bool FRigBaseElementDetails::IsAnyControlOfAnimationType(ERigControlAnimationType InType) const
 {
-	for(TWeakObjectPtr<UDetailsViewWrapperObject> ObjectBeingCustomized : ObjectsBeingCustomized)
+	return ContainsElementByPredicate([InType](const FPerElementInfo& Info)
 	{
-		if(ObjectBeingCustomized.IsValid())
+		if(const FRigControlElement* ControlElement = Info.GetElement<FRigControlElement>())
 		{
-			if(ObjectBeingCustomized->IsChildOf<FRigControlElement>())
-			{
-				const FRigControlElement ControlElement = ObjectBeingCustomized->GetContent<FRigControlElement>();
-				if(ControlElement.Settings.AnimationType == InType)
-				{
-					return true;
-				}
-			}
+			return ControlElement->Settings.AnimationType == InType;
 		}
-	}
-	return false;
+		return false;
+	});
 }
 
 bool FRigBaseElementDetails::IsAnyControlNotOfAnimationType(ERigControlAnimationType InType) const
 {
-	for(TWeakObjectPtr<UDetailsViewWrapperObject> ObjectBeingCustomized : ObjectsBeingCustomized)
+	return ContainsElementByPredicate([InType](const FPerElementInfo& Info)
 	{
-		if(ObjectBeingCustomized.IsValid())
+		if(const FRigControlElement* ControlElement = Info.GetElement<FRigControlElement>())
 		{
-			if(ObjectBeingCustomized->IsChildOf<FRigControlElement>())
-			{
-				const FRigControlElement ControlElement = ObjectBeingCustomized->GetContent<FRigControlElement>();
-				if(ControlElement.Settings.AnimationType != InType)
-				{
-					return true;
-				}
-			}
+			return ControlElement->Settings.AnimationType != InType;
 		}
-	}
-	return false;
+		return false;
+	});
 }
 
 bool FRigBaseElementDetails::IsAnyControlOfValueType(ERigControlType InType) const
 {
-	for(TWeakObjectPtr<UDetailsViewWrapperObject> ObjectBeingCustomized : ObjectsBeingCustomized)
+	return ContainsElementByPredicate([InType](const FPerElementInfo& Info)
 	{
-		if(ObjectBeingCustomized.IsValid())
+		if(const FRigControlElement* ControlElement = Info.GetElement<FRigControlElement>())
 		{
-			if(ObjectBeingCustomized->IsChildOf<FRigControlElement>())
-			{
-				const FRigControlElement ControlElement = ObjectBeingCustomized->GetContent<FRigControlElement>();
-				if(ControlElement.Settings.ControlType == InType)
-				{
-					return true;
-				}
-			}
+			return ControlElement->Settings.ControlType == InType;
 		}
-	}
-	return false;
+		return false;
+	});
 }
 
 bool FRigBaseElementDetails::IsAnyControlNotOfValueType(ERigControlType InType) const
 {
-	for(TWeakObjectPtr<UDetailsViewWrapperObject> ObjectBeingCustomized : ObjectsBeingCustomized)
+	return ContainsElementByPredicate([InType](const FPerElementInfo& Info)
 	{
-		if(ObjectBeingCustomized.IsValid())
+		if(const FRigControlElement* ControlElement = Info.GetElement<FRigControlElement>())
 		{
-			if(ObjectBeingCustomized->IsChildOf<FRigControlElement>())
-			{
-				const FRigControlElement ControlElement = ObjectBeingCustomized->GetContent<FRigControlElement>();
-				if(ControlElement.Settings.ControlType != InType)
-				{
-					return true;
-				}
-			}
+			return ControlElement->Settings.ControlType != InType;
 		}
-	}
-	return false;
+		return false;
+	});
+}
+
+bool FRigBaseElementDetails::IsAnyElementProcedural() const
+{
+	return ContainsElementByPredicate([](const FPerElementInfo& Info)
+	{
+		return Info.IsProcedural();
+	});
+}
+
+const FRigBaseElementDetails::FPerElementInfo* FRigBaseElementDetails::FindElementByPredicate(const TFunction<bool(const FPerElementInfo&)>& InPredicate) const
+{
+	return PerElementInfos.FindByPredicate(InPredicate);
+}
+
+bool FRigBaseElementDetails::ContainsElementByPredicate(const TFunction<bool(const FPerElementInfo&)>& InPredicate) const
+{
+	return PerElementInfos.ContainsByPredicate(InPredicate);
 }
 
 void FRigBaseElementDetails::RegisterSectionMappings(FPropertyEditorModule& PropertyEditorModule)
@@ -886,12 +856,12 @@ void FRigBaseElementDetails::RegisterSectionMappings(FPropertyEditorModule& Prop
 
 FReply FRigBaseElementDetails::OnSelectParentElementInHierarchyClicked()
 {
-	if (HierarchyBeingCustomized)
+	if (PerElementInfos.Num() == 1)
 	{
 		FRigElementKey Key = GetElementKey();
 		if (Key.IsValid())
 		{
-			const FRigElementKey ParentKey = HierarchyBeingCustomized->GetFirstParent(Key);
+			const FRigElementKey ParentKey = PerElementInfos[0].GetHierarchy()->GetFirstParent(GetElementKey());
 			if(ParentKey.IsValid())
 			{
 				return OnSelectElementClicked(ParentKey);
@@ -903,11 +873,11 @@ FReply FRigBaseElementDetails::OnSelectParentElementInHierarchyClicked()
 
 FReply FRigBaseElementDetails::OnSelectElementClicked(const FRigElementKey& InKey)
 {
-	if (HierarchyBeingCustomized)
+	if (PerElementInfos.Num() == 1)
 	{
 		if (InKey.IsValid())
 		{
-			HierarchyBeingCustomized->GetController(true)->SetSelection({InKey});
+			 PerElementInfos[0].GetHierarchy()->GetController(true)->SetSelection({InKey});
 		}	
 	}
 	return FReply::Handled();
@@ -932,21 +902,18 @@ void FRigTransformElementDetails::RegisterSectionMappings(FPropertyEditorModule&
 
 void FRigTransformElementDetails::CustomizeTransform(IDetailLayoutBuilder& DetailBuilder)
 {
-	URigHierarchy* HierarchyBeingDebugged = HierarchyBeingCustomized;
-	if (UControlRig* DebuggedRig = Cast<UControlRig>(BlueprintBeingCustomized->GetObjectBeingDebugged()))
+	if(PerElementInfos.IsEmpty())
 	{
- 		if(!DebuggedRig->IsConstructionModeEnabled())
-		{
-			HierarchyBeingDebugged = DebuggedRig->GetHierarchy();
-		}
+		return;
 	}
-
+	
 	TArray<FRigElementKey> Keys = GetElementKeys();
-	Keys = HierarchyBeingDebugged->SortKeys(Keys);
+	Keys = PerElementInfos[0].GetHierarchy()->SortKeys(Keys);
 
+	const bool bIsProcedural = IsAnyElementProcedural();
 	const bool bAllControls = !IsAnyElementNotOfType(ERigElementType::Control) && !IsAnyControlOfValueType(ERigControlType::Bool);;
 	const bool bAllAnimationChannels = !IsAnyControlNotOfAnimationType(ERigControlAnimationType::AnimationChannel);
-	if(bAllAnimationChannels)
+	if(bAllControls && bAllAnimationChannels)
 	{
 		return;
 	}
@@ -1016,30 +983,36 @@ void FRigTransformElementDetails::CustomizeTransform(IDetailLayoutBuilder& Detai
 		const ERigTransformElementDetailsTransform::Type CurrentTransformType = TransformTypes[Index];
 
 		bool bIsTransformEnabled = true;
-		if (IsAnyElementOfType(ERigElementType::Control))
+
+		if(bIsProcedural)
 		{
-			bIsTransformEnabled = IsAnyControlOfValueType(ERigControlType::EulerTransform) ||
-				IsAnyControlOfValueType(ERigControlType::Transform) ||
-				CurrentTransformType == ERigTransformElementDetailsTransform::Offset;
+			// procedural items only allow editing of the current transform
+			bIsTransformEnabled = CurrentTransformType == ERigTransformElementDetailsTransform::Current; 
+		}
 
-			if(!bIsTransformEnabled)
+		if(bIsTransformEnabled)
+		{
+			if (IsAnyElementOfType(ERigElementType::Control))
 			{
-				ButtonTooltips[Index] = FText::FromString(
-					FString::Printf(TEXT("%s\n%s"),
-						*ButtonTooltips[Index].ToString(), 
-						TEXT("Only transform controls can be edited here. Refer to the 'Value' section instead.")));
+				bIsTransformEnabled = IsAnyControlOfValueType(ERigControlType::EulerTransform) ||
+					IsAnyControlOfValueType(ERigControlType::Transform) ||
+					CurrentTransformType == ERigTransformElementDetailsTransform::Offset;
+
+				if(!bIsTransformEnabled)
+				{
+					ButtonTooltips[Index] = FText::FromString(
+						FString::Printf(TEXT("%s\n%s"),
+							*ButtonTooltips[Index].ToString(), 
+							TEXT("Only transform controls can be edited here. Refer to the 'Value' section instead.")));
+				}
 			}
-		}
-		else if (IsAnyElementOfType(ERigElementType::Bone) && CurrentTransformType == ERigTransformElementDetailsTransform::Initial)
-		{
-			for(TWeakObjectPtr<UDetailsViewWrapperObject> ObjectBeingCustomized : ObjectsBeingCustomized)
+			else if (IsAnyElementOfType(ERigElementType::Bone) && CurrentTransformType == ERigTransformElementDetailsTransform::Initial)
 			{
-				if(ObjectBeingCustomized.IsValid())
+				for(const FPerElementInfo& Info : PerElementInfos)
 				{
-					if(ObjectBeingCustomized->IsChildOf<FRigBoneElement>())
+					if(const FRigBoneElement* BoneElement = Info.GetElement<FRigBoneElement>())
 					{
-						const FRigBoneElement BoneElement = ObjectBeingCustomized->GetContent<FRigBoneElement>();
-						bIsTransformEnabled = BoneElement.BoneType == ERigBoneType::User;
+						bIsTransformEnabled = BoneElement->BoneType == ERigBoneType::User;
 
 						if(!bIsTransformEnabled)
 						{
@@ -1049,10 +1022,9 @@ void FRigTransformElementDetails::CustomizeTransform(IDetailLayoutBuilder& Detai
 									TEXT("Imported Bones' initial transform cannot be edited.")));
 						}
 					}
-				}
-			}			
+				}			
+			}
 		}
-
 		bTransformsEnabled.Add(bIsTransformEnabled);
 	}
 
@@ -1079,7 +1051,6 @@ void FRigTransformElementDetails::CustomizeTransform(IDetailLayoutBuilder& Detai
 	IDetailCategoryBuilder& TransformCategory = DetailBuilder.EditCategory(TEXT("Transform"), LOCTEXT("Transform", "Transform"));
 	AddChoiceWidgetRow(TransformCategory, FText::FromString(TEXT("TransformType")), TransformChoiceWidget.ToSharedRef());
 
-
 	SAdvancedTransformInputBox<FEulerTransform>::FArguments TransformWidgetArgs = SAdvancedTransformInputBox<FEulerTransform>::FArguments()
 	.DisplayToggle(false)
 	.DisplayRelativeWorld(true)
@@ -1116,7 +1087,6 @@ void FRigTransformElementDetails::CustomizeTransform(IDetailLayoutBuilder& Detai
 		TransformWidgetArgs.IsEnabled(bTransformsEnabled[Index]);
 
 		CreateEulerTransformValueWidgetRow(
-			HierarchyBeingDebugged,
 			Keys,
 			TransformWidgetArgs,
 			TransformCategory,
@@ -1129,17 +1099,7 @@ void FRigTransformElementDetails::CustomizeTransform(IDetailLayoutBuilder& Detai
 
 bool FRigTransformElementDetails::IsCurrentLocalEnabled() const
 {
-	for(TWeakObjectPtr<UDetailsViewWrapperObject> ObjectBeingCustomized : ObjectsBeingCustomized)
-	{
-		if(ObjectBeingCustomized.IsValid())
-		{
-			if(ObjectBeingCustomized->GetContent<FRigBaseElement>().GetType() == ERigElementType::Control)
-			{
-				return false;
-			}
-		}
-	}
-	return true;
+	return IsAnyElementOfType(ERigElementType::Control);
 }
 
 void FRigTransformElementDetails::AddChoiceWidgetRow(IDetailCategoryBuilder& InCategory, const FText& InSearchText, TSharedRef<SWidget> InWidget)
@@ -1163,7 +1123,6 @@ void FRigTransformElementDetails::AddChoiceWidgetRow(IDetailCategoryBuilder& InC
 
 FDetailWidgetRow& FRigTransformElementDetails::CreateTransformComponentValueWidgetRow(
 	ERigControlType InControlType,
-	URigHierarchy* HierarchyBeingDebugged,
 	const TArray<FRigElementKey>& Keys,
 	SAdvancedTransformInputBox<FEulerTransform>::FArguments TransformWidgetArgs,
 	IDetailCategoryBuilder& CategoryBuilder,
@@ -1201,7 +1160,6 @@ FDetailWidgetRow& FRigTransformElementDetails::CreateTransformComponentValueWidg
 	TransformWidgetArgs.ConstructScale(InControlType == ERigControlType::Scale);
 
 	return CreateEulerTransformValueWidgetRow(
-		GetHierarchyBeingDebugged(),
 		Keys,
 		TransformWidgetArgs,
 		CategoryBuilder,
@@ -1215,7 +1173,6 @@ FDetailWidgetRow& FRigTransformElementDetails::CreateTransformComponentValueWidg
 TSharedPtr<TArray<ERigControlValueType>> FRigControlElementDetails::PickedValueTypes;
 
 FDetailWidgetRow& FRigTransformElementDetails::CreateEulerTransformValueWidgetRow(
-	URigHierarchy* HierarchyBeingDebugged,
 	const TArray<FRigElementKey>& Keys,
 	SAdvancedTransformInputBox<FEulerTransform>::FArguments TransformWidgetArgs,
 	IDetailCategoryBuilder& CategoryBuilder,
@@ -1225,6 +1182,13 @@ FDetailWidgetRow& FRigTransformElementDetails::CreateEulerTransformValueWidgetRo
 	ERigControlValueType ValueType,
 	TSharedPtr<SWidget> NameContent)
 {
+	URigHierarchy* Hierarchy = PerElementInfos[0].GetHierarchy();
+	URigHierarchy* HierarchyToChange = PerElementInfos[0].GetDefaultHierarchy();
+	if(ValueType == ERigControlValueType::Current)
+	{
+		HierarchyToChange = Hierarchy;
+	}
+	
 	const FRigElementTransformWidgetSettings& Settings = FRigElementTransformWidgetSettings::FindOrAdd(ValueType, CurrentTransformType, TransformWidgetArgs);
 
 	const bool bDisplayRelativeWorldOnCurrent = TransformWidgetArgs._DisplayRelativeWorld; 
@@ -1279,7 +1243,7 @@ FDetailWidgetRow& FRigTransformElementDetails::CreateEulerTransformValueWidgetRo
 			TransformWidgetArgs.AllowEditRotationRepresentation(false);
 			TransformWidgetArgs.DisplayRelativeWorld(false);
 			TransformWidgetArgs.DisplayToggle(true);
-			TransformWidgetArgs.OnGetToggleChecked_Lambda([Keys, HierarchyBeingDebugged, ValueType]
+			TransformWidgetArgs.OnGetToggleChecked_Lambda([Keys, Hierarchy, ValueType]
 				(
 					ESlateTransformComponent::Type Component,
 					ESlateRotationRepresentation::Type RotationRepresentation,
@@ -1290,7 +1254,7 @@ FDetailWidgetRow& FRigTransformElementDetails::CreateEulerTransformValueWidgetRo
 
 					for(const FRigElementKey& Key : Keys)
 					{
-						if(const FRigControlElement* ControlElement = HierarchyBeingDebugged->Find<FRigControlElement>(Key))
+						if(const FRigControlElement* ControlElement = Hierarchy->Find<FRigControlElement>(Key))
 						{
 							TOptional<bool> Value;
 
@@ -1424,7 +1388,7 @@ FDetailWidgetRow& FRigTransformElementDetails::CreateEulerTransformValueWidgetRo
 					return FirstValue.GetValue() ? ECheckBoxState::Checked : ECheckBoxState::Unchecked;
 				});
 				
-			TransformWidgetArgs.OnToggleChanged_Lambda([ValueType, Keys, this]
+			TransformWidgetArgs.OnToggleChanged_Lambda([ValueType, Keys, this, Hierarchy]
 			(
 				ESlateTransformComponent::Type Component,
 				ESlateRotationRepresentation::Type RotationRepresentation,
@@ -1440,11 +1404,11 @@ FDetailWidgetRow& FRigTransformElementDetails::CreateEulerTransformValueWidgetRo
 				const bool Value = CheckState == ECheckBoxState::Checked;
 
 				FScopedTransaction Transaction(LOCTEXT("ChangeLimitToggle", "Change Limit Toggle"));
-				HierarchyBeingCustomized->Modify();
+				Hierarchy->Modify();
 
 				for(const FRigElementKey& Key : Keys)
 				{
-					if(FRigControlElement* ControlElement = HierarchyBeingCustomized->Find<FRigControlElement>(Key))
+					if(FRigControlElement* ControlElement = Hierarchy->Find<FRigControlElement>(Key))
 					{
 						switch(ControlElement->Settings.ControlType)
 						{
@@ -1552,7 +1516,7 @@ FDetailWidgetRow& FRigTransformElementDetails::CreateEulerTransformValueWidgetRo
 							}
 						}
 						
-						HierarchyBeingCustomized->SetControlSettings(ControlElement, ControlElement->Settings, true, true, true);
+						Hierarchy->SetControlSettings(ControlElement, ControlElement->Settings, true, true, true);
 					}
 				}
 			});
@@ -1569,7 +1533,7 @@ FDetailWidgetRow& FRigTransformElementDetails::CreateEulerTransformValueWidgetRo
 		}
 	}
 
-	auto GetRelativeAbsoluteTransforms = [CurrentTransformType, Keys, HierarchyBeingDebugged](
+	auto GetRelativeAbsoluteTransforms = [CurrentTransformType, Keys, Hierarchy](
 		const FRigElementKey& Key,
 		ERigTransformElementDetailsTransform::Type InTransformType = ERigTransformElementDetailsTransform::Max
 		) -> TPair<FEulerTransform, FEulerTransform>
@@ -1585,10 +1549,10 @@ FDetailWidgetRow& FRigTransformElementDetails::CreateEulerTransformValueWidgetRo
 		const bool bInitial = InTransformType == ERigTransformElementDetailsTransform::Initial; 
 		if(bInitial || InTransformType == ERigTransformElementDetailsTransform::Current)
 		{
-			RelativeTransform.FromFTransform(HierarchyBeingDebugged->GetLocalTransform(Key, bInitial));
-			AbsoluteTransform.FromFTransform(HierarchyBeingDebugged->GetGlobalTransform(Key, bInitial));
+			RelativeTransform.FromFTransform(Hierarchy->GetLocalTransform(Key, bInitial));
+			AbsoluteTransform.FromFTransform(Hierarchy->GetGlobalTransform(Key, bInitial));
 
-			if(FRigControlElement* ControlElement = HierarchyBeingDebugged->Find<FRigControlElement>(Key))
+			if(FRigControlElement* ControlElement = Hierarchy->Find<FRigControlElement>(Key))
 			{
 				switch(ControlElement->Settings.ControlType)
 				{
@@ -1597,7 +1561,7 @@ FDetailWidgetRow& FRigTransformElementDetails::CreateEulerTransformValueWidgetRo
 					case ERigControlType::Transform:
 					case ERigControlType::TransformNoScale:
 					{
-						RelativeTransform.Rotation = HierarchyBeingDebugged->GetControlPreferredRotator(ControlElement, bInitial);
+						RelativeTransform.Rotation = Hierarchy->GetControlPreferredRotator(ControlElement, bInitial);
 						break;
 					}
 					default:
@@ -1609,14 +1573,14 @@ FDetailWidgetRow& FRigTransformElementDetails::CreateEulerTransformValueWidgetRo
 		}
 		else
 		{
-			if(FRigControlElement* ControlElement = HierarchyBeingDebugged->Find<FRigControlElement>(Key))
+			if(FRigControlElement* ControlElement = Hierarchy->Find<FRigControlElement>(Key))
 			{
 				const ERigControlType ControlType = ControlElement->Settings.ControlType;
 
 				if(InTransformType == ERigTransformElementDetailsTransform::Offset)
 				{
-					RelativeTransform.FromFTransform(HierarchyBeingDebugged->GetControlOffsetTransform(ControlElement, ERigTransformType::InitialLocal));
-					AbsoluteTransform.FromFTransform(HierarchyBeingDebugged->GetControlOffsetTransform(ControlElement, ERigTransformType::InitialGlobal));
+					RelativeTransform.FromFTransform(Hierarchy->GetControlOffsetTransform(ControlElement, ERigTransformType::InitialLocal));
+					AbsoluteTransform.FromFTransform(Hierarchy->GetControlOffsetTransform(ControlElement, ERigTransformType::InitialGlobal));
 				}
 				else if(InTransformType == ERigTransformElementDetailsTransform::Minimum)
 				{
@@ -1625,7 +1589,7 @@ FDetailWidgetRow& FRigTransformElementDetails::CreateEulerTransformValueWidgetRo
 						case ERigControlType::Position:
 						{
 							const FVector Data = 
-								(FVector)HierarchyBeingDebugged->GetControlValue(ControlElement, ERigControlValueType::Minimum)
+								(FVector)Hierarchy->GetControlValue(ControlElement, ERigControlValueType::Minimum)
 								.Get<FVector3f>();
 							AbsoluteTransform = RelativeTransform = FEulerTransform(Data, FRotator::ZeroRotator, FVector::OneVector);
 							break;
@@ -1633,7 +1597,7 @@ FDetailWidgetRow& FRigTransformElementDetails::CreateEulerTransformValueWidgetRo
 						case ERigControlType::Rotator:
 						{
 							const FVector Data = 
-								(FVector)HierarchyBeingDebugged->GetControlValue(ControlElement, ERigControlValueType::Minimum)
+								(FVector)Hierarchy->GetControlValue(ControlElement, ERigControlValueType::Minimum)
 								.Get<FVector3f>();
 							FRotator Rotator = FRotator::MakeFromEuler(Data);
 							AbsoluteTransform = RelativeTransform = FEulerTransform(FVector::ZeroVector, Rotator, FVector::OneVector);
@@ -1642,7 +1606,7 @@ FDetailWidgetRow& FRigTransformElementDetails::CreateEulerTransformValueWidgetRo
 						case ERigControlType::Scale:
 						{
 							const FVector Data = 
-								(FVector)HierarchyBeingDebugged->GetControlValue(ControlElement, ERigControlValueType::Minimum)
+								(FVector)Hierarchy->GetControlValue(ControlElement, ERigControlValueType::Minimum)
 								.Get<FVector3f>();
 							AbsoluteTransform = RelativeTransform = FEulerTransform(FVector::ZeroVector, FRotator::ZeroRotator, Data);
 							break;
@@ -1650,7 +1614,7 @@ FDetailWidgetRow& FRigTransformElementDetails::CreateEulerTransformValueWidgetRo
 						case ERigControlType::EulerTransform:
 						{
 							const FRigControlValue::FEulerTransform_Float EulerTransform = 
-								HierarchyBeingDebugged->GetControlValue(ControlElement, ERigControlValueType::Minimum)
+								Hierarchy->GetControlValue(ControlElement, ERigControlValueType::Minimum)
 								.Get<FRigControlValue::FEulerTransform_Float>();
 							AbsoluteTransform = RelativeTransform = EulerTransform.ToTransform();
 							break;
@@ -1664,7 +1628,7 @@ FDetailWidgetRow& FRigTransformElementDetails::CreateEulerTransformValueWidgetRo
 						case ERigControlType::Position:
 						{
 							const FVector Data = 
-								(FVector)HierarchyBeingDebugged->GetControlValue(ControlElement, ERigControlValueType::Maximum)
+								(FVector)Hierarchy->GetControlValue(ControlElement, ERigControlValueType::Maximum)
 								.Get<FVector3f>();
 							AbsoluteTransform = RelativeTransform = FEulerTransform(Data, FRotator::ZeroRotator, FVector::OneVector);
 							break;
@@ -1672,7 +1636,7 @@ FDetailWidgetRow& FRigTransformElementDetails::CreateEulerTransformValueWidgetRo
 						case ERigControlType::Rotator:
 						{
 							const FVector Data = 
-								(FVector)HierarchyBeingDebugged->GetControlValue(ControlElement, ERigControlValueType::Maximum)
+								(FVector)Hierarchy->GetControlValue(ControlElement, ERigControlValueType::Maximum)
 								.Get<FVector3f>();
 							FRotator Rotator = FRotator::MakeFromEuler(Data);
 							AbsoluteTransform = RelativeTransform = FEulerTransform(FVector::ZeroVector, Rotator, FVector::OneVector);
@@ -1681,7 +1645,7 @@ FDetailWidgetRow& FRigTransformElementDetails::CreateEulerTransformValueWidgetRo
 						case ERigControlType::Scale:
 						{
 							const FVector Data = 
-								(FVector)HierarchyBeingDebugged->GetControlValue(ControlElement, ERigControlValueType::Maximum)
+								(FVector)Hierarchy->GetControlValue(ControlElement, ERigControlValueType::Maximum)
 								.Get<FVector3f>();
 							AbsoluteTransform = RelativeTransform = FEulerTransform(FVector::ZeroVector, FRotator::ZeroRotator, Data);
 							break;
@@ -1689,7 +1653,7 @@ FDetailWidgetRow& FRigTransformElementDetails::CreateEulerTransformValueWidgetRo
 						case ERigControlType::EulerTransform:
 						{
 							const FRigControlValue::FEulerTransform_Float EulerTransform = 
-								HierarchyBeingDebugged->GetControlValue(ControlElement, ERigControlValueType::Maximum)
+								Hierarchy->GetControlValue(ControlElement, ERigControlValueType::Maximum)
 								.Get<FRigControlValue::FEulerTransform_Float>();
 							AbsoluteTransform = RelativeTransform = EulerTransform.ToTransform();
 							break;
@@ -1732,7 +1696,7 @@ FDetailWidgetRow& FRigTransformElementDetails::CreateEulerTransformValueWidgetRo
 		return bIsRelative ? RelativeTransform : AbsoluteTransform;
 	};
 
-	auto SetSingleTransform = [CurrentTransformType, GetRelativeAbsoluteTransforms, this, HierarchyBeingDebugged](
+	auto SetSingleTransform = [CurrentTransformType, GetRelativeAbsoluteTransforms, this, Hierarchy](
 		const FRigElementKey& Key,
 		FEulerTransform InTransform,
 		bool bIsRelative,
@@ -1742,16 +1706,16 @@ FDetailWidgetRow& FRigTransformElementDetails::CreateEulerTransformValueWidgetRo
 		const bool bInitial = CurrentTransformType == ERigTransformElementDetailsTransform::Initial; 
 
 		bool bConstructionModeEnabled = false;
-		if (UControlRig* DebuggedRig = Cast<UControlRig>(BlueprintBeingCustomized->GetObjectBeingDebugged()))
+		if (UControlRig* DebuggedRig = Cast<UControlRig>(PerElementInfos[0].GetBlueprint()->GetObjectBeingDebugged()))
 		{
 			bConstructionModeEnabled = DebuggedRig->IsConstructionModeEnabled();
 		}
 
 		TArray<URigHierarchy*> HierarchiesToUpdate;
-		HierarchiesToUpdate.Add(HierarchyBeingDebugged);
+		HierarchiesToUpdate.Add(Hierarchy);
 		if(!bCurrent || bConstructionModeEnabled)
 		{
-			HierarchiesToUpdate.Add(HierarchyBeingCustomized);
+			HierarchiesToUpdate.Add(PerElementInfos[0].GetDefaultHierarchy());
 		}
 
 		for(URigHierarchy* HierarchyToUpdate : HierarchiesToUpdate)
@@ -1909,7 +1873,7 @@ FDetailWidgetRow& FRigTransformElementDetails::CreateEulerTransformValueWidgetRo
 		IsComponentRelative,
 		GetSingleTransform,
 		SetSingleTransform,
-		CurrentTransformType
+		HierarchyToChange
 	](
 		ESlateTransformComponent::Type Component,
 		ESlateRotationRepresentation::Type Representation,
@@ -1929,7 +1893,7 @@ FDetailWidgetRow& FRigTransformElementDetails::CreateEulerTransformValueWidgetRo
 				if(!SliderTransaction.IsValid())
 				{
 					SliderTransaction = MakeShareable(new FScopedTransaction(NSLOCTEXT("ControlRigElementDetails", "ChangeNumericValue", "Change Numeric Value")));
-					HierarchyBeingCustomized->Modify();
+					HierarchyToChange->Modify();
 				}
 							
 				SetSingleTransform(Key, Transform, bIsRelative, false);
@@ -1944,7 +1908,7 @@ FDetailWidgetRow& FRigTransformElementDetails::CreateEulerTransformValueWidgetRo
 		IsComponentRelative,
 		GetSingleTransform,
 		SetSingleTransform,
-		CurrentTransformType
+		HierarchyToChange
 	](
 		ESlateTransformComponent::Type Component,
 		ESlateRotationRepresentation::Type Representation,
@@ -1958,7 +1922,7 @@ FDetailWidgetRow& FRigTransformElementDetails::CreateEulerTransformValueWidgetRo
 			FScopedTransaction Transaction(LOCTEXT("ChangeNumericValue", "Change Numeric Value"));
 			if(!SliderTransaction.IsValid())
 			{
-				HierarchyBeingCustomized->Modify();
+				HierarchyToChange->Modify();
 			}
 			
 			for(const FRigElementKey& Key : Keys)
@@ -2023,7 +1987,7 @@ FDetailWidgetRow& FRigTransformElementDetails::CreateEulerTransformValueWidgetRo
 		}
 	});
 
-	TransformWidgetArgs.OnPasteFromClipboard_Lambda([this, Keys, IsComponentRelative, ConformComponentRelative, GetSingleTransform, SetSingleTransform](
+	TransformWidgetArgs.OnPasteFromClipboard_Lambda([this, Keys, IsComponentRelative, ConformComponentRelative, GetSingleTransform, SetSingleTransform, HierarchyToChange](
 		ESlateTransformComponent::Type InComponent
 		)
 	{
@@ -2046,7 +2010,7 @@ FDetailWidgetRow& FRigTransformElementDetails::CreateEulerTransformValueWidgetRo
 		}
 
 		FScopedTransaction Transaction(LOCTEXT("PasteTransform", "Paste Transform"));
-		HierarchyBeingCustomized->Modify();
+		HierarchyToChange->Modify();
 
 		for(const FRigElementKey& Key : Keys)
 		{
@@ -2180,11 +2144,11 @@ FDetailWidgetRow& FRigTransformElementDetails::CreateEulerTransformValueWidgetRo
 		return false;
 	});
 
-	TransformWidgetArgs.OnResetToDefault_Lambda([this, CurrentTransformType, Keys, GetSingleTransform, SetSingleTransform](
+	TransformWidgetArgs.OnResetToDefault_Lambda([this, CurrentTransformType, Keys, GetSingleTransform, SetSingleTransform, HierarchyToChange](
 		ESlateTransformComponent::Type InComponent)
 	{
 		FScopedTransaction Transaction(LOCTEXT("ResetTransformToDefault", "Reset Transform to Default"));
-		HierarchyBeingCustomized->Modify();
+		HierarchyToChange->Modify();
 
 		for(const FRigElementKey& Key : Keys)
 		{
@@ -2280,26 +2244,6 @@ void FRigControlElementDetails::CustomizeDetails(IDetailLayoutBuilder& DetailBui
 {
 	FRigTransformElementDetails::CustomizeDetails(DetailBuilder);
 
-	if (HierarchyBeingCustomized == nullptr)
-	{
-		return;
-	}
-
-	ControlElements.Reset();
-	ObjectPerControl.Reset();;
-	
-	for(TWeakObjectPtr<UDetailsViewWrapperObject> ObjectBeingCustomized : ObjectsBeingCustomized)
-	{
-		if(ObjectBeingCustomized.IsValid())
-		{
-			if(ObjectBeingCustomized->IsChildOf<FRigControlElement>())
-			{
-				ControlElements.Add(ObjectBeingCustomized->GetContent<FRigControlElement>());
-				ObjectPerControl.Add(ObjectBeingCustomized.Get());
-			}
-		}
-	}
-
 	CustomizeControl(DetailBuilder);
 	CustomizeValue(DetailBuilder);
 	CustomizeTransform(DetailBuilder);
@@ -2309,21 +2253,31 @@ void FRigControlElementDetails::CustomizeDetails(IDetailLayoutBuilder& DetailBui
 
 void FRigControlElementDetails::CustomizeValue(IDetailLayoutBuilder& DetailBuilder)
 {
-	if(ControlElements.Num() == 0)
+	if(PerElementInfos.IsEmpty())
 	{
 		return;
 	}
 
+	if(IsAnyElementNotOfType(ERigElementType::Control))
+	{
+		return;
+	}
+
+	URigHierarchy* Hierarchy = PerElementInfos[0].GetHierarchy();
+
 	// only show this section if all controls are the same type
-	const ERigControlType ControlType = ControlElements[0].Settings.ControlType;
+	const FRigControlElement* FirstControlElement = PerElementInfos[0].GetElement<FRigControlElement>();
+	const ERigControlType ControlType = FirstControlElement->Settings.ControlType;
 	bool bAllAnimationChannels = true;
-	for(const FRigControlElement& ControlElement : ControlElements)
+	
+	for(const FPerElementInfo& Info : PerElementInfos)
 	{
-		if(ControlElement.Settings.ControlType != ControlType)
+		const FRigControlElement* ControlElement = Info.GetElement<FRigControlElement>();
+		if(ControlElement->Settings.ControlType != ControlType)
 		{
 			return;
 		}
-		if(ControlElement.Settings.AnimationType != ERigControlAnimationType::AnimationChannel)
+		if(ControlElement->Settings.AnimationType != ERigControlAnimationType::AnimationChannel)
 		{
 			bAllAnimationChannels = false;
 		}
@@ -2355,7 +2309,7 @@ void FRigControlElementDetails::CustomizeValue(IDetailLayoutBuilder& DetailBuild
 	// transform types already got filtered out earlier.
 	// integers with enums don't have limits either
 	if(ControlType != ERigControlType::Bool &&
-		(ControlType != ERigControlType::Integer || ControlElements[0].Settings.ControlEnum.IsNull()))
+		(ControlType != ERigControlType::Integer || FirstControlElement->Settings.ControlEnum.IsNull()))
 	{
 		Labels.Append({
 			LOCTEXT("Min", "Min"),
@@ -2396,12 +2350,12 @@ void FRigControlElementDetails::CustomizeValue(IDetailLayoutBuilder& DetailBuild
 	AddChoiceWidgetRow(ValueCategory, FText::FromString(TEXT("ValueType")), ValueTypeChoiceWidget.ToSharedRef());
 
 	TArray<FRigElementKey> Keys = GetElementKeys();
-	Keys = GetHierarchyBeingDebugged()->SortKeys(Keys);
+	Keys = Hierarchy->SortKeys(Keys);
 
 	for(int32 Index=0; Index < ValueTypes.Num(); Index++)
 	{
 		const ERigControlValueType ValueType = ValueTypes[Index];
-		
+
 		const TAttribute<EVisibility> VisibilityAttribute =
 			TAttribute<EVisibility>::CreateLambda([ValueType, ValueTypeChoiceWidget]()-> EVisibility
 			{
@@ -2425,7 +2379,7 @@ void FRigControlElementDetails::CustomizeValue(IDetailLayoutBuilder& DetailBuild
 				bool bIsEnum = false;
 				for(const FRigElementKey& Key : Keys)
 				{
-					if(const FRigControlElement* ControlElement = HierarchyBeingCustomized->Find<FRigControlElement>(Key))
+					if(const FRigControlElement* ControlElement = Hierarchy->Find<FRigControlElement>(Key))
 					{
 						if(!ControlElement->Settings.ControlEnum.IsNull())
 						{
@@ -2461,7 +2415,6 @@ void FRigControlElementDetails::CustomizeValue(IDetailLayoutBuilder& DetailBuild
 
 				CreateTransformComponentValueWidgetRow(
 					ControlType,
-					GetHierarchyBeingDebugged(),
 					GetElementKeys(),
 					TransformWidgetArgs,
 					ValueCategory,
@@ -2481,6 +2434,22 @@ void FRigControlElementDetails::CustomizeValue(IDetailLayoutBuilder& DetailBuild
 
 void FRigControlElementDetails::CustomizeControl(IDetailLayoutBuilder& DetailBuilder)
 {
+	if(PerElementInfos.IsEmpty())
+	{
+		return;
+	}
+
+	if(IsAnyElementNotOfType(ERigElementType::Control))
+	{
+		return;
+	}
+
+	const bool bIsProcedural = IsAnyElementProcedural();
+	const bool bIsEnabled = !bIsProcedural;
+	
+	URigHierarchy* Hierarchy = PerElementInfos[0].GetHierarchy();
+	URigHierarchy* HierarchyToChange = PerElementInfos[0].GetDefaultHierarchy();
+
 	const TSharedPtr<IPropertyHandle> SettingsHandle = DetailBuilder.GetProperty(TEXT("Settings"));
 	DetailBuilder.HideProperty(SettingsHandle);
 
@@ -2493,6 +2462,7 @@ void FRigControlElementDetails::CustomizeControl(IDetailLayoutBuilder& DetailBui
 
 	const TSharedPtr<IPropertyHandle> DisplayNameHandle = SettingsHandle->GetChildHandle(TEXT("DisplayName"));
 	ControlCategory.AddCustomRow(DisplayNameLabelText)
+	.IsEnabled(bIsEnabled)
 	.NameContent()
 	[
 		DisplayNameHandle->CreatePropertyNameWidget(DisplayNameLabelText)
@@ -2507,7 +2477,7 @@ void FRigControlElementDetails::CustomizeControl(IDetailLayoutBuilder& DetailBui
 		{
 			return OnVerifyDisplayNameChanged(InText, OutErrorMessage, GetElementKey());
 		})
-		.IsEnabled(ObjectsBeingCustomized.Num() == 1)
+		.IsEnabled(bIsEnabled && (PerElementInfos.Num() == 1))
 	];
 
 	const TSharedRef<IPropertyUtilities> PropertyUtilities = DetailBuilder.GetPropertyUtilities();
@@ -2515,20 +2485,18 @@ void FRigControlElementDetails::CustomizeControl(IDetailLayoutBuilder& DetailBui
 	// when control type changes, we have to refresh detail panel
 	const TSharedPtr<IPropertyHandle> AnimationTypeHandle = SettingsHandle->GetChildHandle(TEXT("AnimationType"));
 	AnimationTypeHandle->SetOnPropertyValueChanged(FSimpleDelegate::CreateLambda(
-		[this, PropertyUtilities]()
+		[this, PropertyUtilities, HierarchyToChange, Hierarchy]()
 		{
 			TArray<FRigControlElement> ControlElementsInView = GetElementsInDetailsView<FRigControlElement>();
-			TArray<FRigControlElement*> ControlElementsInHierarchy = GetElementsInHierarchy<FRigControlElement>();
-			check(ControlElementsInView.Num() == ControlElementsInHierarchy.Num());
 
-			if (this->HierarchyBeingCustomized && ControlElementsInHierarchy.Num() > 0)
+			if (HierarchyToChange && ControlElementsInView.Num() == PerElementInfos.Num())
 			{
-				this->HierarchyBeingCustomized->Modify();
+				HierarchyToChange->Modify();
 				
 				for(int32 ControlIndex = 0; ControlIndex< ControlElementsInView.Num(); ControlIndex++)
 				{
 					const FRigControlElement& ViewElement = ControlElementsInView[ControlIndex];
-					FRigControlElement* ControlElement = ControlElementsInHierarchy[ControlIndex];
+					FRigControlElement* ControlElement = PerElementInfos[ControlIndex].GetDefaultElement<FRigControlElement>();
 					
 					ControlElement->Settings.AnimationType = ViewElement.Settings.AnimationType;
 
@@ -2568,15 +2536,15 @@ void FRigControlElementDetails::CustomizeControl(IDetailLayoutBuilder& DetailBui
 						}
 					}
 
-					this->HierarchyBeingCustomized->SetControlSettings(ControlElement, ControlElement->Settings, true, true, true);
-					ObjectsBeingCustomized[ControlIndex]->SetContent<FRigControlElement>(*ControlElement);
+					HierarchyToChange->SetControlSettings(ControlElement, ControlElement->Settings, true, true, true);
+					PerElementInfos[ControlIndex].WrapperObject->SetContent<FRigControlElement>(*ControlElement);
 
-					if (this->HierarchyBeingCustomized != this->BlueprintBeingCustomized->Hierarchy)
+					if (HierarchyToChange != Hierarchy)
 					{
-						if(FRigControlElement* OtherControlElement = this->BlueprintBeingCustomized->Hierarchy->Find<FRigControlElement>(ControlElement->GetKey()))
+						if(FRigControlElement* OtherControlElement = PerElementInfos[0].GetElement<FRigControlElement>())
 						{
 							OtherControlElement->Settings = ControlElement->Settings;
-							this->BlueprintBeingCustomized->Hierarchy->SetControlSettings(OtherControlElement, OtherControlElement->Settings, true, true, true);
+							Hierarchy->SetControlSettings(OtherControlElement, OtherControlElement->Settings, true, true, true);
 						}
 					}
 				}
@@ -2586,7 +2554,8 @@ void FRigControlElementDetails::CustomizeControl(IDetailLayoutBuilder& DetailBui
 		}
 	));
 
-	ControlCategory.AddProperty(AnimationTypeHandle.ToSharedRef());
+	ControlCategory.AddProperty(AnimationTypeHandle.ToSharedRef())
+	.IsEnabled(bIsEnabled);
 
 	// when control type changes, we have to refresh detail panel
 	const TSharedPtr<IPropertyHandle> ControlTypeHandle = SettingsHandle->GetChildHandle(TEXT("ControlType"));
@@ -2598,47 +2567,42 @@ void FRigControlElementDetails::CustomizeControl(IDetailLayoutBuilder& DetailBui
 		}
 	));
 
-	ControlCategory.AddProperty(ControlTypeHandle.ToSharedRef());
+	ControlCategory.AddProperty(ControlTypeHandle.ToSharedRef())
+	.IsEnabled(bIsEnabled);
 
 	const bool bSupportsShape = !IsAnyControlOfAnimationType(ERigControlAnimationType::AnimationChannel) &&
 		!IsAnyControlOfAnimationType(ERigControlAnimationType::VisualCue);
 
-	if (HierarchyBeingCustomized != nullptr)
+	if (HierarchyToChange != nullptr)
 	{
 		bool bEnableGroupWithParentControl = true;
-		for(TWeakObjectPtr<UDetailsViewWrapperObject> ObjectBeingCustomized : ObjectsBeingCustomized)
+		for(const FPerElementInfo& Info : PerElementInfos)
 		{
-			if(ObjectBeingCustomized.IsValid())
+			if(const FRigControlElement* ControlElement = Info.GetElement<FRigControlElement>())
 			{
-				if(ObjectBeingCustomized->IsChildOf<FRigControlElement>())
+				bool bSingleEnableGroupWithParentControl = false;
+				if(const FRigControlElement* ParentElement =
+					Cast<FRigControlElement>(Info.GetHierarchy()->GetFirstParent(ControlElement)))
 				{
-					bool bSingleEnableGroupWithParentControl = false;
-					const FRigElementKey Key = ObjectBeingCustomized->GetContent<FRigControlElement>().GetKey();
-					if(const FRigControlElement* ControlElement = HierarchyBeingCustomized->Find<FRigControlElement>(Key))
+					if(ControlElement->Settings.IsAnimatable() &&
+						Info.GetHierarchy()->GetChildren(ControlElement).IsEmpty())
 					{
-						if(const FRigControlElement* ParentElement =
-							Cast<FRigControlElement>(HierarchyBeingCustomized->GetFirstParent(ControlElement)))
-						{
-							if(ControlElement->Settings.IsAnimatable() &&
-								HierarchyBeingCustomized->GetChildren(ControlElement).IsEmpty())
-							{
-								bSingleEnableGroupWithParentControl = true;
-							}
-						}
+						bSingleEnableGroupWithParentControl = true;
 					}
+				}
 
-					if(!bSingleEnableGroupWithParentControl)
-					{
-						bEnableGroupWithParentControl = false;
-						break;
-					}
+				if(!bSingleEnableGroupWithParentControl)
+				{
+					bEnableGroupWithParentControl = false;
+					break;
 				}
 			}
 		}
 		if(bEnableGroupWithParentControl)
 		{
 			const TSharedPtr<IPropertyHandle> GroupWithParentControlHandle = SettingsHandle->GetChildHandle(TEXT("bGroupWithParentControl"));
-			ControlCategory.AddProperty(GroupWithParentControlHandle.ToSharedRef()).DisplayName(FText::FromString(TEXT("Group Channels")));
+			ControlCategory.AddProperty(GroupWithParentControlHandle.ToSharedRef()).DisplayName(FText::FromString(TEXT("Group Channels")))
+			.IsEnabled(bIsEnabled);
 		}
 	}
 	
@@ -2648,61 +2612,61 @@ void FRigControlElementDetails::CustomizeControl(IDetailLayoutBuilder& DetailBui
 		IsAnyControlNotOfValueType(ERigControlType::Vector2D)))
 	{
 		const TSharedPtr<IPropertyHandle> PrimaryAxisHandle = SettingsHandle->GetChildHandle(TEXT("PrimaryAxis"));
-		ControlCategory.AddProperty(PrimaryAxisHandle.ToSharedRef()).DisplayName(FText::FromString(TEXT("Primary Axis")));
+		ControlCategory.AddProperty(PrimaryAxisHandle.ToSharedRef()).DisplayName(FText::FromString(TEXT("Primary Axis")))
+		.IsEnabled(bIsEnabled);
 	}
 
 	if(IsAnyControlOfValueType(ERigControlType::Integer))
 	{
 		const TSharedPtr<IPropertyHandle> ControlEnumHandle = SettingsHandle->GetChildHandle(TEXT("ControlEnum"));
-		ControlCategory.AddProperty(ControlEnumHandle.ToSharedRef()).DisplayName(FText::FromString(TEXT("Control Enum")));
+		ControlCategory.AddProperty(ControlEnumHandle.ToSharedRef()).DisplayName(FText::FromString(TEXT("Control Enum")))
+		.IsEnabled(bIsEnabled);
 
 		ControlEnumHandle->SetOnPropertyValueChanged(FSimpleDelegate::CreateLambda(
 			[this, PropertyUtilities]()
 			{
 				PropertyUtilities->ForceRefresh();
 
-				if (this->HierarchyBeingCustomized != nullptr)
+				for(int32 ControlIndex = 0; ControlIndex < PerElementInfos.Num(); ControlIndex++)
 				{
-					for(int32 ControlIndex = 0; ControlIndex < ControlElements.Num(); ControlIndex++)
+					FPerElementInfo& Info = PerElementInfos[ControlIndex];
+					const FRigControlElement ControlInView = Info.WrapperObject->GetContent<FRigControlElement>();
+					FRigControlElement* ControlBeingCustomized = Info.GetDefaultElement<FRigControlElement>();
+					
+					const UEnum* ControlEnum = ControlBeingCustomized->Settings.ControlEnum;
+					if (ControlEnum != nullptr)
 					{
-						const FRigControlElement ControlInView = ControlElements[ControlIndex];
-						FRigControlElement* ControlBeingCustomized = HierarchyBeingCustomized->Find<FRigControlElement>(ControlInView.GetKey());
-						
-						const UEnum* ControlEnum = ControlBeingCustomized->Settings.ControlEnum;
-						if (ControlEnum != nullptr)
+						int32 Maximum = (int32)ControlEnum->GetMaxEnumValue() - 1;
+						ControlBeingCustomized->Settings.MinimumValue.Set<int32>(0);
+						ControlBeingCustomized->Settings.MaximumValue.Set<int32>(Maximum);
+						ControlBeingCustomized->Settings.LimitEnabled.Reset();
+						ControlBeingCustomized->Settings.LimitEnabled.Add(true);
+						Info.GetDefaultHierarchy()->SetControlSettings(ControlBeingCustomized, ControlBeingCustomized->Settings, true, true, true);
+
+						FRigControlValue InitialValue = Info.GetDefaultHierarchy()->GetControlValue(ControlBeingCustomized, ERigControlValueType::Initial);
+						FRigControlValue CurrentValue = Info.GetDefaultHierarchy()->GetControlValue(ControlBeingCustomized, ERigControlValueType::Current);
+
+						ControlBeingCustomized->Settings.ApplyLimits(InitialValue);
+						ControlBeingCustomized->Settings.ApplyLimits(CurrentValue);
+						Info.GetDefaultHierarchy()->SetControlValue(ControlBeingCustomized, InitialValue, ERigControlValueType::Initial, false, false, true);
+						Info.GetDefaultHierarchy()->SetControlValue(ControlBeingCustomized, CurrentValue, ERigControlValueType::Current, false, false, true);
+
+						if (UControlRig* DebuggedRig = Cast<UControlRig>(Info.GetBlueprint()->GetObjectBeingDebugged()))
 						{
-							int32 Maximum = (int32)ControlEnum->GetMaxEnumValue() - 1;
-							ControlBeingCustomized->Settings.MinimumValue.Set<int32>(0);
-							ControlBeingCustomized->Settings.MaximumValue.Set<int32>(Maximum);
-							ControlBeingCustomized->Settings.LimitEnabled.Reset();
-							ControlBeingCustomized->Settings.LimitEnabled.Add(true);
-							HierarchyBeingCustomized->SetControlSettings(ControlBeingCustomized, ControlBeingCustomized->Settings, true, true, true);
-
-							FRigControlValue InitialValue = HierarchyBeingCustomized->GetControlValue(ControlBeingCustomized, ERigControlValueType::Initial);
-							FRigControlValue CurrentValue = HierarchyBeingCustomized->GetControlValue(ControlBeingCustomized, ERigControlValueType::Current);
-
-							ControlBeingCustomized->Settings.ApplyLimits(InitialValue);
-							ControlBeingCustomized->Settings.ApplyLimits(CurrentValue);
-							HierarchyBeingCustomized->SetControlValue(ControlBeingCustomized, InitialValue, ERigControlValueType::Initial, false, false, true);
-							HierarchyBeingCustomized->SetControlValue(ControlBeingCustomized, CurrentValue, ERigControlValueType::Current, false, false, true);
-
-							if (UControlRig* DebuggedRig = Cast<UControlRig>(BlueprintBeingCustomized->GetObjectBeingDebugged()))
+							URigHierarchy* DebuggedHierarchy = DebuggedRig->GetHierarchy();
+							if(FRigControlElement* DebuggedControlElement = DebuggedHierarchy->Find<FRigControlElement>(ControlBeingCustomized->GetKey()))
 							{
-								URigHierarchy* DebuggedHierarchy = DebuggedRig->GetHierarchy();
-								if(FRigControlElement* DebuggedControlElement = DebuggedHierarchy->Find<FRigControlElement>(ControlBeingCustomized->GetKey()))
-								{
-									DebuggedControlElement->Settings.MinimumValue.Set<int32>(0);
-                                    DebuggedControlElement->Settings.MaximumValue.Set<int32>(Maximum);
-									DebuggedHierarchy->SetControlSettings(DebuggedControlElement, DebuggedControlElement->Settings, true, true, true);
+								DebuggedControlElement->Settings.MinimumValue.Set<int32>(0);
+                                DebuggedControlElement->Settings.MaximumValue.Set<int32>(Maximum);
+								DebuggedHierarchy->SetControlSettings(DebuggedControlElement, DebuggedControlElement->Settings, true, true, true);
 
-                                    DebuggedHierarchy->SetControlValue(DebuggedControlElement, InitialValue, ERigControlValueType::Initial);
-                                    DebuggedHierarchy->SetControlValue(DebuggedControlElement, CurrentValue, ERigControlValueType::Current);
-								}
+                                DebuggedHierarchy->SetControlValue(DebuggedControlElement, InitialValue, ERigControlValueType::Initial);
+                                DebuggedHierarchy->SetControlValue(DebuggedControlElement, CurrentValue, ERigControlValueType::Current);
 							}
 						}
-
-						ObjectPerControl[ControlIndex]->SetContent<FRigControlElement>(*ControlBeingCustomized);
 					}
+
+					Info.WrapperObject->SetContent<FRigControlElement>(*ControlBeingCustomized);
 				}
 			}
 		));
@@ -2714,11 +2678,11 @@ void FRigControlElementDetails::CustomizeControl(IDetailLayoutBuilder& DetailBui
 	if(bSupportsShape)
 	{
 		const TSharedPtr<IPropertyHandle> AvailableSpacesHandle = CustomizationHandle->GetChildHandle(TEXT("AvailableSpaces"));
-		ControlCategory.AddProperty(AvailableSpacesHandle.ToSharedRef());
+		ControlCategory.AddProperty(AvailableSpacesHandle.ToSharedRef())
+		.IsEnabled(bIsEnabled);
 	}
 
 	TArray<FRigElementKey> Keys = GetElementKeys();
-	URigHierarchy* HierarchyBeingDebugged = GetHierarchyBeingDebugged();
 
 	if(bSupportsShape)
 	{
@@ -2726,11 +2690,16 @@ void FRigControlElementDetails::CustomizeControl(IDetailLayoutBuilder& DetailBui
 		
 		ControlCategory
 		.AddProperty(DrawLimitsHandle.ToSharedRef()).DisplayName(FText::FromString(TEXT("Draw Limits")))
-		.IsEnabled(TAttribute<bool>::CreateLambda([Keys, HierarchyBeingDebugged]() -> bool
+		.IsEnabled(TAttribute<bool>::CreateLambda([Keys, Hierarchy, bIsEnabled]() -> bool
 		{
+			if(!bIsEnabled)
+			{
+				return false;
+			}
+			
 			for(const FRigElementKey& Key : Keys)
 			{
-				if(const FRigControlElement* ControlElement = HierarchyBeingDebugged->Find<FRigControlElement>(Key))
+				if(const FRigControlElement* ControlElement = Hierarchy->Find<FRigControlElement>(Key))
 				{
 					if(ControlElement->Settings.LimitEnabled.Contains(FRigControlLimitEnabled(true, true)))
 					{
@@ -2744,7 +2713,8 @@ void FRigControlElementDetails::CustomizeControl(IDetailLayoutBuilder& DetailBui
 	
 	if(!IsAnyControlNotOfAnimationType(ERigControlAnimationType::ProxyControl))
 	{
-		ControlCategory.AddProperty(SettingsHandle->GetChildHandle(TEXT("DrivenControls")).ToSharedRef());
+		ControlCategory.AddProperty(SettingsHandle->GetChildHandle(TEXT("DrivenControls")).ToSharedRef())
+		.IsEnabled(bIsEnabled);
 	}
 }
 
@@ -2757,17 +2727,23 @@ void FRigControlElementDetails::CustomizeAnimationChannels(IDetailLayoutBuilder&
 	}
 
 	// only show this if only one control is selected
-	if((ControlElements.Num() != 1) || HierarchyBeingCustomized == nullptr)
+	if(PerElementInfos.Num() != 1)
 	{
 		return;
 	}
 
-	const FRigControlElement* ControlElement = HierarchyBeingCustomized->Find<FRigControlElement>(ControlElements[0].GetKey());
+	const FRigControlElement* ControlElement = PerElementInfos[0].GetElement<FRigControlElement>();
 	if(ControlElement == nullptr)
 	{
 		return;
 	}
 
+	const bool bIsProcedural = IsAnyElementProcedural();
+	const bool bIsEnabled = !bIsProcedural;
+
+	URigHierarchy* Hierarchy = PerElementInfos[0].GetHierarchy();
+	URigHierarchy* HierarchyToChange = PerElementInfos[0].GetDefaultHierarchy();
+
 	IDetailCategoryBuilder& Category = DetailBuilder.EditCategory(TEXT("AnimationChannels"), LOCTEXT("AnimationChannels", "Animation Channels"));
 	
 	const TSharedRef<IPropertyUtilities> PropertyUtilities = DetailBuilder.GetPropertyUtilities();
@@ -2777,6 +2753,7 @@ void FRigControlElementDetails::CustomizeAnimationChannels(IDetailLayoutBuilder&
 	.HAlign(HAlign_Right)
 	[
 		SNew(SButton)
+		.IsEnabled(bIsEnabled)
 		.ButtonStyle(FAppStyle::Get(), "SimpleButton")
 		.ContentPadding(FMargin(1, 0))
 		.OnClicked(this, &FRigControlElementDetails::OnAddAnimationChannelClicked)
@@ -2792,7 +2769,7 @@ void FRigControlElementDetails::CustomizeAnimationChannels(IDetailLayoutBuilder&
 	Category.HeaderContent(HeaderContentWidget);
 
 	bool bHasAnimationChannels = false;
-	const FRigBaseElementChildrenArray ChildElements = HierarchyBeingCustomized->GetChildren(ControlElement);
+	const FRigBaseElementChildrenArray ChildElements = Hierarchy->GetChildren(ControlElement);
 	for(const FRigBaseElement* ChildElement : ChildElements)
 	{
 		if(const FRigControlElement* ChildControlElement = Cast<FRigControlElement>(ChildElement))
@@ -2802,7 +2779,7 @@ void FRigControlElementDetails::CustomizeAnimationChannels(IDetailLayoutBuilder&
 				bHasAnimationChannels = true;
 				const FRigElementKey ChildElementKey = ChildElement->GetKey();
 				
-				const TPair<const FSlateBrush*, FSlateColor> BrushAndColor = SRigHierarchyItem::GetBrushForElementType(HierarchyBeingCustomized, ChildElementKey);
+				const TPair<const FSlateBrush*, FSlateColor> BrushAndColor = SRigHierarchyItem::GetBrushForElementType(Hierarchy, ChildElementKey);
 
 				static TArray<TSharedPtr<ERigControlType>> ControlValueTypes;
 				if(ControlValueTypes.IsEmpty())
@@ -2822,6 +2799,7 @@ void FRigControlElementDetails::CustomizeAnimationChannels(IDetailLayoutBuilder&
 				TSharedPtr<SWidget> NameContent;
 
 				SAssignNew(NameContent, SHorizontalBox)
+				.IsEnabled(bIsEnabled)
 
 				+ SHorizontalBox::Slot()
 				.MaxWidth(32)
@@ -2934,7 +2912,6 @@ void FRigControlElementDetails::CustomizeAnimationChannels(IDetailLayoutBuilder&
 
 						WidgetRow = &CreateTransformComponentValueWidgetRow(
 							ChildControlElement->Settings.ControlType,
-							GetHierarchyBeingDebugged(),
 							ChildElementKeys,
 							TransformWidgetArgs,
 							Category,
@@ -2956,7 +2933,6 @@ void FRigControlElementDetails::CustomizeAnimationChannels(IDetailLayoutBuilder&
 						.Visibility(EVisibility::Visible);
 
 						WidgetRow = &CreateEulerTransformValueWidgetRow(
-							GetHierarchyBeingDebugged(),
 							ChildElementKeys,
 							TransformWidgetArgs,
 							Category,
@@ -2981,19 +2957,16 @@ void FRigControlElementDetails::CustomizeAnimationChannels(IDetailLayoutBuilder&
 				if(WidgetRow)
 				{
 					WidgetRow->AddCustomContextMenuAction(FUIAction(
-						FExecuteAction::CreateLambda([this, ChildElementKeys]()
+						FExecuteAction::CreateLambda([this, ChildElementKeys, HierarchyToChange]()
 						{
-							if(BlueprintBeingCustomized)
+							if(URigHierarchyController* Controller = HierarchyToChange->GetController(true))
 							{
-								if(URigHierarchyController* Controller = BlueprintBeingCustomized->GetHierarchyController())
+								FScopedTransaction(LOCTEXT("DeleteAnimationChannels", "Delete Animation Channels"));
+								HierarchyToChange->Modify();
+								
+								for(const FRigElementKey& KeyToRemove : ChildElementKeys)
 								{
-									FScopedTransaction(LOCTEXT("DeleteAnimationChannels", "Delete Animation Channels"));
-									BlueprintBeingCustomized->Hierarchy->Modify();
-									
-									for(const FRigElementKey& KeyToRemove : ChildElementKeys)
-									{
-										Controller->RemoveElement(KeyToRemove, true, true);
-									}
+									Controller->RemoveElement(KeyToRemove, true, true);
 								}
 							}
 						})),
@@ -3019,6 +2992,7 @@ void FRigControlElementDetails::CustomizeAnimationChannels(IDetailLayoutBuilder&
 			.Padding(FMargin(0.f, 0.f, 0.f, 0.f))
 			[
 				SNew(STextBlock)
+				.IsEnabled(bIsEnabled)
 				.Font(IDetailLayoutBuilder::GetDetailFont())
 				.Text(LOCTEXT("NoAnimationChannels", "No animation channels"))
 			]
@@ -3028,20 +3002,30 @@ void FRigControlElementDetails::CustomizeAnimationChannels(IDetailLayoutBuilder&
 
 FReply FRigControlElementDetails::OnAddAnimationChannelClicked()
 {
+	if(!IsAnyElementNotOfType(ERigElementType::Control) || IsAnyElementProcedural())
+	{
+		return FReply::Handled();
+	}
+
+	const FRigElementKey& Key = PerElementInfos[0].GetElement()->GetKey();
+	URigHierarchy* HierarchyToChange = PerElementInfos[0].GetDefaultHierarchy();
+
 	static const FString ChannelName = TEXT("Channel");
 	FRigControlSettings Settings;
 	Settings.AnimationType = ERigControlAnimationType::AnimationChannel;
 	Settings.ControlType = ERigControlType::Float;
 	Settings.MinimumValue = FRigControlValue::Make<float>(0.f);
 	Settings.MaximumValue = FRigControlValue::Make<float>(1.f);
-	Settings.DisplayName = HierarchyBeingCustomized->GetSafeNewDisplayName(ControlElements[0].GetKey(), ChannelName);
-	HierarchyBeingCustomized->GetController(true)->AddAnimationChannel(*ChannelName, ControlElements[0].GetKey(), Settings, true, true);
+	Settings.DisplayName = HierarchyToChange->GetSafeNewDisplayName(Key, ChannelName);
+	HierarchyToChange->GetController(true)->AddAnimationChannel(*ChannelName, Key, Settings, true, true);
 	return FReply::Handled();
 }
 
 TSharedRef<ITableRow> FRigControlElementDetails::HandleGenerateAnimationChannelTypeRow(TSharedPtr<ERigControlType> ControlType, const TSharedRef<STableViewBase>& OwnerTable, FRigElementKey ControlKey)
 {
-	TPair<const FSlateBrush*, FSlateColor> BrushAndColor = SRigHierarchyItem::GetBrushForElementType(HierarchyBeingCustomized, ControlKey);
+	const URigHierarchy* HierarchyToChange = PerElementInfos[0].GetDefaultHierarchy();
+
+	TPair<const FSlateBrush*, FSlateColor> BrushAndColor = SRigHierarchyItem::GetBrushForElementType(HierarchyToChange, ControlKey);
 	BrushAndColor.Value = SRigHierarchyItem::GetColorForControlType(*ControlType.Get(), nullptr);
 
 	return SNew(STableRow<TSharedPtr<ERigControlType>>, OwnerTable)
@@ -3081,173 +3065,171 @@ void FRigControlElementDetails::HandleControlTypeChanged(TSharedPtr<ERigControlT
 
 void FRigControlElementDetails::HandleControlTypeChanged(ERigControlType ControlType, TArray<FRigElementKey> ControlKeys, const TSharedRef<IPropertyUtilities> PropertyUtilities)
 {
-	TArray<FRigControlElement*> ControlElementsInHierarchy = GetElementsInHierarchy<FRigControlElement>(ControlKeys);
-	if (this->HierarchyBeingCustomized && ControlElementsInHierarchy.Num() > 0)
+	for(FPerElementInfo& Info : PerElementInfos)
 	{
-		this->HierarchyBeingCustomized->Modify();
+		URigHierarchy* Hierarchy = Info.GetHierarchy();
+		URigHierarchy* HierarchyToChange = Info.GetDefaultHierarchy();
+		HierarchyToChange->Modify();
 		
-		for(int32 ControlIndex = 0; ControlIndex< ControlElementsInHierarchy.Num(); ControlIndex++)
-		{
-			FRigControlElement* ControlElement = ControlElementsInHierarchy[ControlIndex];
-			
-			FRigControlValue ValueToSet;
+		FRigControlElement* ControlElement = Info.GetDefaultElement<FRigControlElement>();
+		
+		FRigControlValue ValueToSet;
 
-			ControlElement->Settings.ControlType = ControlType;
-			ControlElement->Settings.LimitEnabled.Reset();
-			ControlElement->Settings.bGroupWithParentControl = false;
+		ControlElement->Settings.ControlType = ControlType;
+		ControlElement->Settings.LimitEnabled.Reset();
+		ControlElement->Settings.bGroupWithParentControl = false;
 
-			switch (ControlElement->Settings.ControlType)
+		switch (ControlElement->Settings.ControlType)
+		{
+			case ERigControlType::Bool:
 			{
-				case ERigControlType::Bool:
-				{
-					ControlElement->Settings.AnimationType = ERigControlAnimationType::AnimationChannel;
-					ValueToSet = FRigControlValue::Make<bool>(false);
-					ControlElement->Settings.bGroupWithParentControl = ControlElement->Settings.IsAnimatable();
-					break;
-				}
-				case ERigControlType::Float:
-				{
-					ValueToSet = FRigControlValue::Make<float>(0.f);
-					ControlElement->Settings.SetupLimitArrayForType(true);
-					ControlElement->Settings.MinimumValue = FRigControlValue::Make<float>(0.f);
-					ControlElement->Settings.MaximumValue = FRigControlValue::Make<float>(100.f);
-					ControlElement->Settings.bGroupWithParentControl = ControlElement->Settings.IsAnimatable();
-					break;
-				}
-				case ERigControlType::Integer:
-				{
-					ValueToSet = FRigControlValue::Make<int32>(0);
-					ControlElement->Settings.SetupLimitArrayForType(true);
-					ControlElement->Settings.MinimumValue = FRigControlValue::Make<int32>(0);
-					ControlElement->Settings.MaximumValue = FRigControlValue::Make<int32>(100);
-					ControlElement->Settings.bGroupWithParentControl = ControlElement->Settings.IsAnimatable();
-					break;
-				}
-				case ERigControlType::Vector2D:
-				{
-					ValueToSet = FRigControlValue::Make<FVector2D>(FVector2D::ZeroVector);
-					ControlElement->Settings.SetupLimitArrayForType(true);
-					ControlElement->Settings.MinimumValue = FRigControlValue::Make<FVector2D>(FVector2D::ZeroVector);
-					ControlElement->Settings.MaximumValue = FRigControlValue::Make<FVector2D>(FVector2D(100.f, 100.f));
-					ControlElement->Settings.bGroupWithParentControl = ControlElement->Settings.IsAnimatable();
-					break;
-				}
-				case ERigControlType::Position:
-				{
-					ValueToSet = FRigControlValue::Make<FVector>(FVector::ZeroVector);
-					ControlElement->Settings.SetupLimitArrayForType(false);
-					ControlElement->Settings.MinimumValue = FRigControlValue::Make<FVector>(-FVector::OneVector);
-					ControlElement->Settings.MaximumValue = FRigControlValue::Make<FVector>(FVector::OneVector);
-					break;
-				}
-				case ERigControlType::Scale:
-				{
-					ValueToSet = FRigControlValue::Make<FVector>(FVector::OneVector);
-					ControlElement->Settings.SetupLimitArrayForType(false);
-					ControlElement->Settings.MinimumValue = FRigControlValue::Make<FVector>(FVector::ZeroVector);
-					ControlElement->Settings.MaximumValue = FRigControlValue::Make<FVector>(FVector::OneVector);
-					break;
-				}
-				case ERigControlType::Rotator:
-				{
-					ValueToSet = FRigControlValue::Make<FRotator>(FRotator::ZeroRotator);
-					ControlElement->Settings.SetupLimitArrayForType(false, false);
-					ControlElement->Settings.MinimumValue = FRigControlValue::Make<FRotator>(FRotator::ZeroRotator);
-					ControlElement->Settings.MaximumValue = FRigControlValue::Make<FRotator>(FRotator(180.f, 180.f, 180.f));
-					break;
-				}
-				case ERigControlType::Transform:
-				{
-					ValueToSet = FRigControlValue::Make<FTransform>(FTransform::Identity);
-					ControlElement->Settings.SetupLimitArrayForType(false, false, false);
-					ControlElement->Settings.MinimumValue = ValueToSet;
-					ControlElement->Settings.MaximumValue = ValueToSet;
-					break;
-				}
-				case ERigControlType::TransformNoScale:
-				{
-					FTransformNoScale Identity = FTransform::Identity;
-					ValueToSet = FRigControlValue::Make<FTransformNoScale>(Identity);
-					ControlElement->Settings.SetupLimitArrayForType(false, false, false);
-					ControlElement->Settings.MinimumValue = ValueToSet;
-					ControlElement->Settings.MaximumValue = ValueToSet;
-					break;
-				}
-				case ERigControlType::EulerTransform:
-				{
-					FEulerTransform Identity = FEulerTransform::Identity;
-					ValueToSet = FRigControlValue::Make<FEulerTransform>(Identity);
-					ControlElement->Settings.SetupLimitArrayForType(false, false, false);
-					ControlElement->Settings.MinimumValue = ValueToSet;
-					ControlElement->Settings.MaximumValue = ValueToSet;
-					break;
-				}
-				default:
-				{
-					ensure(false);
-					break;
-				}
+				ControlElement->Settings.AnimationType = ERigControlAnimationType::AnimationChannel;
+				ValueToSet = FRigControlValue::Make<bool>(false);
+				ControlElement->Settings.bGroupWithParentControl = ControlElement->Settings.IsAnimatable();
+				break;
 			}
-
-			this->HierarchyBeingCustomized->SetControlSettings(ControlElement, ControlElement->Settings, true, true, true);
-			this->HierarchyBeingCustomized->SetControlValue(ControlElement, ValueToSet, ERigControlValueType::Initial, true, false, true);
-			this->HierarchyBeingCustomized->SetControlValue(ControlElement, ValueToSet, ERigControlValueType::Current, true, false, true);
-
-			if(ObjectsBeingCustomized.IsValidIndex(ControlIndex))
+			case ERigControlType::Float:
 			{
-				if(ObjectsBeingCustomized[ControlIndex]->GetContent<FRigBaseElement>().GetKey() == ControlElement->GetKey())
-				{
-					ObjectsBeingCustomized[ControlIndex]->SetContent<FRigControlElement>(*ControlElement);
-				}
+				ValueToSet = FRigControlValue::Make<float>(0.f);
+				ControlElement->Settings.SetupLimitArrayForType(true);
+				ControlElement->Settings.MinimumValue = FRigControlValue::Make<float>(0.f);
+				ControlElement->Settings.MaximumValue = FRigControlValue::Make<float>(100.f);
+				ControlElement->Settings.bGroupWithParentControl = ControlElement->Settings.IsAnimatable();
+				break;
 			}
-
-			if (this->HierarchyBeingCustomized != this->BlueprintBeingCustomized->Hierarchy)
+			case ERigControlType::Integer:
 			{
-				if(FRigControlElement* OtherControlElement = this->BlueprintBeingCustomized->Hierarchy->Find<FRigControlElement>(ControlElement->GetKey()))
-				{
-					OtherControlElement->Settings = ControlElement->Settings;
-					this->BlueprintBeingCustomized->Hierarchy->SetControlSettings(OtherControlElement, OtherControlElement->Settings, true, true, true);
-					this->BlueprintBeingCustomized->Hierarchy->SetControlValue(OtherControlElement, ValueToSet, ERigControlValueType::Initial, true);
-					this->BlueprintBeingCustomized->Hierarchy->SetControlValue(OtherControlElement, ValueToSet, ERigControlValueType::Current, true);
-				}
+				ValueToSet = FRigControlValue::Make<int32>(0);
+				ControlElement->Settings.SetupLimitArrayForType(true);
+				ControlElement->Settings.MinimumValue = FRigControlValue::Make<int32>(0);
+				ControlElement->Settings.MaximumValue = FRigControlValue::Make<int32>(100);
+				ControlElement->Settings.bGroupWithParentControl = ControlElement->Settings.IsAnimatable();
+				break;
 			}
-			else
+			case ERigControlType::Vector2D:
+			{
+				ValueToSet = FRigControlValue::Make<FVector2D>(FVector2D::ZeroVector);
+				ControlElement->Settings.SetupLimitArrayForType(true);
+				ControlElement->Settings.MinimumValue = FRigControlValue::Make<FVector2D>(FVector2D::ZeroVector);
+				ControlElement->Settings.MaximumValue = FRigControlValue::Make<FVector2D>(FVector2D(100.f, 100.f));
+				ControlElement->Settings.bGroupWithParentControl = ControlElement->Settings.IsAnimatable();
+				break;
+			}
+			case ERigControlType::Position:
+			{
+				ValueToSet = FRigControlValue::Make<FVector>(FVector::ZeroVector);
+				ControlElement->Settings.SetupLimitArrayForType(false);
+				ControlElement->Settings.MinimumValue = FRigControlValue::Make<FVector>(-FVector::OneVector);
+				ControlElement->Settings.MaximumValue = FRigControlValue::Make<FVector>(FVector::OneVector);
+				break;
+			}
+			case ERigControlType::Scale:
 			{
-				this->BlueprintBeingCustomized->PropagateHierarchyFromBPToInstances();
+				ValueToSet = FRigControlValue::Make<FVector>(FVector::OneVector);
+				ControlElement->Settings.SetupLimitArrayForType(false);
+				ControlElement->Settings.MinimumValue = FRigControlValue::Make<FVector>(FVector::ZeroVector);
+				ControlElement->Settings.MaximumValue = FRigControlValue::Make<FVector>(FVector::OneVector);
+				break;
+			}
+			case ERigControlType::Rotator:
+			{
+				ValueToSet = FRigControlValue::Make<FRotator>(FRotator::ZeroRotator);
+				ControlElement->Settings.SetupLimitArrayForType(false, false);
+				ControlElement->Settings.MinimumValue = FRigControlValue::Make<FRotator>(FRotator::ZeroRotator);
+				ControlElement->Settings.MaximumValue = FRigControlValue::Make<FRotator>(FRotator(180.f, 180.f, 180.f));
+				break;
+			}
+			case ERigControlType::Transform:
+			{
+				ValueToSet = FRigControlValue::Make<FTransform>(FTransform::Identity);
+				ControlElement->Settings.SetupLimitArrayForType(false, false, false);
+				ControlElement->Settings.MinimumValue = ValueToSet;
+				ControlElement->Settings.MaximumValue = ValueToSet;
+				break;
+			}
+			case ERigControlType::TransformNoScale:
+			{
+				FTransformNoScale Identity = FTransform::Identity;
+				ValueToSet = FRigControlValue::Make<FTransformNoScale>(Identity);
+				ControlElement->Settings.SetupLimitArrayForType(false, false, false);
+				ControlElement->Settings.MinimumValue = ValueToSet;
+				ControlElement->Settings.MaximumValue = ValueToSet;
+				break;
+			}
+			case ERigControlType::EulerTransform:
+			{
+				FEulerTransform Identity = FEulerTransform::Identity;
+				ValueToSet = FRigControlValue::Make<FEulerTransform>(Identity);
+				ControlElement->Settings.SetupLimitArrayForType(false, false, false);
+				ControlElement->Settings.MinimumValue = ValueToSet;
+				ControlElement->Settings.MaximumValue = ValueToSet;
+				break;
+			}
+			default:
+			{
+				ensure(false);
+				break;
 			}
 		}
-		
-		PropertyUtilities->ForceRefresh();
+
+		HierarchyToChange->SetControlSettings(ControlElement, ControlElement->Settings, true, true, true);
+		HierarchyToChange->SetControlValue(ControlElement, ValueToSet, ERigControlValueType::Initial, true, false, true);
+		HierarchyToChange->SetControlValue(ControlElement, ValueToSet, ERigControlValueType::Current, true, false, true);
+
+		Info.WrapperObject->SetContent<FRigControlElement>(*ControlElement);
+
+		if (HierarchyToChange != Hierarchy)
+		{
+			if(FRigControlElement* OtherControlElement = Info.GetElement<FRigControlElement>())
+			{
+				OtherControlElement->Settings = ControlElement->Settings;
+				Hierarchy->SetControlSettings(OtherControlElement, OtherControlElement->Settings, true, true, true);
+				Hierarchy->SetControlValue(OtherControlElement, ValueToSet, ERigControlValueType::Initial, true);
+				Hierarchy->SetControlValue(OtherControlElement, ValueToSet, ERigControlValueType::Current, true);
+			}
+		}
+		else
+		{
+			Info.GetBlueprint()->PropagateHierarchyFromBPToInstances();
+		}
 	}
+	
+	PropertyUtilities->ForceRefresh();
 }
 
 void FRigControlElementDetails::CustomizeShape(IDetailLayoutBuilder& DetailBuilder)
 {
-	for(TWeakObjectPtr<UDetailsViewWrapperObject> ObjectBeingCustomized : ObjectsBeingCustomized)
+	if(PerElementInfos.IsEmpty())
 	{
-		if(ObjectBeingCustomized.IsValid())
+		return;
+	}
+
+	if(ContainsElementByPredicate([](const FPerElementInfo& Info)
+	{
+		if(const FRigControlElement* ControlElement = Info.GetElement<FRigControlElement>())
 		{
-			if(ObjectBeingCustomized->IsChildOf<FRigControlElement>())
-			{
-				const FRigControlElement ControlElement = ObjectBeingCustomized->GetContent<FRigControlElement>();
-				if(!ControlElement.Settings.SupportsShape())
-				{
-					return;
-				}
-			}
+			return !ControlElement->Settings.SupportsShape();
 		}
+		return true;
+	}))
+	{
+		return;
 	}
 
+	const bool bIsProcedural = IsAnyElementProcedural();
+	const bool bIsEnabled = !bIsProcedural;
+
 	TSharedPtr<IPropertyHandle> ShapeHandle = DetailBuilder.GetProperty(TEXT("Shape"));
 	TSharedPtr<IPropertyHandle> InitialHandle = ShapeHandle->GetChildHandle(TEXT("Initial"));
 	TSharedPtr<IPropertyHandle> LocalHandle = InitialHandle->GetChildHandle(TEXT("Local"));
 	ShapeTransformHandle = LocalHandle->GetChildHandle(TEXT("Transform"));
 	
 	ShapeNameList.Reset();
-	if (BlueprintBeingCustomized)
+	
+	if (UControlRigBlueprint* Blueprint = PerElementInfos[0].GetBlueprint())
 	{
-		const bool bUseNameSpace = BlueprintBeingCustomized->ShapeLibraries.Num() > 1;
-		for(TSoftObjectPtr<UControlRigShapeLibrary>& ShapeLibrary : BlueprintBeingCustomized->ShapeLibraries)
+		const bool bUseNameSpace = Blueprint->ShapeLibraries.Num() > 1;
+		for(TSoftObjectPtr<UControlRigShapeLibrary>& ShapeLibrary : Blueprint->ShapeLibraries)
 		{
 			if (!ShapeLibrary.IsValid())
 			{
@@ -3272,14 +3254,18 @@ void FRigControlElementDetails::CustomizeShape(IDetailLayoutBuilder& DetailBuild
 	if(!IsAnyControlNotOfAnimationType(ERigControlAnimationType::ProxyControl))
 	{
 		ShapeCategory.AddProperty(SettingsHandle->GetChildHandle(TEXT("ShapeVisibility")).ToSharedRef())
+		.IsEnabled(bIsEnabled)
 		.DisplayName(FText::FromString(TEXT("Visibility Mode")));
 	}
 
 	ShapeCategory.AddProperty(SettingsHandle->GetChildHandle(TEXT("bShapeVisible")).ToSharedRef())
+	.IsEnabled(bIsEnabled)
 	.DisplayName(FText::FromString(TEXT("Visible")));
 
 	IDetailGroup& ShapePropertiesGroup = ShapeCategory.AddGroup(TEXT("Shape Properties"), LOCTEXT("ShapeProperties", "Shape Properties"));
-	ShapePropertiesGroup.HeaderRow().NameContent()
+	ShapePropertiesGroup.HeaderRow()
+	.IsEnabled(bIsEnabled)
+	.NameContent()
 	[
 		SNew(STextBlock)
 		.Font(IDetailLayoutBuilder::GetDetailFont())
@@ -3289,25 +3275,29 @@ void FRigControlElementDetails::CustomizeShape(IDetailLayoutBuilder& DetailBuild
 	.CopyAction(FUIAction(
 		FExecuteAction::CreateSP(this, &FRigControlElementDetails::OnCopyShapeProperties)))
 	.PasteAction(FUIAction(
-		FExecuteAction::CreateSP(this, &FRigControlElementDetails::OnPasteShapeProperties)));
+		FExecuteAction::CreateSP(this, &FRigControlElementDetails::OnPasteShapeProperties),
+		FCanExecuteAction::CreateLambda([bIsEnabled]() { return bIsEnabled; })));
 	
 	// setup shape transform
 	SAdvancedTransformInputBox<FEulerTransform>::FArguments TransformWidgetArgs = SAdvancedTransformInputBox<FEulerTransform>::FArguments()
+	.IsEnabled(bIsEnabled)
 	.DisplayToggle(false)
 	.DisplayRelativeWorld(false)
 	.Font(IDetailLayoutBuilder::GetDetailFont());
 
 	TArray<FRigElementKey> Keys = GetElementKeys();
-	URigHierarchy* HierarchyBeingDebugged = GetHierarchyBeingDebugged();
-	Keys = HierarchyBeingDebugged->SortKeys(Keys);
+	Keys = PerElementInfos[0].GetHierarchy()->SortKeys(Keys);
 
-	auto GetShapeTransform = [HierarchyBeingDebugged](
+	auto GetShapeTransform = [this](
 		const FRigElementKey& Key
 		) -> FEulerTransform
 	{
-		if(FRigControlElement* ControlElement = HierarchyBeingDebugged->Find<FRigControlElement>(Key))
+		if(const FPerElementInfo& Info = FindElement(Key))
 		{
-			return FEulerTransform(HierarchyBeingDebugged->GetControlShapeTransform(ControlElement, ERigTransformType::InitialLocal));
+			if(FRigControlElement* ControlElement = Info.GetElement<FRigControlElement>())
+			{
+				return FEulerTransform(Info.GetHierarchy()->GetControlShapeTransform(ControlElement, ERigTransformType::InitialLocal));
+			}
 		}
 		return FEulerTransform::Identity;
 	};
@@ -3318,13 +3308,16 @@ void FRigControlElementDetails::CustomizeShape(IDetailLayoutBuilder& DetailBuild
 		bool bSetupUndo
 		)
 	{
-		if(FRigControlElement* ControlElement = HierarchyBeingCustomized->Find<FRigControlElement>(Key))
+		if(const FPerElementInfo& Info = FindElement(Key))
 		{
-			HierarchyBeingCustomized->SetControlShapeTransform(ControlElement, InTransform.ToFTransform(), ERigTransformType::InitialLocal, bSetupUndo, true, bSetupUndo);
+			if(const FRigControlElement* ControlElement = Info.GetDefaultElement<FRigControlElement>())
+			{
+				Info.GetDefaultHierarchy()->SetControlShapeTransform((FRigControlElement*)ControlElement, InTransform.ToFTransform(), ERigTransformType::InitialLocal, bSetupUndo, true, bSetupUndo);
+			}
 		}
 	};
 
-	TransformWidgetArgs.OnGetNumericValue_Lambda([Keys, HierarchyBeingDebugged, GetShapeTransform](
+	TransformWidgetArgs.OnGetNumericValue_Lambda([Keys, GetShapeTransform](
 		ESlateTransformComponent::Type Component,
 		ESlateRotationRepresentation::Type Representation,
 		ESlateTransformSubComponent::Type SubComponent) -> TOptional<FVector::FReal>
@@ -3381,7 +3374,7 @@ void FRigControlElementDetails::CustomizeShape(IDetailLayoutBuilder& DetailBuild
 				if(!SliderTransaction.IsValid())
 				{
 					SliderTransaction = MakeShareable(new FScopedTransaction(NSLOCTEXT("ControlRigElementDetails", "ChangeNumericValue", "Change Numeric Value")));
-					HierarchyBeingCustomized->Modify();
+					PerElementInfos[0].GetDefaultHierarchy()->Modify();
 				}
 				SetShapeTransform(Key, Transform, false);
 			}
@@ -3403,7 +3396,7 @@ void FRigControlElementDetails::CustomizeShape(IDetailLayoutBuilder& DetailBuild
 	{
 		{
 			FScopedTransaction Transaction(LOCTEXT("ChangeNumericValue", "Change Numeric Value"));
-			this->HierarchyBeingCustomized->Modify();
+			PerElementInfos[0].GetDefaultHierarchy()->Modify();
 
 			for(const FRigElementKey& Key : Keys)
 			{
@@ -3484,7 +3477,7 @@ void FRigControlElementDetails::CustomizeShape(IDetailLayoutBuilder& DetailBuild
 		}
 
 		FScopedTransaction Transaction(LOCTEXT("PasteTransform", "Paste Transform"));
-		HierarchyBeingCustomized->Modify();
+		PerElementInfos[0].GetDefaultHierarchy()->Modify();
 
 		for(const FRigElementKey& Key : Keys)
 		{
@@ -3606,7 +3599,7 @@ void FRigControlElementDetails::CustomizeShape(IDetailLayoutBuilder& DetailBuild
 		ESlateTransformComponent::Type InComponent)
 	{
 		FScopedTransaction Transaction(LOCTEXT("ResetTransformToDefault", "Reset Transform to Default"));
-		HierarchyBeingCustomized->Modify();
+		PerElementInfos[0].GetDefaultHierarchy()->Modify();
 
 		for(const FRigElementKey& Key : Keys)
 		{
@@ -3641,6 +3634,12 @@ void FRigControlElementDetails::CustomizeShape(IDetailLayoutBuilder& DetailBuild
 		}
 	});
 
+	TArray<FRigControlElement*> ControlElements;
+	Algo::Transform(PerElementInfos, ControlElements, [](const FPerElementInfo& Info)
+	{
+		return Info.GetElement<FRigControlElement>();
+	});
+
 	SAdvancedTransformInputBox<FEulerTransform>::ConstructGroupedTransformRows(
 		ShapeCategory, 
 		LOCTEXT("ShapeTransform", "Shape Transform"), 
@@ -3652,19 +3651,21 @@ void FRigControlElementDetails::CustomizeShape(IDetailLayoutBuilder& DetailBuild
 	.NameContent()
 	[
 		SNew(STextBlock)
+		.IsEnabled(bIsEnabled)
 		.Text(FText::FromString(TEXT("Shape")))
 		.Font(IDetailLayoutBuilder::GetDetailFont())
 		.IsEnabled(this, &FRigControlElementDetails::IsShapeEnabled)
 	]
 	.ValueContent()
 	[
-		SAssignNew(ShapeNameListWidget, SControlRigShapeNameList, ControlElements, BlueprintBeingCustomized)
+		SAssignNew(ShapeNameListWidget, SControlRigShapeNameList, ControlElements, PerElementInfos[0].GetBlueprint())
 		.OnGetNameListContent(this, &FRigControlElementDetails::GetShapeNameList)
 		.IsEnabled(this, &FRigControlElementDetails::IsShapeEnabled)
 	];
 
 	ShapeColorHandle = SettingsHandle->GetChildHandle(TEXT("ShapeColor"));
 	ShapePropertiesGroup.AddPropertyRow(ShapeColorHandle.ToSharedRef())
+	.IsEnabled(bIsEnabled)
 	.DisplayName(FText::FromString(TEXT("Color")));
 }
 
@@ -3698,18 +3699,19 @@ void FRigControlElementDetails::RegisterSectionMappings(FPropertyEditorModule& P
 
 bool FRigControlElementDetails::IsShapeEnabled() const
 {
-	for(TWeakObjectPtr<UDetailsViewWrapperObject> ObjectBeingCustomized : ObjectsBeingCustomized)
+	if(IsAnyElementProcedural())
+	{
+		 return false;
+	}
+	
+	return ContainsElementByPredicate([](const FPerElementInfo& Info)
 	{
-		if(ObjectBeingCustomized.IsValid())
+		if(const FRigControlElement* ControlElement = Info.GetElement<FRigControlElement>())
 		{
-			if(ObjectBeingCustomized->IsChildOf<FRigControlElement>())
-			{
-				const FRigControlElement ControlElement = ObjectBeingCustomized->GetContent<FRigControlElement>();
-				return ControlElement.Settings.SupportsShape();
-			}
+			return ControlElement->Settings.SupportsShape();
 		}
-	}
-	return false;
+		return false;
+	});
 }
 
 const TArray<TSharedPtr<FString>>& FRigControlElementDetails::GetShapeNameList() const
@@ -3721,28 +3723,23 @@ FText FRigControlElementDetails::GetDisplayName() const
 {
 	FName DisplayName(NAME_None);
 
-	for(int32 ObjectIndex = 0; ObjectIndex < ObjectsBeingCustomized.Num(); ObjectIndex++)
+	for(int32 ObjectIndex = 0; ObjectIndex < PerElementInfos.Num(); ObjectIndex++)
 	{
-		TWeakObjectPtr<UDetailsViewWrapperObject> ObjectBeingCustomized = ObjectsBeingCustomized[ObjectIndex];
-		if(ObjectBeingCustomized.IsValid())
+		const FPerElementInfo& Info = PerElementInfos[ObjectIndex];
+		if(const FRigControlElement* ControlElement = Info.GetElement<FRigControlElement>())
 		{
-			if(ObjectBeingCustomized->IsChildOf<FRigControlElement>())
-			{
-				const FRigControlElement ControlElement = ObjectBeingCustomized->GetContent<FRigControlElement>();
-
- 				const FName ThisDisplayName =
- 					(ControlElement.IsAnimationChannel()) ?
- 					ControlElement.GetDisplayName() :
-					ControlElement.Settings.DisplayName;
+ 			const FName ThisDisplayName =
+ 				(ControlElement->IsAnimationChannel()) ?
+ 				ControlElement->GetDisplayName() :
+				ControlElement->Settings.DisplayName;
 
-				if(ObjectIndex == 0)
-				{
-					DisplayName = ThisDisplayName;
-				}
-				else if(DisplayName != ThisDisplayName)
-				{
-					return ControlRigDetailsMultipleValues;
-				}
+			if(ObjectIndex == 0)
+			{
+				DisplayName = ThisDisplayName;
+			}
+			else if(DisplayName != ThisDisplayName)
+			{
+				return ControlRigDetailsMultipleValues;
 			}
 		}
 	}
@@ -3756,16 +3753,12 @@ FText FRigControlElementDetails::GetDisplayName() const
 
 void FRigControlElementDetails::SetDisplayName(const FText& InNewText, ETextCommit::Type InCommitType)
 {
-	for(int32 ObjectIndex = 0; ObjectIndex < ObjectsBeingCustomized.Num(); ObjectIndex++)
+	for(int32 ObjectIndex = 0; ObjectIndex < PerElementInfos.Num(); ObjectIndex++)
 	{
-		TWeakObjectPtr<UDetailsViewWrapperObject> ObjectBeingCustomized = ObjectsBeingCustomized[ObjectIndex];
-		if(ObjectBeingCustomized.IsValid())
+		const FPerElementInfo& Info = PerElementInfos[ObjectIndex];
+		if(const FRigControlElement* ControlElement = Info.GetElement<FRigControlElement>())
 		{
-			if(ObjectBeingCustomized->IsChildOf<FRigControlElement>())
-			{
-				FRigControlElement ControlElement = ObjectBeingCustomized->GetContent<FRigControlElement>();
-				SetDisplayNameForElement(InNewText, InCommitType, ControlElement.GetKey());
-			}
+			SetDisplayNameForElement(InNewText, InCommitType, ControlElement->GetKey());
 		}
 	}
 }
@@ -3777,41 +3770,20 @@ void FRigControlElementDetails::SetDisplayNameForElement(const FText& InNewText,
 		return;
 	}
 
-	const FName DisplayName = InNewText.IsEmpty() ? FName(NAME_None) : FName(*InNewText.ToString());
-
-	if(HierarchyBeingCustomized)
+	const FPerElementInfo& Info = FindElement(InKeyToRename);
+	if(!Info.IsValid())
 	{
-		if(const FRigControlElement* ControlElement = HierarchyBeingCustomized->Find<FRigControlElement>(InKeyToRename))
-		{
-			if(ControlElement->GetDisplayName() == DisplayName)
-			{
-				return;
-			}
-		}
+		return;
 	}
 
-	for(int32 ObjectIndex = 0; ObjectIndex < ObjectsBeingCustomized.Num(); ObjectIndex++)
+	if(Info.IsProcedural())
 	{
-		TWeakObjectPtr<UDetailsViewWrapperObject> ObjectBeingCustomized = ObjectsBeingCustomized[ObjectIndex];
-		if(ObjectBeingCustomized.IsValid())
-		{
-			if(ObjectBeingCustomized->IsChildOf<FRigControlElement>())
-			{
-				FRigControlElement ControlElement = ObjectBeingCustomized->GetContent<FRigControlElement>();
-				if(ControlElement.GetKey() == InKeyToRename)
-				{
-					ControlElement.Settings.DisplayName = DisplayName;
-					ObjectBeingCustomized->SetContent<FRigControlElement>(ControlElement);
-				}
-			}
-		}
+		return;
 	}
 
-	if(HierarchyBeingCustomized)
-	{
-		FScopedTransaction Transaction(LOCTEXT("SetElementDisplayName", "Set Element Display Name"));
-		HierarchyBeingCustomized->GetController(true)->SetDisplayName(InKeyToRename, DisplayName, true, true);
-	}
+	const FName DisplayName = InNewText.IsEmpty() ? FName(NAME_None) : FName(*InNewText.ToString());
+	const bool bRename = IsAnyControlOfAnimationType(ERigControlAnimationType::AnimationChannel);
+	Info.GetDefaultHierarchy()->GetController(true)->SetDisplayName(InKeyToRename, DisplayName, bRename, true, true);
 }
 
 bool FRigControlElementDetails::OnVerifyDisplayNameChanged(const FText& InText, FText& OutErrorMessage, const FRigElementKey& InKeyToRename)
@@ -3823,10 +3795,21 @@ bool FRigControlElementDetails::OnVerifyDisplayNameChanged(const FText& InText,
 		return false;
 	}
 
+	const FPerElementInfo& Info = FindElement(InKeyToRename);
+	if(!Info.IsValid())
+	{
+		return false;
+	}
+
+	if(Info.IsProcedural())
+	{
+		return false;
+	}
+
 	// make sure there is no duplicate
-	if(URigHierarchy* Hierarchy = GetHierarchy())
+	if(const URigHierarchy* Hierarchy = Info.GetDefaultHierarchy())
 	{
-		if(const FRigControlElement* ControlElement = Hierarchy->Find<FRigControlElement>(InKeyToRename))
+		if(const FRigControlElement* ControlElement = Info.GetDefaultElement<FRigControlElement>())
 		{
 			if(const FRigBaseElement* ParentElement = Hierarchy->GetFirstParent(ControlElement))
 			{
@@ -3845,19 +3828,15 @@ bool FRigControlElementDetails::OnVerifyDisplayNameChanged(const FText& InText,
 void FRigControlElementDetails::OnCopyShapeProperties()
 {
 	FString Value;
-	if (ObjectsBeingCustomized.Num() > 0)
+
+	if (!PerElementInfos.IsEmpty())
 	{
-		if(ObjectsBeingCustomized[0].IsValid())
+		if(const FRigControlElement* ControlElement = PerElementInfos[0].GetElement<FRigControlElement>())
 		{
-			if(ObjectsBeingCustomized[0]->IsChildOf<FRigControlElement>())
-			{
-				const FRigControlElement ControlElement = ObjectsBeingCustomized[0]->GetContent<FRigControlElement>();
-
-				Value = FString::Printf(TEXT("(ShapeName=\"%s\",ShapeColor=%s,Transform=%s)"),
-					*ControlElement.Settings.ShapeName.ToString(),
-					*ControlElement.Settings.ShapeColor.ToString(),
-					*ControlElement.Shape.Initial.Local.Transform.ToString());
-			}
+			Value = FString::Printf(TEXT("(ShapeName=\"%s\",ShapeColor=%s,Transform=%s)"),
+				*ControlElement->Settings.ShapeName.ToString(),
+				*ControlElement->Settings.ShapeColor.ToString(),
+				*ControlElement->Shape.Initial.Local.Transform.ToString());
 		}
 	}
 		
@@ -3937,21 +3916,26 @@ FDetailWidgetRow& FRigControlElementDetails::CreateBoolValueWidgetRow(
 	TAttribute<EVisibility> Visibility,
 	TSharedPtr<SWidget> NameContent)
 {
-	const bool bCurrent = ValueType == ERigControlValueType::Current;
-	const bool bInitial = ValueType == ERigControlValueType::Initial;
-
-	URigHierarchy* HierarchyBeingDebugged = GetHierarchyBeingDebugged();
-	URigHierarchy* HierarchyToChange = bCurrent ? HierarchyBeingDebugged : HierarchyBeingCustomized;
-
 	const static TCHAR* TrueText = TEXT("True");
 	const static TCHAR* FalseText = TEXT("False");
+	
+	const bool bIsProcedural = IsAnyElementProcedural();
+	const bool bIsEnabled = !bIsProcedural || ValueType == ERigControlValueType::Current;
+
+	URigHierarchy* Hierarchy = PerElementInfos[0].GetHierarchy();
+	URigHierarchy* HierarchyToChange = PerElementInfos[0].GetDefaultHierarchy();
+	if(ValueType == ERigControlValueType::Current)
+	{
+		HierarchyToChange = Hierarchy;
+	}
 
 	if(!NameContent.IsValid())
 	{
 		SAssignNew(NameContent, STextBlock)
 		.Text(Label)
 		.ToolTipText(Tooltip)
-		.Font(IDetailLayoutBuilder::GetDetailFont());
+		.Font(IDetailLayoutBuilder::GetDetailFont())
+		.IsEnabled(bIsEnabled);
 	}
 
 	FDetailWidgetRow& WidgetRow = CategoryBuilder.AddCustomRow(Label)
@@ -3965,12 +3949,12 @@ FDetailWidgetRow& FRigControlElementDetails::CreateBoolValueWidgetRow(
 	.ValueContent()
 	[
 		SNew(SCheckBox)
-		.IsChecked_Lambda([ValueType, Keys, HierarchyBeingDebugged]() -> ECheckBoxState
+		.IsChecked_Lambda([ValueType, Keys, Hierarchy]() -> ECheckBoxState
 		{
-			const bool FirstValue = HierarchyBeingDebugged->GetControlValue<bool>(Keys[0], ValueType);
+			const bool FirstValue = Hierarchy->GetControlValue<bool>(Keys[0], ValueType);
 			for(int32 Index = 1; Index < Keys.Num(); Index++)
 			{
-				const bool SecondValue = HierarchyBeingDebugged->GetControlValue<bool>(Keys[Index], ValueType);
+				const bool SecondValue = Hierarchy->GetControlValue<bool>(Keys[Index], ValueType);
 				if(FirstValue != SecondValue)
 				{
 					return ECheckBoxState::Undetermined;
@@ -3993,11 +3977,12 @@ FDetailWidgetRow& FRigControlElementDetails::CreateBoolValueWidgetRow(
 				HierarchyToChange->SetControlValue(Key, FRigControlValue::Make<bool>(Value), ValueType, true, true); 
 			}
 		})
+		.IsEnabled(bIsEnabled)
 	]
 	.CopyAction(FUIAction(
-	FExecuteAction::CreateLambda([ValueType, Keys, HierarchyBeingDebugged]()
+	FExecuteAction::CreateLambda([ValueType, Keys, Hierarchy]()
 		{
-			const bool FirstValue = HierarchyBeingDebugged->GetControlValue<bool>(Keys[0], ValueType);
+			const bool FirstValue = Hierarchy->GetControlValue<bool>(Keys[0], ValueType);
 			FPlatformApplicationMisc::ClipboardCopy(FirstValue ? TrueText : FalseText);
 		}),
 		FCanExecuteAction())
@@ -4016,14 +4001,19 @@ FDetailWidgetRow& FRigControlElementDetails::CreateBoolValueWidgetRow(
 				HierarchyToChange->SetControlValue(Key, FRigControlValue::Make<bool>(Value), ValueType, true, true); 
 			}
 		}),
-		FCanExecuteAction())
+		FCanExecuteAction::CreateLambda([bIsEnabled]() { return bIsEnabled; }))
 	)
 	.OverrideResetToDefault(FResetToDefaultOverride::Create(
-		TAttribute<bool>::CreateLambda([ValueType, Keys, HierarchyBeingDebugged]() -> bool
+		TAttribute<bool>::CreateLambda([ValueType, Keys, Hierarchy, bIsEnabled]() -> bool
 		{
-			const bool FirstValue = HierarchyBeingDebugged->GetControlValue<bool>(Keys[0], ValueType);
+			if(!bIsEnabled)
+			{
+				return false;
+			}
+			
+			const bool FirstValue = Hierarchy->GetControlValue<bool>(Keys[0], ValueType);
 			const bool ReferenceValue = ValueType == ERigControlValueType::Initial ? false :
-				HierarchyBeingDebugged->GetControlValue<bool>(Keys[0], ERigControlValueType::Initial);
+				Hierarchy->GetControlValue<bool>(Keys[0], ERigControlValueType::Initial);
 
 			return FirstValue != ReferenceValue;
 		}),
@@ -4076,21 +4066,28 @@ FDetailWidgetRow& FRigControlElementDetails::CreateEnumValueWidgetRow(
 	TAttribute<EVisibility> Visibility,
 	TSharedPtr<SWidget> NameContent)
 {
-	const bool bCurrent = ValueType == ERigControlValueType::Current;
-	const bool bInitial = ValueType == ERigControlValueType::Initial;
-	
-	URigHierarchy* HierarchyBeingDebugged = GetHierarchyBeingDebugged();
-	URigHierarchy* HierarchyToChange = bCurrent ? HierarchyBeingDebugged : HierarchyBeingCustomized;
+	const bool bIsProcedural = IsAnyElementProcedural();
+	const bool bIsEnabled = !bIsProcedural || ValueType == ERigControlValueType::Current;
+
+	URigHierarchy* Hierarchy = PerElementInfos[0].GetHierarchy();
+	URigHierarchy* HierarchyToChange = PerElementInfos[0].GetDefaultHierarchy();
+	if(ValueType == ERigControlValueType::Current)
+	{
+		HierarchyToChange = Hierarchy;
+	}
 
 	UEnum* Enum = nullptr;
 	for(const FRigElementKey& Key : Keys)
 	{
-		if(const FRigControlElement* ControlElement = HierarchyBeingCustomized->Find<FRigControlElement>(Key))
+		if(const FPerElementInfo& Info = FindElement(Key))
 		{
-			Enum = ControlElement->Settings.ControlEnum.Get();
-			if(Enum)
+			if(const FRigControlElement* ControlElement = Info.GetElement<FRigControlElement>())
 			{
-				break;
+				Enum = ControlElement->Settings.ControlEnum.Get();
+				if(Enum)
+				{
+					break;
+				}
 			}
 		}
 	}
@@ -4102,7 +4099,8 @@ FDetailWidgetRow& FRigControlElementDetails::CreateEnumValueWidgetRow(
 		SAssignNew(NameContent, STextBlock)
 		.Text(Label)
 		.ToolTipText(Tooltip)
-		.Font(IDetailLayoutBuilder::GetDetailFont());
+		.Font(IDetailLayoutBuilder::GetDetailFont())
+		.IsEnabled(bIsEnabled);
 	}
 
 	FDetailWidgetRow& WidgetRow = CategoryBuilder.AddCustomRow(Label)
@@ -4116,12 +4114,12 @@ FDetailWidgetRow& FRigControlElementDetails::CreateEnumValueWidgetRow(
 	.ValueContent()
 	[
 		SNew(SEnumComboBox, Enum)
-		.CurrentValue_Lambda([ValueType, Keys, HierarchyBeingDebugged]() -> int32
+		.CurrentValue_Lambda([ValueType, Keys, Hierarchy]() -> int32
 		{
-			const int32 FirstValue = HierarchyBeingDebugged->GetControlValue<int32>(Keys[0], ValueType);
+			const int32 FirstValue = Hierarchy->GetControlValue<int32>(Keys[0], ValueType);
 			for(int32 Index = 1; Index < Keys.Num(); Index++)
 			{
-				const int32 SecondValue = HierarchyBeingDebugged->GetControlValue<int32>(Keys[Index], ValueType);
+				const int32 SecondValue = Hierarchy->GetControlValue<int32>(Keys[Index], ValueType);
 				if(FirstValue != SecondValue)
 				{
 					return INDEX_NONE;
@@ -4144,11 +4142,12 @@ FDetailWidgetRow& FRigControlElementDetails::CreateEnumValueWidgetRow(
 			}
 		})
 		.Font(FAppStyle::GetFontStyle(TEXT("MenuItem.Font")))
+		.IsEnabled(bIsEnabled)
 	]
 	.CopyAction(FUIAction(
-	FExecuteAction::CreateLambda([ValueType, Keys, HierarchyBeingDebugged]()
+	FExecuteAction::CreateLambda([ValueType, Keys, Hierarchy]()
 		{
-			const int32 FirstValue = HierarchyBeingDebugged->GetControlValue<int32>(Keys[0], ValueType);
+			const int32 FirstValue = Hierarchy->GetControlValue<int32>(Keys[0], ValueType);
 			FPlatformApplicationMisc::ClipboardCopy(*FString::FromInt(FirstValue));
 		}),
 		FCanExecuteAction())
@@ -4172,14 +4171,19 @@ FDetailWidgetRow& FRigControlElementDetails::CreateEnumValueWidgetRow(
 				HierarchyToChange->SetControlValue(Key, FRigControlValue::Make<int32>(Value), ValueType, true, true); 
 			}
 		}),
-		FCanExecuteAction())
+		FCanExecuteAction::CreateLambda([bIsEnabled]() { return bIsEnabled; }))
 	)
 	.OverrideResetToDefault(FResetToDefaultOverride::Create(
-		TAttribute<bool>::CreateLambda([ValueType, Keys, HierarchyBeingDebugged]() -> bool
+		TAttribute<bool>::CreateLambda([ValueType, Keys, Hierarchy, bIsEnabled]() -> bool
 		{
-			const int32 FirstValue = HierarchyBeingDebugged->GetControlValue<int32>(Keys[0], ValueType);
+			if(!bIsEnabled)
+			{
+				return false;
+			}
+			
+			const int32 FirstValue = Hierarchy->GetControlValue<int32>(Keys[0], ValueType);
 			const int32 ReferenceValue = ValueType == ERigControlValueType::Initial ? 0 :
-				HierarchyBeingDebugged->GetControlValue<int32>(Keys[0], ERigControlValueType::Initial);
+				Hierarchy->GetControlValue<int32>(Keys[0], ERigControlValueType::Initial);
 
 			return FirstValue != ReferenceValue;
 		}),
@@ -4208,12 +4212,16 @@ FDetailWidgetRow& FRigControlElementDetails::CreateVector2DValueWidgetRow(
 	TAttribute<EVisibility> Visibility,
 	TSharedPtr<SWidget> NameContent)
 {
-	const bool bCurrent = ValueType == ERigControlValueType::Current;
-	const bool bInitial = ValueType == ERigControlValueType::Initial;
+	const bool bIsProcedural = IsAnyElementProcedural();
+	const bool bIsEnabled = !bIsProcedural || ValueType == ERigControlValueType::Current;
 	const bool bShowToggle = (ValueType == ERigControlValueType::Minimum) || (ValueType == ERigControlValueType::Maximum);
 	
-	URigHierarchy* HierarchyBeingDebugged = GetHierarchyBeingDebugged();
-	URigHierarchy* HierarchyToChange = bCurrent ? HierarchyBeingDebugged : HierarchyBeingCustomized;
+	URigHierarchy* Hierarchy = PerElementInfos[0].GetHierarchy();
+	URigHierarchy* HierarchyToChange = PerElementInfos[0].GetDefaultHierarchy();
+	if(ValueType == ERigControlValueType::Current)
+	{
+		HierarchyToChange = Hierarchy;
+	}
 
 	using SNumericVector2DInputBox = SNumericVectorInputBox<float, FVector2f, 2>;
 	TSharedPtr<SNumericVector2DInputBox> VectorInputBox;
@@ -4224,13 +4232,13 @@ FDetailWidgetRow& FRigControlElementDetails::CreateVector2DValueWidgetRow(
 
 	if(bShowToggle)
 	{
-		auto ToggleChecked = [ValueType, Keys, HierarchyBeingDebugged](int32 Index) -> ECheckBoxState
+		auto ToggleChecked = [ValueType, Keys, Hierarchy](int32 Index) -> ECheckBoxState
 		{
 			TOptional<bool> FirstValue;
 
 			for(const FRigElementKey& Key : Keys)
 			{
-				if(const FRigControlElement* ControlElement = HierarchyBeingDebugged->Find<FRigControlElement>(Key))
+				if(const FRigControlElement* ControlElement = Hierarchy->Find<FRigControlElement>(Key))
 				{
 					if(ControlElement->Settings.LimitEnabled.Num() == 2)
 					{
@@ -4302,12 +4310,12 @@ FDetailWidgetRow& FRigControlElementDetails::CreateVector2DValueWidgetRow(
 		});
 	}
 
-	auto GetValue = [ValueType, Keys, HierarchyBeingDebugged](int32 Component) -> TOptional<float>
+	auto GetValue = [ValueType, Keys, Hierarchy](int32 Component) -> TOptional<float>
 	{
-		const float FirstValue = HierarchyBeingDebugged->GetControlValue<FVector3f>(Keys[0], ValueType).Component(Component);
+		const float FirstValue = Hierarchy->GetControlValue<FVector3f>(Keys[0], ValueType).Component(Component);
 		for(int32 Index = 1; Index < Keys.Num(); Index++)
 		{
-			const float SecondValue = HierarchyBeingDebugged->GetControlValue<FVector3f>(Keys[Index], ValueType).Component(Component);
+			const float SecondValue = Hierarchy->GetControlValue<FVector3f>(Keys[Index], ValueType).Component(Component);
 			if(FirstValue != SecondValue)
 			{
 				return TOptional<float>();
@@ -4316,7 +4324,7 @@ FDetailWidgetRow& FRigControlElementDetails::CreateVector2DValueWidgetRow(
 		return FirstValue;
 	};
 
-	auto OnValueChanged = [ValueType, Keys, HierarchyBeingDebugged, HierarchyToChange, this]
+	auto OnValueChanged = [ValueType, Keys, Hierarchy, HierarchyToChange, this]
 		(TOptional<float> InValue, ETextCommit::Type InCommitType, bool bSetupUndo, int32 Component)
 		{
 			if(!InValue.IsSet())
@@ -4328,7 +4336,7 @@ FDetailWidgetRow& FRigControlElementDetails::CreateVector2DValueWidgetRow(
 		
 			for(const FRigElementKey& Key : Keys)
 			{
-				FVector3f Vector = HierarchyBeingDebugged->GetControlValue<FVector3f>(Key, ValueType);
+				FVector3f Vector = Hierarchy->GetControlValue<FVector3f>(Key, ValueType);
 				if(!FMath::IsNearlyEqual(Vector.Component(Component), Value))
 				{
 					if(!SliderTransaction.IsValid())
@@ -4352,7 +4360,8 @@ FDetailWidgetRow& FRigControlElementDetails::CreateVector2DValueWidgetRow(
 		SAssignNew(NameContent, STextBlock)
 		.Text(Label)
 		.ToolTipText(Tooltip)
-		.Font(IDetailLayoutBuilder::GetDetailFont());
+		.Font(IDetailLayoutBuilder::GetDetailFont())
+		.IsEnabled(bIsEnabled);
 	}
 
 	WidgetRow
@@ -4398,11 +4407,12 @@ FDetailWidgetRow& FRigControlElementDetails::CreateVector2DValueWidgetRow(
 		 .ToggleYChecked(ToggleYChecked)
 		 .OnToggleXChanged(OnToggleXChanged)
 		 .OnToggleYChanged(OnToggleYChanged)
+		 .IsEnabled(bIsEnabled)
 	]
 	.CopyAction(FUIAction(
-	FExecuteAction::CreateLambda([ValueType, Keys, HierarchyBeingDebugged]()
+	FExecuteAction::CreateLambda([ValueType, Keys, Hierarchy]()
 		{
-			const FVector3f Data3 = HierarchyBeingDebugged->GetControlValue<FVector3f>(Keys[0], ValueType);
+			const FVector3f Data3 = Hierarchy->GetControlValue<FVector3f>(Keys[0], ValueType);
 			const FVector2f Data(Data3.X, Data3.Y);
 			FString Content = Data.ToString();
 			FPlatformApplicationMisc::ClipboardCopy(*Content);
@@ -4432,17 +4442,17 @@ FDetailWidgetRow& FRigControlElementDetails::CreateVector2DValueWidgetRow(
 				HierarchyToChange->SetControlValue(Key, FRigControlValue::Make<FVector3f>(Data3), ValueType, true, true); 
 			}
 		}),
-		FCanExecuteAction())
+		FCanExecuteAction::CreateLambda([bIsEnabled]() { return bIsEnabled; }))
 	);
 
-	if(ValueType == ERigControlValueType::Current || ValueType == ERigControlValueType::Initial)
+	if((ValueType == ERigControlValueType::Current || ValueType == ERigControlValueType::Initial) && bIsEnabled)
 	{
 		WidgetRow.OverrideResetToDefault(FResetToDefaultOverride::Create(
-			TAttribute<bool>::CreateLambda([ValueType, Keys, HierarchyBeingDebugged]() -> bool
+			TAttribute<bool>::CreateLambda([ValueType, Keys, Hierarchy]() -> bool
 			{
-				const FVector3f FirstValue = HierarchyBeingDebugged->GetControlValue<FVector3f>(Keys[0], ValueType);
+				const FVector3f FirstValue = Hierarchy->GetControlValue<FVector3f>(Keys[0], ValueType);
 				const FVector3f ReferenceValue = ValueType == ERigControlValueType::Initial ? FVector3f::ZeroVector :
-					HierarchyBeingDebugged->GetControlValue<FVector3f>(Keys[0], ERigControlValueType::Initial);
+					Hierarchy->GetControlValue<FVector3f>(Keys[0], ERigControlValueType::Initial);
 
 				return !(FirstValue - ReferenceValue).IsNearlyZero();
 			}),
