Control Rig: Construction Event

diff --git a/Engine/Plugins/Animation/ControlRig/Source/ControlRigEditor/Private/ControlRigElementDetails.h b/Engine/Plugins/Animation/ControlRig/Source/ControlRigEditor/Private/ControlRigElementDetails.h
--- a/Engine/Plugins/Animation/ControlRig/Source/ControlRigEditor/Private/ControlRigElementDetails.h
+++ b/Engine/Plugins/Animation/ControlRig/Source/ControlRigEditor/Private/ControlRigElementDetails.h
@@ -24,6 +24,7 @@
 #include "Internationalization/FastDecimalFormat.h"
 #include "ScopedTransaction.h"
 #include "Styling/AppStyle.h"
+#include "Algo/Transform.h"
 
 class IPropertyHandle;
 
@@ -115,8 +116,6 @@ public:
 	virtual void CustomizeDetails(IDetailLayoutBuilder& DetailBuilder) override;
 	virtual void BeginDestroy() {};
 
-	URigHierarchy* GetHierarchy() const { return HierarchyBeingCustomized; }
-	URigHierarchy* GetHierarchyBeingDebugged() const;
 	FRigElementKey GetElementKey() const;
 	FText GetName() const;
 	void SetName(const FText& InNewText, ETextCommit::Type InCommitType);
@@ -133,55 +132,80 @@ public:
 	TArray<T> GetElementsInDetailsView(const TArray<FRigElementKey>& InFilter = TArray<FRigElementKey>()) const
 	{
 		TArray<T> Elements;
-		for(TWeakObjectPtr<UDetailsViewWrapperObject> ObjectBeingCustomized : ObjectsBeingCustomized)
+		for(const FPerElementInfo& Info : PerElementInfos)
 		{
-			if(ObjectBeingCustomized.IsValid())
+			T Content = Info.WrapperObject->GetContent<T>();
+			if(!InFilter.IsEmpty() && !InFilter.Contains(Content.GetKey()))
 			{
-				T Content = ObjectBeingCustomized->GetContent<T>();
-				if(!InFilter.IsEmpty() && !InFilter.Contains(Content.GetKey()))
-				{
-					continue;
-				}
-				Elements.Add(Content);
+				continue;
 			}
+			Elements.Add(Content);
 		}
 		return Elements;
 	}
 
-	template<typename T>
-	TArray<T*> GetElementsInHierarchy(const TArray<FRigElementKey>& InFilter = TArray<FRigElementKey>()) const
+	struct FPerElementInfo
 	{
-		TArray<T*> Elements;
-		if(HierarchyBeingCustomized)
+		FPerElementInfo()
+			: WrapperObject()
+			, Element()
+			, DefaultElement()
+		{}
+
+		bool IsValid() const { return Element.IsValid(); }
+		bool IsProcedural() const { return Element.IsValid() && Element.Get()->IsProcedural(); }
+		operator bool() const { return IsValid(); }
+
+		URigHierarchy* GetHierarchy() const { return (URigHierarchy*)Element.GetHierarchy(); }
+		URigHierarchy* GetDefaultHierarchy() const
 		{
-			if(InFilter.IsEmpty())
+			if(DefaultElement.IsValid())
 			{
-				for(TWeakObjectPtr<UDetailsViewWrapperObject> ObjectBeingCustomized : ObjectsBeingCustomized)
-				{
-					if(ObjectBeingCustomized.IsValid())
-					{
-						const FRigElementKey Key = ObjectBeingCustomized->GetContent<FRigBaseElement>().GetKey();
-						Elements.Add(HierarchyBeingCustomized->Find<T>(Key));
-					}
-				}
+				return (URigHierarchy*)DefaultElement.GetHierarchy();
 			}
-			else
+			return GetHierarchy();
+		}
+
+		UControlRigBlueprint* GetBlueprint() const
+		{
+			if(const UControlRig* ControlRig = GetHierarchy()->GetTypedOuter<UControlRig>())
 			{
-				for(const FRigElementKey& FilterKey : InFilter)
-				{
-					Elements.Add(HierarchyBeingCustomized->Find<T>(FilterKey));
-				}
+				return Cast<UControlRigBlueprint>(ControlRig->GetClass()->ClassGeneratedBy);
 			}
+			return GetDefaultHierarchy()->GetTypedOuter<UControlRigBlueprint>();
 		}
-		return Elements;
-	}
 
+		template<typename T = FRigBaseElement>
+		T* GetElement() const
+		{
+			return (T*)Element.Get<T>();
+		}
+
+		template<typename T = FRigBaseElement>
+		T* GetDefaultElement() const
+		{
+			if(DefaultElement)
+			{
+				return (T*)DefaultElement.Get<T>();
+			}
+			return GetElement<T>();
+		}
+
+		TWeakObjectPtr<UDetailsViewWrapperObject> WrapperObject;
+		FRigElementHandle Element;
+		FRigElementHandle DefaultElement;
+	};
+
+	const FPerElementInfo& FindElement(const FRigElementKey& InKey) const;
 	bool IsAnyElementOfType(ERigElementType InType) const;
 	bool IsAnyElementNotOfType(ERigElementType InType) const;
 	bool IsAnyControlOfAnimationType(ERigControlAnimationType InType) const;
 	bool IsAnyControlNotOfAnimationType(ERigControlAnimationType InType) const;
 	bool IsAnyControlOfValueType(ERigControlType InType) const;
 	bool IsAnyControlNotOfValueType(ERigControlType InType) const;
+	bool IsAnyElementProcedural() const;
+	const FPerElementInfo* FindElementByPredicate(const TFunction<bool(const FPerElementInfo&)>& InPredicate) const;
+	bool ContainsElementByPredicate(const TFunction<bool(const FPerElementInfo&)>& InPredicate) const;
 
 	static void RegisterSectionMappings(FPropertyEditorModule& PropertyEditorModule);
 	virtual void RegisterSectionMappings(FPropertyEditorModule& PropertyEditorModule, UClass* InClass);
@@ -191,9 +215,8 @@ public:
 
 protected:
 
-	UControlRigBlueprint* BlueprintBeingCustomized;
-	URigHierarchy* HierarchyBeingCustomized;
-	TArray<TWeakObjectPtr<UDetailsViewWrapperObject>> ObjectsBeingCustomized;
+	TArray<FPerElementInfo> PerElementInfos;
+	
 	TSharedPtr<SButton> SelectParentElementButton;
 };
 
@@ -232,7 +255,6 @@ protected:
 
 	FDetailWidgetRow& CreateTransformComponentValueWidgetRow(
 		ERigControlType InControlType,
-		URigHierarchy* HierarchyBeingDebugged,
 		const TArray<FRigElementKey>& Keys,
 		SAdvancedTransformInputBox<FEulerTransform>::FArguments TransformWidgetArgs,
 		IDetailCategoryBuilder& CategoryBuilder, 
@@ -243,7 +265,6 @@ protected:
 		TSharedPtr<SWidget> NameContent = TSharedPtr<SWidget>());
 
 	FDetailWidgetRow& CreateEulerTransformValueWidgetRow(
-		URigHierarchy* HierarchyBeingDebugged,
 		const TArray<FRigElementKey>& Keys,
 		SAdvancedTransformInputBox<FEulerTransform>::FArguments TransformWidgetArgs,
 		IDetailCategoryBuilder& CategoryBuilder, 
@@ -377,12 +398,16 @@ private:
 		TAttribute<EVisibility> Visibility = EVisibility::Visible,
 		TSharedPtr<SWidget> NameContent = TSharedPtr<SWidget>())
 	{
-		const bool bCurrent = ValueType == ERigControlValueType::Current;
-		const bool bInitial = ValueType == ERigControlValueType::Initial;
 		const bool bShowToggle = (ValueType == ERigControlValueType::Minimum) || (ValueType == ERigControlValueType::Maximum);
-		
-		URigHierarchy* HierarchyBeingDebugged = GetHierarchyBeingDebugged();
-		URigHierarchy* HierarchyToChange = bCurrent ? HierarchyBeingDebugged : HierarchyBeingCustomized;
+		const bool bIsProcedural = IsAnyElementProcedural();
+		const bool bIsEnabled = !bIsProcedural || ValueType == ERigControlValueType::Current;
+
+		URigHierarchy* Hierarchy = PerElementInfos[0].GetHierarchy();
+		URigHierarchy* HierarchyToChange = PerElementInfos[0].GetDefaultHierarchy();
+		if(ValueType == ERigControlValueType::Current)
+		{
+			HierarchyToChange = Hierarchy;
+		}
 
 		TSharedPtr<SNumericEntryBox<T>> NumericEntryBox;
 		
@@ -393,13 +418,13 @@ private:
 		if(bShowToggle)
 		{
 			ToggleChecked = TAttribute<ECheckBoxState>::CreateLambda(
-				[ValueType, Keys, HierarchyBeingDebugged]() -> ECheckBoxState
+				[ValueType, Keys, Hierarchy]() -> ECheckBoxState
 				{
 					TOptional<bool> FirstValue;
 
 					for(const FRigElementKey& Key : Keys)
 					{
-						if(const FRigControlElement* ControlElement = HierarchyBeingDebugged->Find<FRigControlElement>(Key))
+						if(const FRigControlElement* ControlElement = Hierarchy->Find<FRigControlElement>(Key))
 						{
 							if(ControlElement->Settings.LimitEnabled.Num() == 1)
 							{
@@ -486,7 +511,8 @@ private:
 			SAssignNew(NameContent, STextBlock)
 			.Text(Label)
 			.ToolTipText(Tooltip)
-			.Font(IDetailLayoutBuilder::GetDetailFont());
+			.Font(IDetailLayoutBuilder::GetDetailFont())
+			.IsEnabled(bIsEnabled);
 		}
 
 		WidgetRow
@@ -504,12 +530,12 @@ private:
 	        .AllowSpin(ValueType == ERigControlValueType::Current || ValueType == ERigControlValueType::Initial)
 	        .LinearDeltaSensitivity(1)
 			.Delta(0.01f)
-	        .Value_Lambda([ValueType, Keys, HierarchyBeingDebugged]() -> TOptional<T>
+	        .Value_Lambda([ValueType, Keys, Hierarchy]() -> TOptional<T>
 	        {
-		        const T FirstValue = HierarchyBeingDebugged->GetControlValue<T>(Keys[0], ValueType);
+		        const T FirstValue = Hierarchy->GetControlValue<T>(Keys[0], ValueType);
 				for(int32 Index = 1; Index < Keys.Num(); Index++)
 				{
-					const T SecondValue = HierarchyBeingDebugged->GetControlValue<T>(Keys[Index], ValueType);
+					const T SecondValue = Hierarchy->GetControlValue<T>(Keys[Index], ValueType);
 					if(FirstValue != SecondValue)
 					{
 						return TOptional<T>();
@@ -525,19 +551,19 @@ private:
 	        {
         		OnValueChanged(InValue, InCommitType, true);
 	        })
-	        .MinSliderValue_Lambda([ValueType, Keys, HierarchyBeingDebugged]() -> TOptional<T>
+	        .MinSliderValue_Lambda([ValueType, Keys, Hierarchy]() -> TOptional<T>
 			 {
 				 if(ValueType == ERigControlValueType::Current || ValueType == ERigControlValueType::Initial)
 				 {
-			 		return HierarchyBeingDebugged->GetControlValue<T>(Keys[0], ERigControlValueType::Minimum);
+			 		return Hierarchy->GetControlValue<T>(Keys[0], ERigControlValueType::Minimum);
 				 }
 				 return TOptional<T>();
 			 })
-			 .MaxSliderValue_Lambda([ValueType, Keys, HierarchyBeingDebugged]() -> TOptional<T>
+			 .MaxSliderValue_Lambda([ValueType, Keys, Hierarchy]() -> TOptional<T>
 			 {
 		 		if(ValueType == ERigControlValueType::Current || ValueType == ERigControlValueType::Initial)
 		 		{
-					 return HierarchyBeingDebugged->GetControlValue<T>(Keys[0], ERigControlValueType::Maximum);
+					 return Hierarchy->GetControlValue<T>(Keys[0], ERigControlValueType::Maximum);
 				 }
 				 return TOptional<T>();
 			 })
@@ -545,11 +571,12 @@ private:
 			 .ToggleChecked(ToggleChecked)
 			 .OnToggleChanged(OnToggleChanged)
 			 .UndeterminedString(NSLOCTEXT("FRigControlElementDetails", "MultipleValues", "Multiple Values"))
+			 .IsEnabled(bIsEnabled)
 		]
 		.CopyAction(FUIAction(
-		FExecuteAction::CreateLambda([ValueType, Keys, HierarchyBeingDebugged]()
+		FExecuteAction::CreateLambda([ValueType, Keys, Hierarchy]()
 			{
-				const T FirstValue = HierarchyBeingDebugged->GetControlValue<T>(Keys[0], ValueType);
+				const T FirstValue = Hierarchy->GetControlValue<T>(Keys[0], ValueType);
 				const FString Content = FastDecimalFormat::NumberToString(
 					FirstValue,
 					FastDecimalFormat::GetCultureAgnosticFormattingRules(),
@@ -584,17 +611,17 @@ private:
 					}
 				}
 			}),
-			FCanExecuteAction())
+			FCanExecuteAction::CreateLambda([bIsEnabled]() { return bIsEnabled; }))
 		);
 
-		if(ValueType == ERigControlValueType::Current || ValueType == ERigControlValueType::Initial)
+		if((ValueType == ERigControlValueType::Current || ValueType == ERigControlValueType::Initial) && bIsEnabled)
 		{
 			WidgetRow.OverrideResetToDefault(FResetToDefaultOverride::Create(
-				TAttribute<bool>::CreateLambda([ValueType, Keys, HierarchyBeingDebugged]() -> bool
+				TAttribute<bool>::CreateLambda([ValueType, Keys, Hierarchy]() -> bool
 				{
-					const T FirstValue = HierarchyBeingDebugged->GetControlValue<T>(Keys[0], ValueType);
+					const T FirstValue = Hierarchy->GetControlValue<T>(Keys[0], ValueType);
 					const T ReferenceValue = ValueType == ERigControlValueType::Initial ? T(0) :
-						HierarchyBeingDebugged->GetControlValue<T>(Keys[0], ERigControlValueType::Initial);
+						Hierarchy->GetControlValue<T>(Keys[0], ERigControlValueType::Initial);
 
 					return !FRigControlElementDetails::Equals(FirstValue, ReferenceValue);
 				}),
@@ -629,8 +656,6 @@ private:
 	TSharedPtr<IPropertyHandle> ShapeColorHandle;
 	TSharedPtr<IPropertyHandle> ShapeTransformHandle;
 
-	TArray<FRigControlElement> ControlElements;
-	TArray<UDetailsViewWrapperObject*> ObjectPerControl;
 	TSharedPtr<SControlRigShapeNameList> ShapeNameListWidget; 
 	static TSharedPtr<TArray<ERigControlValueType>> PickedValueTypes;
 };
diff --git a/Engine/Plugins/Animation/ControlRig/Source/ControlRigEditor/Private/ControlRigHierarchyCommands.cpp b/Engine/Plugins/Animation/ControlRig/Source/ControlRigEditor/Private/ControlRigHierarchyCommands.cpp
--- a/Engine/Plugins/Animation/ControlRig/Source/ControlRigEditor/Private/ControlRigHierarchyCommands.cpp
+++ b/Engine/Plugins/Animation/ControlRig/Source/ControlRigEditor/Private/ControlRigHierarchyCommands.cpp
@@ -36,7 +36,6 @@ void FControlRigHierarchyCommands::RegisterCommands()
 	UI_COMMAND(ShowNulls, "Show Nulls", "Whether to show or hide nulls", EUserInterfaceActionType::ToggleButton, FInputChord());
 	UI_COMMAND(ShowRigidBodies, "Show RigidBodies", "Whether to show or hide rigidbodies", EUserInterfaceActionType::ToggleButton, FInputChord());
 	UI_COMMAND(ShowReferences, "Show References", "Whether to show or hide references", EUserInterfaceActionType::ToggleButton, FInputChord());
-	UI_COMMAND(ShowDynamicHierarchy, "Dynamic Hierarchy", "Whether to show the static or dynamic hierarchy", EUserInterfaceActionType::ToggleButton, FInputChord());
 	UI_COMMAND(ToggleControlShapeTransformEdit, "Toggle Shape Transform Edit", "Toggle Editing Selected Control's Shape Transform", EUserInterfaceActionType::ToggleButton, FInputChord(EKeys::Period, EModifierKey::Control)); 
 	UI_COMMAND(SpaceSwitching, "Space Switching", "Space switching on the control", EUserInterfaceActionType::Button, FInputChord(EKeys::Tab)); 
 	UI_COMMAND(ShowIconColors, "Show Icon Colors", "Whether to tint the icons with the element color", EUserInterfaceActionType::ToggleButton, FInputChord());
diff --git a/Engine/Plugins/Animation/ControlRig/Source/ControlRigEditor/Private/ControlRigHierarchyCommands.h b/Engine/Plugins/Animation/ControlRig/Source/ControlRigEditor/Private/ControlRigHierarchyCommands.h
--- a/Engine/Plugins/Animation/ControlRig/Source/ControlRigEditor/Private/ControlRigHierarchyCommands.h
+++ b/Engine/Plugins/Animation/ControlRig/Source/ControlRigEditor/Private/ControlRigHierarchyCommands.h
@@ -108,9 +108,6 @@ public:
 	/* Show references */
 	TSharedPtr< FUICommandInfo > ShowReferences;
 
-	/* Show dynamic hierarchy */
-	TSharedPtr< FUICommandInfo > ShowDynamicHierarchy;
-
 	/** Toggle Shape Transform Edit*/
 	TSharedPtr< FUICommandInfo > ToggleControlShapeTransformEdit;
 
diff --git a/Engine/Plugins/Animation/ControlRig/Source/ControlRigEditor/Private/EditMode/ControlRigEditMode.cpp b/Engine/Plugins/Animation/ControlRig/Source/ControlRigEditor/Private/EditMode/ControlRigEditMode.cpp
--- a/Engine/Plugins/Animation/ControlRig/Source/ControlRigEditor/Private/EditMode/ControlRigEditMode.cpp
+++ b/Engine/Plugins/Animation/ControlRig/Source/ControlRigEditor/Private/EditMode/ControlRigEditMode.cpp
@@ -144,6 +144,7 @@ FControlRigEditMode::FControlRigEditMode()
 	, bIsChangingCoordSystem(false)
 	, InteractionType((uint8)EControlRigInteractionType::None)
 	, bShowControlsAsOverlay(false)
+	, HierarchyHashBeforeConstruction(INDEX_NONE)
 {
 	ControlProxy = NewObject<UControlRigDetailPanelControlProxies>(GetTransientPackage(), NAME_None);
 	ControlProxy->SetFlags(RF_Transactional);
@@ -1223,7 +1224,7 @@ bool FControlRigEditMode::HandleClick(FEditorViewportClient* InViewportClient, H
 								GEditor->NoteSelectionChange();
 							}
 						}
-						ClearRigElementSelection(FRigElementTypeHelper::ToMask(ERigElementType::Control));
+						ClearRigElementSelection(ValidControlTypeMask());
 						SetRigElementSelection(ShapeActor->ControlRig.Get(),ERigElementType::Control, ControlName, true);
 					}
 
@@ -1286,7 +1287,7 @@ bool FControlRigEditMode::HandleClick(FEditorViewportClient* InViewportClient, H
 											}
 											else
 											{
-												ClearRigElementSelection(FRigElementTypeHelper::ToMask(ERigElementType::Control));
+												ClearRigElementSelection(ValidControlTypeMask());
 												SetRigElementSelection(ControlRig,ERigElementType::Control, ControlName, true);
 											}
 											return true;
@@ -1324,7 +1325,7 @@ bool FControlRigEditMode::HandleClick(FEditorViewportClient* InViewportClient, H
 			}
 			else
 			{
-				ClearRigElementSelection(FRigElementTypeHelper::ToMask(ERigElementType::Control));
+				ClearRigElementSelection(ValidControlTypeMask());
 				SetRigElementSelection(FKBoneProxy->ControlRig,ERigElementType::Control, ControlName, true);
 			}
 			return true;
@@ -1556,7 +1557,7 @@ bool FControlRigEditMode::BoxSelect(FBox& InBox, bool InSelect)
 	const bool bShiftDown = LevelViewportClient->Viewport->KeyState(EKeys::LeftShift) || LevelViewportClient->Viewport->KeyState(EKeys::RightShift);
 	if (!bShiftDown)
 	{
-		ClearRigElementSelection(FRigElementTypeHelper::ToMask(ERigElementType::Control));
+		ClearRigElementSelection(ValidControlTypeMask());
 	}
 
 	// Select all actors that are within the selection box area.  Be aware that certain modes do special processing below.	
@@ -1629,7 +1630,7 @@ bool FControlRigEditMode::FrustumSelect(const FConvexVolume& InFrustum, FEditorV
 	const bool bShiftDown = InViewportClient->Viewport->KeyState(EKeys::LeftShift) || InViewportClient->Viewport->KeyState(EKeys::RightShift);
 	if (!bShiftDown)
 	{
-		ClearRigElementSelection(FRigElementTypeHelper::ToMask(ERigElementType::Control));
+		ClearRigElementSelection(ValidControlTypeMask());
 	}
 
 	for (TPair<UControlRig*, TArray<AControlRigShapeActor*>>& Pairs : ControlRigShapeActors)
@@ -1698,7 +1699,7 @@ bool FControlRigEditMode::InputDelta(FEditorViewportClient* InViewportClient, FV
 	{
 		for (TPair<UControlRig*, TArray<AControlRigShapeActor*>>& Pairs : ControlRigShapeActors)
 		{
-			if (AreRigElementsSelected(FRigElementTypeHelper::ToMask(ERigElementType::Control), Pairs.Key))
+			if (AreRigElementsSelected(ValidControlTypeMask(), Pairs.Key))
 			{
 				FTransform ComponentTransform = GetHostingSceneComponentTransform(Pairs.Key);
 
@@ -1883,17 +1884,9 @@ void FControlRigEditMode::ClearRigElementSelection(uint32 InTypes)
 // internal private function that doesn't use guarding.
 void FControlRigEditMode::SetRigElementSelectionInternal(UControlRig* ControlRig, ERigElementType Type, const FName& InRigElementName, bool bSelected)
 {
-	UControlRigBlueprint* Blueprint = Cast<UControlRigBlueprint>(ControlRig->GetClass()->ClassGeneratedBy);
-	if (IsInLevelEditor())
+	if(URigHierarchyController* Controller = ControlRig->GetHierarchy()->GetController())
 	{
-		if(URigHierarchyController* Controller = ControlRig->GetHierarchy()->GetController())
-		{
-			Controller->SelectElement(FRigElementKey(InRigElementName, Type), bSelected);
-		}
-	}
-	else if (Blueprint)
-	{
-		Blueprint->GetHierarchyController()->SelectElement(FRigElementKey(InRigElementName, Type), bSelected);
+		Controller->SelectElement(FRigElementKey(InRigElementName, Type), bSelected);
 	}
 }
 
@@ -2046,9 +2039,8 @@ void FControlRigEditMode::RecalcPivotTransform()
 			// Use average location as pivot location
 			FVector PivotLocation = FVector::ZeroVector;
 
-			// @todo: support bones also
 			TArray<FRigElementKey> SelectedRigElements = GetSelectedRigElements(ControlRig);
-			if (AreRigElementsSelected(FRigElementTypeHelper::ToMask(ERigElementType::Control),ControlRig))
+			if (AreRigElementsSelected(ValidControlTypeMask(),ControlRig))
 			{
 				FTransform LastTransform = FTransform::Identity;
 
@@ -3207,7 +3199,7 @@ bool FControlRigEditMode::AreRigElementSelectedAndMovable(UControlRig* ControlRi
 {
 	const UControlRigEditModeSettings* Settings = GetDefault<UControlRigEditModeSettings>();
 
-	if (!Settings || Settings->bHideControlShapes || !ControlRig || !AreRigElementsSelected(FRigElementTypeHelper::ToMask(ERigElementType::Control), ControlRig))
+	if (!Settings || Settings->bHideControlShapes || !ControlRig || !AreRigElementsSelected(ValidControlTypeMask(), ControlRig))
 	{
 		return false;
 	}
@@ -3439,6 +3431,11 @@ void FControlRigEditMode::OnHierarchyModified_AnyThread(ERigHierarchyNotificatio
 		return;
 	}
 
+	if(IsConstructionEventRunning())
+	{
+		return;
+	}
+
 	if(IsInGameThread())
 	{
 		OnHierarchyModified(InNotif, InHierarchy, InElement);
@@ -3497,6 +3494,41 @@ void FControlRigEditMode::OnControlModified(UControlRig* Subject, FRigControlEle
 	*/
 }
 
+void FControlRigEditMode::OnPreConstruction_AnyThread(UControlRig* InRig, const EControlRigState InState,
+	const FName& InEventName)
+{
+	HierarchyHashBeforeConstruction = InRig->GetHierarchy()->GetTopologyHash(false, true);
+}
+
+void FControlRigEditMode::OnPostConstruction_AnyThread(UControlRig* InRig, const EControlRigState InState,
+	const FName& InEventName)
+{
+	uint32 HierarchyHash = INDEX_NONE;
+	Swap(HierarchyHash, HierarchyHashBeforeConstruction);
+
+	if(HierarchyHash != InRig->GetHierarchy()->GetTopologyHash(false, true))
+	{
+		auto Task = [this, InRig]()
+		{
+			RequestToRecreateControlShapeActors(InRig);
+			RecreateControlShapeActors();
+			HandleSelectionChanged();
+		};
+				
+		if(IsInGameThread())
+		{
+			Task();
+		}
+		else
+		{
+			FFunctionGraphTask::CreateAndDispatchWhenReady([Task]()
+			{
+				Task();
+			}, TStatId(), NULL, ENamedThreads::GameThread);
+		}
+	}
+}
+
 void FControlRigEditMode::OnWidgetModeChanged(UE::Widget::EWidgetMode InWidgetMode)
 {
 	const UControlRigEditModeSettings* Settings = GetDefault<UControlRigEditModeSettings>();
@@ -3538,7 +3570,7 @@ bool FControlRigEditMode::CanChangeControlShapeTransform()
 				// do not allow multi-select
 				if (SelectedRigElements.Num() == 1)
 				{
-					if (AreRigElementsSelected(FRigElementTypeHelper::ToMask(ERigElementType::Control),ControlRig))
+					if (AreRigElementsSelected(ValidControlTypeMask(),ControlRig))
 					{
 						// only enable for a Control with Gizmo enabled and visible
 						if (FRigControlElement* ControlElement = ControlRig->GetHierarchy()->Find<FRigControlElement>(SelectedRigElements[0]))
@@ -3676,7 +3708,7 @@ void FControlRigEditMode::MoveControlShape(AControlRigShapeActor* ShapeActor, co
 				{
 					bPrintPythonCommands = World->IsPreviewWorld();
 				}
-				ControlRig->SetControlGlobalTransform(ShapeActor->ControlName, NewTransform, true, Context, true, bPrintPythonCommands, true);			// assumes it's attached to actor
+				ControlRig->SetControlGlobalTransform(ShapeActor->ControlName, NewTransform, true, Context, false, bPrintPythonCommands, true);			// assumes it's attached to actor
 				ShapeActor->SetGlobalTransform(CurrentTransform);
 				if (bCalcLocal)
 				{
@@ -3947,6 +3979,8 @@ void FControlRigEditMode::AddControlRigInternal(UControlRig* InControlRig)
 	InControlRig->PostInitInstanceIfRequired();
 	InControlRig->GetHierarchy()->OnModified().RemoveAll(this);
 	InControlRig->GetHierarchy()->OnModified().AddSP(this, &FControlRigEditMode::OnHierarchyModified_AnyThread);
+	InControlRig->OnPreConstructionForUI_AnyThread().AddSP(this, &FControlRigEditMode::OnPreConstruction_AnyThread);
+	InControlRig->OnPostConstruction_AnyThread().AddSP(this, &FControlRigEditMode::OnPostConstruction_AnyThread);
 
 	//needed for the control rig track editor delegates to get hooked up
 	if (WeakSequencer.IsValid())
@@ -3999,6 +4033,9 @@ void FControlRigEditMode::RemoveControlRig(UControlRig* InControlRig)
 {
 	InControlRig->ControlModified().RemoveAll(this);
 	InControlRig->GetHierarchy()->OnModified().RemoveAll(this);
+	InControlRig->OnPreConstructionForUI_AnyThread().RemoveAll(this);
+	InControlRig->OnPostConstruction_AnyThread().RemoveAll(this);
+	
 	int32 Index = RuntimeControlRigs.Find(InControlRig);
 	TStrongObjectPtr<UControlRigEditModeDelegateHelper>* DelegateHelper = DelegateHelpers.Find(InControlRig);
 	if (DelegateHelper && DelegateHelper->IsValid())
diff --git a/Engine/Plugins/Animation/ControlRig/Source/ControlRigEditor/Private/EditMode/ControlRigEditMode.h b/Engine/Plugins/Animation/ControlRig/Source/ControlRigEditor/Private/EditMode/ControlRigEditMode.h
--- a/Engine/Plugins/Animation/ControlRig/Source/ControlRigEditor/Private/EditMode/ControlRigEditMode.h
+++ b/Engine/Plugins/Animation/ControlRig/Source/ControlRigEditor/Private/EditMode/ControlRigEditMode.h
@@ -163,6 +163,8 @@ public:
 	void OnHierarchyModified(ERigHierarchyNotification InNotif, URigHierarchy* InHierarchy, const FRigBaseElement* InElement);
 	void OnHierarchyModified_AnyThread(ERigHierarchyNotification InNotif, URigHierarchy* InHierarchy, const FRigBaseElement* InElement);
 	void OnControlModified(UControlRig* Subject, FRigControlElement* InControlElement, const FRigControlModifiedContext& Context);
+	void OnPreConstruction_AnyThread(UControlRig* InRig, const EControlRigState InState, const FName& InEventName);
+	void OnPostConstruction_AnyThread(UControlRig* InRig, const EControlRigState InState, const FName& InEventName);
 
 	/** return true if it can be removed from preview scene 
 	- this is to ensure preview scene doesn't remove shape actors */
@@ -174,7 +176,10 @@ public:
 	Control Rig, otherwise will recreate all of them*/
 	void RequestToRecreateControlShapeActors(UControlRig* ControlRig = nullptr); 
 
-
+	static uint32 ValidControlTypeMask()
+	{
+		return FRigElementTypeHelper::ToMask(ERigElementType::Control);
+	}
 
 protected:
 
@@ -435,6 +440,9 @@ private:
 	uint8 InteractionType;
 	bool bShowControlsAsOverlay;
 
+	bool IsConstructionEventRunning() const { return HierarchyHashBeforeConstruction != INDEX_NONE; }
+	uint32 HierarchyHashBeforeConstruction;
+
 	friend class FControlRigEditorModule;
 	friend class FControlRigEditor;
 	friend class FControlRigEditModeGenericDetails;
diff --git a/Engine/Plugins/Animation/ControlRig/Source/ControlRigEditor/Private/Editor/ControlRigEditor.cpp b/Engine/Plugins/Animation/ControlRig/Source/ControlRigEditor/Private/Editor/ControlRigEditor.cpp
--- a/Engine/Plugins/Animation/ControlRig/Source/ControlRigEditor/Private/Editor/ControlRigEditor.cpp
+++ b/Engine/Plugins/Animation/ControlRig/Source/ControlRigEditor/Private/Editor/ControlRigEditor.cpp
@@ -99,7 +99,6 @@
 #include "ToolMenus.h"
 #include "Styling/AppStyle.h"
 #include "AssetRegistry/AssetRegistryModule.h"
-#include "ControlRig/Private/Units/Execution/RigUnit_InteractionExecution.h"
 
 #define LOCTEXT_NAMESPACE "ControlRigEditor"
 
@@ -224,6 +223,7 @@ FControlRigEditor::FControlRigEditor()
 	, RigHierarchyTabCount(0)
 	, HaltedAtNode(nullptr)
 	, bSuspendDetailsPanelRefresh(false)
+	, HierarchyHashBeforeConstruction(INDEX_NONE)
 {
 }
 
@@ -288,6 +288,19 @@ UControlRigBlueprint* FControlRigEditor::GetControlRigBlueprint() const
 	return Cast<UControlRigBlueprint>(GetBlueprintObj());
 }
 
+URigHierarchy* FControlRigEditor::GetHierarchyBeingDebugged() const
+{
+	if(UControlRigBlueprint* RigBlueprint = GetControlRigBlueprint())
+	{
+		if(UControlRig* RigBeingDebugged = Cast<UControlRig>(RigBlueprint->GetObjectBeingDebugged()))
+		{
+			return RigBeingDebugged->GetHierarchy();
+		}
+		return RigBlueprint->Hierarchy;
+	}
+	return nullptr;
+}
+
 void FControlRigEditor::ExtendMenu()
 {
 	if(MenuExtender.IsValid())
@@ -989,7 +1002,7 @@ void FControlRigEditor::SetEventQueue(TArray<FName> InEventQueue, bool bCompile)
 		
 		// need to clear selection before remove transient control
 		// because active selection will trigger transient control recreation after removal	
-		PreviousSelection = RigBlueprint->Hierarchy->GetSelectedKeys();
+		PreviousSelection = GetHierarchyBeingDebugged()->GetSelectedKeys();
 		RigBlueprint->GetHierarchyController()->ClearSelection();
 		
 		// need to copy here since the removal changes the iterator
@@ -1021,10 +1034,7 @@ void FControlRigEditor::SetEventQueue(TArray<FName> InEventQueue, bool bCompile)
 
 	if (FControlRigEditMode* EditMode = GetEditMode())
 	{
-		if (UControlRigBlueprint* RigBlueprint = Cast<UControlRigBlueprint>(GetBlueprintObj()))
-		{
-			EditMode->RecreateControlShapeActors(RigBlueprint->Hierarchy->GetSelectedKeys());
-		}
+		EditMode->RecreateControlShapeActors(GetHierarchyBeingDebugged()->GetSelectedKeys());
 
 		UControlRigEditModeSettings* Settings = GetMutableDefault<UControlRigEditModeSettings>();
 		Settings->bDisplayNulls = IsConstructionModeEnabled();
@@ -1032,10 +1042,7 @@ void FControlRigEditor::SetEventQueue(TArray<FName> InEventQueue, bool bCompile)
 
 	if (PreviousSelection.Num() > 0)
 	{
-		if (UControlRigBlueprint* RigBlueprint = Cast<UControlRigBlueprint>(GetBlueprintObj()))
-		{
-			RigBlueprint->GetHierarchyController()->SetSelection(PreviousSelection);
-		}
+		GetHierarchyBeingDebugged()->GetController(true)->SetSelection(PreviousSelection);
 		SetDetailViewForRigElements();
 	}
 }
@@ -1294,6 +1301,8 @@ void FControlRigEditor::HandleSetObjectBeingDebugged(UObject* InObject)
 		if(!PreviouslyDebuggedControlRig->HasAnyFlags(RF_BeginDestroyed))
 		{
 			PreviouslyDebuggedControlRig->GetHierarchy()->OnModified().RemoveAll(this);
+			PreviouslyDebuggedControlRig->OnPreConstructionForUI_AnyThread().RemoveAll(this);
+			PreviouslyDebuggedControlRig->OnPostConstruction_AnyThread().RemoveAll(this);
 		}
 	}
 
@@ -1317,7 +1326,7 @@ void FControlRigEditor::HandleSetObjectBeingDebugged(UObject* InObject)
 		bool bIsExternalControlRig = DebuggedControlRig != ControlRig;
 		bool bShouldExecute = (!bIsExternalControlRig) && bExecutionControlRig;
 		DebuggedControlRig->ControlRigLog = &ControlRigLog;
-		DebuggedControlRig->DynamicHierarchy->HierarchyForSelectionPtr = GetControlRigBlueprint()->Hierarchy;
+		GetControlRigBlueprint()->Hierarchy->HierarchyForSelectionPtr = DebuggedControlRig->DynamicHierarchy;
 
 		UControlRigSkeletalMeshComponent* EditorSkelComp = Cast<UControlRigSkeletalMeshComponent>(GetPersonaToolkit()->GetPreviewScene()->GetPreviewMeshComponent());
 		if (EditorSkelComp)
@@ -1356,6 +1365,8 @@ void FControlRigEditor::HandleSetObjectBeingDebugged(UObject* InObject)
 		}
 
 		DebuggedControlRig->GetHierarchy()->OnModified().AddSP(this, &FControlRigEditor::OnHierarchyModified_AnyThread);
+		DebuggedControlRig->OnPreConstructionForUI_AnyThread().AddSP(this, &FControlRigEditor::OnPreConstruction_AnyThread);
+		DebuggedControlRig->OnPostConstruction_AnyThread().AddSP(this, &FControlRigEditor::OnPostConstruction_AnyThread);
 
 		if(EditorSkelComp)
 		{
@@ -1576,34 +1587,19 @@ void FControlRigEditor::SetDetailViewForRigElements()
 	ClearDetailObject();
 
 	UControlRigBlueprint* RigBlueprint = Cast<UControlRigBlueprint>(GetBlueprintObj());
-	URigHierarchy* Hierarchy = RigBlueprint->Hierarchy;
-	
-	if (!IsConstructionModeEnabled())
-	{
-		if (UControlRig* DebuggedControlRig = Cast<UControlRig>(RigBlueprint->GetObjectBeingDebugged()))
-		{
-			Hierarchy = DebuggedControlRig->GetHierarchy();
-		}
-	}
-
+	URigHierarchy* HierarchyBeingDebugged = GetHierarchyBeingDebugged();
 	TArray<UObject*> Objects;
 
-	TArray<FRigElementKey> CurrentSelection = RigBlueprint->Hierarchy->GetSelectedKeys();
+	TArray<FRigElementKey> CurrentSelection = HierarchyBeingDebugged->GetSelectedKeys();
 	for(const FRigElementKey& SelectedKey : CurrentSelection)
 	{
-		FRigBaseElement* Element = Hierarchy->Find(SelectedKey);
+		FRigBaseElement* Element = HierarchyBeingDebugged->Find(SelectedKey);
 		if (Element == nullptr)
 		{
 			continue;
 		}
 
-		URigHierarchy* SelectionHierarchy = Hierarchy;
-		if (Hierarchy->HierarchyForSelectionPtr.IsValid())
-		{
-			SelectionHierarchy = Hierarchy->HierarchyForSelectionPtr.Get();
-		}
-
-		UDetailsViewWrapperObject* WrapperObject = UDetailsViewWrapperObject::MakeInstance(Element->GetElementStruct(), (uint8*)Element, SelectionHierarchy);
+		UDetailsViewWrapperObject* WrapperObject = UDetailsViewWrapperObject::MakeInstance(Element->GetElementStruct(), (uint8*)Element, HierarchyBeingDebugged);
 		WrapperObject->GetWrappedPropertyChangedChainEvent().AddSP(this, &FControlRigEditor::OnWrappedPropertyChangedChainEvent);
 		WrapperObject->AddToRoot();
 
@@ -2036,7 +2032,7 @@ void FControlRigEditor::Compile()
 
 		if (FControlRigEditMode* EditMode = GetEditMode())
 		{
-			EditMode->RecreateControlShapeActors(RigBlueprint->Hierarchy->GetSelectedKeys());
+			EditMode->RecreateControlShapeActors(GetHierarchyBeingDebugged()->GetSelectedKeys());
 		}
 	}
 
@@ -2915,14 +2911,7 @@ void FControlRigEditor::HandleControlRigExecutedEvent(UControlRig* InControlRig,
 			DebuggedControlRig = ControlRig;
 		}
 
-		URigHierarchy* Hierarchy = ControlRigBP->Hierarchy; 
-		if(!IsConstructionModeEnabled())
-		{
-			if(DebuggedControlRig)
-			{
-				Hierarchy = DebuggedControlRig->GetHierarchy();
-			}
-		}
+		URigHierarchy* Hierarchy = GetHierarchyBeingDebugged(); 
 
 		TArray< TWeakObjectPtr<UObject> > SelectedObjects = Inspector->GetSelectedObjects();
 		for (TWeakObjectPtr<UObject> SelectedObject : SelectedObjects)
@@ -3922,7 +3911,11 @@ void FControlRigEditor::OnPinControlNameListComboBox(const TArray<TSharedPtr<FSt
 
 bool FControlRigEditor::IsConstructionModeEnabled() const
 {
-	return LastEventQueue.Contains(FRigUnit_PrepareForExecution::EventName);
+	if(ControlRig)
+	{
+		return ControlRig->IsConstructionModeEnabled();
+	}
+	return false;
 }
 
 void FControlRigEditor::HandlePreviewSceneCreated(const TSharedRef<IPersonaPreviewScene>& InPersonaPreviewScene)
@@ -4021,7 +4014,7 @@ void FControlRigEditor::UpdateControlRig()
 				ControlRig->ExecutionType = ERigExecutionType::Editing;
 				ControlRig->ControlRigLog = &ControlRigLog;
 
-				ControlRig->InitializeFromCDO();
+				ControlRig->Initialize(true);
  			}
 
 			ControlRig->PreviewInstance = PreviewInstance;
@@ -4088,12 +4081,7 @@ void FControlRigEditor::CacheNameLists()
 		TArray<UEdGraph*> EdGraphs;
 		ControlRigBP->GetAllGraphs(EdGraphs);
 
-		URigHierarchy* Hierarchy = ControlRigBP->Hierarchy;
-		if(UControlRig* DebuggedControlRig = Cast<UControlRig>(GetBlueprintObj()->GetObjectBeingDebugged()))
-		{
-			Hierarchy = DebuggedControlRig->GetHierarchy();
-		}
-
+		URigHierarchy* Hierarchy = GetHierarchyBeingDebugged();
 		for (UEdGraph* Graph : EdGraphs)
 		{
 			UControlRigGraph* RigGraph = Cast<UControlRigGraph>(Graph);
@@ -4101,7 +4089,7 @@ void FControlRigEditor::CacheNameLists()
 			{
 				continue;
 			}
-			RigGraph->CacheNameLists(Hierarchy, &ControlRigBP->DrawContainer);
+			RigGraph->CacheNameLists(Hierarchy, &ControlRigBP->DrawContainer, ControlRigBP->ShapeLibraries);
 		}
 	}
 }
@@ -4283,9 +4271,9 @@ FTransform FControlRigEditor::GetRigElementTransform(const FRigElementKey& InEle
 
 	if (bLocal)
 	{
-		return GetControlRigBlueprint()->Hierarchy->GetLocalTransform(InElement);
+		return GetHierarchyBeingDebugged()->GetLocalTransform(InElement);
 	}
-	return GetControlRigBlueprint()->Hierarchy->GetGlobalTransform(InElement);
+	return GetHierarchyBeingDebugged()->GetGlobalTransform(InElement);
 }
 
 void FControlRigEditor::SetRigElementTransform(const FRigElementKey& InElement, const FTransform& InTransform, bool bLocal)
@@ -4431,6 +4419,12 @@ void FControlRigEditor::OnFinishedChangingProperties(const FPropertyChangedEvent
 			return;
 		}
 
+		if (PropertyChangedEvent.MemberProperty->GetNameCPP() == GET_MEMBER_NAME_STRING_CHECKED(UControlRigBlueprint, HierarchySettings))
+		{
+			ControlRigBP->PropagateHierarchyFromBPToInstances();
+			return;
+		}
+
 		if (PropertyChangedEvent.MemberProperty->GetNameCPP() == GET_MEMBER_NAME_STRING_CHECKED(UControlRigBlueprint, DrawContainer))
 		{
 			ControlRigBP->PropagateDrawInstructionsFromBPToInstances();
@@ -4889,7 +4883,7 @@ void FControlRigEditor::OnHierarchyChanged()
 
 		FBlueprintEditorUtils::MarkBlueprintAsModified(GetControlRigBlueprint());
 		
-		TArray<const FRigBaseElement*> SelectedElements = ControlRigBP->Hierarchy->GetSelectedElements();
+		TArray<const FRigBaseElement*> SelectedElements = GetHierarchyBeingDebugged()->GetSelectedElements();
 		for(const FRigBaseElement* SelectedElement : SelectedElements)
 		{
 			ControlRigBP->Hierarchy->OnModified().Broadcast(ERigHierarchyNotification::ElementSelected, ControlRigBP->Hierarchy, SelectedElement);
@@ -5096,39 +5090,6 @@ void FControlRigEditor::OnHierarchyModified(ERigHierarchyNotification InNotif, U
 
 			break;
 		}
-		case ERigHierarchyNotification::ElementSelected:
-		case ERigHierarchyNotification::ElementDeselected:
-		{
-			check(InElement);
-			const bool bSelected = InNotif == ERigHierarchyNotification::ElementSelected;
-
-			if (InElement->GetType() == ERigElementType::Bone)
-			{
-				SynchronizeViewportBoneSelection();
-			}
-
-			if (bSelected)
-			{
-				SetDetailViewForRigElements();
-			}
-			else
-			{
-				TArray<FRigElementKey> CurrentSelection = RigBlueprint->Hierarchy->GetSelectedKeys();
-				if (CurrentSelection.Num() > 0)
-				{
-					if(FRigBaseElement* LastSelectedElement = InHierarchy->Find(CurrentSelection.Last()))
-					{
-						OnHierarchyModified(ERigHierarchyNotification::ElementSelected, InHierarchy, LastSelectedElement);
-					}
-				}
-				else
-				{
-					ClearDetailObject();
-				}
-			}
-				
-			break;
-		}
 		default:
 		{
 			break;
@@ -5138,6 +5099,11 @@ void FControlRigEditor::OnHierarchyModified(ERigHierarchyNotification InNotif, U
 
 void FControlRigEditor::OnHierarchyModified_AnyThread(ERigHierarchyNotification InNotif, URigHierarchy* InHierarchy, const FRigBaseElement* InElement)
 {
+	if(IsConstructionEventRunning())
+	{
+		return;
+	}
+	
 	FRigElementKey Key;
 	if(InElement)
 	{
@@ -5156,7 +5122,7 @@ void FControlRigEditor::OnHierarchyModified_AnyThread(ERigHierarchyNotification
 	}
 
 	TWeakObjectPtr<URigHierarchy> WeakHierarchy = InHierarchy;
-	FFunctionGraphTask::CreateAndDispatchWhenReady([this, InNotif, WeakHierarchy, Key]()
+	auto Task = [this, InNotif, WeakHierarchy, Key]()
     {
 		if(!WeakHierarchy.IsValid())
     	{
@@ -5167,17 +5133,47 @@ void FControlRigEditor::OnHierarchyModified_AnyThread(ERigHierarchyNotification
 
 		switch(InNotif)
 		{
+			case ERigHierarchyNotification::ElementSelected:
+			case ERigHierarchyNotification::ElementDeselected:
+			{
+				if(Element)
+				{
+					const bool bSelected = InNotif == ERigHierarchyNotification::ElementSelected;
+
+					if (Element->GetType() == ERigElementType::Bone)
+					{
+						SynchronizeViewportBoneSelection();
+					}
+
+					if (bSelected)
+					{
+						SetDetailViewForRigElements();
+					}
+					else
+					{
+						TArray<FRigElementKey> CurrentSelection = GetHierarchyBeingDebugged()->GetSelectedKeys();
+						if (CurrentSelection.Num() > 0)
+						{
+							if(FRigBaseElement* LastSelectedElement = WeakHierarchy.Get()->Find(CurrentSelection.Last()))
+							{
+								OnHierarchyModified(ERigHierarchyNotification::ElementSelected,  WeakHierarchy.Get(), LastSelectedElement);
+							}
+						}
+						else
+						{
+							ClearDetailObject();
+						}
+					}
+				}						
+				break;
+			}
 			case ERigHierarchyNotification::ElementAdded:
 			case ERigHierarchyNotification::ElementRemoved:
 			case ERigHierarchyNotification::ElementRenamed:
 			case ERigHierarchyNotification::ParentChanged:
             case ERigHierarchyNotification::HierarchyReset:
 			{
-				// only recache the name lists of this is the main hierarchy
-				if(WeakHierarchy.Get()->GetOuter()->IsA<UControlRigBlueprint>())
-				{
-					CacheNameLists();
-				}
+				CacheNameLists();
 				break;
 			}
 			case ERigHierarchyNotification::ControlSettingChanged:
@@ -5227,7 +5223,19 @@ void FControlRigEditor::OnHierarchyModified_AnyThread(ERigHierarchyNotification
 			}
 		}
 		
-    }, TStatId(), NULL, ENamedThreads::GameThread);
+    };
+
+	if(IsInGameThread())
+	{
+		Task();
+	}
+	else
+	{
+		FFunctionGraphTask::CreateAndDispatchWhenReady([Task]()
+		{
+			Task();
+		}, TStatId(), NULL, ENamedThreads::GameThread);
+	}
 }
 
 void FControlRigEditor::SynchronizeViewportBoneSelection()
@@ -5243,7 +5251,7 @@ void FControlRigEditor::SynchronizeViewportBoneSelection()
 	{
 		EditorSkelComp->BonesOfInterest.Reset();
 
-		TArray<const FRigBaseElement*> SelectedBones = RigBlueprint->Hierarchy->GetSelectedElements(ERigElementType::Bone);
+		TArray<const FRigBaseElement*> SelectedBones = GetHierarchyBeingDebugged()->GetSelectedElements(ERigElementType::Bone);
 		for (const FRigBaseElement* SelectedBone : SelectedBones)
 		{
  			const int32 BoneIndex = EditorSkelComp->GetReferenceSkeleton().FindBoneIndex(SelectedBone->GetName());
@@ -5376,7 +5384,7 @@ void FControlRigEditor::CreateRigHierarchyToGraphDragAndDropMenu() const
 				{
 					const FControlRigRigHierarchyToGraphDragAndDropContext& DragDropContext = MainContext->GetRigHierarchyToGraphDragAndDropContext();
 
-					UControlRigBlueprint* Blueprint = MainContext->GetControlRigBlueprint();
+					URigHierarchy* Hierarchy = ControlRigEditor->GetHierarchyBeingDebugged();
 					const TArray<FRigElementKey>& DraggedKeys = DragDropContext.DraggedElementKeys;
 					UEdGraph* Graph = DragDropContext.Graph.Get();
 					const FVector2D& NodePosition = DragDropContext.NodePosition;
@@ -5401,7 +5409,7 @@ void FControlRigEditor::CreateRigHierarchyToGraphDragAndDropMenu() const
 						}
 						else if(DraggedKey.Type == ERigElementType::Control)
 						{
-							if(const FRigControlElement* ControlElement = Blueprint->Hierarchy->Find<FRigControlElement>(DraggedKey))
+							if(const FRigControlElement* ControlElement = Hierarchy->Find<FRigControlElement>(DraggedKey))
 							{
 								const uint8 DraggedAnimationType = ControlElement->IsAnimationChannel() ? 1 : 0; 
 								if(DraggedAnimationTypes == 2)
@@ -5590,7 +5598,7 @@ void FControlRigEditor::CreateRigHierarchyToGraphDragAndDropMenu() const
 						Section.AddEntry(SetRelativeTransformEntry);
 					}
 
-					if (DraggedKeys.Num() > 0 && Blueprint != nullptr)
+					if (DraggedKeys.Num() > 0 && Hierarchy != nullptr)
 					{
 						FToolMenuEntry& ItemArraySeparator = Section.AddSeparator(TEXT("ItemArraySeparator"));
 						ItemArraySeparator.InsertPosition.Name = TEXT("SetRelativeTransformEntry"),
@@ -5602,7 +5610,7 @@ void FControlRigEditor::CreateRigHierarchyToGraphDragAndDropMenu() const
 							LOCTEXT("CreateItemArray_ToolTip", "Creates an item array from the selected elements in the hierarchy"),
 							FSlateIcon(),
 							FUIAction(
-								FExecuteAction::CreateLambda([ControlRigEditor, Blueprint, DraggedKeys, NodePosition]()
+								FExecuteAction::CreateLambda([ControlRigEditor, DraggedKeys, NodePosition]()
 									{
 										if (URigVMController* Controller = ControlRigEditor->GetFocusedController())
 										{
@@ -5698,8 +5706,8 @@ void FControlRigEditor::HandleMakeElementGetterSetter(ERigElementGetterSetterTyp
 		return;
 	}
 
-	UControlRigBlueprint* Blueprint = Cast<UControlRigBlueprint>(GetBlueprintObj());
-	if (Blueprint == nullptr)
+	URigHierarchy* Hierarchy = GetHierarchyBeingDebugged();
+	if (Hierarchy == nullptr)
 	{
 		return;
 	}
@@ -5736,11 +5744,11 @@ void FControlRigEditor::HandleMakeElementGetterSetter(ERigElementGetterSetterTyp
 		TArray<FName> ChannelPins;
 		TMap<FName, FRigVMTemplateArgumentType> PinsToResolve; 
 
-		if(FRigControlElement* ControlElement = Blueprint->Hierarchy->Find<FRigControlElement>(Key))
+		if(FRigControlElement* ControlElement = Hierarchy->Find<FRigControlElement>(Key))
 		{
 			if(ControlElement->IsAnimationChannel())
 			{
-				if(const FRigControlElement* ParentControlElement = Cast<FRigControlElement>(Blueprint->Hierarchy->GetFirstParent(ControlElement)))
+				if(const FRigControlElement* ParentControlElement = Cast<FRigControlElement>(Hierarchy->GetFirstParent(ControlElement)))
 				{
 					NameValue = ParentControlElement->GetName();
 				}
@@ -5871,7 +5879,7 @@ void FControlRigEditor::HandleMakeElementGetterSetter(ERigElementGetterSetterTyp
 				{
 					if (Key.Type == ERigElementType::Control)
 					{
-						FRigControlElement* ControlElement = Blueprint->Hierarchy->Find<FRigControlElement>(Key);
+						FRigControlElement* ControlElement = Hierarchy->Find<FRigControlElement>(Key);
 						check(ControlElement);
 						
 						switch (ControlElement->Settings.ControlType)
@@ -5959,7 +5967,7 @@ void FControlRigEditor::HandleMakeElementGetterSetter(ERigElementGetterSetterTyp
 				{
 					if (Key.Type == ERigElementType::Control)
 					{
-						FRigControlElement* ControlElement = Blueprint->Hierarchy->Find<FRigControlElement>(Key);
+						FRigControlElement* ControlElement = Hierarchy->Find<FRigControlElement>(Key);
 						check(ControlElement);
 
 						switch (ControlElement->Settings.ControlType)
@@ -5994,7 +6002,7 @@ void FControlRigEditor::HandleMakeElementGetterSetter(ERigElementGetterSetterTyp
 								StructTemplate = FRigUnit_SetControlVector::StaticStruct();
 								NewNode.ValuePinName = TEXT("Vector");
 								NewNode.ValueType = ERigControlType::Position;
-								NewNode.Value = FRigControlValue::Make<FVector>(Blueprint->Hierarchy->GetGlobalTransform(Key).GetLocation());
+								NewNode.Value = FRigControlValue::Make<FVector>(Hierarchy->GetGlobalTransform(Key).GetLocation());
 								break;
 							}
 							case ERigControlType::Scale:
@@ -6003,7 +6011,7 @@ void FControlRigEditor::HandleMakeElementGetterSetter(ERigElementGetterSetterTyp
 								StructTemplate = FRigUnit_SetControlVector::StaticStruct();
 								NewNode.ValuePinName = TEXT("Vector");
 								NewNode.ValueType = ERigControlType::Scale;
-								NewNode.Value = FRigControlValue::Make<FVector>(Blueprint->Hierarchy->GetGlobalTransform(Key).GetScale3D());
+								NewNode.Value = FRigControlValue::Make<FVector>(Hierarchy->GetGlobalTransform(Key).GetScale3D());
 								break;
 							}
 							case ERigControlType::Rotator:
@@ -6012,7 +6020,7 @@ void FControlRigEditor::HandleMakeElementGetterSetter(ERigElementGetterSetterTyp
 								StructTemplate = FRigUnit_SetControlRotator::StaticStruct();
 								NewNode.ValuePinName = TEXT("Rotator");
 								NewNode.ValueType = ERigControlType::Rotator;
-								NewNode.Value = FRigControlValue::Make<FRotator>(Blueprint->Hierarchy->GetGlobalTransform(Key).Rotator());
+								NewNode.Value = FRigControlValue::Make<FRotator>(Hierarchy->GetGlobalTransform(Key).Rotator());
 								break;
 							}
 							case ERigControlType::Transform:
@@ -6022,7 +6030,7 @@ void FControlRigEditor::HandleMakeElementGetterSetter(ERigElementGetterSetterTyp
 								StructTemplate = FRigUnit_SetTransform::StaticStruct();
 								NewNode.ValuePinName = TEXT("Transform");
 								NewNode.ValueType = ERigControlType::Transform;
-								NewNode.Value = FRigControlValue::Make<FTransform>(Blueprint->Hierarchy->GetGlobalTransform(Key));
+								NewNode.Value = FRigControlValue::Make<FTransform>(Hierarchy->GetGlobalTransform(Key));
 								break;
 							}
 							default:
@@ -6036,7 +6044,7 @@ void FControlRigEditor::HandleMakeElementGetterSetter(ERigElementGetterSetterTyp
 						StructTemplate = FRigUnit_SetTransform::StaticStruct();
 						NewNode.ValuePinName = TEXT("Transform");
 						NewNode.ValueType = ERigControlType::Transform;
-						NewNode.Value = FRigControlValue::Make<FTransform>(Blueprint->Hierarchy->GetGlobalTransform(Key));
+						NewNode.Value = FRigControlValue::Make<FTransform>(Hierarchy->GetGlobalTransform(Key));
 					}
 					break;
 				}
@@ -6053,7 +6061,7 @@ void FControlRigEditor::HandleMakeElementGetterSetter(ERigElementGetterSetterTyp
 					StructTemplate = FRigUnit_SetRotation::StaticStruct();
 					NewNode.ValuePinName = TEXT("Rotation");
 					NewNode.ValueType = ERigControlType::Rotator;
-					NewNode.Value = FRigControlValue::Make<FRotator>(Blueprint->Hierarchy->GetGlobalTransform(Key).Rotator());
+					NewNode.Value = FRigControlValue::Make<FRotator>(Hierarchy->GetGlobalTransform(Key).Rotator());
 					break;
 				}
 				case ERigElementGetterSetterType_Translation:
@@ -6061,7 +6069,7 @@ void FControlRigEditor::HandleMakeElementGetterSetter(ERigElementGetterSetterTyp
 					StructTemplate = FRigUnit_SetTranslation::StaticStruct();
 					NewNode.ValuePinName = TEXT("Translation");
 					NewNode.ValueType = ERigControlType::Position;
-					NewNode.Value = FRigControlValue::Make<FVector>(Blueprint->Hierarchy->GetGlobalTransform(Key).GetLocation());
+					NewNode.Value = FRigControlValue::Make<FVector>(Hierarchy->GetGlobalTransform(Key).GetLocation());
 					break;
 				}
 				case ERigElementGetterSetterType_Offset:
@@ -6087,7 +6095,7 @@ void FControlRigEditor::HandleMakeElementGetterSetter(ERigElementGetterSetterTyp
 			NodePositionIncrement = FVector2D(380.f, 0.f);
 		}
 
-		FName Name = FControlRigBlueprintUtils::ValidateName(Blueprint, StructTemplate->GetName());
+		FName Name = FControlRigBlueprintUtils::ValidateName(GetControlRigBlueprint(), StructTemplate->GetName());
 		if (URigVMUnitNode* ModelNode = GetFocusedController()->AddUnitNode(StructTemplate, FRigUnit::GetMethodName(), NodePosition, FString(), true, true))
 		{
 			FString ItemTypeStr = StaticEnum<ERigElementType>()->GetDisplayNameTextByValue((int64)Key.Type).ToString();
@@ -6153,7 +6161,7 @@ void FControlRigEditor::HandleMakeElementGetterSetter(ERigElementGetterSetterTyp
 				}
 			}
 
-			UControlRigUnitNodeSpawner::HookupMutableNode(ModelNode, Blueprint);
+			UControlRigUnitNodeSpawner::HookupMutableNode(ModelNode, GetControlRigBlueprint());
 		}
 
 		NodePosition += NodePositionIncrement;
@@ -6237,6 +6245,36 @@ void FControlRigEditor::HandleOnControlModified(UControlRig* Subject, FRigContro
 
 			if (!NewDefaultValue.IsEmpty())
 			{
+				bool bRequiresRecompile = true;
+
+				if(ControlRig)
+				{
+					if(TSharedPtr<FRigVMParserAST> AST = Pin->GetGraph()->GetDiagnosticsAST())
+					{
+						FRigVMASTProxy PinProxy = FRigVMASTProxy::MakeFromUObject(Pin);
+						if(const FRigVMExprAST* PinExpr = AST->GetExprForSubject(PinProxy))
+						{
+							if(PinExpr->IsA(FRigVMExprAST::Var))
+							{
+								const FString PinHash = URigVMCompiler::GetPinHash(Pin, PinExpr->To<FRigVMVarExprAST>(), false);
+								if(const FRigVMOperand* OperandForPin = Blueprint->PinToOperandMap.Find(PinHash))
+								{
+									if(URigVM* VM = ControlRig->GetVM())
+									{
+										// only operands which are shared across multiple instructions require recompile
+										if(!VM->GetByteCode().IsOperandShared(*OperandForPin))
+										{
+											VM->SetPropertyValueFromString(*OperandForPin, NewDefaultValue);
+											bRequiresRecompile = false;
+										}
+									}
+								}
+							}
+						}
+					}
+				}
+				
+				TGuardValue<bool> DisableBlueprintNotifs(Blueprint->bSuspendModelNotificationsForSelf, !bRequiresRecompile);
 				GetFocusedController()->SetPinDefaultValue(Pin->GetPinPath(), NewDefaultValue, true, true, true);
 			}
 		}
@@ -6265,9 +6303,12 @@ void FControlRigEditor::HandleOnControlModified(UControlRig* Subject, FRigContro
 			{
 				const FTransform GlobalTransform = ControlRig->GetControlGlobalTransform(ControlElement->GetName());
 				Blueprint->Hierarchy->SetGlobalTransform(ElementKey, GlobalTransform);
-				Blueprint->Hierarchy->SetInitialGlobalTransform(ElementKey, GlobalTransform);
 				Hierarchy->SetGlobalTransform(ElementKey, GlobalTransform);
-				Hierarchy->SetInitialGlobalTransform(ElementKey, GlobalTransform);
+				if (IsConstructionModeEnabled())
+				{
+					Blueprint->Hierarchy->SetInitialGlobalTransform(ElementKey, GlobalTransform);
+					Hierarchy->SetInitialGlobalTransform(ElementKey, GlobalTransform);
+				}
 			}
 		}
 	}
@@ -6275,18 +6316,20 @@ void FControlRigEditor::HandleOnControlModified(UControlRig* Subject, FRigContro
 	{
 		FRigControlElement* SourceControlElement = Hierarchy->Find<FRigControlElement>(ControlElement->GetKey());
 		FRigControlElement* TargetControlElement = Blueprint->Hierarchy->Find<FRigControlElement>(ControlElement->GetKey());
+		if(SourceControlElement && TargetControlElement)
+		{
+			TargetControlElement->Settings = SourceControlElement->Settings;
 
-		TargetControlElement->Settings = SourceControlElement->Settings;
+			// only fire the setting change if the interaction is not currently ongoing
+			if(!Subject->ElementsBeingInteracted.Contains(ControlElement->GetKey()))
+			{
+				Blueprint->Hierarchy->OnModified().Broadcast(ERigHierarchyNotification::ControlSettingChanged, Blueprint->Hierarchy, TargetControlElement);
+			}
 
-		// only fire the setting change if the interaction is not currently ongoing
-		if(!Subject->ElementsBeingInteracted.Contains(ControlElement->GetKey()))
-		{
-			Blueprint->Hierarchy->OnModified().Broadcast(ERigHierarchyNotification::ControlSettingChanged, Blueprint->Hierarchy, TargetControlElement);
+			// we copy the pose including the weights since we want the topology to align during construction mode.
+			// i.e. dynamic reparenting should be reset here.
+			TargetControlElement->CopyPose(SourceControlElement, true, true, true);
 		}
-
-		// we copy the pose including the weights since we want the topology to align during construction mode.
-		// i.e. dynamic reparenting should be reset here.
-		TargetControlElement->CopyPose(SourceControlElement, true, true, true);
 	}
 }
 
@@ -6495,18 +6538,28 @@ void FControlRigEditor::UpdateGraphCompilerErrors()
 			// reset all nodes and store them in the map
 			bool bFoundWarning = false;
 			bool bFoundError = false;
+			
 			for (UEdGraphNode* GraphNode : Graph->Nodes)
 			{
 				if (UControlRigGraphNode* ControlRigGraphNode = Cast<UControlRigGraphNode>(GraphNode))
 				{
 					bFoundError = bFoundError || ControlRigGraphNode->ErrorType <= (int32)EMessageSeverity::Error;
 					bFoundWarning = bFoundWarning || ControlRigGraphNode->ErrorType <= (int32)EMessageSeverity::Warning;
+
+					if(ControlRigGraphNode->ErrorType <= (int32)EMessageSeverity::Warning)
+					{
+						if(!ControlRig->VM->WasInstructionVisitedDuringLastRun(ControlRigGraphNode->GetInstructionIndex(true)) &&
+							!ControlRig->VM->WasInstructionVisitedDuringLastRun(ControlRigGraphNode->GetInstructionIndex(false)))
+						{
+							continue;
+						}
+					}
 				}
+
 				GraphNode->ErrorType = int32(EMessageSeverity::Info) + 1;
 			}
 
 			// update the nodes' error messages
-			bool bFoundErrorOrWarningInLog = false;
 			for (const FControlRigLog::FLogEntry& Entry : ControlRigLog.Entries)
 			{
 				URigVMNode* ModelNode = Cast<URigVMNode>(ByteCode.GetSubjectForInstruction(Entry.InstructionIndex));
@@ -6523,7 +6576,6 @@ void FControlRigEditor::UpdateGraphCompilerErrors()
 
 				bFoundError = bFoundError || Entry.Severity <= EMessageSeverity::Error;
 				bFoundWarning = bFoundWarning || Entry.Severity <= EMessageSeverity::Warning;
-				bFoundErrorOrWarningInLog = bFoundErrorOrWarningInLog || Entry.Severity <= EMessageSeverity::Warning;
 
 				int32 ErrorType = (int32)Entry.Severity;
 				if (GraphNode->ErrorType < ErrorType)
@@ -6540,47 +6592,18 @@ void FControlRigEditor::UpdateGraphCompilerErrors()
 					GraphNode->ErrorType = ErrorType;
 				}
 			}
-			bAnyErrorsLeft = bFoundErrorOrWarningInLog;
 
+			bAnyErrorsLeft = false;
 			for (UEdGraphNode* GraphNode : Graph->Nodes)
 			{
 				GraphNode->bHasCompilerMessage = GraphNode->ErrorType <= int32(EMessageSeverity::Info);
+				bAnyErrorsLeft = bAnyErrorsLeft || GraphNode->bHasCompilerMessage; 
 			}
 
 			if (bFoundError)
 			{
 				Blueprint->Status = BS_Error;
 				Blueprint->MarkPackageDirty();
-
-				/*
-				if (bFoundErrorOrWarningInLog)
-				{
-					FNotificationInfo Info(LOCTEXT("ControlRigBlueprintCompilerUnitErrorMessage", "There has been a compiler error.\nCheck the Execution Stack view."));
-					Info.bUseSuccessFailIcons = true;
-					Info.Image = FAppStyle::GetBrush(TEXT("MessageLog.Error"));
-					Info.bFireAndForget = true;
-					Info.FadeOutDuration = 5.0f;
-					Info.ExpireDuration = 5.0f;
-					TSharedPtr<SNotificationItem> NotificationPtr = FSlateNotificationManager::Get().AddNotification(Info);
-					NotificationPtr->SetCompletionState(SNotificationItem::CS_Success);
-				}
-				*/
-			}
-			else if (bFoundWarning)
-			{
-				/*
-				if (bFoundErrorOrWarningInLog)
-				{
-					FNotificationInfo Info(LOCTEXT("ControlRigBlueprintCompilerUnitWarningMessage", "There has been a compiler warning.\nCheck the Execution Stack view."));
-					Info.bUseSuccessFailIcons = true;
-					Info.Image = FAppStyle::GetBrush(TEXT("MessageLog.Warning"));
-					Info.bFireAndForget = true;
-					Info.FadeOutDuration = 5.0f;
-					Info.ExpireDuration = 5.0f;
-					TSharedPtr<SNotificationItem> NotificationPtr = FSlateNotificationManager::Get().AddNotification(Info);
-					NotificationPtr->SetCompletionState(SNotificationItem::CS_Success);
-				}
-				*/
 			}
 		}
 
@@ -6790,5 +6813,43 @@ FString* FControlRigEditor::GetSnippetStorage(int32 InSnippetIndex)
 	return nullptr;
 }
 
+void FControlRigEditor::OnPreConstruction_AnyThread(UControlRig* InRig, const EControlRigState InState,
+	const FName& InEventName)
+{
+	HierarchyHashBeforeConstruction = InRig->GetHierarchy()->GetTopologyHash(false);
+}
+
+void FControlRigEditor::OnPostConstruction_AnyThread(UControlRig* InRig, const EControlRigState InState,
+	const FName& InEventName)
+{
+	uint32 HierarchyHash = INDEX_NONE;
+	Swap(HierarchyHash, HierarchyHashBeforeConstruction);
+
+	if(HierarchyHash != InRig->GetHierarchy()->GetTopologyHash(false))
+	{
+		auto Task = [this, InRig]()
+		{
+			CacheNameLists();
+			SynchronizeViewportBoneSelection();
+			RebindToSkeletalMeshComponent();
+			if(DetailViewShowsAnyRigElement())
+			{
+				SetDetailViewForRigElements();
+			}
+		};
+				
+		if(IsInGameThread())
+		{
+			Task();
+		}
+		else
+		{
+			FFunctionGraphTask::CreateAndDispatchWhenReady([Task]()
+			{
+				Task();
+			}, TStatId(), NULL, ENamedThreads::GameThread);
+		}
+	}
+}
 
 #undef LOCTEXT_NAMESPACE
diff --git a/Engine/Plugins/Animation/ControlRig/Source/ControlRigEditor/Private/Editor/ControlRigEditor.h b/Engine/Plugins/Animation/ControlRig/Source/ControlRigEditor/Private/Editor/ControlRigEditor.h
--- a/Engine/Plugins/Animation/ControlRig/Source/ControlRigEditor/Private/Editor/ControlRigEditor.h
+++ b/Engine/Plugins/Animation/ControlRig/Source/ControlRigEditor/Private/Editor/ControlRigEditor.h
@@ -145,6 +145,9 @@ public:
 	// Gets the Control Rig Blueprint being edited/viewed
 	UControlRigBlueprint* GetControlRigBlueprint() const;
 
+	// returns the hierarchy being debugged
+	URigHierarchy* GetHierarchyBeingDebugged() const;
+
 	void SetDetailObjects(const TArray<UObject*>& InObjects);
 	void SetDetailObjects(const TArray<UObject*>& InObjects, bool bChangeUISelectionState);
 	void SetDetailViewForRigElements();
@@ -483,6 +486,12 @@ protected:
 
 	FDelegateHandle PropertyChangedHandle;
 
+	void OnPreConstruction_AnyThread(UControlRig* InRig, const EControlRigState InState, const FName& InEventName);
+	void OnPostConstruction_AnyThread(UControlRig* InRig, const EControlRigState InState, const FName& InEventName);
+
+	bool IsConstructionEventRunning() const { return HierarchyHashBeforeConstruction != INDEX_NONE; }
+	uint32 HierarchyHashBeforeConstruction;
+
 	static const TArray<FName> ForwardsSolveEventQueue;
 	static const TArray<FName> BackwardsSolveEventQueue;
 	static const TArray<FName> ConstructionEventQueue;
diff --git a/Engine/Plugins/Animation/ControlRig/Source/ControlRigEditor/Private/Editor/SRigHierarchy.cpp b/Engine/Plugins/Animation/ControlRig/Source/ControlRigEditor/Private/Editor/SRigHierarchy.cpp
--- a/Engine/Plugins/Animation/ControlRig/Source/ControlRigEditor/Private/Editor/SRigHierarchy.cpp
+++ b/Engine/Plugins/Animation/ControlRig/Source/ControlRigEditor/Private/Editor/SRigHierarchy.cpp
@@ -113,7 +113,7 @@ void SRigHierarchy::Construct(const FArguments& InArgs, TSharedRef<FControlRigEd
 
 	// setup all delegates for the rig hierarchy widget
 	FRigTreeDelegates Delegates;
-	Delegates.OnGetHierarchy = FOnGetRigTreeHierarchy::CreateSP(this, &SRigHierarchy::GetHierarchyForTopology);
+	Delegates.OnGetHierarchy = FOnGetRigTreeHierarchy::CreateSP(this, &SRigHierarchy::GetHierarchyForTreeView);
 	Delegates.OnGetDisplaySettings = FOnGetRigTreeDisplaySettings::CreateSP(this, &SRigHierarchy::GetDisplaySettings);
 	Delegates.OnRenameElement = FOnRigTreeRenameElement::CreateSP(this, &SRigHierarchy::HandleRenameElement);
 	Delegates.OnVerifyElementNameChanged = FOnRigTreeVerifyElementNameChanged::CreateSP(this, &SRigHierarchy::HandleVerifyNameChanged);
@@ -217,17 +217,6 @@ void SRigHierarchy::Construct(const FArguments& InArgs, TSharedRef<FControlRigEd
 		[
 			SNew(SBorder)
 			.Padding(0.0f)
-			.BorderImage(FControlRigEditorStyle::Get().GetBrush( "ControlRig.Viewport.Border"))
-			.BorderBackgroundColor_Lambda([this]()
-			{
-				FLinearColor Color = FLinearColor::Transparent;
-				if(DisplaySettings.bShowDynamicHierarchy)
-				{
-					Color = UControlRigEditorSettings::Get()->DynamicHierarchyBorderColor;
-				}
-				return Color;
-			})
-			.Padding(0.0f)
 			.ShowEffectWhenDisabled(false)
 			[
 				SNew(SBorder)
@@ -256,6 +245,8 @@ void SRigHierarchy::Construct(const FArguments& InArgs, TSharedRef<FControlRigEd
 	];
 
 	bIsChangingRigHierarchy = false;
+	HierarchyHashBeforeConstruction = INDEX_NONE;
+	
 	RefreshTreeView();
 
 	if (ControlRigEditor.IsValid())
@@ -298,17 +289,20 @@ void SRigHierarchy::BindCommands()
 	const FControlRigHierarchyCommands& Commands = FControlRigHierarchyCommands::Get();
 
 	CommandList->MapAction(Commands.AddBoneItem,
-		FExecuteAction::CreateSP(this, &SRigHierarchy::HandleNewItem, ERigElementType::Bone, false));
+		FExecuteAction::CreateSP(this, &SRigHierarchy::HandleNewItem, ERigElementType::Bone, false),
+		FCanExecuteAction::CreateSP(this, &SRigHierarchy::IsNonProceduralElementSelected));
 
 	CommandList->MapAction(Commands.AddControlItem,
-		FExecuteAction::CreateSP(this, &SRigHierarchy::HandleNewItem, ERigElementType::Control, false));
+		FExecuteAction::CreateSP(this, &SRigHierarchy::HandleNewItem, ERigElementType::Control, false),
+		FCanExecuteAction::CreateSP(this, &SRigHierarchy::IsNonProceduralElementSelected));
 
 	CommandList->MapAction(Commands.AddAnimationChannelItem,
 		FExecuteAction::CreateSP(this, &SRigHierarchy::HandleNewItem, ERigElementType::Control, true),
-		FCanExecuteAction::CreateSP(this, &SRigHierarchy::IsControlSelected));
+		FCanExecuteAction::CreateSP(this, &SRigHierarchy::IsControlSelected, false));
 
 	CommandList->MapAction(Commands.AddNullItem,
-		FExecuteAction::CreateSP(this, &SRigHierarchy::HandleNewItem, ERigElementType::Null, false));
+		FExecuteAction::CreateSP(this, &SRigHierarchy::HandleNewItem, ERigElementType::Null, false),
+		FCanExecuteAction::CreateSP(this, &SRigHierarchy::IsNonProceduralElementSelected));
 
 	CommandList->MapAction(Commands.DuplicateItem,
 		FExecuteAction::CreateSP(this, &SRigHierarchy::HandleDuplicateItem),
@@ -345,45 +339,44 @@ void SRigHierarchy::BindCommands()
 	CommandList->MapAction(
 		Commands.ResetTransform,
 		FExecuteAction::CreateSP(this, &SRigHierarchy::HandleResetTransform, true),
-		FCanExecuteAction::CreateSP(this, &SRigHierarchy::IsMultiSelected));
+		FCanExecuteAction::CreateSP(this, &SRigHierarchy::IsMultiSelected, true));
 
 	CommandList->MapAction(
 		Commands.ResetAllTransforms,
-		FExecuteAction::CreateSP(this, &SRigHierarchy::HandleResetTransform, false),
-		FCanExecuteAction::CreateSP(this, &SRigHierarchy::CanPasteItems));
+		FExecuteAction::CreateSP(this, &SRigHierarchy::HandleResetTransform, false));
 
 	CommandList->MapAction(
 		Commands.SetInitialTransformFromClosestBone,
 		FExecuteAction::CreateSP(this, &SRigHierarchy::HandleSetInitialTransformFromClosestBone),
-		FCanExecuteAction::CreateSP(this, &SRigHierarchy::IsControlOrNullSelected));
+		FCanExecuteAction::CreateSP(this, &SRigHierarchy::IsControlOrNullSelected, false));
 
 	CommandList->MapAction(
 		Commands.SetInitialTransformFromCurrentTransform,
 		FExecuteAction::CreateSP(this, &SRigHierarchy::HandleSetInitialTransformFromCurrentTransform),
-		FCanExecuteAction::CreateSP(this, &SRigHierarchy::IsMultiSelected));
+		FCanExecuteAction::CreateSP(this, &SRigHierarchy::IsMultiSelected, false));
 
 	CommandList->MapAction(
 		Commands.SetShapeTransformFromCurrent,
 		FExecuteAction::CreateSP(this, &SRigHierarchy::HandleSetShapeTransformFromCurrent),
-		FCanExecuteAction::CreateSP(this, &SRigHierarchy::IsControlSelected));
+		FCanExecuteAction::CreateSP(this, &SRigHierarchy::IsControlSelected, false));
 
 	CommandList->MapAction(
 		Commands.FrameSelection,
 		FExecuteAction::CreateSP(this, &SRigHierarchy::HandleFrameSelection),
-		FCanExecuteAction::CreateSP(this, &SRigHierarchy::IsMultiSelected));
+		FCanExecuteAction::CreateSP(this, &SRigHierarchy::IsMultiSelected, true));
 
 	CommandList->MapAction(
 		Commands.ControlBoneTransform,
 		FExecuteAction::CreateSP(this, &SRigHierarchy::HandleControlBoneOrSpaceTransform),
-		FCanExecuteAction::CreateSP(this, &SRigHierarchy::IsSingleBoneSelected),
+		FCanExecuteAction::CreateSP(this, &SRigHierarchy::IsSingleBoneSelected, false),
 		FIsActionChecked(),
-		FIsActionButtonVisible::CreateSP(this, &SRigHierarchy::IsSingleBoneSelected)
+		FIsActionButtonVisible::CreateSP(this, &SRigHierarchy::IsSingleBoneSelected, false)
 		);
 
 	CommandList->MapAction(
 		Commands.Unparent,
 		FExecuteAction::CreateSP(this, &SRigHierarchy::HandleUnparent),
-		FCanExecuteAction::CreateSP(this, &SRigHierarchy::IsMultiSelected));
+		FCanExecuteAction::CreateSP(this, &SRigHierarchy::IsMultiSelected, false));
 
 	CommandList->MapAction(
 		Commands.FilteringFlattensHierarchy,
@@ -433,17 +426,6 @@ void SRigHierarchy::BindCommands()
 		FCanExecuteAction(),
 		FIsActionChecked::CreateLambda([this]() { return DisplaySettings.bShowReferences; }));
 
-	CommandList->MapAction(
-		Commands.ShowDynamicHierarchy,
-		FExecuteAction::CreateLambda([this]()
-		{
-			DisplaySettings.bShowDynamicHierarchy = !DisplaySettings.bShowDynamicHierarchy;
-			HandleSetObjectBeingDebugged(GetControlRigEditor()->GetControlRigBlueprint()->GetObjectBeingDebugged());
-			RefreshTreeView();
-		}),
-		FCanExecuteAction(),
-		FIsActionChecked::CreateLambda([this]() { return DisplaySettings.bShowDynamicHierarchy; }));
-
 	CommandList->MapAction(
 		Commands.ToggleControlShapeTransformEdit,
 		FExecuteAction::CreateLambda([this]()
@@ -454,7 +436,7 @@ void SRigHierarchy::BindCommands()
 	CommandList->MapAction(
 		Commands.SpaceSwitching,
 		FExecuteAction::CreateSP(this, &SRigHierarchy::HandleTestSpaceSwitching),
-		FCanExecuteAction::CreateSP(this, &SRigHierarchy::IsControlSelected));
+		FCanExecuteAction::CreateSP(this, &SRigHierarchy::IsControlSelected, true));
 
 	CommandList->MapAction(
 		Commands.ShowIconColors,
@@ -474,7 +456,7 @@ FReply SRigHierarchy::OnKeyDown(const FGeometry& MyGeometry, const FKeyEvent& In
 
 EVisibility SRigHierarchy::IsToolbarVisible() const
 {
-	if (URigHierarchy* Hierarchy = GetDebuggedHierarchy())
+	if (URigHierarchy* Hierarchy = GetHierarchy())
 	{
 		if (Hierarchy->Num(ERigElementType::Bone) > 0)
 		{
@@ -486,7 +468,7 @@ EVisibility SRigHierarchy::IsToolbarVisible() const
 
 EVisibility SRigHierarchy::IsSearchbarVisible() const
 {
-	if (URigHierarchy* Hierarchy = GetDebuggedHierarchy())
+	if (URigHierarchy* Hierarchy = GetHierarchy())
 	{
 		if ((Hierarchy->Num(ERigElementType::Bone) +
 			Hierarchy->Num(ERigElementType::Null) +
@@ -635,7 +617,7 @@ void SRigHierarchy::OnHierarchyModified(ERigHierarchyNotification InNotif, URigH
 		return;
 	}
 
-	if (bIsChangingRigHierarchy)
+	if (bIsChangingRigHierarchy || IsConstructionEventRunning())
 	{
 		return;
 	}
@@ -687,30 +669,27 @@ void SRigHierarchy::OnHierarchyModified(ERigHierarchyNotification InNotif, URigH
 		}
 		case ERigHierarchyNotification::ParentWeightsChanged:
 		{
-			if(DisplaySettings.bShowDynamicHierarchy)
+			if(URigHierarchy* Hierarchy = GetHierarchy())
 			{
-				if(URigHierarchy* Hierarchy = GetDebuggedHierarchy())
+				if(InElement)
 				{
-					if(InElement)
+					TArray<FRigElementWeight> ParentWeights = Hierarchy->GetParentWeightArray(InElement->GetKey());
+					if(ParentWeights.Num() > 0)
 					{
-						TArray<FRigElementWeight> ParentWeights = Hierarchy->GetParentWeightArray(InElement->GetKey());
-						if(ParentWeights.Num() > 0)
+						TArray<FRigElementKey> ParentKeys = Hierarchy->GetParents(InElement->GetKey());
+						check(ParentKeys.Num() == ParentWeights.Num());
+						for(int32 ParentIndex=0;ParentIndex<ParentKeys.Num();ParentIndex++)
 						{
-							TArray<FRigElementKey> ParentKeys = Hierarchy->GetParents(InElement->GetKey());
-							check(ParentKeys.Num() == ParentWeights.Num());
-							for(int32 ParentIndex=0;ParentIndex<ParentKeys.Num();ParentIndex++)
+							if(ParentWeights[ParentIndex].IsAlmostZero())
 							{
-								if(ParentWeights[ParentIndex].IsAlmostZero())
-								{
-									continue;
-								}
-	
-								if(TreeView->ReparentElement(InElement->GetKey(), ParentKeys[ParentIndex]))
-								{
-									RefreshTreeView(false);
-								}
-								break;
+								continue;
 							}
+
+							if(TreeView->ReparentElement(InElement->GetKey(), ParentKeys[ParentIndex]))
+							{
+								RefreshTreeView(false);
+							}
+							break;
 						}
 					}
 				}
@@ -781,10 +760,6 @@ void SRigHierarchy::OnHierarchyModified(ERigHierarchyNotification InNotif, URigH
 
 void SRigHierarchy::OnHierarchyModified_AnyThread(ERigHierarchyNotification InNotif, URigHierarchy* InHierarchy, const FRigBaseElement* InElement)
 {
-	if(!DisplaySettings.bShowDynamicHierarchy)
-	{
-		return;
-	}
 	if(bIsChangingRigHierarchy)
 	{
 		return;
@@ -800,6 +775,11 @@ void SRigHierarchy::OnHierarchyModified_AnyThread(ERigHierarchyNotification InNo
 		return;
 	}
 
+	if(IsConstructionEventRunning())
+	{
+		return;
+	}
+
 	if(IsInGameThread())
 	{
 		OnHierarchyModified(InNotif, InHierarchy, InElement);
@@ -861,11 +841,74 @@ void SRigHierarchy::HandleSetObjectBeingDebugged(UObject* InObject)
 		ControlRigBeingDebuggedPtr = ControlRig;
 		ControlRig->GetHierarchy()->OnModified().RemoveAll(this);
 		ControlRig->GetHierarchy()->OnModified().AddSP(this, &SRigHierarchy::OnHierarchyModified_AnyThread);
+		ControlRig->OnPreConstructionForUI_AnyThread().RemoveAll(this);
+		ControlRig->OnPreConstructionForUI_AnyThread().AddSP(this, &SRigHierarchy::OnPreConstruction_AnyThread);
+		ControlRig->OnPostConstruction_AnyThread().RemoveAll(this);
+		ControlRig->OnPostConstruction_AnyThread().AddSP(this, &SRigHierarchy::OnPostConstruction_AnyThread);
 	}
 
 	RefreshTreeView();
 }
 
+void SRigHierarchy::OnPreConstruction_AnyThread(UControlRig* InRig, const EControlRigState InState, const FName& InEventName)
+{
+	if(InRig != ControlRigBeingDebuggedPtr.Get())
+	{
+		return;
+	}
+	HierarchyHashBeforeConstruction = InRig->GetHierarchy()->GetTopologyHash(false);
+	SelectionBeforeConstruction = InRig->GetHierarchy()->GetSelectedKeys();
+}
+
+void SRigHierarchy::OnPostConstruction_AnyThread(UControlRig* InRig, const EControlRigState InState, const FName& InEventName)
+{
+	if(InRig != ControlRigBeingDebuggedPtr.Get())
+	{
+		return;
+	}
+
+	uint32 HierarchyHash = INDEX_NONE;
+	Swap(HierarchyHash, HierarchyHashBeforeConstruction);
+			
+	if(HierarchyHash != InRig->GetHierarchy()->GetTopologyHash(false))
+	{
+		auto Task = [this]()
+		{
+			TGuardValue<bool> GuardRigHierarchyChanges(bIsChangingRigHierarchy, true);
+
+			RefreshTreeView(true);
+
+			TreeView->ClearSelection();
+			if(!SelectionBeforeConstruction.IsEmpty())
+			{
+				for (int32 RootIndex = 0; RootIndex < TreeView->RootElements.Num(); ++RootIndex)
+				{
+					for(const FRigElementKey& Key : SelectionBeforeConstruction)
+					{
+						TSharedPtr<FRigTreeElement> Found = TreeView->FindElement(Key, TreeView->RootElements[RootIndex]);
+						if (Found.IsValid())
+						{
+							TreeView->SetItemSelection(Found, true, ESelectInfo::OnNavigation);
+						}
+					}
+				}
+			}
+		};
+				
+		if(IsInGameThread())
+		{
+			Task();
+		}
+		else
+		{
+			FFunctionGraphTask::CreateAndDispatchWhenReady([Task]()
+			{
+				Task();
+			}, TStatId(), NULL, ENamedThreads::GameThread);
+		}
+	}
+}
+
 void SRigHierarchy::ClearDetailPanel() const
 {
 	if(ControlRigEditor.IsValid())
@@ -884,7 +927,6 @@ TSharedRef< SWidget > SRigHierarchy::CreateFilterMenu()
 	MenuBuilder.BeginSection("FilterOptions", LOCTEXT("OptionsMenuHeading", "Options"));
 	{
 		MenuBuilder.AddMenuEntry(Actions.FilteringFlattensHierarchy);
-		MenuBuilder.AddMenuEntry(Actions.ShowDynamicHierarchy);
 		MenuBuilder.AddMenuEntry(Actions.HideParentsWhenFiltering);
 	}
 	MenuBuilder.EndSection();
@@ -1073,7 +1115,7 @@ void SRigHierarchy::CreateContextMenu() const
 	{
 		UToolMenu* Menu = ToolMenus->RegisterMenu(MenuName);
 		
-		Menu->AddDynamicSection(NAME_None, FNewToolMenuDelegate::CreateLambda([](UToolMenu* InMenu)
+		Menu->AddDynamicSection(NAME_None, FNewToolMenuDelegate::CreateLambda([this](UToolMenu* InMenu)
 			{
 				UControlRigContextMenuContext* MainContext = InMenu->FindContext<UControlRigContextMenuContext>();
 				
@@ -1111,19 +1153,53 @@ void SRigHierarchy::CreateContextMenu() const
 					ElementsSection.AddMenuEntry(Commands.RenameItem);
 					ElementsSection.AddMenuEntry(Commands.MirrorItem);
 
-					if (RigHierarchyPanel->IsSingleBoneSelected() || RigHierarchyPanel->IsControlSelected())
+					if(RigHierarchyPanel->IsProceduralElementSelected() && ControlRigBlueprint.IsValid())
+					{
+						ElementsSection.AddMenuEntry(
+							"SelectSpawnerNode",
+							LOCTEXT("SelectSpawnerNode", "Select Spawner Node"),
+							LOCTEXT("SelectSpawnerNode_Tooltip", "Selects the node that spawn / added this element."),
+							FSlateIcon(),
+							FUIAction(FExecuteAction::CreateLambda([this]() {
+								const TArray<const FRigBaseElement*> Elements = GetHierarchy()->GetSelectedElements();
+								for(const FRigBaseElement* Element : Elements)
+								{
+									if(Element->IsProcedural())
+									{
+										const int32 InstructionIndex = Element->GetCreatedAtInstructionIndex();
+										if(UControlRig* ControlRig = Cast<UControlRig>(ControlRigBlueprint->GetObjectBeingDebugged()))
+										{
+											if(ControlRig->VM)
+											{
+												if(URigVMNode* Node = Cast<URigVMNode>(ControlRig->VM->GetByteCode().GetSubjectForInstruction(InstructionIndex)))
+												{
+													if(URigVMController* Controller = ControlRigBlueprint->GetController(Node->GetGraph()))
+													{
+														Controller->SelectNode(Node);
+														Controller->RequestJumpToHyperlinkDelegate.ExecuteIfBound(Node);
+													}
+												}
+											}
+										}
+									}
+								}
+							})
+						));
+					}
+
+					if (RigHierarchyPanel->IsSingleBoneSelected(false) || RigHierarchyPanel->IsControlSelected(false))
 					{
 						FToolMenuSection& InteractionSection = InMenu->AddSection(TEXT("Interaction"), LOCTEXT("InteractionHeader", "Interaction"));
-						if(RigHierarchyPanel->IsSingleBoneSelected())
+						if(RigHierarchyPanel->IsSingleBoneSelected(false))
 						{
 							InteractionSection.AddMenuEntry(Commands.ControlBoneTransform);
 						}
-						else if(RigHierarchyPanel->IsControlSelected())
+						else if(RigHierarchyPanel->IsControlSelected(false))
 						{
 							InteractionSection.AddMenuEntry(Commands.SpaceSwitching);
 						}
 					}
-					
+
 					FToolMenuSection& CopyPasteSection = InMenu->AddSection(TEXT("Copy&Paste"), LOCTEXT("Copy&PasteHeader", "Copy & Paste"));
 					CopyPasteSection.AddMenuEntry(Commands.CopyItems);
 					CopyPasteSection.AddMenuEntry(Commands.PasteItems);
@@ -1340,7 +1416,7 @@ void SRigHierarchy::ImportHierarchy(const FAssetData& InAssetData)
 		return;
 	}
 
-	URigHierarchy* Hierarchy = GetHierarchy();
+	URigHierarchy* Hierarchy = GetDefaultHierarchy();
 	USkeletalMesh* Mesh = Cast<USkeletalMesh>(InAssetData.GetAsset());
 	if (Mesh && Hierarchy)
 	{
@@ -1410,36 +1486,57 @@ void SRigHierarchy::ImportHierarchy(const FAssetData& InAssetData)
 	}
 }
 
-bool SRigHierarchy::IsMultiSelected() const
+bool SRigHierarchy::IsMultiSelected(bool bIncludeProcedural) const
 {
-	return GetSelectedKeys().Num() > 0;
+	if(GetSelectedKeys().Num() > 0)
+	{
+		if(!bIncludeProcedural && IsProceduralElementSelected())
+		{
+			return false;
+		}
+		return true;
+	}
+	return false;
 }
 
-bool SRigHierarchy::IsSingleSelected() const
+bool SRigHierarchy::IsSingleSelected(bool bIncludeProcedural) const
 {
-	return GetSelectedKeys().Num() == 1;
+	if(GetSelectedKeys().Num() == 1)
+	{
+		if(!bIncludeProcedural && IsProceduralElementSelected())
+		{
+			return false;
+		}
+		return true;
+	}
+	return false;
 }
 
-bool SRigHierarchy::IsSingleBoneSelected() const
+bool SRigHierarchy::IsSingleBoneSelected(bool bIncludeProcedural) const
 {
-	if(!IsSingleSelected())
+	if(!IsSingleSelected(bIncludeProcedural))
 	{
 		return false;
 	}
 	return GetSelectedKeys()[0].Type == ERigElementType::Bone;
 }
 
-bool SRigHierarchy::IsSingleNullSelected() const
+bool SRigHierarchy::IsSingleNullSelected(bool bIncludeProcedural) const
 {
-	if(!IsSingleSelected())
+	if(!IsSingleSelected(bIncludeProcedural))
 	{
 		return false;
 	}
 	return GetSelectedKeys()[0].Type == ERigElementType::Null;
 }
 
-bool SRigHierarchy::IsControlSelected() const
+bool SRigHierarchy::IsControlSelected(bool bIncludeProcedural) const
 {
+	if(!bIncludeProcedural && IsProceduralElementSelected())
+	{
+		return false;
+	}
+	
 	TArray<FRigElementKey> SelectedKeys = GetSelectedKeys();
 	for (const FRigElementKey& SelectedKey : SelectedKeys)
 	{
@@ -1451,8 +1548,13 @@ bool SRigHierarchy::IsControlSelected() const
 	return false;
 }
 
-bool SRigHierarchy::IsControlOrNullSelected() const
+bool SRigHierarchy::IsControlOrNullSelected(bool bIncludeProcedural) const
 {
+	if(!bIncludeProcedural && IsProceduralElementSelected())
+	{
+		return false;
+	}
+
 	TArray<FRigElementKey> SelectedKeys = GetSelectedKeys();
 	for (const FRigElementKey& SelectedKey : SelectedKeys)
 	{
@@ -1468,6 +1570,32 @@ bool SRigHierarchy::IsControlOrNullSelected() const
 	return false;
 }
 
+bool SRigHierarchy::IsProceduralElementSelected() const
+{
+	TArray<FRigElementKey> SelectedKeys = GetSelectedKeys();
+	for (const FRigElementKey& SelectedKey : SelectedKeys)
+	{
+		if(!GetHierarchy()->IsProcedural(SelectedKey))
+		{
+			return false;
+		}
+	}
+	return true;
+}
+
+bool SRigHierarchy::IsNonProceduralElementSelected() const
+{
+	TArray<FRigElementKey> SelectedKeys = GetSelectedKeys();
+	for (const FRigElementKey& SelectedKey : SelectedKeys)
+	{
+		if(GetHierarchy()->IsProcedural(SelectedKey))
+		{
+			return false;
+		}
+	}
+	return true;
+}
+
 void SRigHierarchy::HandleDeleteItem()
 {
 	if(!ControlRigEditor.IsValid())
@@ -1475,7 +1603,12 @@ void SRigHierarchy::HandleDeleteItem()
 		return;
 	}
 
-	URigHierarchy* Hierarchy = GetHierarchy();
+	if(!CanDeleteItem())
+	{
+		return;
+	}
+
+	URigHierarchy* Hierarchy = GetDefaultHierarchy();
  	if (Hierarchy)
  	{
 		TArray<FRigElementKey> RemovedItems;
@@ -1545,7 +1678,7 @@ void SRigHierarchy::HandleDeleteItem()
 
 bool SRigHierarchy::CanDeleteItem() const
 {
-	return IsMultiSelected();
+	return IsMultiSelected(false);
 }
 
 /** Create Item */
@@ -1556,7 +1689,7 @@ void SRigHierarchy::HandleNewItem(ERigElementType InElementType, bool bIsAnimati
 		return;
 	}
 
-	URigHierarchy* Hierarchy = GetHierarchy();
+	URigHierarchy* Hierarchy = GetDefaultHierarchy();
 	if (Hierarchy)
 	{
 		// unselect current selected item
@@ -1662,7 +1795,7 @@ void SRigHierarchy::HandleNewItem(ERigElementType InElementType, bool bIsAnimati
 /** Check whether we can deleting the selected item(s) */
 bool SRigHierarchy::CanDuplicateItem() const
 {
-	return IsMultiSelected();
+	return IsMultiSelected(false);
 }
 
 /** Duplicate Item */
@@ -1673,7 +1806,7 @@ void SRigHierarchy::HandleDuplicateItem()
 		return;
 	}
 
-	URigHierarchy* Hierarchy = GetHierarchy();
+	URigHierarchy* Hierarchy = GetDefaultHierarchy();
 	if (Hierarchy)
 	{
 		ClearDetailPanel();
@@ -1710,7 +1843,7 @@ void SRigHierarchy::HandleMirrorItem()
 		return;
 	}
 	
-	URigHierarchy* Hierarchy = GetHierarchy();
+	URigHierarchy* Hierarchy = GetDefaultHierarchy();
 	if (Hierarchy)
 	{
 		URigHierarchyController* Controller = Hierarchy->GetController(true);
@@ -1758,7 +1891,7 @@ void SRigHierarchy::HandleMirrorItem()
 /** Check whether we can deleting the selected item(s) */
 bool SRigHierarchy::CanRenameItem() const
 {
-	if(IsSingleSelected())
+	if(IsSingleSelected(false))
 	{
 		const FRigElementKey Key = GetSelectedKeys()[0];
 		if(Key.Type == ERigElementType::RigidBody ||
@@ -1768,7 +1901,7 @@ bool SRigHierarchy::CanRenameItem() const
 		}
 		if(Key.Type == ERigElementType::Control)
 		{
-			if(URigHierarchy* DebuggedHierarchy = GetDebuggedHierarchy())
+			if(URigHierarchy* DebuggedHierarchy = GetHierarchy())
 			{
 				if(FRigControlElement* ControlElement = DebuggedHierarchy->Find<FRigControlElement>(Key))
 				{
@@ -1797,7 +1930,7 @@ void SRigHierarchy::HandleRenameItem()
 		return;
 	}
 
-	URigHierarchy* Hierarchy = GetHierarchy();
+	URigHierarchy* Hierarchy = GetDefaultHierarchy();
 	if (Hierarchy)
 	{
 		FScopedTransaction Transaction(LOCTEXT("HierarchyTreeRenameSelected", "Rename selected item from hierarchy"));
@@ -1833,7 +1966,7 @@ void SRigHierarchy::HandleRenameItem()
 bool SRigHierarchy::CanPasteItems() const
 {
 	return true;
-}
+}	
 
 bool SRigHierarchy::CanCopyOrPasteItems() const
 {
@@ -1842,7 +1975,7 @@ bool SRigHierarchy::CanCopyOrPasteItems() const
 
 void SRigHierarchy::HandleCopyItems()
 {
-	if (URigHierarchy* Hierarchy = GetDebuggedHierarchy())
+	if (URigHierarchy* Hierarchy = GetHierarchy())
 	{
 		URigHierarchyController* Controller = Hierarchy->GetController(true);
 		check(Controller);
@@ -1855,7 +1988,7 @@ void SRigHierarchy::HandleCopyItems()
 
 void SRigHierarchy::HandlePasteItems()
 {
-	if (URigHierarchy* Hierarchy = GetHierarchy())
+	if (URigHierarchy* Hierarchy = GetDefaultHierarchy())
 	{
 		TGuardValue<bool> GuardRigHierarchyChanges(bIsChangingRigHierarchy, true);
 		TGuardValue<bool> SuspendBlueprintNotifs(ControlRigBlueprint->bSuspendAllNotifications, true);
@@ -1911,7 +2044,7 @@ void SRigHierarchy::HandlePasteGlobalTransforms()
 
 void SRigHierarchy::HandlePasteTransforms(ERigTransformType::Type InTransformType, bool bAffectChildren)
 {
-	if (URigHierarchy* Hierarchy = GetHierarchy())
+	if (URigHierarchy* Hierarchy = GetDefaultHierarchy())
 	{
 		FString Content;
 		FPlatformApplicationMisc::ClipboardPaste(Content);
@@ -1926,7 +2059,7 @@ void SRigHierarchy::HandlePasteTransforms(ERigTransformType::Type InTransformTyp
 			return;
 		}
 		
-		URigHierarchy* DebuggedHierarchy = GetDebuggedHierarchy();
+		URigHierarchy* DebuggedHierarchy = GetHierarchy();
 
 		const TArray<FRigElementKey> CurrentSelection = Hierarchy->GetSelectedKeys();
 		const int32 Count = FMath::Min<int32>(CurrentSelection.Num(), Data.Elements.Num());
@@ -1960,15 +2093,6 @@ void SRigHierarchy::HandlePasteTransforms(ERigTransformType::Type InTransformTyp
 }
 
 URigHierarchy* SRigHierarchy::GetHierarchy() const
-{
-	if (ControlRigBlueprint.IsValid())
-	{
-		return ControlRigBlueprint->Hierarchy;
-	}
-	return nullptr;
-}
-
-URigHierarchy* SRigHierarchy::GetDebuggedHierarchy() const
 {
 	if (ControlRigBlueprint.IsValid())
 	{
@@ -1984,17 +2108,16 @@ URigHierarchy* SRigHierarchy::GetDebuggedHierarchy() const
 			return CurrentRig->GetHierarchy();
 		}
 	}
-	return GetHierarchy();
+	return GetDefaultHierarchy();
 }
 
-const URigHierarchy* SRigHierarchy::GetHierarchyForTopology() const
+URigHierarchy* SRigHierarchy::GetDefaultHierarchy() const
 {
-	URigHierarchy* Hierarchy = GetHierarchy();
-	if(DisplaySettings.bShowDynamicHierarchy && ControlRigBeingDebuggedPtr.IsValid())
+	if (ControlRigBlueprint.IsValid())
 	{
-		Hierarchy = ControlRigBeingDebuggedPtr->GetHierarchy();
+		return ControlRigBlueprint->Hierarchy;
 	}
-	return Hierarchy;
+	return nullptr;
 }
 
 
@@ -2052,6 +2175,11 @@ TOptional<EItemDropZone> SRigHierarchy::OnCanAcceptDrop(const FDragDropEvent& Dr
 					return ReturnDropZone;
 				}
 
+				if(Hierarchy->IsProcedural(DraggedKey))
+				{
+					return ReturnDropZone;
+				}
+
 				if(Hierarchy->IsParentedTo(TargetItem->Key, DraggedKey))
 				{
 					return ReturnDropZone;
@@ -2059,6 +2187,12 @@ TOptional<EItemDropZone> SRigHierarchy::OnCanAcceptDrop(const FDragDropEvent& Dr
 			}
 		}
 
+		// don't allow dragging onto procedural items
+		if(!GetDefaultHierarchy()->Contains(TargetItem->Key))
+		{
+			return ReturnDropZone;
+		}
+
 		switch (TargetItem->Key.Type)
 		{
 			case ERigElementType::Bone:
@@ -2152,7 +2286,7 @@ FName SRigHierarchy::HandleRenameElement(const FRigElementKey& OldKey, const FSt
 	{
 		FScopedTransaction Transaction(LOCTEXT("HierarchyRename", "Rename Hierarchy Element"));
 
-		URigHierarchy* Hierarchy = GetHierarchy();
+		URigHierarchy* Hierarchy = GetDefaultHierarchy();
 		URigHierarchyController* Controller = Hierarchy->GetController(true);
 		check(Controller);
 
@@ -2266,11 +2400,11 @@ FReply SRigHierarchy::OnDrop(const FGeometry& MyGeometry, const FDragDropEvent&
 
 void SRigHierarchy::HandleResetTransform(bool bSelectionOnly)
 {
-	if ((IsMultiSelected() || !bSelectionOnly) && ControlRigEditor.IsValid())
+	if ((IsMultiSelected(true) || !bSelectionOnly) && ControlRigEditor.IsValid())
 	{
 		if (UControlRigBlueprint* Blueprint = ControlRigEditor.Pin()->GetControlRigBlueprint())
 		{
-			if (URigHierarchy* DebuggedHierarchy = GetDebuggedHierarchy())
+			if (URigHierarchy* DebuggedHierarchy = GetHierarchy())
 			{
 				FScopedTransaction Transaction(LOCTEXT("HierarchyResetTransforms", "Reset Transforms"));
 
@@ -2302,11 +2436,11 @@ void SRigHierarchy::HandleResetTransform(bool bSelectionOnly)
 
 void SRigHierarchy::HandleSetInitialTransformFromCurrentTransform()
 {
-	if (IsMultiSelected())
+	if (IsMultiSelected(false))
 	{
 		if (UControlRigBlueprint* Blueprint = ControlRigEditor.Pin()->GetControlRigBlueprint())
 		{
-			if (URigHierarchy* DebuggedHierarchy = GetDebuggedHierarchy())
+			if (URigHierarchy* DebuggedHierarchy = GetHierarchy())
 			{
 				FScopedTransaction Transaction(LOCTEXT("HierarchySetInitialTransforms", "Set Initial Transforms"));
 
@@ -2407,8 +2541,10 @@ void SRigHierarchy::HandleControlBoneOrSpaceTransform()
 		if (SelectedKeys[0].Type == ERigElementType::Bone ||
 			SelectedKeys[0].Type == ERigElementType::Null)
 		{
-			Blueprint->AddTransientControl(SelectedKeys[0]);
-			return;
+			if(!DebuggedControlRig->GetHierarchy()->IsProcedural(SelectedKeys[0]))
+			{
+				Blueprint->AddTransientControl(SelectedKeys[0]);
+			}
 		}
 	}
 }
@@ -2505,7 +2641,7 @@ void SRigHierarchy::HandleUnparent()
 
 bool SRigHierarchy::FindClosestBone(const FVector& Point, FName& OutRigElementName, FTransform& OutGlobalTransform) const
 {
-	if (URigHierarchy* DebuggedHierarchy = GetDebuggedHierarchy())
+	if (URigHierarchy* DebuggedHierarchy = GetHierarchy())
 	{
 		float NearestDistance = BIG_NUMBER;
 
@@ -2558,8 +2694,8 @@ FReply SRigHierarchy::ReparentOrMatchTransform(const TArray<FRigElementKey>& Dra
 {
 	bool bMatchTransforms = !bReparentItems;
 
-	URigHierarchy* Hierarchy = GetHierarchy();
-	URigHierarchy* DebuggedHierarchy = GetDebuggedHierarchy();
+	URigHierarchy* DebuggedHierarchy = GetHierarchy();
+	URigHierarchy* Hierarchy = GetDefaultHierarchy();
 
 	if (Hierarchy && ControlRigBlueprint.IsValid())
 	{
@@ -2716,11 +2852,11 @@ FReply SRigHierarchy::ReparentOrMatchTransform(const TArray<FRigElementKey>& Dra
 
 void SRigHierarchy::HandleSetInitialTransformFromClosestBone()
 {
-	if (IsControlOrNullSelected())
+	if (IsControlOrNullSelected(false))
 	{
 		if (UControlRigBlueprint* Blueprint = ControlRigEditor.Pin()->GetControlRigBlueprint())
 		{
-			if (URigHierarchy* DebuggedHierarchy = GetDebuggedHierarchy())
+			if (URigHierarchy* DebuggedHierarchy = GetHierarchy())
 			{
 				FScopedTransaction Transaction(LOCTEXT("HierarchySetInitialTransforms", "Set Initial Transforms"));
 
@@ -2796,11 +2932,11 @@ void SRigHierarchy::HandleSetInitialTransformFromClosestBone()
 
 void SRigHierarchy::HandleSetShapeTransformFromCurrent()
 {
-	if (IsControlSelected())
+	if (IsControlSelected(false))
 	{
 		if (UControlRigBlueprint* Blueprint = ControlRigEditor.Pin()->GetControlRigBlueprint())
 		{
-			if (URigHierarchy* DebuggedHierarchy = GetDebuggedHierarchy())
+			if (URigHierarchy* DebuggedHierarchy = GetHierarchy())
 			{
 				FScopedTransaction Transaction(LOCTEXT("HierarchySetShapeTransforms", "Set Shape Transforms"));
 
diff --git a/Engine/Plugins/Animation/ControlRig/Source/ControlRigEditor/Private/Editor/SRigHierarchy.h b/Engine/Plugins/Animation/ControlRig/Source/ControlRigEditor/Private/Editor/SRigHierarchy.h
--- a/Engine/Plugins/Animation/ControlRig/Source/ControlRigEditor/Private/Editor/SRigHierarchy.h
+++ b/Engine/Plugins/Animation/ControlRig/Source/ControlRigEditor/Private/Editor/SRigHierarchy.h
@@ -7,6 +7,7 @@
 #include "DragAndDrop/GraphNodeDragDropOp.h"
 #include "Engine/SkeletalMesh.h"
 #include "SRigHierarchyTreeView.h"
+#include "Units/RigUnitContext.h"
 #include "SRigHierarchy.generated.h"
 
 class SRigHierarchy;
@@ -179,17 +180,19 @@ private:
 	/** Command list we bind to */
 	TSharedPtr<FUICommandList> CommandList;
 
-	bool IsMultiSelected() const;
-	bool IsSingleSelected() const;
-	bool IsSingleBoneSelected() const;
-	bool IsSingleNullSelected() const;
-	bool IsControlSelected() const;
-	bool IsControlOrNullSelected() const;
+	bool IsMultiSelected(bool bIncludeProcedural) const;
+	bool IsSingleSelected(bool bIncludeProcedural) const;
+	bool IsSingleBoneSelected(bool bIncludeProcedural) const;
+	bool IsSingleNullSelected(bool bIncludeProcedural) const;
+	bool IsControlSelected(bool bIncludeProcedural) const;
+	bool IsControlOrNullSelected(bool bIncludeProcedural) const;
+	bool IsProceduralElementSelected() const;
+	bool IsNonProceduralElementSelected() const;
 
 	URigHierarchy* GetHierarchy() const;
-	URigHierarchy* GetDebuggedHierarchy() const;
-	const URigHierarchy* GetHierarchyForTopology() const;
-
+	URigHierarchy* GetDefaultHierarchy() const;
+	const URigHierarchy* GetHierarchyForTreeView() const { return GetHierarchy(); }
+	
 	void ImportHierarchy(const FAssetData& InAssetData);
 	void CreateImportMenu(FMenuBuilder& MenuBuilder);
 	void CreateRefreshMenu(FMenuBuilder& MenuBuilder);
@@ -220,6 +223,12 @@ private:
 	void OnHierarchyModified_AnyThread(ERigHierarchyNotification InNotif, URigHierarchy* InHierarchy, const FRigBaseElement* InElement);
 	void HandleRefreshEditorFromBlueprint(UControlRigBlueprint* InBlueprint);
 	void HandleSetObjectBeingDebugged(UObject* InObject);
+	void OnPreConstruction_AnyThread(UControlRig* InRig, const EControlRigState InState, const FName& InEventName);
+	void OnPostConstruction_AnyThread(UControlRig* InRig, const EControlRigState InState, const FName& InEventName);
+
+	bool IsConstructionEventRunning() const { return HierarchyHashBeforeConstruction != INDEX_NONE; }
+	uint32 HierarchyHashBeforeConstruction;
+	TArray<FRigElementKey> SelectionBeforeConstruction;
 
 public:
 	FName HandleRenameElement(const FRigElementKey& OldKey, const FString& NewName);
diff --git a/Engine/Plugins/Animation/ControlRig/Source/ControlRigEditor/Private/Editor/SRigHierarchyTreeView.cpp b/Engine/Plugins/Animation/ControlRig/Source/ControlRigEditor/Private/Editor/SRigHierarchyTreeView.cpp
--- a/Engine/Plugins/Animation/ControlRig/Source/ControlRigEditor/Private/Editor/SRigHierarchyTreeView.cpp
+++ b/Engine/Plugins/Animation/ControlRig/Source/ControlRigEditor/Private/Editor/SRigHierarchyTreeView.cpp
@@ -15,6 +15,7 @@
 #include "Widgets/Text/SInlineEditableTextBlock.h"
 #include "ControlRig.h"
 #include "ControlRigEditorStyle.h"
+#include "SRigHierarchy.h"
 #include "Settings/ControlRigSettings.h"
 #include "Graph/ControlRigGraphSchema.h"
 #include "Styling/AppStyle.h"
@@ -35,6 +36,7 @@ FRigTreeElement::FRigTreeElement(const FRigElementKey& InKey, TWeakPtr<SRigHiera
 	ChannelName = NAME_None;
 	bIsTransient = false;
 	bIsAnimationChannel = false;
+	bIsProcedural = false;
 	bSupportsRename = InSupportsRename;
 	FilterResult = InFilterResult;
 
@@ -42,13 +44,18 @@ FRigTreeElement::FRigTreeElement(const FRigElementKey& InKey, TWeakPtr<SRigHiera
 	{
 		if(const URigHierarchy* Hierarchy = InTreeView.Pin()->GetRigTreeDelegates().GetHierarchy())
 		{
-			if(const FRigControlElement* ControlElement = Hierarchy->Find<FRigControlElement>(Key))
+			if(const FRigBaseElement* Element = Hierarchy->Find(Key))
 			{
-				bIsTransient = ControlElement->Settings.bIsTransientControl;
-				bIsAnimationChannel = ControlElement->IsAnimationChannel();
-				if(bIsAnimationChannel)
+				bIsProcedural = Element->IsProcedural();
+				
+				if(const FRigControlElement* ControlElement = Cast<FRigControlElement>(Element))
 				{
-					ChannelName = ControlElement->GetDisplayName();
+					bIsTransient = ControlElement->Settings.bIsTransientControl;
+					bIsAnimationChannel = ControlElement->IsAnimationChannel();
+					if(bIsAnimationChannel)
+					{
+						ChannelName = ControlElement->GetDisplayName();
+					}
 				}
 			}
 
@@ -86,7 +93,10 @@ void FRigTreeElement::RefreshDisplaySettings(const URigHierarchy* InHierarchy, c
 	{
 		IconColor = FilterResult == ERigTreeFilterResult::Shown ? FSlateColor::UseForeground() : FSlateColor(FLinearColor::Gray * 0.5f);
 	}
-	TextColor = FilterResult == ERigTreeFilterResult::Shown ? FSlateColor::UseForeground() : FSlateColor(FLinearColor::Gray * 0.5f);
+
+	TextColor = FilterResult == ERigTreeFilterResult::Shown ?
+		(InHierarchy->IsProcedural(Key) ? FSlateColor(FLinearColor(0.9f, 0.8f, 0.4f)) : FSlateColor::UseForeground()) :
+		(InHierarchy->IsProcedural(Key) ? FSlateColor(FLinearColor(0.9f, 0.8f, 0.4f) * 0.5f) : FSlateColor(FLinearColor::Gray * 0.5f));
 }
 
 //////////////////////////////////////////////////////////////
@@ -249,6 +259,18 @@ bool SRigHierarchyTreeView::AddElement(FRigElementKey InKey, FRigElementKey InPa
 		return false;
 	}
 
+	// skip transient controls
+	if(const URigHierarchy* Hierarchy = Delegates.GetHierarchy())
+	{
+		if(const FRigControlElement* ControlElement = Hierarchy->Find<FRigControlElement>(InKey))
+		{
+			if(ControlElement->Settings.bIsTransientControl)
+			{
+				return false;
+			}
+		}
+	}
+
 	const FRigTreeDisplaySettings& Settings = Delegates.GetDisplaySettings();
 	const bool bSupportsRename = Delegates.OnRenameElement.IsBound();
 
@@ -402,24 +424,23 @@ bool SRigHierarchyTreeView::AddElement(const FRigBaseElement* InElement)
 		if(const URigHierarchy* Hierarchy = Delegates.GetHierarchy())
 		{
 			FRigElementKey ParentKey = Hierarchy->GetFirstParent(InElement->GetKey());
-			if(Settings.bShowDynamicHierarchy)
+
+			TArray<FRigElementWeight> ParentWeights = Hierarchy->GetParentWeightArray(InElement->GetKey());
+			if(ParentWeights.Num() > 0)
 			{
-				TArray<FRigElementWeight> ParentWeights = Hierarchy->GetParentWeightArray(InElement->GetKey());
-				if(ParentWeights.Num() > 0)
+				TArray<FRigElementKey> ParentKeys = Hierarchy->GetParents(InElement->GetKey());
+				check(ParentKeys.Num() == ParentWeights.Num());
+				for(int32 ParentIndex=0;ParentIndex<ParentKeys.Num();ParentIndex++)
 				{
-					TArray<FRigElementKey> ParentKeys = Hierarchy->GetParents(InElement->GetKey());
-					check(ParentKeys.Num() == ParentWeights.Num());
-					for(int32 ParentIndex=0;ParentIndex<ParentKeys.Num();ParentIndex++)
+					if(ParentWeights[ParentIndex].IsAlmostZero())
 					{
-						if(ParentWeights[ParentIndex].IsAlmostZero())
-						{
-							continue;
-						}
-						ParentKey = ParentKeys[ParentIndex];
-						break;
+						continue;
 					}
+					ParentKey = ParentKeys[ParentIndex];
+					break;
 				}
 			}
+
 			if (ParentKey.IsValid())
 			{
 				if(const FRigBaseElement* ParentElement = Hierarchy->Find(ParentKey))
diff --git a/Engine/Plugins/Animation/ControlRig/Source/ControlRigEditor/Private/Editor/SRigHierarchyTreeView.h b/Engine/Plugins/Animation/ControlRig/Source/ControlRigEditor/Private/Editor/SRigHierarchyTreeView.h
--- a/Engine/Plugins/Animation/ControlRig/Source/ControlRigEditor/Private/Editor/SRigHierarchyTreeView.h
+++ b/Engine/Plugins/Animation/ControlRig/Source/ControlRigEditor/Private/Editor/SRigHierarchyTreeView.h
@@ -25,7 +25,6 @@ struct CONTROLRIGEDITOR_API FRigTreeDisplaySettings
 		bShowNulls = true;
 		bShowRigidBodies = true;
 		bShowReferences = true;
-		bShowDynamicHierarchy = false;
 		bShowIconColors = true;
 	}
 	
@@ -55,9 +54,6 @@ struct CONTROLRIGEDITOR_API FRigTreeDisplaySettings
 	/** Whether or not to show references in the hierarchy */
 	bool bShowReferences;
 
-	/** Whether or not to show the static or dynamic hierarchy */
-	bool bShowDynamicHierarchy;
-
 	/** Whether to tint the icons with the element color */
 	bool bShowIconColors;
 };
@@ -174,6 +170,7 @@ public:
 	FName ChannelName;
 	bool bIsTransient;
 	bool bIsAnimationChannel;
+	bool bIsProcedural;
 	bool bSupportsRename;
 	TArray<TSharedPtr<FRigTreeElement>> Children;
 
diff --git a/Engine/Plugins/Animation/ControlRig/Source/ControlRigEditor/Private/Editor/SRigSpacePickerWidget.cpp b/Engine/Plugins/Animation/ControlRig/Source/ControlRigEditor/Private/Editor/SRigSpacePickerWidget.cpp
--- a/Engine/Plugins/Animation/ControlRig/Source/ControlRigEditor/Private/Editor/SRigSpacePickerWidget.cpp
+++ b/Engine/Plugins/Animation/ControlRig/Source/ControlRigEditor/Private/Editor/SRigSpacePickerWidget.cpp
@@ -212,7 +212,10 @@ void SRigSpacePickerWidget::SetControls(
 	{
 		if(HierarchyModifiedHandle.IsValid())
 		{
-			Hierarchy->OnModified().Remove(HierarchyModifiedHandle);
+			if(IsValid(Hierarchy) && !Hierarchy->HasAnyFlags(RF_BeginDestroyed))
+			{
+				Hierarchy->OnModified().Remove(HierarchyModifiedHandle);
+			}
 			HierarchyModifiedHandle.Reset();
 		}
 	}
diff --git a/Engine/Plugins/Animation/ControlRig/Source/ControlRigEditor/Private/Graph/ControlRigGraphPanelPinFactory.cpp b/Engine/Plugins/Animation/ControlRig/Source/ControlRigEditor/Private/Graph/ControlRigGraphPanelPinFactory.cpp
--- a/Engine/Plugins/Animation/ControlRig/Source/ControlRigEditor/Private/Graph/ControlRigGraphPanelPinFactory.cpp
+++ b/Engine/Plugins/Animation/ControlRig/Source/ControlRigEditor/Private/Graph/ControlRigGraphPanelPinFactory.cpp
@@ -96,6 +96,12 @@ TSharedPtr<SGraphPin> FControlRigGraphPanelPinFactory::CreatePin(UEdGraphPin* In
 						.OnGetNameFromSelection_UObject(RigGraph, &UControlRigGraph::GetSelectedElementsNameList)
 						.OnGetNameListContent_UObject(RigGraph, &UControlRigGraph::GetDrawingNameList);
 				}
+				else if (CustomWidgetName == TEXT("ShapeName"))
+				{
+					return SNew(SControlRigGraphPinNameList, InPin)
+						.ModelPin(ModelPin)
+						.OnGetNameListContent_UObject(RigGraph, &UControlRigGraph::GetShapeNameList);
+				}
 				else if (CustomWidgetName == TEXT("VariableName"))
 				{
 					return SNew(SControlRigGraphPinVariableName, InPin);
diff --git a/Engine/Source/Developer/RigVMDeveloper/Private/RigVMModel/RigVMController.cpp b/Engine/Source/Developer/RigVMDeveloper/Private/RigVMModel/RigVMController.cpp
--- a/Engine/Source/Developer/RigVMDeveloper/Private/RigVMModel/RigVMController.cpp
+++ b/Engine/Source/Developer/RigVMDeveloper/Private/RigVMModel/RigVMController.cpp
@@ -16195,6 +16195,7 @@ bool URigVMController::ChangePinType(URigVMPin* InPin, const FString& InCPPType,
 						{
 							if(MemberName.Equals(PinName))
 							{
+								PostProcessDefaultValue(InPin, MemberValue);
 								ApplyResolvedDefaultValue(InController, InPin, SubPinPath, MemberValue, bSetupUndoRedo);
 								break;
 							}
diff --git a/Engine/Source/Runtime/RigVM/Private/RigVMCore/RigVMByteCode.cpp b/Engine/Source/Runtime/RigVM/Private/RigVMCore/RigVMByteCode.cpp
--- a/Engine/Source/Runtime/RigVM/Private/RigVMCore/RigVMByteCode.cpp
+++ b/Engine/Source/Runtime/RigVM/Private/RigVMCore/RigVMByteCode.cpp
@@ -1731,6 +1731,22 @@ FRigVMOperandArray FRigVMByteCode::GetOperandsForOp(const FRigVMInstruction& InI
 	return FRigVMOperandArray();
 }
 
+TArray<int32> FRigVMByteCode::GetInstructionsForOperand(const FRigVMOperand& InOperand) const
+{
+	TArray<int32> InstructionIndices;
+	
+	const FRigVMInstructionArray Instructions = GetInstructions();
+	for(int32 InstructionIndex = 0; InstructionIndex < GetNumInstructions(); InstructionIndex++)
+	{
+		if(GetOperandsForOp(Instructions[InstructionIndex]).Contains(InOperand))
+		{
+			InstructionIndices.Add(InstructionIndex);
+		}
+
+	}
+	return InstructionIndices;
+}
+
 uint64 FRigVMByteCode::GetOpAlignment(ERigVMOpCode InOpCode) const
 {
 	switch (InOpCode)
