Control Rig: Restore backwards compat for python API

diff --git a/Engine/Plugins/Animation/ControlRig/Source/ControlRigDeveloper/Private/ControlRigBlueprint.cpp b/Engine/Plugins/Animation/ControlRig/Source/ControlRigDeveloper/Private/ControlRigBlueprint.cpp
--- a/Engine/Plugins/Animation/ControlRig/Source/ControlRigDeveloper/Private/ControlRigBlueprint.cpp
+++ b/Engine/Plugins/Animation/ControlRig/Source/ControlRigDeveloper/Private/ControlRigBlueprint.cpp
@@ -1832,7 +1832,15 @@ URigVMController* UControlRigBlueprint::GetController(const URigVMGraph* InGraph
 
 URigVMController* UControlRigBlueprint::GetControllerByName(const FString InGraphName) const
 {
-	for (URigVMGraph* Graph : GetAllModels())
+	if(InGraphName.IsEmpty())
+	{
+		if(const URigVMGraph* DefaultModel = GetRigVMClient()->GetDefaultModel())
+		{
+			return GetController(DefaultModel);
+		}
+	}
+	
+	for (const URigVMGraph* Graph : GetAllModels())
 	{
 		if (Graph->GetName() == InGraphName || Graph->GetGraphName() == InGraphName)
 		{
diff --git a/Engine/Plugins/Animation/ControlRig/Source/ControlRigDeveloper/Public/ControlRigBlueprint.h b/Engine/Plugins/Animation/ControlRig/Source/ControlRigDeveloper/Public/ControlRigBlueprint.h
--- a/Engine/Plugins/Animation/ControlRig/Source/ControlRigDeveloper/Public/ControlRigBlueprint.h
+++ b/Engine/Plugins/Animation/ControlRig/Source/ControlRigDeveloper/Public/ControlRigBlueprint.h
@@ -319,7 +319,7 @@ public:
 	URigVMController* GetController(const URigVMGraph* InGraph = nullptr) const;
 
 	UFUNCTION(BlueprintCallable, Category = "Control Rig Blueprint")
-	URigVMController* GetControllerByName(const FString InGraphName) const;
+	URigVMController* GetControllerByName(const FString InGraphName = TEXT("")) const;
 
 	UFUNCTION(BlueprintCallable, Category = "Control Rig Blueprint")
 	URigVMController* GetOrCreateController(URigVMGraph* InGraph = nullptr);
