Control Rig: Fix null ptr issue inside of Graph model

diff --git a/Engine/Plugins/Animation/ControlRig/Source/ControlRigDeveloper/Private/Graph/ControlRigGraph.cpp b/Engine/Plugins/Animation/ControlRig/Source/ControlRigDeveloper/Private/Graph/ControlRigGraph.cpp
--- a/Engine/Plugins/Animation/ControlRig/Source/ControlRigDeveloper/Private/Graph/ControlRigGraph.cpp
+++ b/Engine/Plugins/Animation/ControlRig/Source/ControlRigDeveloper/Private/Graph/ControlRigGraph.cpp
@@ -205,7 +205,12 @@ void UControlRigGraph::CacheNameLists(URigHierarchy* InHierarchy, const FControl
 
 	for(const TSoftObjectPtr<UControlRigShapeLibrary>& ShapeLibrary : ShapeLibraries)
 	{
-		if(ShapeLibrary.IsNull())
+		if(ShapeLibrary.IsNull() || !ShapeLibrary.IsValid())
+		{
+			ShapeLibrary.LoadSynchronous();
+		}
+
+		if(ShapeLibrary.IsNull() || !ShapeLibrary.IsValid())
 		{
 			continue;
 		}
