Control Rig: Export aggregate node properly

diff --git a/Engine/Plugins/Animation/ControlRig/Source/ControlRigDeveloper/Private/ControlRigBlueprint.cpp b/Engine/Plugins/Animation/ControlRig/Source/ControlRigDeveloper/Private/ControlRigBlueprint.cpp
--- a/Engine/Plugins/Animation/ControlRig/Source/ControlRigDeveloper/Private/ControlRigBlueprint.cpp
+++ b/Engine/Plugins/Animation/ControlRig/Source/ControlRigDeveloper/Private/ControlRigBlueprint.cpp
@@ -1908,11 +1908,17 @@ TArray<FString> UControlRigBlueprint::GeneratePythonCommands(const FString InNew
 	
 	// Create graphs
 	{
+		TArray<URigVMGraph*> AllModels = GetAllModels();
+		AllModels.RemoveAll([](const URigVMGraph* GraphToRemove) -> bool
+		{
+			return GraphToRemove->GetTypedOuter<URigVMAggregateNode>() != nullptr;
+		});
+		
 		// Find all graphs to process and sort them by dependencies
 		TArray<URigVMGraph*> ProcessedGraphs;
-		while (ProcessedGraphs.Num() < GetAllModels().Num())
+		while (ProcessedGraphs.Num() < AllModels.Num())
 		{
-			for (URigVMGraph* Graph : GetAllModels())
+			for (URigVMGraph* Graph : AllModels)
 			{
 				if (ProcessedGraphs.Contains(Graph))
 				{
@@ -1937,10 +1943,13 @@ TArray<FString> UControlRigBlueprint::GeneratePythonCommands(const FString InNew
 					}
 					else if (URigVMCollapseNode* CollapseNode = Cast<URigVMCollapseNode>(Node))
 					{
-						if (!ProcessedGraphs.Contains(CollapseNode->GetContainedGraph()))
+						if(!CollapseNode->IsA<URigVMAggregateNode>())
 						{
-							bFoundUnprocessedReference = true;
-							break;
+							if (!ProcessedGraphs.Contains(CollapseNode->GetContainedGraph()))
+							{
+								bFoundUnprocessedReference = true;
+								break;
+							}
 						}
 					}
 				}
diff --git a/Engine/Source/Developer/RigVMDeveloper/Private/RigVMModel/RigVMController.cpp b/Engine/Source/Developer/RigVMDeveloper/Private/RigVMModel/RigVMController.cpp
--- a/Engine/Source/Developer/RigVMDeveloper/Private/RigVMModel/RigVMController.cpp
+++ b/Engine/Source/Developer/RigVMDeveloper/Private/RigVMModel/RigVMController.cpp
@@ -486,6 +486,12 @@ TArray<FString> URigVMController::GetAddNodePythonCommands(URigVMNode* Node) con
 		TArray<FString> InnerNodeCommands = GetAddNodePythonCommands(AggregateNode->GetFirstInnerNode());
 		Commands.Append(InnerNodeCommands);
 
+		// set_node_position(node_name, position)
+		Commands.Add(FString::Printf(TEXT("blueprint.get_controller_by_name('%s').set_node_position_by_name('%s', %s)"),
+				*GraphName,
+				*AggregateNode->GetName(),
+				*RigVMPythonUtils::Vector2DToPythonString(AggregateNode->GetPosition())));
+
 		// add commands for any additional aggregate pin
 		const TArray<URigVMPin*> AggregatePins = AggregateNode->IsInputAggregate() ? AggregateNode->GetAggregateInputs() : AggregateNode->GetAggregateOutputs();
 
