Fix for Control Rig exported vectors not working when used on AnimBP
Pin default value is now generated using FBlueprintEditorUtils::PropertyValueToString, which generates a valid default string for anim graph pin editors

diff --git a/Engine/Plugins/Animation/ControlRig/Source/ControlRigDeveloper/Private/AnimGraphNode_ControlRig.cpp b/Engine/Plugins/Animation/ControlRig/Source/ControlRigDeveloper/Private/AnimGraphNode_ControlRig.cpp
--- a/Engine/Plugins/Animation/ControlRig/Source/ControlRigDeveloper/Private/AnimGraphNode_ControlRig.cpp
+++ b/Engine/Plugins/Animation/ControlRig/Source/ControlRigDeveloper/Private/AnimGraphNode_ControlRig.cpp
@@ -13,6 +13,7 @@
 #include "AnimationGraphSchema.h"
 #include "ControlRigBlueprintGeneratedClass.h"
 #include "ControlRigBlueprint.h"
+#include "Misc/DefaultValueHelper.h"
 
 #define LOCTEXT_NAMESPACE "AnimGraphNode_ControlRig"
 
@@ -171,9 +172,10 @@ void UAnimGraphNode_ControlRig::CreateCustomPins(TArray<UEdGraphPin*>* OldPins)
 						{
 							FString DefaultValue;
 							
+							// The format the graph pins editor use is different of what property exporter produces, so we use BlueprintEditorUtils to generate the default string
 							// Variable.Memory here points to the corresponding property in the Control Rig BP CDO, it was initialized in UAnimGraphNode_ControlRig::RebuildExposedProperties 
-							PropertyIt->ExportTextItem_Direct(DefaultValue, Variable.Memory, nullptr, nullptr, PPF_None, nullptr);
-							
+							FBlueprintEditorUtils::PropertyValueToString_Direct(*PropertyIt, Variable.Memory, DefaultValue, this);
+
 							if (!DefaultValue.IsEmpty())
 							{
 								AnimGraphDefaultSchema->TrySetDefaultValue(*NewPin, DefaultValue);
@@ -783,6 +785,80 @@ void UAnimGraphNode_ControlRig::PostEditChangeProperty(FPropertyChangedEvent& Pr
 	}
 }
 
+void UAnimGraphNode_ControlRig::PostReconstructNode()
+{
+	// Fix default values that were serialized directly with property exported strings (not valid for pin editors)
+#if WITH_EDITOR
+	if (!IsTemplate())
+	{
+		static UScriptStruct* VectorStruct = TBaseStructure<FVector>::Get();
+		static UScriptStruct* TransformStruct = TBaseStructure<FTransform>::Get();
+
+		// fix up any pin data if it needs to 
+		for (UEdGraphPin* CurrentPin : Pins)
+		{
+			const FName& PinCategory = CurrentPin->PinType.PinCategory;
+
+			// Target only struct types vector and transform 
+			// For rotator, if it was created with incorect format, it was serialized with DefaultValue = L"0, 0, 0", so it can not be corrected without risking changing users value
+			if (PinCategory == UEdGraphSchema_K2::PC_Struct)
+			{
+				const UAnimationGraphSchema* AnimGraphDefaultSchema = GetDefault<UAnimationGraphSchema>();
+				const FName& PinName = CurrentPin->GetFName();
+
+				if (CurrentPin->PinType.PinSubCategoryObject == VectorStruct)
+				{
+					// If InitFromString succeeds, it has bad format, re-encode
+					FVector FixVector = FVector::ZeroVector;
+					if (FixVector.InitFromString(CurrentPin->DefaultValue))
+					{
+						const FString DefaultValue = FString::Printf(TEXT("%f,%f,%f"), FixVector.X, FixVector.Y, FixVector.Z);
+						AnimGraphDefaultSchema->TrySetDefaultValue(*CurrentPin, DefaultValue);
+					}
+				}
+				else if (CurrentPin->PinType.PinSubCategoryObject == TransformStruct)
+				{
+					// If pin was serialized with bad format, transforms would stay empty, so try to get the value from ControlRig or set to Identity as last resort
+					if (CurrentPin->DefaultValue.IsEmpty())
+					{
+						FTransform FixTransform = FTransform::Identity;
+						FString DefaultValue = FixTransform.ToString();
+
+						const FRigVMExternalVariable* RigVMExternalVariable = (CurrentPin->Direction == EGPD_Input) ? 
+							InputVariables.Find(CurrentPin->GetFName()) : (CurrentPin->Direction == EGPD_Output) ? 
+							OutputVariables.Find(CurrentPin->GetFName()) : nullptr;
+
+						if (RigVMExternalVariable != nullptr && RigVMExternalVariable->IsValid())
+						{
+							if (UControlRigBlueprintGeneratedClass* GeneratedClass = Cast<UControlRigBlueprintGeneratedClass>(GetTargetClass()))
+							{
+								const FName& PropertyName = CurrentPin->GetFName();
+
+								for (TFieldIterator<FProperty> PropertyIt(GeneratedClass); PropertyIt; ++PropertyIt)
+								{
+									if (PropertyIt->GetFName() == PropertyName)
+									{
+										// The format the graph pins editor use is different of what property exporter produces, so we use BlueprintEditorUtils to generate the default string
+										// Variable.Memory here points to the corresponding property in the Control Rig BP CDO, it was initialized in UAnimGraphNode_ControlRig::RebuildExposedProperties 
+										FBlueprintEditorUtils::PropertyValueToString_Direct(*PropertyIt, RigVMExternalVariable->Memory, DefaultValue, this);
+										break;
+									}
+								}
+							}
+						}
+
+						AnimGraphDefaultSchema->SetPinAutogeneratedDefaultValue(CurrentPin, DefaultValue);
+						AnimGraphDefaultSchema->TrySetDefaultValue(*CurrentPin, DefaultValue);
+					}
+				}
+			}
+		}
+	}
+#endif // WITH_EDITOR
+
+	Super::PostReconstructNode();
+}
+
 void UAnimGraphNode_ControlRig::CustomizePinData(UEdGraphPin* Pin, FName SourcePropertyName, int32 ArrayIndex) const
 {
 	Super::CustomizePinData(Pin, SourcePropertyName, ArrayIndex);
diff --git a/Engine/Plugins/Animation/ControlRig/Source/ControlRigDeveloper/Public/AnimGraphNode_ControlRig.h b/Engine/Plugins/Animation/ControlRig/Source/ControlRigDeveloper/Public/AnimGraphNode_ControlRig.h
--- a/Engine/Plugins/Animation/ControlRig/Source/ControlRigDeveloper/Public/AnimGraphNode_ControlRig.h
+++ b/Engine/Plugins/Animation/ControlRig/Source/ControlRigDeveloper/Public/AnimGraphNode_ControlRig.h
@@ -19,6 +19,7 @@ class UAnimGraphNode_ControlRig : public UAnimGraphNode_CustomProperty
 
 	virtual void CustomizeDetails(IDetailLayoutBuilder& DetailBuilder) override;
 	virtual void PostEditChangeProperty(FPropertyChangedEvent& PropertyChangedEvent) override;
+	virtual void PostReconstructNode() override;
 
 private:
 	// UEdGraphNode interface
