Control Rig: Follow up fixes after construction event

diff --git a/Engine/Plugins/Animation/ControlRig/Source/ControlRig/Private/ControlRig.cpp b/Engine/Plugins/Animation/ControlRig/Source/ControlRig/Private/ControlRig.cpp
--- a/Engine/Plugins/Animation/ControlRig/Source/ControlRig/Private/ControlRig.cpp
+++ b/Engine/Plugins/Animation/ControlRig/Source/ControlRig/Private/ControlRig.cpp
@@ -1800,7 +1800,7 @@ FRigControlElement* UControlRig::FindControl(const FName& InControlName) const
 
 bool UControlRig::IsConstructionModeEnabled() const
 {
-	return EventQueue.Contains(FRigUnit_PrepareForExecution::EventName);
+	return EventQueue.Num() == 1 && EventQueue.Contains(FRigUnit_PrepareForExecution::EventName);
 }
 
 FTransform UControlRig::SetupControlFromGlobalTransform(const FName& InControlName, const FTransform& InGlobalTransform)
diff --git a/Engine/Plugins/Animation/ControlRig/Source/ControlRig/Private/Rigs/AdditiveControlRig.cpp b/Engine/Plugins/Animation/ControlRig/Source/ControlRig/Private/Rigs/AdditiveControlRig.cpp
--- a/Engine/Plugins/Animation/ControlRig/Source/ControlRig/Private/Rigs/AdditiveControlRig.cpp
+++ b/Engine/Plugins/Animation/ControlRig/Source/ControlRig/Private/Rigs/AdditiveControlRig.cpp
@@ -5,6 +5,7 @@
 #include "Engine/SkeletalMesh.h"
 #include "IControlRigObjectBinding.h"
 #include "Components/SkeletalMeshComponent.h"
+#include "Execution/RigUnit_PrepareForExecution.h"
 #include "Rigs/RigHierarchyController.h"
 #include "Units/Execution/RigUnit_BeginExecution.h"
 
@@ -41,21 +42,36 @@ FName UAdditiveControlRig::GetNullName(const FName& InBoneName)
 
 bool UAdditiveControlRig::ExecuteUnits(FRigUnitContext& InOutContext, const FName& InEventName)
 {
-	if (InEventName != FRigUnit_BeginExecution::EventName)
+	if (InEventName == FRigUnit_BeginExecution::EventName)
 	{
-		// still return true - we want the evaluation code
-		// to simply carry on for other events.
-		return true;
+		for (FRigUnit_AddBoneTransform& Unit : AddBoneRigUnits)
+		{
+			FName ControlName = GetControlName(Unit.Bone);
+			const int32 Index = GetHierarchy()->GetIndex(FRigElementKey(ControlName, ERigElementType::Control));
+			Unit.Transform = GetHierarchy()->GetLocalTransform(Index);
+			Unit.ExecuteContext.Hierarchy = GetHierarchy();
+			Unit.ExecuteContext.EventName = InEventName;
+			Unit.Execute(InOutContext);
+		}
 	}
-
-	for (FRigUnit_AddBoneTransform& Unit : AddBoneRigUnits)
+	else if(InEventName == FRigUnit_PrepareForExecution::EventName)
 	{
-		FName ControlName = GetControlName(Unit.Bone);
-		const int32 Index = GetHierarchy()->GetIndex(FRigElementKey(ControlName, ERigElementType::Control));
-		Unit.Transform = GetHierarchy()->GetLocalTransform(Index);
-		Unit.ExecuteContext.Hierarchy = GetHierarchy();
-		Unit.ExecuteContext.EventName = InEventName;
-		Unit.Execute(InOutContext);
+		if (USkeletalMeshComponent* SkeletalMeshComponent = Cast<USkeletalMeshComponent>(GetObjectBinding()->GetBoundObject()))
+		{
+			CreateRigElements(SkeletalMeshComponent->GetSkeletalMesh());
+		}
+
+		// add units and initialize
+		AddBoneRigUnits.Reset();
+
+		GetHierarchy()->ForEach<FRigBoneElement>([&](FRigBoneElement* BoneElement) -> bool
+		{
+			FRigUnit_AddBoneTransform NewUnit;
+			NewUnit.Bone = BoneElement->GetName();
+			NewUnit.bPropagateToChildren = true;
+			AddBoneRigUnits.Add(NewUnit);
+			return true;
+		});
 	}
 
 	return true;
@@ -72,27 +88,9 @@ void UAdditiveControlRig::Initialize(bool bInitRigUnits /*= true*/)
 		return;
 	}
 
-	// we do this after Initialize because Initialize will copy from CDO. 
-	// create hierarchy from the incoming skeleton
-	if (USkeletalMeshComponent* SkeletalMeshComponent = Cast<USkeletalMeshComponent>(GetObjectBinding()->GetBoundObject()))
-	{
-		CreateRigElements(SkeletalMeshComponent->GetSkeletalMesh());
-	}
-
-	// add units and initialize
-	AddBoneRigUnits.Reset();
-
-	GetHierarchy()->ForEach<FRigBoneElement>([&](FRigBoneElement* BoneElement) -> bool
-	{
-		FRigUnit_AddBoneTransform NewUnit;
-        NewUnit.Bone = BoneElement->GetName();
-        NewUnit.bPropagateToChildren = true;
-        AddBoneRigUnits.Add(NewUnit);
-		return true;
-    });
-
-	// execute init
 	Execute(EControlRigState::Init, FRigUnit_BeginExecution::EventName);
+	FRigUnitContext DefaultContext;
+	ExecuteUnits(DefaultContext, FRigUnit_PrepareForExecution::EventName);
 }
 
 void UAdditiveControlRig::CreateRigElements(const FReferenceSkeleton& InReferenceSkeleton, const FSmartNameMapping* InSmartNameMapping)
diff --git a/Engine/Plugins/Animation/ControlRig/Source/ControlRig/Private/Rigs/RigHierarchy.cpp b/Engine/Plugins/Animation/ControlRig/Source/ControlRig/Private/Rigs/RigHierarchy.cpp
--- a/Engine/Plugins/Animation/ControlRig/Source/ControlRig/Private/Rigs/RigHierarchy.cpp
+++ b/Engine/Plugins/Animation/ControlRig/Source/ControlRig/Private/Rigs/RigHierarchy.cpp
@@ -1582,7 +1582,7 @@ bool URigHierarchy::SetParentWeight(FRigBaseElement* InChild, int32 InParentInde
 			EnsureCacheValidity();
 			
 #if WITH_EDITOR
-			if (ensure(!bPropagatingChange))
+			if (!bPropagatingChange)
 			{
 				TGuardValue<bool> bPropagatingChangeGuardValue(bPropagatingChange, true);
 			
@@ -1704,7 +1704,7 @@ bool URigHierarchy::SetParentWeightArray(FRigBaseElement* InChild,  const TArray
 			PropagateDirtyFlags(MultiParentElement, ERigTransformType::IsInitial(LocalType), bAffectChildren);
 			
 #if WITH_EDITOR
-			if (ensure(!bPropagatingChange))
+			if (!bPropagatingChange)
 			{
 				TGuardValue<bool> bPropagatingChangeGuardValue(bPropagatingChange, true);
 			
@@ -2299,7 +2299,7 @@ void URigHierarchy::SetPose(const FRigPose& InPose, ERigTransformType::Type InTr
 	{
 		return;
 	}
-	
+
 	for(const FRigPoseElement& PoseElement : InPose)
 	{
 		FCachedRigElement Index = PoseElement.Index;
@@ -2735,7 +2735,7 @@ void URigHierarchy::SetTransform(FRigTransformElement* InTransformElement, const
 			bSetupUndo);
 	}
 
-	if (ensure(!bPropagatingChange))
+	if (!bPropagatingChange)
 	{
 		TGuardValue<bool> bPropagatingChangeGuardValue(bPropagatingChange, true);
 			
@@ -2954,7 +2954,7 @@ void URigHierarchy::SetControlOffsetTransform(FRigControlElement* InControlEleme
             bSetupUndo);
 	}
 
-	if (ensure(!bPropagatingChange))
+	if (!bPropagatingChange)
 	{
 		TGuardValue<bool> bPropagatingChangeGuardValue(bPropagatingChange, true);
 			
@@ -3086,7 +3086,7 @@ void URigHierarchy::SetControlShapeTransform(FRigControlElement* InControlElemen
 	}
 
 #if WITH_EDITOR
-	if (ensure(!bPropagatingChange))
+	if (!bPropagatingChange)
 	{
 		TGuardValue<bool> bPropagatingChangeGuardValue(bPropagatingChange, true);
 			
@@ -3154,7 +3154,7 @@ void URigHierarchy::SetControlSettings(FRigControlElement* InControlElement, FRi
 	Notify(ERigHierarchyNotification::ControlSettingChanged, InControlElement);
 	
 #if WITH_EDITOR
-	if (ensure(!bPropagatingChange))
+	if (!bPropagatingChange)
 	{
 		TGuardValue<bool> bPropagatingChangeGuardValue(bPropagatingChange, true);
 			
@@ -3413,7 +3413,7 @@ void URigHierarchy::SetControlValue(FRigControlElement* InControlElement, const
 				Notify(ERigHierarchyNotification::ControlSettingChanged, InControlElement);
 
 #if WITH_EDITOR
-				if (ensure(!bPropagatingChange))
+				if (!bPropagatingChange)
 				{
 					TGuardValue<bool> bPropagatingChangeGuardValue(bPropagatingChange, true);
 			
@@ -3474,7 +3474,7 @@ void URigHierarchy::SetControlVisibility(FRigControlElement* InControlElement, b
 	}
 
 #if WITH_EDITOR
-	if (ensure(!bPropagatingChange))
+	if (!bPropagatingChange)
 	{
 		TGuardValue<bool> bPropagatingChangeGuardValue(bPropagatingChange, true);
 			
@@ -3534,7 +3534,7 @@ void URigHierarchy::SetCurveValue(FRigCurveElement* InCurveElement, float InValu
 		PushCurveToStack(InCurveElement->GetKey(), PreviousValue, InCurveElement->Value, bPreviousIsValueSet, true, bSetupUndo);
 	}
 
-	if (ensure(!bPropagatingChange))
+	if (!bPropagatingChange)
 	{
 		TGuardValue<bool> bPropagatingChangeGuardValue(bPropagatingChange, true);
 			
@@ -3581,7 +3581,7 @@ void URigHierarchy::UnsetCurveValue(FRigCurveElement* InCurveElement, bool bSetu
 		PushCurveToStack(InCurveElement->GetKey(), InCurveElement->Value, InCurveElement->Value, bPreviousIsValueSet, false, bSetupUndo);
 	}
 
-	if (ensure(!bPropagatingChange))
+	if (!bPropagatingChange)
 	{
 		TGuardValue<bool> bPropagatingChangeGuardValue(bPropagatingChange, true);
 			
diff --git a/Engine/Plugins/Animation/ControlRig/Source/ControlRig/Private/Rigs/RigHierarchyCache.cpp b/Engine/Plugins/Animation/ControlRig/Source/ControlRig/Private/Rigs/RigHierarchyCache.cpp
--- a/Engine/Plugins/Animation/ControlRig/Source/ControlRig/Private/Rigs/RigHierarchyCache.cpp
+++ b/Engine/Plugins/Animation/ControlRig/Source/ControlRig/Private/Rigs/RigHierarchyCache.cpp
@@ -22,12 +22,16 @@ bool FCachedRigElement::UpdateCache(const FRigElementKey& InKey, const URigHiera
 	{
 		if(!IsValid() || !IsIdentical(InKey, InHierarchy))
 		{
+			// have to create a copy since Reset below
+			// potentially resets the InKey as well.
+			const FRigElementKey KeyToResolve = InKey; 
+
 			Reset();
 
-			int32 Idx = InHierarchy->GetIndex(InKey);
+			int32 Idx = InHierarchy->GetIndex(KeyToResolve);
 			if(Idx != INDEX_NONE)
 			{
-				Key = InKey;
+				Key = KeyToResolve;
 				Index = (uint16)Idx;
 				Element = InHierarchy->Get(Index);
 			}
diff --git a/Engine/Plugins/Animation/ControlRig/Source/ControlRigEditor/Private/ControlRigElementDetails.cpp b/Engine/Plugins/Animation/ControlRig/Source/ControlRigEditor/Private/ControlRigElementDetails.cpp
--- a/Engine/Plugins/Animation/ControlRig/Source/ControlRigEditor/Private/ControlRigElementDetails.cpp
+++ b/Engine/Plugins/Animation/ControlRig/Source/ControlRigEditor/Private/ControlRigElementDetails.cpp
@@ -86,6 +86,15 @@ void RigElementKeyDetails_GetCustomizedInfo(TSharedRef<IPropertyHandle> InStruct
 		{
 			break;
 		}
+
+		if(const UControlRig* ControlRig = Object->GetTypedOuter<UControlRig>())
+		{
+			OutBlueprint = Cast<UControlRigBlueprint>(ControlRig->GetClass()->ClassGeneratedBy);
+			if(OutBlueprint)
+			{
+				break;
+			}
+		}
 	}
 
 	if (OutBlueprint == nullptr)
diff --git a/Engine/Plugins/Animation/ControlRig/Source/ControlRigEditor/Private/EditMode/ControlRigEditMode.cpp b/Engine/Plugins/Animation/ControlRig/Source/ControlRigEditor/Private/EditMode/ControlRigEditMode.cpp
--- a/Engine/Plugins/Animation/ControlRig/Source/ControlRigEditor/Private/EditMode/ControlRigEditMode.cpp
+++ b/Engine/Plugins/Animation/ControlRig/Source/ControlRigEditor/Private/EditMode/ControlRigEditMode.cpp
@@ -144,7 +144,7 @@ FControlRigEditMode::FControlRigEditMode()
 	, bIsChangingCoordSystem(false)
 	, InteractionType((uint8)EControlRigInteractionType::None)
 	, bShowControlsAsOverlay(false)
-	, HierarchyHashBeforeConstruction(INDEX_NONE)
+	, bIsConstructionEventRunning(false)
 {
 	ControlProxy = NewObject<UControlRigDetailPanelControlProxies>(GetTransientPackage(), NAME_None);
 	ControlProxy->SetFlags(RF_Transactional);
@@ -3431,7 +3431,7 @@ void FControlRigEditMode::OnHierarchyModified_AnyThread(ERigHierarchyNotificatio
 		return;
 	}
 
-	if(IsConstructionEventRunning())
+	if(bIsConstructionEventRunning)
 	{
 		return;
 	}
@@ -3497,17 +3497,25 @@ void FControlRigEditMode::OnControlModified(UControlRig* Subject, FRigControlEle
 void FControlRigEditMode::OnPreConstruction_AnyThread(UControlRig* InRig, const EControlRigState InState,
 	const FName& InEventName)
 {
-	HierarchyHashBeforeConstruction = InRig->GetHierarchy()->GetTopologyHash(false, true);
+	bIsConstructionEventRunning = true;
 }
 
 void FControlRigEditMode::OnPostConstruction_AnyThread(UControlRig* InRig, const EControlRigState InState,
 	const FName& InEventName)
 {
-	uint32 HierarchyHash = INDEX_NONE;
-	Swap(HierarchyHash, HierarchyHashBeforeConstruction);
+	bIsConstructionEventRunning = false;
 
-	if(HierarchyHash != InRig->GetHierarchy()->GetTopologyHash(false, true))
+	const int32 RigIndex = RuntimeControlRigs.Find(InRig);
+	if(!LastHierarchyHash.IsValidIndex(RigIndex))
 	{
+		return;
+	}
+	
+	const int32 HierarchyHash = InRig->GetHierarchy()->GetTopologyHash(false, true);
+	if(LastHierarchyHash[RigIndex] != HierarchyHash)
+	{
+		LastHierarchyHash[RigIndex] = HierarchyHash;
+		
 		auto Task = [this, InRig]()
 		{
 			RequestToRecreateControlShapeActors(InRig);
@@ -3974,12 +3982,12 @@ AControlRigShapeActor* FControlRigEditMode::GetControlShapeFromControlName(UCont
 void FControlRigEditMode::AddControlRigInternal(UControlRig* InControlRig)
 {
 	RuntimeControlRigs.AddUnique(InControlRig);
+	LastHierarchyHash.Add(INDEX_NONE);
 
 	InControlRig->SetControlsVisible(true);
 	InControlRig->PostInitInstanceIfRequired();
 	InControlRig->GetHierarchy()->OnModified().RemoveAll(this);
 	InControlRig->GetHierarchy()->OnModified().AddSP(this, &FControlRigEditMode::OnHierarchyModified_AnyThread);
-	InControlRig->OnPreConstructionForUI_AnyThread().AddSP(this, &FControlRigEditMode::OnPreConstruction_AnyThread);
 	InControlRig->OnPostConstruction_AnyThread().AddSP(this, &FControlRigEditMode::OnPostConstruction_AnyThread);
 
 	//needed for the control rig track editor delegates to get hooked up
@@ -4049,6 +4057,11 @@ void FControlRigEditMode::RemoveControlRig(UControlRig* InControlRig)
 	{
 		RuntimeControlRigs.RemoveAt(Index);
 	}
+	if (LastHierarchyHash.IsValidIndex(Index))
+	{
+		LastHierarchyHash.RemoveAt(Index);
+	}
+
 	//needed for the control rig track editor delegates to get removed
 	if (WeakSequencer.IsValid())
 	{
diff --git a/Engine/Plugins/Animation/ControlRig/Source/ControlRigEditor/Private/EditMode/ControlRigEditMode.h b/Engine/Plugins/Animation/ControlRig/Source/ControlRigEditor/Private/EditMode/ControlRigEditMode.h
--- a/Engine/Plugins/Animation/ControlRig/Source/ControlRigEditor/Private/EditMode/ControlRigEditMode.h
+++ b/Engine/Plugins/Animation/ControlRig/Source/ControlRigEditor/Private/EditMode/ControlRigEditMode.h
@@ -440,8 +440,8 @@ private:
 	uint8 InteractionType;
 	bool bShowControlsAsOverlay;
 
-	bool IsConstructionEventRunning() const { return HierarchyHashBeforeConstruction != INDEX_NONE; }
-	uint32 HierarchyHashBeforeConstruction;
+	bool bIsConstructionEventRunning;
+	TArray<uint32> LastHierarchyHash;
 
 	friend class FControlRigEditorModule;
 	friend class FControlRigEditor;
diff --git a/Engine/Plugins/Animation/ControlRig/Source/ControlRigEditor/Private/Editor/ControlRigEditor.cpp b/Engine/Plugins/Animation/ControlRig/Source/ControlRigEditor/Private/Editor/ControlRigEditor.cpp
--- a/Engine/Plugins/Animation/ControlRig/Source/ControlRigEditor/Private/Editor/ControlRigEditor.cpp
+++ b/Engine/Plugins/Animation/ControlRig/Source/ControlRigEditor/Private/Editor/ControlRigEditor.cpp
@@ -223,7 +223,8 @@ FControlRigEditor::FControlRigEditor()
 	, RigHierarchyTabCount(0)
 	, HaltedAtNode(nullptr)
 	, bSuspendDetailsPanelRefresh(false)
-	, HierarchyHashBeforeConstruction(INDEX_NONE)
+	, bIsConstructionEventRunning(false)
+	, LastHierarchyHash(INDEX_NONE)
 {
 }
 
@@ -1367,6 +1368,7 @@ void FControlRigEditor::HandleSetObjectBeingDebugged(UObject* InObject)
 		DebuggedControlRig->GetHierarchy()->OnModified().AddSP(this, &FControlRigEditor::OnHierarchyModified_AnyThread);
 		DebuggedControlRig->OnPreConstructionForUI_AnyThread().AddSP(this, &FControlRigEditor::OnPreConstruction_AnyThread);
 		DebuggedControlRig->OnPostConstruction_AnyThread().AddSP(this, &FControlRigEditor::OnPostConstruction_AnyThread);
+		LastHierarchyHash = INDEX_NONE;
 
 		if(EditorSkelComp)
 		{
@@ -5099,7 +5101,7 @@ void FControlRigEditor::OnHierarchyModified(ERigHierarchyNotification InNotif, U
 
 void FControlRigEditor::OnHierarchyModified_AnyThread(ERigHierarchyNotification InNotif, URigHierarchy* InHierarchy, const FRigBaseElement* InElement)
 {
-	if(IsConstructionEventRunning())
+	if(bIsConstructionEventRunning)
 	{
 		return;
 	}
@@ -6816,17 +6818,19 @@ FString* FControlRigEditor::GetSnippetStorage(int32 InSnippetIndex)
 void FControlRigEditor::OnPreConstruction_AnyThread(UControlRig* InRig, const EControlRigState InState,
 	const FName& InEventName)
 {
-	HierarchyHashBeforeConstruction = InRig->GetHierarchy()->GetTopologyHash(false);
+	bIsConstructionEventRunning = true;
 }
 
 void FControlRigEditor::OnPostConstruction_AnyThread(UControlRig* InRig, const EControlRigState InState,
 	const FName& InEventName)
 {
-	uint32 HierarchyHash = INDEX_NONE;
-	Swap(HierarchyHash, HierarchyHashBeforeConstruction);
-
-	if(HierarchyHash != InRig->GetHierarchy()->GetTopologyHash(false))
+	bIsConstructionEventRunning = false;
+	
+	const int32 HierarchyHash = InRig->GetHierarchy()->GetTopologyHash(false);
+	if(LastHierarchyHash != HierarchyHash)
 	{
+		LastHierarchyHash = HierarchyHash;
+		
 		auto Task = [this, InRig]()
 		{
 			CacheNameLists();
diff --git a/Engine/Plugins/Animation/ControlRig/Source/ControlRigEditor/Private/Editor/ControlRigEditor.h b/Engine/Plugins/Animation/ControlRig/Source/ControlRigEditor/Private/Editor/ControlRigEditor.h
--- a/Engine/Plugins/Animation/ControlRig/Source/ControlRigEditor/Private/Editor/ControlRigEditor.h
+++ b/Engine/Plugins/Animation/ControlRig/Source/ControlRigEditor/Private/Editor/ControlRigEditor.h
@@ -489,8 +489,8 @@ protected:
 	void OnPreConstruction_AnyThread(UControlRig* InRig, const EControlRigState InState, const FName& InEventName);
 	void OnPostConstruction_AnyThread(UControlRig* InRig, const EControlRigState InState, const FName& InEventName);
 
-	bool IsConstructionEventRunning() const { return HierarchyHashBeforeConstruction != INDEX_NONE; }
-	uint32 HierarchyHashBeforeConstruction;
+	bool bIsConstructionEventRunning;
+	uint32 LastHierarchyHash;
 
 	static const TArray<FName> ForwardsSolveEventQueue;
 	static const TArray<FName> BackwardsSolveEventQueue;
diff --git a/Engine/Plugins/Animation/ControlRig/Source/ControlRigEditor/Private/Editor/SRigHierarchy.cpp b/Engine/Plugins/Animation/ControlRig/Source/ControlRigEditor/Private/Editor/SRigHierarchy.cpp
--- a/Engine/Plugins/Animation/ControlRig/Source/ControlRigEditor/Private/Editor/SRigHierarchy.cpp
+++ b/Engine/Plugins/Animation/ControlRig/Source/ControlRigEditor/Private/Editor/SRigHierarchy.cpp
@@ -245,7 +245,8 @@ void SRigHierarchy::Construct(const FArguments& InArgs, TSharedRef<FControlRigEd
 	];
 
 	bIsChangingRigHierarchy = false;
-	HierarchyHashBeforeConstruction = INDEX_NONE;
+	LastHierarchyHash = INDEX_NONE;
+	bIsConstructionEventRunning = false;
 	
 	RefreshTreeView();
 
@@ -617,7 +618,7 @@ void SRigHierarchy::OnHierarchyModified(ERigHierarchyNotification InNotif, URigH
 		return;
 	}
 
-	if (bIsChangingRigHierarchy || IsConstructionEventRunning())
+	if (bIsChangingRigHierarchy || bIsConstructionEventRunning)
 	{
 		return;
 	}
@@ -775,7 +776,7 @@ void SRigHierarchy::OnHierarchyModified_AnyThread(ERigHierarchyNotification InNo
 		return;
 	}
 
-	if(IsConstructionEventRunning())
+	if(bIsConstructionEventRunning)
 	{
 		return;
 	}
@@ -845,6 +846,7 @@ void SRigHierarchy::HandleSetObjectBeingDebugged(UObject* InObject)
 		ControlRig->OnPreConstructionForUI_AnyThread().AddSP(this, &SRigHierarchy::OnPreConstruction_AnyThread);
 		ControlRig->OnPostConstruction_AnyThread().RemoveAll(this);
 		ControlRig->OnPostConstruction_AnyThread().AddSP(this, &SRigHierarchy::OnPostConstruction_AnyThread);
+		LastHierarchyHash = INDEX_NONE;
 	}
 
 	RefreshTreeView();
@@ -856,7 +858,7 @@ void SRigHierarchy::OnPreConstruction_AnyThread(UControlRig* InRig, const EContr
 	{
 		return;
 	}
-	HierarchyHashBeforeConstruction = InRig->GetHierarchy()->GetTopologyHash(false);
+	bIsConstructionEventRunning = true;
 	SelectionBeforeConstruction = InRig->GetHierarchy()->GetSelectedKeys();
 }
 
@@ -867,11 +869,13 @@ void SRigHierarchy::OnPostConstruction_AnyThread(UControlRig* InRig, const ECont
 		return;
 	}
 
-	uint32 HierarchyHash = INDEX_NONE;
-	Swap(HierarchyHash, HierarchyHashBeforeConstruction);
-			
-	if(HierarchyHash != InRig->GetHierarchy()->GetTopologyHash(false))
+	bIsConstructionEventRunning = false;
+
+	const int32 HierarchyHash = InRig->GetHierarchy()->GetTopologyHash(false);
+	if(LastHierarchyHash != HierarchyHash)
 	{
+		LastHierarchyHash = HierarchyHash;
+		
 		auto Task = [this]()
 		{
 			TGuardValue<bool> GuardRigHierarchyChanges(bIsChangingRigHierarchy, true);
diff --git a/Engine/Plugins/Animation/ControlRig/Source/ControlRigEditor/Private/Editor/SRigHierarchy.h b/Engine/Plugins/Animation/ControlRig/Source/ControlRigEditor/Private/Editor/SRigHierarchy.h
--- a/Engine/Plugins/Animation/ControlRig/Source/ControlRigEditor/Private/Editor/SRigHierarchy.h
+++ b/Engine/Plugins/Animation/ControlRig/Source/ControlRigEditor/Private/Editor/SRigHierarchy.h
@@ -226,8 +226,8 @@ private:
 	void OnPreConstruction_AnyThread(UControlRig* InRig, const EControlRigState InState, const FName& InEventName);
 	void OnPostConstruction_AnyThread(UControlRig* InRig, const EControlRigState InState, const FName& InEventName);
 
-	bool IsConstructionEventRunning() const { return HierarchyHashBeforeConstruction != INDEX_NONE; }
-	uint32 HierarchyHashBeforeConstruction;
+	bool bIsConstructionEventRunning;
+	uint32 LastHierarchyHash;
 	TArray<FRigElementKey> SelectionBeforeConstruction;
 
 public:
