Control Rig: Remove unneeded include.

#rnx

diff --git a/Engine/Plugins/Animation/ControlRig/Source/ControlRig/Private/ControlRigAnimInstance.cpp b/Engine/Plugins/Animation/ControlRig/Source/ControlRig/Private/ControlRigAnimInstance.cpp
--- a/Engine/Plugins/Animation/ControlRig/Source/ControlRig/Private/ControlRigAnimInstance.cpp
+++ b/Engine/Plugins/Animation/ControlRig/Source/ControlRig/Private/ControlRigAnimInstance.cpp
@@ -1,7 +1,8 @@
 // Copyright Epic Games, Inc. All Rights Reserved.
 
 #include "ControlRigAnimInstance.h"
-#include "HAL/PlatformProcess.h"
+
+#include "Animation/AnimNodeBase.h"
 
 ////////////////////////////////////////////////////////////////////////////////////////
 void FControlRigAnimInstanceProxy::Initialize(UAnimInstance* InAnimInstance)
diff --git a/Engine/Plugins/Animation/ControlRig/Source/ControlRig/Public/ControlRigAnimInstance.h b/Engine/Plugins/Animation/ControlRig/Source/ControlRig/Public/ControlRigAnimInstance.h
--- a/Engine/Plugins/Animation/ControlRig/Source/ControlRig/Public/ControlRigAnimInstance.h
+++ b/Engine/Plugins/Animation/ControlRig/Source/ControlRig/Public/ControlRigAnimInstance.h
@@ -4,7 +4,6 @@
 
 #include "Animation/AnimInstance.h"
 #include "Animation/AnimInstanceProxy.h"
-#include "AnimNodes/AnimNode_PoseSnapshot.h"
 #include "ControlRigAnimInstance.generated.h"
 
 /** Proxy override for this UAnimInstance-derived class */
@@ -23,7 +22,7 @@ public:
 	{
 	}
 
-	virtual ~FControlRigAnimInstanceProxy();
+	virtual ~FControlRigAnimInstanceProxy() override;
 
 	// FAnimInstanceProxy interface
 	virtual void Initialize(UAnimInstance* InAnimInstance) override;
