Control Rig: Avoid running construction event consecutively if event doesn't exist

diff --git a/Engine/Plugins/Animation/ControlRig/Source/ControlRig/Private/ControlRig.cpp b/Engine/Plugins/Animation/ControlRig/Source/ControlRig/Private/ControlRig.cpp
--- a/Engine/Plugins/Animation/ControlRig/Source/ControlRig/Private/ControlRig.cpp
+++ b/Engine/Plugins/Animation/ControlRig/Source/ControlRig/Private/ControlRig.cpp
@@ -1356,12 +1356,9 @@ bool UControlRig::Execute(const EControlRigState InState, const FName& InEventNa
 		bInteractionJustBegan = false;
 	}
 
-	if(bSuccess)
+	if(bIsConstructionEvent)
 	{
-		if(bIsConstructionEvent)
-		{
-			bRequiresConstructionEvent = false;
-		}
+		bRequiresConstructionEvent = false;
 	}
 
 	return bSuccess;
