Control Rig: Follow up fixes for introduction of RigVMClient

diff --git a/Engine/Plugins/Animation/ControlRig/Source/ControlRigDeveloper/Private/ControlRigBlueprint.cpp b/Engine/Plugins/Animation/ControlRig/Source/ControlRigDeveloper/Private/ControlRigBlueprint.cpp
--- a/Engine/Plugins/Animation/ControlRig/Source/ControlRigDeveloper/Private/ControlRigBlueprint.cpp
+++ b/Engine/Plugins/Animation/ControlRig/Source/ControlRigDeveloper/Private/ControlRigBlueprint.cpp
@@ -134,6 +134,7 @@ UControlRigBlueprint::UControlRigBlueprint(const FObjectInitializer& ObjectIniti
 	FunctionLibraryEdGraph->bEditable = 0;
 	FunctionLibraryEdGraph->bAllowDeletion = 0;
 	FunctionLibraryEdGraph->bIsFunctionDefinition = false;
+	FunctionLibraryEdGraph->ModelNodePath = RigVMClient.GetFunctionLibrary()->GetNodePath();
 	FunctionLibraryEdGraph->Initialize(this);
 
 	Validator = ObjectInitializer.CreateDefaultSubobject<UControlRigValidator>(this, TEXT("ControlRigValidator"));
@@ -320,6 +321,11 @@ UObject* UControlRigBlueprint::GetEditorObjectForRigVMGraph(URigVMGraph* InVMGra
 			return nullptr;
 		}
 
+		if(InVMGraph->IsA<URigVMFunctionLibrary>())
+		{
+			return FunctionLibraryEdGraph;
+		}
+
 		TArray<UEdGraph*> EdGraphs;
 		GetAllGraphs(EdGraphs);
 
@@ -3198,6 +3204,11 @@ void UControlRigBlueprint::RebuildGraphFromModel()
 		}
 	}
 
+	if(FunctionLibraryEdGraph && RigVMClient.GetFunctionLibrary())
+	{
+		FunctionLibraryEdGraph->ModelNodePath = RigVMClient.GetFunctionLibrary()->GetNodePath();
+	}
+
 	TArray<URigVMGraph*> RigGraphs = RigVMClient.GetAllModels(true, true);
 
 	for (int32 RigGraphIndex = 0; RigGraphIndex < RigGraphs.Num(); RigGraphIndex++)
@@ -4826,10 +4837,12 @@ UEdGraph* UControlRigBlueprint::CreateEdGraph(URigVMGraph* InModel, bool bForce)
 {
 	check(InModel);
 
+#if WITH_EDITORONLY_DATA
 	if(InModel->IsA<URigVMFunctionLibrary>())
 	{
 		return FunctionLibraryEdGraph;
 	}
+#endif
 	
 	if(bForce)
 	{
diff --git a/Engine/Plugins/Animation/ControlRig/Source/ControlRigEditor/Private/ControlRigGraphDetails.cpp b/Engine/Plugins/Animation/ControlRig/Source/ControlRigEditor/Private/ControlRigGraphDetails.cpp
--- a/Engine/Plugins/Animation/ControlRig/Source/ControlRigEditor/Private/ControlRigGraphDetails.cpp
+++ b/Engine/Plugins/Animation/ControlRig/Source/ControlRigEditor/Private/ControlRigGraphDetails.cpp
@@ -1593,6 +1593,10 @@ void FControlRigWrappedNodeDetails::CustomizeLiveValues(IDetailLayoutBuilder& De
 
 	UDetailsViewWrapperObject* FirstWrapper = ObjectsBeingCustomized[0].Get();
 	URigVMNode* FirstNode = NodesBeingCustomized[0].Get();
+	if(FirstNode->GetTypedOuter<URigVMFunctionLibrary>())
+	{
+		return;
+	}
 
 	TSharedPtr<FRigVMParserAST> AST = FirstNode->GetGraph()->GetRuntimeAST(BlueprintBeingCustomized->VMCompileSettings.ASTSettings, false);
 	if(!AST.IsValid())
diff --git a/Engine/Plugins/Animation/ControlRig/Source/ControlRigEditor/Private/Editor/ControlRigEditor.cpp b/Engine/Plugins/Animation/ControlRig/Source/ControlRigEditor/Private/Editor/ControlRigEditor.cpp
--- a/Engine/Plugins/Animation/ControlRig/Source/ControlRigEditor/Private/Editor/ControlRigEditor.cpp
+++ b/Engine/Plugins/Animation/ControlRig/Source/ControlRigEditor/Private/Editor/ControlRigEditor.cpp
@@ -1518,11 +1518,14 @@ void FControlRigEditor::SetDetailObjects(const TArray<UObject*>& InObjects, bool
 	{
 		if (URigVMLibraryNode* LibraryNode = Cast<URigVMLibraryNode>(InObject))
 		{
-			UEdGraph* EdGraph = GetControlRigBlueprint()->GetEdGraph(LibraryNode->GetContainedGraph());
-			if (EdGraph)
+			if(!LibraryNode->IsA<URigVMFunctionReferenceNode>())
 			{
-				FilteredObjects.AddUnique(EdGraph);
-				continue;
+				UEdGraph* EdGraph = GetControlRigBlueprint()->GetEdGraph(LibraryNode->GetContainedGraph());
+				if (EdGraph)
+				{
+					FilteredObjects.AddUnique(EdGraph);
+					continue;
+				}
 			}
 		}
 		else if (Cast<URigVMFunctionEntryNode>(InObject) || Cast<URigVMFunctionReturnNode>(InObject))
@@ -1534,7 +1537,8 @@ void FControlRigEditor::SetDetailObjects(const TArray<UObject*>& InObjects, bool
 				continue;
 			}
 		}
-		else if (URigVMNode* ModelNode = Cast<URigVMNode>(InObject))
+
+		if (URigVMNode* ModelNode = Cast<URigVMNode>(InObject))
 		{
 			// check if we know the dynamic class already
 			bool bClassCreated = UDetailsViewWrapperObject::GetClassForNodes(ModelNodes, false) == nullptr;
