Sequencer: Set control shape actor selectability not based on whether the controls are visible in the viewport because the controls can also be selected in the Anim Outliner or in Sequencer. This fixes an issue where if you select a control in the Anim Outliner or in Sequencer, the transform gizmo won't appear because the control is not considered selected.

diff --git a/Engine/Plugins/Animation/ControlRig/Source/ControlRigEditor/Private/EditMode/ControlRigEditMode.cpp b/Engine/Plugins/Animation/ControlRig/Source/ControlRigEditor/Private/EditMode/ControlRigEditMode.cpp
--- a/Engine/Plugins/Animation/ControlRig/Source/ControlRigEditor/Private/EditMode/ControlRigEditMode.cpp
+++ b/Engine/Plugins/Animation/ControlRig/Source/ControlRigEditor/Private/EditMode/ControlRigEditMode.cpp
@@ -3199,7 +3199,7 @@ bool FControlRigEditMode::AreRigElementSelectedAndMovable(UControlRig* ControlRi
 {
 	const UControlRigEditModeSettings* Settings = GetDefault<UControlRigEditModeSettings>();
 
-	if (!Settings || Settings->bHideControlShapes || !ControlRig || !AreRigElementsSelected(ValidControlTypeMask(), ControlRig))
+	if (!Settings || !ControlRig || !AreRigElementsSelected(ValidControlTypeMask(), ControlRig))
 	{
 		return false;
 	}
@@ -3955,8 +3955,7 @@ void FControlRigEditMode::TickControlShape(AControlRigShapeActor* ShapeActor, co
 					bControlsHiddenInViewport);
 				
 				ShapeActor->SetSelectable(
-					ControlElement->Settings.IsSelectable(bRespectVisibilityForSelection) &&
-					!bControlsHiddenInViewport);
+					ControlElement->Settings.IsSelectable(bRespectVisibilityForSelection));
 			}
 		}
 	}
