Control Rig: Fix fallout from construction event - interaction event working again

diff --git a/Engine/Plugins/Animation/ControlRig/Source/ControlRig/Private/ControlRig.cpp b/Engine/Plugins/Animation/ControlRig/Source/ControlRig/Private/ControlRig.cpp
--- a/Engine/Plugins/Animation/ControlRig/Source/ControlRig/Private/ControlRig.cpp
+++ b/Engine/Plugins/Animation/ControlRig/Source/ControlRig/Private/ControlRig.cpp
@@ -69,6 +69,7 @@ UControlRig::UControlRig(const FObjectInitializer& ObjectInitializer)
 #endif 
 	, DataSourceRegistry(nullptr)
 	, EventQueue()
+	, EventQueueToRun()
 #if WITH_EDITOR
 	, PreviewInstance(nullptr)
 #endif
@@ -284,7 +285,7 @@ void UControlRig::Evaluate_AnyThread()
 
 	// create a copy since we need to change it here temporarily,
 	// and UI / the rig may change the event queue while it is running
-	TArray<FName> EventQueueToRun = EventQueue;
+	EventQueueToRun = EventQueue;
 
 	if(InteractionType != (uint8)EControlRigInteractionType::None)
 	{
@@ -323,6 +324,8 @@ void UControlRig::Evaluate_AnyThread()
 		}
 #endif
 	}
+
+	EventQueueToRun.Reset();
 }
 
 
@@ -628,10 +631,15 @@ bool UControlRig::Execute(const EControlRigState InState, const FName& InEventNa
 	{
 		return false;
 	}
+
+	if(EventQueueToRun.IsEmpty())
+	{
+		EventQueueToRun = EventQueue;
+	}
 	
-	const bool bIsEventInQueue = GetEventQueue().Contains(InEventName);
-	const bool bIsEventFirstInQueue = !GetEventQueue().IsEmpty() && GetEventQueue()[0] == InEventName; 
-	const bool bIsEventLastInQueue = !GetEventQueue().IsEmpty() && GetEventQueue().Last() == InEventName;
+	const bool bIsEventInQueue = EventQueueToRun.Contains(InEventName);
+	const bool bIsEventFirstInQueue = !EventQueueToRun.IsEmpty() && EventQueueToRun[0] == InEventName; 
+	const bool bIsEventLastInQueue = !EventQueueToRun.IsEmpty() && EventQueueToRun.Last() == InEventName;
 	const bool bIsInitializingMemory = InState == EControlRigState::Init;
 	const bool bIsExecutingInstructions = InState == EControlRigState::Update;
 	const bool bIsConstructionEvent = InEventName == FRigUnit_PrepareForExecution::EventName;
@@ -1010,9 +1018,9 @@ bool UControlRig::Execute(const EControlRigState InState, const FName& InEventNa
 #if WITH_EDITOR
 		// only set a valid first entry event when execution
 		// has passed the initialization stage and there are multiple events present
-		if (EventQueue.Num() >= 2 && VM && !bIsInitializingMemory)
+		if (EventQueueToRun.Num() >= 2 && VM && !bIsInitializingMemory)
 		{
-			VM->SetFirstEntryEventInEventQueue(EventQueue[0]);
+			VM->SetFirstEntryEventInEventQueue(EventQueueToRun[0]);
 		}
 
 		// Transform Overrride is generated using a Transient Control 
@@ -1797,7 +1805,7 @@ FRigControlElement* UControlRig::FindControl(const FName& InControlName) const
 
 bool UControlRig::IsConstructionModeEnabled() const
 {
-	return EventQueue.Num() == 1 && EventQueue.Contains(FRigUnit_PrepareForExecution::EventName);
+	return EventQueueToRun.Num() == 1 && EventQueueToRun.Contains(FRigUnit_PrepareForExecution::EventName);
 }
 
 FTransform UControlRig::SetupControlFromGlobalTransform(const FName& InControlName, const FTransform& InGlobalTransform)
@@ -1874,7 +1882,7 @@ void UControlRig::HandleOnControlModified(UControlRig* Subject, FRigControlEleme
 void UControlRig::HandleExecutionReachedExit(const FName& InEventName)
 {
 #if WITH_EDITOR
-	if (EventQueue.Last() == InEventName)
+	if (EventQueueToRun.Last() == InEventName)
 	{
 		if(URigVM* SnapShotVM = GetSnapshotVM(false))
 		{
diff --git a/Engine/Plugins/Animation/ControlRig/Source/ControlRig/Public/ControlRig.h b/Engine/Plugins/Animation/ControlRig/Source/ControlRig/Public/ControlRig.h
--- a/Engine/Plugins/Animation/ControlRig/Source/ControlRig/Public/ControlRig.h
+++ b/Engine/Plugins/Animation/ControlRig/Source/ControlRig/Public/ControlRig.h
@@ -583,6 +583,7 @@ private:
 	/** The event name used during an update */
 	UPROPERTY(transient)
 	TArray<FName> EventQueue;
+	TArray<FName> EventQueueToRun;
 
 	/** Copy the VM from the default object */
 	void InstantiateVMFromCDO();
