Control Rig: Update the shape transform from current instead of initial

diff --git a/Engine/Plugins/Animation/ControlRig/Source/ControlRig/Private/ControlRigGizmoActor.cpp b/Engine/Plugins/Animation/ControlRig/Source/ControlRig/Private/ControlRigGizmoActor.cpp
--- a/Engine/Plugins/Animation/ControlRig/Source/ControlRig/Private/ControlRigGizmoActor.cpp
+++ b/Engine/Plugins/Animation/ControlRig/Source/ControlRig/Private/ControlRigGizmoActor.cpp
@@ -164,7 +164,7 @@ bool AControlRigShapeActor::UpdateControlSettings(
 	// update the shape transform
 	if(bShapeTransformChanged)
 	{
-		const FTransform ShapeTransform = InControlElement->Shape.Get(ERigTransformType::InitialLocal);
+		const FTransform ShapeTransform = InControlElement->Shape.Get(ERigTransformType::CurrentLocal);
 		StaticMeshComponent->SetRelativeTransform(MeshTransform * ShapeTransform);
 	}
 	
