Control Rig: Use displayname namespace also for root elements

diff --git a/Engine/Plugins/Animation/ControlRig/Source/ControlRig/Private/Rigs/RigHierarchy.cpp b/Engine/Plugins/Animation/ControlRig/Source/ControlRig/Private/Rigs/RigHierarchy.cpp
--- a/Engine/Plugins/Animation/ControlRig/Source/ControlRig/Private/Rigs/RigHierarchy.cpp
+++ b/Engine/Plugins/Animation/ControlRig/Source/ControlRig/Private/Rigs/RigHierarchy.cpp
@@ -1017,38 +1017,55 @@ FName URigHierarchy::GetSafeNewDisplayName(const FRigElementKey& InParentElement
 	{
 		return NAME_None;
 	}
-	
+
+	TArray<FRigElementKey> KeysToCheck;
 	if(InParentElement.IsValid())
 	{
-		FString SanitizedName = InPotentialNewDisplayName;
-		SanitizeName(SanitizedName);
-		FString Name = SanitizedName;
-
-		const TArray<FRigElementKey> ChildKeys = GetChildren(InParentElement);
-		TArray<FString> DisplayNames;
-		Algo::Transform(ChildKeys, DisplayNames, [this](const FRigElementKey& InKey) -> FString
+		KeysToCheck = GetChildren(InParentElement);
+	}
+	else
+	{
+		// get all of the root elements
+		for(const FRigBaseElement* Element : Elements)
 		{
-			if(const FRigBaseElement* BaseElement = Find(InKey))
+			if(!Element->IsA<FRigTransformElement>())
 			{
-				return BaseElement->GetDisplayName().ToString();
+				continue;
 			}
-			return FString();
-		});
-
-		int32 Suffix = 1;
-		while (DisplayNames.Contains(Name))
-		{
-			FString BaseString = SanitizedName;
-			if (BaseString.Len() > GetMaxNameLength() - 4)
+			
+			if(GetNumberOfParents(Element) == 0)
 			{
-				BaseString.LeftChopInline(BaseString.Len() - (GetMaxNameLength() - 4));
+				KeysToCheck.Add(Element->GetKey());
 			}
-			Name = *FString::Printf(TEXT("%s_%d"), *BaseString, ++Suffix);
 		}
+	}
+
+	FString SanitizedName = InPotentialNewDisplayName;
+	SanitizeName(SanitizedName);
+	FString Name = SanitizedName;
 
-		return *Name;
+	TArray<FString> DisplayNames;
+	Algo::Transform(KeysToCheck, DisplayNames, [this](const FRigElementKey& InKey) -> FString
+	{
+		if(const FRigBaseElement* BaseElement = Find(InKey))
+		{
+			return BaseElement->GetDisplayName().ToString();
+		}
+		return FString();
+	});
+
+	int32 Suffix = 1;
+	while (DisplayNames.Contains(Name))
+	{
+		FString BaseString = SanitizedName;
+		if (BaseString.Len() > GetMaxNameLength() - 4)
+		{
+			BaseString.LeftChopInline(BaseString.Len() - (GetMaxNameLength() - 4));
+		}
+		Name = *FString::Printf(TEXT("%s_%d"), *BaseString, ++Suffix);
 	}
-	return *InPotentialNewDisplayName;
+
+	return *Name;
 }
 
 FEdGraphPinType URigHierarchy::GetControlPinType(FRigControlElement* InControlElement) const
diff --git a/Engine/Plugins/Animation/ControlRig/Source/ControlRig/Private/Rigs/RigHierarchyController.cpp b/Engine/Plugins/Animation/ControlRig/Source/ControlRig/Private/Rigs/RigHierarchyController.cpp
--- a/Engine/Plugins/Animation/ControlRig/Source/ControlRig/Private/Rigs/RigHierarchyController.cpp
+++ b/Engine/Plugins/Animation/ControlRig/Source/ControlRig/Private/Rigs/RigHierarchyController.cpp
@@ -2017,13 +2017,16 @@ FName URigHierarchyController::SetDisplayName(FRigControlElement* InControlEleme
 		return NAME_None;
 	}
 
+	FRigElementKey ParentElementKey;
 	if(const FRigBaseElement* ParentElement = Hierarchy->GetFirstParent(InControlElement))
 	{
-		const FString DesiredDisplayName = InDisplayName.IsNone() ? FString() : InDisplayName.ToString();
-		const FName DisplayName = Hierarchy->GetSafeNewDisplayName(ParentElement->GetKey(), DesiredDisplayName);
-		InControlElement->Settings.DisplayName = DisplayName;
+		ParentElementKey = ParentElement->GetKey();
 	}
-	
+
+	const FString DesiredDisplayName = InDisplayName.IsNone() ? FString() : InDisplayName.ToString();
+	const FName DisplayName = Hierarchy->GetSafeNewDisplayName(ParentElementKey, DesiredDisplayName);
+	InControlElement->Settings.DisplayName = DisplayName;
+
 	Hierarchy->IncrementTopologyVersion();
 	Notify(ERigHierarchyNotification::ControlSettingChanged, InControlElement);
 
diff --git a/Engine/Plugins/Animation/ControlRig/Source/ControlRigEditor/Private/ControlRigElementDetails.cpp b/Engine/Plugins/Animation/ControlRig/Source/ControlRigEditor/Private/ControlRigElementDetails.cpp
--- a/Engine/Plugins/Animation/ControlRig/Source/ControlRigEditor/Private/ControlRigElementDetails.cpp
+++ b/Engine/Plugins/Animation/ControlRig/Source/ControlRigEditor/Private/ControlRigElementDetails.cpp
@@ -3735,7 +3735,7 @@ FText FRigControlElementDetails::GetDisplayName() const
 	for(int32 ObjectIndex = 0; ObjectIndex < PerElementInfos.Num(); ObjectIndex++)
 	{
 		const FPerElementInfo& Info = PerElementInfos[ObjectIndex];
-		if(const FRigControlElement* ControlElement = Info.GetElement<FRigControlElement>())
+		if(const FRigControlElement* ControlElement = Info.GetDefaultElement<FRigControlElement>())
 		{
  			const FName ThisDisplayName =
  				(ControlElement->IsAnimationChannel()) ?
@@ -3765,7 +3765,7 @@ void FRigControlElementDetails::SetDisplayName(const FText& InNewText, ETextComm
 	for(int32 ObjectIndex = 0; ObjectIndex < PerElementInfos.Num(); ObjectIndex++)
 	{
 		const FPerElementInfo& Info = PerElementInfos[ObjectIndex];
-		if(const FRigControlElement* ControlElement = Info.GetElement<FRigControlElement>())
+		if(const FRigControlElement* ControlElement = Info.GetDefaultElement<FRigControlElement>())
 		{
 			SetDisplayNameForElement(InNewText, InCommitType, ControlElement->GetKey());
 		}
