Control Rig: Don't use detail wrapper for comments

diff --git a/Engine/Plugins/Animation/ControlRig/Source/ControlRigEditor/Private/Editor/ControlRigEditor.cpp b/Engine/Plugins/Animation/ControlRig/Source/ControlRigEditor/Private/Editor/ControlRigEditor.cpp
--- a/Engine/Plugins/Animation/ControlRig/Source/ControlRigEditor/Private/Editor/ControlRigEditor.cpp
+++ b/Engine/Plugins/Animation/ControlRig/Source/ControlRigEditor/Private/Editor/ControlRigEditor.cpp
@@ -1520,8 +1520,7 @@ void FControlRigEditor::SetDetailObjects(const TArray<UObject*>& InObjects, bool
 		{
 			if(!LibraryNode->IsA<URigVMFunctionReferenceNode>())
 			{
-				UEdGraph* EdGraph = GetControlRigBlueprint()->GetEdGraph(LibraryNode->GetContainedGraph());
-				if (EdGraph)
+				if (UEdGraph* EdGraph = GetControlRigBlueprint()->GetEdGraph(LibraryNode->GetContainedGraph()))
 				{
 					FilteredObjects.AddUnique(EdGraph);
 					continue;
@@ -1530,13 +1529,23 @@ void FControlRigEditor::SetDetailObjects(const TArray<UObject*>& InObjects, bool
 		}
 		else if (Cast<URigVMFunctionEntryNode>(InObject) || Cast<URigVMFunctionReturnNode>(InObject))
 		{
-			UEdGraph* EdGraph = GetControlRigBlueprint()->GetEdGraph(CastChecked<URigVMNode>(InObject)->GetGraph());
-			if (EdGraph)
+			if (UEdGraph* EdGraph = GetControlRigBlueprint()->GetEdGraph(CastChecked<URigVMNode>(InObject)->GetGraph()))
 			{
 				FilteredObjects.AddUnique(EdGraph);
 				continue;
 			}
 		}
+		else if (URigVMCommentNode* CommentNode = Cast<URigVMCommentNode>(InObject))
+		{
+			if (UControlRigGraph* EdGraph = Cast<UControlRigGraph>(GetControlRigBlueprint()->GetEdGraph(CastChecked<URigVMNode>(InObject)->GetGraph())))
+			{
+				if(UEdGraphNode* EdGraphNode = EdGraph->FindNodeForModelNodeName(CommentNode->GetFName()))
+				{
+					FilteredObjects.AddUnique(EdGraphNode);
+					continue;
+				}
+			}
+		}
 
 		if (URigVMNode* ModelNode = Cast<URigVMNode>(InObject))
 		{
