Fixed control rig layer instance proxy to not re-locate internal control rig nodes

Control rig nodes (and others) cannot be re-located in editor due to re-instancing delegate registration
Switched to shared ptr to hold internal control rig nodes.

diff --git a/Engine/Plugins/Animation/ControlRig/Source/ControlRig/Private/Sequencer/ControlRigLayerInstanceProxy.cpp b/Engine/Plugins/Animation/ControlRig/Source/ControlRig/Private/Sequencer/ControlRigLayerInstanceProxy.cpp
--- a/Engine/Plugins/Animation/ControlRig/Source/ControlRig/Private/Sequencer/ControlRigLayerInstanceProxy.cpp
+++ b/Engine/Plugins/Animation/ControlRig/Source/ControlRig/Private/Sequencer/ControlRigLayerInstanceProxy.cpp
@@ -45,13 +45,13 @@ void FControlRigLayerInstanceProxy::ConstructNodes()
 {
 	if (ControlRigNodes.Num() > 0)
 	{
-		FAnimNode_ControlRig_ExternalSource* ParentNode = &ControlRigNodes[0];
+		FAnimNode_ControlRig_ExternalSource* ParentNode = ControlRigNodes[0].Get();
 		CurrentRoot = static_cast<FAnimNode_Base*>(ParentNode);
 
 		// index 0 - (N-1) is the order of operation
 		for (int32 Index = 1; Index < ControlRigNodes.Num(); ++Index)
 		{
-			FAnimNode_ControlRig_ExternalSource* CurrentNode = &ControlRigNodes[Index];
+			FAnimNode_ControlRig_ExternalSource* CurrentNode = ControlRigNodes[Index].Get();
 			ParentNode->Source.SetLinkNode(CurrentNode);
 			ParentNode = CurrentNode;
 		}
@@ -131,8 +131,8 @@ void FControlRigLayerInstanceProxy::AddControlRigTrack(int32 ControlRigID, UCont
 
 	if(!Node)
 	{
-		FAnimNode_ControlRig_ExternalSource* Parent = (ControlRigNodes.Num() > 0)? &ControlRigNodes.Last() : nullptr;
-		Node = &ControlRigNodes.AddDefaulted_GetRef();
+		FAnimNode_ControlRig_ExternalSource* Parent = (ControlRigNodes.Num() > 0)? ControlRigNodes.Last().Get() : nullptr;
+		Node = ControlRigNodes.Add_GetRef(MakeShared<FAnimNode_ControlRig_ExternalSource>()).Get();
 		SequencerToControlRigNodeMap.FindOrAdd(ControlRigID) = Node;
 		if (Parent)
 		{
@@ -182,13 +182,13 @@ void FControlRigLayerInstanceProxy::RemoveControlRigTrack(int32 ControlRigID)
 		FAnimNode_ControlRig_ExternalSource* Parent = nullptr;
 		for (int32 Index = 0; Index < ControlRigNodes.Num(); ++Index)
 		{
-			FAnimNode_ControlRig_ExternalSource* Current = &ControlRigNodes[Index];
+			FAnimNode_ControlRig_ExternalSource* Current = ControlRigNodes[Index].Get();
 
 			if (Current == Node)
 			{
 				// we need to delete this one
 				// find next child one
-				FAnimNode_ControlRig_ExternalSource* Child = (ControlRigNodes.IsValidIndex(Index + 1)) ? &ControlRigNodes[Index + 1] : nullptr;
+				FAnimNode_ControlRig_ExternalSource* Child = (ControlRigNodes.IsValidIndex(Index + 1)) ? ControlRigNodes[Index + 1].Get() : nullptr;
 
 				// if no parent, change root
 				if (Parent == nullptr)
@@ -244,9 +244,9 @@ UAnimSequencerInstance* FControlRigLayerInstanceProxy::GetSequencerAnimInstance(
 
 UControlRig* FControlRigLayerInstanceProxy::GetFirstAvailableControlRig() const
 {
-	for (const FAnimNode_ControlRig_ExternalSource& ControlRigNode : ControlRigNodes)
+	for (const TSharedPtr<FAnimNode_ControlRig_ExternalSource>& ControlRigNode : ControlRigNodes)
 	{
-		if (UControlRig* Rig = ControlRigNode.GetControlRig())
+		if (UControlRig* Rig = ControlRigNode->GetControlRig())
 		{
 			return Rig;
 		}
diff --git a/Engine/Plugins/Animation/ControlRig/Source/ControlRig/Private/Sequencer/ControlRigLayerInstanceProxy.h b/Engine/Plugins/Animation/ControlRig/Source/ControlRig/Private/Sequencer/ControlRigLayerInstanceProxy.h
--- a/Engine/Plugins/Animation/ControlRig/Source/ControlRig/Private/Sequencer/ControlRigLayerInstanceProxy.h
+++ b/Engine/Plugins/Animation/ControlRig/Source/ControlRig/Private/Sequencer/ControlRigLayerInstanceProxy.h
@@ -117,7 +117,7 @@ protected:
 	FAnimNode_Base* CurrentRoot;
 
 	/** ControlRig Nodes */
-	TArray<FAnimNode_ControlRig_ExternalSource> ControlRigNodes;
+	TArray<TSharedPtr<FAnimNode_ControlRig_ExternalSource>> ControlRigNodes;
 
 	/** mapping from sequencer index to internal player index */
 	TMap<int32, FAnimNode_ControlRig_ExternalSource*> SequencerToControlRigNodeMap;
