Sequencer: Added blueprint functions to get/set control rig parameter section masks

diff --git a/Engine/Plugins/Animation/ControlRig/Source/ControlRigEditor/Private/ControlRigSequencerEditorLibrary.cpp b/Engine/Plugins/Animation/ControlRig/Source/ControlRigEditor/Private/ControlRigSequencerEditorLibrary.cpp
--- a/Engine/Plugins/Animation/ControlRig/Source/ControlRigEditor/Private/ControlRigSequencerEditorLibrary.cpp
+++ b/Engine/Plugins/Animation/ControlRig/Source/ControlRigEditor/Private/ControlRigSequencerEditorLibrary.cpp
@@ -2412,4 +2412,93 @@ FRigElementKey UControlRigSequencerEditorLibrary::GetWorldSpaceReferenceKey()
 	return URigHierarchy::GetWorldSpaceReferenceKey();
 }
 
+bool UControlRigSequencerEditorLibrary::GetControlsMask(UMovieSceneSection* InSection, FName ControlName)
+{
+	UMovieSceneControlRigParameterSection* ParameterSection = Cast<UMovieSceneControlRigParameterSection>(InSection);
+	if (!ParameterSection)
+	{
+		FFrame::KismetExecutionMessage(TEXT("Cannot call GetControlsMask without a UMovieSceneControlRigParameterSection"), ELogVerbosity::Error);
+		return false;
+	}
+	
+	UControlRig* ControlRig = ParameterSection->GetControlRig();
+	if (!ControlRig)
+	{
+		FFrame::KismetExecutionMessage(TEXT("Section does not have a control rig"), ELogVerbosity::Error);
+		return false;
+	}
+
+	TArray<FRigControlElement*> Controls;
+	ControlRig->GetControlsInOrder(Controls);
+	int32 Index = 0;
+	for (const FRigControlElement* RigControl : Controls)
+	{
+		if (RigControl->GetName() == ControlName)
+		{
+			return ParameterSection->GetControlsMask(Index);
+		}
+		++Index;
+	}
+
+	FFrame::KismetExecutionMessage(*FString::Printf(TEXT("Control Name ('%s') not found"), *ControlName.ToString()), ELogVerbosity::Error);
+	return false;
+}
+
+void UControlRigSequencerEditorLibrary::SetControlsMask(UMovieSceneSection* InSection, const TArray<FName>& ControlNames, bool bVisible)
+{
+	UMovieSceneControlRigParameterSection* ParameterSection = Cast<UMovieSceneControlRigParameterSection>(InSection);
+	if (!ParameterSection)
+	{
+		FFrame::KismetExecutionMessage(TEXT("Cannot call SetControlsMask without a UMovieSceneControlRigParameterSection"), ELogVerbosity::Error);
+		return;
+	}
+
+	UControlRig* ControlRig = ParameterSection->GetControlRig();
+	if (!ControlRig)
+	{
+		FFrame::KismetExecutionMessage(TEXT("Section does not have a control rig"), ELogVerbosity::Error);
+		return;
+	}
+
+	ParameterSection->Modify();
+
+	TArray<FRigControlElement*> Controls;
+	ControlRig->GetControlsInOrder(Controls);
+	int32 Index = 0;
+	for (const FRigControlElement* RigControl : Controls)
+	{
+		if (ControlNames.Contains(RigControl->GetName()))
+		{
+			ParameterSection->SetControlsMask(Index, bVisible);
+		}
+		++Index;
+	}
+}
+
+void UControlRigSequencerEditorLibrary::ShowAllControls(UMovieSceneSection* InSection)
+{
+	UMovieSceneControlRigParameterSection* ParameterSection = Cast<UMovieSceneControlRigParameterSection>(InSection);
+	if (!ParameterSection)
+	{
+		FFrame::KismetExecutionMessage(TEXT("Cannot call ShowAllControls without a UMovieSceneControlRigParameterSection"), ELogVerbosity::Error);
+		return;
+	}
+
+	ParameterSection->Modify();
+	ParameterSection->FillControlsMask(true);
+}
+
+void UControlRigSequencerEditorLibrary::HideAllControls(UMovieSceneSection* InSection)
+{
+	UMovieSceneControlRigParameterSection* ParameterSection = Cast<UMovieSceneControlRigParameterSection>(InSection);
+	if (!ParameterSection)
+	{
+		FFrame::KismetExecutionMessage(TEXT("Cannot call HideAllControls without a UMovieSceneControlRigParameterSection"), ELogVerbosity::Error);
+		return;
+	}
+
+	ParameterSection->Modify();
+	ParameterSection->FillControlsMask(false);
+}
+
 #undef LOCTEXT_NAMESPACE
diff --git a/Engine/Plugins/Animation/ControlRig/Source/ControlRigEditor/Public/ControlRigSequencerEditorLibrary.h b/Engine/Plugins/Animation/ControlRig/Source/ControlRigEditor/Public/ControlRigSequencerEditorLibrary.h
--- a/Engine/Plugins/Animation/ControlRig/Source/ControlRigEditor/Public/ControlRigSequencerEditorLibrary.h
+++ b/Engine/Plugins/Animation/ControlRig/Source/ControlRigEditor/Public/ControlRigSequencerEditorLibrary.h
@@ -887,4 +887,19 @@ public:
 	UFUNCTION(BlueprintCallable, Category = "Editor Scripting | Sequencer Tools | Control Rig")
 	bool RenameControlRigControlChannels(ULevelSequence* InSequence, UControlRig* InControlRig, const TArray<FName>& InOldControlNames, const TArray<FName>& InNewControlNames);
 
+	/** Get the controls mask for the given ControlName */
+	UFUNCTION(BlueprintPure, Category = "Editor Scripting | Sequencer Tools | Control Rig")
+	static bool GetControlsMask(UMovieSceneSection* InSection, FName ControlName);
+
+	/** Set the controls mask for the given ControlNames */
+	UFUNCTION(BlueprintCallable, Category = "Editor Scripting | Sequencer Tools | Control Rig")
+	static void SetControlsMask(UMovieSceneSection* InSection, const TArray<FName>& ControlNames, bool bVisible);
+
+	/** Shows all of the controls for the given section */
+	UFUNCTION(BlueprintCallable, Category = "Editor Scripting | Sequencer Tools | Control Rig")
+	static void ShowAllControls(UMovieSceneSection* InSection);
+
+	/** Hides all of the controls for the given section */
+	UFUNCTION(BlueprintCallable, Category = "Editor Scripting | Sequencer Tools | Control Rig")
+	static void HideAllControls(UMovieSceneSection* InSection);
 };
