Sequencer: Control Rig: Wrap setting keys up in a transaction so we don't do one per control.
#trivial

diff --git a/Engine/Plugins/Animation/ControlRig/Source/ControlRigEditor/Private/Sequencer/ControlRigParameterTrackEditor.cpp b/Engine/Plugins/Animation/ControlRig/Source/ControlRigEditor/Private/Sequencer/ControlRigParameterTrackEditor.cpp
--- a/Engine/Plugins/Animation/ControlRig/Source/ControlRigEditor/Private/Sequencer/ControlRigParameterTrackEditor.cpp
+++ b/Engine/Plugins/Animation/ControlRig/Source/ControlRigEditor/Private/Sequencer/ControlRigParameterTrackEditor.cpp
@@ -1534,9 +1534,13 @@ void FControlRigParameterTrackEditor::OnAddTransformKeysForSelectedObjects(EMovi
 
 	TMap<UControlRig*, TArray<FRigElementKey>> SelectedControls;
 	ControlRigEditMode->GetAllSelectedControls(SelectedControls);
-
+	if (SelectedControls.Num() <= 0)
+	{
+		return;
+	}
 	const EControlRigContextChannelToKey ChannelsToKey = static_cast<EControlRigContextChannelToKey>(Channel); 
-	
+	FScopedTransaction KeyTransaction(LOCTEXT("SetKeysOnControls", "Set Keys On Controls"), !GIsTransacting);
+
 	for (const TPair<UControlRig*, TArray<FRigElementKey>>& Selection : SelectedControls)
 	{
 		UControlRig* ControlRig = Selection.Key;
