ControlRigTransformableHandle: Fix non-unity build break, add missing include

diff --git a/Engine/Plugins/Animation/ControlRig/Source/ControlRig/Public/Constraints/ControlRigTransformableHandle.h b/Engine/Plugins/Animation/ControlRig/Source/ControlRig/Public/Constraints/ControlRigTransformableHandle.h
--- a/Engine/Plugins/Animation/ControlRig/Source/ControlRig/Public/Constraints/ControlRigTransformableHandle.h
+++ b/Engine/Plugins/Animation/ControlRig/Source/ControlRig/Public/Constraints/ControlRigTransformableHandle.h
@@ -4,6 +4,7 @@
 
 #include "TransformableHandle.h"
 #include "Rigs/RigHierarchyDefines.h"
+#include "ControlRig.h"
 
 #include "ControlRigTransformableHandle.generated.h"
 
