Control Rig: Maintain initial pose in certain scenraios before construction

diff --git a/Engine/Plugins/Animation/ControlRig/Source/ControlRig/Private/ControlRig.cpp b/Engine/Plugins/Animation/ControlRig/Source/ControlRig/Private/ControlRig.cpp
--- a/Engine/Plugins/Animation/ControlRig/Source/ControlRig/Private/ControlRig.cpp
+++ b/Engine/Plugins/Animation/ControlRig/Source/ControlRig/Private/ControlRig.cpp
@@ -943,7 +943,21 @@ bool UControlRig::Execute(const EControlRigState InState, const FName& InEventNa
 #if WITH_EDITOR
 					FTransientControlScope TransientControlScope(GetHierarchy());
 	#endif
-					GetHierarchy()->ResetToDefault();
+					{
+						// maintain the initial pose if it ever was set by the client
+						FRigPose InitialPose;
+						if(!bResetInitialTransformsBeforeConstruction)
+						{
+							InitialPose = GetHierarchy()->GetPose(true, ERigElementType::ToResetAfterConstructionEvent, FRigElementKeyCollection());
+						}
+						
+						GetHierarchy()->ResetToDefault();
+
+						if(InitialPose.Num() > 0)
+						{
+							GetHierarchy()->SetPose(InitialPose, ERigTransformType::InitialLocal);
+						}
+					}
 
 					{
 	#if WITH_EDITOR
