Control Rig: Guard SRigSpacePickerWidget against invalid objects

diff --git a/Engine/Plugins/Animation/ControlRig/Source/ControlRigEditor/Private/Editor/SRigSpacePickerWidget.cpp b/Engine/Plugins/Animation/ControlRig/Source/ControlRigEditor/Private/Editor/SRigSpacePickerWidget.cpp
--- a/Engine/Plugins/Animation/ControlRig/Source/ControlRigEditor/Private/Editor/SRigSpacePickerWidget.cpp
+++ b/Engine/Plugins/Animation/ControlRig/Source/ControlRigEditor/Private/Editor/SRigSpacePickerWidget.cpp
@@ -212,7 +212,7 @@ void SRigSpacePickerWidget::SetControls(
 	{
 		if(HierarchyModifiedHandle.IsValid())
 		{
-			if(IsValid(Hierarchy) && !Hierarchy->HasAnyFlags(RF_BeginDestroyed))
+			if(IsValid(Hierarchy) && !Hierarchy->HasAnyFlags(RF_BeginDestroyed) && !Hierarchy->HasAnyFlags(RF_TagGarbageTemp) && !Hierarchy->HasAnyFlags(RF_Garbage))
 			{
 				Hierarchy->OnModified().Remove(HierarchyModifiedHandle);
 			}
