Control Rig: Metadata support - first pass.

Not yet enabled in the user interface.

diff --git a/Engine/Plugins/Animation/ControlRig/Source/ControlRig/Private/Rigs/RigHierarchy.cpp b/Engine/Plugins/Animation/ControlRig/Source/ControlRig/Private/Rigs/RigHierarchy.cpp
--- a/Engine/Plugins/Animation/ControlRig/Source/ControlRig/Private/Rigs/RigHierarchy.cpp
+++ b/Engine/Plugins/Animation/ControlRig/Source/ControlRig/Private/Rigs/RigHierarchy.cpp
@@ -790,6 +790,52 @@ bool URigHierarchy::IsProcedural(const FRigBaseElement* InElement) const
 	return InElement->IsProcedural();
 }
 
+TArray<FName> URigHierarchy::GetMetadataNames(FRigElementKey InItem)
+{
+	TArray<FName> Names;
+	if(const FRigBaseElement* Element = Find(InItem))
+	{
+		for(int32 Index = 0; Index < Element->NumMetadata(); Index++)
+		{
+			Names.Add(Element->GetMetadata(Index)->GetName());
+		}
+	}
+	return Names;
+}
+
+ERigMetadataType URigHierarchy::GetMetadataType(FRigElementKey InItem, FName InMetadataName)
+{
+	if(const FRigBaseElement* Element = Find(InItem))
+	{
+		for(int32 Index = 0; Index < Element->NumMetadata(); Index++)
+		{
+			if(const FRigBaseMetadata* Metadata = Element->GetMetadata(InMetadataName))
+			{
+				return Metadata->GetType();
+			}
+		}
+	}
+	return ERigMetadataType::Invalid;
+}
+
+bool URigHierarchy::RemoveMetadata(FRigElementKey InItem, FName InMetadataName)
+{
+	if(FRigBaseElement* Element = Find(InItem))
+	{
+		return Element->RemoveMetadata(InMetadataName);
+	}
+	return false;
+}
+
+bool URigHierarchy::RemoveAllMetadata(FRigElementKey InItem)
+{
+	if(FRigBaseElement* Element = Find(InItem))
+	{
+		return Element->RemoveAllMetadata();
+	}
+	return false;
+}
+
 TArray<const FRigBaseElement*> URigHierarchy::GetSelectedElements(ERigElementType InTypeFilter) const
 {
 	LLM_SCOPE_BYNAME(TEXT("Animation/ControlRig"));
diff --git a/Engine/Plugins/Animation/ControlRig/Source/ControlRig/Private/Rigs/RigHierarchyElements.cpp b/Engine/Plugins/Animation/ControlRig/Source/ControlRig/Private/Rigs/RigHierarchyElements.cpp
--- a/Engine/Plugins/Animation/ControlRig/Source/ControlRig/Private/Rigs/RigHierarchyElements.cpp
+++ b/Engine/Plugins/Animation/ControlRig/Source/ControlRig/Private/Rigs/RigHierarchyElements.cpp
@@ -7,11 +7,17 @@
 #include "ControlRigObjectVersion.h"
 #include "ControlRigGizmoLibrary.h"
 #include "AnimationCoreLibrary.h"
+#include "Algo/Transform.h"
 
 ////////////////////////////////////////////////////////////////////////////////
 // FRigBaseElement
 ////////////////////////////////////////////////////////////////////////////////
 
+FRigBaseElement::~FRigBaseElement()
+{
+	RemoveAllMetadata();
+}
+
 UScriptStruct* FRigBaseElement::GetElementStruct() const
 {
 	switch(GetType())
@@ -71,6 +77,21 @@ void FRigBaseElement::Save(FArchive& Ar, URigHierarchy* Hierarchy, ESerializatio
 	if(SerializationPhase == ESerializationPhase::StaticData)
 	{
 		Ar << Key;
+
+		static const UEnum* MetadataTypeEnum = StaticEnum<ERigMetadataType>();
+
+		int32 MetadataNum = Metadata.Num();
+		Ar << MetadataNum;
+
+		for(FRigBaseMetadata* Md : Metadata)
+		{
+			FName MetadataName = Md->GetName();
+			FName MetadataTypeName = MetadataTypeEnum->GetNameByValue((int64)Md->GetType());
+
+			Ar << MetadataName;
+			Ar << MetadataTypeName;
+			Md->Serialize(Ar, false);
+		}
 	}
 }
 
@@ -86,7 +107,140 @@ void FRigBaseElement::Load(FArchive& Ar, URigHierarchy* Hierarchy, ESerializatio
 		Key = LoadedKey;
 
 		NameString = Key.Name.ToString();
+
+		RemoveAllMetadata();
+		
+		if (Ar.CustomVer(FControlRigObjectVersion::GUID) >= FControlRigObjectVersion::HierarchyElementMetadata)
+		{
+			static const UEnum* MetadataTypeEnum = StaticEnum<ERigMetadataType>();
+
+			int32 MetadataNum = 0;
+			Ar << MetadataNum;
+
+			for(int32 MetadataIndex = 0; MetadataIndex < MetadataNum; MetadataIndex++)
+			{
+				FName MetadataName(NAME_None);
+				FName MetadataTypeName(NAME_None);
+				Ar << MetadataName;
+				Ar << MetadataTypeName;
+				
+				const ERigMetadataType MetadataType = (ERigMetadataType)MetadataTypeEnum->GetValueByName(MetadataTypeName);
+				FRigBaseMetadata* Md = FRigBaseMetadata::MakeMetadata(this, MetadataName, MetadataType);
+				Md->Serialize(Ar, true);
+				Metadata.Add(Md);
+			}
+
+			for(int32 MetadataIndex = 0; MetadataIndex < Metadata.Num(); MetadataIndex++)
+			{
+				MetadataNameToIndex.Add(Metadata[MetadataIndex]->GetName(), MetadataIndex);
+			}
+		}
+	}
+}
+
+bool FRigBaseElement::RemoveMetadata(const FName& InName)
+{
+	if(const int32* MetadataIndexPtr = MetadataNameToIndex.Find(InName))
+	{
+		const int32 MetadataIndex = *MetadataIndexPtr;
+		{
+			FRigBaseMetadata* Md = Metadata[MetadataIndex];
+			if(const UScriptStruct* Struct = Md->GetMetadataStruct())
+			{
+				Struct->DestroyStruct(Md, 1);
+			}
+			FMemory::Free(Md);
+		}
+		MetadataNameToIndex.Remove(InName);
+		Metadata.RemoveAt(MetadataIndex);
+		for(TPair<FName, int32>& Pair : MetadataNameToIndex)
+		{
+			if(Pair.Value > MetadataIndex)
+			{
+				Pair.Value--;
+			}
+		}
+		return true;
+	}
+	return false;
+}
+
+bool FRigBaseElement::RemoveAllMetadata()
+{
+	if(!Metadata.IsEmpty())
+	{
+		for(FRigBaseMetadata* Md : Metadata)
+		{
+			if(const UScriptStruct* Struct = Md->GetMetadataStruct())
+			{
+				Struct->DestroyStruct(Md, 1);
+			}
+			FMemory::Free(Md);
+		}
+		Metadata.Reset();
+		MetadataNameToIndex.Reset();
+		return true;
 	}
+	return false;
+}
+
+void FRigBaseElement::CopyFrom(URigHierarchy* InHierarchy, FRigBaseElement* InOther, URigHierarchy* InOtherHierarchy)
+{
+	// remember all previous names
+	TArray<FName> RemainingNames;
+	Algo::Transform(Metadata, RemainingNames, [](const FRigBaseMetadata* Md) -> FName
+	{
+		return Md->GetName();
+	});
+
+	// copy over all metadata. this also takes care of potential type changes
+	for(const FRigBaseMetadata* InOtherMd : InOther->Metadata)
+	{
+		FRigBaseMetadata* Md = SetupValidMetadata(InOtherMd->Name, InOtherMd->Type);
+		check(Md);
+		Md->SetValueData(InOtherMd->GetValueData(), InOtherMd->GetValueSize());
+		RemainingNames.Remove(InOtherMd->Name);
+	}
+
+	// remove all remaining metadata
+	for(const FName& NameToRemove : RemainingNames)
+	{
+		RemoveMetadata(NameToRemove);
+	}
+
+	// rebuild the name map
+	MetadataNameToIndex.Reset();
+	for(int32 MetadataIndex = 0; MetadataIndex < Metadata.Num(); MetadataIndex++)
+	{
+		MetadataNameToIndex.Add(Metadata[MetadataIndex]->GetName(), MetadataIndex);
+	}
+}
+
+FRigBaseMetadata* FRigBaseElement::SetupValidMetadata(const FName& InName, ERigMetadataType InType)
+{
+	if(const int32* MetadataIndex = MetadataNameToIndex.Find(InName))
+	{
+		FRigBaseMetadata* Md = Metadata[*MetadataIndex];
+		if(Md->GetType() == InType)
+		{
+			return Md;
+		}
+
+		if(const UScriptStruct* Struct = Md->GetMetadataStruct())
+		{
+			Struct->DestroyStruct(Md, 1);
+		}
+		FMemory::Free(Md);
+
+		Md = FRigBaseMetadata::MakeMetadata(this, InName, InType);
+		Metadata[*MetadataIndex] = Md;
+		return Md;
+	}
+
+	FRigBaseMetadata* Md = FRigBaseMetadata::MakeMetadata(this, InName, InType);
+	const int32 MetadataIndex = Metadata.Add(Md);
+	MetadataNameToIndex.Add(InName, MetadataIndex);
+	return Md;
 }
 
 ////////////////////////////////////////////////////////////////////////////////
diff --git a/Engine/Plugins/Animation/ControlRig/Source/ControlRig/Private/Rigs/RigHierarchyMetadata.cpp b/Engine/Plugins/Animation/ControlRig/Source/ControlRig/Private/Rigs/RigHierarchyMetadata.cpp
new file mode 100644
--- /dev/null
+++ b/Engine/Plugins/Animation/ControlRig/Source/ControlRig/Private/Rigs/RigHierarchyMetadata.cpp
@@ -0,0 +1,134 @@
+// Copyright Epic Games, Inc. All Rights Reserved.
+
+#include "Rigs/RigHierarchyMetadata.h"
+#include "Rigs/RigHierarchyElements.h"
+#include "Rigs/RigHierarchy.h"
+#include "ControlRigObjectVersion.h"
+
+////////////////////////////////////////////////////////////////////////////////
+// FRigBaseMetadata
+////////////////////////////////////////////////////////////////////////////////
+
+UScriptStruct* FRigBaseMetadata::GetMetadataStruct() const
+{
+	return GetMetadataStruct(GetType());
+}
+
+UScriptStruct* FRigBaseMetadata::GetMetadataStruct(const ERigMetadataType& InType)
+{
+	switch(InType)
+	{
+		case ERigMetadataType::Bool:
+		{
+			return FRigBoolMetadata::StaticStruct();
+		}
+		case ERigMetadataType::BoolArray:
+		{
+			return FRigBoolArrayMetadata::StaticStruct();
+		}
+		case ERigMetadataType::Float:
+		{
+			return FRigFloatMetadata::StaticStruct();
+		}
+		case ERigMetadataType::FloatArray:
+		{
+			return FRigFloatArrayMetadata::StaticStruct();
+		}
+		case ERigMetadataType::Int32:
+		{
+			return FRigInt32Metadata::StaticStruct();
+		}
+		case ERigMetadataType::Int32Array:
+		{
+			return FRigInt32ArrayMetadata::StaticStruct();
+		}
+		case ERigMetadataType::Name:
+		{
+			return FRigNameMetadata::StaticStruct();
+		}
+		case ERigMetadataType::NameArray:
+		{
+			return FRigNameArrayMetadata::StaticStruct();
+		}
+		case ERigMetadataType::Vector:
+		{
+			return FRigVectorMetadata::StaticStruct();
+		}
+		case ERigMetadataType::VectorArray:
+		{
+			return FRigVectorArrayMetadata::StaticStruct();
+		}
+		case ERigMetadataType::Rotator:
+		{
+			return FRigRotatorMetadata::StaticStruct();
+		}
+		case ERigMetadataType::RotatorArray:
+		{
+			return FRigRotatorArrayMetadata::StaticStruct();
+		}
+		case ERigMetadataType::Quat:
+		{
+			return FRigQuatMetadata::StaticStruct();
+		}
+		case ERigMetadataType::QuatArray:
+		{
+			return FRigQuatArrayMetadata::StaticStruct();
+		}
+		case ERigMetadataType::Transform:
+		{
+			return FRigTransformMetadata::StaticStruct();
+		}
+		case ERigMetadataType::TransformArray:
+		{
+			return FRigTransformArrayMetadata::StaticStruct();
+		}
+		case ERigMetadataType::RigElementKey:
+		{
+			return FRigElementKeyMetadata::StaticStruct();
+		}
+		case ERigMetadataType::RigElementKeyArray:
+		{
+			return FRigElementKeyArrayMetadata::StaticStruct();
+		}
+		default:
+		{
+			break;
+		}
+	}
+	return StaticStruct();
+}
+
+FRigBaseMetadata* FRigBaseMetadata::MakeMetadata(const FRigBaseElement* InElement, const FName& InName, ERigMetadataType InType)
+{
+	check(InType != ERigMetadataType::Invalid);
+	
+	UScriptStruct* Struct = GetMetadataStruct(InType);
+	check(Struct);
+	
+	FRigBaseMetadata* Md = (FRigBaseMetadata*)FMemory::Malloc(Struct->GetStructureSize());
+	Struct->InitializeStruct(Md, 1);
+
+	Md->Element = InElement;
+	Md->Name = InName;
+	Md->Type = InType;
+	return Md;
+}
+
+void FRigBaseMetadata::Serialize(FArchive& Ar, bool bIsLoading)
+{
+	Ar.UsingCustomVersion(FControlRigObjectVersion::GUID);
+	Ar << Name;
+	Ar << Type;
+}
+
+const FRigElementKey& FRigBaseMetadata::GetKey() const
+{
+	if(Element)
+	{
+		return Element->GetKey();
+	}
+
+	static const FRigElementKey EmptyKey;
+	return EmptyKey;
+}
+
diff --git a/Engine/Plugins/Animation/ControlRig/Source/ControlRig/Public/ControlRigObjectVersion.h b/Engine/Plugins/Animation/ControlRig/Source/ControlRig/Public/ControlRigObjectVersion.h
--- a/Engine/Plugins/Animation/ControlRig/Source/ControlRig/Public/ControlRigObjectVersion.h
+++ b/Engine/Plugins/Animation/ControlRig/Source/ControlRig/Public/ControlRigObjectVersion.h
@@ -69,6 +69,9 @@ struct CONTROLRIG_API FControlRigObjectVersion
 		// Added preferred euler angles to controls
 		PreferredEulerAnglesForControls,
 
+		// Added rig hierarchy element metadata
+		HierarchyElementMetadata,
+
 		// -----<new versions can be added above this line>-------------------------------------------------
 		VersionPlusOne,
 		LatestVersion = VersionPlusOne - 1
diff --git a/Engine/Plugins/Animation/ControlRig/Source/ControlRig/Public/Rigs/RigHierarchy.h b/Engine/Plugins/Animation/ControlRig/Source/ControlRig/Public/Rigs/RigHierarchy.h
--- a/Engine/Plugins/Animation/ControlRig/Source/ControlRig/Public/Rigs/RigHierarchy.h
+++ b/Engine/Plugins/Animation/ControlRig/Source/ControlRig/Public/Rigs/RigHierarchy.h
@@ -836,6 +836,477 @@ public:
 		return GetKeysOfType<FRigReferenceElement>(bTraverse);
 	}
 
+	/**
+	 * Returns the name of metadata for a given element
+	 * @param InItem The element key to return the metadata keys for
+	 */
+	UFUNCTION(BlueprintPure, Category = URigHierarchy)
+	TArray<FName> GetMetadataNames(FRigElementKey InItem);
+
+	/**
+	 * Returns the type of metadata given its name the item it is stored under
+	 * @param InItem The element key to return the metadata type for
+	 * @param InMetadataName The name of the metadata to return the type for
+	 */
+	UFUNCTION(BlueprintPure, Category = URigHierarchy)
+	ERigMetadataType GetMetadataType(FRigElementKey InItem, FName InMetadataName);
+
+	/**
+	 * Removes the metadata under a given element 
+	 * @param InItem The element key to search under
+	 * @param InMetadataName The name of the metadata to remove
+	 */
+	UFUNCTION(BlueprintCallable, Category = URigHierarchy)
+	bool RemoveMetadata(FRigElementKey InItem, FName InMetadataName);
+
+	/**
+     * Removes all of the metadata under a given item 
+	 * @param InItem The element key to search under
+	 */
+	UFUNCTION(BlueprintCallable, Category = URigHierarchy)
+	bool RemoveAllMetadata(FRigElementKey InItem);
+
+	/**
+	 * Queries and returns the value of bool metadata
+	 * @param InItem The element key to return the metadata for
+	 * @param InMetadataName The name of the metadata to query
+	 * @param DefaultValue The default value to fall back on
+	 */
+	UFUNCTION(BlueprintPure, Category = URigHierarchy)
+	FORCEINLINE bool GetBoolMetadata(FRigElementKey InItem, FName InMetadataName, bool DefaultValue) const
+	{
+		return GetMetadata<bool>(InItem, ERigMetadataType::Bool, InMetadataName, DefaultValue);
+	}
+
+	/**
+	 * Queries and returns the value of bool array metadata
+	 * @param InItem The element key to return the metadata for
+	 * @param InMetadataName The name of the metadata to query
+	 */
+	UFUNCTION(BlueprintPure, Category = URigHierarchy)
+	FORCEINLINE TArray<bool> GetBoolArrayMetadata(FRigElementKey InItem, FName InMetadataName) const
+	{
+		return GetArrayMetadata<bool>(InItem, ERigMetadataType::BoolArray, InMetadataName);
+	}
+
+	/**
+	 * Sets the metadata to a bool value
+	 * @param InItem The element key to set the metadata for
+	 * @param InMetadataName The name of the metadata to set
+	 * @param InValue The value to set
+	 * @return Returns true if setting the metadata was successful
+	 */
+	UFUNCTION(BlueprintCallable, Category = URigHierarchy)
+	FORCEINLINE bool SetBoolMetadata(FRigElementKey InItem, FName InMetadataName, bool InValue)
+	{
+		return SetMetadata<bool>(InItem, ERigMetadataType::Bool, InMetadataName, InValue);
+	}
+
+	/**
+	 * Sets the metadata to a bool array value
+	 * @param InItem The element key to set the metadata for
+	 * @param InMetadataName The name of the metadata to set
+	 * @param InValue The value to set
+	 * @return Returns true if setting the metadata was successful
+	 */
+	UFUNCTION(BlueprintCallable, Category = URigHierarchy)
+	FORCEINLINE bool SetBoolArrayMetadata(FRigElementKey InItem, FName InMetadataName, TArray<bool> InValue)
+	{
+		return SetArrayMetadata<bool>(InItem, ERigMetadataType::BoolArray, InMetadataName, InValue);
+	}
+
+	/**
+	 * Queries and returns the value of float metadata
+	 * @param InItem The element key to return the metadata for
+	 * @param InMetadataName The name of the metadata to query
+	 * @param DefaultValue The default value to fall back on
+	 */
+	UFUNCTION(BlueprintPure, Category = URigHierarchy)
+	FORCEINLINE float GetFloatMetadata(FRigElementKey InItem, FName InMetadataName, float DefaultValue) const
+	{
+		return GetMetadata<float>(InItem, ERigMetadataType::Float, InMetadataName, DefaultValue);
+	}
+
+	/**
+	 * Queries and returns the value of float array metadata
+	 * @param InItem The element key to return the metadata for
+	 * @param InMetadataName The name of the metadata to query
+	 */
+	UFUNCTION(BlueprintPure, Category = URigHierarchy)
+	FORCEINLINE TArray<float> GetFloatArrayMetadata(FRigElementKey InItem, FName InMetadataName) const
+	{
+		return GetArrayMetadata<float>(InItem, ERigMetadataType::FloatArray, InMetadataName);
+	}
+
+	/**
+	 * Sets the metadata to a float value
+	 * @param InItem The element key to set the metadata for
+	 * @param InMetadataName The name of the metadata to set
+	 * @param InValue The value to set
+	 * @return Returns true if setting the metadata was successful
+	 */
+	UFUNCTION(BlueprintCallable, Category = URigHierarchy)
+	FORCEINLINE bool SetFloatMetadata(FRigElementKey InItem, FName InMetadataName, float InValue)
+	{
+		return SetMetadata<float>(InItem, ERigMetadataType::Float, InMetadataName, InValue);
+	}
+
+	/**
+	 * Sets the metadata to a float array value
+	 * @param InItem The element key to set the metadata for
+	 * @param InMetadataName The name of the metadata to set
+	 * @param InValue The value to set
+	 * @return Returns true if setting the metadata was successful
+	 */
+	UFUNCTION(BlueprintCallable, Category = URigHierarchy)
+	FORCEINLINE bool SetFloatArrayMetadata(FRigElementKey InItem, FName InMetadataName, TArray<float> InValue)
+	{
+		return SetArrayMetadata<float>(InItem, ERigMetadataType::FloatArray, InMetadataName, InValue);
+	}
+
+	/**
+	 * Queries and returns the value of int32 metadata
+	 * @param InItem The element key to return the metadata for
+	 * @param InMetadataName The name of the metadata to query
+	 * @param DefaultValue The default value to fall back on
+	 */
+	UFUNCTION(BlueprintPure, Category = URigHierarchy)
+	FORCEINLINE int32 GetInt32Metadata(FRigElementKey InItem, FName InMetadataName, int32 DefaultValue) const
+	{
+		return GetMetadata<int32>(InItem, ERigMetadataType::Int32, InMetadataName, DefaultValue);
+	}
+
+	/**
+	 * Queries and returns the value of int32 array metadata
+	 * @param InItem The element key to return the metadata for
+	 * @param InMetadataName The name of the metadata to query
+	 */
+	UFUNCTION(BlueprintPure, Category = URigHierarchy)
+	FORCEINLINE TArray<int32> GetInt32ArrayMetadata(FRigElementKey InItem, FName InMetadataName) const
+	{
+		return GetArrayMetadata<int32>(InItem, ERigMetadataType::Int32Array, InMetadataName);
+	}
+
+	/**
+	 * Sets the metadata to a int32 value
+	 * @param InItem The element key to set the metadata for
+	 * @param InMetadataName The name of the metadata to set
+	 * @param InValue The value to set
+	 * @return Returns true if setting the metadata was successful
+	 */
+	UFUNCTION(BlueprintCallable, Category = URigHierarchy)
+	FORCEINLINE bool SetInt32Metadata(FRigElementKey InItem, FName InMetadataName, int32 InValue)
+	{
+		return SetMetadata<int32>(InItem, ERigMetadataType::Int32, InMetadataName, InValue);
+	}
+
+	/**
+	 * Sets the metadata to a int32 array value
+	 * @param InItem The element key to set the metadata for
+	 * @param InMetadataName The name of the metadata to set
+	 * @param InValue The value to set
+	 * @return Returns true if setting the metadata was successful
+	 */
+	UFUNCTION(BlueprintCallable, Category = URigHierarchy)
+	FORCEINLINE bool SetInt32ArrayMetadata(FRigElementKey InItem, FName InMetadataName, TArray<int32> InValue)
+	{
+		return SetArrayMetadata<int32>(InItem, ERigMetadataType::Int32Array, InMetadataName, InValue);
+	}
+
+	/**
+	 * Queries and returns the value of FName metadata
+	 * @param InItem The element key to return the metadata for
+	 * @param InMetadataName The name of the metadata to query
+	 * @param DefaultValue The default value to fall back on
+	 */
+	UFUNCTION(BlueprintPure, Category = URigHierarchy)
+	FORCEINLINE FName GetNameMetadata(FRigElementKey InItem, FName InMetadataName, FName DefaultValue) const
+	{
+		return GetMetadata<FName>(InItem, ERigMetadataType::Name, InMetadataName, DefaultValue);
+	}
+
+	/**
+	 * Queries and returns the value of FName array metadata
+	 * @param InItem The element key to return the metadata for
+	 * @param InMetadataName The name of the metadata to query
+	 */
+	UFUNCTION(BlueprintPure, Category = URigHierarchy)
+	FORCEINLINE TArray<FName> GetNameArrayMetadata(FRigElementKey InItem, FName InMetadataName) const
+	{
+		return GetArrayMetadata<FName>(InItem, ERigMetadataType::NameArray, InMetadataName);
+	}
+
+	/**
+	 * Sets the metadata to a FName value
+	 * @param InItem The element key to set the metadata for
+	 * @param InMetadataName The name of the metadata to set
+	 * @param InValue The value to set
+	 * @return Returns true if setting the metadata was successful
+	 */
+	UFUNCTION(BlueprintCallable, Category = URigHierarchy)
+	FORCEINLINE bool SetNameMetadata(FRigElementKey InItem, FName InMetadataName, FName InValue)
+	{
+		return SetMetadata<FName>(InItem, ERigMetadataType::Name, InMetadataName, InValue);
+	}
+
+	/**
+	 * Sets the metadata to a FName array value
+	 * @param InItem The element key to set the metadata for
+	 * @param InMetadataName The name of the metadata to set
+	 * @param InValue The value to set
+	 * @return Returns true if setting the metadata was successful
+	 */
+	UFUNCTION(BlueprintCallable, Category = URigHierarchy)
+	FORCEINLINE bool SetNameArrayMetadata(FRigElementKey InItem, FName InMetadataName, TArray<FName> InValue)
+	{
+		return SetArrayMetadata<FName>(InItem, ERigMetadataType::FloatArray, InMetadataName, InValue);
+	}
+
+	/**
+	 * Queries and returns the value of FVector metadata
+	 * @param InItem The element key to return the metadata for
+	 * @param InMetadataName The name of the metadata to query
+	 * @param DefaultValue The default value to fall back on
+	 */
+	UFUNCTION(BlueprintPure, Category = URigHierarchy)
+	FORCEINLINE FVector GetVectorMetadata(FRigElementKey InItem, FName InMetadataName, FVector DefaultValue) const
+	{
+		return GetMetadata<FVector>(InItem, ERigMetadataType::Vector, InMetadataName, DefaultValue);
+	}
+
+	/**
+	 * Queries and returns the value of FVector array metadata
+	 * @param InItem The element key to return the metadata for
+	 * @param InMetadataName The name of the metadata to query
+	 */
+	UFUNCTION(BlueprintPure, Category = URigHierarchy)
+	FORCEINLINE TArray<FVector> GetVectorArrayMetadata(FRigElementKey InItem, FName InMetadataName) const
+	{
+		return GetArrayMetadata<FVector>(InItem, ERigMetadataType::VectorArray, InMetadataName);
+	}
+
+	/**
+	 * Sets the metadata to a FVector value
+	 * @param InItem The element key to set the metadata for
+	 * @param InMetadataName The name of the metadata to set
+	 * @param InValue The value to set
+	 * @return Returns true if setting the metadata was successful
+	 */
+	UFUNCTION(BlueprintCallable, Category = URigHierarchy)
+	FORCEINLINE bool SetVectorMetadata(FRigElementKey InItem, FName InMetadataName, FVector InValue)
+	{
+		return SetMetadata<FVector>(InItem, ERigMetadataType::Vector, InMetadataName, InValue);
+	}
+
+	/**
+	 * Sets the metadata to a FVector array value
+	 * @param InItem The element key to set the metadata for
+	 * @param InMetadataName The name of the metadata to set
+	 * @param InValue The value to set
+	 * @return Returns true if setting the metadata was successful
+	 */
+	UFUNCTION(BlueprintCallable, Category = URigHierarchy)
+	FORCEINLINE bool SetVectorArrayMetadata(FRigElementKey InItem, FName InMetadataName, TArray<FVector> InValue)
+	{
+		return SetArrayMetadata<FVector>(InItem, ERigMetadataType::VectorArray, InMetadataName, InValue);
+	}
+
+	/**
+	 * Queries and returns the value of FRotator metadata
+	 * @param InItem The element key to return the metadata for
+	 * @param InMetadataName The name of the metadata to query
+	 * @param DefaultValue The default value to fall back on
+	 */
+	UFUNCTION(BlueprintPure, Category = URigHierarchy)
+	FORCEINLINE FRotator GetRotatorMetadata(FRigElementKey InItem, FName InMetadataName, FRotator DefaultValue) const
+	{
+		return GetMetadata<FRotator>(InItem, ERigMetadataType::Rotator, InMetadataName, DefaultValue);
+	}
+
+	/**
+	 * Queries and returns the value of FRotator array metadata
+	 * @param InItem The element key to return the metadata for
+	 * @param InMetadataName The name of the metadata to query
+	 */
+	UFUNCTION(BlueprintPure, Category = URigHierarchy)
+	FORCEINLINE TArray<FRotator> GetRotatorArrayMetadata(FRigElementKey InItem, FName InMetadataName) const
+	{
+		return GetArrayMetadata<FRotator>(InItem, ERigMetadataType::RotatorArray, InMetadataName);
+	}
+
+	/**
+	 * Sets the metadata to a FRotator value
+	 * @param InItem The element key to set the metadata for
+	 * @param InMetadataName The name of the metadata to set
+	 * @param InValue The value to set
+	 * @return Returns true if setting the metadata was successful
+	 */
+	UFUNCTION(BlueprintCallable, Category = URigHierarchy)
+	FORCEINLINE bool SetRotatorMetadata(FRigElementKey InItem, FName InMetadataName, FRotator InValue)
+	{
+		return SetMetadata<FRotator>(InItem, ERigMetadataType::Rotator, InMetadataName, InValue);
+	}
+
+	/**
+	 * Sets the metadata to a FRotator array value
+	 * @param InItem The element key to set the metadata for
+	 * @param InMetadataName The name of the metadata to set
+	 * @param InValue The value to set
+	 * @return Returns true if setting the metadata was successful
+	 */
+	UFUNCTION(BlueprintCallable, Category = URigHierarchy)
+	FORCEINLINE bool SetRotatorArrayMetadata(FRigElementKey InItem, FName InMetadataName, TArray<FRotator> InValue)
+	{
+		return SetArrayMetadata<FRotator>(InItem, ERigMetadataType::RotatorArray, InMetadataName, InValue);
+	}
+
+	/**
+	 * Queries and returns the value of FQuat metadata
+	 * @param InItem The element key to return the metadata for
+	 * @param InMetadataName The name of the metadata to query
+	 * @param DefaultValue The default value to fall back on
+	 */
+	UFUNCTION(BlueprintPure, Category = URigHierarchy)
+	FORCEINLINE FQuat GetQuatMetadata(FRigElementKey InItem, FName InMetadataName, FQuat DefaultValue) const
+	{
+		return GetMetadata<FQuat>(InItem, ERigMetadataType::Quat, InMetadataName, DefaultValue);
+	}
+
+	/**
+	 * Queries and returns the value of FQuat array metadata
+	 * @param InItem The element key to return the metadata for
+	 * @param InMetadataName The name of the metadata to query
+	 */
+	UFUNCTION(BlueprintPure, Category = URigHierarchy)
+	FORCEINLINE TArray<FQuat> GetQuatArrayMetadata(FRigElementKey InItem, FName InMetadataName) const
+	{
+		return GetArrayMetadata<FQuat>(InItem, ERigMetadataType::QuatArray, InMetadataName);
+	}
+
+	/**
+	 * Sets the metadata to a FQuat value
+	 * @param InItem The element key to set the metadata for
+	 * @param InMetadataName The name of the metadata to set
+	 * @param InValue The value to set
+	 * @return Returns true if setting the metadata was successful
+	 */
+	UFUNCTION(BlueprintCallable, Category = URigHierarchy)
+	FORCEINLINE bool SetQuatMetadata(FRigElementKey InItem, FName InMetadataName, FQuat InValue)
+	{
+		return SetMetadata<FQuat>(InItem, ERigMetadataType::Quat, InMetadataName, InValue);
+	}
+
+	/**
+	 * Sets the metadata to a FQuat array value
+	 * @param InItem The element key to set the metadata for
+	 * @param InMetadataName The name of the metadata to set
+	 * @param InValue The value to set
+	 * @return Returns true if setting the metadata was successful
+	 */
+	UFUNCTION(BlueprintCallable, Category = URigHierarchy)
+	FORCEINLINE bool SetQuatArrayMetadata(FRigElementKey InItem, FName InMetadataName, TArray<FQuat> InValue)
+	{
+		return SetArrayMetadata<FQuat>(InItem, ERigMetadataType::VectorArray, InMetadataName, InValue);
+	}
+
+	/**
+	 * Queries and returns the value of FTransform metadata
+	 * @param InItem The element key to return the metadata for
+	 * @param InMetadataName The name of the metadata to query
+	 * @param DefaultValue The default value to fall back on
+	 */
+	UFUNCTION(BlueprintPure, Category = URigHierarchy)
+	FORCEINLINE FTransform GetTransformMetadata(FRigElementKey InItem, FName InMetadataName, FTransform DefaultValue) const
+	{
+		return GetMetadata<FTransform>(InItem, ERigMetadataType::Transform, InMetadataName, DefaultValue);
+	}
+
+	/**
+	 * Queries and returns the value of FTransform array metadata
+	 * @param InItem The element key to return the metadata for
+	 * @param InMetadataName The name of the metadata to query
+	 */
+	UFUNCTION(BlueprintPure, Category = URigHierarchy)
+	FORCEINLINE TArray<FTransform> GetTransformArrayMetadata(FRigElementKey InItem, FName InMetadataName) const
+	{
+		return GetArrayMetadata<FTransform>(InItem, ERigMetadataType::TransformArray, InMetadataName);
+	}
+
+	/**
+	 * Sets the metadata to a FTransform value
+	 * @param InItem The element key to set the metadata for
+	 * @param InMetadataName The name of the metadata to set
+	 * @param InValue The value to set
+	 * @return Returns true if setting the metadata was successful
+	 */
+	UFUNCTION(BlueprintCallable, Category = URigHierarchy)
+	FORCEINLINE bool SetTransformMetadata(FRigElementKey InItem, FName InMetadataName, FTransform InValue)
+	{
+		return SetMetadata<FTransform>(InItem, ERigMetadataType::Transform, InMetadataName, InValue);
+	}
+
+	/**
+	 * Sets the metadata to a FTransform array value
+	 * @param InItem The element key to set the metadata for
+	 * @param InMetadataName The name of the metadata to set
+	 * @param InValue The value to set
+	 * @return Returns true if setting the metadata was successful
+	 */
+	UFUNCTION(BlueprintCallable, Category = URigHierarchy)
+	FORCEINLINE bool SetTransformArrayMetadata(FRigElementKey InItem, FName InMetadataName, TArray<FTransform> InValue)
+	{
+		return SetArrayMetadata<FTransform>(InItem, ERigMetadataType::VectorArray, InMetadataName, InValue);
+	}
+
+	/**
+	 * Queries and returns the value of FRigElementKey metadata
+	 * @param InItem The element key to return the metadata for
+	 * @param InMetadataName The name of the metadata to query
+	 * @param DefaultValue The default value to fall back on
+	 */
+	UFUNCTION(BlueprintPure, Category = URigHierarchy)
+	FORCEINLINE FRigElementKey GetRigElementKeyMetadata(FRigElementKey InItem, FName InMetadataName, FRigElementKey DefaultValue) const
+	{
+		return GetMetadata<FRigElementKey>(InItem, ERigMetadataType::RigElementKey, InMetadataName, DefaultValue);
+	}
+
+	/**
+	 * Queries and returns the value of FRigElementKey array metadata
+	 * @param InItem The element key to return the metadata for
+	 * @param InMetadataName The name of the metadata to query
+	 */
+	UFUNCTION(BlueprintPure, Category = URigHierarchy)
+	FORCEINLINE TArray<FRigElementKey> GetRigElementKeyArrayMetadata(FRigElementKey InItem, FName InMetadataName) const
+	{
+		return GetArrayMetadata<FRigElementKey>(InItem, ERigMetadataType::RigElementKeyArray, InMetadataName);
+	}
+
+	/**
+	 * Sets the metadata to a FRigElementKey value
+	 * @param InItem The element key to set the metadata for
+	 * @param InMetadataName The name of the metadata to set
+	 * @param InValue The value to set
+	 * @return Returns true if setting the metadata was successful
+	 */
+	UFUNCTION(BlueprintCallable, Category = URigHierarchy)
+	FORCEINLINE bool SetRigElementKeyMetadata(FRigElementKey InItem, FName InMetadataName, FRigElementKey InValue)
+	{
+		return SetMetadata<FRigElementKey>(InItem, ERigMetadataType::RigElementKey, InMetadataName, InValue);
+	}
+
+	/**
+	 * Sets the metadata to a FRigElementKey array value
+	 * @param InItem The element key to set the metadata for
+	 * @param InMetadataName The name of the metadata to set
+	 * @param InValue The value to set
+	 * @return Returns true if setting the metadata was successful
+	 */
+	UFUNCTION(BlueprintCallable, Category = URigHierarchy)
+	FORCEINLINE bool SetRigElementKeyArrayMetadata(FRigElementKey InItem, FName InMetadataName, TArray<FRigElementKey> InValue)
+	{
+		return SetArrayMetadata<FRigElementKey>(InItem, ERigMetadataType::VectorArray, InMetadataName, InValue);
+	}
+
 	/**
 	 * Returns the selected elements
 	 * @InTypeFilter The types to retrieve the selection for
@@ -2887,6 +3358,7 @@ private:
 		{
 			new(&Elements[Index]) ElementType();
 		}
+		Elements[0].OwnedInstances = Num;
 		return Elements;
 	}
 
@@ -3285,7 +3757,67 @@ private:
 #if WITH_EDITOR
 	static TArray<FString> ControlSettingsToPythonCommands(const FRigControlSettings& Settings, const FString& NameSettings);
 #endif
-	
+
+	template<typename T>
+	FORCEINLINE const T& GetMetadata(const FRigElementKey& InItem, ERigMetadataType InType, const FName& InMetadataName, const T& DefaultValue) const
+	{
+		return GetMetadata<T>(Find(InItem), InType, InMetadataName, DefaultValue);
+	}
+
+	template<typename T>
+	FORCEINLINE const T& GetMetadata(const FRigBaseElement* InElement, ERigMetadataType InType, const FName& InMetadataName, const T& DefaultValue) const
+	{
+		if(InElement)
+		{
+			if(FRigBaseMetadata* Metadata = InElement->GetMetadata(InMetadataName, InType))
+			{
+				return *(const T*)Metadata->GetValueData();
+			}
+		}
+		return DefaultValue;
+	}
+
+	template<typename T>
+	FORCEINLINE const TArray<T>& GetArrayMetadata(const FRigElementKey& InItem, ERigMetadataType InType, const FName& InMetadataName) const
+	{
+		return GetArrayMetadata<T>(Find(InItem), InType, InMetadataName);
+	}
+
+	template<typename T>
+	FORCEINLINE const TArray<T>& GetArrayMetadata(const FRigBaseElement* InElement, ERigMetadataType InType, const FName& InMetadataName) const
+	{
+		static const TArray<T> EmptyArray;
+		return GetMetadata<TArray<T>>(InElement, InType, InMetadataName, EmptyArray);
+	}
+
+	template<typename T>
+	FORCEINLINE bool SetMetadata(const FRigElementKey& InItem, ERigMetadataType InType, const FName& InMetadataName, const T& InValue)
+	{
+		return SetMetadata<T>(Find(InItem), InType, InMetadataName, InValue);
+	}
+
+	template<typename T>
+	FORCEINLINE bool SetMetadata(FRigBaseElement* InElement, ERigMetadataType InType, const FName& InMetadataName, const T& InValue)
+	{
+		if(InElement)
+		{
+			return InElement->SetMetaData(InMetadataName, InType, &InValue, sizeof(T));
+		}
+		return false;
+	}
+
+	template<typename T>
+	FORCEINLINE bool SetArrayMetadata(const FRigElementKey& InItem, ERigMetadataType InType, const FName& InMetadataName, const TArray<T>& InValue)
+	{
+		return SetMetadata<TArray<T>>(Find(InItem), InType, InMetadataName, InValue);
+	}
+
+	template<typename T>
+	FORCEINLINE bool SetArrayMetadata(FRigBaseElement* InElement, ERigMetadataType InType, const FName& InMetadataName, const TArray<T>& InValue)
+	{
+		return SetMetadata<TArray<T>>(InElement, InType, InMetadataName, InValue);
+	}
+
 protected:
 	
 	bool bEnableCacheValidityCheck;
diff --git a/Engine/Plugins/Animation/ControlRig/Source/ControlRig/Public/Rigs/RigHierarchyDefines.h b/Engine/Plugins/Animation/ControlRig/Source/ControlRig/Public/Rigs/RigHierarchyDefines.h
--- a/Engine/Plugins/Animation/ControlRig/Source/ControlRig/Public/Rigs/RigHierarchyDefines.h
+++ b/Engine/Plugins/Animation/ControlRig/Source/ControlRig/Public/Rigs/RigHierarchyDefines.h
@@ -40,6 +40,35 @@ enum class ERigBoneType : uint8
 	User
 };
 
+/* 
+ * The type of meta data stored on an element
+ */
+UENUM(BlueprintType)
+enum class ERigMetadataType : uint8
+{
+	Bool,
+	BoolArray,
+	Float,
+	FloatArray,
+	Int32,
+	Int32Array,
+	Name,
+	NameArray,
+	Vector,
+	VectorArray,
+	Rotator,
+	RotatorArray,
+	Quat,
+	QuatArray,
+	Transform,
+	TransformArray,
+	RigElementKey,
+	RigElementKeyArray,
+
+	/** MAX - invalid */
+	Invalid UMETA(Hidden),
+};
+
 UENUM()
 enum class ERigHierarchyNotification : uint8
 {
diff --git a/Engine/Plugins/Animation/ControlRig/Source/ControlRig/Public/Rigs/RigHierarchyElements.h b/Engine/Plugins/Animation/ControlRig/Source/ControlRig/Public/Rigs/RigHierarchyElements.h
--- a/Engine/Plugins/Animation/ControlRig/Source/ControlRig/Public/Rigs/RigHierarchyElements.h
+++ b/Engine/Plugins/Animation/ControlRig/Source/ControlRig/Public/Rigs/RigHierarchyElements.h
@@ -4,6 +4,7 @@
 
 #include "CoreMinimal.h"
 #include "RigHierarchyDefines.h"
+#include "RigHierarchyMetadata.h"
 #include "RigHierarchyElements.generated.h"
 
 struct FRigUnitContext;
@@ -515,7 +516,7 @@ public:
 	, OwnedInstances(0)
 	{}
 
-	virtual ~FRigBaseElement(){}
+	virtual ~FRigBaseElement();
 
 	enum ESerializationPhase
 	{
@@ -543,6 +544,9 @@ protected:
 	UPROPERTY(BlueprintReadOnly, Transient, Category = RigElement, meta = (AllowPrivateAccess = "true"))
 	int32 CreatedAtInstructionIndex;
 
+	TArray<FRigBaseMetadata*> Metadata;
+	TMap<FName,int32> MetadataNameToIndex;
+
 	FORCEINLINE static bool IsClassOf(const FRigBaseElement* InElement)
 	{
 		return true;
@@ -566,6 +570,39 @@ public:
 	FORCEINLINE int32 GetCreatedAtInstructionIndex() const { return CreatedAtInstructionIndex; }
 	FORCEINLINE bool IsProcedural() const { return CreatedAtInstructionIndex != INDEX_NONE; }
 
+	FORCEINLINE int32 NumMetadata() const { return Metadata.Num(); }
+	FORCEINLINE FRigBaseMetadata* GetMetadata(int32 InIndex) const { return Metadata[InIndex]; }
+	FORCEINLINE FRigBaseMetadata* GetMetadata(const FName& InName) const
+	{
+		if(const int32* MetadataIndex = MetadataNameToIndex.Find(InName))
+		{
+			return GetMetadata(*MetadataIndex);
+		}
+		return nullptr;
+	}
+	FORCEINLINE FRigBaseMetadata* GetMetadata(const FName& InName, ERigMetadataType InType) const
+	{
+		if(const int32* MetadataIndex = MetadataNameToIndex.Find(InName))
+		{
+			FRigBaseMetadata* Md = GetMetadata(*MetadataIndex);
+			if(Md->GetType() == InType)
+			{
+				return Md;
+			}
+		}
+		return nullptr;
+	}
+	FORCEINLINE bool SetMetaData(const FName& InName, ERigMetadataType InType, const void* InData, int32 InSize)
+	{
+		if(FRigBaseMetadata* Md = SetupValidMetadata(InName, InType))
+		{
+			return Md->SetValueData(InData, InSize);
+		}
+		return false;
+	}
+	bool RemoveMetadata(const FName& InName);
+	bool RemoveAllMetadata();
+
 	template<typename T>
 	FORCEINLINE bool IsA() const { return T::IsClassOf(this); }
 
@@ -621,7 +658,10 @@ public:
 protected:
 
 	// helper function to be called as part of URigHierarchy::CopyHierarchy
-	virtual void  CopyFrom(URigHierarchy* InHierarchy, FRigBaseElement* InOther, URigHierarchy* InOtherHierarchy) {}
+	virtual void CopyFrom(URigHierarchy* InHierarchy, FRigBaseElement* InOther, URigHierarchy* InOtherHierarchy);
+	
+	// sets up the metadata and ensures the right type
+	FRigBaseMetadata* SetupValidMetadata(const FName& InName, ERigMetadataType InType);
 
 	mutable uint16 TopologyVersion;
 	mutable FRigBaseElementChildrenArray CachedChildren;
@@ -631,6 +671,7 @@ protected:
 
 	friend class URigHierarchy;
 	friend class URigHierarchyController;
+	friend struct FRigUnit_SetMetadataBase;
 };
 
 USTRUCT(BlueprintType)
diff --git a/Engine/Plugins/Animation/ControlRig/Source/ControlRig/Public/Rigs/RigHierarchyMetadata.h b/Engine/Plugins/Animation/ControlRig/Source/ControlRig/Public/Rigs/RigHierarchyMetadata.h
new file mode 100644
--- /dev/null
+++ b/Engine/Plugins/Animation/ControlRig/Source/ControlRig/Public/Rigs/RigHierarchyMetadata.h
@@ -0,0 +1,764 @@
+// Copyright Epic Games, Inc. All Rights Reserved.
+
+#pragma once
+
+#include "CoreMinimal.h"
+#include "RigHierarchyDefines.h"
+#include "RigHierarchyMetadata.generated.h"
+
+struct FRigBaseElement;
+class URigHierarchyController;
+
+#define DECLARE_RIG_METADATA_METHODS(MetadataType) \
+template<typename T> \
+friend FORCEINLINE const T* Cast(const MetadataType* InMetadata) \
+{ \
+   return Cast<T>((const FRigBaseMetadata*) InMetadata); \
+} \
+template<typename T> \
+friend FORCEINLINE T* Cast(MetadataType* InMetadata) \
+{ \
+   return Cast<T>((FRigBaseMetadata*) InMetadata); \
+} \
+template<typename T> \
+friend FORCEINLINE const T* CastChecked(const MetadataType* InMetadata) \
+{ \
+	return CastChecked<T>((const FRigBaseMetadata*) InMetadata); \
+} \
+template<typename T> \
+friend FORCEINLINE T* CastChecked(MetadataType* InMetadata) \
+{ \
+	return CastChecked<T>((FRigBaseMetadata*) InMetadata); \
+}
+
+USTRUCT()
+struct CONTROLRIG_API FRigBaseMetadata
+{
+	GENERATED_BODY()
+
+public:
+
+	FRigBaseMetadata()
+    	: Element(nullptr)
+		, Name(NAME_None)
+		, Type(ERigMetadataType::Invalid)
+		, ValueProperty(nullptr)
+	{}
+
+	virtual ~FRigBaseMetadata(){}
+
+protected:
+
+	const FRigBaseElement* Element;
+
+	UPROPERTY()
+	FName Name;
+
+	UPROPERTY()
+	ERigMetadataType Type;
+
+	mutable const FProperty* ValueProperty;
+
+	FORCEINLINE static bool IsClassOf(const FRigBaseMetadata* InMetadata)
+	{
+		return true;
+	}
+
+public:
+
+	UScriptStruct* GetMetadataStruct() const;
+	virtual void Serialize(FArchive& Ar, bool bIsLoading);
+
+	FORCEINLINE bool IsValid() const { return (GetType() != ERigMetadataType::Invalid) && GetKey().IsValid(); }
+	FORCEINLINE const FRigBaseElement* GetElement() const { return Element; }
+	const FRigElementKey& GetKey() const;
+	FORCEINLINE const FName& GetName() const { return Name; }
+	FORCEINLINE const ERigMetadataType& GetType() const { return Type; }
+	FORCEINLINE bool IsArray() const
+	{
+		return GetValueProperty()->IsA<FArrayProperty>();
+	}
+	FORCEINLINE const void* GetValueData() const
+	{
+		return GetValueProperty()->ContainerPtrToValuePtr<void>(this);
+	}
+	FORCEINLINE const int32& GetValueSize() const
+	{
+		return GetValueProperty()->ElementSize;
+	}
+	FORCEINLINE bool SetValueData(const void* InData, int32 InSize)
+	{
+		if(InData)
+		{
+			if(InSize == GetValueSize())
+			{
+				return SetValueDataImpl(InData);
+			}
+		}
+		return false;
+	}
+
+	template<typename T>
+	FORCEINLINE bool IsA() const { return T::IsClassOf(this); }
+
+	template<typename T>
+    friend FORCEINLINE const T* Cast(const FRigBaseMetadata* InMetadata)
+	{
+		if(InMetadata)
+		{
+			if(InMetadata->IsA<T>())
+			{
+				return static_cast<const T*>(InMetadata);
+			}
+		}
+		return nullptr;
+	}
+
+	template<typename T>
+    friend FORCEINLINE T* Cast(FRigBaseMetadata* InMetadata)
+	{
+		if(InMetadata)
+		{
+			if(InMetadata->IsA<T>())
+			{
+				return static_cast<T*>(InMetadata);
+			}
+		}
+		return nullptr;
+	}
+
+	template<typename T>
+    friend FORCEINLINE const T* CastChecked(const FRigBaseMetadata* InMetadata)
+	{
+		const T* CastElement = Cast<T>(InMetadata);
+		check(CastElement);
+		return CastElement;
+	}
+
+	template<typename T>
+    friend FORCEINLINE T* CastChecked(FRigBaseMetadata* InMetadata)
+	{
+		T* CastElement = Cast<T>(InMetadata);
+		check(CastElement);
+		return CastElement;
+	}
+
+protected:
+
+	static UScriptStruct* GetMetadataStruct(const ERigMetadataType& InType);
+	static FRigBaseMetadata* MakeMetadata(const FRigBaseElement* InElement, const FName& InName, ERigMetadataType InType);
+
+	FORCEINLINE const FProperty* GetValueProperty() const
+	{
+		if(ValueProperty == nullptr)
+		{
+			static constexpr TCHAR ValueName[] = TEXT("Value");
+			ValueProperty = GetMetadataStruct()->FindPropertyByName(ValueName);
+		}
+		return ValueProperty;
+	}
+
+	FORCEINLINE void* GetValueData()
+	{
+		return GetValueProperty()->ContainerPtrToValuePtr<void>(this);
+	}
+
+	FORCEINLINE virtual bool SetValueDataImpl(const void* InData)
+	{
+		const FProperty* Property = GetValueProperty();
+		if(!Property->Identical(GetValueData(), InData))
+		{
+			Property->CopyCompleteValue(GetValueData(), InData);
+			return true;
+		}
+		return false;
+	}
+
+	friend struct FRigBaseElement;
+	friend class URigHierarchy;
+	friend class URigHierarchyController;
+};
+
+USTRUCT()
+struct CONTROLRIG_API FRigBoolMetadata : public FRigBaseMetadata
+{
+public:
+	
+	GENERATED_BODY()
+	DECLARE_RIG_METADATA_METHODS(FRigBoolMetadata)
+
+	FRigBoolMetadata()
+		: FRigBaseMetadata()
+		, Value(false)
+	{}
+	virtual ~FRigBoolMetadata() override {}
+	FORCEINLINE virtual void Serialize(FArchive& Ar, bool bIsLoading) override
+	{
+		Super::Serialize(Ar, bIsLoading);
+		Ar << Value;
+	}
+	FORCEINLINE const bool& GetValue() const { return Value; }
+	FORCEINLINE bool& GetValue() { return Value; }
+	FORCEINLINE bool SetValue(const bool& InValue) { return SetValueData(&InValue, sizeof(bool)); }
+
+	FORCEINLINE static bool IsClassOf(const FRigBaseMetadata* InMetadata)
+	{
+		return InMetadata->GetType() == ERigMetadataType::Bool;
+	}
+
+protected:
+
+	UPROPERTY()
+	bool Value;
+};
+
+USTRUCT()
+struct CONTROLRIG_API FRigBoolArrayMetadata : public FRigBaseMetadata
+{
+public:
+	
+	GENERATED_BODY()
+	DECLARE_RIG_METADATA_METHODS(FRigBoolArrayMetadata)
+
+	FRigBoolArrayMetadata()
+		: FRigBaseMetadata()
+	{}
+	virtual ~FRigBoolArrayMetadata() override {}
+	FORCEINLINE virtual void Serialize(FArchive& Ar, bool bIsLoading) override
+	{
+		Super::Serialize(Ar, bIsLoading);
+		Ar << Value;
+	}
+	FORCEINLINE const TArray<bool>& GetValue() const { return Value; }
+	FORCEINLINE TArray<bool>& GetValue() { return Value; }
+	FORCEINLINE bool SetValue(const TArray<bool>& InValue) { return SetValueData(&InValue, sizeof(TArray<bool>)); }
+
+	FORCEINLINE static bool IsClassOf(const FRigBaseMetadata* InMetadata)
+	{
+		return InMetadata->GetType() == ERigMetadataType::BoolArray;
+	}
+
+protected:
+
+	UPROPERTY()
+	TArray<bool> Value;
+};
+
+USTRUCT()
+struct CONTROLRIG_API FRigFloatMetadata : public FRigBaseMetadata
+{
+public:
+	
+	GENERATED_BODY()
+	DECLARE_RIG_METADATA_METHODS(FRigFloatMetadata)
+
+	FRigFloatMetadata()
+		: FRigBaseMetadata()
+		, Value(0.f)
+	{}
+	virtual ~FRigFloatMetadata() override {}
+	FORCEINLINE virtual void Serialize(FArchive& Ar, bool bIsLoading) override
+	{
+		Super::Serialize(Ar, bIsLoading);
+		Ar << Value;
+	}
+	FORCEINLINE const float& GetValue() const { return Value; } 
+	FORCEINLINE float& GetValue() { return Value; }
+	FORCEINLINE bool SetValue(const float& InValue) { return SetValueData(&InValue, sizeof(float)); }
+
+	FORCEINLINE static bool IsClassOf(const FRigBaseMetadata* InMetadata)
+    {
+    	return InMetadata->GetType() == ERigMetadataType::Float;
+    }
+
+protected:
+
+	UPROPERTY()
+	float Value;
+};
+
+USTRUCT()
+struct CONTROLRIG_API FRigFloatArrayMetadata : public FRigBaseMetadata
+{
+public:
+	
+	GENERATED_BODY()
+	DECLARE_RIG_METADATA_METHODS(FRigFloatArrayMetadata)
+
+	FRigFloatArrayMetadata()
+		: FRigBaseMetadata()
+	{}
+	virtual ~FRigFloatArrayMetadata() override {}
+	FORCEINLINE virtual void Serialize(FArchive& Ar, bool bIsLoading) override
+	{
+		Super::Serialize(Ar, bIsLoading);
+		Ar << Value;
+	}
+	FORCEINLINE const TArray<float>& GetValue() const { return Value; }
+	FORCEINLINE TArray<float>& GetValue() { return Value; }
+	FORCEINLINE bool SetValue(const TArray<float>& InValue) { return SetValueData(&InValue, sizeof(TArray<float>)); }
+
+	FORCEINLINE static bool IsClassOf(const FRigBaseMetadata* InMetadata)
+	{
+		return InMetadata->GetType() == ERigMetadataType::FloatArray;
+	}
+
+protected:
+
+	UPROPERTY()
+	TArray<float> Value;
+};
+
+USTRUCT()
+struct CONTROLRIG_API FRigInt32Metadata : public FRigBaseMetadata
+{
+public:
+	
+	GENERATED_BODY()
+	DECLARE_RIG_METADATA_METHODS(FRigInt32Metadata)
+
+	FRigInt32Metadata()
+		: FRigBaseMetadata()
+		, Value(0)
+	{}
+	virtual ~FRigInt32Metadata() override {}
+	FORCEINLINE virtual void Serialize(FArchive& Ar, bool bIsLoading) override
+	{
+		Super::Serialize(Ar, bIsLoading);
+		Ar << Value;
+	}
+	FORCEINLINE const int32& GetValue() const { return Value; } 
+	FORCEINLINE int32& GetValue() { return Value; }
+	FORCEINLINE bool SetValue(const int32& InValue) { return SetValueData(&InValue, sizeof(int32)); }
+
+	FORCEINLINE static bool IsClassOf(const FRigBaseMetadata* InMetadata)
+    {
+    	return InMetadata->GetType() == ERigMetadataType::Int32;
+    }
+
+protected:
+
+	UPROPERTY()
+	int32 Value;
+};
+
+USTRUCT()
+struct CONTROLRIG_API FRigInt32ArrayMetadata : public FRigBaseMetadata
+{
+public:
+	
+	GENERATED_BODY()
+	DECLARE_RIG_METADATA_METHODS(FRigInt32ArrayMetadata)
+
+	FRigInt32ArrayMetadata()
+		: FRigBaseMetadata()
+	{}
+	virtual ~FRigInt32ArrayMetadata() override {}
+	FORCEINLINE virtual void Serialize(FArchive& Ar, bool bIsLoading) override
+	{
+		Super::Serialize(Ar, bIsLoading);
+		Ar << Value;
+	}
+	FORCEINLINE const TArray<int32>& GetValue() const { return Value; }
+	FORCEINLINE TArray<int32>& GetValue() { return Value; }
+	FORCEINLINE bool SetValue(const TArray<int32>& InValue) { return SetValueData(&InValue, sizeof(TArray<int32>)); }
+
+	FORCEINLINE static bool IsClassOf(const FRigBaseMetadata* InMetadata)
+	{
+		return InMetadata->GetType() == ERigMetadataType::Int32Array;
+	}
+
+protected:
+
+	UPROPERTY()
+	TArray<int32> Value;
+};
+
+USTRUCT()
+struct CONTROLRIG_API FRigNameMetadata : public FRigBaseMetadata
+{
+public:
+	
+	GENERATED_BODY()
+	DECLARE_RIG_METADATA_METHODS(FRigNameMetadata)
+
+	FRigNameMetadata()
+		: FRigBaseMetadata()
+		, Value(NAME_None)
+	{}
+	virtual ~FRigNameMetadata() override {}
+	FORCEINLINE virtual void Serialize(FArchive& Ar, bool bIsLoading) override
+	{
+		Super::Serialize(Ar, bIsLoading);
+		Ar << Value;
+	}
+	FORCEINLINE const FName& GetValue() const { return Value; } 
+	FORCEINLINE FName& GetValue() { return Value; }
+	FORCEINLINE bool SetValue(const FName& InValue) { return SetValueData(&InValue, sizeof(FName)); }
+
+	FORCEINLINE static bool IsClassOf(const FRigBaseMetadata* InMetadata)
+    {
+    	return InMetadata->GetType() == ERigMetadataType::Name;
+    }
+
+protected:
+
+	UPROPERTY()
+	FName Value;
+};
+
+USTRUCT()
+struct CONTROLRIG_API FRigNameArrayMetadata : public FRigBaseMetadata
+{
+public:
+	
+	GENERATED_BODY()
+	DECLARE_RIG_METADATA_METHODS(FRigNameArrayMetadata)
+
+	FRigNameArrayMetadata()
+		: FRigBaseMetadata()
+	{}
+	virtual ~FRigNameArrayMetadata() override {}
+	FORCEINLINE virtual void Serialize(FArchive& Ar, bool bIsLoading) override
+	{
+		Super::Serialize(Ar, bIsLoading);
+		Ar << Value;
+	}
+	FORCEINLINE const TArray<FName>& GetValue() const { return Value; }
+	FORCEINLINE TArray<FName>& GetValue() { return Value; }
+	FORCEINLINE bool SetValue(const TArray<FName>& InValue) { return SetValueData(&InValue, sizeof(TArray<FName>)); }
+
+	FORCEINLINE static bool IsClassOf(const FRigBaseMetadata* InMetadata)
+	{
+		return InMetadata->GetType() == ERigMetadataType::NameArray;
+	}
+
+protected:
+
+	UPROPERTY()
+	TArray<FName> Value;
+};
+
+USTRUCT()
+struct CONTROLRIG_API FRigVectorMetadata : public FRigBaseMetadata
+{
+public:
+	
+	GENERATED_BODY()
+	DECLARE_RIG_METADATA_METHODS(FRigVectorMetadata)
+
+	FRigVectorMetadata()
+		: FRigBaseMetadata()
+		, Value(FVector::ZeroVector)
+	{}
+	virtual ~FRigVectorMetadata() override {}
+	FORCEINLINE virtual void Serialize(FArchive& Ar, bool bIsLoading) override
+	{
+		Super::Serialize(Ar, bIsLoading);
+		Ar << Value;
+	}
+	FORCEINLINE const FVector& GetValue() const { return Value; } 
+	FORCEINLINE FVector& GetValue() { return Value; }
+	FORCEINLINE bool SetValue(const FVector& InValue) { return SetValueData(&InValue, sizeof(FVector)); }
+
+	FORCEINLINE static bool IsClassOf(const FRigBaseMetadata* InMetadata)
+    {
+    	return InMetadata->GetType() == ERigMetadataType::Vector;
+    }
+
+protected:
+
+	UPROPERTY()
+	FVector Value;
+};
+
+USTRUCT()
+struct CONTROLRIG_API FRigVectorArrayMetadata : public FRigBaseMetadata
+{
+public:
+	
+	GENERATED_BODY()
+	DECLARE_RIG_METADATA_METHODS(FRigVectorArrayMetadata)
+
+	FRigVectorArrayMetadata()
+		: FRigBaseMetadata()
+	{}
+	virtual ~FRigVectorArrayMetadata() override {}
+	FORCEINLINE virtual void Serialize(FArchive& Ar, bool bIsLoading) override
+	{
+		Super::Serialize(Ar, bIsLoading);
+		Ar << Value;
+	}
+	FORCEINLINE const TArray<FVector>& GetValue() const { return Value; }
+	FORCEINLINE TArray<FVector>& GetValue() { return Value; }
+	FORCEINLINE bool SetValue(const TArray<FVector>& InValue) { return SetValueData(&InValue, sizeof(TArray<FVector>)); }
+
+	FORCEINLINE static bool IsClassOf(const FRigBaseMetadata* InMetadata)
+	{
+		return InMetadata->GetType() == ERigMetadataType::VectorArray;
+	}
+
+protected:
+
+	UPROPERTY()
+	TArray<FVector> Value;
+};
+
+USTRUCT()
+struct CONTROLRIG_API FRigRotatorMetadata : public FRigBaseMetadata
+{
+public:
+	
+	GENERATED_BODY()
+	DECLARE_RIG_METADATA_METHODS(FRigRotatorMetadata)
+
+	FRigRotatorMetadata()
+		: FRigBaseMetadata()
+		, Value(FRotator::ZeroRotator)
+	{}
+	virtual ~FRigRotatorMetadata() override {}
+	FORCEINLINE virtual void Serialize(FArchive& Ar, bool bIsLoading) override
+	{
+		Super::Serialize(Ar, bIsLoading);
+		Ar << Value;
+	}
+	FORCEINLINE const FRotator& GetValue() const { return Value; } 
+	FORCEINLINE FRotator& GetValue() { return Value; }
+	FORCEINLINE bool SetValue(const FRotator& InValue) { return SetValueData(&InValue, sizeof(FRotator)); }
+
+	FORCEINLINE static bool IsClassOf(const FRigBaseMetadata* InMetadata)
+    {
+    	return InMetadata->GetType() == ERigMetadataType::Rotator;
+    }
+
+protected:
+
+	UPROPERTY()
+	FRotator Value;
+};
+
+USTRUCT()
+struct CONTROLRIG_API FRigRotatorArrayMetadata : public FRigBaseMetadata
+{
+public:
+	
+	GENERATED_BODY()
+	DECLARE_RIG_METADATA_METHODS(FRigRotatorArrayMetadata)
+
+	FRigRotatorArrayMetadata()
+		: FRigBaseMetadata()
+	{}
+	virtual ~FRigRotatorArrayMetadata() override {}
+	FORCEINLINE virtual void Serialize(FArchive& Ar, bool bIsLoading) override
+	{
+		Super::Serialize(Ar, bIsLoading);
+		Ar << Value;
+	}
+	FORCEINLINE const TArray<FRotator>& GetValue() const { return Value; }
+	FORCEINLINE TArray<FRotator>& GetValue() { return Value; }
+	FORCEINLINE bool SetValue(const TArray<FRotator>& InValue) { return SetValueData(&InValue, sizeof(TArray<FRotator>)); }
+
+	FORCEINLINE static bool IsClassOf(const FRigBaseMetadata* InMetadata)
+	{
+		return InMetadata->GetType() == ERigMetadataType::RotatorArray;
+	}
+
+protected:
+
+	UPROPERTY()
+	TArray<FRotator> Value;
+};
+
+USTRUCT()
+struct CONTROLRIG_API FRigQuatMetadata : public FRigBaseMetadata
+{
+public:
+	
+	GENERATED_BODY()
+	DECLARE_RIG_METADATA_METHODS(FRigQuatMetadata)
+
+	FRigQuatMetadata()
+		: FRigBaseMetadata()
+		, Value(FQuat::Identity)
+	{}
+	virtual ~FRigQuatMetadata() override {}
+	FORCEINLINE virtual void Serialize(FArchive& Ar, bool bIsLoading) override
+	{
+		Super::Serialize(Ar, bIsLoading);
+		Ar << Value;
+	}
+	FORCEINLINE const FQuat& GetValue() const { return Value; } 
+	FORCEINLINE FQuat& GetValue() { return Value; }
+	FORCEINLINE bool SetValue(const FQuat& InValue) { return SetValueData(&InValue, sizeof(FQuat)); }
+
+	FORCEINLINE static bool IsClassOf(const FRigBaseMetadata* InMetadata)
+    {
+    	return InMetadata->GetType() == ERigMetadataType::Quat;
+    }
+
+protected:
+
+	UPROPERTY()
+	FQuat Value;
+};
+
+USTRUCT()
+struct CONTROLRIG_API FRigQuatArrayMetadata : public FRigBaseMetadata
+{
+public:
+	
+	GENERATED_BODY()
+	DECLARE_RIG_METADATA_METHODS(FRigQuatArrayMetadata)
+
+	FRigQuatArrayMetadata()
+		: FRigBaseMetadata()
+	{}
+	virtual ~FRigQuatArrayMetadata() override {}
+	FORCEINLINE virtual void Serialize(FArchive& Ar, bool bIsLoading) override
+	{
+		Super::Serialize(Ar, bIsLoading);
+		Ar << Value;
+	}
+	FORCEINLINE const TArray<FQuat>& GetValue() const { return Value; }
+	FORCEINLINE TArray<FQuat>& GetValue() { return Value; }
+	FORCEINLINE bool SetValue(const TArray<FQuat>& InValue) { return SetValueData(&InValue, sizeof(TArray<FQuat>)); }
+
+	FORCEINLINE static bool IsClassOf(const FRigBaseMetadata* InMetadata)
+	{
+		return InMetadata->GetType() == ERigMetadataType::QuatArray;
+	}
+
+protected:
+
+	UPROPERTY()
+	TArray<FQuat> Value;
+};
+
+USTRUCT()
+struct CONTROLRIG_API FRigTransformMetadata : public FRigBaseMetadata
+{
+public:
+	
+	GENERATED_BODY()
+	DECLARE_RIG_METADATA_METHODS(FRigTransformMetadata)
+
+	FRigTransformMetadata()
+		: FRigBaseMetadata()
+		, Value(FTransform::Identity)
+	{}
+	virtual ~FRigTransformMetadata() override {}
+	FORCEINLINE virtual void Serialize(FArchive& Ar, bool bIsLoading) override
+	{
+		Super::Serialize(Ar, bIsLoading);
+		Ar << Value;
+	}
+	FORCEINLINE const FTransform& GetValue() const { return Value; } 
+	FORCEINLINE FTransform& GetValue() { return Value; }
+	FORCEINLINE bool SetValue(const FTransform& InValue) { return SetValueData(&InValue, sizeof(FTransform)); }
+
+	FORCEINLINE static bool IsClassOf(const FRigBaseMetadata* InMetadata)
+    {
+    	return InMetadata->GetType() == ERigMetadataType::Transform;
+    }
+
+protected:
+
+	UPROPERTY()
+	FTransform Value;
+};
+
+USTRUCT()
+struct CONTROLRIG_API FRigTransformArrayMetadata : public FRigBaseMetadata
+{
+public:
+	
+	GENERATED_BODY()
+	DECLARE_RIG_METADATA_METHODS(FRigTransformArrayMetadata)
+
+	FRigTransformArrayMetadata()
+		: FRigBaseMetadata()
+	{}
+	virtual ~FRigTransformArrayMetadata() override {}
+	FORCEINLINE virtual void Serialize(FArchive& Ar, bool bIsLoading) override
+	{
+		Super::Serialize(Ar, bIsLoading);
+		Ar << Value;
+	}
+	FORCEINLINE const TArray<FTransform>& GetValue() const { return Value; }
+	FORCEINLINE TArray<FTransform>& GetValue() { return Value; }
+	FORCEINLINE bool SetValue(const TArray<FTransform>& InValue) { return SetValueData(&InValue, sizeof(TArray<FTransform>)); }
+
+	FORCEINLINE static bool IsClassOf(const FRigBaseMetadata* InMetadata)
+	{
+		return InMetadata->GetType() == ERigMetadataType::TransformArray;
+	}
+
+protected:
+
+	UPROPERTY()
+	TArray<FTransform> Value;
+};
+
+USTRUCT()
+struct CONTROLRIG_API FRigElementKeyMetadata : public FRigBaseMetadata
+{
+public:
+	
+	GENERATED_BODY()
+	DECLARE_RIG_METADATA_METHODS(FRigElementKeyMetadata)
+
+	FRigElementKeyMetadata()
+		: FRigBaseMetadata()
+	{}
+	virtual ~FRigElementKeyMetadata() override {}
+	FORCEINLINE virtual void Serialize(FArchive& Ar, bool bIsLoading) override
+	{
+		Super::Serialize(Ar, bIsLoading);
+		Ar << Value;
+	}
+	FORCEINLINE const FRigElementKey& GetValue() const { return Value; } 
+	FORCEINLINE FRigElementKey& GetValue() { return Value; }
+	FORCEINLINE bool SetValue(const FRigElementKey& InValue) { return SetValueData(&InValue, sizeof(FRigElementKey)); }
+
+	FORCEINLINE static bool IsClassOf(const FRigBaseMetadata* InMetadata)
+    {
+    	return InMetadata->GetType() == ERigMetadataType::RigElementKey;
+    }
+
+protected:
+
+	UPROPERTY()
+	FRigElementKey Value;
+};
+
+USTRUCT()
+struct CONTROLRIG_API FRigElementKeyArrayMetadata : public FRigBaseMetadata
+{
+public:
+	
+	GENERATED_BODY()
+	DECLARE_RIG_METADATA_METHODS(FRigElementKeyArrayMetadata)
+
+	FRigElementKeyArrayMetadata()
+		: FRigBaseMetadata()
+	{}
+	virtual ~FRigElementKeyArrayMetadata() override {}
+	FORCEINLINE virtual void Serialize(FArchive& Ar, bool bIsLoading) override
+	{
+		Super::Serialize(Ar, bIsLoading);
+		Ar << Value;
+	}
+	FORCEINLINE const TArray<FRigElementKey>& GetValue() const { return Value; }
+	FORCEINLINE TArray<FRigElementKey>& GetValue() { return Value; }
+	FORCEINLINE bool SetValue(const TArray<FRigElementKey>& InValue) { return SetValueData(&InValue, sizeof(TArray<FRigElementKey>)); }
+
+	FORCEINLINE static bool IsClassOf(const FRigBaseMetadata* InMetadata)
+	{
+		return InMetadata->GetType() == ERigMetadataType::RigElementKeyArray;
+	}
+
+protected:
+
+	UPROPERTY()
+	TArray<FRigElementKey> Value;
+};
