Control Rig: Ensure to full preferred types for template nodes

diff --git a/Engine/Plugins/Animation/ControlRig/Source/ControlRigDeveloper/Private/ControlRigBlueprint.cpp b/Engine/Plugins/Animation/ControlRig/Source/ControlRigDeveloper/Private/ControlRigBlueprint.cpp
--- a/Engine/Plugins/Animation/ControlRig/Source/ControlRigDeveloper/Private/ControlRigBlueprint.cpp
+++ b/Engine/Plugins/Animation/ControlRig/Source/ControlRigDeveloper/Private/ControlRigBlueprint.cpp
@@ -1384,6 +1384,13 @@ void UControlRigBlueprint::RefreshAllModels(EControlRigBlueprintLoadType InLoadT
 				{
 					TemplateNode->InvalidateCache();
 					TemplateNode->PostLoad();
+
+					if(!TemplateNode->HasWildCardPin() &&
+						!TemplateNode->IsSingleton() &&
+						TemplateNode->PreferredPermutationPairs.IsEmpty())
+					{
+						TemplateNode->InitializeFilteredPermutationsFromTypes();
+					}
 				}
 			}
 		}
@@ -4868,13 +4875,12 @@ UEdGraph* UControlRigBlueprint::CreateEdGraph(URigVMGraph* InModel, bool bForce)
 	UControlRigGraph* ControlRigGraph = NewObject<UControlRigGraph>(this, *GraphName, RF_Transactional);
 	ControlRigGraph->Schema = UControlRigGraphSchema::StaticClass();
 	ControlRigGraph->bAllowDeletion = true;
+	ControlRigGraph->ModelNodePath = InModel->GetNodePath();
+	ControlRigGraph->Initialize(this);
 	
 	FBlueprintEditorUtils::AddUbergraphPage(this, ControlRigGraph);
 	LastEditedDocuments.AddUnique(ControlRigGraph);
 
-	ControlRigGraph->Initialize(this);
-	ControlRigGraph->ModelNodePath = InModel->GetNodePath();
-
 	return ControlRigGraph;
 }
 
