ControlRig: Use default hierarchy when performing unparent in hierarchy tree

diff --git a/Engine/Plugins/Animation/ControlRig/Source/ControlRigEditor/Private/Editor/SRigHierarchy.cpp b/Engine/Plugins/Animation/ControlRig/Source/ControlRigEditor/Private/Editor/SRigHierarchy.cpp
--- a/Engine/Plugins/Animation/ControlRig/Source/ControlRigEditor/Private/Editor/SRigHierarchy.cpp
+++ b/Engine/Plugins/Animation/ControlRig/Source/ControlRigEditor/Private/Editor/SRigHierarchy.cpp
@@ -2582,7 +2582,7 @@ void SRigHierarchy::HandleUnparent()
 		TGuardValue<bool> GuardRigHierarchyChanges(bIsChangingRigHierarchy, true);
 		TGuardValue<bool> SuspendBlueprintNotifs(ControlRigBlueprint->bSuspendAllNotifications, true);
 
-		URigHierarchy* Hierarchy = GetHierarchy();
+		URigHierarchy* Hierarchy = GetDefaultHierarchy();
 		check(Hierarchy);
 
 		URigHierarchyController* Controller = Hierarchy->GetController(true);
