Constraints in Sequencer:
- dynamic behavior or parent constraints so that child transform in constraint space can be dynamically changed over time.
- transform compensation when changing child transforms
- Aim constraint: dynamic look-at axis and input roll

to be discussed: this has to be done for scene components as well once the constraint channel available thru sequencer

diff --git a/Engine/Plugins/Animation/ControlRig/Source/ControlRig/Private/Constraints/ControlRigTransformableHandle.cpp b/Engine/Plugins/Animation/ControlRig/Source/ControlRig/Private/Constraints/ControlRigTransformableHandle.cpp
--- a/Engine/Plugins/Animation/ControlRig/Source/ControlRig/Private/Constraints/ControlRigTransformableHandle.cpp
+++ b/Engine/Plugins/Animation/ControlRig/Source/ControlRig/Private/Constraints/ControlRigTransformableHandle.cpp
@@ -16,6 +16,12 @@ UTransformableControlHandle::~UTransformableControlHandle()
 	UnregisterDelegates();
 }
 
+void UTransformableControlHandle::PostLoad()
+{
+	Super::PostLoad();
+	RegisterDelegates();
+}
+
 bool UTransformableControlHandle::IsValid() const
 {
 	if (!ControlRig.IsValid() || ControlName == NAME_None)
@@ -169,6 +175,7 @@ void UTransformableControlHandle::UnregisterDelegates() const
 		{
 			Hierarchy->OnModified().RemoveAll(this);
 		}
+		ControlRig->ControlModified().RemoveAll(this);
 	}
 }
 
@@ -186,6 +193,8 @@ void UTransformableControlHandle::RegisterDelegates()
 		{
 			Hierarchy->OnModified().AddUObject(this, &UTransformableControlHandle::OnHierarchyModified);
 		}
+		
+		ControlRig->ControlModified().AddUObject(this, &UTransformableControlHandle::OnControlModified);
 	}
 }
 
@@ -227,6 +236,30 @@ void UTransformableControlHandle::OnHierarchyModified(
 	}
 }
 
+void UTransformableControlHandle::OnControlModified(
+	UControlRig* InControlRig,
+	FRigControlElement* InControl,
+	const FRigControlModifiedContext& InContext)
+{
+	if (!InControlRig || !InControl)
+	{
+		return;
+	}
+
+	if (!ControlRig.IsValid() || ControlName == NAME_None)
+	{
+		return;
+	}
+
+	if (ControlRig == InControlRig && InControl->GetName() == ControlName)
+	{
+		if(OnHandleModified.IsBound())
+		{
+			OnHandleModified.Broadcast(this, InContext.bConstraintUpdate);
+		}
+	}
+}
+
 #if WITH_EDITOR
 
 FName UTransformableControlHandle::GetName() const
diff --git a/Engine/Plugins/Animation/ControlRig/Source/ControlRig/Private/Sequencer/MovieSceneControlRigParameterSection.cpp b/Engine/Plugins/Animation/ControlRig/Source/ControlRig/Private/Sequencer/MovieSceneControlRigParameterSection.cpp
--- a/Engine/Plugins/Animation/ControlRig/Source/ControlRig/Private/Sequencer/MovieSceneControlRigParameterSection.cpp
+++ b/Engine/Plugins/Animation/ControlRig/Source/ControlRig/Private/Sequencer/MovieSceneControlRigParameterSection.cpp
@@ -884,8 +884,9 @@ void UMovieSceneControlRigParameterSection::PostEditImport()
 void UMovieSceneControlRigParameterSection::PostLoad()
 {
 	Super::PostLoad();
-
+	RegisterConstraintsHandles();
 }
+
 bool UMovieSceneControlRigParameterSection::HasScalarParameter(FName InParameterName) const
 {
 	for (const FScalarParameterNameAndCurve& ScalarParameterNameAndCurve : ScalarParameterNamesAndCurves)
@@ -1276,6 +1277,34 @@ void UMovieSceneControlRigParameterSection::AddSpaceChannel(FName InControlName,
 	}
 }
 
+void UMovieSceneControlRigParameterSection::RegisterConstraintsHandles()
+{
+	if (OnConstraintRemovedHandle.IsValid())
+	{
+		return;
+	}
+
+	if (!ControlRig)
+	{
+		return;
+	}
+
+	// register constraint deletion handle
+	FConstraintsManagerController& Controller = FConstraintsManagerController::Get(ControlRig->GetWorld());
+	OnConstraintRemovedHandle = Controller.OnConstraintRemoved().AddLambda([this](FName InConstraintName)
+	{
+		Modify();
+		const int32 NumRemoved = ConstraintsChannels.RemoveAll( [InConstraintName](const FConstraintAndActiveChannel& InChannel)
+		{
+			return InChannel.Constraint.IsValid() ? InChannel.Constraint->GetFName() == InConstraintName : false;
+		});
+		if (NumRemoved)
+		{
+			ReconstructChannelProxy();
+		}
+	});
+}
+
 bool UMovieSceneControlRigParameterSection::HasConstraintChannel(const FName& InConstraintName) const
 {
 	return ConstraintsChannels.ContainsByPredicate( [InConstraintName](const FConstraintAndActiveChannel& InChannel)
@@ -1308,26 +1337,7 @@ void UMovieSceneControlRigParameterSection::AddConstraintChannel(UTickableConstr
 		{
 			ReconstructChannelProxy();
 		}
-
-		if (!OnConstraintRemovedHandle.IsValid())
-		{
-			if (ControlRig)
-			{
-				FConstraintsManagerController& Controller = FConstraintsManagerController::Get(ControlRig->GetWorld());
-				OnConstraintRemovedHandle = Controller.OnConstraintRemoved().AddLambda([this](FName InConstraintName)
-				{
-					Modify();
-					const int32 NumRemoved = ConstraintsChannels.RemoveAll( [InConstraintName](const FConstraintAndActiveChannel& InChannel)
-					{
-						return InChannel.Constraint.IsValid() ? InChannel.Constraint->GetFName() == InConstraintName : false;
-					});
-					if (NumRemoved)
-					{
-						ReconstructChannelProxy();
-					}
-				});
-			}
-		}
+		RegisterConstraintsHandles();
 	}
 }
 
@@ -1395,7 +1405,7 @@ void UMovieSceneControlRigParameterSection::ReconstructChannelProxy()
 		const FConstraintsManagerController& Controller = FConstraintsManagerController::Get(ControlRig->GetWorld());
 		auto GetConstraints = [&Controller, this](const FName& InControlName)
 		{
-			static constexpr bool bSorted = true;
+			static constexpr bool bSorted = true;			
 			const uint32 ControlHash = HashCombine(GetTypeHash(ControlRig.Get()), GetTypeHash(InControlName));
 			return Controller.GetParentConstraints(ControlHash, bSorted);
 		};
diff --git a/Engine/Plugins/Animation/ControlRig/Source/ControlRig/Private/Sequencer/MovieSceneControlRigParameterTemplate.cpp b/Engine/Plugins/Animation/ControlRig/Source/ControlRig/Private/Sequencer/MovieSceneControlRigParameterTemplate.cpp
--- a/Engine/Plugins/Animation/ControlRig/Source/ControlRig/Private/Sequencer/MovieSceneControlRigParameterTemplate.cpp
+++ b/Engine/Plugins/Animation/ControlRig/Source/ControlRig/Private/Sequencer/MovieSceneControlRigParameterTemplate.cpp
@@ -1786,7 +1786,7 @@ void FMovieSceneControlRigParameterTemplate::EvaluateCurvesWithMasks(const FMovi
 			const FConstraintAndActiveChannel& ConstraintAndActiveChannel = Constraints[Index];
 			ConstraintAndActiveChannel.ActiveChannel.Evaluate(Time, Value);
 			
-			Values.ConstraintsValues.Emplace(ConstraintAndActiveChannel.Constraint, Value);
+			Values.ConstraintsValues.Emplace(ConstraintAndActiveChannel.Constraint.Get(), Value);
 		}
 		
 		for (int32 Index = 0; Index < Bools.Num(); ++Index)
diff --git a/Engine/Plugins/Animation/ControlRig/Source/ControlRig/Public/Constraints/ControlRigTransformableHandle.h b/Engine/Plugins/Animation/ControlRig/Source/ControlRig/Public/Constraints/ControlRigTransformableHandle.h
--- a/Engine/Plugins/Animation/ControlRig/Source/ControlRig/Public/Constraints/ControlRigTransformableHandle.h
+++ b/Engine/Plugins/Animation/ControlRig/Source/ControlRig/Public/Constraints/ControlRigTransformableHandle.h
@@ -26,32 +26,35 @@ class CONTROLRIG_API UTransformableControlHandle : public UTransformableHandle
 	
 public:
 	virtual ~UTransformableControlHandle();
+
+	virtual void PostLoad() override;
 	
-	/** @todo document */
+	/** Sanity check to ensure that ControlRig and ControlName are safe to use. */
 	virtual bool IsValid() const override;
 
-	/** @todo document */
+	/** Sets the global transform of the control. */
 	virtual void SetGlobalTransform(const FTransform& InGlobal) const override;
-	/** @todo document */
+	/** Sets the local transform of the control. */
 	virtual void SetLocalTransform(const FTransform& InLocal) const override;
-	/** @todo document */
+	/** Gets the global transform of the control. */
 	virtual FTransform GetGlobalTransform() const  override;
-	/** @todo document */
+	/** Sets the local transform of the control. */
 	virtual FTransform GetLocalTransform() const  override;
 
-	/** @todo document */
+	/** Returns the target object containing the tick function (e.i. SkeletalComponent bound to ControlRig). */
 	virtual UObject* GetPrerequisiteObject() const override;
-	/** @todo document */
+	/** Returns ths SkeletalComponent tick function. */
 	virtual FTickFunction* GetTickFunction() const override;
 
-	/** @todo document */
+	/** Generates a hash value based on ControlRig and ControlName. */
 	virtual uint32 GetHash() const override;
 	/** @todo document */
 	virtual TWeakObjectPtr<UObject> GetTarget() const override;
-	/** @todo document */
+
+	/** Returns the skeletal mesh bound to ControlRig. */
 	USkeletalMeshComponent* GetSkeletalMesh() const;
 
-	/** @todo document */
+	/** Registers/Unregisters useful delegates to track changes in the control's transform. */
 	void UnregisterDelegates() const;
 	void RegisterDelegates();
 
@@ -60,13 +63,19 @@ public:
 	virtual FName GetName() const override;
 #endif
 
-	/** @todo document */
+	/** The ControlRig that this handle is pointing at. */
 	UPROPERTY(BlueprintReadOnly, Category = "Object")
 	TSoftObjectPtr<UControlRig> ControlRig;
 
-	/** @todo document */
+	/** The ControlName of the control that this handle is pointing at. */
 	UPROPERTY(BlueprintReadOnly, Category = "Object")
 	FName ControlName;
+
+	/** @todo document */
+	void OnControlModified(
+		UControlRig* InControlRig,
+		FRigControlElement* InControl,
+		const FRigControlModifiedContext& InContext);
 	
 private:
 
diff --git a/Engine/Plugins/Animation/ControlRig/Source/ControlRig/Public/Rigs/RigHierarchyDefines.h b/Engine/Plugins/Animation/ControlRig/Source/ControlRig/Public/Rigs/RigHierarchyDefines.h
--- a/Engine/Plugins/Animation/ControlRig/Source/ControlRig/Public/Rigs/RigHierarchyDefines.h
+++ b/Engine/Plugins/Animation/ControlRig/Source/ControlRig/Public/Rigs/RigHierarchyDefines.h
@@ -1385,7 +1385,7 @@ struct CONTROLRIG_API FRigControlModifiedContext
 	uint32 KeyMask;
 	float LocalTime;
 	FName EventName;
-	
+	bool bConstraintUpdate = false;
 };
 
 /*
diff --git a/Engine/Plugins/Animation/ControlRig/Source/ControlRig/Public/Sequencer/MovieSceneControlRigParameterSection.h b/Engine/Plugins/Animation/ControlRig/Source/ControlRig/Public/Sequencer/MovieSceneControlRigParameterSection.h
--- a/Engine/Plugins/Animation/ControlRig/Source/ControlRig/Public/Sequencer/MovieSceneControlRigParameterSection.h
+++ b/Engine/Plugins/Animation/ControlRig/Source/ControlRig/Public/Sequencer/MovieSceneControlRigParameterSection.h
@@ -117,7 +117,7 @@ struct CONTROLRIG_API FConstraintAndActiveChannel
 	{};
 
 	UPROPERTY()
-	TWeakObjectPtr<UTickableConstraint> Constraint;
+	TSoftObjectPtr<UTickableConstraint> Constraint;
 
 	UPROPERTY()
 	FMovieSceneConstraintChannel ActiveChannel;
@@ -417,6 +417,7 @@ public:
 
 	/** todo */
 	FDelegateHandle OnConstraintRemovedHandle;
+	void RegisterConstraintsHandles();
 	bool HasConstraintChannel(const FName& InConstraintName) const;
 	FConstraintAndActiveChannel* GetConstraintChannel(const FName& InConstraintName);
 	void AddConstraintChannel(UTickableConstraint* InConstraint, bool bReconstructChannel);
diff --git a/Engine/Plugins/Animation/ControlRig/Source/ControlRigEditor/Private/EditMode/ControlRigEditMode.cpp b/Engine/Plugins/Animation/ControlRig/Source/ControlRigEditor/Private/EditMode/ControlRigEditMode.cpp
--- a/Engine/Plugins/Animation/ControlRig/Source/ControlRigEditor/Private/EditMode/ControlRigEditMode.cpp
+++ b/Engine/Plugins/Animation/ControlRig/Source/ControlRigEditor/Private/EditMode/ControlRigEditMode.cpp
@@ -61,6 +61,7 @@
 #include "InteractiveToolManager.h"
 #include "EdModeInteractiveToolsContext.h"
 #include "Editor/EditorPerProjectUserSettings.h"
+#include "TransformConstraint.h"
 
 void UControlRigEditModeDelegateHelper::OnPoseInitialized()
 {
@@ -952,7 +953,7 @@ bool FControlRigEditMode::EndTracking(FEditorViewportClient* InViewportClient, F
 	InteractionType = (uint8)EControlRigInteractionType::None;
 	
 	if (InteractionScopes.Num() > 0)
-	{
+	{		
 		if (bManipulatorMadeChange)
 		{
 			bManipulatorMadeChange = false;
@@ -2514,7 +2515,6 @@ void FControlRigEditMode::OpenSpacePickerWidget()
 		{
 			if (WeakSequencer.IsValid())
 			{
-
 				if (const FRigControlElement* ControlElement = InHierarchy->Find<FRigControlElement>(InControlKey))
 				{
 					ISequencer* Sequencer = WeakSequencer.Pin().Get();
@@ -3604,17 +3604,71 @@ bool FControlRigEditMode::CanChangeControlShapeTransform()
 	return false;
 }
 
-void FControlRigEditMode::SetControlShapeTransform(AControlRigShapeActor* ShapeActor, const FTransform& InTransform)
+void FControlRigEditMode::SetControlShapeTransform(
+	const AControlRigShapeActor* InShapeActor,
+	const FTransform& InGlobalTransform,
+	const FTransform& InToWorldTransform,
+	const FRigControlModifiedContext& InContext,
+	const bool bPrintPython) const
 {
-	if (UControlRig* ControlRig = ShapeActor->ControlRig.Get())
+	UControlRig* ControlRig = InShapeActor->ControlRig.Get();
+	if (!ControlRig)
 	{
-		ControlRig->SetControlGlobalTransform(ShapeActor->ControlName, InTransform);
+		return;
 	}
+
+	static constexpr bool bNotify = true, bFixEuler = true, bUndo = false;
+	if (!IsInLevelEditor())
+	{
+		// assumes it's attached to actor
+		ControlRig->SetControlGlobalTransform(
+			InShapeActor->ControlName, InGlobalTransform, bNotify, InContext, bUndo, bPrintPython, bFixEuler);
+		return;
+	}
+	
+	// find the last constraint in the stack (this could be cached on mouse press)
+	TArray< TObjectPtr<UTickableConstraint> > Constraints;
+	FTransformConstraintUtils::GetParentConstraints(ControlRig->GetWorld(), InShapeActor, Constraints);
+	const int32 LastActiveIndex = Constraints.FindLastByPredicate([](const TObjectPtr<UTickableConstraint>& InConstraint)
+	{
+		if (const UTickableTransformConstraint* TransformConstraint = Cast<UTickableTransformConstraint>(InConstraint.Get()))
+		{
+			return InConstraint->Active && TransformConstraint->bDynamicOffset;
+		}
+		return false;
+	});
+
+	// set the global space
+	// assumes it's attached to actor
+	ControlRig->SetControlGlobalTransform(
+		InShapeActor->ControlName, InGlobalTransform, bNotify, InContext, bUndo, bPrintPython, bFixEuler);
+
+	if (!Constraints.IsValidIndex(LastActiveIndex))
+	{
+		return;
+	}
+	
+	// switch to constraint space
+	const UTickableTransformConstraint* Constraint = Cast<UTickableTransformConstraint>(Constraints[LastActiveIndex]);
+	const FTransform ParentTransform = Constraint->GetParentGlobalTransform();
+	const FTransform ChildTransform = InGlobalTransform * InToWorldTransform;
+	FTransform LocalTransform = ControlRig->GetControlLocalTransform(InShapeActor->ControlName);
+	
+	const ETransformConstraintType ConstraintType = static_cast<ETransformConstraintType>(Constraint->GetType());
+	LocalTransform = FTransformConstraintUtils::ComputeRelativeTransform(LocalTransform,
+		ChildTransform, ParentTransform, ConstraintType);
+
+	FRigControlModifiedContext Context = InContext;
+	Context.bConstraintUpdate = false;
+	
+	ControlRig->SetControlLocalTransform(InShapeActor->ControlName, LocalTransform, bNotify, Context, bUndo, bFixEuler);
+	
+	Constraint->Evaluate();
 }
 
-FTransform FControlRigEditMode::GetControlShapeTransform(AControlRigShapeActor* ShapeActor) const
+FTransform FControlRigEditMode::GetControlShapeTransform(const AControlRigShapeActor* ShapeActor)
 {
-	if (UControlRig* ControlRig = ShapeActor->ControlRig.Get())
+	if (const UControlRig* ControlRig = ShapeActor->ControlRig.Get())
 	{
 		return ControlRig->GetControlGlobalTransform(ShapeActor->ControlName);
 	}
@@ -3706,6 +3760,7 @@ void FControlRigEditMode::MoveControlShape(AControlRigShapeActor* ShapeActor, co
 				FTransform NewTransform = CurrentTransform.GetRelativeTransform(ToWorldTransform);
 				FRigControlModifiedContext Context;
 				Context.EventName = FRigUnit_BeginExecution::EventName;
+				Context.bConstraintUpdate = true;
 				if (bCalcLocal)
 				{
 					InOutLocal = ControlRig->GetControlLocalTransform(ShapeActor->ControlName);
@@ -3716,13 +3771,14 @@ void FControlRigEditMode::MoveControlShape(AControlRigShapeActor* ShapeActor, co
 				{
 					bPrintPythonCommands = World->IsPreviewWorld();
 				}
-				ControlRig->SetControlGlobalTransform(ShapeActor->ControlName, NewTransform, true, Context, false, bPrintPythonCommands, true);			// assumes it's attached to actor
+
+				SetControlShapeTransform(ShapeActor, NewTransform, ToWorldTransform, Context, bPrintPythonCommands);			
+				
 				ShapeActor->SetGlobalTransform(CurrentTransform);
 				if (bCalcLocal)
 				{
 					FTransform NewLocal = ControlRig->GetControlLocalTransform(ShapeActor->ControlName);
 					InOutLocal = NewLocal.GetRelativeTransform(InOutLocal);
-
 				}
 
 				ControlRig->Evaluate_AnyThread();
diff --git a/Engine/Plugins/Animation/ControlRig/Source/ControlRigEditor/Private/EditMode/ControlRigEditMode.h b/Engine/Plugins/Animation/ControlRig/Source/ControlRigEditor/Private/EditMode/ControlRigEditMode.h
--- a/Engine/Plugins/Animation/ControlRig/Source/ControlRigEditor/Private/EditMode/ControlRigEditMode.h
+++ b/Engine/Plugins/Animation/ControlRig/Source/ControlRigEditor/Private/EditMode/ControlRigEditMode.h
@@ -393,8 +393,9 @@ private:
 
 	/* Check on tick to see if movie scene has changed, returns true if it has*/
 	bool CheckMovieSceneSig();
-	void SetControlShapeTransform(AControlRigShapeActor* ShapeActor, const FTransform& InTransform);
-	FTransform GetControlShapeTransform(AControlRigShapeActor* ShapeActor) const;
+	void SetControlShapeTransform( const AControlRigShapeActor* InShapeActor, const FTransform& InGlobalTransform,
+		const FTransform& InToWorldTransform, const FRigControlModifiedContext& InContext, const bool bPrintPython) const;
+	static FTransform GetControlShapeTransform(const AControlRigShapeActor* ShapeActor);
 	void MoveControlShape(AControlRigShapeActor* ShapeActor, const bool bTranslation, FVector& InDrag, 
 		const bool bRotation, FRotator& InRot, const bool bScale, FVector& InScale, const FTransform& ToWorldTransform,
 		bool bUseLocal, bool bCalcLocal, FTransform& InOutLocal);
diff --git a/Engine/Plugins/Animation/ControlRig/Source/ControlRigEditor/Private/Editor/Constraints/SConstraintsWidget.cpp b/Engine/Plugins/Animation/ControlRig/Source/ControlRigEditor/Private/Editor/Constraints/SConstraintsWidget.cpp
--- a/Engine/Plugins/Animation/ControlRig/Source/ControlRigEditor/Private/Editor/Constraints/SConstraintsWidget.cpp
+++ b/Engine/Plugins/Animation/ControlRig/Source/ControlRigEditor/Private/Editor/Constraints/SConstraintsWidget.cpp
@@ -235,6 +235,7 @@ FReply SDroppableConstraintItem::OnDragDetected(const FGeometry& MyGeometry, con
 FReply SDroppableConstraintItem::CreateSelectionPicker() const
 {
 	// FIXME temp approach for selecting the parent
+	FSlateApplication::Get().DismissAllMenus();
 	
 	static const FActorPickerModeModule& ActorPickerMode = FModuleManager::Get().GetModuleChecked<FActorPickerModeModule>("ActorPickerMode");
 
@@ -245,7 +246,9 @@ FReply SDroppableConstraintItem::CreateSelectionPicker() const
 		FOnShouldFilterActor(), 
 		FOnActorSelected::CreateLambda([ConstraintsCreationWidget, ConstraintTypeCopy](AActor* InActor)
 		{
-			SDroppableConstraintItem::CreateConstraint(InActor, ConstraintsCreationWidget, ConstraintTypeCopy);
+			const FOnConstraintCreated CreationDelegate = ConstraintsCreationWidget.IsValid() ?
+				ConstraintsCreationWidget->OnConstraintCreated : FOnConstraintCreated();
+			SDroppableConstraintItem::CreateConstraint(InActor, CreationDelegate, ConstraintTypeCopy);
 		}) );
 
 	
@@ -254,7 +257,7 @@ FReply SDroppableConstraintItem::CreateSelectionPicker() const
 
 void SDroppableConstraintItem::CreateConstraint(
 	AActor* InParent,
-	TSharedPtr<SConstraintsCreationWidget> InCreationWidget,
+	FOnConstraintCreated InCreationDelegate,
 	const ETransformConstraintType InConstraintType)
 {
 	if (!InParent)
@@ -290,10 +293,9 @@ void SDroppableConstraintItem::CreateConstraint(
 	}
 
 	// update list
-	const bool bCreationBound = InCreationWidget.IsValid() ? InCreationWidget->OnConstraintCreated.IsBound() : false;
-	if (bCreated && bCreationBound)
+	if (bCreated && InCreationDelegate.IsBound())
 	{
-		InCreationWidget->OnConstraintCreated.Execute();
+		InCreationDelegate.Execute();
 	}
 }
 
@@ -571,6 +573,7 @@ void SConstraintsEditionWidget::Construct(const FArguments& InArgs)
 				.ListItemsSource( &ListItems )
 				.OnGenerateRow(this, &SConstraintsEditionWidget::OnGenerateWidgetForItem)
 				.OnContextMenuOpening(this, &SConstraintsEditionWidget::CreateContextMenu)
+				.OnMouseButtonDoubleClick(this, &SConstraintsEditionWidget::OnItemDoubleClicked)
 			]
 		]
 	];
@@ -848,6 +851,34 @@ TSharedPtr<SWidget> SConstraintsEditionWidget::CreateContextMenu()
 			}
 		})),
 		NAME_None,
+		
+		EUserInterfaceActionType::Button);
+		MenuBuilder.AddMenuEntry(
+		LOCTEXT("CompensateKeyLabel", "Compensate Key"),
+		FText::Format(LOCTEXT("CompensateKeyTooltip", "Compensate transform key for {0}."), FText::FromName(ListItems[Index]->Name)),
+		FSlateIcon(),
+		FUIAction(FExecuteAction::CreateLambda([Constraint]()
+		{
+			if (UTickableTransformConstraint* TransformConstraint = Cast<UTickableTransformConstraint>(Constraint))
+			{
+				FConstraintChannelHelper::Compensate(TransformConstraint);
+			}
+		})),
+		NAME_None,
+		EUserInterfaceActionType::Button);
+
+		MenuBuilder.AddMenuEntry(
+		LOCTEXT("CompensateAllKeysLabel", "Compensate All Keys"),
+		FText::Format(LOCTEXT("CompensateAllKeysTooltip", "Compensate all transform keys for {0}."), FText::FromName(ListItems[Index]->Name)),
+		FSlateIcon(),
+		FUIAction(FExecuteAction::CreateLambda([Constraint]()
+		{
+			if (UTickableTransformConstraint* TransformConstraint = Cast<UTickableTransformConstraint>(Constraint))
+			{
+				FConstraintChannelHelper::Compensate(TransformConstraint, true);
+			}
+		})),
+		NAME_None,
 		EUserInterfaceActionType::Button);
 	}
 	MenuBuilder.EndSection();
@@ -855,6 +886,25 @@ TSharedPtr<SWidget> SConstraintsEditionWidget::CreateContextMenu()
 	return MenuBuilder.MakeWidget();
 }
 
+void SConstraintsEditionWidget::OnItemDoubleClicked(ItemSharedPtr InItem)
+{
+	const int32 Index = ListItems.IndexOfByKey(InItem);
+	if (!ListItems.IsValidIndex(Index))
+	{
+		return;
+	}
+
+	UWorld* World = GCurrentLevelEditingViewportClient->GetWorld();
+	const FConstraintsManagerController& Controller = FConstraintsManagerController::Get(World);
+	UTickableConstraint* Constraint = Controller.GetConstraint(InItem->Name);
+	if (!Constraint)
+	{
+		return;
+	}
+
+	Constraint->SetActive(!Constraint->Active);
+}
+
 END_SLATE_FUNCTION_BUILD_OPTIMIZATION
 
 #undef LOCTEXT_NAMESPACE
diff --git a/Engine/Plugins/Animation/ControlRig/Source/ControlRigEditor/Private/Editor/Constraints/SConstraintsWidget.h b/Engine/Plugins/Animation/ControlRig/Source/ControlRigEditor/Private/Editor/Constraints/SConstraintsWidget.h
--- a/Engine/Plugins/Animation/ControlRig/Source/ControlRigEditor/Private/Editor/Constraints/SConstraintsWidget.h
+++ b/Engine/Plugins/Animation/ControlRig/Source/ControlRigEditor/Private/Editor/Constraints/SConstraintsWidget.h
@@ -11,6 +11,8 @@ class SConstraintsCreationWidget;
 class SConstraintsEditionWidget;
 class UTickableConstraint;
 
+DECLARE_DELEGATE(FOnConstraintCreated);
+
 /**
  * FConstraintInfo
  */
@@ -84,7 +86,7 @@ private:
 	/** Creates the constraint between the current selection and the picked actor. */
 	static void CreateConstraint(
 		AActor* InParent,
-		TSharedPtr<SConstraintsCreationWidget> InConstraints,
+		FOnConstraintCreated InDelegate,
 		const ETransformConstraintType InConstraintType);
 
 	/** TSharedPtr to the tree item. */
@@ -106,8 +108,6 @@ class CONTROLRIGEDITOR_API SConstraintsCreationWidget : public SCompoundWidget
 {
 public:
 
-	DECLARE_DELEGATE(FOnConstraintCreated);
-	
 	SLATE_BEGIN_ARGS(SConstraintsCreationWidget)	{}
 
 		SLATE_EVENT(FOnConstraintCreated, OnConstraintCreated)
@@ -235,6 +235,7 @@ private:
 
 	/** @todo documentation. */
 	TSharedPtr< SWidget > CreateContextMenu();
+	void OnItemDoubleClicked(ItemSharedPtr InItem);
 	
 	void RegisterSelectionChanged();
 	void UnregisterSelectionChanged();
diff --git a/Engine/Plugins/Animation/ControlRig/Source/ControlRigEditor/Private/Sequencer/ConstraintChannelHelper.cpp b/Engine/Plugins/Animation/ControlRig/Source/ControlRigEditor/Private/Sequencer/ConstraintChannelHelper.cpp
--- a/Engine/Plugins/Animation/ControlRig/Source/ControlRigEditor/Private/Sequencer/ConstraintChannelHelper.cpp
+++ b/Engine/Plugins/Animation/ControlRig/Source/ControlRigEditor/Private/Sequencer/ConstraintChannelHelper.cpp
@@ -12,12 +12,15 @@
 #include "MovieSceneToolHelpers.h"
 
 #include "TransformConstraint.h"
+#include "Algo/Copy.h"
 #include "Algo/Unique.h"
 #include "Channels/MovieSceneConstraintChannel.h"
 
 #include "Tools/BakingHelper.h"
 #include "Tools/ConstraintBaker.h"
 
+bool FConstraintChannelHelper::bDoNotCompensate = false;
+
 namespace
 {
 	bool CanAddKey(const FConstraintAndActiveChannel* Channel, const FFrameNumber& InTime, bool& ActiveValue)
@@ -105,9 +108,7 @@ void FConstraintChannelHelper::AddChildTransformKey(
 	{
 		const FTransform LocalTransform = ControlHandle->GetLocalTransform();
 		
-		static constexpr bool bNotify = true;
-		static constexpr bool bUndo = false;
-		static constexpr bool bFixEuler = true;
+		static constexpr bool bNotify = true, bUndo = false, bFixEuler = true;
 	
 		FRigControlModifiedContext KeyframeContext;
 		KeyframeContext.SetKey = EControlRigSetKey::Always;
@@ -272,84 +273,232 @@ TArrayView<FMovieSceneFloatChannel*> FConstraintChannelHelper::GetTransformFloat
 
 namespace
 {
-	
-void GetComponentGlobalTransforms(
-	const UTransformableControlHandle* InControlHandle,
-	const TSharedPtr<ISequencer>& InSequencer,
-	TArray<FFrameNumber> InFrames,
-	TArray<FTransform>& OutTransforms)
-{
-	const USkeletalMeshComponent* SkeletalMeshComponent = InControlHandle->GetSkeletalMesh();
-	AActor* Actor = SkeletalMeshComponent ? SkeletalMeshComponent->GetTypedOuter< AActor >() : nullptr;
-	if (Actor)
-	{
-		FActorForWorldTransforms ControlRigActorSelection;
-		ControlRigActorSelection.Actor = Actor;
-		MovieSceneToolHelpers::GetActorWorldTransforms(InSequencer.Get(), ControlRigActorSelection, InFrames, OutTransforms);		
-	}
-	else
-	{
-		OutTransforms.SetNum(InFrames.Num());
-		for (FTransform& Transform : OutTransforms)
-		{
-			Transform = FTransform::Identity;
-		}
-	}
-}
 
-struct FTransformEvaluator
+struct FCompensationEvaluator
 {
 public:
-	TArray<FTransform> ChildTransforms;
-	TArray<FTransform> ParentGlobalTransforms;
-	TArray<FTransform> ComponentGlobalTransforms;
-
-	FTransformEvaluator(UTickableTransformConstraint* InConstraint)
+	TArray<FTransform> ChildLocals;
+	TArray<FTransform> ChildGlobals;
+	TArray<FTransform> SpaceGlobals;
+	
+	FCompensationEvaluator(UTickableTransformConstraint* InConstraint)
 		: Constraint(InConstraint)
+		, Handle(InConstraint ? InConstraint->ChildTRSHandle : nullptr)
 	{}
-	
-	void ComputeFrames(const TSharedPtr<ISequencer>& InSequencer, const TArray<FFrameNumber>& InFrames)
+
+	void ComputeLocalTransforms(
+		UWorld* InWorld, const TSharedPtr<ISequencer>& InSequencer,
+		const TArray<FFrameNumber>& InFrames, const bool bToActive)
 	{
-		if (!Constraint)
+		if (InFrames.IsEmpty())
 		{
 			return;
 		}
 		
-		const UTransformableControlHandle* ControlHandle = Cast<UTransformableControlHandle>(Constraint->ChildTRSHandle);
-		if (!ControlHandle)
+		const TArray<UTickableTransformConstraint*> Constraints = GetHandleTransformConstraints(InWorld);
+		if (Constraints.IsEmpty())
 		{
 			return;
 		}
 
-		// todo: compute child, parent and component in one single pass for faster computation
-		FConstraintBaker::GetHandleTransforms(
-		   ControlHandle->ControlRig->GetWorld(),
-		   InSequencer,
-		   ControlHandle,
-		   InFrames,
-		   bLocal,
-		   ChildTransforms);
-	
-		FConstraintBaker::GetHandleTransforms(
-			ControlHandle->ControlRig->GetWorld(),
-			InSequencer,
-			Constraint->ParentTRSHandle,
-			InFrames,
-			bLocal,
-			ParentGlobalTransforms);
-
-		GetComponentGlobalTransforms(
-			ControlHandle,
-			InSequencer,
-			InFrames,
-			ComponentGlobalTransforms);
+		// find last active constraint in the list that is different than the on we want to compensate for
+		auto GetLastActiveConstraint = [this, Constraints]()
+		{
+			// find last active constraint in the list that is different than the on we want to compensate for
+			const int32 LastActiveIndex = Constraints.FindLastByPredicate([this](const UTickableTransformConstraint* InConstraint)
+			{
+				return InConstraint->Active && InConstraint->bDynamicOffset && InConstraint != Constraint;
+			});
+
+			// if found, return its parent global transform
+			return LastActiveIndex > INDEX_NONE ?Constraints[LastActiveIndex] : nullptr;
+		};
+		
+		const UMovieScene* MovieScene = InSequencer->GetFocusedMovieSceneSequence()->GetMovieScene();
+        const FFrameRate TickResolution = MovieScene->GetTickResolution();
+        const EMovieScenePlayerStatus::Type PlaybackStatus = InSequencer->GetPlaybackStatus();
+
+		const int32 NumFrames = InFrames.Num();
+
+		// resize arrays to num frames + 1 as we also evaluate at InFrames[0]-1
+		ChildLocals.SetNum(NumFrames+1);
+        ChildGlobals.SetNum(NumFrames+1);
+		SpaceGlobals.SetNum(NumFrames+1);
+
+		const ETransformConstraintType ConstraintType = static_cast<ETransformConstraintType>(Constraint->GetType());
+		
+        for (int32 Index = 0; Index < NumFrames+1; ++Index)
+        {
+        	const FFrameNumber FrameNumber = (Index == 0) ? InFrames[0]-1 : InFrames[Index-1];
+        
+        	// evaluate animation
+        	const FMovieSceneEvaluationRange EvaluationRange = FMovieSceneEvaluationRange(FFrameTime(FrameNumber), TickResolution);
+        	const FMovieSceneContext Context = FMovieSceneContext(EvaluationRange, PlaybackStatus).SetHasJumped(true);
+    
+        	InSequencer->GetEvaluationTemplate().Evaluate(Context, *InSequencer);
+        
+        	// evaluate constraints
+        	for (const UTickableTransformConstraint* InConstraint: Constraints)
+        	{
+        		InConstraint->Evaluate();
+        	}
+        
+        	// evaluate ControlRig?
+        	// ControlRig->Evaluate_AnyThread();
+
+        	FTransform& ChildLocal = ChildLocals[Index];
+        	FTransform& ChildGlobal = ChildGlobals[Index];
+        	FTransform& SpaceGlobal = SpaceGlobals[Index];
+
+        	// store child transforms        	
+        	ChildLocal = Handle->GetLocalTransform();
+        	ChildGlobal = Handle->GetGlobalTransform();
+        	
+        	const UTickableTransformConstraint* LastConstraint = GetLastActiveConstraint();
+        	const ETransformConstraintType LastConstraintType = LastConstraint ?
+				static_cast<ETransformConstraintType>(LastConstraint->GetType()) : ETransformConstraintType::Parent; 
+
+        	// store constraint/parent space global transform
+        	if (bToActive)
+       		{
+        		// if activating the constraint, store last constraint or parent space at T[0]-1
+        		// and constraint space for all other times
+				if (Index == 0)
+				{
+					if (LastConstraint)
+					{
+						SpaceGlobal = LastConstraint->GetParentGlobalTransform();
+						ChildLocal = FTransformConstraintUtils::ComputeRelativeTransform(
+							ChildLocal, ChildGlobal, SpaceGlobal, LastConstraintType);
+					}
+				}
+				else
+				{
+					SpaceGlobal = Constraint->GetParentGlobalTransform();
+					ChildLocal = FTransformConstraintUtils::ComputeRelativeTransform(
+						ChildLocal, ChildGlobal, SpaceGlobal, ConstraintType);
+				}
+       		}
+            else
+            {
+            	// if deactivating the constraint, store constraint space at T[0]-1
+            	// and last constraint or parent space for all other times
+            	if (Index == 0)
+            	{
+            		SpaceGlobal = Constraint->GetParentGlobalTransform();
+            		ChildLocal = FTransformConstraintUtils::ComputeRelativeTransform(
+            			ChildLocal, ChildGlobal, SpaceGlobal, ConstraintType);
+            	}
+                else
+                {
+                	if (LastConstraint)
+                	{
+                		SpaceGlobal = LastConstraint->GetParentGlobalTransform();
+                		ChildLocal = FTransformConstraintUtils::ComputeRelativeTransform(
+                			ChildLocal, ChildGlobal, SpaceGlobal, LastConstraintType);
+                	}
+                }
+            }
+        }
+	}
+
+	void ComputeCompensation(UWorld* InWorld, const TSharedPtr<ISequencer>& InSequencer, const FFrameNumber& InTime)
+	{
+		const TArray<UTickableTransformConstraint*> Constraints = GetHandleTransformConstraints(InWorld);
+		if (Constraints.IsEmpty())
+		{
+			return;
+		}
+
+		// find last active constraint in the list that is different than the on we want to compensate for
+		auto GetLastActiveConstraint = [this, Constraints]()
+		{
+			// find last active constraint in the list that is different than the on we want to compensate for
+			const int32 LastActiveIndex = Constraints.FindLastByPredicate([this](const UTickableTransformConstraint* InConstraint)
+			{
+				return InConstraint->Active && InConstraint->bDynamicOffset;
+			});
+
+			// if found, return its parent global transform
+			return LastActiveIndex > INDEX_NONE ? Constraints[LastActiveIndex] : nullptr;
+		};
+
+		auto EvaluateAt = [InSequencer, &Constraints](const FFrameNumber InFrame)
+		{
+			const UMovieScene* MovieScene = InSequencer->GetFocusedMovieSceneSequence()->GetMovieScene();
+			const FFrameRate TickResolution = MovieScene->GetTickResolution();
+			const EMovieScenePlayerStatus::Type PlaybackStatus = InSequencer->GetPlaybackStatus();
+
+			const FMovieSceneEvaluationRange EvaluationRange0 = FMovieSceneEvaluationRange(FFrameTime(InFrame), TickResolution);
+			const FMovieSceneContext Context0 = FMovieSceneContext(EvaluationRange0, PlaybackStatus).SetHasJumped(true);
+
+			InSequencer->GetEvaluationTemplate().Evaluate(Context0, *InSequencer);
+
+			for (const UTickableTransformConstraint* InConstraint: Constraints)
+			{
+				InConstraint->Evaluate();
+			}
+
+			// ControlRig->Evaluate_AnyThread();
+		};
+
+		// allocate
+		ChildLocals.SetNum(1);
+		ChildGlobals.SetNum(1);
+		SpaceGlobals.SetNum(1);
+		
+		// evaluate at InTime and store global
+		EvaluateAt(InTime);
+		ChildGlobals[0] = Handle->GetGlobalTransform();
+
+		// evaluate at InTime-1 and store local
+		EvaluateAt(InTime-1);
+        ChildLocals[0] = Handle->GetLocalTransform();
+
+		// if constraint at T-1 then switch to its space
+        if (const UTickableTransformConstraint* LastConstraint = GetLastActiveConstraint())
+        {
+        	SpaceGlobals[0] = LastConstraint->GetParentGlobalTransform();
+        	
+        	const ETransformConstraintType LastConstraintType = static_cast<ETransformConstraintType>(LastConstraint->GetType());
+        	ChildLocals[0] = FTransformConstraintUtils::ComputeRelativeTransform(ChildLocals[0],
+		ChildGlobals[0], SpaceGlobals[0], LastConstraintType);
+        }
+        else // switch to parent space
+        {
+        	const FTransform ChildLocal = ChildLocals[0];
+        	Handle->SetGlobalTransform(ChildGlobals[0]);
+        	ChildLocals[0] = Handle->GetLocalTransform();
+        	Handle->SetLocalTransform(ChildLocal);
+        }
 	}
 	
 private:
-	static constexpr bool bLocal = false; // store global transforms
+
+	TArray<UTickableTransformConstraint*> GetHandleTransformConstraints(UWorld* InWorld) const
+	{
+		TArray<UTickableTransformConstraint*> TransformConstraints;
+		if (Handle)
+		{
+			// get sorted transform constraints
+			const FConstraintsManagerController& Controller = FConstraintsManagerController::Get(InWorld);
+			static constexpr bool bSorted = true;
+			TArray<TObjectPtr<UTickableConstraint>> Constraints = Controller.GetParentConstraints(Handle->GetHash(), bSorted);
+			for (const TObjectPtr<UTickableConstraint>& InConstraint: Constraints)
+			{
+				if (UTickableTransformConstraint* TransformConstraint = Cast<UTickableTransformConstraint>(InConstraint.Get()))
+				{
+					TransformConstraints.Add(TransformConstraint);
+				}
+			}
+		}
+		return TransformConstraints;
+	}
+	
 	UTickableTransformConstraint* Constraint = nullptr;
+	UTransformableHandle* Handle = nullptr;
 };
-
+	
 }
 
 void FConstraintChannelHelper::SmartControlConstraintKey(
@@ -377,6 +526,8 @@ void FConstraintChannelHelper::SmartControlConstraintKey(
 			bool ActiveValueToBeSet = false;
 			if (CanAddKey(Channel, Time, ActiveValueToBeSet))
 			{
+				TGuardValue<bool> CompensateGuard(bDoNotCompensate, true);
+				
 				// set constraint as dynamic
 				InConstraint->bDynamicOffset = true;
 				
@@ -388,13 +539,11 @@ void FConstraintChannelHelper::SmartControlConstraintKey(
 				TArray<FFrameNumber> FramesToCompensate;
 				GetFramesToCompensate(Channel->ActiveChannel, ActiveValueToBeSet, Time, Channels, FramesToCompensate);
 
-				// store child, parent and component global transforms for these frames
-				FTransformEvaluator Evaluator(InConstraint);
-				Evaluator.ComputeFrames(InSequencer, FramesToCompensate);
-				TArray<FTransform>& ChildTransforms = Evaluator.ChildTransforms;
-				const TArray<FTransform>& ParentTransforms = Evaluator.ParentGlobalTransforms;
-				const TArray<FTransform>& ComponentTransforms = Evaluator.ComponentGlobalTransforms;
-
+				// store child and space transforms for these frames
+				FCompensationEvaluator Evaluator(InConstraint);
+				Evaluator.ComputeLocalTransforms(ControlRig->GetWorld(), InSequencer, FramesToCompensate, ActiveValueToBeSet);
+				TArray<FTransform>& ChildLocals = Evaluator.ChildLocals;
+				
 				// store tangents at this time
 				TArray<FMovieSceneTangentData> Tangents;
 				FControlRigSpaceChannelHelpers::EvaluateTangentAtThisTime(ControlRig, Section, ControlName, Time, Tangents);
@@ -404,14 +553,10 @@ void FConstraintChannelHelper::SmartControlConstraintKey(
 				// add child's transform key at Time-1 to keep animation
 				{
 					const FFrameNumber TimeMinusOne(Time -1);
-
-					const bool bPrevAsLocal = !ActiveValueToBeSet;
-					const FTransform& SpaceTransformTime = bPrevAsLocal ? ParentTransforms[0] : // constraint space  
-																		  ComponentTransforms[0]; // component space
-					const FTransform ChildTransformAtMinusOne = ChildTransforms[0].GetRelativeTransform(SpaceTransformTime);
+					
 					FBakingHelper::AddTransformKeys(ControlRig, ControlName, {TimeMinusOne},
-						{ChildTransformAtMinusOne}, ChannelsToKey, TickResolution, bPrevAsLocal);
-				
+					 	{ChildLocals[0]}, ChannelsToKey, TickResolution, true);
+					
 					// set tangents at Time-1
 					FControlRigSpaceChannelHelpers::SetTangentsAtThisTime(ControlRig, Section, ControlName, TimeMinusOne, Tangents);
 				}
@@ -425,19 +570,13 @@ void FConstraintChannelHelper::SmartControlConstraintKey(
 
 				// compensate
 				{
-					const bool bNextAsLocal = ActiveValueToBeSet;
-					for (int32 Index = 0; Index < ChildTransforms.Num(); ++Index)
-					{
-						const FTransform& SpaceTransform = bNextAsLocal ? ParentTransforms[Index] : // constraint space 
-																		  ComponentTransforms[Index]; // component space
-						FTransform& ChildTransform = ChildTransforms[Index];
-						ChildTransform = ChildTransform.GetRelativeTransform(SpaceTransform);
-					}
+					// we need to remove the first transforms as we store NumFrames+1 transforms
+					ChildLocals.RemoveAt(0);
 
 					// add keys
 					FBakingHelper::AddTransformKeys( ControlRig, ControlHandle->ControlName, FramesToCompensate,
-						ChildTransforms, ChannelsToKey, TickResolution, bNextAsLocal);
-				
+					 	ChildLocals, ChannelsToKey, TickResolution, true);
+					
 					// set tangents at Time
 					FControlRigSpaceChannelHelpers::SetTangentsAtThisTime(ControlRig, Section, ControlName, Time, Tangents);
 				}
@@ -508,177 +647,110 @@ void FConstraintChannelHelper::GetFramesToCompensate(
 	OutFramesAfter.SetNum(Algo::Unique(OutFramesAfter));
 }
 
-// namespace
-// {
-// 	using IndexRange = ::TRange<int32>;
-// 	IndexRange GetRange(const TArrayView<const bool>& Values, const int32 Offset) 
-// 	{
-// 		IndexRange Range;
-// 			
-// 		int32 FirstActive = INDEX_NONE;
-// 		for (int32 Index = Offset; Index < Values.Num(); ++Index)
-// 		{
-// 			if (Values[Index])
-// 			{
-// 				FirstActive = Index;
-// 				break;
-// 			}
-// 		}
-//
-// 		if (FirstActive == INDEX_NONE)
-// 		{
-// 			return Range;
-// 		}
-//
-// 		Range.SetLowerBound(TRangeBound<int32>(FirstActive));
-// 		Range.SetUpperBound(TRangeBound<int32>(FirstActive));
-//
-// 		for (int32 NextInactive = FirstActive+1; NextInactive < Values.Num(); ++NextInactive)
-// 		{
-// 			if (Values[NextInactive] == false)
-// 			{
-// 				Range.SetUpperBound(TRangeBound<int32>(NextInactive));
-// 				return Range;
-// 			}
-// 		}
-// 			
-// 		return Range;
-// 	};
-//
-// 	TArray<IndexRange> GetIndexRanges(const FMovieSceneConstraintChannel* Channel)
-// 	{
-// 		TArray<IndexRange> Ranges;
-//
-// 		const TMovieSceneChannelData<const bool> ChannelData = Channel->GetData();
-// 		const TArrayView<const bool> Values = ChannelData.GetValues();
-// 		if (Values.Num() == 1)
-// 		{
-// 			if (Values[0] == true)
-// 			{
-// 				Ranges.Emplace(0);
-// 			}
-// 			return Ranges;
-// 		}
-// 		
-// 		int32 Offset = 0;
-// 		while (Values.IsValidIndex(Offset))
-// 		{
-// 			TRange<int32> Range = GetRange(Values, Offset);
-// 			if (!Range.IsEmpty())
-// 			{
-// 				Ranges.Emplace(Range);
-// 				Offset = Range.GetUpperBoundValue()+1;
-// 			}
-// 			else
-// 			{
-// 				Offset = INDEX_NONE;
-// 			}
-// 		}
-// 			
-// 		return Ranges;
-// 	}
-// }
-//
-// void DrawKeys(
-// 	FMovieSceneConstraintChannel* Channel,
-// 	TArrayView<const FKeyHandle> InKeyHandles,
-// 	const UMovieSceneSection* InOwner,
-// 	TArrayView<FKeyDrawParams> OutKeyDrawParams)
-// {
-// 	const UMovieSceneControlRigParameterSection* Section = Cast<UMovieSceneControlRigParameterSection>(InOwner);
-// 	if (!Section)
-// 	{
-// 		return;
-// 	}
-// 	
-// 	static const FName SquareKeyBrushName("Sequencer.KeySquare");
-// 	static const FName FilledBorderBrushName("FilledBorder");
-//
-// 	static FKeyDrawParams Params;
-// 	Params.FillBrush = FAppStyle::GetBrush(FilledBorderBrushName);
-// 	Params.BorderBrush = FAppStyle::GetBrush(SquareKeyBrushName);
-// 	
-// 	for (FKeyDrawParams& Param : OutKeyDrawParams)
-// 	{
-// 		Param = Params;
-// 	}
-// }
-//
-// void DrawExtra(FMovieSceneConstraintChannel* Channel, const UMovieSceneSection* Owner,const FGeometry& AllottedGeometry, FSequencerSectionPainter& Painter)
-// {
-// 	const UMovieSceneControlRigParameterSection* Section = Cast<UMovieSceneControlRigParameterSection>(Owner);
-// 	if (!Section || !Channel)
-// 	{
-// 		return;
-// 	}
-// 	
-// 	// using namespace UE::Sequencer;
-//
-// 	// get index range
-// 	TArray<IndexRange> IndexRanges = GetIndexRanges(Channel);
-// 	if (IndexRanges.IsEmpty())
-// 	{
-// 		return;
-// 	}
-// 	
-// 	// convert to bar range
-// 	TArray<FKeyBarCurveModel::FBarRange> Ranges;
-// 	const TMovieSceneChannelData<const bool> ChannelData = Channel->GetData();
-// 	const TArrayView<const FFrameNumber> Times = ChannelData.GetTimes();
-//
-// 	// const bool IsFirstValueTrue = IndexRanges[0].GetLowerBoundValue() == 0;
-// 	
-// 	for (const IndexRange& ActiveRange: IndexRanges)
-// 	{
-// 		FKeyBarCurveModel::FBarRange BarRange;
-// 		FFrameRate TickResolution = Section->GetTypedOuter<UMovieScene>()->GetTickResolution();
-// 		double LowerValue = Times[ActiveRange.GetLowerBoundValue()] / TickResolution;
-// 		double UpperValue = Times[ActiveRange.GetUpperBoundValue()] / TickResolution;
-// 	
-// 		BarRange.Range.SetLowerBound(TRangeBound<double>(LowerValue));
-// 		BarRange.Range.SetUpperBound(TRangeBound<double>(UpperValue));
-// 	
-// 		BarRange.Name = "Constraint";
-// 		BarRange.Color = FLinearColor(.2, .5, .1);
-// 		static FLinearColor ZebraTint = FLinearColor::White.CopyWithNewOpacity(0.01f);
-// 		BarRange.Color = BarRange.Color * (1.f - ZebraTint.A) + ZebraTint * ZebraTint.A;
-// 		
-// 		Ranges.Add(BarRange);
-// 	}
-//
-// 	// draw bars
-// 	const FSlateBrush* WhiteBrush = FAppStyle::GetBrush("WhiteBrush");
-// 	static constexpr ESlateDrawEffect DrawEffects = ESlateDrawEffect::None;
-//
-// 	const FSlateFontInfo FontInfo = FCoreStyle::Get().GetFontStyle("ToolTip.LargerFont");
-// 	const TSharedRef<FSlateFontMeasure> FontMeasure = FSlateApplication::Get().GetRenderer()->GetFontMeasureService();
-//
-// 	const FVector2D& LocalSize = AllottedGeometry.GetLocalSize();
-// 	static constexpr float LaneTop = 0;
-//
-// 	const FTimeToPixel& TimeToPixel = Painter.GetTimeConverter();
-// 	const double InputMin = TimeToPixel.PixelToSeconds(0.f);
-// 	const double InputMax = TimeToPixel.PixelToSeconds(Painter.SectionGeometry.GetLocalSize().X);
-//
-// 	for (int32 Index = 0; Index < Ranges.Num(); ++Index)
-// 	{
-// 		const FKeyBarCurveModel::FBarRange& Range = Ranges[Index];
-//
-// 		double LowerSeconds = /*(Index == 0 && IsFirstValueTrue) ? InputMin : */Range.Range.GetLowerBoundValue();
-// 		double UpperSeconds = Range.Range.GetUpperBoundValue();
-// 		if (UpperSeconds == Range.Range.GetLowerBoundValue())
-// 		{
-// 			UpperSeconds = InputMax;
-// 		}
-//
-// 		const double BoxStart = TimeToPixel.SecondsToPixel(LowerSeconds);
-// 		const double BoxEnd = TimeToPixel.SecondsToPixel(UpperSeconds);
-// 		const double BoxSize = BoxEnd - BoxStart;
-//
-// 		const FVector2D Size = FVector2D(BoxSize, LocalSize.Y);
-// 		const FVector2D Translation = FVector2D(BoxStart, LaneTop);
-// 		const FPaintGeometry BoxGeometry = AllottedGeometry.ToPaintGeometry(Size, FSlateLayoutTransform(Translation));
-// 		
-// 		FSlateDrawElement::MakeBox(Painter.DrawElements, Painter.LayerId, BoxGeometry, WhiteBrush, DrawEffects, Range.Color);
-// 	}
-// }
+void FConstraintChannelHelper::Compensate(UTickableTransformConstraint* InConstraint, const bool bAllTimes)
+{
+	const TWeakPtr<ISequencer> WeakSequencer = FBakingHelper::GetSequencer();
+	if (!WeakSequencer.IsValid() || !WeakSequencer.Pin()->GetFocusedMovieSceneSequence())
+	{
+		return;
+	}
+
+	const UTransformableControlHandle* ControlHandle = Cast<UTransformableControlHandle>(InConstraint->ChildTRSHandle);
+	if (!ControlHandle)
+	{
+		return;
+	}
+
+	const TSharedPtr<ISequencer> Sequencer = WeakSequencer.Pin();
+	if (const UMovieSceneControlRigParameterSection* Section = GetControlSection(ControlHandle, Sequencer))
+	{
+		const FFrameRate TickResolution = Sequencer->GetFocusedTickResolution();
+		const FFrameTime FrameTime = Sequencer->GetLocalTime().ConvertTo(TickResolution);
+		const FFrameNumber Time = FrameTime.GetFrame();
+
+		const TOptional<FFrameNumber> OptTime = bAllTimes ? TOptional<FFrameNumber>() : TOptional<FFrameNumber>(Time);
+		CompensateIfNeeded(ControlHandle->ControlRig.Get(), Sequencer, Section, OptTime);
+	}
+}
+
+void FConstraintChannelHelper::CompensateIfNeeded(
+	const UControlRig* ControlRig,
+	const TSharedPtr<ISequencer>& InSequencer,
+	const UMovieSceneControlRigParameterSection* Section,
+	const TOptional<FFrameNumber>& OptionalTime)
+{
+	if (bDoNotCompensate)
+	{
+		return;
+	}
+
+	TGuardValue<bool> CompensateGuard(bDoNotCompensate, true);
+	
+	// Frames to compensate
+	TArray<FFrameNumber> OptionalTimeArray;
+	if (OptionalTime.IsSet())
+	{
+		OptionalTimeArray.Add(OptionalTime.GetValue());
+	}
+	
+	auto GetSpaceTimesToCompensate = [&OptionalTimeArray](const FConstraintAndActiveChannel& Channel)->TArrayView<const FFrameNumber>
+	{
+		if (OptionalTimeArray.IsEmpty())
+		{
+			return Channel.ActiveChannel.GetData().GetTimes();
+		}
+		return OptionalTimeArray;
+	};
+	
+	// keyframe context
+	FRigControlModifiedContext KeyframeContext;
+	KeyframeContext.SetKey = EControlRigSetKey::Always;
+	const FFrameRate TickResolution = InSequencer->GetFocusedTickResolution();
+	
+	bool bNeedsEvaluation = false;
+
+	// gather all transform constraints
+	TArray<FConstraintAndActiveChannel> TransformConstraintsChannels;
+	Algo::CopyIf(Section->GetConstraintsChannels(), TransformConstraintsChannels,
+		[](const FConstraintAndActiveChannel& InChannel) 
+		{
+			return InChannel.Constraint.IsValid() && InChannel.Constraint->IsA<UTickableTransformConstraint>(); 
+		}
+	);
+	
+	// compensate constraints
+	for (const FConstraintAndActiveChannel& Channel: TransformConstraintsChannels)
+	{
+		const TArrayView<const FFrameNumber> FramesToCompensate = GetSpaceTimesToCompensate(Channel);
+		for (const FFrameNumber& Time : FramesToCompensate)
+		{
+			const FFrameNumber TimeMinusOne(Time-1);
+			
+			bool CurrentValue = false, PreviousValue = false;
+			Channel.ActiveChannel.Evaluate(TimeMinusOne, PreviousValue);
+			Channel.ActiveChannel.Evaluate(Time, CurrentValue);
+
+			if (CurrentValue != PreviousValue) //if they are the same no need to do anything
+			{
+				UTickableTransformConstraint* Constraint = Cast<UTickableTransformConstraint>(Channel.Constraint.Get()); 
+				
+				// compute transform to set
+				// if switching from active to inactive then we must add a key at T-1 in the constraint space
+				// if switching from inactive to active then we must add a key at T-1 in the previous constraint or parent space
+				FCompensationEvaluator Evaluator(Constraint);
+				Evaluator.ComputeCompensation(ControlRig->GetWorld(), InSequencer, Time);
+				const TArray<FTransform>& LocalTransforms = Evaluator.ChildLocals;
+				
+				const EMovieSceneTransformChannel ChannelsToKey = FConstraintBaker::GetChannelsToKey(Constraint);
+				FConstraintBaker::AddTransformKeys(
+					InSequencer, Constraint->ChildTRSHandle, {TimeMinusOne}, LocalTransforms, ChannelsToKey);
+				bNeedsEvaluation = true;
+			}
+		}
+	}
+	
+	if (bNeedsEvaluation)
+	{
+		InSequencer->ForceEvaluate();
+	}
+}
\ No newline at end of file
diff --git a/Engine/Plugins/Animation/ControlRig/Source/ControlRigEditor/Private/Sequencer/ControlRigParameterTrackEditor.cpp b/Engine/Plugins/Animation/ControlRig/Source/ControlRigEditor/Private/Sequencer/ControlRigParameterTrackEditor.cpp
--- a/Engine/Plugins/Animation/ControlRig/Source/ControlRigEditor/Private/Sequencer/ControlRigParameterTrackEditor.cpp
+++ b/Engine/Plugins/Animation/ControlRig/Source/ControlRigEditor/Private/Sequencer/ControlRigParameterTrackEditor.cpp
@@ -72,6 +72,7 @@
 #include "Animation/SkeletalMeshActor.h"
 #include "TimerManager.h"
 #include "BakeToControlRigSettings.h"
+#include "ConstraintChannelHelper.h"
 #include "UnrealEdGlobals.h"
 #include "Editor/UnrealEdEngine.h"
 #include "Toolkits/IToolkitHost.h"
@@ -2798,11 +2799,18 @@ FKeyPropertyResult FControlRigParameterTrackEditor::AddKeysToControlRigHandle(US
 		if (KeyPropertyResult.bKeyCreated)
 		{
 			UMovieSceneControlRigParameterSection* ParamSection = Cast<UMovieSceneControlRigParameterSection>(Track->GetSectionToKey());
-			if (ParamSection && ParamSection->GetControlRig())
+			if (UControlRig* SectionControlRig = ParamSection ? ParamSection->GetControlRig() : nullptr)
 			{
 				TOptional<FFrameNumber> OptionalKeyTime = KeyTime;
-				FControlRigSpaceChannelHelpers::CompensateIfNeeded(ParamSection->GetControlRig(), GetSequencer().Get(), ParamSection,
+
+				// compensate spaces
+				FControlRigSpaceChannelHelpers::CompensateIfNeeded(
+					SectionControlRig, GetSequencer().Get(), ParamSection,
 					RigControlName, OptionalKeyTime);
+
+				// compensate constraints
+				FConstraintChannelHelper::CompensateIfNeeded(
+					SectionControlRig, GetSequencer(), ParamSection, OptionalKeyTime);
 			}
 		}
 	}
diff --git a/Engine/Plugins/Animation/ControlRig/Source/ControlRigEditor/Private/Tools/ConstraintBaker.cpp b/Engine/Plugins/Animation/ControlRig/Source/ControlRigEditor/Private/Tools/ConstraintBaker.cpp
--- a/Engine/Plugins/Animation/ControlRig/Source/ControlRigEditor/Private/Tools/ConstraintBaker.cpp
+++ b/Engine/Plugins/Animation/ControlRig/Source/ControlRigEditor/Private/Tools/ConstraintBaker.cpp
@@ -7,6 +7,7 @@
 
 #include "ISequencer.h"
 
+#include "ControlRig.h"
 #include "IKeyArea.h"
 #include "LevelSequence.h"
 #include "MovieSceneToolHelpers.h"
@@ -128,7 +129,7 @@ void FConstraintBaker::DoIt(UTickableTransformConstraint* InConstraint)
 
 	// bake to channel curves
 	const EMovieSceneTransformChannel Channels = GetChannelsToKey(InConstraint);
-	BakeChild(Sequencer, InConstraint->ChildTRSHandle, Frames, Transforms, Channels);
+	AddTransformKeys(Sequencer, InConstraint->ChildTRSHandle, Frames, Transforms, Channels);
 
 	// disable constraint
 	InConstraint->SetActive(false);
@@ -203,7 +204,7 @@ void FConstraintBaker::GetHandleTransforms(
 	}
 }
 
-void FConstraintBaker::BakeChild(
+void FConstraintBaker::AddTransformKeys(
 	const TSharedPtr<ISequencer>& InSequencer,
 	UTransformableHandle* InHandle,
 	const TArray<FFrameNumber>& InFrames,
diff --git a/Engine/Plugins/Animation/ControlRig/Source/ControlRigEditor/Public/ConstraintChannelHelper.h b/Engine/Plugins/Animation/ControlRig/Source/ControlRigEditor/Public/ConstraintChannelHelper.h
--- a/Engine/Plugins/Animation/ControlRig/Source/ControlRigEditor/Public/ConstraintChannelHelper.h
+++ b/Engine/Plugins/Animation/ControlRig/Source/ControlRigEditor/Public/ConstraintChannelHelper.h
@@ -4,6 +4,7 @@
 
 #include "Templates/SharedPointer.h"
 
+class UControlRig;
 class ISequencer;
 class UMovieScene3DTransformSection;
 class UTransformableHandle;
@@ -24,14 +25,26 @@ struct FConstraintChannelHelper
 public:
 	/** @todo documentation. */
 	static void AddConstraintKey(UTickableTransformConstraint* InConstraint);
+	
+	/** Adds an active key if needed and does the compensation when switching. */
 	static void SmartConstraintKey(UTickableTransformConstraint* InConstraint);
 
-	/** @todo documentation. */
+	/** Adds a transform key on InHandle at InTime. */
 	static void AddChildTransformKey(
 		const UTransformableHandle* InHandle,
 		const FFrameNumber& InTime,
 		const TSharedPtr<ISequencer>& InSequencer);
 
+	/** Compensate transform on handles when a constraint switches state. */
+	static void Compensate(
+		UTickableTransformConstraint* InConstraint,
+		const bool bAllTimes = false);
+	static void CompensateIfNeeded(
+		const UControlRig* ControlRig,
+		const TSharedPtr<ISequencer>& InSequencer,
+		const UMovieSceneControlRigParameterSection* Section,
+		const TOptional<FFrameNumber>& OptionalTime);
+	
 private:
 
 	/** BEGIN CONTROL RIG SECTION */
@@ -74,17 +87,6 @@ private:
 	static TArrayView<FMovieSceneFloatChannel*> GetTransformFloatChannels(
 		const UTransformableHandle* InHandle,
 		const TSharedPtr<ISequencer>& InSequencer);
-};
 
-/** Key drawing overrides */
-// void DrawKeys(
-// 	FMovieSceneConstraintChannel* Channel,
-// 	TArrayView<const FKeyHandle> InKeyHandles,
-// 	const UMovieSceneSection* InOwner,
-// 	TArrayView<FKeyDrawParams> OutKeyDrawParams);
-//
-// void DrawExtra(
-// 	FMovieSceneConstraintChannel* Channel,
-// 	const UMovieSceneSection* Owner,
-// 	const FGeometry& AllottedGeometry,
-// 	FSequencerSectionPainter& Painter);
+	static bool bDoNotCompensate;
+};
diff --git a/Engine/Plugins/Animation/ControlRig/Source/ControlRigEditor/Public/Tools/ConstraintBaker.h b/Engine/Plugins/Animation/ControlRig/Source/ControlRigEditor/Public/Tools/ConstraintBaker.h
--- a/Engine/Plugins/Animation/ControlRig/Source/ControlRigEditor/Public/Tools/ConstraintBaker.h
+++ b/Engine/Plugins/Animation/ControlRig/Source/ControlRigEditor/Public/Tools/ConstraintBaker.h
@@ -16,16 +16,25 @@ public:
 	/** @todo documentation. (subject to changes to handle start / stop as parameters) */
 	static void DoIt(UTickableTransformConstraint* InConstraint);
 
-	/** @todo documentation. */
+	/** Stores InHandle local (or global) transforms at InFrames. */
 	static void GetHandleTransforms(
 		UWorld* InWorld,
 		const TSharedPtr<ISequencer>& InSequencer,
 		const UTransformableHandle* InHandle,
 		const TArray<FFrameNumber>& InFrames,
 		const bool bLocal,
-		TArray<FTransform>& OutTransforms); 
-
+		TArray<FTransform>& OutTransforms);
+	
+	/** Returns the channels to key based on the constraint's type. */
 	static EMovieSceneTransformChannel GetChannelsToKey(const UTickableTransformConstraint* InConstraint);
+
+	/** Add InTransforms keys at InFrames into the InHandle transform animation channels. */
+	static void AddTransformKeys(
+		const TSharedPtr<ISequencer>& InSequencer,
+		UTransformableHandle* InHandle,
+		const TArray<FFrameNumber>& InFrames,
+		const TArray<FTransform>& InTransforms,
+		const EMovieSceneTransformChannel& InChannels);
 	
 private:
 
@@ -37,13 +46,5 @@ private:
 		const TArray<FFrameNumber>& InFrames,
 		const bool bLocal,
 		TArray<FTransform>& OutTransforms);
-
-	/** Bake the resulting transforms into the child's transform animation channels. */
-	static void BakeChild(
-		const TSharedPtr<ISequencer>& InSequencer,
-		UTransformableHandle* InHandle,
-		const TArray<FFrameNumber>& InFrames,
-		const TArray<FTransform>& InTransforms,
-		const EMovieSceneTransformChannel& InChannels);
 };
 
diff --git a/Engine/Source/Editor/MovieSceneTools/Private/Channels/ConstraintChannelEditor.cpp b/Engine/Source/Editor/MovieSceneTools/Private/Channels/ConstraintChannelEditor.cpp
--- a/Engine/Source/Editor/MovieSceneTools/Private/Channels/ConstraintChannelEditor.cpp
+++ b/Engine/Source/Editor/MovieSceneTools/Private/Channels/ConstraintChannelEditor.cpp
@@ -7,6 +7,7 @@
 #include "KeyBarCurveModel.h"
 #include "SequencerSectionPainter.h"
 #include "TimeToPixel.h"
+#include "Fonts/FontMeasure.h"
 #include "Styling/AppStyle.h"
 #include "Styling/CoreStyle.h"
 #include "Framework/Application/SlateApplication.h"
@@ -116,8 +117,6 @@ void DrawExtra(
 	{
 		return;
 	}
-	
-	// using namespace UE::Sequencer;
 
 	// get index range
 	TArray<IndexRange> IndexRanges = GetIndexRanges(Channel);
@@ -125,7 +124,7 @@ void DrawExtra(
 	{
 		return;
 	}
-	
+
 	// convert to bar range
 	TArray<FKeyBarCurveModel::FBarRange> Ranges;
 	const TMovieSceneChannelData<const bool> ChannelData = Channel->GetData();
@@ -151,7 +150,7 @@ void DrawExtra(
 		Ranges.Add(BarRange);
 	}
 
-	// draw bars
+	// draw bars + labels
 	const FSlateBrush* WhiteBrush = FAppStyle::GetBrush("WhiteBrush");
 	static constexpr ESlateDrawEffect DrawEffects = ESlateDrawEffect::None;
 
@@ -165,10 +164,34 @@ void DrawExtra(
 	const double InputMin = TimeToPixel.PixelToSeconds(0.f);
 	const double InputMax = TimeToPixel.PixelToSeconds(Painter.SectionGeometry.GetLocalSize().X);
 
+	// label data
+	FText Label;
+	bool bDrawLabel = false;
+	
+	const FMovieSceneChannelProxy& Proxy = Owner->GetChannelProxy();
+	const TArrayView<FMovieSceneConstraintChannel*> ConstraintChannels = Proxy.GetChannels<FMovieSceneConstraintChannel>();
+	const int32 ChannelIndex = ConstraintChannels.IndexOfByKey(Channel);
+	if (ConstraintChannels.IsValidIndex(ChannelIndex))
+	{
+		const TArrayView<const FMovieSceneChannelMetaData> MetaData = Proxy.GetMetaData<FMovieSceneConstraintChannel>();
+		if (MetaData.IsValidIndex(ChannelIndex))
+		{
+			const FName& ConstraintName = MetaData[ChannelIndex].Name;
+			if (ConstraintName != NAME_None)
+			{
+				Label = FText::FromName(ConstraintName);
+				bDrawLabel = true;
+			}
+		}
+	}
+	const FVector2D TextSize = bDrawLabel ? FontMeasure->Measure(Label, FontInfo) : FVector2D::ZeroVector;
+	static constexpr double LabelPixelOffset = 10.0;  
+	
 	for (int32 Index = 0; Index < Ranges.Num(); ++Index)
 	{
 		const FKeyBarCurveModel::FBarRange& Range = Ranges[Index];
 
+		// draw bar
 		double LowerSeconds = /*(Index == 0 && IsFirstValueTrue) ? InputMin : */Range.Range.GetLowerBoundValue();
 		double UpperSeconds = Range.Range.GetUpperBoundValue();
 		if (UpperSeconds == Range.Range.GetLowerBoundValue())
@@ -185,5 +208,28 @@ void DrawExtra(
 		const FPaintGeometry BoxGeometry = AllottedGeometry.ToPaintGeometry(Size, FSlateLayoutTransform(Translation));
 		
 		FSlateDrawElement::MakeBox(Painter.DrawElements, Painter.LayerId, BoxGeometry, WhiteBrush, DrawEffects, Range.Color);
+
+		// draw label
+		if (bDrawLabel)
+		{
+			const double LabelPos = BoxStart + LabelPixelOffset;
+			const FVector2D Position(LabelPos, LaneTop + (LocalSize.Y - TextSize.Y) * .5f);
+			const FPaintGeometry LabelGeometry = AllottedGeometry.ToPaintGeometry(FSlateLayoutTransform(Position));
+
+			const double LabelMaxSize = BoxSize-LabelPixelOffset;
+			if (TextSize.X < LabelMaxSize)
+			{
+				FSlateDrawElement::MakeText(Painter.DrawElements, Painter.LayerId, LabelGeometry, Label, FontInfo);				
+			}
+			else
+			{
+				// crop
+				const int32 End = FontMeasure->FindLastWholeCharacterIndexBeforeOffset(Label, FontInfo, FMath::RoundToInt(LabelMaxSize));
+				if (End >= 0)
+				{
+					FSlateDrawElement::MakeText(Painter.DrawElements, Painter.LayerId, LabelGeometry, Label.ToString(), 0, End, FontInfo);
+				}
+			}
+		}
 	}
 }
diff --git a/Engine/Source/Runtime/Experimental/Animation/Constraints/Private/TransformConstraint.cpp b/Engine/Source/Runtime/Experimental/Animation/Constraints/Private/TransformConstraint.cpp
--- a/Engine/Source/Runtime/Experimental/Animation/Constraints/Private/TransformConstraint.cpp
+++ b/Engine/Source/Runtime/Experimental/Animation/Constraints/Private/TransformConstraint.cpp
@@ -45,6 +45,32 @@ void UTickableTransformConstraint::PostEditChangeProperty(FPropertyChangedEvent&
 
 #endif
 
+void UTickableTransformConstraint::UnregisterDelegates() const
+{
+	if (ChildTRSHandle)
+	{
+		ChildTRSHandle->HandleModified().RemoveAll(this);
+	}
+	if (ParentTRSHandle)
+	{
+		ParentTRSHandle->HandleModified().RemoveAll(this);
+	}
+}
+
+void UTickableTransformConstraint::RegisterDelegates()
+{
+	UnregisterDelegates();
+
+	if (ChildTRSHandle)
+	{
+		ChildTRSHandle->HandleModified().AddUObject(this, &UTickableTransformConstraint::OnHandleModified);
+	}
+	if (ParentTRSHandle)
+	{
+		ParentTRSHandle->HandleModified().AddUObject(this, &UTickableTransformConstraint::OnHandleModified);
+	}	
+}
+
 void UTickableTransformConstraint::Setup()
 {
 	if (!ParentTRSHandle->IsValid() || !ChildTRSHandle->IsValid())
@@ -55,6 +81,7 @@ void UTickableTransformConstraint::Setup()
 	
 	ComputeOffset();
 	SetupDependencies();
+	RegisterDelegates();
 }
 
 void UTickableTransformConstraint::SetupDependencies()
@@ -82,12 +109,14 @@ void UTickableTransformConstraint::PostLoad()
 {
 	Super::PostLoad();
 	ConstraintTick.RegisterFunction(GetFunction() );
+	RegisterDelegates();
 }
 
 void UTickableTransformConstraint::PostDuplicate(bool bDuplicateForPIE)
 {
 	Super::PostDuplicate(bDuplicateForPIE);
 	ConstraintTick.RegisterFunction(GetFunction());
+	RegisterDelegates();
 }
 
 uint32 UTickableTransformConstraint::GetTargetHash() const
@@ -148,6 +177,9 @@ FTransform UTickableTransformConstraint::GetParentLocalTransform() const
 	return ParentTRSHandle->IsValid() ? ParentTRSHandle->GetLocalTransform() : FTransform::Identity;
 }
 
+void UTickableTransformConstraint::OnHandleModified(UTransformableHandle* InHandle, bool bUpdate)
+{}
+
 /** 
  * UTickableTranslationConstraint
  **/
@@ -157,17 +189,42 @@ UTickableTranslationConstraint::UTickableTranslationConstraint()
 	Type = ETransformConstraintType::Translation;
 }
 
+#if WITH_EDITOR
+
+void UTickableTranslationConstraint::PostEditChangeProperty(FPropertyChangedEvent& PropertyChangedEvent)
+{
+	Super::PostEditChangeProperty(PropertyChangedEvent);
+
+	const FName PropertyName = PropertyChangedEvent.GetPropertyName();
+	if (PropertyName == GET_MEMBER_NAME_CHECKED(UTickableTransformConstraint, bDynamicOffset))
+	{
+		if (bDynamicOffset)
+		{
+			Cache.CachedInputHash = CalculateInputHash();
+			
+			const FTransform ChildGlobalTransform = GetChildGlobalTransform();
+			const FTransform ParentWorldTransform = GetParentGlobalTransform();
+			DynamicOffsetTranslation = ChildGlobalTransform.GetLocation() - ParentWorldTransform.GetLocation();
+			
+			Evaluate();
+		}
+		return;
+	}
+}
+
+#endif
+
 void UTickableTranslationConstraint::ComputeOffset()
 {
-	UClass* StaticClass = UTickableTranslationConstraint::StaticClass();
-	
 	const FTransform InitParentTransform = GetParentGlobalTransform();
 	const FTransform InitChildTransform = GetChildGlobalTransform();
 	
 	OffsetTranslation = FVector::ZeroVector;
+	DynamicOffsetTranslation = FVector::ZeroVector;
 	if (bMaintainOffset)
 	{
 		OffsetTranslation = InitChildTransform.GetLocation() - InitParentTransform.GetLocation();
+		DynamicOffsetTranslation = OffsetTranslation;
 	}	
 }
 
@@ -188,7 +245,9 @@ FConstraintTickFunction::ConstraintFunction UTickableTranslationConstraint::GetF
 
 		const FVector ParentTranslation = GetParentGlobalTransform().GetLocation();
 		FTransform Transform = GetChildGlobalTransform();
-		FVector NewTranslation = bMaintainOffset ? ParentTranslation + OffsetTranslation : ParentTranslation;
+		FVector NewTranslation = bDynamicOffset ? ParentTranslation + DynamicOffsetTranslation :
+								bMaintainOffset ? ParentTranslation + OffsetTranslation :
+								ParentTranslation;
 		if (ClampedWeight < 1.0f - KINDA_SMALL_NUMBER)
 		{
 			NewTranslation = FMath::Lerp<FVector>(Transform.GetLocation(), NewTranslation, ClampedWeight);
@@ -199,6 +258,54 @@ FConstraintTickFunction::ConstraintFunction UTickableTranslationConstraint::GetF
 	};
 }
 
+void UTickableTranslationConstraint::OnHandleModified(UTransformableHandle* InHandle, bool bUpdate)
+{
+	if (!bDynamicOffset)
+	{
+		return;
+	}
+	
+	if (!InHandle || InHandle != ChildTRSHandle)
+	{
+		return;
+	}
+
+	const uint32 InputHash = CalculateInputHash();
+	
+	// update dynamic offset
+	if (InputHash != Cache.CachedInputHash)
+	{
+		Cache.CachedInputHash = InputHash;
+
+		if (bUpdate)
+		{
+			const FTransform ParentWorldTransform = GetParentGlobalTransform();
+			const FTransform ChildGlobalTransform = GetChildGlobalTransform();
+			DynamicOffsetTranslation = ChildGlobalTransform.GetLocation() - ParentWorldTransform.GetLocation();
+		}
+		else
+		{
+			const FTransform ChildLocalTransform = GetChildLocalTransform();
+			DynamicOffsetTranslation = ChildLocalTransform.GetTranslation();
+		}
+	}
+}
+
+uint32 UTickableTranslationConstraint::CalculateInputHash() const
+{
+	uint32 Hash = 0;
+
+	// local location hash
+	const FTransform ChildLocalTransform = GetChildLocalTransform();
+	Hash = HashCombine(Hash, GetTypeHash(ChildLocalTransform.GetTranslation() ));
+
+	// global location hash
+	const FTransform ChildGlobalTransform = GetChildGlobalTransform();
+	Hash = HashCombine(Hash, GetTypeHash(ChildGlobalTransform.GetTranslation() ));
+	
+	return Hash;
+}
+
 /** 
  * UTickableRotationConstraint
  **/
@@ -208,16 +315,42 @@ UTickableRotationConstraint::UTickableRotationConstraint()
 	Type = ETransformConstraintType::Rotation;
 }
 
+#if WITH_EDITOR
+
+void UTickableRotationConstraint::PostEditChangeProperty(FPropertyChangedEvent& PropertyChangedEvent)
+{
+	Super::PostEditChangeProperty(PropertyChangedEvent);
+
+	const FName PropertyName = PropertyChangedEvent.GetPropertyName();
+	if (PropertyName == GET_MEMBER_NAME_CHECKED(UTickableTransformConstraint, bDynamicOffset))
+	{
+		if (bDynamicOffset)
+		{
+			Cache.CachedInputHash = CalculateInputHash();
+			
+			const FTransform ChildGlobalTransform = GetChildGlobalTransform();
+			const FTransform ParentWorldTransform = GetParentGlobalTransform();
+			DynamicOffsetRotation = ParentWorldTransform.GetRotation().Inverse() * ChildGlobalTransform.GetRotation();
+			
+			Evaluate();
+		}
+		return;
+	}
+}
+
+#endif
+
 void UTickableRotationConstraint::ComputeOffset()
 {
 	const FTransform InitParentTransform = GetParentGlobalTransform();
 	const FTransform InitChildTransform = GetChildGlobalTransform();
 	
-	OffsetRotation = FQuat::Identity;
+	OffsetRotation = DynamicOffsetRotation = FQuat::Identity;
 	if (bMaintainOffset)
 	{
 		OffsetRotation = InitParentTransform.GetRotation().Inverse() * InitChildTransform.GetRotation();
 		OffsetRotation.Normalize();
+		DynamicOffsetRotation = OffsetRotation; 
 	}
 }
 
@@ -239,7 +372,9 @@ FConstraintTickFunction::ConstraintFunction UTickableRotationConstraint::GetFunc
 		const FQuat ParentRotation = GetParentGlobalTransform().GetRotation();
 		FTransform Transform = GetChildGlobalTransform();
 
-		FQuat NewRotation = bMaintainOffset ? ParentRotation * OffsetRotation : ParentRotation;
+		FQuat NewRotation = bDynamicOffset ? ParentRotation * DynamicOffsetRotation :
+							bMaintainOffset ? ParentRotation * OffsetRotation :
+							ParentRotation;
 		if (ClampedWeight < 1.0f - KINDA_SMALL_NUMBER)
 		{
 			NewRotation = FQuat::Slerp(Transform.GetRotation(), NewRotation, ClampedWeight);
@@ -250,6 +385,53 @@ FConstraintTickFunction::ConstraintFunction UTickableRotationConstraint::GetFunc
 	};
 }
 
+void UTickableRotationConstraint::OnHandleModified(UTransformableHandle* InHandle, bool bUpdate)
+{
+	if (!bDynamicOffset)
+	{
+		return;
+	}
+	
+	if (!InHandle || InHandle != ChildTRSHandle)
+	{
+		return;
+	}
+
+	const uint32 InputHash = CalculateInputHash();
+	
+	// update dynamic offset
+	if (InputHash != Cache.CachedInputHash)
+	{
+		Cache.CachedInputHash = InputHash;
+
+		if (bUpdate)
+		{
+			const FTransform ParentWorldTransform = GetParentGlobalTransform();
+			const FTransform ChildGlobalTransform = GetChildGlobalTransform();
+			DynamicOffsetRotation = ParentWorldTransform.GetRotation().Inverse() * ChildGlobalTransform.GetRotation();
+		}
+		else
+		{
+			DynamicOffsetRotation = GetChildLocalTransform().GetRotation();
+		}
+	}
+}
+
+uint32 UTickableRotationConstraint::CalculateInputHash() const
+{
+	uint32 Hash = 0;
+
+	// local rotation hash
+	const FTransform ChildLocalTransform = GetChildLocalTransform();
+	Hash = HashCombine(Hash, GetTypeHash(ChildLocalTransform.GetRotation().Euler() ));
+
+	// global rotation hash
+	const FTransform ChildGlobalTransform = GetChildGlobalTransform();
+	Hash = HashCombine(Hash, GetTypeHash(ChildGlobalTransform.GetRotation().Euler() ));
+	
+	return Hash;
+}
+
 /** 
  * UTickableScaleConstraint
  **/
@@ -318,23 +500,27 @@ void UTickableParentConstraint::ComputeOffset()
 	const FTransform InitChildTransform = GetChildGlobalTransform();
 	
 	OffsetTransform = FTransform::Identity;
+	DynamicOffsetTransform = FTransform::Identity;
 	if (bMaintainOffset)
 	{
 		OffsetTransform = InitChildTransform.GetRelativeTransform(InitParentTransform);
+		DynamicOffsetTransform = OffsetTransform; 
 	}
 }
 
-uint32 UTickableParentConstraint::CalculateDependenciesHash() const
+uint32 UTickableParentConstraint::CalculateInputHash() const
 {
 	uint32 Hash = 0;
 	
-	Hash = HashCombine(Hash, GetTypeHash(Active));
-	Hash = HashCombine(Hash, GetTypeHash(Weight));
-
-	const FTransform ParentGlobal = GetParentGlobalTransform();
-	Hash = HashCombine(Hash, GetTypeHash(ParentGlobal.GetTranslation()));
-	Hash = HashCombine(Hash, GetTypeHash(ParentGlobal.GetRotation().Euler() ));
-	Hash = HashCombine(Hash, GetTypeHash(ParentGlobal.GetScale3D()));
+	const FTransform ChildLocalTransform = GetChildLocalTransform();
+	Hash = HashCombine(Hash, GetTypeHash(ChildLocalTransform.GetTranslation()));
+	Hash = HashCombine(Hash, GetTypeHash(ChildLocalTransform.GetRotation().Euler() ));
+	Hash = HashCombine(Hash, GetTypeHash(ChildLocalTransform.GetScale3D()));
+	
+	const FTransform ChildGlobalTransform = GetChildLocalTransform();
+	Hash = HashCombine(Hash, GetTypeHash(ChildGlobalTransform.GetTranslation()));
+	Hash = HashCombine(Hash, GetTypeHash(ChildGlobalTransform.GetRotation().Euler() ));
+	Hash = HashCombine(Hash, GetTypeHash(ChildGlobalTransform.GetScale3D()));
 	
 	return Hash;
 }
@@ -369,50 +555,72 @@ FConstraintTickFunction::ConstraintFunction UTickableParentConstraint::GetFuncti
 
 		const FTransform ParentTransform = GetParentGlobalTransform();
 		
-		// if bDynamicOffset is on, means we assume that the incoming local child transform is coming from the child
-		// TRS animation (for instance) which is local to the constraint's parent
-		if (bDynamicOffset)
-		{
-			const uint32 DependenciesHash = CalculateDependenciesHash();
+		FTransform TargetTransform = bDynamicOffset ? DynamicOffsetTransform * ParentTransform :
+									bMaintainOffset ? OffsetTransform * ParentTransform :
+									ParentTransform;
+		//apply weight if needed
+		LerpTransform(GetChildGlobalTransform(), TargetTransform);
+		SetChildGlobalTransform(TargetTransform);
+	};
+}
+
+void UTickableParentConstraint::OnHandleModified(UTransformableHandle* InHandle, bool bUpdate)
+{
+	if (!bDynamicOffset)
+	{
+		return;
+	}
+	
+	if (!InHandle || InHandle != ChildTRSHandle)
+	{
+		return;
+	}
+
+	const uint32 InputHash = CalculateInputHash();
 
+	// update dynamic offset
+	if (InputHash != Cache.CachedInputHash)
+	{
+		Cache.CachedInputHash = InputHash;
+		
+		if (bUpdate)
+		{
+			const FTransform ParentWorldTransform = GetParentGlobalTransform();
 			const FTransform ChildGlobalTransform = GetChildGlobalTransform();
-			const bool bIsGlobalDirty = !ChildGlobalTransform.Equals(Cache.LastGlobalSet);
-			const bool bIsInputHashDirty = DependenciesHash != Cache.CachedInputHash;
+			DynamicOffsetTransform = ChildGlobalTransform.GetRelativeTransform(ParentWorldTransform);
+		}
+		else
+		{
+			DynamicOffsetTransform = GetChildLocalTransform();
+		}
+	}
+}
 
-			// we avoid computation if nothing has changed
-			if (!bIsGlobalDirty && !bIsInputHashDirty)
-			{
-				return;
-			}
+#if WITH_EDITOR
 
-			FTransform LocalTransform = Cache.LastLocalSet;
-			if (bIsGlobalDirty)
-			{
-				LocalTransform = GetChildLocalTransform();
-			}
-			
-			FTransform TargetTransform = LocalTransform * ParentTransform;
-			//apply weight if needed
-			LerpTransform(ChildGlobalTransform, TargetTransform);
-			
-			SetChildGlobalTransform(TargetTransform);
+void UTickableParentConstraint::PostEditChangeProperty(FPropertyChangedEvent& PropertyChangedEvent)
+{
+	Super::PostEditChangeProperty(PropertyChangedEvent);
 
-			// update cache
-			Cache.CachedInputHash = DependenciesHash;
-			Cache.LastGlobalSet = TargetTransform;
-			Cache.LastLocalSet = LocalTransform;
+	const FName PropertyName = PropertyChangedEvent.GetPropertyName();
+	if (PropertyName == GET_MEMBER_NAME_CHECKED(UTickableTransformConstraint, bDynamicOffset))
+	{
+		if(bDynamicOffset)
+		{
+			Cache.CachedInputHash = CalculateInputHash();
 			
-			return;
+			const FTransform ChildGlobalTransform = GetChildGlobalTransform();
+			const FTransform ParentWorldTransform = GetParentGlobalTransform();
+			DynamicOffsetTransform = ChildGlobalTransform.GetRelativeTransform(ParentWorldTransform);
+			
+			Evaluate();
 		}
-
-		// static behavior
-		FTransform TargetTransform = bMaintainOffset ? OffsetTransform * ParentTransform : ParentTransform;
-		//apply weight if needed
-		LerpTransform(GetChildGlobalTransform(), TargetTransform);
-		SetChildGlobalTransform(TargetTransform);
-	};
+		return;
+	}
 }
 
+#endif
+
 /** 
  * UTickableLookAtConstraint
  **/
@@ -440,16 +648,49 @@ FConstraintTickFunction::ConstraintFunction UTickableLookAtConstraint::GetFuncti
 		const FTransform ParentTransform = GetParentGlobalTransform();
 		const FTransform ChildTransform = GetChildGlobalTransform();
 		
-		const FVector LookAtDir = ParentTransform.GetLocation() - ChildTransform.GetLocation();
-		const FRotator LookAtRotation = LookAtDir.Rotation();
+		const FVector LookAtDir = (ParentTransform.GetLocation() - ChildTransform.GetLocation()).GetSafeNormal();
 
-		FTransform Transform = GetChildGlobalTransform();
-		Transform.SetRotation(LookAtRotation.Quaternion());
+		if (!LookAtDir.IsNearlyZero() && !Axis.IsNearlyZero())
+		{
+			const FVector AxisToOrient = ChildTransform.TransformVectorNoScale(Axis).GetSafeNormal();
 		
-		SetChildGlobalTransform(Transform);
+			FQuat Rotation = FindQuatBetweenNormals(AxisToOrient, LookAtDir);
+			Rotation = Rotation * ChildTransform.GetRotation();
+
+			FTransform Transform = ChildTransform;
+			Transform.SetRotation(Rotation.GetNormalized());
+			SetChildGlobalTransform(Transform);
+		}
 	};
 }
 
+FQuat UTickableLookAtConstraint::FindQuatBetweenNormals(const FVector& A, const FVector& B)
+{
+	const FQuat::FReal Dot = FVector::DotProduct(A, B);
+	FQuat::FReal W = 1 + Dot;
+	FQuat Result;
+
+	if (W < SMALL_NUMBER)
+	{
+		// A and B point in opposite directions
+		W = 2 - W;
+		Result = FQuat( -A.Y * B.Z + A.Z * B.Y, -A.Z * B.X + A.X * B.Z, -A.X * B.Y + A.Y * B.X, W).GetNormalized();
+
+		const FVector Normal = FMath::Abs(A.X) > FMath::Abs(A.Y) ? FVector::YAxisVector : FVector::XAxisVector;
+		const FVector BiNormal = FVector::CrossProduct(A, Normal);
+		const FVector TauNormal = FVector::CrossProduct(A, BiNormal);
+		Result = Result * FQuat(TauNormal, PI);
+	}
+	else
+	{
+		//Axis = FVector::CrossProduct(A, B);
+		Result = FQuat( A.Y * B.Z - A.Z * B.Y, A.Z * B.X - A.X * B.Z, A.X * B.Y - A.Y * B.X, W);
+	}
+
+	Result.Normalize();
+	return Result;
+};
+
 /** 
  * FTransformConstraintUtils
  **/
@@ -499,6 +740,7 @@ UTransformableComponentHandle* FTransformConstraintUtils::CreateHandleForSceneCo
 		ComponentHandle = NewObject<UTransformableComponentHandle>(Outer);
 		ComponentHandle->Component = InSceneComponent;
 		InSceneComponent->SetMobility(EComponentMobility::Movable);
+		ComponentHandle->RegisterDelegates();
 	}
 	return ComponentHandle;
 }
@@ -648,3 +890,47 @@ bool FTransformConstraintUtils::AddConstraint(
 
 	return true;
 }
+
+FTransform FTransformConstraintUtils::ComputeRelativeTransform(
+	const FTransform& InChildLocal,
+	const FTransform& InChildWorld,
+	const FTransform& InSpaceWorld,
+	const ETransformConstraintType InType)
+{
+	switch (InType)
+	{
+	case ETransformConstraintType::Translation:
+		{
+			FTransform RelativeTransform = InChildLocal;
+			RelativeTransform.SetLocation(InChildWorld.GetLocation() - InSpaceWorld.GetLocation());
+			return RelativeTransform;
+		}
+	case ETransformConstraintType::Rotation:
+		{
+			FTransform RelativeTransform = InChildLocal;
+			FQuat RelativeRotation = InSpaceWorld.GetRotation().Inverse() * InChildWorld.GetRotation();
+			RelativeRotation.Normalize();
+			RelativeTransform.SetRotation(RelativeRotation);
+			return RelativeTransform;
+		}
+	case ETransformConstraintType::Scale:
+		{
+			FTransform RelativeTransform = InChildLocal;
+			const FVector SpaceScale = InSpaceWorld.GetScale3D();
+			FVector RelativeScale = InChildWorld.GetScale3D();
+			RelativeScale[0] = FMath::Abs(SpaceScale[0]) > KINDA_SMALL_NUMBER ? RelativeScale[0] / SpaceScale[0] : 0.f;
+			RelativeScale[1] = FMath::Abs(SpaceScale[1]) > KINDA_SMALL_NUMBER ? RelativeScale[1] / SpaceScale[1] : 0.f;
+			RelativeScale[2] = FMath::Abs(SpaceScale[2]) > KINDA_SMALL_NUMBER ? RelativeScale[2] / SpaceScale[2] : 0.f;
+			RelativeTransform.SetScale3D(RelativeScale);
+			return RelativeTransform;
+		}
+	case ETransformConstraintType::Parent:
+		return InChildWorld.GetRelativeTransform(InSpaceWorld);
+	case ETransformConstraintType::LookAt:
+		return InChildLocal;
+	default:
+		break;
+	}
+	
+	return InChildWorld.GetRelativeTransform(InSpaceWorld);
+}
\ No newline at end of file
diff --git a/Engine/Source/Runtime/Experimental/Animation/Constraints/Private/TransformableHandle.cpp b/Engine/Source/Runtime/Experimental/Animation/Constraints/Private/TransformableHandle.cpp
--- a/Engine/Source/Runtime/Experimental/Animation/Constraints/Private/TransformableHandle.cpp
+++ b/Engine/Source/Runtime/Experimental/Animation/Constraints/Private/TransformableHandle.cpp
@@ -5,19 +5,27 @@
 
 #include "Components/SceneComponent.h"
 #include "GameFramework/Actor.h"
+#include "Engine/Engine.h"
 
 /**
  * UTransformableHandle
  */
 
 UTransformableHandle::~UTransformableHandle()
-{}
+{
+	OnHandleModified.Clear();
+}
 
 void UTransformableHandle::PostLoad()
 {
 	Super::PostLoad();
 }
 
+UTransformableHandle::FHandleModifiedEvent& UTransformableHandle::HandleModified()
+{
+	return OnHandleModified;
+}
+
 /**
  * UTransformableComponentHandle
  */
@@ -76,6 +84,88 @@ TWeakObjectPtr<UObject> UTransformableComponentHandle::GetTarget() const
 	return Component;
 }
 
+void UTransformableComponentHandle::UnregisterDelegates() const
+{
+#if WITH_EDITOR
+	FCoreUObjectDelegates::OnObjectPropertyChanged.RemoveAll(this);
+	GEngine->OnActorMoving().RemoveAll(this);
+#endif
+}
+
+void UTransformableComponentHandle::RegisterDelegates()
+{
+	UnregisterDelegates();
+
+#if WITH_EDITOR
+	GEngine->OnActorMoving().AddUObject(this, &UTransformableComponentHandle::OnActorMoving);
+	FCoreUObjectDelegates::OnObjectPropertyChanged.AddUObject(this, &UTransformableComponentHandle::OnPostPropertyChanged);
+#endif
+}
+
+void UTransformableComponentHandle::OnActorMoving(AActor* InActor)
+{
+	if (!Component.IsValid())
+	{
+		return;
+	}
+
+	const USceneComponent* SceneComponent = InActor ? InActor->GetRootComponent() : nullptr;
+	if (SceneComponent != Component)
+	{
+		return;
+	}
+
+	if(OnHandleModified.IsBound())
+	{
+		OnHandleModified.Broadcast(this, true);
+	}
+}
+
+void UTransformableComponentHandle::OnPostPropertyChanged(
+	UObject* InObject,
+	FPropertyChangedEvent& InPropertyChangedEvent)
+{
+	if (!Component.IsValid())
+	{
+		return;
+	}
+
+	USceneComponent* SceneComponent = Cast<USceneComponent>(InObject);
+	if (!SceneComponent)
+	{
+		if (const AActor* Actor = Cast<AActor>(InObject))
+		{
+			SceneComponent = Actor->GetRootComponent();
+		}
+	}
+	
+	if (SceneComponent!= Component)
+	{
+		return;
+	}
+
+	const FProperty* MemberProperty = InPropertyChangedEvent.MemberProperty;
+	if (!MemberProperty)
+	{
+		return;
+	}
+	
+	const FName MemberPropertyName = MemberProperty->GetFName();
+	const bool bTransformationChanged =
+		(MemberPropertyName == USceneComponent::GetRelativeLocationPropertyName() ||
+			MemberPropertyName == USceneComponent::GetRelativeRotationPropertyName() ||
+			MemberPropertyName == USceneComponent::GetRelativeScale3DPropertyName());
+	if (!bTransformationChanged)
+	{
+		return;
+	}
+
+	if(OnHandleModified.IsBound())
+	{
+		OnHandleModified.Broadcast(this, true);
+	}
+}
+
 #if WITH_EDITOR
 FName UTransformableComponentHandle::GetName() const
 {
diff --git a/Engine/Source/Runtime/Experimental/Animation/Constraints/Public/TransformConstraint.h b/Engine/Source/Runtime/Experimental/Animation/Constraints/Public/TransformConstraint.h
--- a/Engine/Source/Runtime/Experimental/Animation/Constraints/Public/TransformConstraint.h
+++ b/Engine/Source/Runtime/Experimental/Animation/Constraints/Public/TransformConstraint.h
@@ -26,70 +26,93 @@ class CONSTRAINTS_API UTickableTransformConstraint : public UTickableConstraint
 
 public:
 
-	/** @todo document */
+	/** Sets up the constraint so that the initial offset is set and dependencies and handles managed. */
 	void Setup();
-	/** @todo document */
+	
+	/** UObjects overrides. */
 	virtual void PostLoad() override;
-	/** @todo document */
 	virtual void PostDuplicate(bool bDuplicateForPIE) override;
 
-	/** @todo document */
+	/** Returns the target hash value (i.e. the child handle's hash). */
 	virtual uint32 GetTargetHash() const override;
-	/** @todo document */
+	
+	/** Test whether an InObject is referenced by that constraint. (i.e. is it's parent or child). */
 	virtual bool ReferencesObject(TWeakObjectPtr<UObject> InObject) const override;
 	
-	/** @todo document */
+	/** The transformable handle representing the parent of that constraint. */
 	UPROPERTY(BlueprintReadWrite, Category = "Handle")
 	TObjectPtr<UTransformableHandle> ParentTRSHandle;
-	/** @todo document */
+	
+	/** The transformable handle representing the child of that constraint. */
 	UPROPERTY(BlueprintReadWrite, Category = "Handle")
 	TObjectPtr<UTransformableHandle> ChildTRSHandle;
 
-	/** @todo document */
+	/** Should that constraint maintain the default offset.  */
 	UPROPERTY(EditAnywhere, BlueprintReadWrite, Category="Offset")
 	bool bMaintainOffset = true;
 
-	/** @todo document */
+	/** Defines how much the constraint will be applied. */
 	// UPROPERTY(EditAnywhere, BlueprintReadWrite,Category="Weight", meta = (Input, ClampMin = "0.0", ClampMax = "1.0", UIMin = "0.0", UIMax = "1.0"))
 	//@benoit when not EditAnywhere?
 	UPROPERTY(meta = (Input, ClampMin = "0.0", ClampMax = "1.0", UIMin = "0.0", UIMax = "1.0"))
 	float Weight = 1.f;
 
+	/** Should the child be able to change it's offset dynamically. */
 	UPROPERTY(EditAnywhere, BlueprintReadWrite, Category="Offset")
 	bool bDynamicOffset = false;
 
-	/** @todo document. */
+	/** Returns the constraint type (Position, Parent, Aim...). */
 	int64 GetType() const;
 
-	/** @todo document */
+	/** Get the current child's global transform. */
 	FTransform GetChildGlobalTransform() const;
-	/** @todo document */
+	
+	/** Get the current child's local transform. */
 	FTransform GetChildLocalTransform() const;
-		
+
+	/** Get the current parent's global transform. */
+	FTransform GetParentGlobalTransform() const;
+	
+	/** Get the current parent's local transform. */
+	FTransform GetParentLocalTransform() const;
+	
 protected:
-	/** @todo document */
+
+	/** Registers/Unregisters useful delegates for both child and parent handles. */
+	void UnregisterDelegates() const;
+	void RegisterDelegates();
+	
+	/**
+	 * Manages changes on the child/parent transformable handle. This can be used to update internal data (i.e. offset)
+	 * when transform changes outside of that system and need to trigger changes within the constraint itself.   
+	*/
+	virtual void OnHandleModified(UTransformableHandle* InHandle, bool bUpdate);
+	
+	/**
+	 * Computes the initial offset that is needed to keep the child's global transform unchanged when creating the
+	 * constraint. This can be called whenever necessary to update the current offset.
+	*/
 	virtual void ComputeOffset() PURE_VIRTUAL(ComputeOffset, return;);
 	
-	/** @todo document */
+	/**
+	 * Sets up dependencies between the parent, the constraint and the child using their respective tick functions.
+	 * It creates a dependency graph between them so that they tick in the right order when evaluated.   
+	*/
 	void SetupDependencies();
 	
-	/** @todo document */
+	/** Set the current child's global transform. */
 	void SetChildGlobalTransform(const FTransform& InGlobal) const;
-	/** @todo document */
+	
+	/** Set the current child's local transform. */
 	void SetChildLocalTransform(const FTransform& InLocal) const;
 
-	/** @todo document */
-	FTransform GetParentGlobalTransform() const;
-	/** @todo document */
-	FTransform GetParentLocalTransform() const;
-
-	/** @todo document */
+	/** Defines the constraint's type (Position, Parent, Aim...). */
 	UPROPERTY()
 	ETransformConstraintType Type = ETransformConstraintType::Parent;
 
 #if WITH_EDITOR
 public:
-	/** @todo document */
+	/** Returns the constraint's label used for UI. */
 	virtual FName GetLabel() const override;
 
 	// UObject interface
@@ -110,16 +133,39 @@ class CONSTRAINTS_API UTickableTranslationConstraint : public UTickableTransform
 public:
 	UTickableTranslationConstraint();
 
-	/** @todo document */
+	/** Returns the position constraint function that the tick function will evaluate. */
 	virtual FConstraintTickFunction::ConstraintFunction GetFunction() const override;
 
 protected:
-	/** @todo document */
+
+	/** Cache data structure to store last child local/global location. */
+	struct FDynamicCache
+	{
+		uint32 CachedInputHash = 0;
+	};
+	mutable FDynamicCache Cache;
+	uint32 CalculateInputHash() const;
+
+	/** Updates the dynamic offset based on external child's transform changes. */
+	virtual void OnHandleModified(UTransformableHandle* InHandle, bool bUpdate) override;
+	
+	/** Computes the child's local translation offset in the parent space. */
 	virtual void ComputeOffset() override;
 
-	/** @todo document */
+	/** Defines the local child's translation offset in the parent space. */
 	UPROPERTY(EditAnywhere, BlueprintReadWrite, Category="Offset")
 	FVector OffsetTranslation = FVector::ZeroVector;
+
+	/** Defines the local child's translation dynamic offset in the parent space. */
+	UPROPERTY(EditAnywhere, BlueprintReadWrite, Category="Offset")
+	FVector DynamicOffsetTranslation = FVector::ZeroVector;
+
+#if WITH_EDITOR
+public:
+	// UObject interface
+	virtual void PostEditChangeProperty(FPropertyChangedEvent& PropertyChangedEvent) override;
+	// End of UObject interface
+#endif
 };
 
 /** 
@@ -134,16 +180,39 @@ class CONSTRAINTS_API UTickableRotationConstraint : public UTickableTransformCon
 public:
 	UTickableRotationConstraint();
 
-	/** @todo document */
+	/** Returns the rotation constraint function that the tick function will evaluate. */
 	virtual FConstraintTickFunction::ConstraintFunction GetFunction() const override;
 
 protected:
-	/** @todo document */
+
+	/** Cache data structure to store last child local/global rotation. */
+	struct FDynamicCache
+	{
+		uint32 CachedInputHash = 0;
+	};
+	mutable FDynamicCache Cache;
+	uint32 CalculateInputHash() const;
+	
+	/** Computes the child's local rotation offset in the parent space. */
 	virtual void ComputeOffset() override;
+
+	/** Updates the dynamic offset based on external child's transform changes. */
+	virtual void OnHandleModified(UTransformableHandle* InHandle, bool bUpdate) override;
 	
-	/** @todo document */
-	UPROPERTY(EditAnywhere, BlueprintReadWrite,Category="Offset")
+	/** Defines the local child's rotation offset in the parent space. */
+	UPROPERTY(EditAnywhere, BlueprintReadWrite, Category="Offset")
 	FQuat OffsetRotation = FQuat::Identity;
+
+	/** Defines the local child's rotation dynamic offset in the parent space. */
+	UPROPERTY(EditAnywhere, BlueprintReadWrite, Category="Offset")
+	FQuat DynamicOffsetRotation = FQuat::Identity;
+
+#if WITH_EDITOR
+public:
+	// UObject interface
+	virtual void PostEditChangeProperty(FPropertyChangedEvent& PropertyChangedEvent) override;
+	// End of UObject interface
+#endif
 };
 
 /** 
@@ -158,15 +227,15 @@ class CONSTRAINTS_API UTickableScaleConstraint : public UTickableTransformConstr
 public:
 	UTickableScaleConstraint();
 	
-	/** @todo document */
+	/** Returns the scale constraint function that the tick function will evaluate. */
 	virtual FConstraintTickFunction::ConstraintFunction GetFunction() const override;
 
 protected:
-	/** @todo document */
+	/** Computes the child's local scale offset in the parent space. */
 	virtual void ComputeOffset() override;
 
-	/** @todo document */
-	UPROPERTY(EditAnywhere, BlueprintReadWrite,Category="Offset")
+	/** Defines the local child's scale offset in the parent space. */
+	UPROPERTY(EditAnywhere, BlueprintReadWrite, Category="Offset")
 	FVector OffsetScale = FVector::OneVector;
 };
 
@@ -182,28 +251,38 @@ class CONSTRAINTS_API UTickableParentConstraint : public UTickableTransformConst
 public:
 	UTickableParentConstraint();
 	
-	/** @todo document */
+	/** Returns the transform constraint function that the tick function will evaluate. */
 	virtual FConstraintTickFunction::ConstraintFunction GetFunction() const override;
 
 protected:
-	/** @todo document */
-	virtual void ComputeOffset() override;
-
-	/** @todo document */
-	UPROPERTY(EditAnywhere, BlueprintReadWrite, Category="Offset")
-	FTransform OffsetTransform = FTransform::Identity;
-
-	/** @todo document */
+	/** Cache data structure to store last child local/global transform. */
 	struct FDynamicCache
 	{
-		FTransform LastGlobalSet = FTransform::Identity;
-		FTransform LastLocalSet = FTransform::Identity;
 		uint32 CachedInputHash = 0;
 	};
 	mutable FDynamicCache Cache;
+	uint32 CalculateInputHash() const;
 
-	/** @todo document */
-	uint32 CalculateDependenciesHash() const;
+	/** Updates the dynamic offset based on external child's transform changes. */
+	virtual void OnHandleModified(UTransformableHandle* InHandle, bool bUpdate) override;
+	
+	/** Computes the child's local transform offset in the parent space. */
+	virtual void ComputeOffset() override;
+
+	/** Defines the local child's transform offset in the parent space. */
+	UPROPERTY(EditAnywhere, BlueprintReadWrite, Category="Offset")
+	FTransform OffsetTransform = FTransform::Identity;
+
+	/** Defines the local child's dynamic transform offset in the parent space. */
+	UPROPERTY(EditAnywhere, BlueprintReadWrite, Category="Offset")
+	FTransform DynamicOffsetTransform = FTransform::Identity;
+
+#if WITH_EDITOR
+public:
+	// UObject interface
+	virtual void PostEditChangeProperty(FPropertyChangedEvent& PropertyChangedEvent) override;
+	// End of UObject interface
+#endif
 };
 
 /** 
@@ -218,16 +297,21 @@ class CONSTRAINTS_API UTickableLookAtConstraint : public UTickableTransformConst
 public:
 	UTickableLookAtConstraint();
 	
-	/** @todo document */
+	/** Returns the look at constraint function that the tick function will evaluate. */
 	virtual FConstraintTickFunction::ConstraintFunction GetFunction() const override;
 
 protected:
 	/** @todo document */
 	virtual void ComputeOffset() override;
 
-	/** @todo document */
-	UPROPERTY(EditAnywhere, BlueprintReadWrite,Category="Offset")
-	FTransform OffsetTransform = FTransform::Identity;
+	/** Defines the aiming axis. */
+	UPROPERTY(EditAnywhere, BlueprintReadWrite, Category="Axis")
+	FVector Axis = FVector::XAxisVector;
+
+private:
+
+	/** Computes the shortest quaternion between A and B. */
+	static FQuat FindQuatBetweenNormals(const FVector& A, const FVector& B);
 };
 
 /** 
@@ -236,8 +320,7 @@ protected:
 
 struct CONSTRAINTS_API FTransformConstraintUtils
 {
-
-	/** @todo document */
+	/** Fills a sorted constraint array that InChild actor is the child of. */
 	static void GetParentConstraints(
 		UWorld* World,
 		const AActor* InChild,
@@ -248,12 +331,12 @@ struct CONSTRAINTS_API FTransformConstraintUtils
 		USceneComponent* InSceneComponent, 
 		UObject* Outer);
 
-	/** @todo document */
+	/** Creates a new transform constraint based on the InType. */
 	static UTickableTransformConstraint* CreateFromType(
 		UWorld* InWorld,
 		const ETransformConstraintType InType);
 
-	/** @todo document */
+	/** Creates respective handles and creates a new InType transform constraint. */	
 	static UTickableTransformConstraint* CreateAndAddFromActors(
 		UWorld* InWorld,
 		AActor* InParent,
@@ -268,4 +351,11 @@ struct CONSTRAINTS_API FTransformConstraintUtils
 		UTransformableHandle* InChildHandle,
 		UTickableTransformConstraint* Constraint,
 		const bool bMaintainOffset = true);
+
+	/** Computes the relative transform between both transform based on the constraint's InType. */
+	static FTransform ComputeRelativeTransform(
+		const FTransform& InChildLocal,
+		const FTransform& InChildWorld,
+		const FTransform& InSpaceWorld,
+		const ETransformConstraintType InType);
 };
diff --git a/Engine/Source/Runtime/Experimental/Animation/Constraints/Public/TransformableHandle.h b/Engine/Source/Runtime/Experimental/Animation/Constraints/Public/TransformableHandle.h
--- a/Engine/Source/Runtime/Experimental/Animation/Constraints/Public/TransformableHandle.h
+++ b/Engine/Source/Runtime/Experimental/Animation/Constraints/Public/TransformableHandle.h
@@ -19,6 +19,9 @@ class CONSTRAINTS_API UTransformableHandle : public UObject
 	GENERATED_BODY()
 	
 public:
+
+	DECLARE_EVENT_TwoParams(UTransformableHandle, FHandleModifiedEvent, UTransformableHandle*, bool);
+	
 	virtual ~UTransformableHandle();
 
 	virtual void PostLoad() override;
@@ -30,7 +33,7 @@ public:
 	virtual void SetGlobalTransform(const FTransform& InGlobal) const PURE_VIRTUAL(SetGlobalTransform, );
 	/** Sets the local transform of the underlying transformable object in it's parent space. */
 	virtual void SetLocalTransform(const FTransform& InLocal) const PURE_VIRTUAL(SetLocalTransform, );
-	/** Gets the transform of the underlying transformable object. */
+	/** Gets the global transform of the underlying transformable object. */
 	virtual FTransform GetGlobalTransform() const PURE_VIRTUAL(GetGlobalTransform, return FTransform::Identity;);
 	/** Gets the local transform of the underlying transformable object in it's parent space. */
 	virtual FTransform GetLocalTransform() const PURE_VIRTUAL(GetLocalTransform, return FTransform::Identity;);
@@ -46,17 +49,20 @@ public:
 	**/
 	virtual FTickFunction* GetTickFunction() const PURE_VIRTUAL(GetTickFunction, return nullptr;);
 
-	/**
-	 * Generates a hash value of the underlying transformable object.
-	**/
+	/** Generates a hash value of the underlying transformable object. */
 	virtual uint32 GetHash() const PURE_VIRTUAL(GetHash, return 0;);
 
 	/** @todo document */
 	virtual TWeakObjectPtr<UObject> GetTarget() const PURE_VIRTUAL(GetTarget, return nullptr;);
 
+	FHandleModifiedEvent& HandleModified();
+
 #if WITH_EDITOR
 	virtual FName GetName() const PURE_VIRTUAL(GetName, return NAME_None;);
 #endif
+
+protected:
+	FHandleModifiedEvent OnHandleModified;
 };
 
 /**
@@ -72,35 +78,42 @@ public:
 	
 	virtual ~UTransformableComponentHandle();
 	
-	/** @todo document */
+	/** Sanity check to ensure that Component. */
 	virtual bool IsValid() const override;
 	
-	/** @todo document */
+	/** Sets the global transform of Component. */
 	virtual void SetGlobalTransform(const FTransform& InGlobal) const override;
-	/** Sets the local transform of the underlying transformable object in it's parent space. */
+	/** Sets the local transform of Component in it's attachment. */
 	virtual void SetLocalTransform(const FTransform& InLocal) const override;
-	/** @todo document */
+	/** Gets the global transform of Component. */
 	virtual FTransform GetGlobalTransform() const override;
-	/** @todo document */
+	/** Gets the local transform of Component in it's attachment. */
 	virtual FTransform GetLocalTransform() const override;
 
-	/** @todo document */
+	/** Returns the target object containing the tick function (e.i. Component). */
 	virtual UObject* GetPrerequisiteObject() const override;
-	/** @todo document */
+	/** Returns Component's tick function. */
 	virtual FTickFunction* GetTickFunction() const override;
 
-	/** @todo document */
+	/** Generates a hash value of Component. */
 	virtual uint32 GetHash() const override;
 
 	/** @todo document */
 	virtual TWeakObjectPtr<UObject> GetTarget() const override;
 
 #if WITH_EDITOR
-	/** @todo document */
 	virtual FName GetName() const override;
 #endif
 	
-	/** @todo document */
+	/** The Component that this handle is pointing at. */
 	UPROPERTY(BlueprintReadOnly, Category = "Object")
 	TWeakObjectPtr<USceneComponent> Component;
+
+	/** Registers/Unregisters useful delegates to track changes in the Component's transform. */
+	void UnregisterDelegates() const;
+	void RegisterDelegates();
+	
+	/** @todo document */
+	void OnActorMoving(AActor* InActor);
+	void OnPostPropertyChanged(UObject* InObject, FPropertyChangedEvent& InPropertyChangedEvent);
 };
\ No newline at end of file
