Sequencer: Anim Tools: Make sure snapper doens't clear on undo, unfortunately we get bad event sent saying sequence is changing when it's not.
#trivial

diff --git a/Engine/Plugins/Animation/ControlRig/Source/ControlRigEditor/Private/EditMode/SControlRigSnapper.cpp b/Engine/Plugins/Animation/ControlRig/Source/ControlRigEditor/Private/EditMode/SControlRigSnapper.cpp
--- a/Engine/Plugins/Animation/ControlRig/Source/ControlRigEditor/Private/EditMode/SControlRigSnapper.cpp
+++ b/Engine/Plugins/Animation/ControlRig/Source/ControlRigEditor/Private/EditMode/SControlRigSnapper.cpp
@@ -386,7 +386,10 @@ SControlRigSnapper::~SControlRigSnapper()
 
 void SControlRigSnapper::OnActivateSequenceChanged(FMovieSceneSequenceIDRef ID)
 {
-	ClearActors();
+	if(!GIsTransacting)
+	{ 
+		ClearActors();
+	}
 }
 
 
