Control Rig: Increase default node run count

diff --git a/Engine/Plugins/Animation/ControlRig/Source/ControlRigDeveloper/Public/ControlRigBlueprint.h b/Engine/Plugins/Animation/ControlRig/Source/ControlRigDeveloper/Public/ControlRigBlueprint.h
--- a/Engine/Plugins/Animation/ControlRig/Source/ControlRigDeveloper/Public/ControlRigBlueprint.h
+++ b/Engine/Plugins/Animation/ControlRig/Source/ControlRigDeveloper/Public/ControlRigBlueprint.h
@@ -116,7 +116,7 @@ struct CONTROLRIGDEVELOPER_API FRigGraphDisplaySettings
 		: bShowNodeInstructionIndex(false)
 		, bShowNodeRunCounts(false)
 		, NodeRunLowerBound(1)
-		, NodeRunLimit(64)
+		, NodeRunLimit(256)
 		, MinMicroSeconds(0.0)
 		, MaxMicroSeconds(1.0)
 		, TotalMicroSeconds(0.0)
