Control Rig: Hierarchy Metadata

Implemented using the new dispatch factories

diff --git a/Engine/Plugins/Animation/ControlRig/Source/ControlRig/Private/Rigs/RigHierarchyMetadata.cpp b/Engine/Plugins/Animation/ControlRig/Source/ControlRig/Private/Rigs/RigHierarchyMetadata.cpp
--- a/Engine/Plugins/Animation/ControlRig/Source/ControlRig/Private/Rigs/RigHierarchyMetadata.cpp
+++ b/Engine/Plugins/Animation/ControlRig/Source/ControlRig/Private/Rigs/RigHierarchyMetadata.cpp
@@ -82,6 +82,14 @@ UScriptStruct* FRigBaseMetadata::GetMetadataStruct(const ERigMetadataType& InTyp
 		{
 			return FRigTransformArrayMetadata::StaticStruct();
 		}
+		case ERigMetadataType::LinearColor:
+		{
+			return FRigLinearColorMetadata::StaticStruct();
+		}
+		case ERigMetadataType::LinearColorArray:
+		{
+			return FRigLinearColorArrayMetadata::StaticStruct();
+		}
 		case ERigMetadataType::RigElementKey:
 		{
 			return FRigElementKeyMetadata::StaticStruct();
diff --git a/Engine/Plugins/Animation/ControlRig/Source/ControlRig/Private/Units/Hierarchy/RigUnit_Metadata.cpp b/Engine/Plugins/Animation/ControlRig/Source/ControlRig/Private/Units/Hierarchy/RigUnit_Metadata.cpp
new file mode 100644
--- /dev/null
+++ b/Engine/Plugins/Animation/ControlRig/Source/ControlRig/Private/Units/Hierarchy/RigUnit_Metadata.cpp
@@ -0,0 +1,403 @@
+// Copyright Epic Games, Inc. All Rights Reserved.
+
+#include "RigUnit_Metadata.h"
+#include "RigVMTypeUtils.h"
+#include "RigVMCore/RigVMStruct.h"
+#include "Units/RigUnitContext.h"
+
+FName FRigDispatch_MetadataBase::ItemArgName = TEXT("Item");
+FName FRigDispatch_MetadataBase::NameArgName = TEXT("Name");
+FName FRigDispatch_MetadataBase::CacheArgName = TEXT("Cache");
+FName FRigDispatch_MetadataBase::DefaultArgName = TEXT("Default");
+FName FRigDispatch_MetadataBase::ValueArgName = TEXT("Value");
+FName FRigDispatch_MetadataBase::FoundArgName = TEXT("Found");
+FName FRigDispatch_MetadataBase::SuccessArgName = TEXT("Success");
+
+#if WITH_EDITOR
+
+FString FRigDispatch_MetadataBase::GetNodeTitle(const FRigVMTemplateTypeMap& InTypes) const
+{
+	if(const TRigVMTypeIndex* ValueTypeIndexPtr = InTypes.Find(ValueArgName))
+	{
+		const TRigVMTypeIndex& ValueTypeIndex = *ValueTypeIndexPtr;
+		if(ValueTypeIndex != RigVMTypeUtils::TypeIndex::WildCard &&
+			ValueTypeIndex != RigVMTypeUtils::TypeIndex::WildCardArray)
+		{
+			static constexpr TCHAR GetMetadataFormat[] = TEXT("Get %s Metadata");
+			static constexpr TCHAR SetMetadataFormat[] = TEXT("Set %s Metadata");
+
+			const FRigVMTemplateArgumentType& ValueType = FRigVMRegistry::Get().GetType(ValueTypeIndex);
+			FString ValueName;
+			if(ValueType.CPPTypeObject == FRigElementKey::StaticStruct())
+			{
+				ValueName = ItemArgName.ToString();
+			}
+			else if(ValueType.CPPTypeObject)
+			{
+				ValueName = ValueType.CPPTypeObject->GetName();
+			}
+			else if(ValueTypeIndex == RigVMTypeUtils::TypeIndex::FName)
+			{
+				ValueName = NameArgName.ToString();
+			}
+			else
+			{
+				ValueName = ValueType.GetBaseCPPType();
+				ValueName = ValueName.Left(1).ToUpper() + ValueName.Mid(1);
+			}
+
+			if(ValueType.IsArray())
+			{
+				ValueName += TEXT(" Array");
+			}
+
+			return FString::Printf(IsSetMetadata() ? SetMetadataFormat : GetMetadataFormat, *ValueName); 
+		}
+	}
+	return FRigDispatchFactory::GetNodeTitle(InTypes);
+}
+
+#endif
+
+TArray<FRigVMTemplateArgument> FRigDispatch_MetadataBase::GetArguments() const
+{
+	if(Arguments.IsEmpty())
+	{
+		if(IsSetMetadata())
+		{
+			ExecuteArgIndex = Arguments.Emplace(FRigVMStruct::ExecuteContextName, ERigVMPinDirection::IO, FRigVMRegistry::Get().GetTypeIndex<FControlRigExecuteContext>());
+		}
+		ItemArgIndex = Arguments.Emplace(ItemArgName, ERigVMPinDirection::Input, FRigVMRegistry::Get().GetTypeIndex<FRigElementKey>());
+		NameArgIndex = Arguments.Emplace(NameArgName, ERigVMPinDirection::Input, RigVMTypeUtils::TypeIndex::FName);
+		CacheArgIndex = Arguments.Emplace(CacheArgName, ERigVMPinDirection::Hidden, FRigVMRegistry::Get().GetTypeIndex<FCachedRigElement>());
+	};
+	return Arguments;
+}
+
+#if WITH_EDITOR
+
+FText FRigDispatch_MetadataBase::GetArgumentTooltip(const FName& InArgumentName, TRigVMTypeIndex InTypeIndex) const
+{
+	if(InArgumentName == ItemArgName)
+	{
+		return NSLOCTEXT("FRigDispatch_MetadataBase", "ItemArgTooltip", "The item storing the metadata");
+	}
+	if(InArgumentName == NameArgName)
+	{
+		return NSLOCTEXT("FRigDispatch_MetadataBase", "NameArgTooltip", "The name of the metadata");
+	}
+	if(InArgumentName == DefaultArgName)
+	{
+		return NSLOCTEXT("FRigDispatch_MetadataBase", "DefaultArgTooltip", "The default value used as a fallback if the metadata does not exist");
+	}
+	if(InArgumentName == ValueArgName)
+	{
+		return NSLOCTEXT("FRigDispatch_MetadataBase", "ValueArgTooltip", "The value to get / set");
+	}
+	if(InArgumentName == FoundArgName)
+	{
+		return NSLOCTEXT("FRigDispatch_MetadataBase", "FoundArgTooltip", "Returns true if the metadata exists with the specific type");
+	}
+	if(InArgumentName == SuccessArgName)
+	{
+		return NSLOCTEXT("FRigDispatch_MetadataBase", "SuccessArgTooltip", "Returns true if the metadata was successfully stored");
+	}
+	return FRigDispatchFactory::GetArgumentTooltip(InArgumentName, InTypeIndex);
+}
+
+#endif
+
+const TArray<TRigVMTypeIndex>& FRigDispatch_MetadataBase::GetValueTypes() const
+{
+	static TArray<TRigVMTypeIndex> Types;
+	if(Types.IsEmpty())
+	{
+		const FRigVMRegistry& Registry = FRigVMRegistry::Get();
+		Types = {
+			RigVMTypeUtils::TypeIndex::Bool,
+			RigVMTypeUtils::TypeIndex::Float,
+			RigVMTypeUtils::TypeIndex::Int32,
+			RigVMTypeUtils::TypeIndex::FName,
+			Registry.GetTypeIndex<FVector>(false),
+			Registry.GetTypeIndex<FRotator>(false),
+			Registry.GetTypeIndex<FQuat>(false),
+			Registry.GetTypeIndex<FTransform>(false),
+			Registry.GetTypeIndex<FLinearColor>(false),
+			Registry.GetTypeIndex<FRigElementKey>(false),
+			RigVMTypeUtils::TypeIndex::BoolArray,
+			RigVMTypeUtils::TypeIndex::FloatArray,
+			RigVMTypeUtils::TypeIndex::Int32Array,
+			RigVMTypeUtils::TypeIndex::FNameArray,
+			Registry.GetTypeIndex<FVector>(true),
+			Registry.GetTypeIndex<FRotator>(true),
+			Registry.GetTypeIndex<FQuat>(true),
+			Registry.GetTypeIndex<FTransform>(true),
+			Registry.GetTypeIndex<FLinearColor>(true),
+			Registry.GetTypeIndex<FRigElementKey>(true)
+		};
+	}
+	return Types;
+}
+
+////////////////////////////////////////////////////////////////////////////////////////////////
+
+TArray<FRigVMTemplateArgument> FRigDispatch_GetMetadata::GetArguments() const
+{
+	if(ValueArgIndex == INDEX_NONE)
+	{
+		Arguments = Super::GetArguments(); 
+		DefaultArgIndex = Arguments.Emplace(DefaultArgName, ERigVMPinDirection::Input, GetValueTypes());
+		ValueArgIndex = Arguments.Emplace(ValueArgName, ERigVMPinDirection::Output, GetValueTypes());
+		FoundArgIndex = Arguments.Emplace(FoundArgName, ERigVMPinDirection::Output, RigVMTypeUtils::TypeIndex::Bool);
+	};
+	return Arguments;
+}
+
+FRigBaseMetadata* FRigDispatch_GetMetadata::FindMetadata(const FRigVMExtendedExecuteContext& InContext,
+                                                              const FRigElementKey& InKey, const FName& InName, ERigMetadataType InType, FCachedRigElement& Cache)
+{
+	const FRigUnitContext& Context = GetRigUnitContext(InContext);
+	if(Cache.UpdateCache(InKey, Context.Hierarchy))
+	{
+		if(FRigBaseElement* Element = Context.Hierarchy->Get(Cache.GetIndex()))
+		{
+			return Element->GetMetadata(InName, InType);
+		}
+	}
+	return nullptr;
+}
+
+FRigVMFunctionPtr FRigDispatch_GetMetadata::GetDispatchFunctionImpl(const FRigVMTemplateTypeMap& InTypes) const
+{
+	const FRigVMRegistry& Registry = FRigVMRegistry::Get();
+	const TRigVMTypeIndex& ValueTypeIndex = InTypes.FindChecked(TEXT("Value"));
+	
+	if(ValueTypeIndex == RigVMTypeUtils::TypeIndex::Bool)
+	{
+		return &FRigDispatch_GetMetadata::GetMetadataDispatch<bool, FRigBoolMetadata, ERigMetadataType::Bool>;
+	}
+	if(ValueTypeIndex == RigVMTypeUtils::TypeIndex::Float)
+	{
+		return &FRigDispatch_GetMetadata::GetMetadataDispatch<float, FRigFloatMetadata, ERigMetadataType::Float>;
+	}
+	if(ValueTypeIndex == RigVMTypeUtils::TypeIndex::Int32)
+	{
+		return &FRigDispatch_GetMetadata::GetMetadataDispatch<int32, FRigInt32Metadata, ERigMetadataType::Int32>;
+	}
+	if(ValueTypeIndex == RigVMTypeUtils::TypeIndex::FName)
+	{
+		return &FRigDispatch_GetMetadata::GetMetadataDispatch<FName, FRigNameMetadata, ERigMetadataType::Name>;
+	}
+	if(ValueTypeIndex == Registry.GetTypeIndex<FVector>(false))
+	{
+		return &FRigDispatch_GetMetadata::GetMetadataDispatch<FVector, FRigVectorMetadata, ERigMetadataType::Vector>;
+	}
+	if(ValueTypeIndex == Registry.GetTypeIndex<FRotator>(false))
+	{
+		return &FRigDispatch_GetMetadata::GetMetadataDispatch<FRotator, FRigRotatorMetadata, ERigMetadataType::Rotator>;
+	}
+	if(ValueTypeIndex == Registry.GetTypeIndex<FQuat>(false))
+	{
+		return &FRigDispatch_GetMetadata::GetMetadataDispatch<FQuat, FRigQuatMetadata, ERigMetadataType::Quat>;
+	}
+	if(ValueTypeIndex == Registry.GetTypeIndex<FTransform>(false))
+	{
+		return &FRigDispatch_GetMetadata::GetMetadataDispatch<FTransform, FRigTransformMetadata, ERigMetadataType::Transform>;
+	}
+	if(ValueTypeIndex == Registry.GetTypeIndex<FLinearColor>(false))
+	{
+		return &FRigDispatch_GetMetadata::GetMetadataDispatch<FLinearColor, FRigLinearColorMetadata, ERigMetadataType::LinearColor>;
+	}
+	if(ValueTypeIndex == Registry.GetTypeIndex<FRigElementKey>(false))
+	{
+		return &FRigDispatch_GetMetadata::GetMetadataDispatch<FRigElementKey, FRigElementKeyMetadata, ERigMetadataType::RigElementKey>;
+	}
+	if(ValueTypeIndex == RigVMTypeUtils::TypeIndex::BoolArray)
+	{
+		return &FRigDispatch_GetMetadata::GetMetadataDispatch<TArray<bool>, FRigBoolArrayMetadata, ERigMetadataType::BoolArray>;
+	}
+	if(ValueTypeIndex == RigVMTypeUtils::TypeIndex::FloatArray)
+	{
+		return &FRigDispatch_GetMetadata::GetMetadataDispatch<TArray<float>, FRigFloatArrayMetadata, ERigMetadataType::FloatArray>;
+	}
+	if(ValueTypeIndex == RigVMTypeUtils::TypeIndex::Int32Array)
+	{
+		return &FRigDispatch_GetMetadata::GetMetadataDispatch<TArray<int32>, FRigInt32ArrayMetadata, ERigMetadataType::Int32Array>;
+	}
+	if(ValueTypeIndex == RigVMTypeUtils::TypeIndex::FNameArray)
+	{
+		return &FRigDispatch_GetMetadata::GetMetadataDispatch<TArray<FName>, FRigNameArrayMetadata, ERigMetadataType::NameArray>;
+	}
+	if(ValueTypeIndex == Registry.GetTypeIndex<FVector>(true))
+	{
+		return &FRigDispatch_GetMetadata::GetMetadataDispatch<TArray<FVector>, FRigVectorArrayMetadata, ERigMetadataType::VectorArray>;
+	}
+	if(ValueTypeIndex == Registry.GetTypeIndex<FRotator>(true))
+	{
+		return &FRigDispatch_GetMetadata::GetMetadataDispatch<TArray<FRotator>, FRigRotatorArrayMetadata, ERigMetadataType::RotatorArray>;
+	}
+	if(ValueTypeIndex == Registry.GetTypeIndex<FQuat>(true))
+	{
+		return &FRigDispatch_GetMetadata::GetMetadataDispatch<TArray<FQuat>, FRigQuatArrayMetadata, ERigMetadataType::QuatArray>;
+	}
+	if(ValueTypeIndex == Registry.GetTypeIndex<FTransform>(true))
+	{
+		return &FRigDispatch_GetMetadata::GetMetadataDispatch<TArray<FTransform>, FRigTransformArrayMetadata, ERigMetadataType::TransformArray>;
+	}
+	if(ValueTypeIndex == Registry.GetTypeIndex<FLinearColor>(true))
+	{
+		return &FRigDispatch_GetMetadata::GetMetadataDispatch<TArray<FLinearColor>, FRigLinearColorArrayMetadata, ERigMetadataType::LinearColorArray>;
+	}
+	if(ValueTypeIndex == Registry.GetTypeIndex<FRigElementKey>(true))
+	{
+		return &FRigDispatch_GetMetadata::GetMetadataDispatch<TArray<FRigElementKey>, FRigElementKeyArrayMetadata, ERigMetadataType::RigElementKeyArray>;
+	}
+
+	return nullptr;
+}
+
+////////////////////////////////////////////////////////////////////////////////////////////////
+
+TArray<FRigVMTemplateArgument> FRigDispatch_SetMetadata::GetArguments() const
+{
+	if(ValueArgIndex == INDEX_NONE)
+	{
+		Arguments = Super::GetArguments(); 
+		ValueArgIndex = Arguments.Emplace(ValueArgName, ERigVMPinDirection::Input, GetValueTypes());
+		SuccessArgIndex = Arguments.Emplace(SuccessArgName, ERigVMPinDirection::Output, RigVMTypeUtils::TypeIndex::Bool);
+	};
+	return Arguments;
+}
+
+FRigBaseMetadata* FRigDispatch_SetMetadata::FindOrAddMetadata(FControlRigExecuteContext& InContext,
+                                                              const FRigElementKey& InKey, const FName& InName, ERigMetadataType InType, FCachedRigElement& Cache)
+{
+	if(Cache.UpdateCache(InKey, InContext.Hierarchy))
+	{
+		if(FRigBaseElement* Element = InContext.Hierarchy->Get(Cache.GetIndex()))
+		{
+			return Element->SetupValidMetadata(InName, InType);
+		}
+	}
+	return nullptr;
+}
+
+FRigVMFunctionPtr FRigDispatch_SetMetadata::GetDispatchFunctionImpl(const FRigVMTemplateTypeMap& InTypes) const
+{
+	const FRigVMRegistry& Registry = FRigVMRegistry::Get();
+	const TRigVMTypeIndex& ValueTypeIndex = InTypes.FindChecked(TEXT("Value"));
+	
+	if(ValueTypeIndex == RigVMTypeUtils::TypeIndex::Bool)
+	{
+		return &FRigDispatch_SetMetadata::SetMetadataDispatch<bool, FRigBoolMetadata, ERigMetadataType::Bool>;
+	}
+	if(ValueTypeIndex == RigVMTypeUtils::TypeIndex::Float)
+	{
+		return &FRigDispatch_SetMetadata::SetMetadataDispatch<float, FRigFloatMetadata, ERigMetadataType::Float>;
+	}
+	if(ValueTypeIndex == RigVMTypeUtils::TypeIndex::Int32)
+	{
+		return &FRigDispatch_SetMetadata::SetMetadataDispatch<int32, FRigInt32Metadata, ERigMetadataType::Int32>;
+	}
+	if(ValueTypeIndex == RigVMTypeUtils::TypeIndex::FName)
+	{
+		return &FRigDispatch_SetMetadata::SetMetadataDispatch<FName, FRigNameMetadata, ERigMetadataType::Name>;
+	}
+	if(ValueTypeIndex == Registry.GetTypeIndex<FVector>(false))
+	{
+		return &FRigDispatch_SetMetadata::SetMetadataDispatch<FVector, FRigVectorMetadata, ERigMetadataType::Vector>;
+	}
+	if(ValueTypeIndex == Registry.GetTypeIndex<FRotator>(false))
+	{
+		return &FRigDispatch_SetMetadata::SetMetadataDispatch<FRotator, FRigRotatorMetadata, ERigMetadataType::Rotator>;
+	}
+	if(ValueTypeIndex == Registry.GetTypeIndex<FQuat>(false))
+	{
+		return &FRigDispatch_SetMetadata::SetMetadataDispatch<FQuat, FRigQuatMetadata, ERigMetadataType::Quat>;
+	}
+	if(ValueTypeIndex == Registry.GetTypeIndex<FTransform>(false))
+	{
+		return &FRigDispatch_SetMetadata::SetMetadataDispatch<FTransform, FRigTransformMetadata, ERigMetadataType::Transform>;
+	}
+	if(ValueTypeIndex == Registry.GetTypeIndex<FLinearColor>(false))
+	{
+		return &FRigDispatch_SetMetadata::SetMetadataDispatch<FLinearColor, FRigLinearColorMetadata, ERigMetadataType::LinearColor>;
+	}
+	if(ValueTypeIndex == Registry.GetTypeIndex<FRigElementKey>(false))
+	{
+		return &FRigDispatch_SetMetadata::SetMetadataDispatch<FRigElementKey, FRigElementKeyMetadata, ERigMetadataType::RigElementKey>;
+	}
+	if(ValueTypeIndex == RigVMTypeUtils::TypeIndex::BoolArray)
+	{
+		return &FRigDispatch_SetMetadata::SetMetadataDispatch<TArray<bool>, FRigBoolArrayMetadata, ERigMetadataType::BoolArray>;
+	}
+	if(ValueTypeIndex == RigVMTypeUtils::TypeIndex::FloatArray)
+	{
+		return &FRigDispatch_SetMetadata::SetMetadataDispatch<TArray<float>, FRigFloatArrayMetadata, ERigMetadataType::FloatArray>;
+	}
+	if(ValueTypeIndex == RigVMTypeUtils::TypeIndex::Int32Array)
+	{
+		return &FRigDispatch_SetMetadata::SetMetadataDispatch<TArray<int32>, FRigInt32ArrayMetadata, ERigMetadataType::Int32Array>;
+	}
+	if(ValueTypeIndex == RigVMTypeUtils::TypeIndex::FNameArray)
+	{
+		return &FRigDispatch_SetMetadata::SetMetadataDispatch<TArray<FName>, FRigNameArrayMetadata, ERigMetadataType::NameArray>;
+	}
+	if(ValueTypeIndex == Registry.GetTypeIndex<FVector>(true))
+	{
+		return &FRigDispatch_SetMetadata::SetMetadataDispatch<TArray<FVector>, FRigVectorArrayMetadata, ERigMetadataType::VectorArray>;
+	}
+	if(ValueTypeIndex == Registry.GetTypeIndex<FRotator>(true))
+	{
+		return &FRigDispatch_SetMetadata::SetMetadataDispatch<TArray<FRotator>, FRigRotatorArrayMetadata, ERigMetadataType::RotatorArray>;
+	}
+	if(ValueTypeIndex == Registry.GetTypeIndex<FQuat>(true))
+	{
+		return &FRigDispatch_SetMetadata::SetMetadataDispatch<TArray<FQuat>, FRigQuatArrayMetadata, ERigMetadataType::QuatArray>;
+	}
+	if(ValueTypeIndex == Registry.GetTypeIndex<FTransform>(true))
+	{
+		return &FRigDispatch_SetMetadata::SetMetadataDispatch<TArray<FTransform>, FRigTransformArrayMetadata, ERigMetadataType::TransformArray>;
+	}
+	if(ValueTypeIndex == Registry.GetTypeIndex<FLinearColor>(true))
+	{
+		return &FRigDispatch_SetMetadata::SetMetadataDispatch<TArray<FLinearColor>, FRigLinearColorArrayMetadata, ERigMetadataType::LinearColorArray>;
+	}
+	if(ValueTypeIndex == Registry.GetTypeIndex<FRigElementKey>(true))
+	{
+		return &FRigDispatch_SetMetadata::SetMetadataDispatch<TArray<FRigElementKey>, FRigElementKeyArrayMetadata, ERigMetadataType::RigElementKeyArray>;
+	}
+
+	return nullptr;
+}
+
+FRigUnit_HasMetadata_Execute()
+{
+	const URigHierarchy* Hierarchy = Context.Hierarchy;
+
+	bFound = false;
+
+	if(CachedIndex.UpdateCache(Item, Hierarchy))
+	{
+		if(const FRigBaseElement* Element = Hierarchy->Get(CachedIndex))
+		{
+			bFound = Element->GetMetadata(Name, Type) != nullptr;
+		}
+	}
+}
+
+FRigUnit_FindItemsWithMetadata_Execute()
+{
+	const URigHierarchy* Hierarchy = Context.Hierarchy;
+
+	Items.Reset();
+
+	Hierarchy->Traverse([&Items, Name, Type](const FRigBaseElement* Element, bool& bContinue)
+	{
+		if(Element->GetMetadata(Name, Type) != nullptr)
+		{
+			Items.Add(Element->GetKey());
+		}
+		bContinue = true;
+	});
+}
\ No newline at end of file
diff --git a/Engine/Plugins/Animation/ControlRig/Source/ControlRig/Private/Units/Hierarchy/RigUnit_Metadata.h b/Engine/Plugins/Animation/ControlRig/Source/ControlRig/Private/Units/Hierarchy/RigUnit_Metadata.h
new file mode 100644
--- /dev/null
+++ b/Engine/Plugins/Animation/ControlRig/Source/ControlRig/Private/Units/Hierarchy/RigUnit_Metadata.h
@@ -0,0 +1,248 @@
+// Copyright Epic Games, Inc. All Rights Reserved.
+
+#pragma once
+
+#include "Units/RigDispatchFactory.h"
+#include "RigUnit_Metadata.generated.h"
+
+USTRUCT(meta=(Abstract, Category="Hierarchy", NodeColor="0.462745, 1,0, 0.329412", Varying))
+struct CONTROLRIG_API FRigDispatch_MetadataBase : public FRigDispatchFactory
+{
+	GENERATED_BODY()
+
+#if WITH_EDITOR
+	virtual FString GetNodeTitle(const FRigVMTemplateTypeMap& InTypes) const override;;
+#endif
+	virtual TArray<FRigVMTemplateArgument> GetArguments() const override;
+	virtual bool IsSetMetadata() const { return false; }
+
+#if WITH_EDITOR
+	virtual FText GetArgumentTooltip(const FName& InArgumentName, TRigVMTypeIndex InTypeIndex) const override;
+#endif
+
+protected:
+
+	const TArray<TRigVMTypeIndex>& GetValueTypes() const;
+
+	mutable TArray<FRigVMTemplateArgument> Arguments;
+	
+	mutable int32 ExecuteArgIndex = INDEX_NONE;
+	mutable int32 ItemArgIndex = INDEX_NONE;
+	mutable int32 NameArgIndex = INDEX_NONE;
+	mutable int32 CacheArgIndex = INDEX_NONE;
+	mutable int32 DefaultArgIndex = INDEX_NONE;
+	mutable int32 ValueArgIndex = INDEX_NONE;
+	mutable int32 FoundArgIndex = INDEX_NONE;
+	mutable int32 SuccessArgIndex = INDEX_NONE;
+	static FName ItemArgName;
+	static FName NameArgName;
+	static FName CacheArgName;
+	static FName DefaultArgName;
+	static FName ValueArgName;
+	static FName FoundArgName;
+	static FName SuccessArgName;
+};
+
+/*
+ * Sets some metadata for the provided item
+ */
+USTRUCT(meta=(DisplayName="Get Metadata"))
+struct CONTROLRIG_API FRigDispatch_GetMetadata : public FRigDispatch_MetadataBase
+{
+	GENERATED_BODY()
+
+	virtual TArray<FRigVMTemplateArgument> GetArguments() const override;
+
+protected:
+
+	static FRigBaseMetadata* FindMetadata(const FRigVMExtendedExecuteContext& InContext, const FRigElementKey& InKey, const FName& InName, ERigMetadataType InType, FCachedRigElement& Cache);
+	virtual FRigVMFunctionPtr GetDispatchFunctionImpl(const FRigVMTemplateTypeMap& InTypes) const override;
+
+#if WITH_EDITOR
+	template<typename ValueType>
+	FORCEINLINE_DEBUGGABLE bool CheckArgumentTypes(FRigVMMemoryHandleArray Handles) const
+	{
+		return CheckArgumentType(Handles[ItemArgIndex].IsType<FRigElementKey>(), ItemArgName) &&
+			CheckArgumentType(Handles[NameArgIndex].IsType<FName>(), NameArgName) &&
+			CheckArgumentType(Handles[CacheArgIndex].IsType<FCachedRigElement>(true), CacheArgName) &&
+			CheckArgumentType(Handles[DefaultArgIndex].IsType<ValueType>(), DefaultArgName) &&
+			CheckArgumentType(Handles[ValueArgIndex].IsType<ValueType>(), ValueArgName) &&
+			CheckArgumentType(Handles[FoundArgIndex].IsType<bool>(), FoundArgName);
+	}
+#endif
+
+	template<typename ValueType, typename MetadataType, ERigMetadataType EnumValue>
+	FORCEINLINE_DEBUGGABLE static void GetMetadataDispatch(FRigVMExtendedExecuteContext& InContext, FRigVMMemoryHandleArray Handles)
+	{
+		const FRigDispatch_GetMetadata* Factory = static_cast<const FRigDispatch_GetMetadata*>(InContext.Factory);
+
+#if WITH_EDITOR
+		if(!Factory->CheckArgumentTypes<ValueType>(Handles))
+		{
+			return;
+		}
+#endif
+
+		// unpack the memory
+		const FRigElementKey& Item = *(const FRigElementKey*)Handles[Factory->ItemArgIndex].GetData();
+		const FName& Name = *(const FName*)Handles[Factory->NameArgIndex].GetData();
+		FCachedRigElement& Cache = *(FCachedRigElement*)Handles[Factory->CacheArgIndex].GetData(false, InContext.GetSlice().GetIndex());
+		const ValueType& Default = *(const ValueType*)Handles[Factory->DefaultArgIndex].GetData();
+		ValueType& Value = *(ValueType*)Handles[Factory->ValueArgIndex].GetData();
+		bool& Found = *(bool*)Handles[Factory->FoundArgIndex].GetData();
+
+		// extract the metadata
+		if (MetadataType* Md = Cast<MetadataType>(FindMetadata(InContext, Item, Name, EnumValue, Cache)))
+		{
+			Value = Md->GetValue();
+			Found = true;
+		}
+		else
+		{
+			Value = Default;
+			Found = false;
+		}
+	}
+};
+
+/*
+ * Returns some metadata for the provided item
+ */
+USTRUCT(meta=(DisplayName="Set Metadata"))
+struct CONTROLRIG_API FRigDispatch_SetMetadata : public FRigDispatch_MetadataBase
+{
+	GENERATED_BODY()
+
+	virtual TArray<FRigVMTemplateArgument> GetArguments() const override;
+	virtual bool IsSetMetadata() const override { return true; }
+
+protected:
+
+	static FRigBaseMetadata* FindOrAddMetadata(FControlRigExecuteContext& InContext, const FRigElementKey& InKey, const FName& InName, ERigMetadataType InType, FCachedRigElement& Cache);
+	virtual FRigVMFunctionPtr GetDispatchFunctionImpl(const FRigVMTemplateTypeMap& InTypes) const override;
+
+#if WITH_EDITOR
+	template<typename ValueType>
+	FORCEINLINE_DEBUGGABLE bool CheckArgumentTypes(FRigVMMemoryHandleArray Handles) const
+	{
+		return CheckArgumentType(Handles[ExecuteArgIndex].IsType<FControlRigExecuteContext>(), ItemArgName) &&
+			CheckArgumentType(Handles[ItemArgIndex].IsType<FRigElementKey>(), ItemArgName) &&
+			CheckArgumentType(Handles[NameArgIndex].IsType<FName>(), NameArgName) &&
+			CheckArgumentType(Handles[CacheArgIndex].IsType<FCachedRigElement>(true), CacheArgName) &&
+			CheckArgumentType(Handles[ValueArgIndex].IsType<ValueType>(), ValueArgName) &&
+			CheckArgumentType(Handles[SuccessArgIndex].IsType<bool>(), SuccessArgName);
+	}
+#endif
+	
+	template<typename ValueType, typename MetadataType, ERigMetadataType EnumValue>
+	FORCEINLINE_DEBUGGABLE static void SetMetadataDispatch(FRigVMExtendedExecuteContext& InContext, FRigVMMemoryHandleArray Handles)
+	{
+		const FRigDispatch_SetMetadata* Factory = static_cast<const FRigDispatch_SetMetadata*>(InContext.Factory);
+
+#if WITH_EDITOR
+		if(!Factory->CheckArgumentTypes<ValueType>(Handles))
+		{
+			return;
+		}
+#endif
+
+		// unpack the memory
+		FControlRigExecuteContext& ExecuteContext = *(FControlRigExecuteContext*)Handles[Factory->ExecuteArgIndex].GetData();
+		const FRigElementKey& Item = *(const FRigElementKey*)Handles[Factory->ItemArgIndex].GetData();
+		const FName& Name = *(const FName*)Handles[Factory->NameArgIndex].GetData();
+		FCachedRigElement& Cache = *(FCachedRigElement*)Handles[Factory->CacheArgIndex].GetData(false, InContext.GetSlice().GetIndex());
+		ValueType& Value = *(ValueType*)Handles[Factory->ValueArgIndex].GetData();
+		bool& Success = *(bool*)Handles[Factory->SuccessArgIndex].GetData();
+
+		// extract the metadata
+		if (MetadataType* Md = Cast<MetadataType>(FindOrAddMetadata(ExecuteContext, Item, Name, EnumValue, Cache)))
+		{
+			Md->GetValue() = Value;
+			Success = true;
+		}
+		else
+		{
+			Success = false;
+		}
+	}
+};
+
+/**
+ * Returns true if a given item in the hierarchy has a specific set of metadata
+ */
+USTRUCT(meta=(DisplayName="Has Metadata", Category="Hierarchy", DocumentationPolicy = "Strict", Keywords="MetadataExists,HasKey,SupportsMetadata", NodeColor="0.462745, 1,0, 0.329412", Varying))
+struct CONTROLRIG_API FRigUnit_HasMetadata : public FRigUnit
+{
+	GENERATED_BODY()
+
+	FRigUnit_HasMetadata()
+		: Item(NAME_None, ERigElementType::Bone)
+		, Name(NAME_None)
+		, Type(ERigMetadataType::Float)
+		, bFound(false)
+		, CachedIndex()
+	{}
+
+	RIGVM_METHOD()
+	virtual void Execute(const FRigUnitContext& Context) override;
+
+	/**
+	 * The item to check the metadata for
+	 */
+	UPROPERTY(meta = (Input, ExpandByDefault))
+	FRigElementKey Item;
+
+	/**
+	 * The name of the metadata to check
+	 */ 
+	UPROPERTY(meta = (Input))
+	FName Name;
+
+	/**
+	 * The type of metadata to check for
+	 */ 
+	UPROPERTY(meta = (Input))
+	ERigMetadataType Type;
+
+	// True if the item has the metadata
+	UPROPERTY(meta=(Output))
+	bool bFound;
+
+	// Used to cache the internally
+	UPROPERTY()
+	FCachedRigElement CachedIndex;
+};
+
+/**
+ * Returns all items containing a specific set of metadata
+ */
+USTRUCT(meta=(DisplayName="Find Items with Metadata", Category="Hierarchy", DocumentationPolicy = "Strict", Keywords="MetadataExists,HasKey,SupportsMetadata", NodeColor="0.462745, 1,0, 0.329412", Varying))
+struct CONTROLRIG_API FRigUnit_FindItemsWithMetadata : public FRigUnit
+{
+	GENERATED_BODY()
+
+	FRigUnit_FindItemsWithMetadata()
+		: Name(NAME_None)
+		, Type(ERigMetadataType::Float)
+		, Items()
+	{}
+
+	RIGVM_METHOD()
+	virtual void Execute(const FRigUnitContext& Context) override;
+
+	/**
+	 * The name of the metadata to find
+	 */ 
+	UPROPERTY(meta = (Input))
+	FName Name;
+
+	/**
+	 * The type of metadata to find
+	 */ 
+	UPROPERTY(meta = (Input))
+	ERigMetadataType Type;
+
+	// The items containing the metadata
+	UPROPERTY(meta=(Output))
+	TArray<FRigElementKey> Items;
+};
\ No newline at end of file
diff --git a/Engine/Plugins/Animation/ControlRig/Source/ControlRig/Public/Rigs/RigHierarchyDefines.h b/Engine/Plugins/Animation/ControlRig/Source/ControlRig/Public/Rigs/RigHierarchyDefines.h
--- a/Engine/Plugins/Animation/ControlRig/Source/ControlRig/Public/Rigs/RigHierarchyDefines.h
+++ b/Engine/Plugins/Animation/ControlRig/Source/ControlRig/Public/Rigs/RigHierarchyDefines.h
@@ -62,6 +62,8 @@ enum class ERigMetadataType : uint8
 	QuatArray,
 	Transform,
 	TransformArray,
+	LinearColor,
+	LinearColorArray,
 	RigElementKey,
 	RigElementKeyArray,
 
diff --git a/Engine/Plugins/Animation/ControlRig/Source/ControlRig/Public/Rigs/RigHierarchyElements.h b/Engine/Plugins/Animation/ControlRig/Source/ControlRig/Public/Rigs/RigHierarchyElements.h
--- a/Engine/Plugins/Animation/ControlRig/Source/ControlRig/Public/Rigs/RigHierarchyElements.h
+++ b/Engine/Plugins/Animation/ControlRig/Source/ControlRig/Public/Rigs/RigHierarchyElements.h
@@ -671,7 +671,7 @@ protected:
 
 	friend class URigHierarchy;
 	friend class URigHierarchyController;
-	friend struct FRigUnit_SetMetadataBase;
+	friend struct FRigDispatch_SetMetadata;
 };
 
 USTRUCT(BlueprintType)
diff --git a/Engine/Plugins/Animation/ControlRig/Source/ControlRig/Public/Rigs/RigHierarchyMetadata.h b/Engine/Plugins/Animation/ControlRig/Source/ControlRig/Public/Rigs/RigHierarchyMetadata.h
--- a/Engine/Plugins/Animation/ControlRig/Source/ControlRig/Public/Rigs/RigHierarchyMetadata.h
+++ b/Engine/Plugins/Animation/ControlRig/Source/ControlRig/Public/Rigs/RigHierarchyMetadata.h
@@ -699,6 +699,71 @@ protected:
 	TArray<FTransform> Value;
 };
 
+USTRUCT()
+struct CONTROLRIG_API FRigLinearColorMetadata : public FRigBaseMetadata
+{
+public:
+	
+	GENERATED_BODY()
+	DECLARE_RIG_METADATA_METHODS(FRigLinearColorMetadata)
+
+	FRigLinearColorMetadata()
+		: FRigBaseMetadata()
+		, Value(FLinearColor::White)
+	{}
+	virtual ~FRigLinearColorMetadata() override {}
+	FORCEINLINE virtual void Serialize(FArchive& Ar, bool bIsLoading) override
+	{
+		Super::Serialize(Ar, bIsLoading);
+		Ar << Value;
+	}
+	FORCEINLINE const FLinearColor& GetValue() const { return Value; } 
+	FORCEINLINE FLinearColor& GetValue() { return Value; }
+	FORCEINLINE bool SetValue(const FLinearColor& InValue) { return SetValueData(&InValue, sizeof(FLinearColor)); }
+
+	FORCEINLINE static bool IsClassOf(const FRigBaseMetadata* InMetadata)
+	{
+		return InMetadata->GetType() == ERigMetadataType::LinearColor;
+	}
+
+protected:
+
+	UPROPERTY()
+	FLinearColor Value;
+};
+
+USTRUCT()
+struct CONTROLRIG_API FRigLinearColorArrayMetadata : public FRigBaseMetadata
+{
+public:
+	
+	GENERATED_BODY()
+	DECLARE_RIG_METADATA_METHODS(FRigLinearColorArrayMetadata)
+
+	FRigLinearColorArrayMetadata()
+		: FRigBaseMetadata()
+	{}
+	virtual ~FRigLinearColorArrayMetadata() override {}
+	FORCEINLINE virtual void Serialize(FArchive& Ar, bool bIsLoading) override
+	{
+		Super::Serialize(Ar, bIsLoading);
+		Ar << Value;
+	}
+	FORCEINLINE const TArray<FLinearColor>& GetValue() const { return Value; }
+	FORCEINLINE TArray<FLinearColor>& GetValue() { return Value; }
+	FORCEINLINE bool SetValue(const TArray<FLinearColor>& InValue) { return SetValueData(&InValue, sizeof(TArray<FLinearColor>)); }
+
+	FORCEINLINE static bool IsClassOf(const FRigBaseMetadata* InMetadata)
+	{
+		return InMetadata->GetType() == ERigMetadataType::LinearColorArray;
+	}
+
+protected:
+
+	UPROPERTY()
+	TArray<FLinearColor> Value;
+};
+
 USTRUCT()
 struct CONTROLRIG_API FRigElementKeyMetadata : public FRigBaseMetadata
 {
diff --git a/Engine/Plugins/Animation/ControlRig/Source/ControlRigEditor/Private/ControlRigElementDetails.cpp b/Engine/Plugins/Animation/ControlRig/Source/ControlRigEditor/Private/ControlRigElementDetails.cpp
--- a/Engine/Plugins/Animation/ControlRig/Source/ControlRigEditor/Private/ControlRigElementDetails.cpp
+++ b/Engine/Plugins/Animation/ControlRig/Source/ControlRigEditor/Private/ControlRigElementDetails.cpp
@@ -635,6 +635,14 @@ void FRigBaseElementDetails::CustomizeDetails(IDetailLayoutBuilder& DetailBuilde
 	}
 
 	DetailBuilder.HideCategory(TEXT("RigElement"));
+
+	// if we are not a bone, control or null
+	if(!IsAnyElementOfType(ERigElementType::Bone) &&
+		!IsAnyElementOfType(ERigElementType::Control) &&
+		!IsAnyElementOfType(ERigElementType::Null))
+	{
+		CustomizeMetadata(DetailBuilder);
+	}
 }
 
 FRigElementKey FRigBaseElementDetails::GetElementKey() const
@@ -861,6 +869,8 @@ void FRigBaseElementDetails::RegisterSectionMappings(FPropertyEditorModule& Prop
 
 void FRigBaseElementDetails::RegisterSectionMappings(FPropertyEditorModule& PropertyEditorModule, UClass* InClass)
 {
+	TSharedRef<FPropertySection> MetadataSection = PropertyEditorModule.FindOrCreateSection(InClass->GetFName(), "Metadata", LOCTEXT("Metadata", "Metadata"));
+	MetadataSection->AddCategory("Metadata");
 }
 
 FReply FRigBaseElementDetails::OnSelectParentElementInHierarchyClicked()
@@ -892,6 +902,38 @@ FReply FRigBaseElementDetails::OnSelectElementClicked(const FRigElementKey& InKe
 	return FReply::Handled();
 }
 
+void FRigBaseElementDetails::CustomizeMetadata(IDetailLayoutBuilder& DetailBuilder)
+{
+	if(PerElementInfos.Num() != 1)
+	{
+		return;
+	}
+	const FRigBaseElement* Element = PerElementInfos[0].Element.Get();
+	if(Element->NumMetadata() == 0)
+	{
+		return;
+	}
+
+	IDetailCategoryBuilder& MetadataCategory = DetailBuilder.EditCategory(TEXT("Metadata"), LOCTEXT("Metadata", "Metadata"));
+	for(int32 Index=0;Index<Element->NumMetadata();Index++)
+	{
+		FRigBaseMetadata* Metadata = Element->GetMetadata(Index);
+		TSharedPtr<FStructOnScope> StructOnScope = MakeShareable(new FStructOnScope(Metadata->GetMetadataStruct(), (uint8*)Metadata));
+
+		FAddPropertyParams Params;
+		Params.CreateCategoryNodes(false);
+		Params.ForceShowProperty();
+		
+		IDetailPropertyRow* Row = MetadataCategory.AddExternalStructureProperty(StructOnScope, TEXT("Value"), EPropertyLocation::Default, Params);
+		if(Row)
+		{
+			(*Row)
+			.DisplayName(FText::FromName(Metadata->GetName()))
+			.IsEnabled(false);
+		}
+	}
+}
+
 TSharedPtr<TArray<ERigTransformElementDetailsTransform::Type>> FRigTransformElementDetails::PickedTransforms;
 
 void FRigTransformElementDetails::CustomizeDetails(IDetailLayoutBuilder& DetailBuilder)
@@ -2247,6 +2289,7 @@ void FRigBoneElementDetails::CustomizeDetails(IDetailLayoutBuilder& DetailBuilde
 {
 	FRigTransformElementDetails::CustomizeDetails(DetailBuilder);
 	CustomizeTransform(DetailBuilder);
+	CustomizeMetadata(DetailBuilder);
 }
 
 void FRigControlElementDetails::CustomizeDetails(IDetailLayoutBuilder& DetailBuilder)
@@ -2258,6 +2301,7 @@ void FRigControlElementDetails::CustomizeDetails(IDetailLayoutBuilder& DetailBui
 	CustomizeTransform(DetailBuilder);
 	CustomizeShape(DetailBuilder);
 	CustomizeAnimationChannels(DetailBuilder);
+	CustomizeMetadata(DetailBuilder);
 }
 
 void FRigControlElementDetails::CustomizeValue(IDetailLayoutBuilder& DetailBuilder)
@@ -3703,7 +3747,7 @@ void FRigControlElementDetails::RegisterSectionMappings(FPropertyEditorModule& P
 	ShapeSection->AddCategory("Shape");
 
 	TSharedRef<FPropertySection> ChannelsSection = PropertyEditorModule.FindOrCreateSection(InClass->GetFName(), "Channels", LOCTEXT("Channels", "Channels"));
-	ShapeSection->AddCategory("AnimationChannels");
+	ChannelsSection->AddCategory("AnimationChannels");
 }
 
 bool FRigControlElementDetails::IsShapeEnabled() const
@@ -4487,6 +4531,7 @@ void FRigNullElementDetails::CustomizeDetails(IDetailLayoutBuilder& DetailBuilde
 {
 	FRigTransformElementDetails::CustomizeDetails(DetailBuilder);
 	CustomizeTransform(DetailBuilder);
+	CustomizeMetadata(DetailBuilder);
 }
 
 #undef LOCTEXT_NAMESPACE
diff --git a/Engine/Plugins/Animation/ControlRig/Source/ControlRigEditor/Private/ControlRigElementDetails.h b/Engine/Plugins/Animation/ControlRig/Source/ControlRigEditor/Private/ControlRigElementDetails.h
--- a/Engine/Plugins/Animation/ControlRig/Source/ControlRigEditor/Private/ControlRigElementDetails.h
+++ b/Engine/Plugins/Animation/ControlRig/Source/ControlRigEditor/Private/ControlRigElementDetails.h
@@ -215,7 +215,10 @@ public:
 
 protected:
 
+	void CustomizeMetadata(IDetailLayoutBuilder& DetailBuilder);
+
 	TArray<FPerElementInfo> PerElementInfos;
+	TArray<TStrongObjectPtr<UDetailsViewWrapperObject>> MetadataWrappers;
 	
 	TSharedPtr<SButton> SelectParentElementButton;
 };
diff --git a/Engine/Source/Runtime/RigVM/Private/RigVMCore/RigVMRegistry.cpp b/Engine/Source/Runtime/RigVM/Private/RigVMCore/RigVMRegistry.cpp
--- a/Engine/Source/Runtime/RigVM/Private/RigVMCore/RigVMRegistry.cpp
+++ b/Engine/Source/Runtime/RigVM/Private/RigVMCore/RigVMRegistry.cpp
@@ -761,26 +761,32 @@ void FRigVMRegistry::Register(const TCHAR* InName, FRigVMFunctionPtr InFunctionP
 #endif
 }
 
-void FRigVMRegistry::RegisterFactory(UScriptStruct* InFactoryStruct)
+const FRigVMDispatchFactory* FRigVMRegistry::RegisterFactory(UScriptStruct* InFactoryStruct)
 {
 	check(InFactoryStruct);
 	check(InFactoryStruct != FRigVMDispatchFactory::StaticStruct());
 	check(InFactoryStruct->IsChildOf(FRigVMDispatchFactory::StaticStruct()));
 
 	// ensure to register factories only once
-	const bool bFactoryAlreadyRegistered = Factories.ContainsByPredicate([InFactoryStruct](const FRigVMDispatchFactory* Factory)
+	const FRigVMDispatchFactory* ExistingFactory = nullptr;
+	const bool bFactoryAlreadyRegistered = Factories.ContainsByPredicate([InFactoryStruct, &ExistingFactory](const FRigVMDispatchFactory* Factory)
 	{
-		return Factory->GetScriptStruct() == InFactoryStruct;
+		if(Factory->GetScriptStruct() == InFactoryStruct)
+		{
+			ExistingFactory = Factory;
+			return true;
+		}
+		return false;
 	});
 	if(bFactoryAlreadyRegistered)
 	{
-		return;
+		return ExistingFactory;
 	}
 
 #if WITH_EDITOR
 	if(InFactoryStruct->HasMetaData(TEXT("Abstract")))
 	{
-		return;
+		return nullptr;
 	}
 #endif
 	
@@ -789,6 +795,7 @@ void FRigVMRegistry::RegisterFactory(UScriptStruct* InFactoryStruct)
 	Factory->FactoryScriptStruct = InFactoryStruct;
 	Factories.Add(Factory);
 	Factory->RegisterDependencyTypes();
+	return Factory;
 }
 
 const FRigVMFunction* FRigVMRegistry::FindFunction(const TCHAR* InName) const
diff --git a/Engine/Source/Runtime/RigVM/Public/RigVMCore/RigVMMemoryStorage.h b/Engine/Source/Runtime/RigVM/Public/RigVMCore/RigVMMemoryStorage.h
--- a/Engine/Source/Runtime/RigVM/Public/RigVMCore/RigVMMemoryStorage.h
+++ b/Engine/Source/Runtime/RigVM/Public/RigVMCore/RigVMMemoryStorage.h
@@ -128,7 +128,8 @@ public:
 	>
 	FORCEINLINE bool IsType(bool bIsHiddenArgument = false) const
 	{
-		return IsTypeArray<T::ElementType>(bIsHiddenArgument);
+		typedef typename T::ElementType ElementType;
+		return IsTypeArray<ElementType>(bIsHiddenArgument);
 	}
 
 	// returns true if the handle is of the given type
diff --git a/Engine/Source/Runtime/RigVM/Public/RigVMCore/RigVMRegistry.h b/Engine/Source/Runtime/RigVM/Public/RigVMCore/RigVMRegistry.h
--- a/Engine/Source/Runtime/RigVM/Public/RigVMCore/RigVMRegistry.h
+++ b/Engine/Source/Runtime/RigVM/Public/RigVMCore/RigVMRegistry.h
@@ -30,7 +30,7 @@ public:
 	void Register(const TCHAR* InName, FRigVMFunctionPtr InFunctionPtr, UScriptStruct* InStruct = nullptr, const TArray<FRigVMFunctionArgument>& InArguments = TArray<FRigVMFunctionArgument>());
 
 	// Registers a dispatch factory given its struct.
-	void RegisterFactory(UScriptStruct* InFactoryStruct);
+	const FRigVMDispatchFactory* RegisterFactory(UScriptStruct* InFactoryStruct);
 
 	// Initializes the registry by storing the defaults
 	void InitializeIfNeeded();
