Control Rig: Add missing spawn integer animation channel node

diff --git a/Engine/Plugins/Animation/ControlRig/Source/ControlRig/Private/Units/Execution/RigUnit_DynamicHierarchy.cpp b/Engine/Plugins/Animation/ControlRig/Source/ControlRig/Private/Units/Execution/RigUnit_DynamicHierarchy.cpp
--- a/Engine/Plugins/Animation/ControlRig/Source/ControlRig/Private/Units/Execution/RigUnit_DynamicHierarchy.cpp
+++ b/Engine/Plugins/Animation/ControlRig/Source/ControlRig/Private/Units/Execution/RigUnit_DynamicHierarchy.cpp
@@ -718,6 +718,44 @@ FRigUnit_HierarchyAddAnimationChannelFloat_Execute()
 	}
 }
 
+FRigUnit_HierarchyAddAnimationChannelInteger_Execute()
+{
+	FString ErrorMessage;
+	if(!FRigUnit_DynamicHierarchyBase::IsValidToRunInContext(Context, ExecuteContext, true, &ErrorMessage))
+	{
+		if(!ErrorMessage.IsEmpty())
+		{
+			UE_CONTROLRIG_RIGUNIT_REPORT_ERROR(TEXT("%s"), *ErrorMessage);
+		}
+		return;
+	}
+
+	Item.Reset();
+
+	if(URigHierarchyController* Controller = ExecuteContext.Hierarchy->GetController(true))
+	{
+		FRigControlSettings ControlSettings;
+		ControlSettings.ControlType = ERigControlType::Integer;
+		ControlSettings.SetupLimitArrayForType(true, true, true);
+		ControlSettings.MinimumValue = FRigControlValue::Make<int32>(MinimumValue);
+		ControlSettings.MaximumValue = FRigControlValue::Make<int32>(MaximumValue);
+		const FRigControlValue Value = FRigControlValue::Make<int32>(InitialValue);
+		
+		FRigHierarchyControllerInstructionBracket InstructionBracket(Controller, RigVMExecuteContext.GetInstructionIndex());
+		Item = Controller->AddAnimationChannel(Name, Parent, ControlSettings, false, false);
+
+		if(Item.IsValid())
+		{
+			ExecuteContext.Hierarchy->SetControlValue(Item, Value, ERigControlValueType::Initial, false, false);
+			ExecuteContext.Hierarchy->SetControlValue(Item, Value, ERigControlValueType::Current, false, false);
+		}
+		else
+		{
+			UE_CONTROLRIG_RIGUNIT_REPORT_ERROR(TEXT("Animation Channel could not be added. Is Parent valid?"));
+		}
+	}
+}
+
 FRigUnit_HierarchyAddAnimationChannelVector2D_Execute()
 {
 	FString ErrorMessage;
diff --git a/Engine/Plugins/Animation/ControlRig/Source/ControlRig/Private/Units/Execution/RigUnit_DynamicHierarchy.h b/Engine/Plugins/Animation/ControlRig/Source/ControlRig/Private/Units/Execution/RigUnit_DynamicHierarchy.h
--- a/Engine/Plugins/Animation/ControlRig/Source/ControlRig/Private/Units/Execution/RigUnit_DynamicHierarchy.h
+++ b/Engine/Plugins/Animation/ControlRig/Source/ControlRig/Private/Units/Execution/RigUnit_DynamicHierarchy.h
@@ -1050,6 +1050,45 @@ struct CONTROLRIG_API FRigUnit_HierarchyAddAnimationChannelFloat : public FRigUn
 	virtual void Execute(const FRigUnitContext& Context) override;
 };
 
+/**
+ * Adds a new animation channel to the hierarchy
+ * Note: This node only runs as part of the construction event.
+ */
+USTRUCT(meta=(DisplayName="Spawn Integer Animation Channel", TemplateName="SpawnAnimationChannel", Keywords="Construction,Create,New,AddAnimationChannel,NewAnimationChannel,CreateAnimationChannel,AddChannel,NewChannel,CreateChannel,SpawnChannel", Varying))
+struct CONTROLRIG_API FRigUnit_HierarchyAddAnimationChannelInteger : public FRigUnit_HierarchyAddElement
+{
+	GENERATED_BODY()
+
+	FRigUnit_HierarchyAddAnimationChannelInteger()
+	{
+		Name = TEXT("NewChannel");
+		InitialValue = 0;
+		MinimumValue = 0;
+		MaximumValue = 100;
+	}
+
+	/*
+	 * The initial value of the new animation channel
+	 */
+	UPROPERTY(meta = (Input))
+	int32 InitialValue;
+
+	/*
+	 * The initial value of the new animation channel
+	 */
+	UPROPERTY(meta = (Input))
+	int32 MinimumValue;
+
+	/*
+	 * The maximum value for the animation channel
+	 */
+	UPROPERTY(meta = (Input))
+	int32 MaximumValue;
+
+	RIGVM_METHOD()
+	virtual void Execute(const FRigUnitContext& Context) override;
+};
+
 /**
  * Adds a new animation channel to the hierarchy
  * Note: This node only runs as part of the construction event.
