Control Rig: String conversion and print support

diff --git a/Engine/Plugins/Animation/ControlRig/Source/ControlRig/Private/Units/Core/RigUnit_Name.h b/Engine/Plugins/Animation/ControlRig/Source/ControlRig/Private/Units/Core/RigUnit_Name.h
--- a/Engine/Plugins/Animation/ControlRig/Source/ControlRig/Private/Units/Core/RigUnit_Name.h
+++ b/Engine/Plugins/Animation/ControlRig/Source/ControlRig/Private/Units/Core/RigUnit_Name.h
@@ -14,7 +14,7 @@ struct CONTROLRIG_API FRigUnit_NameBase : public FRigUnit
 /**
  * Concatenates two strings together to make a new string
  */
-USTRUCT(meta = (DisplayName = "Concat", Keywords = "Add,+,Combine,Merge,Append"))
+USTRUCT(meta = (DisplayName = "Concat", TemplateName = "Concat", Keywords = "Add,+,Combine,Merge,Append"))
 struct CONTROLRIG_API FRigUnit_NameConcat : public FRigUnit_NameBase
 {
 	GENERATED_BODY()
@@ -27,20 +27,20 @@ struct CONTROLRIG_API FRigUnit_NameConcat : public FRigUnit_NameBase
 	RIGVM_METHOD()
 	virtual void Execute(const FRigUnitContext& Context) override;
 
-	UPROPERTY(meta=(Input))
+	UPROPERTY(meta=(Input, Aggregate))
 	FName A;
 
-	UPROPERTY(meta=(Input))
+	UPROPERTY(meta=(Input, Aggregate))
 	FName B;
 
-	UPROPERTY(meta=(Output))
+	UPROPERTY(meta=(Output, Aggregate))
 	FName Result;
 };
 
 /**
  * Returns the left or right most characters from the string chopping the given number of characters from the start or the end
  */
-USTRUCT(meta = (DisplayName = "Chop", Keywords = "Truncate,-,Remove,Subtract,Split"))
+USTRUCT(meta = (DisplayName = "Chop", TemplateName = "Chop", Keywords = "Truncate,-,Remove,Subtract,Split"))
 struct CONTROLRIG_API FRigUnit_NameTruncate : public FRigUnit_NameBase
 {
 	GENERATED_BODY()
@@ -78,7 +78,7 @@ struct CONTROLRIG_API FRigUnit_NameTruncate : public FRigUnit_NameBase
 /**
  * Replace all occurrences of a substring in this string
  */
-USTRUCT(meta = (DisplayName = "Replace", Keywords = "Search,Emplace,Find"))
+USTRUCT(meta = (DisplayName = "Replace", TemplateName = "Replace", Keywords = "Search,Emplace,Find"))
 struct CONTROLRIG_API FRigUnit_NameReplace : public FRigUnit_NameBase
 {
 	GENERATED_BODY()
@@ -107,7 +107,7 @@ struct CONTROLRIG_API FRigUnit_NameReplace : public FRigUnit_NameBase
 /**
  * Tests whether this string ends with given string
  */
-USTRUCT(meta = (DisplayName = "Ends With", Keywords = "Right"))
+USTRUCT(meta = (DisplayName = "Ends With", TemplateName = "EndsWith", Keywords = "Right"))
 struct CONTROLRIG_API FRigUnit_EndsWith : public FRigUnit_NameBase
 {
 	GENERATED_BODY()
@@ -134,7 +134,7 @@ struct CONTROLRIG_API FRigUnit_EndsWith : public FRigUnit_NameBase
 /**
  * Tests whether this string starts with given string
  */
-USTRUCT(meta = (DisplayName = "Starts With", Keywords = "Left"))
+USTRUCT(meta = (DisplayName = "Starts With", TemplateName = "StartsWith", Keywords = "Left"))
 struct CONTROLRIG_API FRigUnit_StartsWith : public FRigUnit_NameBase
 {
 	GENERATED_BODY()
@@ -161,7 +161,7 @@ struct CONTROLRIG_API FRigUnit_StartsWith : public FRigUnit_NameBase
 /**
  * Returns true or false if a given name exists in another given name
  */
-USTRUCT(meta = (DisplayName = "Contains", Keywords = "Contains,Find,Has,Search"))
+USTRUCT(meta = (DisplayName = "Contains", TemplateName = "Contains", Keywords = "Contains,Find,Has,Search"))
 struct CONTROLRIG_API FRigUnit_Contains : public FRigUnit_NameBase
 {
 	GENERATED_BODY()
diff --git a/Engine/Plugins/Animation/ControlRig/Source/ControlRig/Private/Units/Core/RigUnit_Print.cpp b/Engine/Plugins/Animation/ControlRig/Source/ControlRig/Private/Units/Core/RigUnit_Print.cpp
new file mode 100644
--- /dev/null
+++ b/Engine/Plugins/Animation/ControlRig/Source/ControlRig/Private/Units/Core/RigUnit_Print.cpp
@@ -0,0 +1,105 @@
+// Copyright Epic Games, Inc. All Rights Reserved.
+
+#include "Units/Core/RigUnit_Print.h"
+#include "Units/RigUnitContext.h"
+#include "Kismet/KismetSystemLibrary.h"
+
+TArray<FRigVMTemplateArgument> FRigDispatch_Print::GetArguments() const
+{
+	const TArray<FRigVMTemplateArgument::ETypeCategory> ValueCategories = {
+		FRigVMTemplateArgument::ETypeCategory_SingleAnyValue,
+		FRigVMTemplateArgument::ETypeCategory_ArrayAnyValue
+	};
+	return {
+		FRigVMTemplateArgument(FRigVMStruct::ExecuteContextName, ERigVMPinDirection::IO, FRigVMRegistry::Get().GetTypeIndex<FControlRigExecuteContext>()),
+		FRigVMTemplateArgument(TEXT("Prefix"), ERigVMPinDirection::Input, RigVMTypeUtils::TypeIndex::FString),
+		FRigVMTemplateArgument(TEXT("Value"), ERigVMPinDirection::Input, ValueCategories),
+		FRigVMTemplateArgument(TEXT("Enabled"), ERigVMPinDirection::Input, RigVMTypeUtils::TypeIndex::Bool),
+		FRigVMTemplateArgument(TEXT("ScreenDuration"), ERigVMPinDirection::Input, RigVMTypeUtils::TypeIndex::Float),
+		FRigVMTemplateArgument(TEXT("ScreenColor"), ERigVMPinDirection::Input, FRigVMRegistry::Get().GetTypeIndex<FLinearColor>())
+	};
+}
+
+FRigVMTemplateTypeMap FRigDispatch_Print::OnNewArgumentType(const FName& InArgumentName,
+	TRigVMTypeIndex InTypeIndex) const
+{
+	FRigVMTemplateTypeMap Types;
+	Types.Add(FRigVMStruct::ExecuteContextName, FRigVMRegistry::Get().GetTypeIndex<FControlRigExecuteContext>());
+	Types.Add(TEXT("Prefix"), RigVMTypeUtils::TypeIndex::FString);
+	Types.Add(TEXT("Value"), InTypeIndex);
+	Types.Add(TEXT("Enabled"), RigVMTypeUtils::TypeIndex::Bool);
+	Types.Add(TEXT("ScreenDuration"), RigVMTypeUtils::TypeIndex::Float);
+	Types.Add(TEXT("ScreenColor"), FRigVMRegistry::Get().GetTypeIndex<FLinearColor>());
+	return Types;
+}
+
+#if WITH_EDITOR
+
+FString FRigDispatch_Print::GetArgumentDefaultValue(const FName& InArgumentName, TRigVMTypeIndex InTypeIndex) const
+{
+	if(InArgumentName == TEXT("Enabled"))
+	{
+		return TEXT("True");
+	}
+	if(InArgumentName == TEXT("ScreenDuration"))
+	{
+		return TEXT("0.050000");
+	}
+	return FRigDispatchFactory::GetArgumentDefaultValue(InArgumentName, InTypeIndex);
+}
+
+FString FRigDispatch_Print::GetArgumentMetaData(const FName& InArgumentName, const FName& InMetaDataKey) const
+{
+	if(InArgumentName == TEXT("ScreenDuration") || InArgumentName == TEXT("ScreenColor"))
+	{
+		if(InMetaDataKey == FRigVMStruct::DetailsOnlyMetaName)
+		{
+			return TEXT("True");
+		}
+	}
+	return FRigDispatchFactory::GetArgumentMetaData(InArgumentName, InMetaDataKey);
+}
+
+#endif
+
+void FRigDispatch_Print::Execute(FRigVMExtendedExecuteContext& InContext, FRigVMMemoryHandleArray Handles)
+{
+#if WITH_EDITOR
+	const FProperty* ValueProperty = Handles[2].GetResolvedProperty(); 
+	check(ValueProperty);
+	check(Handles[1].IsString());
+	check(Handles[3].IsBool());
+	check(Handles[4].IsFloat());
+	check(Handles[5].IsType<FLinearColor>());
+
+	const FString& Prefix = *(const FString*)Handles[1].GetData();
+	const bool bEnabled = *(const bool*)Handles[3].GetData();
+	const float& ScreenDuration = *(const float*)Handles[4].GetData();
+	const FLinearColor& ScreenColor = *(const FLinearColor*)Handles[5].GetData();
+	const uint8* Value = Handles[2].GetData();
+	
+	if(!bEnabled)
+	{
+		return;
+	}
+
+	FString String;
+	ValueProperty->ExportText_Direct(String, Value, Value, nullptr, PPF_None, nullptr);
+
+	FString ObjectPath;
+	if(InContext.VM)
+	{
+		ObjectPath = InContext.VM->GetName();
+	}
+
+	static constexpr TCHAR LogFormat[] = TEXT("%s[%04d] %s%s");
+	UE_LOG(LogControlRig, Display, LogFormat, *ObjectPath, InContext.PublicData.GetInstructionIndex(), *Prefix, *String);
+
+	const FRigUnitContext& Context = GetRigUnitContext(InContext);
+	if(ScreenDuration > SMALL_NUMBER && Context.World)
+	{
+		static constexpr TCHAR PrintStringFormat[] = TEXT("[%04d] %s%s");
+		UKismetSystemLibrary::PrintString(Context.World, FString::Printf(PrintStringFormat, InContext.PublicData.GetInstructionIndex(), *Prefix, *String), true, false, ScreenColor, ScreenDuration);
+	}
+#endif
+}
diff --git a/Engine/Plugins/Animation/ControlRig/Source/ControlRig/Private/Units/Core/RigUnit_Print.h b/Engine/Plugins/Animation/ControlRig/Source/ControlRig/Private/Units/Core/RigUnit_Print.h
new file mode 100644
--- /dev/null
+++ b/Engine/Plugins/Animation/ControlRig/Source/ControlRig/Private/Units/Core/RigUnit_Print.h
@@ -0,0 +1,34 @@
+// Copyright Epic Games, Inc. All Rights Reserved.
+
+#pragma once
+
+#include "Units/RigUnit.h"
+#include "Units/RigDispatchFactory.h"
+#include "RigUnit_Print.generated.h"
+
+/*
+ * Prints any value to the log
+ */
+USTRUCT(meta=(DisplayName = "Print", NodeColor = "0.8, 0, 0.2, 1"))
+struct FRigDispatch_Print : public FRigDispatchFactory
+{
+	GENERATED_BODY()
+
+public:
+
+	virtual TArray<FRigVMTemplateArgument> GetArguments() const override;
+	virtual FRigVMTemplateTypeMap OnNewArgumentType(const FName& InArgumentName, TRigVMTypeIndex InTypeIndex) const override;
+
+#if WITH_EDITOR
+	virtual FString GetArgumentDefaultValue(const FName& InArgumentName, TRigVMTypeIndex InTypeIndex) const override;
+	virtual FString GetArgumentMetaData(const FName& InArgumentName, const FName& InMetaDataKey) const override;
+#endif
+
+protected:
+
+	virtual FRigVMFunctionPtr GetDispatchFunctionImpl(const FRigVMTemplateTypeMap& InTypes) const override
+	{
+		return &FRigDispatch_Print::Execute;
+	}
+	static void Execute(FRigVMExtendedExecuteContext& InContext, FRigVMMemoryHandleArray Handles);
+};
diff --git a/Engine/Plugins/Animation/ControlRig/Source/ControlRig/Private/Units/Core/RigUnit_String.cpp b/Engine/Plugins/Animation/ControlRig/Source/ControlRig/Private/Units/Core/RigUnit_String.cpp
new file mode 100644
--- /dev/null
+++ b/Engine/Plugins/Animation/ControlRig/Source/ControlRig/Private/Units/Core/RigUnit_String.cpp
@@ -0,0 +1,369 @@
+// Copyright Epic Games, Inc. All Rights Reserved.
+
+#include "Units/Core/RigUnit_String.h"
+#include "Units/RigUnitContext.h"
+
+FRigUnit_StringToName_Execute()
+{
+	Result = *Value;
+}
+
+FRigUnit_NameToString_Execute()
+{
+	Result = Value.ToString();
+}
+
+FRigUnit_StringConcat_Execute()
+{
+	Result = A + B;
+}
+
+FRigUnit_StringTruncate_Execute()
+{
+	Remainder = Name;
+	Chopped = FString();
+
+	if(Name.IsEmpty() || Count <= 0)
+	{
+		return;
+	}
+
+	if (FromEnd)
+	{
+		Remainder = Name.LeftChop(Count);
+		Chopped = Name.Right(Count);
+	}
+	else
+	{
+		Remainder = Name.RightChop(Count);
+		Chopped = Name.Left(Count);
+	}
+}
+
+FRigUnit_StringReplace_Execute()
+{
+	Result = Name.Replace(*Old, *New, ESearchCase::CaseSensitive);
+}
+
+FRigUnit_StringEndsWith_Execute()
+{
+	Result = Name.EndsWith(Ending, ESearchCase::CaseSensitive);
+}
+
+FRigUnit_StringStartsWith_Execute()
+{
+	Result = Name.StartsWith(Start, ESearchCase::CaseSensitive);
+}
+
+FRigUnit_StringContains_Execute()
+{
+	Result = Name.Contains(Search, ESearchCase::CaseSensitive);
+}
+
+FRigUnit_StringLength_Execute()
+{
+	Length = Value.Len();
+}
+
+FRigUnit_StringTrimWhitespace_Execute()
+{
+	Result = Value;
+	Result.TrimStartAndEndInline();
+}
+
+FRigUnit_StringToUppercase_Execute()
+{
+	Result = Value.ToUpper();
+}
+
+FRigUnit_StringToLowercase_Execute()
+{
+	Result = Value.ToLower();
+}
+
+FRigUnit_StringReverse_Execute()
+{
+	Reverse = Value.Reverse();
+}
+
+FRigUnit_StringLeft_Execute()
+{
+	Result = Value.Left(Count);
+}
+
+FRigUnit_StringRight_Execute()
+{
+	Result = Value.Right(Count);
+}
+
+FRigUnit_StringMiddle_Execute()
+{
+	if(Count < 0)
+	{
+		Result = Value.Mid(Start);
+	}
+	else
+	{
+		Result = Value.Mid(Start, Count);
+	}
+}
+
+FRigUnit_StringFind_Execute()
+{
+	Index = INDEX_NONE;
+
+	if(!Value.IsEmpty() && !Search.IsEmpty())
+	{
+		Index = Value.Find(Search, ESearchCase::CaseSensitive);
+	}
+
+	Found = Index != INDEX_NONE;
+}
+
+FRigUnit_StringSplit_Execute()
+{
+	Result.Reset();
+	if(!Value.IsEmpty())
+	{
+		if(Separator.IsEmpty())
+		{
+			UE_CONTROLRIG_RIGUNIT_REPORT_ERROR(TEXT("Separator is empty."));
+			return;
+		}
+
+		FString ValueRemaining = Value;
+		FString Left, Right;
+		while(ValueRemaining.Split(TEXT("."), &Left, &Right, ESearchCase::CaseSensitive, ESearchDir::FromStart))
+		{
+			Result.Add(Left);
+			Left.Empty();
+			ValueRemaining = Right;
+		}
+
+		if (!Right.IsEmpty())
+		{
+			Result.Add(Right);
+		}
+	}
+}
+
+FRigUnit_StringJoin_Execute()
+{
+	Result.Reset();
+	if(!Values.IsEmpty())
+	{
+		Result = FString::Join(Values, *Separator);
+	}
+}
+
+FRigUnit_StringPadInteger_Execute()
+{
+	static constexpr TCHAR Format02Digits[] = TEXT("%02d");
+	static constexpr TCHAR Format03Digits[] = TEXT("%03d");
+	static constexpr TCHAR Format04Digits[] = TEXT("%04d");
+	static constexpr TCHAR Format05Digits[] = TEXT("%05d");
+	static constexpr TCHAR Format06Digits[] = TEXT("%06d");
+	static constexpr TCHAR Format07Digits[] = TEXT("%07d");
+	static constexpr TCHAR Format08Digits[] = TEXT("%08d");
+	static constexpr TCHAR Format09Digits[] = TEXT("%09d");
+	static constexpr TCHAR Format10Digits[] = TEXT("%10d");
+	static constexpr TCHAR Format11Digits[] = TEXT("%11d");
+	static constexpr TCHAR Format12Digits[] = TEXT("%12d");
+	static constexpr TCHAR Format13Digits[] = TEXT("%13d");
+	static constexpr TCHAR Format14Digits[] = TEXT("%14d");
+	static constexpr TCHAR Format15Digits[] = TEXT("%15d");
+	static constexpr TCHAR Format16Digits[] = TEXT("%16d");
+
+	switch(Digits)
+	{
+		case 2:
+		{
+			Result = FString::Printf(Format02Digits, Value); 
+			break;
+		}
+		case 3:
+		{
+			Result = FString::Printf(Format03Digits, Value); 
+			break;
+		}
+		case 4:
+		{
+			Result = FString::Printf(Format04Digits, Value); 
+			break;
+		}
+		case 5:
+		{
+			Result = FString::Printf(Format05Digits, Value); 
+			break;
+		}
+		case 6:
+		{
+			Result = FString::Printf(Format06Digits, Value); 
+			break;
+		}
+		case 7:
+		{
+			Result = FString::Printf(Format07Digits, Value); 
+			break;
+		}
+		case 8:
+		{
+			Result = FString::Printf(Format08Digits, Value); 
+			break;
+		}
+		case 9:
+		{
+			Result = FString::Printf(Format09Digits, Value); 
+			break;
+		}
+		case 10:
+		{
+			Result = FString::Printf(Format10Digits, Value); 
+			break;
+		}
+		case 11:
+		{
+			Result = FString::Printf(Format11Digits, Value); 
+			break;
+		}
+		case 12:
+		{
+			Result = FString::Printf(Format12Digits, Value); 
+			break;
+		}
+		case 13:
+		{
+			Result = FString::Printf(Format13Digits, Value); 
+			break;
+		}
+		case 14:
+		{
+			Result = FString::Printf(Format14Digits, Value); 
+			break;
+		}
+		case 15:
+		{
+			Result = FString::Printf(Format15Digits, Value); 
+			break;
+		}
+		case 16:
+		{
+			Result = FString::Printf(Format16Digits, Value); 
+			break;
+		}
+		default:
+		{
+			Result = FString::FormatAsNumber(Value);
+			break;
+		}
+	}
+}
+
+TArray<FRigVMTemplateArgument> FRigDispatch_ToString::GetArguments() const
+{
+	const TArray<FRigVMTemplateArgument::ETypeCategory> ValueCategories = {
+		FRigVMTemplateArgument::ETypeCategory_SingleAnyValue,
+		FRigVMTemplateArgument::ETypeCategory_ArrayAnyValue
+	};
+	return {
+		FRigVMTemplateArgument(TEXT("Value"), ERigVMPinDirection::Input, ValueCategories),
+		FRigVMTemplateArgument(TEXT("Result"), ERigVMPinDirection::Output, RigVMTypeUtils::TypeIndex::FString)
+	};
+}
+
+FRigVMTemplateTypeMap FRigDispatch_ToString::OnNewArgumentType(const FName& InArgumentName,
+	TRigVMTypeIndex InTypeIndex) const
+{
+	FRigVMTemplateTypeMap Types;
+	Types.Add(TEXT("Value"), InTypeIndex);
+	Types.Add(TEXT("Result"), RigVMTypeUtils::TypeIndex::FString);
+	return Types;
+}
+
+void FRigDispatch_ToString::Execute(FRigVMExtendedExecuteContext& InContext, FRigVMMemoryHandleArray Handles)
+{
+	const FProperty* ValueProperty = Handles[0].GetResolvedProperty(); 
+	check(ValueProperty);
+	check(Handles[1].IsString());
+
+	const uint8* Value = Handles[0].GetData();
+	FString& Result = *(FString*)Handles[1].GetData();
+	Result.Reset();
+	ValueProperty->ExportText_Direct(Result, Value, Value, nullptr, PPF_None, nullptr);
+}
+
+TArray<FRigVMTemplateArgument> FRigDispatch_FromString::GetArguments() const
+{
+	const TArray<FRigVMTemplateArgument::ETypeCategory> ValueCategories = {
+		FRigVMTemplateArgument::ETypeCategory_SingleAnyValue,
+		FRigVMTemplateArgument::ETypeCategory_ArrayAnyValue
+	};
+	return {
+		FRigVMTemplateArgument(TEXT("String"), ERigVMPinDirection::Input, RigVMTypeUtils::TypeIndex::FString),
+		FRigVMTemplateArgument(TEXT("Result"), ERigVMPinDirection::Output, ValueCategories)
+	};
+}
+
+FRigVMTemplateTypeMap FRigDispatch_FromString::OnNewArgumentType(const FName& InArgumentName,
+	TRigVMTypeIndex InTypeIndex) const
+{
+	FRigVMTemplateTypeMap Types;
+	Types.Add(TEXT("String"), RigVMTypeUtils::TypeIndex::FString);
+	Types.Add(TEXT("Result"), InTypeIndex);
+	return Types;
+}
+
+void FRigDispatch_FromString::Execute(FRigVMExtendedExecuteContext& InContext, FRigVMMemoryHandleArray Handles)
+{
+	const FProperty* ValueProperty = Handles[1].GetResolvedProperty(); 
+	check(ValueProperty);
+	check(Handles[0].IsString());
+
+	const FString& String = *(const FString*)Handles[0].GetData();
+	uint8* Value = Handles[1].GetData();
+
+	class FRigDispatch_FromString_ErrorPipe : public FOutputDevice
+	{
+	public:
+
+		TArray<FString> Errors;
+
+		FRigDispatch_FromString_ErrorPipe()
+			: FOutputDevice()
+		{
+		}
+
+		virtual void Serialize(const TCHAR* V, ELogVerbosity::Type Verbosity, const class FName& Category) override
+		{
+			Errors.Add(FString::Printf(TEXT("Error convert to string: %s"), V));
+		}
+	};
+
+	FRigDispatch_FromString_ErrorPipe ErrorPipe;
+	ValueProperty->ImportText_Direct(*String, Value, nullptr, PPF_None, &ErrorPipe);
+
+	if(!ErrorPipe.Errors.IsEmpty())
+	{
+		ValueProperty->InitializeValue(Value);
+		for(const FString& Error : ErrorPipe.Errors)
+		{
+			const FRigUnitContext& Context = GetRigUnitContext(InContext);
+
+#if WITH_EDITOR
+			if(Context.Log != nullptr)
+			{
+				Context.Log->Report(EMessageSeverity::Error, InContext.PublicData.GetFunctionName(), InContext.PublicData.GetInstructionIndex(), Error);
+			}
+#endif
+
+			FString ObjectPath;
+			if(InContext.VM)
+			{
+				ObjectPath = InContext.VM->GetName();
+			}
+			
+			static constexpr TCHAR ErrorLogFormat[] = TEXT("%s: [%04d] %s");
+			UE_LOG(LogControlRig, Error, ErrorLogFormat, *ObjectPath, InContext.PublicData.GetInstructionIndex(), *Error);
+		}
+	}
+}
diff --git a/Engine/Plugins/Animation/ControlRig/Source/ControlRig/Private/Units/Core/RigUnit_String.h b/Engine/Plugins/Animation/ControlRig/Source/ControlRig/Private/Units/Core/RigUnit_String.h
new file mode 100644
--- /dev/null
+++ b/Engine/Plugins/Animation/ControlRig/Source/ControlRig/Private/Units/Core/RigUnit_String.h
@@ -0,0 +1,593 @@
+// Copyright Epic Games, Inc. All Rights Reserved.
+
+#pragma once
+
+#include "Units/RigUnit.h"
+#include "Units/RigDispatchFactory.h"
+#include "RigUnit_String.generated.h"
+
+USTRUCT(meta=(Abstract, Category="Core|String", NodeColor = "0.462745, 1,0, 0.329412"))
+struct CONTROLRIG_API FRigUnit_StringBase : public FRigUnit
+{
+	GENERATED_BODY()
+};
+
+/**
+* Makes a name from a string
+*/
+USTRUCT(meta=(DisplayName="To Name", TemplateName="Cast", Keywords="Make,Construct"))
+struct CONTROLRIG_API FRigUnit_StringToName : public FRigUnit_StringBase
+{
+	GENERATED_BODY()
+
+	FRigUnit_StringToName()
+	{
+		Value = FString();
+		Result = NAME_None;
+	}
+
+	RIGVM_METHOD()
+	virtual void Execute(const FRigUnitContext& Context) override;
+
+	UPROPERTY(meta=(Input))
+	FString Value;
+
+	UPROPERTY(meta=(Output))
+	FName Result;
+};
+
+/**
+* Makes a string from a name
+*/
+USTRUCT(meta=(DisplayName="To Name", TemplateName="Cast", Keywords="Make,Construct"))
+struct CONTROLRIG_API FRigUnit_NameToString : public FRigUnit_StringBase
+{
+	GENERATED_BODY()
+
+	FRigUnit_NameToString()
+	{
+		Value = NAME_None;
+		Result = FString();
+	}
+
+	RIGVM_METHOD()
+	virtual void Execute(const FRigUnitContext& Context) override;
+
+	UPROPERTY(meta=(Input))
+	FName Value;
+
+	UPROPERTY(meta=(Output))
+	FString Result;
+};
+
+/**
+ * Concatenates two strings together to make a new string
+ */
+USTRUCT(meta = (DisplayName = "Concat", TemplateName = "Concat", Keywords = "Add,+,Combine,Merge,Append"))
+struct CONTROLRIG_API FRigUnit_StringConcat : public FRigUnit_StringBase
+{
+	GENERATED_BODY()
+
+	FRigUnit_StringConcat()
+	{
+		A = B = Result = FString();
+	}
+
+	RIGVM_METHOD()
+	virtual void Execute(const FRigUnitContext& Context) override;
+
+	UPROPERTY(meta=(Input, Aggregate))
+	FString A;
+
+	UPROPERTY(meta=(Input, Aggregate))
+	FString B;
+
+	UPROPERTY(meta=(Output, Aggregate))
+	FString Result;
+};
+
+/**
+ * Returns the left or right most characters from the string chopping the given number of characters from the start or the end
+ */
+USTRUCT(meta = (DisplayName = "Chop", TemplateName = "Chop", Keywords = "Truncate,-,Remove,Subtract,Split"))
+struct CONTROLRIG_API FRigUnit_StringTruncate : public FRigUnit_StringBase
+{
+	GENERATED_BODY()
+
+	FRigUnit_StringTruncate()
+	{
+		Name = FString();
+		Count = 1;
+		FromEnd = true;
+	}
+
+	RIGVM_METHOD()
+	virtual void Execute(const FRigUnitContext& Context) override;
+
+	UPROPERTY(meta=(Input))
+	FString Name;
+
+	// Number of characters to remove from left or right
+	UPROPERTY(meta=(Input))
+	int32 Count;
+
+	// if set to true the characters will be removed from the end
+	UPROPERTY(meta=(Input))
+	bool FromEnd;
+
+	// the part of the string without the chopped characters
+	UPROPERTY(meta=(Output))
+	FString Remainder;
+
+	// the part of the name that has been chopped off
+	UPROPERTY(meta = (Output))
+	FString Chopped;
+};
+
+/**
+ * Replace all occurrences of a substring in this string
+ */
+USTRUCT(meta = (DisplayName = "Replace", TemplateName = "Replace", Keywords = "Search,Emplace,Find"))
+struct CONTROLRIG_API FRigUnit_StringReplace : public FRigUnit_StringBase
+{
+	GENERATED_BODY()
+
+	FRigUnit_StringReplace()
+	{
+		Name = Old = New = FString();
+	}
+
+	RIGVM_METHOD()
+	virtual void Execute(const FRigUnitContext& Context) override;
+
+	UPROPERTY(meta = (Input))
+	FString Name;
+
+	UPROPERTY(meta = (Input))
+	FString Old;
+
+	UPROPERTY(meta = (Input))
+	FString New;
+
+	UPROPERTY(meta = (Output))
+	FString Result;
+};
+
+/**
+ * Tests whether this string ends with given string
+ */
+USTRUCT(meta = (DisplayName = "Ends With", TemplateName = "EndsWith", Keywords = "Right"))
+struct CONTROLRIG_API FRigUnit_StringEndsWith : public FRigUnit_StringBase
+{
+	GENERATED_BODY()
+
+	FRigUnit_StringEndsWith()
+	{
+		Name = Ending = FString();
+		Result = false;
+	}
+
+	RIGVM_METHOD()
+	virtual void Execute(const FRigUnitContext& Context) override;
+
+	UPROPERTY(meta = (Input))
+	FString Name;
+
+	UPROPERTY(meta = (Input))
+	FString Ending;
+
+	UPROPERTY(meta = (Output))
+	bool Result;
+};
+
+/**
+ * Tests whether this string starts with given string
+ */
+USTRUCT(meta = (DisplayName = "Starts With", TemplateName = "StartsWith", Keywords = "Left"))
+struct CONTROLRIG_API FRigUnit_StringStartsWith : public FRigUnit_StringBase
+{
+	GENERATED_BODY()
+
+	FRigUnit_StringStartsWith()
+	{
+		Name = Start = FString();
+		Result = false;
+	}
+
+	RIGVM_METHOD()
+	virtual void Execute(const FRigUnitContext& Context) override;
+
+	UPROPERTY(meta = (Input))
+	FString Name;
+
+	UPROPERTY(meta = (Input))
+	FString Start;
+
+	UPROPERTY(meta = (Output))
+	bool Result;
+};
+
+/**
+ * Returns true or false if a given name exists in another given name
+ */
+USTRUCT(meta = (DisplayName = "Contains", TemplateName = "Contains", Keywords = "Contains,Find,Has,Search"))
+struct CONTROLRIG_API FRigUnit_StringContains : public FRigUnit_StringBase
+{
+	GENERATED_BODY()
+
+		FRigUnit_StringContains()
+	{
+		Name = Search = FString();
+		Result = false;
+	}
+
+	RIGVM_METHOD()
+	virtual void Execute(const FRigUnitContext& Context) override;
+
+	UPROPERTY(meta = (Input))
+	FString Name;
+
+	UPROPERTY(meta = (Input))
+	FString Search;
+
+	UPROPERTY(meta = (Output))
+	bool Result;
+};
+
+/**
+ * Returns the length of a string 
+ */
+USTRUCT(meta = (DisplayName = "Length"))
+struct CONTROLRIG_API FRigUnit_StringLength : public FRigUnit_StringBase
+{
+	GENERATED_BODY()
+
+	FRigUnit_StringLength()
+	{
+		Value = FString();
+		Length = 0;
+	}
+
+	RIGVM_METHOD()
+	virtual void Execute(const FRigUnitContext& Context) override;
+
+	UPROPERTY(meta = (Input))
+	FString Value;
+
+	UPROPERTY(meta = (Output))
+	int32 Length;
+};
+
+/**
+ * Trims the whitespace from a string (start and end)
+ */
+USTRUCT(meta = (DisplayName = "Trim Whitespace", Keywords = "Space,WhiteSpace,Remove,Truncate"))
+struct CONTROLRIG_API FRigUnit_StringTrimWhitespace : public FRigUnit_StringBase
+{
+	GENERATED_BODY()
+
+	FRigUnit_StringTrimWhitespace()
+	{
+		Value = Result = FString();
+	}
+
+	RIGVM_METHOD()
+	virtual void Execute(const FRigUnitContext& Context) override;
+
+	UPROPERTY(meta = (Input))
+	FString Value;
+
+	UPROPERTY(meta = (Output))
+	FString Result;
+};
+
+/**
+ * Converts the string to upper case
+ */
+USTRUCT(meta = (DisplayName = "To Uppercase"))
+struct CONTROLRIG_API FRigUnit_StringToUppercase : public FRigUnit_StringBase
+{
+	GENERATED_BODY()
+
+	FRigUnit_StringToUppercase()
+	{
+		Value = Result = FString();
+	}
+
+	RIGVM_METHOD()
+	virtual void Execute(const FRigUnitContext& Context) override;
+
+	UPROPERTY(meta = (Input))
+	FString Value;
+
+	UPROPERTY(meta = (Output))
+	FString Result;
+};
+
+/**
+ * Converts the string to lower case
+ */
+USTRUCT(meta = (DisplayName = "To Lowercase"))
+struct CONTROLRIG_API FRigUnit_StringToLowercase : public FRigUnit_StringBase
+{
+	GENERATED_BODY()
+
+	FRigUnit_StringToLowercase()
+	{
+		Value = Result = FString();
+	}
+
+	RIGVM_METHOD()
+	virtual void Execute(const FRigUnitContext& Context) override;
+
+	UPROPERTY(meta = (Input))
+	FString Value;
+
+	UPROPERTY(meta = (Output))
+	FString Result;
+};
+
+/**
+ * Returns the reverse of the input string
+ */
+USTRUCT(meta = (DisplayName = "Reverse"))
+struct CONTROLRIG_API FRigUnit_StringReverse : public FRigUnit_StringBase
+{
+	GENERATED_BODY()
+
+	FRigUnit_StringReverse()
+	{
+		Value = Reverse = FString();
+	}
+
+	RIGVM_METHOD()
+	virtual void Execute(const FRigUnitContext& Context) override;
+
+	UPROPERTY(meta = (Input))
+	FString Value;
+
+	UPROPERTY(meta = (Output))
+	FString Reverse;
+};
+
+/**
+ * Returns the left most characters of a string
+ */
+USTRUCT(meta = (DisplayName = "Left", Keywords = "Start,Begin"))
+struct CONTROLRIG_API FRigUnit_StringLeft : public FRigUnit_StringBase
+{
+	GENERATED_BODY()
+
+	FRigUnit_StringLeft()
+	{
+		Value = Result = FString();
+		Count = 1;
+	}
+
+	RIGVM_METHOD()
+	virtual void Execute(const FRigUnitContext& Context) override;
+
+	UPROPERTY(meta = (Input))
+	FString Value;
+
+	UPROPERTY(meta = (Input))
+	int32 Count;
+
+	UPROPERTY(meta = (Output))
+	FString Result;
+};
+
+/**
+ * Returns the right most characters of a string
+ */
+USTRUCT(meta = (DisplayName = "Right", Keywords = "End"))
+struct CONTROLRIG_API FRigUnit_StringRight : public FRigUnit_StringBase
+{
+	GENERATED_BODY()
+
+	FRigUnit_StringRight()
+	{
+		Value = Result = FString();
+		Count = 1;
+	}
+
+	RIGVM_METHOD()
+	virtual void Execute(const FRigUnitContext& Context) override;
+
+	UPROPERTY(meta = (Input))
+	FString Value;
+
+	UPROPERTY(meta = (Input))
+	int32 Count;
+
+	UPROPERTY(meta = (Output))
+	FString Result;
+};
+
+/**
+ * Returns the middle section of a string
+ */
+USTRUCT(meta = (DisplayName = "Middle", Keywords = "Within,Center"))
+struct CONTROLRIG_API FRigUnit_StringMiddle : public FRigUnit_StringBase
+{
+	GENERATED_BODY()
+
+	FRigUnit_StringMiddle()
+	{
+		Value = Result = FString();
+		Start = 0;
+		Count = -1;
+	}
+
+	RIGVM_METHOD()
+	virtual void Execute(const FRigUnitContext& Context) override;
+
+	UPROPERTY(meta = (Input))
+	FString Value;
+
+	// the index of the first character
+	UPROPERTY(meta = (Input))
+	int32 Start;
+
+	// if count is set to -1 all character from Start will be returned 
+	UPROPERTY(meta = (Input))
+	int32 Count;
+
+	UPROPERTY(meta = (Output))
+	FString Result;
+};
+
+/**
+ * Finds a string within another string
+ */
+USTRUCT(meta = (DisplayName = "Find", Keywords = "IndexOf"))
+struct CONTROLRIG_API FRigUnit_StringFind : public FRigUnit_StringBase
+{
+	GENERATED_BODY()
+
+	FRigUnit_StringFind()
+	{
+		Value = Search = FString();
+		Found = false;
+		Index = INDEX_NONE;
+	}
+
+	RIGVM_METHOD()
+	virtual void Execute(const FRigUnitContext& Context) override;
+
+	UPROPERTY(meta = (Input))
+	FString Value;
+
+	UPROPERTY(meta = (Input))
+	FString Search;
+
+	UPROPERTY(meta = (Output))
+    bool Found;
+
+	UPROPERTY(meta = (Output))
+	int32 Index;
+};
+
+/**
+ * Splits a string into multiple sections given a separator
+ */
+USTRUCT(meta = (DisplayName = "Split"))
+struct CONTROLRIG_API FRigUnit_StringSplit : public FRigUnit_StringBase
+{
+	GENERATED_BODY()
+
+	FRigUnit_StringSplit()
+	{
+		Value = Separator = FString();
+	}
+
+	RIGVM_METHOD()
+	virtual void Execute(const FRigUnitContext& Context) override;
+
+	UPROPERTY(meta = (Input))
+	FString Value;
+
+	UPROPERTY(meta = (Input))
+	FString Separator;
+
+	UPROPERTY(meta = (Output))
+	TArray<FString> Result;
+};
+
+/**
+ * Joins a string into multiple sections given a separator
+ */
+USTRUCT(meta = (DisplayName = "Join", Keywords = "Combine"))
+struct CONTROLRIG_API FRigUnit_StringJoin : public FRigUnit_StringBase
+{
+	GENERATED_BODY()
+
+	FRigUnit_StringJoin()
+	{
+		Result = Separator = FString();
+	}
+
+	RIGVM_METHOD()
+	virtual void Execute(const FRigUnitContext& Context) override;
+
+	UPROPERTY(meta = (Input))
+	TArray<FString> Values;
+
+	UPROPERTY(meta = (Input))
+	FString Separator;
+
+	UPROPERTY(meta = (Output))
+	FString Result;
+};
+
+/**
+ * Converts an integer number to a string with padding
+ */
+USTRUCT(meta = (DisplayName = "Pad Integer", Keywords = "FromInt,Number,LeadingZeroes"))
+struct CONTROLRIG_API FRigUnit_StringPadInteger : public FRigUnit_StringBase
+{
+	GENERATED_BODY()
+
+	FRigUnit_StringPadInteger()
+	{
+		Value = 0;
+		Digits = 4;
+		Result = FString();
+	}
+
+	RIGVM_METHOD()
+	virtual void Execute(const FRigUnitContext& Context) override;
+
+	UPROPERTY(meta = (Input))
+	int32 Value;
+
+	UPROPERTY(meta = (Input))
+	int32 Digits;
+
+	UPROPERTY(meta = (Output))
+	FString Result;
+};
+
+/*
+ * Converts any value to string
+ */
+USTRUCT(meta=(DisplayName = "To String", odeColor = "0.462745, 1,0, 0.329412"))
+struct FRigDispatch_ToString : public FRigDispatchFactory
+{
+	GENERATED_BODY()
+
+public:
+
+	virtual TArray<FRigVMTemplateArgument> GetArguments() const override;
+	virtual FRigVMTemplateTypeMap OnNewArgumentType(const FName& InArgumentName, TRigVMTypeIndex InTypeIndex) const override;
+
+protected:
+
+	virtual FRigVMFunctionPtr GetDispatchFunctionImpl(const FRigVMTemplateTypeMap& InTypes) const override
+	{
+		return &FRigDispatch_ToString::Execute;
+	}
+	static void Execute(FRigVMExtendedExecuteContext& InContext, FRigVMMemoryHandleArray Handles);
+};
+
+/*
+ * Converts a string into any value
+ */
+USTRUCT(meta=(DisplayName = "From String", NodeColor = "0.462745, 1,0, 0.329412"))
+struct FRigDispatch_FromString : public FRigDispatchFactory
+{
+	GENERATED_BODY()
+
+public:
+
+	virtual TArray<FRigVMTemplateArgument> GetArguments() const override;
+	virtual FRigVMTemplateTypeMap OnNewArgumentType(const FName& InArgumentName, TRigVMTypeIndex InTypeIndex) const override;
+
+protected:
+
+	virtual FRigVMFunctionPtr GetDispatchFunctionImpl(const FRigVMTemplateTypeMap& InTypes) const override
+	{
+		return &FRigDispatch_FromString::Execute;
+	}
+	static void Execute(FRigVMExtendedExecuteContext& InContext, FRigVMMemoryHandleArray Handles);
+};
\ No newline at end of file
diff --git a/Engine/Source/Developer/RigVMDeveloper/Private/RigVMModel/RigVMPin.cpp b/Engine/Source/Developer/RigVMDeveloper/Private/RigVMModel/RigVMPin.cpp
--- a/Engine/Source/Developer/RigVMDeveloper/Private/RigVMModel/RigVMPin.cpp
+++ b/Engine/Source/Developer/RigVMDeveloper/Private/RigVMModel/RigVMPin.cpp
@@ -1019,7 +1019,10 @@ FName URigVMPin::GetCustomWidgetName() const
 					}
 				}
 			}
-			else if(const FRigVMTemplate* Template = UnitNode->GetTemplate())
+		}
+		else if(const URigVMTemplateNode* TemplateNode = Cast<URigVMTemplateNode>(GetNode()))
+		{
+			if(const FRigVMTemplate* Template = TemplateNode->GetTemplate())
 			{
 				const FString MetaData = Template->GetArgumentMetaData(GetFName(), FRigVMStruct::CustomWidgetMetaName);
 				if(!MetaData.IsEmpty())
@@ -1249,9 +1252,9 @@ bool URigVMPin::CanBeBoundToVariable(const FRigVMExternalVariable& InExternalVar
 bool URigVMPin::ShowInDetailsPanelOnly() const
 {
 #if WITH_EDITOR
-	if (URigVMUnitNode* UnitNode = Cast<URigVMUnitNode>(GetNode()))
+	if (GetParentPin() == nullptr)
 	{
-		if (GetParentPin() == nullptr)
+		if (URigVMUnitNode* UnitNode = Cast<URigVMUnitNode>(GetNode()))
 		{
 			if (UScriptStruct* ScriptStruct = UnitNode->GetScriptStruct())
 			{
@@ -1264,6 +1267,13 @@ bool URigVMPin::ShowInDetailsPanelOnly() const
 				}
 			}
 		}
+		else if(const URigVMTemplateNode* TemplateNode = Cast<URigVMTemplateNode>(GetNode()))
+		{
+			if(const FRigVMTemplate* Template = TemplateNode->GetTemplate())
+			{
+				return !Template->GetArgumentMetaData(GetFName(), FRigVMStruct::DetailsOnlyMetaName).IsEmpty();
+			}
+		}
 	}
 #endif
 	return false;
