Control Rig: Add linearcolor metadata blueprint support

diff --git a/Engine/Plugins/Animation/ControlRig/Source/ControlRig/Public/Rigs/RigHierarchy.h b/Engine/Plugins/Animation/ControlRig/Source/ControlRig/Public/Rigs/RigHierarchy.h
--- a/Engine/Plugins/Animation/ControlRig/Source/ControlRig/Public/Rigs/RigHierarchy.h
+++ b/Engine/Plugins/Animation/ControlRig/Source/ControlRig/Public/Rigs/RigHierarchy.h
@@ -1206,7 +1206,7 @@ public:
 	UFUNCTION(BlueprintCallable, Category = URigHierarchy)
 	FORCEINLINE bool SetQuatArrayMetadata(FRigElementKey InItem, FName InMetadataName, TArray<FQuat> InValue)
 	{
-		return SetArrayMetadata<FQuat>(InItem, ERigMetadataType::VectorArray, InMetadataName, InValue);
+		return SetArrayMetadata<FQuat>(InItem, ERigMetadataType::QuatArray, InMetadataName, InValue);
 	}
 
 	/**
@@ -1255,7 +1255,56 @@ public:
 	UFUNCTION(BlueprintCallable, Category = URigHierarchy)
 	FORCEINLINE bool SetTransformArrayMetadata(FRigElementKey InItem, FName InMetadataName, TArray<FTransform> InValue)
 	{
-		return SetArrayMetadata<FTransform>(InItem, ERigMetadataType::VectorArray, InMetadataName, InValue);
+		return SetArrayMetadata<FTransform>(InItem, ERigMetadataType::TransformArray, InMetadataName, InValue);
+	}
+
+		/**
+	 * Queries and returns the value of FLinearColor metadata
+	 * @param InItem The element key to return the metadata for
+	 * @param InMetadataName The name of the metadata to query
+	 * @param DefaultValue The default value to fall back on
+	 */
+	UFUNCTION(BlueprintPure, Category = URigHierarchy)
+	FORCEINLINE FLinearColor GetLinearColorMetadata(FRigElementKey InItem, FName InMetadataName, FLinearColor DefaultValue) const
+	{
+		return GetMetadata<FLinearColor>(InItem, ERigMetadataType::LinearColor, InMetadataName, DefaultValue);
+	}
+
+	/**
+	 * Queries and returns the value of FLinearColor array metadata
+	 * @param InItem The element key to return the metadata for
+	 * @param InMetadataName The name of the metadata to query
+	 */
+	UFUNCTION(BlueprintPure, Category = URigHierarchy)
+	FORCEINLINE TArray<FLinearColor> GetLinearColorArrayMetadata(FRigElementKey InItem, FName InMetadataName) const
+	{
+		return GetArrayMetadata<FLinearColor>(InItem, ERigMetadataType::LinearColorArray, InMetadataName);
+	}
+
+	/**
+	 * Sets the metadata to a FLinearColor value
+	 * @param InItem The element key to set the metadata for
+	 * @param InMetadataName The name of the metadata to set
+	 * @param InValue The value to set
+	 * @return Returns true if setting the metadata was successful
+	 */
+	UFUNCTION(BlueprintCallable, Category = URigHierarchy)
+	FORCEINLINE bool SetLinearColorMetadata(FRigElementKey InItem, FName InMetadataName, FLinearColor InValue)
+	{
+		return SetMetadata<FLinearColor>(InItem, ERigMetadataType::LinearColor, InMetadataName, InValue);
+	}
+
+	/**
+	 * Sets the metadata to a FLinearColor array value
+	 * @param InItem The element key to set the metadata for
+	 * @param InMetadataName The name of the metadata to set
+	 * @param InValue The value to set
+	 * @return Returns true if setting the metadata was successful
+	 */
+	UFUNCTION(BlueprintCallable, Category = URigHierarchy)
+	FORCEINLINE bool SetLinearColorArrayMetadata(FRigElementKey InItem, FName InMetadataName, TArray<FLinearColor> InValue)
+	{
+		return SetArrayMetadata<FLinearColor>(InItem, ERigMetadataType::LinearColorArray, InMetadataName, InValue);
 	}
 
 	/**
@@ -1304,7 +1353,7 @@ public:
 	UFUNCTION(BlueprintCallable, Category = URigHierarchy)
 	FORCEINLINE bool SetRigElementKeyArrayMetadata(FRigElementKey InItem, FName InMetadataName, TArray<FRigElementKey> InValue)
 	{
-		return SetArrayMetadata<FRigElementKey>(InItem, ERigMetadataType::VectorArray, InMetadataName, InValue);
+		return SetArrayMetadata<FRigElementKey>(InItem, ERigMetadataType::RigElementKeyArray, InMetadataName, InValue);
 	}
 
 	/**
