Control Rig: Maintain control names in UI as valid until re-chosen

diff --git a/Engine/Plugins/Animation/ControlRig/Source/ControlRig/Private/Units/RigDispatchFactory.cpp b/Engine/Plugins/Animation/ControlRig/Source/ControlRig/Private/Units/RigDispatchFactory.cpp
--- a/Engine/Plugins/Animation/ControlRig/Source/ControlRig/Private/Units/RigDispatchFactory.cpp
+++ b/Engine/Plugins/Animation/ControlRig/Source/ControlRig/Private/Units/RigDispatchFactory.cpp
@@ -60,6 +60,11 @@ FString FRigDispatchFactory::GetArgumentDefaultValue(const FName& InArgumentName
 		static FString DefaultValueString = GetDefaultValueForStruct(FEulerTransform::Identity); 
 		return DefaultValueString;
 	}
+	if(InTypeIndex == FRigVMRegistry::Get().GetTypeIndex<FLinearColor>())
+	{
+		static FString DefaultValueString = GetDefaultValueForStruct(FLinearColor::White); 
+		return DefaultValueString;
+	}
 	if(InTypeIndex == FRigVMRegistry::Get().GetTypeIndex<FRigElementKey>())
 	{
 		static FString DefaultValueString = GetDefaultValueForStruct(FRigElementKey(NAME_None, ERigElementType::Bone)); 
diff --git a/Engine/Plugins/Animation/ControlRig/Source/ControlRigDeveloper/Private/Graph/ControlRigGraph.cpp b/Engine/Plugins/Animation/ControlRig/Source/ControlRigDeveloper/Private/Graph/ControlRigGraph.cpp
--- a/Engine/Plugins/Animation/ControlRig/Source/ControlRigDeveloper/Private/Graph/ControlRigGraph.cpp
+++ b/Engine/Plugins/Animation/ControlRig/Source/ControlRigDeveloper/Private/Graph/ControlRigGraph.cpp
@@ -176,13 +176,14 @@ void UControlRigGraph::CacheNameLists(URigHierarchy* InHierarchy, const FControl
 		TArray<TSharedPtr<FString>>& RigidBodyNameList = ElementNameLists.FindChecked(ERigElementType::RigidBody);
 		TArray<TSharedPtr<FString>>& ReferenceNameList = ElementNameLists.FindChecked(ERigElementType::Reference);
 		
-		CacheNameListForHierarchy<FRigBaseElement>(InHierarchy, AllNameList);
-		CacheNameListForHierarchy<FRigBoneElement>(InHierarchy, BoneNameList);
-		CacheNameListForHierarchy<FRigNullElement>(InHierarchy, NullNameList);
-		CacheNameListForHierarchy<FRigControlElement>(InHierarchy, ControlNameList);
-		CacheNameListForHierarchy<FRigCurveElement>(InHierarchy, CurveNameList);
-		CacheNameListForHierarchy<FRigRigidBodyElement>(InHierarchy, RigidBodyNameList);
-		CacheNameListForHierarchy<FRigReferenceElement>(InHierarchy, ReferenceNameList);
+		CacheNameListForHierarchy<FRigBaseElement>(InHierarchy, AllNameList, false);
+		CacheNameListForHierarchy<FRigBoneElement>(InHierarchy, BoneNameList, false);
+		CacheNameListForHierarchy<FRigNullElement>(InHierarchy, NullNameList, false);
+		CacheNameListForHierarchy<FRigControlElement>(InHierarchy, ControlNameList, false);
+		CacheNameListForHierarchy<FRigControlElement>(InHierarchy, ControlNameListWithoutAnimationChannels, true);
+		CacheNameListForHierarchy<FRigCurveElement>(InHierarchy, CurveNameList, false);
+		CacheNameListForHierarchy<FRigRigidBodyElement>(InHierarchy, RigidBodyNameList, false);
+		CacheNameListForHierarchy<FRigReferenceElement>(InHierarchy, ReferenceNameList, false);
 
 		LastHierarchyTopologyVersion = InHierarchy->GetTopologyVersion();
 	}
diff --git a/Engine/Plugins/Animation/ControlRig/Source/ControlRigDeveloper/Public/Graph/ControlRigGraph.h b/Engine/Plugins/Animation/ControlRig/Source/ControlRigDeveloper/Public/Graph/ControlRigGraph.h
--- a/Engine/Plugins/Animation/ControlRig/Source/ControlRigDeveloper/Public/Graph/ControlRigGraph.h
+++ b/Engine/Plugins/Animation/ControlRig/Source/ControlRigDeveloper/Public/Graph/ControlRigGraph.h
@@ -53,6 +53,10 @@ public:
 	{
 		return GetElementNameList(ERigElementType::Control);
 	}
+	FORCEINLINE const TArray<TSharedPtr<FString>>* GetControlNameListWithoutAnimationChannels(URigVMPin* InPin = nullptr) const
+	{
+		return &ControlNameListWithoutAnimationChannels;
+	}
 	FORCEINLINE const TArray<TSharedPtr<FString>>* GetNullNameList(URigVMPin* InPin = nullptr) const
 	{
 		return GetElementNameList(ERigElementType::Null);
@@ -102,14 +106,14 @@ private:
 	}
 
 	template<class T>
-	void CacheNameListForHierarchy(URigHierarchy* InHierarchy, TArray<TSharedPtr<FString>>& OutNameList)
+	void CacheNameListForHierarchy(URigHierarchy* InHierarchy, TArray<TSharedPtr<FString>>& OutNameList, bool bFilter = true)
 	{
         TArray<FString> Names;
 		for (auto Element : *InHierarchy)
 		{
 			if(Element->IsA<T>())
 			{
-				if(IncludeElementInNameList<T>(Cast<T>(Element)))
+				if(!bFilter || IncludeElementInNameList<T>(Cast<T>(Element)))
 				{
 					Names.Add(Element->GetName().ToString());
 				}
@@ -144,6 +148,7 @@ private:
 	}
 
 	TMap<ERigElementType, TArray<TSharedPtr<FString>>> ElementNameLists;
+	TArray<TSharedPtr<FString>>	ControlNameListWithoutAnimationChannels;
 	TArray<TSharedPtr<FString>> DrawingNameList;
 	TArray<TSharedPtr<FString>> EntryNameList;
 	TArray<TSharedPtr<FString>> ShapeNameList;
diff --git a/Engine/Plugins/Animation/ControlRig/Source/ControlRigEditor/Private/ControlRigGraphDetails.cpp b/Engine/Plugins/Animation/ControlRig/Source/ControlRigEditor/Private/ControlRigGraphDetails.cpp
--- a/Engine/Plugins/Animation/ControlRig/Source/ControlRigEditor/Private/ControlRigGraphDetails.cpp
+++ b/Engine/Plugins/Animation/ControlRig/Source/ControlRigEditor/Private/ControlRigGraphDetails.cpp
@@ -1381,7 +1381,7 @@ void FControlRigWrappedNodeDetails::CustomizeDetails(IDetailLayoutBuilder& Detai
         		}
         		else if (CustomWidgetName == TEXT("ControlName"))
         		{
-        			NameList = GraphBeingCustomized->GetControlNameList();
+        			NameList = GraphBeingCustomized->GetControlNameListWithoutAnimationChannels();
         		}
         		else if (CustomWidgetName == TEXT("SpaceName"))
         		{
diff --git a/Engine/Plugins/Animation/ControlRig/Source/ControlRigEditor/Private/Graph/ControlRigGraphPanelPinFactory.cpp b/Engine/Plugins/Animation/ControlRig/Source/ControlRigEditor/Private/Graph/ControlRigGraphPanelPinFactory.cpp
--- a/Engine/Plugins/Animation/ControlRig/Source/ControlRigEditor/Private/Graph/ControlRigGraphPanelPinFactory.cpp
+++ b/Engine/Plugins/Animation/ControlRig/Source/ControlRigEditor/Private/Graph/ControlRigGraphPanelPinFactory.cpp
@@ -60,7 +60,8 @@ TSharedPtr<SGraphPin> FControlRigGraphPanelPinFactory::CreatePin(UEdGraphPin* In
 					return SNew(SControlRigGraphPinNameList, InPin)
 						.ModelPin(ModelPin)
 						.OnGetNameFromSelection_UObject(RigGraph, &UControlRigGraph::GetSelectedElementsNameList)
-						.OnGetNameListContent_UObject(RigGraph, &UControlRigGraph::GetControlNameList);
+						.OnGetNameListContent_UObject(RigGraph, &UControlRigGraph::GetControlNameListWithoutAnimationChannels)
+						.OnGetNameListContentForValidation_UObject(RigGraph, &UControlRigGraph::GetControlNameList);
 				}
 				else if (CustomWidgetName == TEXT("SpaceName"))
 				{
diff --git a/Engine/Plugins/Animation/ControlRig/Source/ControlRigEditor/Private/Graph/SControlRigGraphPinNameList.cpp b/Engine/Plugins/Animation/ControlRig/Source/ControlRigEditor/Private/Graph/SControlRigGraphPinNameList.cpp
--- a/Engine/Plugins/Animation/ControlRig/Source/ControlRigEditor/Private/Graph/SControlRigGraphPinNameList.cpp
+++ b/Engine/Plugins/Animation/ControlRig/Source/ControlRigEditor/Private/Graph/SControlRigGraphPinNameList.cpp
@@ -28,10 +28,11 @@ void SControlRigGraphPinNameList::Construct(const FArguments& InArgs, UEdGraphPi
 {
 	this->ModelPin = InArgs._ModelPin;
 	this->OnGetNameListContent = InArgs._OnGetNameListContent;
+	this->OnGetNameListContentForValidation = InArgs._OnGetNameListContentForValidation;
 	this->OnGetNameFromSelection = InArgs._OnGetNameFromSelection;
 	this->bMarkupInvalidItems = InArgs._MarkupInvalidItems;
 
-	CurrentList = GetNameList();
+	UpdateNameLists();
 	SGraphPin::Construct(SGraphPin::FArguments(), InGraphPinObj);
 }
 
@@ -112,9 +113,24 @@ TSharedRef<SWidget>	SControlRigGraphPinNameList::GetDefaultValueWidget()
 		];
 }
 
-const TArray<TSharedPtr<FString>>* SControlRigGraphPinNameList::GetNameList() const
+const TArray<TSharedPtr<FString>>* SControlRigGraphPinNameList::GetNameList(bool bForContent) const
 {
 	const TArray<TSharedPtr<FString>>* Result = nullptr;
+	
+	// if we are looking for the validation name list, try to get it from the
+	// dedicated delegate first - but fall back on the content list nevertheless
+	if(!bForContent)
+	{
+		if (OnGetNameListContentForValidation.IsBound())
+		{
+			Result = OnGetNameListContentForValidation.Execute(ModelPin);
+			if(Result)
+			{
+				return Result;
+			}
+		}
+	}
+
 	if (OnGetNameListContent.IsBound())
 	{
 		Result = OnGetNameListContent.Execute(ModelPin);
@@ -143,6 +159,12 @@ void SControlRigGraphPinNameList::SetNameListText(const FText& NewTypeInValue, E
 	}
 }
 
+void SControlRigGraphPinNameList::UpdateNameLists()
+{
+	CurrentList = GetNameList(true);
+	ValidationList = GetNameList(false);
+}
+
 FSlateColor SControlRigGraphPinNameList::GetNameColor() const
 {
 	if(bMarkupInvalidItems)
@@ -150,7 +172,7 @@ FSlateColor SControlRigGraphPinNameList::GetNameColor() const
 		FString CurrentItem = GetNameListText().ToString();
 		
 		bool bFound = false;
-		for (TSharedPtr<FString> Item : (*CurrentList))
+		for (TSharedPtr<FString> Item : (*ValidationList))
 		{
 			if (Item->Equals(CurrentItem))
 			{
@@ -187,7 +209,8 @@ void SControlRigGraphPinNameList::OnNameListChanged(TSharedPtr<FString> NewSelec
 
 void SControlRigGraphPinNameList::OnNameListComboBox()
 {
-	CurrentList = GetNameList();
+	UpdateNameLists();
+
 	TSharedPtr<FString> CurrentlySelected;
 	for (TSharedPtr<FString> Item : (*CurrentList))
 	{
@@ -231,7 +254,7 @@ FReply SControlRigGraphPinNameList::OnGetSelectedClicked()
 						if (URigVMPin* ParentPin = ModelPin->GetParentPin())
 						{
 							Graph->GetController()->SetPinDefaultValue(ParentPin->GetPinPath(), DefaultValue, true, true, false, true);
-							CurrentList = GetNameList();
+							UpdateNameLists();
 						}
 					}
 
@@ -243,7 +266,7 @@ FReply SControlRigGraphPinNameList::OnGetSelectedClicked()
 						if (Element->GetType() == ERigElementType::Bone)
 						{
 							Graph->GetController()->SetPinDefaultValue(ModelPin->GetPinPath(), Keys[0].Name.ToString(), true, true, false, true);
-							CurrentList = GetNameList();
+							UpdateNameLists();
 						}
 					}
 					
@@ -268,7 +291,7 @@ FReply SControlRigGraphPinNameList::OnGetSelectedClicked()
 									if(Key.Type == ElementType)
 									{
 										Graph->GetController()->SetPinDefaultValue(ModelPin->GetPinPath(), Key.Name.ToString(), true, true, false, true);
-										CurrentList = GetNameList();
+										UpdateNameLists();
 									}
 								}
 								break;
diff --git a/Engine/Plugins/Animation/ControlRig/Source/ControlRigEditor/Private/Graph/SControlRigGraphPinNameList.h b/Engine/Plugins/Animation/ControlRig/Source/ControlRigEditor/Private/Graph/SControlRigGraphPinNameList.h
--- a/Engine/Plugins/Animation/ControlRig/Source/ControlRigEditor/Private/Graph/SControlRigGraphPinNameList.h
+++ b/Engine/Plugins/Animation/ControlRig/Source/ControlRigEditor/Private/Graph/SControlRigGraphPinNameList.h
@@ -22,6 +22,7 @@ public:
 
 		SLATE_ARGUMENT(URigVMPin*, ModelPin)
 		SLATE_EVENT(FOnGetNameListContent, OnGetNameListContent)
+		SLATE_EVENT(FOnGetNameListContent, OnGetNameListContentForValidation)
 		SLATE_EVENT(FOnGetNameFromSelection, OnGetNameFromSelection)
 		SLATE_ARGUMENT(bool, MarkupInvalidItems)
 
@@ -34,20 +35,23 @@ protected:
 	virtual TSharedRef<SWidget>	GetDefaultValueWidget() override;
 	//~ End SGraphPin Interface
 
-	const TArray<TSharedPtr<FString>>* GetNameList() const;
+	const TArray<TSharedPtr<FString>>* GetNameList(bool bForContent = true) const;
 	FText GetNameListText() const;
 	FSlateColor GetNameColor() const;
 	virtual void SetNameListText(const FText& NewTypeInValue, ETextCommit::Type CommitInfo);
+	void UpdateNameLists();
 
 	TSharedRef<SWidget> MakeNameListItemWidget(TSharedPtr<FString> InItem);
 	void OnNameListChanged(TSharedPtr<FString> NewSelection, ESelectInfo::Type SelectInfo);
 	void OnNameListComboBox();
 
 	FOnGetNameListContent OnGetNameListContent;
+	FOnGetNameListContent OnGetNameListContentForValidation;
 	URigVMPin* ModelPin;
 	TSharedPtr<SControlRigGraphPinNameListValueWidget> NameListComboBox;
 	TArray<TSharedPtr<FString>> EmptyList;
 	const TArray<TSharedPtr<FString>>* CurrentList;
+	const TArray<TSharedPtr<FString>>* ValidationList;
 	bool bMarkupInvalidItems;
 
 	/** Helper buttons. */
