Control Rig: More string operations fixes, added string equals

diff --git a/Engine/Plugins/Animation/ControlRig/Source/ControlRig/Private/Units/Core/RigUnit_String.cpp b/Engine/Plugins/Animation/ControlRig/Source/ControlRig/Private/Units/Core/RigUnit_String.cpp
--- a/Engine/Plugins/Animation/ControlRig/Source/ControlRig/Private/Units/Core/RigUnit_String.cpp
+++ b/Engine/Plugins/Animation/ControlRig/Source/ControlRig/Private/Units/Core/RigUnit_String.cpp
@@ -3,16 +3,6 @@
 #include "Units/Core/RigUnit_String.h"
 #include "Units/RigUnitContext.h"
 
-FRigUnit_StringToName_Execute()
-{
-	Result = *Value;
-}
-
-FRigUnit_NameToString_Execute()
-{
-	Result = Value.ToString();
-}
-
 FRigUnit_StringConcat_Execute()
 {
 	Result = A + B;
@@ -108,6 +98,11 @@ FRigUnit_StringMiddle_Execute()
 	}
 }
 
+FRigUnit_StringEquals_Execute()
+{
+	Result = A.Equals(B, ESearchCase::CaseSensitive);
+}
+
 FRigUnit_StringFind_Execute()
 {
 	Index = INDEX_NONE;
@@ -280,6 +275,38 @@ FRigVMTemplateTypeMap FRigDispatch_ToString::OnNewArgumentType(const FName& InAr
 	return Types;
 }
 
+FRigVMFunctionPtr FRigDispatch_ToString::GetDispatchFunctionImpl(const FRigVMTemplateTypeMap& InTypes) const
+{
+	// treat name and string special
+	const TRigVMTypeIndex& ValueTypeIndex = InTypes.FindChecked(TEXT("Value"));
+	if(ValueTypeIndex == RigVMTypeUtils::TypeIndex::FName)
+	{
+		return [](FRigVMExtendedExecuteContext& InContext, FRigVMMemoryHandleArray Handles)
+		{
+			check(Handles[0].IsName());
+			check(Handles[1].IsString());
+
+			const FName& Value = *(const FName*)Handles[0].GetData();
+			FString& Result = *(FString*)Handles[1].GetData();
+			Result = Value.ToString();
+		};
+	}
+	if(ValueTypeIndex == RigVMTypeUtils::TypeIndex::FString)
+	{
+		return [](FRigVMExtendedExecuteContext& InContext, FRigVMMemoryHandleArray Handles)
+		{
+			check(Handles[0].IsString());
+			check(Handles[1].IsString());
+
+			const FString& Value = *(const FString*)Handles[0].GetData();
+			FString& Result = *(FString*)Handles[1].GetData();
+			Result = Value;
+		};
+	}
+	
+	return &FRigDispatch_ToString::Execute;
+}
+
 void FRigDispatch_ToString::Execute(FRigVMExtendedExecuteContext& InContext, FRigVMMemoryHandleArray Handles)
 {
 	const FProperty* ValueProperty = Handles[0].GetResolvedProperty(); 
@@ -313,6 +340,38 @@ FRigVMTemplateTypeMap FRigDispatch_FromString::OnNewArgumentType(const FName& In
 	return Types;
 }
 
+FRigVMFunctionPtr FRigDispatch_FromString::GetDispatchFunctionImpl(const FRigVMTemplateTypeMap& InTypes) const
+{
+	// treat name and string special
+	const TRigVMTypeIndex& ResultTypeIndex = InTypes.FindChecked(TEXT("Result"));
+	if(ResultTypeIndex == RigVMTypeUtils::TypeIndex::FName)
+	{
+		return [](FRigVMExtendedExecuteContext& InContext, FRigVMMemoryHandleArray Handles)
+		{
+			check(Handles[0].IsString());
+			check(Handles[1].IsName());
+
+			const FString& String = *(FString*)Handles[0].GetData();
+			FName& Result = *(FName*)Handles[1].GetData();
+			Result = *String;
+		};
+	}
+	if(ResultTypeIndex == RigVMTypeUtils::TypeIndex::FString)
+	{
+		return [](FRigVMExtendedExecuteContext& InContext, FRigVMMemoryHandleArray Handles)
+		{
+			check(Handles[0].IsString());
+			check(Handles[1].IsString());
+
+			const FString& String = *(FString*)Handles[0].GetData();
+			FString& Result = *(FString*)Handles[1].GetData();
+			Result = String;
+		};
+	}
+
+	return &FRigDispatch_FromString::Execute;
+}
+
 void FRigDispatch_FromString::Execute(FRigVMExtendedExecuteContext& InContext, FRigVMMemoryHandleArray Handles)
 {
 	const FProperty* ValueProperty = Handles[1].GetResolvedProperty(); 
diff --git a/Engine/Plugins/Animation/ControlRig/Source/ControlRig/Private/Units/Core/RigUnit_String.h b/Engine/Plugins/Animation/ControlRig/Source/ControlRig/Private/Units/Core/RigUnit_String.h
--- a/Engine/Plugins/Animation/ControlRig/Source/ControlRig/Private/Units/Core/RigUnit_String.h
+++ b/Engine/Plugins/Animation/ControlRig/Source/ControlRig/Private/Units/Core/RigUnit_String.h
@@ -12,54 +12,6 @@ struct CONTROLRIG_API FRigUnit_StringBase : public FRigUnit
 	GENERATED_BODY()
 };
 
-/**
-* Makes a name from a string
-*/
-USTRUCT(meta=(DisplayName="To Name", TemplateName="Cast", Keywords="Make,Construct"))
-struct CONTROLRIG_API FRigUnit_StringToName : public FRigUnit_StringBase
-{
-	GENERATED_BODY()
-
-	FRigUnit_StringToName()
-	{
-		Value = FString();
-		Result = NAME_None;
-	}
-
-	RIGVM_METHOD()
-	virtual void Execute(const FRigUnitContext& Context) override;
-
-	UPROPERTY(meta=(Input))
-	FString Value;
-
-	UPROPERTY(meta=(Output))
-	FName Result;
-};
-
-/**
-* Makes a string from a name
-*/
-USTRUCT(meta=(DisplayName="To Name", TemplateName="Cast", Keywords="Make,Construct"))
-struct CONTROLRIG_API FRigUnit_NameToString : public FRigUnit_StringBase
-{
-	GENERATED_BODY()
-
-	FRigUnit_NameToString()
-	{
-		Value = NAME_None;
-		Result = FString();
-	}
-
-	RIGVM_METHOD()
-	virtual void Execute(const FRigUnitContext& Context) override;
-
-	UPROPERTY(meta=(Input))
-	FName Value;
-
-	UPROPERTY(meta=(Output))
-	FString Result;
-};
-
 /**
  * Concatenates two strings together to make a new string
  */
@@ -437,6 +389,33 @@ struct CONTROLRIG_API FRigUnit_StringMiddle : public FRigUnit_StringBase
 	FString Result;
 };
 
+/**
+ * Compares two strings (case sensitive)
+ */
+USTRUCT(meta = (DisplayName = "Equals", TemplateName = "Equals", Keywords = "Same,Match"))
+struct CONTROLRIG_API FRigUnit_StringEquals : public FRigUnit_StringBase
+{
+	GENERATED_BODY()
+
+	FRigUnit_StringEquals()
+	{
+		A = B = FString();
+		Result = false;
+	}
+
+	RIGVM_METHOD()
+	virtual void Execute(const FRigUnitContext& Context) override;
+
+	UPROPERTY(meta = (Input))
+	FString A;
+
+	UPROPERTY(meta = (Input))
+	FString B;
+
+	UPROPERTY(meta = (Output))
+	bool Result;
+};
+
 /**
  * Finds a string within another string
  */
@@ -563,10 +542,7 @@ public:
 
 protected:
 
-	virtual FRigVMFunctionPtr GetDispatchFunctionImpl(const FRigVMTemplateTypeMap& InTypes) const override
-	{
-		return &FRigDispatch_ToString::Execute;
-	}
+	virtual FRigVMFunctionPtr GetDispatchFunctionImpl(const FRigVMTemplateTypeMap& InTypes) const override;
 	static void Execute(FRigVMExtendedExecuteContext& InContext, FRigVMMemoryHandleArray Handles);
 };
 
@@ -585,9 +561,6 @@ public:
 
 protected:
 
-	virtual FRigVMFunctionPtr GetDispatchFunctionImpl(const FRigVMTemplateTypeMap& InTypes) const override
-	{
-		return &FRigDispatch_FromString::Execute;
-	}
+	virtual FRigVMFunctionPtr GetDispatchFunctionImpl(const FRigVMTemplateTypeMap& InTypes) const override;
 	static void Execute(FRigVMExtendedExecuteContext& InContext, FRigVMMemoryHandleArray Handles);
 };
\ No newline at end of file
