Control Rig: Remove strong dependencies from User interface on data model - use TWeakObejctPtr

diff --git a/Engine/Plugins/Animation/ControlRig/Source/ControlRigDeveloper/Private/Graph/ControlRigGraphNode.cpp b/Engine/Plugins/Animation/ControlRig/Source/ControlRigDeveloper/Private/Graph/ControlRigGraphNode.cpp
--- a/Engine/Plugins/Animation/ControlRig/Source/ControlRigDeveloper/Private/Graph/ControlRigGraphNode.cpp
+++ b/Engine/Plugins/Animation/ControlRig/Source/ControlRigDeveloper/Private/Graph/ControlRigGraphNode.cpp
@@ -1016,15 +1016,15 @@ URigVMController* UControlRigGraphNode::GetController() const
 URigVMNode* UControlRigGraphNode::GetModelNode() const
 {
 	UControlRigGraphNode* MutableThis = (UControlRigGraphNode*)this;
-	if (CachedModelNode)
+	if (CachedModelNode.IsValid())
 	{
-		if (CachedModelNode->GetOuter() == GetTransientPackage())
+		if (CachedModelNode.Get()->GetOuter() == GetTransientPackage())
 		{
-			MutableThis->CachedModelNode = nullptr;
+			MutableThis->CachedModelNode.Reset();
 		}
 		else
 		{
-			return CachedModelNode;
+			return CachedModelNode.Get();
 		}
 	}
 
@@ -1034,14 +1034,16 @@ URigVMNode* UControlRigGraphNode::GetModelNode() const
 
 		if (Graph->TemplateController != nullptr)
 		{
-			return MutableThis->CachedModelNode = Graph->TemplateController->GetGraph()->FindNode(ModelNodePath);
+			MutableThis->CachedModelNode = TWeakObjectPtr<URigVMNode>(Graph->TemplateController->GetGraph()->FindNode(ModelNodePath));
+			return MutableThis->CachedModelNode.Get();
 		}
 
 #endif
 
 		if (URigVMGraph* Model = GetModel())
 		{
-			return MutableThis->CachedModelNode = Model->FindNode(ModelNodePath);
+			MutableThis->CachedModelNode = TWeakObjectPtr<URigVMNode>(Model->FindNode(ModelNodePath));
+			return MutableThis->CachedModelNode.Get();
 		}
 	}
 
@@ -1059,12 +1061,15 @@ FName UControlRigGraphNode::GetModelNodeName() const
 
 URigVMPin* UControlRigGraphNode::GetModelPinFromPinPath(const FString& InPinPath) const
 {
-	if (TObjectPtr<URigVMPin> const* CachedModelPinPtr = CachedModelPins.Find(InPinPath))
+	if (TWeakObjectPtr<URigVMPin> const* CachedModelPinPtr = CachedModelPins.Find(InPinPath))
 	{
-		URigVMPin* CachedModelPin = *CachedModelPinPtr;
-		if (!CachedModelPin->HasAnyFlags(RF_Transient) && CachedModelPin->GetNode())
+		if(CachedModelPinPtr->IsValid())
 		{
-			return CachedModelPin;
+			URigVMPin* CachedModelPin = CachedModelPinPtr->Get();
+			if (!CachedModelPin->HasAnyFlags(RF_Transient) && CachedModelPin->GetNode())
+			{
+				return CachedModelPin;
+			}
 		}
 	}
 
diff --git a/Engine/Plugins/Animation/ControlRig/Source/ControlRigDeveloper/Public/Graph/ControlRigGraphNode.h b/Engine/Plugins/Animation/ControlRig/Source/ControlRigDeveloper/Public/Graph/ControlRigGraphNode.h
--- a/Engine/Plugins/Animation/ControlRig/Source/ControlRigDeveloper/Public/Graph/ControlRigGraphNode.h
+++ b/Engine/Plugins/Animation/ControlRig/Source/ControlRigDeveloper/Public/Graph/ControlRigGraphNode.h
@@ -5,6 +5,7 @@
 
 #include "CoreMinimal.h"
 #include "UObject/ObjectMacros.h"
+#include "UObject/WeakObjectPtrTemplates.h"
 #include "EdGraph/EdGraphNodeUtils.h"
 #include "K2Node.h"
 #include "UObject/SoftObjectPath.h"
@@ -40,10 +41,10 @@ private:
 	FString ModelNodePath;
 
 	UPROPERTY(transient)
-	TObjectPtr<URigVMNode> CachedModelNode;
+	TWeakObjectPtr<URigVMNode> CachedModelNode;
 
 	UPROPERTY(transient)
-	TMap<FString, TObjectPtr<URigVMPin>> CachedModelPins;
+	TMap<FString, TWeakObjectPtr<URigVMPin>> CachedModelPins;
 
 	/** The property we represent. For template nodes this represents the struct/property type name. */
 	UPROPERTY()
