Sequencer: Control Rig: Crash when keying onto addtive section when the sections have different number of controls due to the wrong time being used due to hidden track editor variable.

diff --git a/Engine/Plugins/Animation/ControlRig/Source/ControlRigEditor/Private/Sequencer/ControlRigParameterTrackEditor.cpp b/Engine/Plugins/Animation/ControlRig/Source/ControlRigEditor/Private/Sequencer/ControlRigParameterTrackEditor.cpp
--- a/Engine/Plugins/Animation/ControlRig/Source/ControlRigEditor/Private/Sequencer/ControlRigParameterTrackEditor.cpp
+++ b/Engine/Plugins/Animation/ControlRig/Source/ControlRigEditor/Private/Sequencer/ControlRigParameterTrackEditor.cpp
@@ -2867,6 +2867,13 @@ void FControlRigParameterTrackEditor::AddControlKeys(
 	UMovieSceneControlRigParameterSection* ParamSection = nullptr;
 	if (Track)
 	{
+		//track editors use a hidden time so we need to set it if we are using non sequencer times when keying.
+		if (InLocalTime != FLT_MAX)
+		{
+			//convert from frame time since conversion may give us one frame less, e.g 1.53333330 * 24000.0/1.0 = 36799.999199999998
+			FFrameTime LocalFrameTime = GetSequencer()->GetFocusedTickResolution().AsFrameTime((double)InLocalTime);
+			BeginKeying(LocalFrameTime.RoundToFrame());
+		}
 		FFrameNumber  FrameTime = GetTimeForKey();
 		UMovieSceneSection* Section = Track->FindSection(FrameTime);
 		ParamSection = Cast<UMovieSceneControlRigParameterSection>(Section);
@@ -2900,7 +2907,7 @@ void FControlRigParameterTrackEditor::AddControlKeys(
 	};
 
 	AnimatablePropertyChanged(FOnKeyProperty::CreateLambda(OnKeyProperty));
-
+	EndKeying(); //fine even if we didn't BeginKeying
 }
 
 bool FControlRigParameterTrackEditor::ModifyOurGeneratedKeysByCurrentAndWeight(UObject* Object, UControlRig* InControlRig, FName RigControlName, UMovieSceneTrack* Track, UMovieSceneSection* SectionToKey, FFrameNumber KeyTime, FGeneratedTrackKeys& GeneratedTotalKeys, float Weight) const
