Control Rig: deprecate equals nodes - replace with dispatch

diff --git a/Engine/Plugins/Animation/ControlRig/Source/ControlRig/Private/Units/Core/RigUnit_CoreDispatch.cpp b/Engine/Plugins/Animation/ControlRig/Source/ControlRig/Private/Units/Core/RigUnit_CoreDispatch.cpp
new file mode 100644
--- /dev/null
+++ b/Engine/Plugins/Animation/ControlRig/Source/ControlRig/Private/Units/Core/RigUnit_CoreDispatch.cpp
@@ -0,0 +1,112 @@
+// Copyright Epic Games, Inc. All Rights Reserved.
+
+#include "Units/Core/RigUnit_CoreDispatch.h"
+#include "Units/RigUnitContext.h"
+#include "EulerTransform.h"
+
+TArray<FRigVMTemplateArgument> FRigDispatch_CoreEquals::GetArguments() const
+{
+	const TArray<FRigVMTemplateArgument::ETypeCategory> ValueCategories = {
+		FRigVMTemplateArgument::ETypeCategory_SingleAnyValue,
+		FRigVMTemplateArgument::ETypeCategory_ArrayAnyValue
+	};
+	return {
+		FRigVMTemplateArgument(TEXT("A"), ERigVMPinDirection::Input, ValueCategories),
+		FRigVMTemplateArgument(TEXT("B"), ERigVMPinDirection::Input, ValueCategories),
+		FRigVMTemplateArgument(TEXT("Result"), ERigVMPinDirection::Output, RigVMTypeUtils::TypeIndex::Bool)
+	};
+}
+
+FRigVMTemplateTypeMap FRigDispatch_CoreEquals::OnNewArgumentType(const FName& InArgumentName,
+	TRigVMTypeIndex InTypeIndex) const
+{
+	FRigVMTemplateTypeMap Types;
+	Types.Add(TEXT("A"), InTypeIndex);
+	Types.Add(TEXT("B"), InTypeIndex);
+	Types.Add(TEXT("Result"), RigVMTypeUtils::TypeIndex::Bool);
+	return Types;
+}
+
+FRigVMFunctionPtr FRigDispatch_CoreEquals::GetDispatchFunctionImpl(const FRigVMTemplateTypeMap& InTypes) const
+{
+	const TRigVMTypeIndex TypeIndex = InTypes.FindChecked(TEXT("A"));
+	check(TypeIndex == InTypes.FindChecked(TEXT("B")));
+	check(InTypes.FindChecked(TEXT("Result")) == RigVMTypeUtils::TypeIndex::Bool);
+
+	if(TypeIndex == RigVMTypeUtils::TypeIndex::Float)
+	{
+		return &FRigDispatch_CoreEquals::Equals<float>;
+	}
+	if(TypeIndex == RigVMTypeUtils::TypeIndex::Double)
+	{
+		return &FRigDispatch_CoreEquals::Equals<double>;
+	}
+	if(TypeIndex == RigVMTypeUtils::TypeIndex::Int32)
+	{
+		return &FRigDispatch_CoreEquals::Equals<int32>;
+	}
+	if(TypeIndex == RigVMTypeUtils::TypeIndex::FName)
+	{
+		return &FRigDispatch_CoreEquals::NameEquals;
+	}
+	if(TypeIndex == RigVMTypeUtils::TypeIndex::FString)
+	{
+		return &FRigDispatch_CoreEquals::StringEquals;
+	}
+	if(TypeIndex == FRigVMRegistry::Get().GetTypeIndex<FVector>())
+	{
+		return &FRigDispatch_CoreEquals::MathTypeEquals<FVector>;
+	}
+	if(TypeIndex == FRigVMRegistry::Get().GetTypeIndex<FVector2D>())
+	{
+		return &FRigDispatch_CoreEquals::MathTypeEquals<FVector2D>;
+	}
+	if(TypeIndex == FRigVMRegistry::Get().GetTypeIndex<FRotator>())
+	{
+		return &FRigDispatch_CoreEquals::MathTypeEquals<FRotator>;
+	}
+	if(TypeIndex == FRigVMRegistry::Get().GetTypeIndex<FQuat>())
+	{
+		return &FRigDispatch_CoreEquals::MathTypeEquals<FQuat>;
+	}
+	if(TypeIndex == FRigVMRegistry::Get().GetTypeIndex<FTransform>())
+	{
+		return &FRigDispatch_CoreEquals::MathTypeEquals<FTransform>;
+	}
+	if(TypeIndex == FRigVMRegistry::Get().GetTypeIndex<FEulerTransform>())
+	{
+		return &FRigDispatch_CoreEquals::MathTypeEquals<FEulerTransform>;
+	}
+	if(TypeIndex == FRigVMRegistry::Get().GetTypeIndex<FLinearColor>())
+	{
+		return &FRigDispatch_CoreEquals::MathTypeEquals<FLinearColor>;
+	}
+
+	return FRigDispatch_CoreBase::GetDispatchFunctionImpl(InTypes);
+}
+
+bool FRigDispatch_CoreEquals::AdaptResult(bool bResult, const FRigVMExtendedExecuteContext& InContext)
+{
+	// if the factory is the not equals factory - let's invert the result
+	if(InContext.Factory->GetScriptStruct() == FRigDispatch_CoreNotEquals::StaticStruct())
+	{
+		return !bResult;
+	}
+	return bResult;
+}
+
+void FRigDispatch_CoreEquals::Execute(FRigVMExtendedExecuteContext& InContext, FRigVMMemoryHandleArray Handles)
+{
+	const FProperty* PropertyA = Handles[0].GetResolvedProperty(); 
+	const FProperty* PropertyB = Handles[1].GetResolvedProperty(); 
+	check(PropertyA);
+	check(PropertyB);
+	check(PropertyA->SameType(PropertyB));
+	check(Handles[2].IsBool());
+
+	const uint8* A = Handles[0].GetData();
+	const uint8* B = Handles[1].GetData();
+	bool& Result = *(bool*)Handles[2].GetData();
+	Result = PropertyA->Identical(A, B);
+	Result = AdaptResult(Result, InContext);
+}
diff --git a/Engine/Plugins/Animation/ControlRig/Source/ControlRig/Private/Units/Core/RigUnit_CoreDispatch.h b/Engine/Plugins/Animation/ControlRig/Source/ControlRig/Private/Units/Core/RigUnit_CoreDispatch.h
new file mode 100644
--- /dev/null
+++ b/Engine/Plugins/Animation/ControlRig/Source/ControlRig/Private/Units/Core/RigUnit_CoreDispatch.h
@@ -0,0 +1,93 @@
+// Copyright Epic Games, Inc. All Rights Reserved.
+
+#pragma once
+
+#include "Units/RigUnit.h"
+#include "Units/RigDispatchFactory.h"
+#include "RigUnit_CoreDispatch.generated.h"
+
+USTRUCT(meta=(Abstract, Category = "Core", NodeColor = "0.762745, 1,0, 0.329412"))
+struct FRigDispatch_CoreBase : public FRigDispatchFactory
+{
+	GENERATED_BODY()
+};
+
+/*
+ * Compares any two values and return true if they are identical
+ */
+USTRUCT(meta=(DisplayName = "Equals", Keywords = "Same,=="))
+struct FRigDispatch_CoreEquals : public FRigDispatch_CoreBase
+{
+	GENERATED_BODY()
+
+public:
+
+	virtual TArray<FRigVMTemplateArgument> GetArguments() const override;
+	virtual FRigVMTemplateTypeMap OnNewArgumentType(const FName& InArgumentName, TRigVMTypeIndex InTypeIndex) const override;
+
+protected:
+
+	virtual FRigVMFunctionPtr GetDispatchFunctionImpl(const FRigVMTemplateTypeMap& InTypes) const override;
+	static bool AdaptResult(bool bResult, const FRigVMExtendedExecuteContext& InContext);
+	static void Execute(FRigVMExtendedExecuteContext& InContext, FRigVMMemoryHandleArray Handles);
+
+	template<typename T>
+	static void Equals(FRigVMExtendedExecuteContext& InContext, FRigVMMemoryHandleArray Handles)
+	{
+		check(Handles[0].IsType<T>());
+		check(Handles[1].IsType<T>());
+		check(Handles[2].IsBool());
+		const T& A = *(const T*)Handles[0].GetData();
+		const T& B = *(const T*)Handles[1].GetData();
+		bool& Result = *(bool*)Handles[2].GetData();
+		Result = AdaptResult(A == B, InContext);
+	}
+
+	static void NameEquals(FRigVMExtendedExecuteContext& InContext, FRigVMMemoryHandleArray Handles)
+	{
+		check(Handles[0].IsType<FName>());
+		check(Handles[1].IsType<FName>());
+		check(Handles[2].IsBool());
+		const FName& A = *(const FName*)Handles[0].GetData();
+		const FName& B = *(const FName*)Handles[1].GetData();
+		bool& Result = *(bool*)Handles[2].GetData();
+		Result = AdaptResult(A.IsEqual(B, ENameCase::CaseSensitive), InContext);
+	}
+
+	static void StringEquals(FRigVMExtendedExecuteContext& InContext, FRigVMMemoryHandleArray Handles)
+	{
+		check(Handles[0].IsType<FString>());
+		check(Handles[1].IsType<FString>());
+		check(Handles[2].IsBool());
+		const FString& A = *(const FString*)Handles[0].GetData();
+		const FString& B = *(const FString*)Handles[1].GetData();
+		bool& Result = *(bool*)Handles[2].GetData();
+		Result = AdaptResult(A.Equals(B, ESearchCase::CaseSensitive), InContext);
+	}
+
+	template<typename T>
+	static void MathTypeEquals(FRigVMExtendedExecuteContext& InContext, FRigVMMemoryHandleArray Handles)
+	{
+		check(Handles[0].IsType<T>());
+		check(Handles[1].IsType<T>());
+		check(Handles[2].IsBool());
+		const T& A = *(const T*)Handles[0].GetData();
+		const T& B = *(const T*)Handles[1].GetData();
+		bool& Result = *(bool*)Handles[2].GetData();
+		Result = AdaptResult(A.Equals(B), InContext);
+	}
+};
+
+/*
+ * Compares any two values and return true if they are identical
+ */
+USTRUCT(meta=(DisplayName = "Not Equals", Keywords = "Different,!=,Xor"))
+struct FRigDispatch_CoreNotEquals : public FRigDispatch_CoreEquals
+{
+	GENERATED_BODY()
+
+	// we are inheriting everything from the equals dispatch,
+	// and due to the check of the factory within FRigDispatch_CoreEquals::Execute we can
+	// rely on that completely. we only need this class for the displayname and
+	// operation specific StaticStruct().
+};
diff --git a/Engine/Plugins/Animation/ControlRig/Source/ControlRig/Private/Units/Core/RigUnit_String.cpp b/Engine/Plugins/Animation/ControlRig/Source/ControlRig/Private/Units/Core/RigUnit_String.cpp
--- a/Engine/Plugins/Animation/ControlRig/Source/ControlRig/Private/Units/Core/RigUnit_String.cpp
+++ b/Engine/Plugins/Animation/ControlRig/Source/ControlRig/Private/Units/Core/RigUnit_String.cpp
@@ -98,11 +98,6 @@ FRigUnit_StringMiddle_Execute()
 	}
 }
 
-FRigUnit_StringEquals_Execute()
-{
-	Result = A.Equals(B, ESearchCase::CaseSensitive);
-}
-
 FRigUnit_StringFind_Execute()
 {
 	Index = INDEX_NONE;
diff --git a/Engine/Plugins/Animation/ControlRig/Source/ControlRig/Private/Units/Core/RigUnit_String.h b/Engine/Plugins/Animation/ControlRig/Source/ControlRig/Private/Units/Core/RigUnit_String.h
--- a/Engine/Plugins/Animation/ControlRig/Source/ControlRig/Private/Units/Core/RigUnit_String.h
+++ b/Engine/Plugins/Animation/ControlRig/Source/ControlRig/Private/Units/Core/RigUnit_String.h
@@ -389,33 +389,6 @@ struct CONTROLRIG_API FRigUnit_StringMiddle : public FRigUnit_StringBase
 	FString Result;
 };
 
-/**
- * Compares two strings (case sensitive)
- */
-USTRUCT(meta = (DisplayName = "Equals", TemplateName = "Equals", Keywords = "Same,Match"))
-struct CONTROLRIG_API FRigUnit_StringEquals : public FRigUnit_StringBase
-{
-	GENERATED_BODY()
-
-	FRigUnit_StringEquals()
-	{
-		A = B = FString();
-		Result = false;
-	}
-
-	RIGVM_METHOD()
-	virtual void Execute(const FRigUnitContext& Context) override;
-
-	UPROPERTY(meta = (Input))
-	FString A;
-
-	UPROPERTY(meta = (Input))
-	FString B;
-
-	UPROPERTY(meta = (Output))
-	bool Result;
-};
-
 /**
  * Finds a string within another string
  */
@@ -530,7 +503,7 @@ struct CONTROLRIG_API FRigUnit_StringPadInteger : public FRigUnit_StringBase
 /*
  * Converts any value to string
  */
-USTRUCT(meta=(DisplayName = "To String", odeColor = "0.462745, 1,0, 0.329412"))
+USTRUCT(meta=(DisplayName = "To String", NodeColor = "0.462745, 1,0, 0.329412"))
 struct FRigDispatch_ToString : public FRigDispatchFactory
 {
 	GENERATED_BODY()
diff --git a/Engine/Plugins/Animation/ControlRig/Source/ControlRig/Private/Units/Execution/RigUnit_Item.cpp b/Engine/Plugins/Animation/ControlRig/Source/ControlRig/Private/Units/Execution/RigUnit_Item.cpp
--- a/Engine/Plugins/Animation/ControlRig/Source/ControlRig/Private/Units/Execution/RigUnit_Item.cpp
+++ b/Engine/Plugins/Animation/ControlRig/Source/ControlRig/Private/Units/Execution/RigUnit_Item.cpp
@@ -3,6 +3,7 @@
 #include "RigUnit_Item.h"
 #include "Units/Core/RigUnit_Name.h"
 #include "Units/RigUnitContext.h"
+#include "Units/Core/RigUnit_CoreDispatch.h"
 
 FRigUnit_ItemExists_Execute()
 {
@@ -42,11 +43,21 @@ FRigUnit_ItemEquals_Execute()
 	Result = (A == B);
 }
 
+FRigVMStructUpgradeInfo FRigUnit_ItemEquals::GetUpgradeInfo() const
+{
+	return FRigVMStructUpgradeInfo::MakeFromStructToFactory(StaticStruct(), FRigDispatch_CoreEquals::StaticStruct());
+}
+
 FRigUnit_ItemNotEquals_Execute()
 {
 	Result = (A != B);
 }
 
+FRigVMStructUpgradeInfo FRigUnit_ItemNotEquals::GetUpgradeInfo() const
+{
+	return FRigVMStructUpgradeInfo::MakeFromStructToFactory(StaticStruct(), FRigDispatch_CoreNotEquals::StaticStruct());
+}
+
 FRigUnit_ItemTypeEquals_Execute()
 {
 	Result = (A.Type == B.Type);
diff --git a/Engine/Plugins/Animation/ControlRig/Source/ControlRig/Private/Units/Execution/RigUnit_Item.h b/Engine/Plugins/Animation/ControlRig/Source/ControlRig/Private/Units/Execution/RigUnit_Item.h
--- a/Engine/Plugins/Animation/ControlRig/Source/ControlRig/Private/Units/Execution/RigUnit_Item.h
+++ b/Engine/Plugins/Animation/ControlRig/Source/ControlRig/Private/Units/Execution/RigUnit_Item.h
@@ -80,7 +80,7 @@ struct CONTROLRIG_API FRigUnit_ItemReplace : public FRigUnit_ItemBase
 /**
 * Returns true if the two items are equal
 */
-USTRUCT(meta=(DisplayName="Item Equals", Keywords=""))
+USTRUCT(meta=(DisplayName="Item Equals", Keywords="", Deprecated="5.1"))
 struct CONTROLRIG_API FRigUnit_ItemEquals : public FRigUnit_ItemBase
 {
 	GENERATED_BODY()
@@ -102,12 +102,15 @@ struct CONTROLRIG_API FRigUnit_ItemEquals : public FRigUnit_ItemBase
 
 	UPROPERTY(meta = (Output))
 	bool Result;
+
+	RIGVM_METHOD()
+	virtual FRigVMStructUpgradeInfo GetUpgradeInfo() const override;
 };
 
 /**
 * Returns true if the two items are not equal
 */
-USTRUCT(meta=(DisplayName="Item Not Equals", Keywords=""))
+USTRUCT(meta=(DisplayName="Item Not Equals", Keywords="", Deprecated="5.1"))
 struct CONTROLRIG_API FRigUnit_ItemNotEquals : public FRigUnit_ItemBase
 {
 	GENERATED_BODY()
@@ -129,6 +132,9 @@ struct CONTROLRIG_API FRigUnit_ItemNotEquals : public FRigUnit_ItemBase
 
 	UPROPERTY(meta = (Output))
 	bool Result;
+
+	RIGVM_METHOD()
+	virtual FRigVMStructUpgradeInfo GetUpgradeInfo() const override;
 };
 
 /**
diff --git a/Engine/Plugins/Animation/ControlRig/Source/ControlRig/Private/Units/Math/RigUnit_MathBool.cpp b/Engine/Plugins/Animation/ControlRig/Source/ControlRig/Private/Units/Math/RigUnit_MathBool.cpp
--- a/Engine/Plugins/Animation/ControlRig/Source/ControlRig/Private/Units/Math/RigUnit_MathBool.cpp
+++ b/Engine/Plugins/Animation/ControlRig/Source/ControlRig/Private/Units/Math/RigUnit_MathBool.cpp
@@ -43,12 +43,22 @@ FRigUnit_MathBoolEquals_Execute()
 	Result = A == B;
 }
 
+FRigVMStructUpgradeInfo FRigUnit_MathBoolEquals::GetUpgradeInfo() const
+{
+	return FRigVMStructUpgradeInfo::MakeFromStructToFactory(StaticStruct(), FRigDispatch_CoreEquals::StaticStruct());
+}
+
 FRigUnit_MathBoolNotEquals_Execute()
 {
     DECLARE_SCOPE_HIERARCHICAL_COUNTER_RIGUNIT()
 	Result = A != B;
 }
 
+FRigVMStructUpgradeInfo FRigUnit_MathBoolNotEquals::GetUpgradeInfo() const
+{
+	return FRigVMStructUpgradeInfo::MakeFromStructToFactory(StaticStruct(), FRigDispatch_CoreNotEquals::StaticStruct());
+}
+
 FRigUnit_MathBoolToggled_Execute()
 {
 	DECLARE_SCOPE_HIERARCHICAL_COUNTER_RIGUNIT()
diff --git a/Engine/Plugins/Animation/ControlRig/Source/ControlRig/Private/Units/Math/RigUnit_MathBool.h b/Engine/Plugins/Animation/ControlRig/Source/ControlRig/Private/Units/Math/RigUnit_MathBool.h
--- a/Engine/Plugins/Animation/ControlRig/Source/ControlRig/Private/Units/Math/RigUnit_MathBool.h
+++ b/Engine/Plugins/Animation/ControlRig/Source/ControlRig/Private/Units/Math/RigUnit_MathBool.h
@@ -3,6 +3,7 @@
 #pragma once
 
 #include "RigUnit_MathBase.h"
+#include "Units/Core/RigUnit_CoreDispatch.h"
 #include "RigUnit_MathBool.generated.h"
 
 USTRUCT(meta=(Abstract, Category="Math|Boolean"))
@@ -147,7 +148,7 @@ struct CONTROLRIG_API FRigUnit_MathBoolOr : public FRigUnit_MathBoolBinaryOp
 /**
  * Returns true if the value A equals B
  */
-USTRUCT(meta=(DisplayName="Equals", TemplateName="Equals", Keywords="Same,=="))
+USTRUCT(meta=(DisplayName="Equals", TemplateName="Equals", Keywords="Same,==", Deprecated="5.1"))
 struct CONTROLRIG_API FRigUnit_MathBoolEquals : public FRigUnit_MathBoolBase
 {
 	GENERATED_BODY()
@@ -168,12 +169,15 @@ struct CONTROLRIG_API FRigUnit_MathBoolEquals : public FRigUnit_MathBoolBase
 
 	UPROPERTY(meta=(Output))
 	bool Result;
+
+	RIGVM_METHOD()
+	virtual FRigVMStructUpgradeInfo GetUpgradeInfo() const override;
 };
 
 /**
  * Returns true if the value A does not equal B
  */
-USTRUCT(meta=(DisplayName="Not Equals", TemplateName="NotEquals", Keywords="Different,!=,Xor"))
+USTRUCT(meta=(DisplayName="Not Equals", TemplateName="NotEquals", Keywords="Different,!=,Xor", Deprecated="5.1"))
 struct CONTROLRIG_API FRigUnit_MathBoolNotEquals : public FRigUnit_MathBoolBase
 {
 	GENERATED_BODY()
@@ -194,6 +198,9 @@ struct CONTROLRIG_API FRigUnit_MathBoolNotEquals : public FRigUnit_MathBoolBase
 
 	UPROPERTY(meta=(Output))
 	bool Result;
+
+	RIGVM_METHOD()
+	virtual FRigVMStructUpgradeInfo GetUpgradeInfo() const override;
 };
 
 /**
diff --git a/Engine/Plugins/Animation/ControlRig/Source/ControlRig/Private/Units/Math/RigUnit_MathDouble.cpp b/Engine/Plugins/Animation/ControlRig/Source/ControlRig/Private/Units/Math/RigUnit_MathDouble.cpp
--- a/Engine/Plugins/Animation/ControlRig/Source/ControlRig/Private/Units/Math/RigUnit_MathDouble.cpp
+++ b/Engine/Plugins/Animation/ControlRig/Source/ControlRig/Private/Units/Math/RigUnit_MathDouble.cpp
@@ -2,6 +2,7 @@
 
 #include "Units/Math/RigUnit_MathDouble.h"
 #include "Units/RigUnitContext.h"
+#include "Units/Core/RigUnit_CoreDispatch.h"
 
 FRigUnit_MathDoubleConstPi_Execute()
 {
@@ -177,12 +178,22 @@ FRigUnit_MathDoubleEquals_Execute()
 	Result = A == B;
 }
 
+FRigVMStructUpgradeInfo FRigUnit_MathDoubleEquals::GetUpgradeInfo() const
+{
+	return FRigVMStructUpgradeInfo::MakeFromStructToFactory(StaticStruct(), FRigDispatch_CoreEquals::StaticStruct());
+}
+
 FRigUnit_MathDoubleNotEquals_Execute()
 {
     DECLARE_SCOPE_HIERARCHICAL_COUNTER_RIGUNIT()
 	Result = A != B;
 }
 
+FRigVMStructUpgradeInfo FRigUnit_MathDoubleNotEquals::GetUpgradeInfo() const
+{
+	return FRigVMStructUpgradeInfo::MakeFromStructToFactory(StaticStruct(), FRigDispatch_CoreNotEquals::StaticStruct());
+}
+
 FRigUnit_MathDoubleGreater_Execute()
 {
     DECLARE_SCOPE_HIERARCHICAL_COUNTER_RIGUNIT()
diff --git a/Engine/Plugins/Animation/ControlRig/Source/ControlRig/Private/Units/Math/RigUnit_MathDouble.h b/Engine/Plugins/Animation/ControlRig/Source/ControlRig/Private/Units/Math/RigUnit_MathDouble.h
--- a/Engine/Plugins/Animation/ControlRig/Source/ControlRig/Private/Units/Math/RigUnit_MathDouble.h
+++ b/Engine/Plugins/Animation/ControlRig/Source/ControlRig/Private/Units/Math/RigUnit_MathDouble.h
@@ -485,7 +485,7 @@ struct CONTROLRIG_API FRigUnit_MathDoubleRemap : public FRigUnit_MathDoubleBase
 /**
  * Returns true if the value A equals B
  */
-USTRUCT(meta=(DisplayName="Equals", TemplateName="Equals", Keywords="Same,=="))
+USTRUCT(meta=(DisplayName="Equals", TemplateName="Equals", Keywords="Same,==", Deprecated="5.1"))
 struct CONTROLRIG_API FRigUnit_MathDoubleEquals : public FRigUnit_MathDoubleBase
 {
 	GENERATED_BODY()
@@ -506,12 +506,15 @@ struct CONTROLRIG_API FRigUnit_MathDoubleEquals : public FRigUnit_MathDoubleBase
 
 	UPROPERTY(meta=(Output))
 	bool Result;
+
+	RIGVM_METHOD()
+	virtual FRigVMStructUpgradeInfo GetUpgradeInfo() const override;
 };
 
 /**
  * Returns true if the value A does not equal B
  */
-USTRUCT(meta=(DisplayName="Not Equals", TemplateName="NotEquals", Keywords="Different,!="))
+USTRUCT(meta=(DisplayName="Not Equals", TemplateName="NotEquals", Keywords="Different,!=", Deprecated="5.1"))
 struct CONTROLRIG_API FRigUnit_MathDoubleNotEquals : public FRigUnit_MathDoubleBase
 {
 	GENERATED_BODY()
@@ -532,6 +535,9 @@ struct CONTROLRIG_API FRigUnit_MathDoubleNotEquals : public FRigUnit_MathDoubleB
 
 	UPROPERTY(meta=(Output))
 	bool Result;
+
+	RIGVM_METHOD()
+	virtual FRigVMStructUpgradeInfo GetUpgradeInfo() const override;
 };
 
 /**
diff --git a/Engine/Plugins/Animation/ControlRig/Source/ControlRig/Private/Units/Math/RigUnit_MathFloat.cpp b/Engine/Plugins/Animation/ControlRig/Source/ControlRig/Private/Units/Math/RigUnit_MathFloat.cpp
--- a/Engine/Plugins/Animation/ControlRig/Source/ControlRig/Private/Units/Math/RigUnit_MathFloat.cpp
+++ b/Engine/Plugins/Animation/ControlRig/Source/ControlRig/Private/Units/Math/RigUnit_MathFloat.cpp
@@ -2,6 +2,7 @@
 
 #include "Units/Math/RigUnit_MathFloat.h"
 #include "Units/RigUnitContext.h"
+#include "Units/Core/RigUnit_CoreDispatch.h"
 
 FRigUnit_MathFloatConstPi_Execute()
 {
@@ -177,12 +178,22 @@ FRigUnit_MathFloatEquals_Execute()
 	Result = A == B;
 }
 
+FRigVMStructUpgradeInfo FRigUnit_MathFloatEquals::GetUpgradeInfo() const
+{
+	return FRigVMStructUpgradeInfo::MakeFromStructToFactory(StaticStruct(), FRigDispatch_CoreEquals::StaticStruct());
+}
+
 FRigUnit_MathFloatNotEquals_Execute()
 {
     DECLARE_SCOPE_HIERARCHICAL_COUNTER_RIGUNIT()
 	Result = A != B;
 }
 
+FRigVMStructUpgradeInfo FRigUnit_MathFloatNotEquals::GetUpgradeInfo() const
+{
+	return FRigVMStructUpgradeInfo::MakeFromStructToFactory(StaticStruct(), FRigDispatch_CoreNotEquals::StaticStruct());
+}
+
 FRigUnit_MathFloatGreater_Execute()
 {
     DECLARE_SCOPE_HIERARCHICAL_COUNTER_RIGUNIT()
diff --git a/Engine/Plugins/Animation/ControlRig/Source/ControlRig/Private/Units/Math/RigUnit_MathFloat.h b/Engine/Plugins/Animation/ControlRig/Source/ControlRig/Private/Units/Math/RigUnit_MathFloat.h
--- a/Engine/Plugins/Animation/ControlRig/Source/ControlRig/Private/Units/Math/RigUnit_MathFloat.h
+++ b/Engine/Plugins/Animation/ControlRig/Source/ControlRig/Private/Units/Math/RigUnit_MathFloat.h
@@ -485,7 +485,7 @@ struct CONTROLRIG_API FRigUnit_MathFloatRemap : public FRigUnit_MathFloatBase
 /**
  * Returns true if the value A equals B
  */
-USTRUCT(meta=(DisplayName="Equals", TemplateName="Equals", Keywords="Same,=="))
+USTRUCT(meta=(DisplayName="Equals", TemplateName="Equals", Keywords="Same,==", Deprecated="5.1"))
 struct CONTROLRIG_API FRigUnit_MathFloatEquals : public FRigUnit_MathFloatBase
 {
 	GENERATED_BODY()
@@ -506,12 +506,15 @@ struct CONTROLRIG_API FRigUnit_MathFloatEquals : public FRigUnit_MathFloatBase
 
 	UPROPERTY(meta=(Output))
 	bool Result;
+
+	RIGVM_METHOD()
+	virtual FRigVMStructUpgradeInfo GetUpgradeInfo() const override;
 };
 
 /**
  * Returns true if the value A does not equal B
  */
-USTRUCT(meta=(DisplayName="Not Equals", TemplateName="NotEquals", Keywords="Different,!="))
+USTRUCT(meta=(DisplayName="Not Equals", TemplateName="NotEquals", Keywords="Different,!=", Deprecated="5.1"))
 struct CONTROLRIG_API FRigUnit_MathFloatNotEquals : public FRigUnit_MathFloatBase
 {
 	GENERATED_BODY()
@@ -532,6 +535,9 @@ struct CONTROLRIG_API FRigUnit_MathFloatNotEquals : public FRigUnit_MathFloatBas
 
 	UPROPERTY(meta=(Output))
 	bool Result;
+
+	RIGVM_METHOD()
+	virtual FRigVMStructUpgradeInfo GetUpgradeInfo() const override;
 };
 
 /**
diff --git a/Engine/Plugins/Animation/ControlRig/Source/ControlRig/Private/Units/Math/RigUnit_MathInt.cpp b/Engine/Plugins/Animation/ControlRig/Source/ControlRig/Private/Units/Math/RigUnit_MathInt.cpp
--- a/Engine/Plugins/Animation/ControlRig/Source/ControlRig/Private/Units/Math/RigUnit_MathInt.cpp
+++ b/Engine/Plugins/Animation/ControlRig/Source/ControlRig/Private/Units/Math/RigUnit_MathInt.cpp
@@ -2,6 +2,7 @@
 
 #include "Units/Math/RigUnit_MathInt.h"
 #include "Units/RigUnitContext.h"
+#include "Units/Core/RigUnit_CoreDispatch.h"
 
 FRigUnit_MathIntAdd_Execute()
 {
@@ -99,12 +100,22 @@ FRigUnit_MathIntEquals_Execute()
 	Result = A == B;
 }
 
+FRigVMStructUpgradeInfo FRigUnit_MathIntEquals::GetUpgradeInfo() const
+{
+	return FRigVMStructUpgradeInfo::MakeFromStructToFactory(StaticStruct(), FRigDispatch_CoreEquals::StaticStruct());
+}
+
 FRigUnit_MathIntNotEquals_Execute()
 {
     DECLARE_SCOPE_HIERARCHICAL_COUNTER_RIGUNIT()
 	Result = A != B;
 }
 
+FRigVMStructUpgradeInfo FRigUnit_MathIntNotEquals::GetUpgradeInfo() const
+{
+	return FRigVMStructUpgradeInfo::MakeFromStructToFactory(StaticStruct(), FRigDispatch_CoreNotEquals::StaticStruct());
+}
+
 FRigUnit_MathIntGreater_Execute()
 {
     DECLARE_SCOPE_HIERARCHICAL_COUNTER_RIGUNIT()
diff --git a/Engine/Plugins/Animation/ControlRig/Source/ControlRig/Private/Units/Math/RigUnit_MathInt.h b/Engine/Plugins/Animation/ControlRig/Source/ControlRig/Private/Units/Math/RigUnit_MathInt.h
--- a/Engine/Plugins/Animation/ControlRig/Source/ControlRig/Private/Units/Math/RigUnit_MathInt.h
+++ b/Engine/Plugins/Animation/ControlRig/Source/ControlRig/Private/Units/Math/RigUnit_MathInt.h
@@ -259,7 +259,7 @@ struct CONTROLRIG_API FRigUnit_MathIntClamp : public FRigUnit_MathIntBase
 /**
  * Returns true if the value A equals B
  */
-USTRUCT(meta=(DisplayName="Equals", TemplateName="Equals", Keywords="Same,=="))
+USTRUCT(meta=(DisplayName="Equals", TemplateName="Equals", Keywords="Same,==", Deprecated="5.1"))
 struct CONTROLRIG_API FRigUnit_MathIntEquals : public FRigUnit_MathIntBase
 {
 	GENERATED_BODY()
@@ -280,12 +280,15 @@ struct CONTROLRIG_API FRigUnit_MathIntEquals : public FRigUnit_MathIntBase
 
 	UPROPERTY(meta=(Output))
 	bool Result;
+
+	RIGVM_METHOD()
+	virtual FRigVMStructUpgradeInfo GetUpgradeInfo() const override;
 };
 
 /**
  * Returns true if the value A does not equal B
  */
-USTRUCT(meta=(DisplayName="Not Equals", TemplateName="NotEquals", Keywords="Different,!="))
+USTRUCT(meta=(DisplayName="Not Equals", TemplateName="NotEquals", Keywords="Different,!=", Deprecated="5.1"))
 struct CONTROLRIG_API FRigUnit_MathIntNotEquals : public FRigUnit_MathIntBase
 {
 	GENERATED_BODY()
@@ -306,6 +309,9 @@ struct CONTROLRIG_API FRigUnit_MathIntNotEquals : public FRigUnit_MathIntBase
 
 	UPROPERTY(meta=(Output))
 	bool Result;
+	
+	RIGVM_METHOD()
+	virtual FRigVMStructUpgradeInfo GetUpgradeInfo() const override;
 };
 
 /**
diff --git a/Engine/Plugins/Animation/ControlRig/Source/ControlRig/Private/Units/Math/RigUnit_MathQuaternion.cpp b/Engine/Plugins/Animation/ControlRig/Source/ControlRig/Private/Units/Math/RigUnit_MathQuaternion.cpp
--- a/Engine/Plugins/Animation/ControlRig/Source/ControlRig/Private/Units/Math/RigUnit_MathQuaternion.cpp
+++ b/Engine/Plugins/Animation/ControlRig/Source/ControlRig/Private/Units/Math/RigUnit_MathQuaternion.cpp
@@ -5,6 +5,7 @@
 #include "Units/RigUnitContext.h"
 #include "AnimationCoreLibrary.h"
 #include "Math/ControlRigMathLibrary.h"
+#include "Units/Core/RigUnit_CoreDispatch.h"
 
 FRigUnit_MathQuaternionFromAxisAndAngle_Execute()
 {
@@ -150,12 +151,22 @@ FRigUnit_MathQuaternionEquals_Execute()
 	Result = A == B;
 }
 
+FRigVMStructUpgradeInfo FRigUnit_MathQuaternionEquals::GetUpgradeInfo() const
+{
+	return FRigVMStructUpgradeInfo::MakeFromStructToFactory(StaticStruct(), FRigDispatch_CoreEquals::StaticStruct());
+}
+
 FRigUnit_MathQuaternionNotEquals_Execute()
 {
     DECLARE_SCOPE_HIERARCHICAL_COUNTER_RIGUNIT()
 	Result = A != B;
 }
 
+FRigVMStructUpgradeInfo FRigUnit_MathQuaternionNotEquals::GetUpgradeInfo() const
+{
+	return FRigVMStructUpgradeInfo::MakeFromStructToFactory(StaticStruct(), FRigDispatch_CoreNotEquals::StaticStruct());
+}
+
 FRigUnit_MathQuaternionSelectBool_Execute()
 {
     DECLARE_SCOPE_HIERARCHICAL_COUNTER_RIGUNIT()
diff --git a/Engine/Plugins/Animation/ControlRig/Source/ControlRig/Private/Units/Math/RigUnit_MathQuaternion.h b/Engine/Plugins/Animation/ControlRig/Source/ControlRig/Private/Units/Math/RigUnit_MathQuaternion.h
--- a/Engine/Plugins/Animation/ControlRig/Source/ControlRig/Private/Units/Math/RigUnit_MathQuaternion.h
+++ b/Engine/Plugins/Animation/ControlRig/Source/ControlRig/Private/Units/Math/RigUnit_MathQuaternion.h
@@ -372,7 +372,7 @@ struct CONTROLRIG_API FRigUnit_MathQuaternionSlerp : public FRigUnit_MathQuatern
 /**
  * Returns true if the value A equals B
  */
-USTRUCT(meta=(DisplayName="Equals", TemplateName="Equals", Keywords="Same,=="))
+USTRUCT(meta=(DisplayName="Equals", TemplateName="Equals", Keywords="Same,==", Deprecated="5.1"))
 struct CONTROLRIG_API FRigUnit_MathQuaternionEquals : public FRigUnit_MathQuaternionBase
 {
 	GENERATED_BODY()
@@ -394,12 +394,15 @@ struct CONTROLRIG_API FRigUnit_MathQuaternionEquals : public FRigUnit_MathQuater
 
 	UPROPERTY(meta=(Output))
 	bool Result;
+
+	RIGVM_METHOD()
+	virtual FRigVMStructUpgradeInfo GetUpgradeInfo() const override;
 };
 
 /**
  * Returns true if the value A does not equal B
  */
-USTRUCT(meta=(DisplayName="Not Equals", TemplateName="NotEquals", Keywords="Different,!="))
+USTRUCT(meta=(DisplayName="Not Equals", TemplateName="NotEquals", Keywords="Different,!=", Deprecated="5.1"))
 struct CONTROLRIG_API FRigUnit_MathQuaternionNotEquals : public FRigUnit_MathQuaternionBase
 {
 	GENERATED_BODY()
@@ -421,6 +424,9 @@ struct CONTROLRIG_API FRigUnit_MathQuaternionNotEquals : public FRigUnit_MathQua
 
 	UPROPERTY(meta=(Output))
 	bool Result;
+
+	RIGVM_METHOD()
+	virtual FRigVMStructUpgradeInfo GetUpgradeInfo() const override;
 };
 
 /**
diff --git a/Engine/Plugins/Animation/ControlRig/Source/ControlRig/Private/Units/Math/RigUnit_MathVector.cpp b/Engine/Plugins/Animation/ControlRig/Source/ControlRig/Private/Units/Math/RigUnit_MathVector.cpp
--- a/Engine/Plugins/Animation/ControlRig/Source/ControlRig/Private/Units/Math/RigUnit_MathVector.cpp
+++ b/Engine/Plugins/Animation/ControlRig/Source/ControlRig/Private/Units/Math/RigUnit_MathVector.cpp
@@ -4,6 +4,7 @@
 #include "Units/Math/RigUnit_MathTransform.h"
 #include "Units/RigUnitContext.h"
 #include "Math/ControlRigMathLibrary.h"
+#include "Units/Core/RigUnit_CoreDispatch.h"
 
 FRigUnit_MathVectorFromFloat_Execute()
 {
@@ -196,12 +197,22 @@ FRigUnit_MathVectorEquals_Execute()
 	Result = A == B;
 }
 
+FRigVMStructUpgradeInfo FRigUnit_MathVectorEquals::GetUpgradeInfo() const
+{
+	return FRigVMStructUpgradeInfo::MakeFromStructToFactory(StaticStruct(), FRigDispatch_CoreEquals::StaticStruct());
+}
+
 FRigUnit_MathVectorNotEquals_Execute()
 {
     DECLARE_SCOPE_HIERARCHICAL_COUNTER_RIGUNIT()
 	Result = A != B;
 }
 
+FRigVMStructUpgradeInfo FRigUnit_MathVectorNotEquals::GetUpgradeInfo() const
+{
+	return FRigVMStructUpgradeInfo::MakeFromStructToFactory(StaticStruct(), FRigDispatch_CoreNotEquals::StaticStruct());
+}
+
 FRigUnit_MathVectorIsNearlyZero_Execute()
 {
     DECLARE_SCOPE_HIERARCHICAL_COUNTER_RIGUNIT()
diff --git a/Engine/Plugins/Animation/ControlRig/Source/ControlRig/Private/Units/Math/RigUnit_MathVector.h b/Engine/Plugins/Animation/ControlRig/Source/ControlRig/Private/Units/Math/RigUnit_MathVector.h
--- a/Engine/Plugins/Animation/ControlRig/Source/ControlRig/Private/Units/Math/RigUnit_MathVector.h
+++ b/Engine/Plugins/Animation/ControlRig/Source/ControlRig/Private/Units/Math/RigUnit_MathVector.h
@@ -376,7 +376,7 @@ struct CONTROLRIG_API FRigUnit_MathVectorRemap : public FRigUnit_MathVectorBase
 /**
  * Returns true if the value A equals B
  */
-USTRUCT(meta=(DisplayName="Equals", TemplateName="Equals", Keywords="Same,=="))
+USTRUCT(meta=(DisplayName="Equals", TemplateName="Equals", Keywords="Same,==", Deprecated="5.1"))
 struct CONTROLRIG_API FRigUnit_MathVectorEquals : public FRigUnit_MathVectorBase
 {
 	GENERATED_BODY()
@@ -398,12 +398,15 @@ struct CONTROLRIG_API FRigUnit_MathVectorEquals : public FRigUnit_MathVectorBase
 
 	UPROPERTY(meta=(Output))
 	bool Result;
+
+	RIGVM_METHOD()
+	virtual FRigVMStructUpgradeInfo GetUpgradeInfo() const override;
 };
 
 /**
  * Returns true if the value A does not equal B
  */
-USTRUCT(meta=(DisplayName="Not Equals", TemplateName="NotEquals", Keywords="Different,!="))
+USTRUCT(meta=(DisplayName="Not Equals", TemplateName="NotEquals", Keywords="Different,!=", Deprecated="5.1"))
 struct CONTROLRIG_API FRigUnit_MathVectorNotEquals : public FRigUnit_MathVectorBase
 {
 	GENERATED_BODY()
@@ -425,6 +428,9 @@ struct CONTROLRIG_API FRigUnit_MathVectorNotEquals : public FRigUnit_MathVectorB
 
 	UPROPERTY(meta=(Output))
 	bool Result;
+
+	RIGVM_METHOD()
+	virtual FRigVMStructUpgradeInfo GetUpgradeInfo() const override;
 };
 
 /**
