Control Rig: Fix various crashes with metadata

diff --git a/Engine/Plugins/Animation/ControlRig/Source/ControlRig/Private/Rigs/RigHierarchyElements.cpp b/Engine/Plugins/Animation/ControlRig/Source/ControlRig/Private/Rigs/RigHierarchyElements.cpp
--- a/Engine/Plugins/Animation/ControlRig/Source/ControlRig/Private/Rigs/RigHierarchyElements.cpp
+++ b/Engine/Plugins/Animation/ControlRig/Source/ControlRig/Private/Rigs/RigHierarchyElements.cpp
@@ -13,6 +13,34 @@
 // FRigBaseElement
 ////////////////////////////////////////////////////////////////////////////////
 
+FRigBaseElement::FRigBaseElement(const FRigBaseElement& InOther)
+{
+	*this = InOther;
+}
+
+FRigBaseElement& FRigBaseElement::operator=(const FRigBaseElement& InOther)
+{
+	Key = InOther.Key;
+	NameString = InOther.NameString;
+	Index = InOther.Index;
+	SubIndex = InOther.SubIndex;
+	bSelected = InOther.bSelected;
+	CreatedAtInstructionIndex = InOther.CreatedAtInstructionIndex;
+	TopologyVersion = InOther.TopologyVersion;
+	CachedChildren.Reset();
+	OwnedInstances = 1;
+
+	RemoveAllMetadata();
+	for(const FRigBaseMetadata* InOtherMd : InOther.Metadata)
+	{
+		FRigBaseMetadata* Md = SetupValidMetadata(InOtherMd->Name, InOtherMd->Type);
+		check(Md);
+		Md->SetValueData(InOtherMd->GetValueData(), InOtherMd->GetValueSize());
+	}
+
+	return *this;
+}
+
 FRigBaseElement::~FRigBaseElement()
 {
 	RemoveAllMetadata();
@@ -143,14 +171,7 @@ bool FRigBaseElement::RemoveMetadata(const FName& InName)
 	if(const int32* MetadataIndexPtr = MetadataNameToIndex.Find(InName))
 	{
 		const int32 MetadataIndex = *MetadataIndexPtr;
-		{
-			FRigBaseMetadata* Md = Metadata[MetadataIndex];
-			if(const UScriptStruct* Struct = Md->GetMetadataStruct())
-			{
-				Struct->DestroyStruct(Md, 1);
-			}
-			FMemory::Free(Md);
-		}
+		FRigBaseMetadata::DestroyMetadata(&Metadata[MetadataIndex]);
 		MetadataNameToIndex.Remove(InName);
 		Metadata.RemoveAt(MetadataIndex);
 		for(TPair<FName, int32>& Pair : MetadataNameToIndex)
@@ -171,11 +192,7 @@ bool FRigBaseElement::RemoveAllMetadata()
 	{
 		for(FRigBaseMetadata* Md : Metadata)
 		{
-			if(const UScriptStruct* Struct = Md->GetMetadataStruct())
-			{
-				Struct->DestroyStruct(Md, 1);
-			}
-			FMemory::Free(Md);
+			FRigBaseMetadata::DestroyMetadata(&Md);
 		}
 		Metadata.Reset();
 		MetadataNameToIndex.Reset();
@@ -218,23 +235,17 @@ void FRigBaseElement::CopyFrom(URigHierarchy* InHierarchy, FRigBaseElement* InOt
 
 FRigBaseMetadata* FRigBaseElement::SetupValidMetadata(const FName& InName, ERigMetadataType InType)
 {
-	if(const int32* MetadataIndex = MetadataNameToIndex.Find(InName))
+	if(const int32* MetadataIndexPtr = MetadataNameToIndex.Find(InName))
 	{
-		FRigBaseMetadata* Md = Metadata[*MetadataIndex];
-		if(Md->GetType() == InType)
-		{
-			return Md;
-		}
-
-		if(const UScriptStruct* Struct = Md->GetMetadataStruct())
+		const int32 MetadataIndex = *MetadataIndexPtr;
+		if(Metadata[MetadataIndex]->GetType() == InType)
 		{
-			Struct->DestroyStruct(Md, 1);
+			return Metadata[MetadataIndex];
 		}
-		FMemory::Free(Md);
 
-		Md = FRigBaseMetadata::MakeMetadata(this, InName, InType);
-		Metadata[*MetadataIndex] = Md;
-		return Md;
+		FRigBaseMetadata::DestroyMetadata(&Metadata[MetadataIndex]);
+		Metadata[MetadataIndex] = FRigBaseMetadata::MakeMetadata(this, InName, InType);
+		return Metadata[MetadataIndex];
 	}
 
 	FRigBaseMetadata* Md = FRigBaseMetadata::MakeMetadata(this, InName, InType);
diff --git a/Engine/Plugins/Animation/ControlRig/Source/ControlRig/Private/Rigs/RigHierarchyMetadata.cpp b/Engine/Plugins/Animation/ControlRig/Source/ControlRig/Private/Rigs/RigHierarchyMetadata.cpp
--- a/Engine/Plugins/Animation/ControlRig/Source/ControlRig/Private/Rigs/RigHierarchyMetadata.cpp
+++ b/Engine/Plugins/Animation/ControlRig/Source/ControlRig/Private/Rigs/RigHierarchyMetadata.cpp
@@ -122,6 +122,19 @@ FRigBaseMetadata* FRigBaseMetadata::MakeMetadata(const FRigBaseElement* InElemen
 	return Md;
 }
 
+void FRigBaseMetadata::DestroyMetadata(FRigBaseMetadata** Metadata)
+{
+	check(Metadata);
+	FRigBaseMetadata* Md = *Metadata;
+	check(Md);
+	if(const UScriptStruct* Struct = Md->GetMetadataStruct())
+	{
+		Struct->DestroyStruct(Md, 1);
+	}
+	FMemory::Free(Md);
+	Md = nullptr;
+}
+
 void FRigBaseMetadata::Serialize(FArchive& Ar, bool bIsLoading)
 {
 	Ar.UsingCustomVersion(FControlRigObjectVersion::GUID);
diff --git a/Engine/Plugins/Animation/ControlRig/Source/ControlRig/Public/Rigs/RigHierarchyElements.h b/Engine/Plugins/Animation/ControlRig/Source/ControlRig/Public/Rigs/RigHierarchyElements.h
--- a/Engine/Plugins/Animation/ControlRig/Source/ControlRig/Public/Rigs/RigHierarchyElements.h
+++ b/Engine/Plugins/Animation/ControlRig/Source/ControlRig/Public/Rigs/RigHierarchyElements.h
@@ -516,6 +516,9 @@ public:
 	, OwnedInstances(0)
 	{}
 
+	FRigBaseElement(const FRigBaseElement& InOther);
+	FRigBaseElement& operator= (const FRigBaseElement& InOther);
+
 	virtual ~FRigBaseElement();
 
 	enum ESerializationPhase
diff --git a/Engine/Plugins/Animation/ControlRig/Source/ControlRig/Public/Rigs/RigHierarchyMetadata.h b/Engine/Plugins/Animation/ControlRig/Source/ControlRig/Public/Rigs/RigHierarchyMetadata.h
--- a/Engine/Plugins/Animation/ControlRig/Source/ControlRig/Public/Rigs/RigHierarchyMetadata.h
+++ b/Engine/Plugins/Animation/ControlRig/Source/ControlRig/Public/Rigs/RigHierarchyMetadata.h
@@ -147,6 +147,7 @@ protected:
 
 	static UScriptStruct* GetMetadataStruct(const ERigMetadataType& InType);
 	static FRigBaseMetadata* MakeMetadata(const FRigBaseElement* InElement, const FName& InName, ERigMetadataType InType);
+	static void DestroyMetadata(FRigBaseMetadata** Metadata);
 
 	FORCEINLINE const FProperty* GetValueProperty() const
 	{
