Control Rig: Raise procedural limit based on the elements in the default hierarchy

diff --git a/Engine/Plugins/Animation/ControlRig/Source/ControlRig/Private/ControlRig.cpp b/Engine/Plugins/Animation/ControlRig/Source/ControlRig/Private/ControlRig.cpp
--- a/Engine/Plugins/Animation/ControlRig/Source/ControlRig/Private/ControlRig.cpp
+++ b/Engine/Plugins/Animation/ControlRig/Source/ControlRig/Private/ControlRig.cpp
@@ -273,6 +273,12 @@ void UControlRig::InitializeFromCDO()
 
 		// copy hierarchy settings
 		HierarchySettings = CDO->HierarchySettings;
+		
+		// increment the procedural limit based on the number of elements in the CDO
+		if(const URigHierarchy* CDOHierarchy = CDO->GetHierarchy())
+		{
+			HierarchySettings.ProceduralElementLimit += CDOHierarchy->Num();
+		}
 
 		// copy vm settings
 		VMRuntimeSettings = CDO->VMRuntimeSettings;
