Control Rig: Introduction of per element tags

diff --git a/Engine/Plugins/Animation/ControlRig/Source/ControlRig/Private/Rigs/RigHierarchy.cpp b/Engine/Plugins/Animation/ControlRig/Source/ControlRig/Private/Rigs/RigHierarchy.cpp
--- a/Engine/Plugins/Animation/ControlRig/Source/ControlRig/Private/Rigs/RigHierarchy.cpp
+++ b/Engine/Plugins/Animation/ControlRig/Source/ControlRig/Private/Rigs/RigHierarchy.cpp
@@ -94,6 +94,7 @@ bool URigHierarchy::bEnableValidityCheckbyDefault = false;
 
 URigHierarchy::URigHierarchy()
 : TopologyVersion(0)
+, MetadataVersion(0)
 , bEnableDirtyPropagation(true)
 , Elements()
 , IndexLookup()
@@ -292,6 +293,7 @@ void URigHierarchy::Reset()
 	LLM_SCOPE_BYNAME(TEXT("Animation/ControlRig"));
 	
 	TopologyVersion = 0;
+	MetadataVersion = 0;
 	bEnableDirtyPropagation = true;
 
 	// walk in reverse since certain elements might not have been allocated themselves
@@ -411,6 +413,7 @@ void URigHierarchy::CopyHierarchy(URigHierarchy* InHierarchy)
 
 
 	TopologyVersion = InHierarchy->GetTopologyVersion();
+	MetadataVersion = InHierarchy->GetMetadataVersion();
 	UpdateAllCachedChildren();
 	
 	EnsureCacheValidity();
@@ -4097,6 +4100,10 @@ FRigBaseElement* URigHierarchy::MakeElement(ERigElementType InElementType, int32
 	{
 		Element->OwnedInstances = InCount;
 	}
+
+	Element->MetadataChangedDelegate.BindUObject(this, &URigHierarchy::OnMetadataChanged);
+	Element->MetadataTagChangedDelegate.BindUObject(this, &URigHierarchy::OnMetadataTagChanged);
+	
 	return Element;
 }
 
@@ -4105,6 +4112,9 @@ void URigHierarchy::DestroyElement(FRigBaseElement*& InElement)
 	LLM_SCOPE_BYNAME(TEXT("Animation/ControlRig"));
 	check(InElement != nullptr);
 
+	InElement->MetadataChangedDelegate.Unbind();
+	InElement->MetadataTagChangedDelegate.Unbind();
+
 	if(InElement->OwnedInstances == 0)
 	{
 		return;
@@ -4115,55 +4125,55 @@ void URigHierarchy::DestroyElement(FRigBaseElement*& InElement)
 	{
 		case ERigElementType::Bone:
 		{
-			FRigBoneElement* Elements = Cast<FRigBoneElement>(InElement);
+			FRigBoneElement* ExistingElements = Cast<FRigBoneElement>(InElement);
 			for(int32 Index=0;Index<Count;Index++)
 			{
-				Elements[Index].~FRigBoneElement(); 
+				ExistingElements[Index].~FRigBoneElement(); 
 			}
 			break;
 		}
 		case ERigElementType::Null:
 		{
-			FRigNullElement* Elements = Cast<FRigNullElement>(InElement);
+			FRigNullElement* ExistingElements = Cast<FRigNullElement>(InElement);
 			for(int32 Index=0;Index<Count;Index++)
 			{
-				Elements[Index].~FRigNullElement(); 
+				ExistingElements[Index].~FRigNullElement(); 
 			}
 			break;
 		}
 		case ERigElementType::Control:
 		{
-			FRigControlElement* Elements = Cast<FRigControlElement>(InElement);
+			FRigControlElement* ExistingElements = Cast<FRigControlElement>(InElement);
 			for(int32 Index=0;Index<Count;Index++)
 			{
-				Elements[Index].~FRigControlElement(); 
+				ExistingElements[Index].~FRigControlElement(); 
 			}
 			break;
 		}
 		case ERigElementType::Curve:
 		{
-			FRigCurveElement* Elements = Cast<FRigCurveElement>(InElement);
+			FRigCurveElement* ExistingElements = Cast<FRigCurveElement>(InElement);
 			for(int32 Index=0;Index<Count;Index++)
 			{
-				Elements[Index].~FRigCurveElement(); 
+				ExistingElements[Index].~FRigCurveElement(); 
 			}
 			break;
 		}
 		case ERigElementType::RigidBody:
 		{
-			FRigRigidBodyElement* Elements = Cast<FRigRigidBodyElement>(InElement);
+			FRigRigidBodyElement* ExistingElements = Cast<FRigRigidBodyElement>(InElement);
 			for(int32 Index=0;Index<Count;Index++)
 			{
-				Elements[Index].~FRigRigidBodyElement(); 
+				ExistingElements[Index].~FRigRigidBodyElement(); 
 			}
 			break;
 		}
 		case ERigElementType::Reference:
 		{
-			FRigReferenceElement* Elements = Cast<FRigReferenceElement>(InElement);
+			FRigReferenceElement* ExistingElements = Cast<FRigReferenceElement>(InElement);
 			for(int32 Index=0;Index<Count;Index++)
 			{
-				Elements[Index].~FRigReferenceElement(); 
+				ExistingElements[Index].~FRigReferenceElement(); 
 			}
 			break;
 		}
@@ -4453,6 +4463,23 @@ void URigHierarchy::CleanupInvalidCaches()
 	EnsureCacheValidity();
 }
 
+void URigHierarchy::OnMetadataChanged(const FRigElementKey& InKey, const FName& InName)
+{
+	IncrementMetadataVersion();
+	if(MetadataChangedDelegate.IsBound())
+	{
+		MetadataChangedDelegate.Broadcast(InKey, InName);
+	}
+}
+
+void URigHierarchy::OnMetadataTagChanged(const FRigElementKey& InKey, const FName& InTag, bool bAdded)
+{
+	if(MetadataTagChangedDelegate.IsBound())
+	{
+		MetadataTagChangedDelegate.Broadcast(InKey, InTag, bAdded);
+	}
+}
+
 void URigHierarchy::EnsureCacheValidityImpl()
 {
 	LLM_SCOPE_BYNAME(TEXT("Animation/ControlRig"));
diff --git a/Engine/Plugins/Animation/ControlRig/Source/ControlRig/Private/Rigs/RigHierarchyCache.cpp b/Engine/Plugins/Animation/ControlRig/Source/ControlRig/Private/Rigs/RigHierarchyCache.cpp
--- a/Engine/Plugins/Animation/ControlRig/Source/ControlRig/Private/Rigs/RigHierarchyCache.cpp
+++ b/Engine/Plugins/Animation/ControlRig/Source/ControlRig/Private/Rigs/RigHierarchyCache.cpp
@@ -24,7 +24,24 @@ bool FCachedRigElement::UpdateCache(const FRigElementKey& InKey, const URigHiera
 		{
 			// have to create a copy since Reset below
 			// potentially resets the InKey as well.
-			const FRigElementKey KeyToResolve = InKey; 
+			const FRigElementKey KeyToResolve = InKey;
+
+			// first try to resolve with the known index.
+			// this happens a lot - where the topology version has
+			// increased - but the known item is still valid
+			if(InHierarchy->IsValidIndex(Index))
+			{
+				if(const FRigBaseElement* PreviousElement = InHierarchy->Get(Index))
+				{
+					if(PreviousElement->GetKey() == KeyToResolve)
+					{
+						Key = KeyToResolve;
+						Element = PreviousElement;
+						ContainerVersion = InHierarchy->GetTopologyVersion();
+						return IsValid();
+					}
+				}
+			}
 
 			Reset();
 
diff --git a/Engine/Plugins/Animation/ControlRig/Source/ControlRig/Private/Rigs/RigHierarchyController.cpp b/Engine/Plugins/Animation/ControlRig/Source/ControlRig/Private/Rigs/RigHierarchyController.cpp
--- a/Engine/Plugins/Animation/ControlRig/Source/ControlRig/Private/Rigs/RigHierarchyController.cpp
+++ b/Engine/Plugins/Animation/ControlRig/Source/ControlRig/Private/Rigs/RigHierarchyController.cpp
@@ -1185,7 +1185,7 @@ TArray<FRigElementKey> URigHierarchyController::ImportFromText(FString InContent
 
 				PastedKeys.Add(ExistingElement->GetKey());
 
-				URigHierarchy::DestroyElement(NewElement);
+				Hierarchy->DestroyElement(NewElement);
 				continue;
 			}
 		}
@@ -1811,7 +1811,7 @@ bool URigHierarchyController::RemoveElement(FRigBaseElement* InElement)
 		Notify(ERigHierarchyNotification::HierarchyReset, nullptr);
 	}
 
-	URigHierarchy::DestroyElement(InElement);
+	Hierarchy->DestroyElement(InElement);
 
 	Hierarchy->EnsureCacheValidity();
 
diff --git a/Engine/Plugins/Animation/ControlRig/Source/ControlRig/Private/Rigs/RigHierarchyElements.cpp b/Engine/Plugins/Animation/ControlRig/Source/ControlRig/Private/Rigs/RigHierarchyElements.cpp
--- a/Engine/Plugins/Animation/ControlRig/Source/ControlRig/Private/Rigs/RigHierarchyElements.cpp
+++ b/Engine/Plugins/Animation/ControlRig/Source/ControlRig/Private/Rigs/RigHierarchyElements.cpp
@@ -162,6 +162,11 @@ void FRigBaseElement::Load(FArchive& Ar, URigHierarchy* Hierarchy, ESerializatio
 			{
 				MetadataNameToIndex.Add(Metadata[MetadataIndex]->GetName(), MetadataIndex);
 			}
+
+			for(int32 MetadataIndex = 0; MetadataIndex < Metadata.Num(); MetadataIndex++)
+			{
+				NotifyMetadataChanged(Metadata[MetadataIndex]->GetName());
+			}
 		}
 	}
 }
@@ -181,6 +186,7 @@ bool FRigBaseElement::RemoveMetadata(const FName& InName)
 				Pair.Value--;
 			}
 		}
+		NotifyMetadataChanged(InName);
 		return true;
 	}
 	return false;
@@ -190,12 +196,19 @@ bool FRigBaseElement::RemoveAllMetadata()
 {
 	if(!Metadata.IsEmpty())
 	{
+		TArray<FName> Names;
+		Names.Reserve(Metadata.Num());
 		for(FRigBaseMetadata* Md : Metadata)
 		{
+			Names.Add(Md->GetName());
 			FRigBaseMetadata::DestroyMetadata(&Md);
 		}
 		Metadata.Reset();
 		MetadataNameToIndex.Reset();
+		for(const FName& Name: Names)
+		{
+			NotifyMetadataChanged(Name);
+		}
 		return true;
 	}
 	return false;
@@ -245,15 +258,34 @@ FRigBaseMetadata* FRigBaseElement::SetupValidMetadata(const FName& InName, ERigM
 
 		FRigBaseMetadata::DestroyMetadata(&Metadata[MetadataIndex]);
 		Metadata[MetadataIndex] = FRigBaseMetadata::MakeMetadata(this, InName, InType);
+		NotifyMetadataChanged(InName);
 		return Metadata[MetadataIndex];
 	}
 
 	FRigBaseMetadata* Md = FRigBaseMetadata::MakeMetadata(this, InName, InType);
 	const int32 MetadataIndex = Metadata.Add(Md);
 	MetadataNameToIndex.Add(InName, MetadataIndex);
+	NotifyMetadataChanged(InName);
 	return Md;
 }
 
+void FRigBaseElement::NotifyMetadataChanged(const FName& InName)
+{
+	MetadataVersion++;
+	if(MetadataChangedDelegate.IsBound())
+	{
+		MetadataChangedDelegate.Execute(GetKey(), InName);
+	}
+}
+
+void FRigBaseElement::NotifyMetadataTagChanged(const FName& InTag, bool bAdded)
+{
+	if(MetadataTagChangedDelegate.IsBound())
+	{
+		MetadataTagChangedDelegate.Execute(GetKey(), InTag, bAdded);
+	}
+}
+
 ////////////////////////////////////////////////////////////////////////////////
 // FRigComputedTransform
 ////////////////////////////////////////////////////////////////////////////////
diff --git a/Engine/Plugins/Animation/ControlRig/Source/ControlRig/Private/Units/Hierarchy/RigUnit_Metadata.cpp b/Engine/Plugins/Animation/ControlRig/Source/ControlRig/Private/Units/Hierarchy/RigUnit_Metadata.cpp
--- a/Engine/Plugins/Animation/ControlRig/Source/ControlRig/Private/Units/Hierarchy/RigUnit_Metadata.cpp
+++ b/Engine/Plugins/Animation/ControlRig/Source/ControlRig/Private/Units/Hierarchy/RigUnit_Metadata.cpp
@@ -5,6 +5,10 @@
 #include "RigVMCore/RigVMStruct.h"
 #include "Units/RigUnitContext.h"
 
+#if WITH_EDITOR
+#include "RigVMModel/RigVMPin.h"
+#endif
+
 FName FRigDispatch_MetadataBase::ItemArgName = TEXT("Item");
 FName FRigDispatch_MetadataBase::NameArgName = TEXT("Name");
 FName FRigDispatch_MetadataBase::CacheArgName = TEXT("Cache");
@@ -153,8 +157,24 @@ TArray<FRigVMTemplateArgument> FRigDispatch_GetMetadata::GetArguments() const
 	return Arguments;
 }
 
+#if WITH_EDITOR
+
+FString FRigDispatch_GetMetadata::GetArgumentMetaData(const FName& InArgumentName, const FName& InMetaDataKey) const
+{
+	if(InArgumentName == TEXT("Name"))
+	{
+		if(InMetaDataKey == FRigVMStruct::CustomWidgetMetaName)
+		{
+			return TEXT("MetadataName");
+		}
+	}
+	return FRigDispatch_MetadataBase::GetArgumentMetaData(InArgumentName, InMetaDataKey);
+}
+
+#endif
+
 FRigBaseMetadata* FRigDispatch_GetMetadata::FindMetadata(const FRigVMExtendedExecuteContext& InContext,
-                                                              const FRigElementKey& InKey, const FName& InName, ERigMetadataType InType, FCachedRigElement& Cache)
+                                                         const FRigElementKey& InKey, const FName& InName, ERigMetadataType InType, FCachedRigElement& Cache)
 {
 	const FRigUnitContext& Context = GetRigUnitContext(InContext);
 	if(Cache.UpdateCache(InKey, Context.Hierarchy))
@@ -371,17 +391,47 @@ FRigVMFunctionPtr FRigDispatch_SetMetadata::GetDispatchFunctionImpl(const FRigVM
 	return nullptr;
 }
 
+FRigUnit_RemoveMetadata_Execute()
+{
+	URigHierarchy* Hierarchy = ExecuteContext.Hierarchy;
+
+	Removed = false;
+
+	if(CachedIndex.UpdateCache(Item, Hierarchy))
+	{
+		if(FRigBaseElement* Element = Hierarchy->Get(CachedIndex))
+		{
+			Removed = Element->RemoveMetadata(Name);
+		}
+	}
+}
+
+FRigUnit_RemoveAllMetadata_Execute()
+{
+	URigHierarchy* Hierarchy = ExecuteContext.Hierarchy;
+
+	Removed = false;
+
+	if(CachedIndex.UpdateCache(Item, Hierarchy))
+	{
+		if(FRigBaseElement* Element = Hierarchy->Get(CachedIndex))
+		{
+			Removed = Element->RemoveAllMetadata();
+		}
+	}
+}
+
 FRigUnit_HasMetadata_Execute()
 {
 	const URigHierarchy* Hierarchy = Context.Hierarchy;
 
-	bFound = false;
+	Found = false;
 
 	if(CachedIndex.UpdateCache(Item, Hierarchy))
 	{
 		if(const FRigBaseElement* Element = Hierarchy->Get(CachedIndex))
 		{
-			bFound = Element->GetMetadata(Name, Type) != nullptr;
+			Found = Element->GetMetadata(Name, Type) != nullptr;
 		}
 	}
 }
@@ -400,4 +450,229 @@ FRigUnit_FindItemsWithMetadata_Execute()
 		}
 		bContinue = true;
 	});
+}
+
+FRigUnit_GetMetadataTags_Execute()
+{
+	Tags.Reset();
+	
+	const URigHierarchy* Hierarchy = Context.Hierarchy;
+	if(CachedIndex.UpdateCache(Item, Hierarchy))
+	{
+		if(const FRigBaseElement* Element = Hierarchy->Get(CachedIndex))
+		{
+			if(const FRigNameArrayMetadata* Md = Cast<FRigNameArrayMetadata>(Element->GetMetadata(URigHierarchy::TagMetadataName, ERigMetadataType::NameArray)))
+			{
+				Tags = Md->GetValue();
+			}
+		}
+	}
+}
+
+FRigUnit_SetMetadataTag_Execute()
+{
+	if(CachedIndex.UpdateCache(Item, ExecuteContext.Hierarchy))
+	{
+		if(FRigBaseElement* Element = ExecuteContext.Hierarchy->Get(CachedIndex))
+		{
+			if(FRigNameArrayMetadata* Md = Cast<FRigNameArrayMetadata>(Element->SetupValidMetadata(URigHierarchy::TagMetadataName, ERigMetadataType::NameArray)))
+			{
+				const int32 LastIndex = Md->GetValue().Num(); 
+				if(Md->GetValue().AddUnique(Tag) == LastIndex)
+				{
+					Element->NotifyMetadataTagChanged(Tag, true);
+				}
+			}
+		}
+	}
+}
+
+FRigUnit_SetMetadataTagArray_Execute()
+{
+	if(CachedIndex.UpdateCache(Item, ExecuteContext.Hierarchy))
+	{
+		if(FRigBaseElement* Element = ExecuteContext.Hierarchy->Get(CachedIndex))
+		{
+			if(FRigNameArrayMetadata* Md = Cast<FRigNameArrayMetadata>(Element->SetupValidMetadata(URigHierarchy::TagMetadataName, ERigMetadataType::NameArray)))
+			{
+				for(const FName& Tag : Tags)
+				{
+					const int32 LastIndex = Md->GetValue().Num(); 
+					if(Md->GetValue().AddUnique(Tag) == LastIndex)
+					{
+						Element->NotifyMetadataTagChanged(Tag, true);
+					}
+				}
+			}
+		}
+	}
+}
+
+FRigUnit_RemoveMetadataTag_Execute()
+{
+	Removed = false;
+	if(CachedIndex.UpdateCache(Item, ExecuteContext.Hierarchy))
+	{
+		if(FRigBaseElement* Element = ExecuteContext.Hierarchy->Get(CachedIndex))
+		{
+			if(FRigNameArrayMetadata* Md = Cast<FRigNameArrayMetadata>(Element->GetMetadata(URigHierarchy::TagMetadataName, ERigMetadataType::NameArray)))
+			{
+				Removed = Md->GetValue().Remove(Tag) > 0;
+				if(Removed)
+				{
+					Element->NotifyMetadataTagChanged(Tag, false);
+				}
+			}
+		}
+	}
+}
+
+FRigUnit_HasMetadataTag_Execute()
+{
+	Found = false;
+	
+	const URigHierarchy* Hierarchy = Context.Hierarchy;
+	if(CachedIndex.UpdateCache(Item, Hierarchy))
+	{
+		if(const FRigBaseElement* Element = Hierarchy->Get(CachedIndex))
+		{
+			if(const FRigNameArrayMetadata* Md = Cast<FRigNameArrayMetadata>(Element->GetMetadata(URigHierarchy::TagMetadataName, ERigMetadataType::NameArray)))
+			{
+				Found = Md->GetValue().Contains(Tag);
+			}
+		}
+	}
+}
+
+FRigUnit_HasMetadataTagArray_Execute()
+{
+	Found = false;
+	
+	const URigHierarchy* Hierarchy = Context.Hierarchy;
+	if(CachedIndex.UpdateCache(Item, Hierarchy))
+	{
+		if(const FRigBaseElement* Element = Hierarchy->Get(CachedIndex))
+		{
+			if(const FRigNameArrayMetadata* Md = Cast<FRigNameArrayMetadata>(Element->GetMetadata(URigHierarchy::TagMetadataName, ERigMetadataType::NameArray)))
+			{
+				Found = true;
+				for(const FName& Tag : Tags)
+				{
+					if(!Md->GetValue().Contains(Tag))
+					{
+						Found = false;
+						break;
+					}
+				}
+			}
+		}
+	}
+}
+
+FRigUnit_FindItemsWithMetadataTag_Execute()
+{
+	const URigHierarchy* Hierarchy = Context.Hierarchy;
+
+	Items.Reset();
+
+	Hierarchy->Traverse([&Items, Tag](const FRigBaseElement* Element, bool& bContinue)
+	{
+		if(const FRigNameArrayMetadata* Md = Cast<FRigNameArrayMetadata>(Element->GetMetadata(URigHierarchy::TagMetadataName, ERigMetadataType::NameArray)))
+		{
+			if(Md->GetValue().Contains(Tag))
+			{
+				Items.Add(Element->GetKey());
+			}
+		}
+		bContinue = true;
+	});
+}
+
+FRigUnit_FindItemsWithMetadataTagArray_Execute()
+{
+	const URigHierarchy* Hierarchy = Context.Hierarchy;
+
+	Items.Reset();
+
+	Hierarchy->Traverse([&Items, Tags](const FRigBaseElement* Element, bool& bContinue)
+	{
+		if(const FRigNameArrayMetadata* Md = Cast<FRigNameArrayMetadata>(Element->GetMetadata(URigHierarchy::TagMetadataName, ERigMetadataType::NameArray)))
+		{
+			bool bFoundAll = true;
+			for(const FName& Tag : Tags)
+			{
+				if(!Md->GetValue().Contains(Tag))
+				{
+					bFoundAll = false;
+					break;
+				}
+			}
+
+			if(bFoundAll)
+			{
+				Items.Add(Element->GetKey());
+			}
+		}
+		bContinue = true;
+	});
+}
+
+FRigUnit_FilterItemsByMetadataTags_Execute()
+{
+	const URigHierarchy* Hierarchy = Context.Hierarchy;
+
+	Result.Reset();
+
+	if(CachedIndices.Num() != Items.Num())
+	{
+		CachedIndices.Reset();
+		CachedIndices.SetNumZeroed(Items.Num());
+	}
+
+	for(int32 Index = 0; Index < Items.Num(); Index++)
+	{
+		if(CachedIndices[Index].UpdateCache(Items[Index], Hierarchy))
+		{
+			const FRigBaseElement* Element = CachedIndices[Index].GetElement();
+			if(const FRigNameArrayMetadata* Md = Cast<FRigNameArrayMetadata>(Element->GetMetadata(URigHierarchy::TagMetadataName, ERigMetadataType::NameArray)))
+			{
+				if(Inclusive)
+				{
+					bool bFoundAll = true;
+					for(const FName& Tag : Tags)
+					{
+						if(!Md->GetValue().Contains(Tag))
+						{
+							bFoundAll = false;
+							break;
+						}
+					}
+					if(bFoundAll)
+					{
+						Result.Add(Element->GetKey());
+					}
+				}
+				else
+				{
+					bool bFoundAny = false;
+					for(const FName& Tag : Tags)
+					{
+						if(Md->GetValue().Contains(Tag))
+						{
+							bFoundAny = true;
+							break;
+						}
+					}
+					if(bFoundAny)
+					{
+						Result.Add(Element->GetKey());
+					}
+				}
+			}
+		}
+		else
+		{
+			UE_CONTROLRIG_RIGUNIT_REPORT_ERROR(TEXT("Item '%s' not found"), *Items[Index].ToString());
+		}
+	}
 }
\ No newline at end of file
diff --git a/Engine/Plugins/Animation/ControlRig/Source/ControlRig/Private/Units/Hierarchy/RigUnit_Metadata.h b/Engine/Plugins/Animation/ControlRig/Source/ControlRig/Private/Units/Hierarchy/RigUnit_Metadata.h
--- a/Engine/Plugins/Animation/ControlRig/Source/ControlRig/Private/Units/Hierarchy/RigUnit_Metadata.h
+++ b/Engine/Plugins/Animation/ControlRig/Source/ControlRig/Private/Units/Hierarchy/RigUnit_Metadata.h
@@ -53,6 +53,10 @@ struct CONTROLRIG_API FRigDispatch_GetMetadata : public FRigDispatch_MetadataBas
 
 	virtual TArray<FRigVMTemplateArgument> GetArguments() const override;
 
+#if WITH_EDITOR
+	virtual FString GetArgumentMetaData(const FName& InArgumentName, const FName& InMetaDataKey) const override;
+#endif
+
 protected:
 
 	static FRigBaseMetadata* FindMetadata(const FRigVMExtendedExecuteContext& InContext, const FRigElementKey& InKey, const FName& InName, ERigMetadataType InType, FCachedRigElement& Cache);
@@ -167,10 +171,81 @@ protected:
 	}
 };
 
+/**
+ * Removes an existing metadata filed from an item
+ */
+USTRUCT(meta=(DisplayName="Remove Metadata", Category="Hierarchy", DocumentationPolicy = "Strict", Keywords="DeleteMetadata", NodeColor="0.462745, 1,0, 0.329412", Varying))
+struct CONTROLRIG_API FRigUnit_RemoveMetadata : public FRigUnitMutable
+{
+	GENERATED_BODY()
+
+	FRigUnit_RemoveMetadata()
+		: Item(NAME_None, ERigElementType::Bone)
+		, Name(NAME_None)
+		, Removed(false)
+		, CachedIndex()
+	{}
+
+	RIGVM_METHOD()
+	virtual void Execute(const FRigUnitContext& Context) override;
+
+	/**
+	 * The item to remove the metadata from 
+	 */
+	UPROPERTY(meta = (Input, ExpandByDefault))
+	FRigElementKey Item;
+
+	/**
+	 * The name of the metadata to remove
+	 */ 
+	UPROPERTY(meta = (Input, CustomWidget="MetadataName"))
+	FName Name;
+
+	// True if the metadata has been removed
+	UPROPERTY(meta=(Output))
+	bool Removed;
+
+	// Used to cache the internally
+	UPROPERTY()
+	FCachedRigElement CachedIndex;
+};
+
+/**
+ * Removes an existing metadata filed from an item
+ */
+USTRUCT(meta=(DisplayName="Remove All Metadata", Category="Hierarchy", DocumentationPolicy = "Strict", Keywords="DeleteMetadata", NodeColor="0.462745, 1,0, 0.329412", Varying))
+struct CONTROLRIG_API FRigUnit_RemoveAllMetadata : public FRigUnitMutable
+{
+	GENERATED_BODY()
+
+	FRigUnit_RemoveAllMetadata()
+		: Item(NAME_None, ERigElementType::Bone)
+		, Removed(false)
+		, CachedIndex()
+	{}
+
+	RIGVM_METHOD()
+	virtual void Execute(const FRigUnitContext& Context) override;
+
+	/**
+	 * The item to remove the metadata from 
+	 */
+	UPROPERTY(meta = (Input, ExpandByDefault))
+	FRigElementKey Item;
+
+	// True if any metadata has been removed
+	UPROPERTY(meta=(Output))
+	bool Removed;
+
+	// Used to cache the internally
+	UPROPERTY()
+	FCachedRigElement CachedIndex;
+};
+
 /**
  * Returns true if a given item in the hierarchy has a specific set of metadata
  */
-USTRUCT(meta=(DisplayName="Has Metadata", Category="Hierarchy", DocumentationPolicy = "Strict", Keywords="MetadataExists,HasKey,SupportsMetadata", NodeColor="0.462745, 1,0, 0.329412", Varying))
+USTRUCT(meta=(DisplayName="Has Metadata", Category="Hierarchy", DocumentationPolicy = "Strict", Keywords="MetadataExists,HasKey,SupportsMetadata,FindMetadata", NodeColor="0.462745, 1,0, 0.329412", Varying))
 struct CONTROLRIG_API FRigUnit_HasMetadata : public FRigUnit
 {
 	GENERATED_BODY()
@@ -179,7 +254,7 @@ struct CONTROLRIG_API FRigUnit_HasMetadata : public FRigUnit
 		: Item(NAME_None, ERigElementType::Bone)
 		, Name(NAME_None)
 		, Type(ERigMetadataType::Float)
-		, bFound(false)
+		, Found(false)
 		, CachedIndex()
 	{}
 
@@ -195,7 +270,7 @@ struct CONTROLRIG_API FRigUnit_HasMetadata : public FRigUnit
 	/**
 	 * The name of the metadata to check
 	 */ 
-	UPROPERTY(meta = (Input))
+	UPROPERTY(meta = (Input, CustomWidget="MetadataName"))
 	FName Name;
 
 	/**
@@ -206,7 +281,7 @@ struct CONTROLRIG_API FRigUnit_HasMetadata : public FRigUnit
 
 	// True if the item has the metadata
 	UPROPERTY(meta=(Output))
-	bool bFound;
+	bool Found;
 
 	// Used to cache the internally
 	UPROPERTY()
@@ -233,7 +308,7 @@ struct CONTROLRIG_API FRigUnit_FindItemsWithMetadata : public FRigUnit
 	/**
 	 * The name of the metadata to find
 	 */ 
-	UPROPERTY(meta = (Input))
+	UPROPERTY(meta = (Input, CustomWidget="MetadataNameAll"))
 	FName Name;
 
 	/**
@@ -245,4 +320,322 @@ struct CONTROLRIG_API FRigUnit_FindItemsWithMetadata : public FRigUnit
 	// The items containing the metadata
 	UPROPERTY(meta=(Output))
 	TArray<FRigElementKey> Items;
+};
+
+/**
+ * Returns the metadata tags on an item
+ */
+USTRUCT(meta=(DisplayName="Get Tags", Category="Hierarchy", DocumentationPolicy = "Strict", Keywords="MetadataExists,HasKey,Tagging,FindTag", NodeColor="0.462745, 1,0, 0.329412", Varying))
+struct CONTROLRIG_API FRigUnit_GetMetadataTags : public FRigUnit
+{
+	GENERATED_BODY()
+
+	FRigUnit_GetMetadataTags()
+		: Item(NAME_None, ERigElementType::Bone)
+		, Tags()
+		, CachedIndex()
+	{}
+
+	RIGVM_METHOD()
+	virtual void Execute(const FRigUnitContext& Context) override;
+
+	/**
+	 * The item to check the metadata for
+	 */
+	UPROPERTY(meta = (Input, ExpandByDefault))
+	FRigElementKey Item;
+
+	/**
+	 * The name of the tag to check
+	 */ 
+	UPROPERTY(meta = (Output))
+	TArray<FName> Tags;
+
+	// Used to cache the internally
+	UPROPERTY()
+	FCachedRigElement CachedIndex;
+};
+
+/**
+ * Sets a single tag on an item 
+ */
+USTRUCT(meta=(DisplayName="Add Tag", Category="Hierarchy", DocumentationPolicy = "Strict", Keywords="MetadataExists,HasKey,Tagging,FindTag,SetTag", NodeColor="0.462745, 1,0, 0.329412", Varying))
+struct CONTROLRIG_API FRigUnit_SetMetadataTag : public FRigUnitMutable
+{
+	GENERATED_BODY()
+
+	FRigUnit_SetMetadataTag()
+		: Item(NAME_None, ERigElementType::Bone)
+		, Tag(NAME_None)
+		, CachedIndex()
+	{}
+
+	RIGVM_METHOD()
+	virtual void Execute(const FRigUnitContext& Context) override;
+
+	/**
+	 * The item to set the metadata for
+	 */
+	UPROPERTY(meta = (Input, ExpandByDefault))
+	FRigElementKey Item;
+
+	/**
+	 * The name of the tag to set
+	 */ 
+	UPROPERTY(meta = (Input))
+	FName Tag;
+
+	// Used to cache the internally
+	UPROPERTY()
+	FCachedRigElement CachedIndex;
+};
+
+/**
+ * Sets multiple tags on an item 
+ */
+USTRUCT(meta=(DisplayName="Add Multiple Tags", Category="Hierarchy", DocumentationPolicy = "Strict", Keywords="AddTags,MetadataExists,HasKey,Tagging,FindTag,SetTags", NodeColor="0.462745, 1,0, 0.329412", Varying))
+struct CONTROLRIG_API FRigUnit_SetMetadataTagArray : public FRigUnitMutable
+{
+	GENERATED_BODY()
+
+	FRigUnit_SetMetadataTagArray()
+		: Item(NAME_None, ERigElementType::Bone)
+		, Tags()
+		, CachedIndex()
+	{}
+
+	RIGVM_METHOD()
+	virtual void Execute(const FRigUnitContext& Context) override;
+
+	/**
+	 * The item to set the metadata for
+	 */
+	UPROPERTY(meta = (Input, ExpandByDefault))
+	FRigElementKey Item;
+
+	/**
+	 * The tags to set for the item
+	 */ 
+	UPROPERTY(meta = (Input))
+	TArray<FName> Tags;
+
+	// Used to cache the internally
+	UPROPERTY()
+	FCachedRigElement CachedIndex;
+};
+
+/**
+ * Removes a tag from an item 
+ */
+USTRUCT(meta=(DisplayName="Remove Tag", Category="Hierarchy", DocumentationPolicy = "Strict", Keywords="DeleteTag", NodeColor="0.462745, 1,0, 0.329412", Varying))
+struct CONTROLRIG_API FRigUnit_RemoveMetadataTag : public FRigUnitMutable
+{
+	GENERATED_BODY()
+
+	FRigUnit_RemoveMetadataTag()
+		: Item(NAME_None, ERigElementType::Bone)
+		, Tag(NAME_None)
+		, CachedIndex()
+	{}
+
+	RIGVM_METHOD()
+	virtual void Execute(const FRigUnitContext& Context) override;
+
+	/**
+	 * The item to set the metadata for
+	 */
+	UPROPERTY(meta = (Input, ExpandByDefault))
+	FRigElementKey Item;
+
+	/**
+	 * The name of the tag to set
+	 */ 
+	UPROPERTY(meta = (Input))
+	FName Tag;
+
+	/**
+	 * Returns true if the removal was successful
+	 */ 
+	UPROPERTY(meta = (Output))
+	bool Removed;
+
+	// Used to cache the internally
+	UPROPERTY()
+	FCachedRigElement CachedIndex;
+};
+
+/**
+ * Returns true if a given item in the hierarchy has a specific tag stored in the metadata
+ */
+USTRUCT(meta=(DisplayName="Has Tag", Category="Hierarchy", DocumentationPolicy = "Strict", Keywords="MetadataExists,HasKey,Tagging,FindTag", NodeColor="0.462745, 1,0, 0.329412", Varying))
+struct CONTROLRIG_API FRigUnit_HasMetadataTag : public FRigUnit
+{
+	GENERATED_BODY()
+
+	FRigUnit_HasMetadataTag()
+		: Item(NAME_None, ERigElementType::Bone)
+		, Tag(NAME_None)
+		, Found(false)
+		, CachedIndex()
+	{}
+
+	RIGVM_METHOD()
+	virtual void Execute(const FRigUnitContext& Context) override;
+
+	/**
+	 * The item to check the metadata for
+	 */
+	UPROPERTY(meta = (Input, ExpandByDefault))
+	FRigElementKey Item;
+
+	/**
+	 * The name of the tag to check
+	 */ 
+	UPROPERTY(meta = (Input, CustomWidget="MetadataTagName"))
+	FName Tag;
+
+	// True if the item has the metadata
+	UPROPERTY(meta=(Output))
+	bool Found;
+
+	// Used to cache the internally
+	UPROPERTY()
+	FCachedRigElement CachedIndex;
+};
+
+/**
+ * Returns true if a given item in the hierarchy has all of the provided tags
+ */
+USTRUCT(meta=(DisplayName="Has Multiple Tags", Category="Hierarchy", DocumentationPolicy = "Strict", Keywords="MetadataExists,HasKey,Tagging,FindTag", NodeColor="0.462745, 1,0, 0.329412", Varying))
+struct CONTROLRIG_API FRigUnit_HasMetadataTagArray : public FRigUnit
+{
+	GENERATED_BODY()
+
+	FRigUnit_HasMetadataTagArray()
+		: Item(NAME_None, ERigElementType::Bone)
+		, Tags()
+		, Found(false)
+		, CachedIndex()
+	{}
+
+	RIGVM_METHOD()
+	virtual void Execute(const FRigUnitContext& Context) override;
+
+	/**
+	 * The item to check the metadata for
+	 */
+	UPROPERTY(meta = (Input, ExpandByDefault))
+	FRigElementKey Item;
+
+	/**
+	 * The name of the tag to check
+	 */ 
+	UPROPERTY(meta = (Input, CustomWidget="MetadataTagName"))
+	TArray<FName> Tags;
+
+	// True if the item has the metadata
+	UPROPERTY(meta=(Output))
+	bool Found;
+
+	// Used to cache the internally
+	UPROPERTY()
+	FCachedRigElement CachedIndex;
+};
+
+/**
+ * Returns all items with a specific tag
+ */
+USTRUCT(meta=(DisplayName="Find Items with Tag", Category="Hierarchy", DocumentationPolicy = "Strict", Keywords="MetadataExists,HasKey,SupportsMetadata", NodeColor="0.462745, 1,0, 0.329412", Varying))
+struct CONTROLRIG_API FRigUnit_FindItemsWithMetadataTag : public FRigUnit
+{
+	GENERATED_BODY()
+
+	FRigUnit_FindItemsWithMetadataTag()
+		: Tag(NAME_None)
+		, Items()
+	{}
+
+	RIGVM_METHOD()
+	virtual void Execute(const FRigUnitContext& Context) override;
+
+	/**
+	 * The name of the tag to find
+	 */ 
+	UPROPERTY(meta = (Input, CustomWidget="MetadataTagNameAll"))
+	FName Tag;
+
+	// The items containing the metadata
+	UPROPERTY(meta=(Output))
+	TArray<FRigElementKey> Items;
+};
+
+/**
+ * Returns all items with a specific tag
+ */
+USTRUCT(meta=(DisplayName="Find Items with multiple Tags", Category="Hierarchy", DocumentationPolicy = "Strict", Keywords="MetadataExists,HasKey,SupportsMetadata", NodeColor="0.462745, 1,0, 0.329412", Varying))
+struct CONTROLRIG_API FRigUnit_FindItemsWithMetadataTagArray : public FRigUnit
+{
+	GENERATED_BODY()
+
+	FRigUnit_FindItemsWithMetadataTagArray()
+		: Tags()
+		, Items()
+	{}
+
+	RIGVM_METHOD()
+	virtual void Execute(const FRigUnitContext& Context) override;
+
+	/**
+	 * The tags to find
+	 */ 
+	UPROPERTY(meta = (Input, CustomWidget="MetadataTagNameAll"))
+	TArray<FName> Tags;
+
+	// The items containing the metadata
+	UPROPERTY(meta=(Output))
+	TArray<FRigElementKey> Items;
+};
+
+/**
+ * Filters an item array by a list of tags
+ */
+USTRUCT(meta=(DisplayName="Filter Items by Tags", Category="Hierarchy", DocumentationPolicy = "Strict", Keywords="MetadataExists,HasKey,SupportsMetadata", NodeColor="0.462745, 1,0, 0.329412", Varying))
+struct CONTROLRIG_API FRigUnit_FilterItemsByMetadataTags : public FRigUnit
+{
+	GENERATED_BODY()
+
+	FRigUnit_FilterItemsByMetadataTags()
+		: Items()
+		, Tags()
+		, Inclusive(true)
+		, Result()
+	{}
+
+	RIGVM_METHOD()
+	virtual void Execute(const FRigUnitContext& Context) override;
+
+	// The items to filter
+	UPROPERTY(meta=(Input))
+	TArray<FRigElementKey> Items;
+
+	/**
+	 * The tags to find
+	 */ 
+	UPROPERTY(meta = (Input, CustomWidget="MetadataTagNameAll"))
+	TArray<FName> Tags;
+
+	/**
+     * If set to true only items with ALL of tags will be returned,
+     * if set to false items with ANY of the tags will be removed
+     */ 
+	UPROPERTY(meta = (Input, CustomWidget="MetadataTagNameAll"))
+	bool Inclusive;
+
+	// The results of the filter
+	UPROPERTY(meta=(Output))
+	TArray<FRigElementKey> Result;
+
+	// Used to cache the internally
+	UPROPERTY()
+	TArray<FCachedRigElement> CachedIndices;
 };
\ No newline at end of file
diff --git a/Engine/Plugins/Animation/ControlRig/Source/ControlRig/Public/Rigs/RigHierarchy.h b/Engine/Plugins/Animation/ControlRig/Source/ControlRig/Public/Rigs/RigHierarchy.h
--- a/Engine/Plugins/Animation/ControlRig/Source/ControlRig/Public/Rigs/RigHierarchy.h
+++ b/Engine/Plugins/Animation/ControlRig/Source/ControlRig/Public/Rigs/RigHierarchy.h
@@ -21,6 +21,8 @@ class URigHierarchyController;
 
 DECLARE_MULTICAST_DELEGATE_ThreeParams(FRigHierarchyModifiedEvent, ERigHierarchyNotification /* type */, URigHierarchy* /* hierarchy */, const FRigBaseElement* /* element */);
 DECLARE_EVENT_FiveParams(URigHierarchy, FRigHierarchyUndoRedoTransformEvent, URigHierarchy*, const FRigElementKey&, ERigTransformType::Type, const FTransform&, bool /* bUndo */);
+DECLARE_MULTICAST_DELEGATE_TwoParams(FRigHierarchyMetadataChangedDelegate, const FRigElementKey& /* Key */, const FName& /* Name */);
+DECLARE_MULTICAST_DELEGATE_ThreeParams(FRigHierarchyMetadataTagChangedDelegate, const FRigElementKey& /* Key */, const FName& /* Tag */, bool /* AddedOrRemoved */);
 
 UENUM()
 enum ERigTransformStackEntryType
@@ -95,6 +97,7 @@ public:
 	typedef TMap<int32, TArray<int32>> TElementDependencyMap;
 	typedef TPair<int32, TArray<int32>> TElementDependencyMapPair;
 	typedef TTuple<int32, int32, int32, ERigTransformType::Type> TInstructionSliceElement;
+	inline static const FName TagMetadataName = TEXT("Tags");
 
 	URigHierarchy();
 
@@ -1356,6 +1359,40 @@ public:
 		return SetArrayMetadata<FRigElementKey>(InItem, ERigMetadataType::RigElementKeyArray, InMetadataName, InValue);
 	}
 
+	/*
+	 * Returns the tags for a given item
+	 * @param InItem The item to return the tags for
+	 */
+	UFUNCTION(BlueprintPure, Category = URigHierarchy)
+	FORCEINLINE TArray<FName> GetTags(FRigElementKey InItem) const
+	{
+		return GetNameArrayMetadata(InItem, TagMetadataName);
+	}
+
+	/*
+	 * Returns true if a given item has a certain tag
+	 * @param InItem The item to return the tags for
+	 * @param InTag The tag to check
+	 */
+	UFUNCTION(BlueprintPure, Category = URigHierarchy)
+	FORCEINLINE bool HasTag(FRigElementKey InItem, FName InTag) const
+	{
+		return GetTags(InItem).Contains(InTag);
+	}
+
+	/*
+     * Sets a tag on an element in the hierarchy
+     * @param InItem The item to set the tag for
+     * @param InTag The tag to set
+     */
+	UFUNCTION(BlueprintCallable, Category = URigHierarchy)
+	FORCEINLINE bool SetTag(FRigElementKey InItem, FName InTag)
+	{
+		TArray<FName> Tags = GetTags(InItem);
+		Tags.AddUnique(InTag);
+		return SetNameArrayMetadata(InItem, TagMetadataName, Tags);
+	}
+
 	/**
 	 * Returns the selected elements
 	 * @InTypeFilter The types to retrieve the selection for
@@ -1478,6 +1515,20 @@ public:
 	 */
 	FRigHierarchyModifiedEvent& OnModified() { return ModifiedEvent; }
 
+	/**
+	 * Returns the MetadataChanged event, which can be used to track metadata changes
+	 * Note: This notification has a very high volume - so the consequences of subscribing
+	 * to it may cause performance slowdowns.
+	 */
+	FRigHierarchyMetadataChangedDelegate& OnMetadataChanged() { return MetadataChangedDelegate; }
+
+	/**
+	 * Returns the MetadataTagChanged event, which can be used to track metadata tag changes
+	 * Note: This notification has a very high volume - so the consequences of subscribing
+	 * to it may cause performance slowdowns.
+	 */
+	FRigHierarchyMetadataTagChangedDelegate& OnMetadataTagChanged() { return MetadataTagChangedDelegate; }
+
 	/**
 	 * Returns the local current or initial value for a given key.
 	 * If the key is invalid FTransform::Identity will be returned.
@@ -2829,6 +2880,16 @@ public:
 	 */
 	void IncrementTopologyVersion();
 
+	/**
+	 * Returns the metadata version of this hierarchy
+	 */
+	uint16 GetMetadataVersion() const { return MetadataVersion; }
+
+	/**
+	 * Increments the metadata version
+	 */
+	FORCEINLINE void IncrementMetadataVersion() { MetadataVersion++; }
+
 	/**
 	 * Returns the current / initial pose of the hierarchy
 	 * @param bInitial If set to true the initial pose will be returned
@@ -3114,6 +3175,8 @@ public:
 private:
 
 	FRigHierarchyModifiedEvent ModifiedEvent;
+	FRigHierarchyMetadataChangedDelegate MetadataChangedDelegate;
+	FRigHierarchyMetadataTagChangedDelegate MetadataTagChangedDelegate;
 	FRigEventDelegate EventDelegate;
 
 public:
@@ -3389,12 +3452,12 @@ private:
 	/*
 	 * Helper function to create an element for a given type
 	 */
-	static FRigBaseElement* MakeElement(ERigElementType InElementType, int32 InCount = 1, int32* OutStructureSize = nullptr); 
+	FRigBaseElement* MakeElement(ERigElementType InElementType, int32 InCount = 1, int32* OutStructureSize = nullptr); 
 
 	/*
 	* Helper function to create an element for a given type
 	*/
-	static void DestroyElement(FRigBaseElement*& InElement);
+	void DestroyElement(FRigBaseElement*& InElement);
 
 	/*
 	 * Templated helper function to create an element
@@ -3449,6 +3512,13 @@ private:
 	UPROPERTY(transient)
 	uint16 TopologyVersion;
 
+	/**
+	 * The metadata version of the hierarchy changes when metadata is being
+	 * created or removed (not when the metadata values changes)
+	 */
+	UPROPERTY(transient)
+	uint16 MetadataVersion;
+
 	/**
 	 * If set to false the dirty flag propagation will be disabled
 	 */
@@ -3867,6 +3937,9 @@ private:
 		return SetMetadata<TArray<T>>(InElement, InType, InMetadataName, InValue);
 	}
 
+	void OnMetadataChanged(const FRigElementKey& InKey, const FName& InName);
+	void OnMetadataTagChanged(const FRigElementKey& InKey, const FName& InTag, bool bAdded);
+
 protected:
 	
 	bool bEnableCacheValidityCheck;
@@ -3961,6 +4034,7 @@ private:
 
 	// certain units are allowed to use this
 	friend struct FRigUnit_AddParent;
+	friend struct FRigUnit_SetDefaultParent;
 
 private:
 	TGuardValue<bool> GuardIsControllerAvailable;
diff --git a/Engine/Plugins/Animation/ControlRig/Source/ControlRig/Public/Rigs/RigHierarchyElements.h b/Engine/Plugins/Animation/ControlRig/Source/ControlRig/Public/Rigs/RigHierarchyElements.h
--- a/Engine/Plugins/Animation/ControlRig/Source/ControlRig/Public/Rigs/RigHierarchyElements.h
+++ b/Engine/Plugins/Animation/ControlRig/Source/ControlRig/Public/Rigs/RigHierarchyElements.h
@@ -12,6 +12,9 @@ struct FRigBaseElement;
 struct FRigControlElement;
 class URigHierarchy;
 
+DECLARE_DELEGATE_RetVal_ThreeParams(FTransform, FRigReferenceGetWorldTransformDelegate, const FRigUnitContext*, const FRigElementKey& /* Key */, bool /* bInitial */);
+DECLARE_DELEGATE_TwoParams(FRigElementMetadataChangedDelegate, const FRigElementKey& /* Key */, const FName& /* Name */);
+DECLARE_DELEGATE_ThreeParams(FRigElementMetadataTagChangedDelegate, const FRigElementKey& /* Key */, const FName& /* Tag */, bool /* AddedOrRemoved */);
 DECLARE_DELEGATE_RetVal_ThreeParams(FTransform, FRigReferenceGetWorldTransformDelegate, const FRigUnitContext*, const FRigElementKey& /* Key */, bool /* bInitial */);
 
 #define DECLARE_RIG_ELEMENT_METHODS(ElementType) \
@@ -513,6 +516,7 @@ public:
 	, bSelected(false)
 	, CreatedAtInstructionIndex(INDEX_NONE)
 	, TopologyVersion(0)
+	, MetadataVersion(0)
 	, OwnedInstances(0)
 	{}
 
@@ -572,6 +576,7 @@ public:
 	FORCEINLINE bool IsSelected() const { return bSelected; }
 	FORCEINLINE int32 GetCreatedAtInstructionIndex() const { return CreatedAtInstructionIndex; }
 	FORCEINLINE bool IsProcedural() const { return CreatedAtInstructionIndex != INDEX_NONE; }
+	FORCEINLINE int32 GetMetadataVersion() const { return MetadataVersion; }
 
 	FORCEINLINE int32 NumMetadata() const { return Metadata.Num(); }
 	FORCEINLINE FRigBaseMetadata* GetMetadata(int32 InIndex) const { return Metadata[InIndex]; }
@@ -666,15 +671,25 @@ protected:
 	// sets up the metadata and ensures the right type
 	FRigBaseMetadata* SetupValidMetadata(const FName& InName, ERigMetadataType InType);
 
+	void NotifyMetadataChanged(const FName& InName);
+	void NotifyMetadataTagChanged(const FName& InTag, bool bAdded);
+
 	mutable uint16 TopologyVersion;
+	mutable uint16 MetadataVersion;
 	mutable FRigBaseElementChildrenArray CachedChildren;
 
 	// used for constructing / destructing the memory. typically == 1
 	int32 OwnedInstances;
 
+	FRigElementMetadataChangedDelegate MetadataChangedDelegate;
+	FRigElementMetadataTagChangedDelegate MetadataTagChangedDelegate;
+
 	friend class URigHierarchy;
 	friend class URigHierarchyController;
 	friend struct FRigDispatch_SetMetadata;
+	friend struct FRigUnit_SetMetadataTag;
+	friend struct FRigUnit_SetMetadataTagArray;
+	friend struct FRigUnit_RemoveMetadataTag;
 };
 
 USTRUCT(BlueprintType)
diff --git a/Engine/Plugins/Animation/ControlRig/Source/ControlRigEditor/Private/Graph/ControlRigGraphPanelPinFactory.cpp b/Engine/Plugins/Animation/ControlRig/Source/ControlRigEditor/Private/Graph/ControlRigGraphPanelPinFactory.cpp
--- a/Engine/Plugins/Animation/ControlRig/Source/ControlRigEditor/Private/Graph/ControlRigGraphPanelPinFactory.cpp
+++ b/Engine/Plugins/Animation/ControlRig/Source/ControlRigEditor/Private/Graph/ControlRigGraphPanelPinFactory.cpp
@@ -200,6 +200,229 @@ TSharedPtr<SGraphPin> FControlRigGraphPanelPinFactory::CreatePin(UEdGraphPin* In
 								return ChannelNames.Names.Get();
 							}
 
+							static TArray<TSharedPtr<FString>> EmptyNameList;
+							return &EmptyNameList;
+						});
+				}
+				else if (CustomWidgetName == TEXT("MetadataName") || CustomWidgetName == TEXT("MetadataNameAll"))
+				{
+					struct FCachedMetadataNames
+					{
+						int32 MetadataVersion;
+						TSharedPtr<TArray<TSharedPtr<FString>>> Names;
+						
+						FCachedMetadataNames()
+						: MetadataVersion(INDEX_NONE)
+						{}
+					};
+
+					const bool bCompleteList = CustomWidgetName == TEXT("MetadataNameAll");
+					
+					return SNew(SControlRigGraphPinNameList, InPin)
+						.ModelPin(ModelPin)
+						.OnGetNameListContent_Lambda([RigGraph, bCompleteList](const URigVMPin* InPin)
+						{
+							if (const UControlRigBlueprint* Blueprint = RigGraph->GetTypedOuter<UControlRigBlueprint>())
+							{
+								if(UControlRig* ControlRig = Cast<UControlRig>(Blueprint->GetObjectBeingDebugged()))
+								{
+									FRigElementKey ItemKey;
+									bool bHasSourceLinks = false;
+
+									// find the pin that holds the item
+									for(URigVMPin* Pin : InPin->GetRootPin()->GetNode()->GetPins())
+									{
+										if(Pin->GetCPPType() == FRigElementKey::StaticStruct()->GetStructCPPName())
+										{
+											if(!bCompleteList)
+											{
+												if(Pin->GetSourceLinks(true).Num() > 0)
+												{
+													bHasSourceLinks = true;
+												}
+											}
+											if(!bHasSourceLinks)
+											{
+												const FString DefaultValue = Pin->GetDefaultValue();
+												if(!DefaultValue.IsEmpty())
+												{
+													FRigElementKey::StaticStruct()->ImportText(*DefaultValue, &ItemKey, nullptr, EPropertyPortFlags::PPF_None, nullptr, FRigElementKey::StaticStruct()->GetName(), true);
+												}
+											}
+											break;
+										}
+									}
+
+									const FRigBaseElement* Element = nullptr;
+									if(ItemKey.IsValid())
+									{
+										if(!bCompleteList && !bHasSourceLinks)
+										{
+											Element = ControlRig->GetHierarchy()->Find(ItemKey);
+											if(Element == nullptr)
+											{
+												static TArray<TSharedPtr<FString>> EmptyNameList;
+												return &EmptyNameList;
+											}
+										}
+									}
+
+									const FString MapHash = Element ?
+										Blueprint->GetPathName() + TEXT("|") + ItemKey.ToString() :
+										Blueprint->GetPathName();
+									const int32 MetadataVersion = Element ? Element->GetMetadataVersion() : ControlRig->GetHierarchy()->GetMetadataVersion(); 
+
+									static TMap<FString, FCachedMetadataNames> MetadataNameLists;
+									FCachedMetadataNames& MetadataNames = MetadataNameLists.FindOrAdd(MapHash);
+
+									if(MetadataNames.MetadataVersion != MetadataVersion)
+									{
+										TArray<FName> Names;
+										if(bCompleteList || bHasSourceLinks || (Element == nullptr))
+										{
+											for(int32 ElementIndex=0; ElementIndex < ControlRig->GetHierarchy()->Num(); ElementIndex++)
+											{
+												const FRigBaseElement* OtherElement = ControlRig->GetHierarchy()->Get(ElementIndex);
+												for(int32 MetadataIndex = 0; MetadataIndex < OtherElement->NumMetadata(); MetadataIndex++)
+												{
+													Names.AddUnique(OtherElement->GetMetadata(MetadataIndex)->GetName());
+												}
+											}
+										}
+										else if(Element)
+										{
+											for(int32 MetadataIndex = 0; MetadataIndex < Element->NumMetadata(); MetadataIndex++)
+											{
+												Names.Add(Element->GetMetadata(MetadataIndex)->GetName());
+											}
+										}
+
+										if(!MetadataNames.Names.IsValid())
+										{
+											MetadataNames.Names = MakeShareable(new TArray<TSharedPtr<FString>>());
+										}
+										MetadataNames.Names->Reset();
+
+										for(const FName& Name : Names)
+										{
+											MetadataNames.Names->Add(MakeShareable(new FString(Name.ToString())));
+										}
+
+										MetadataNames.Names->Sort([](const TSharedPtr<FString>& A, const TSharedPtr<FString>& B)
+										{
+											return A.Get() > B.Get();
+										});
+										MetadataNames.Names->Insert(MakeShareable(new FString(FName(NAME_None).ToString())), 0);
+
+										MetadataNames.MetadataVersion = MetadataVersion;
+									}
+									return MetadataNames.Names.Get();
+								}
+							}
+
+							static TArray<TSharedPtr<FString>> EmptyNameList;
+							return &EmptyNameList;
+						});
+				}
+				else if (CustomWidgetName == TEXT("MetadataTagName") || CustomWidgetName == TEXT("MetadataTagNameAll"))
+				{
+					bool bCompleteList = CustomWidgetName == TEXT("MetadataTagNameAll");
+
+					struct FCachedMetadataTagNames
+					{
+						TSharedPtr<TArray<TSharedPtr<FString>>> Names;
+						FCachedMetadataTagNames()
+						{}
+					};
+
+					return SNew(SControlRigGraphPinNameList, InPin)
+						.ModelPin(ModelPin)
+						.OnGetNameListContent_Lambda([RigGraph, bCompleteList](const URigVMPin* InPin)
+						{
+							if (const UControlRigBlueprint* Blueprint = RigGraph->GetTypedOuter<UControlRigBlueprint>())
+							{
+								if(UControlRig* ControlRig = Cast<UControlRig>(Blueprint->GetObjectBeingDebugged()))
+								{
+									FRigElementKey ItemKey;
+									bool bHasSourceLinks = false;
+
+									// find the pin that holds the item
+									for(URigVMPin* Pin : InPin->GetRootPin()->GetNode()->GetPins())
+									{
+										if(Pin->GetCPPType() == FRigElementKey::StaticStruct()->GetStructCPPName())
+										{
+											if(!bCompleteList)
+											{
+												if(Pin->GetSourceLinks(true).Num() > 0)
+												{
+													bHasSourceLinks = true;
+												}
+											}
+											if(!bHasSourceLinks)
+											{
+												const FString DefaultValue = Pin->GetDefaultValue();
+												if(!DefaultValue.IsEmpty())
+												{
+													FRigElementKey::StaticStruct()->ImportText(*DefaultValue, &ItemKey, nullptr, EPropertyPortFlags::PPF_None, nullptr, FRigElementKey::StaticStruct()->GetName(), true);
+												}
+											}
+											break;
+										}
+									}
+
+									const FString MapHash = ItemKey.IsValid() ?
+										Blueprint->GetPathName() + TEXT("|") + ItemKey.ToString() :
+										Blueprint->GetPathName();
+
+									static TMap<FString, FCachedMetadataTagNames> MetadataTagNameLists;
+									FCachedMetadataTagNames& MetadataTagNames = MetadataTagNameLists.FindOrAdd(MapHash);
+
+									TArray<FName> Tags; 
+									if(bCompleteList || !ItemKey.IsValid())
+									{
+										for(int32 ElementIndex=0; ElementIndex < ControlRig->GetHierarchy()->Num(); ElementIndex++)
+										{
+											const FRigBaseElement* Element = ControlRig->GetHierarchy()->Get(ElementIndex);
+											if(const FRigNameArrayMetadata* Md = Cast<FRigNameArrayMetadata>(Element->GetMetadata(URigHierarchy::TagMetadataName, ERigMetadataType::NameArray)))
+											{
+												for(const FName& Tag : Md->GetValue())
+												{
+													Tags.AddUnique(Tag);
+												}
+											}
+										}
+									}
+									else
+									{
+										if(const FRigBaseElement* Element = ControlRig->GetHierarchy()->Find(ItemKey))
+										{
+											if(const FRigNameArrayMetadata* Md = Cast<FRigNameArrayMetadata>(Element->GetMetadata(URigHierarchy::TagMetadataName, ERigMetadataType::NameArray)))
+											{
+												Tags = Md->GetValue();
+											}
+										}
+									}
+
+									if(!MetadataTagNames.Names.IsValid())
+									{
+										MetadataTagNames.Names = MakeShareable(new TArray<TSharedPtr<FString>>());
+									}
+									MetadataTagNames.Names->Reset();
+
+									for(const FName& Tag : Tags)
+									{
+										MetadataTagNames.Names->Add(MakeShareable(new FString(Tag.ToString())));
+									}
+									MetadataTagNames.Names->Sort([](const TSharedPtr<FString>& A, const TSharedPtr<FString>& B)
+									{
+										return A.Get() > B.Get();
+									});
+									MetadataTagNames.Names->Insert(MakeShareable(new FString(FName(NAME_None).ToString())), 0);
+
+									return MetadataTagNames.Names.Get();
+								}
+							}
+
 							static TArray<TSharedPtr<FString>> EmptyNameList;
 							return &EmptyNameList;
 						});
