ControlRig: Follow up to 21037858 - minor metadata fix

diff --git a/Engine/Plugins/Animation/ControlRig/Source/ControlRig/Private/Units/Hierarchy/RigUnit_Metadata.h b/Engine/Plugins/Animation/ControlRig/Source/ControlRig/Private/Units/Hierarchy/RigUnit_Metadata.h
--- a/Engine/Plugins/Animation/ControlRig/Source/ControlRig/Private/Units/Hierarchy/RigUnit_Metadata.h
+++ b/Engine/Plugins/Animation/ControlRig/Source/ControlRig/Private/Units/Hierarchy/RigUnit_Metadata.h
@@ -628,7 +628,7 @@ struct CONTROLRIG_API FRigUnit_FilterItemsByMetadataTags : public FRigUnit
      * If set to true only items with ALL of tags will be returned,
      * if set to false items with ANY of the tags will be removed
      */ 
-	UPROPERTY(meta = (Input, CustomWidget="MetadataTagNameAll"))
+	UPROPERTY(meta = (Input))
 	bool Inclusive;
 
 	// The results of the filter
