Control Rig: Set default parent node

diff --git a/Engine/Plugins/Animation/ControlRig/Source/ControlRig/Private/Units/Execution/RigUnit_DynamicHierarchy.cpp b/Engine/Plugins/Animation/ControlRig/Source/ControlRig/Private/Units/Execution/RigUnit_DynamicHierarchy.cpp
--- a/Engine/Plugins/Animation/ControlRig/Source/ControlRig/Private/Units/Execution/RigUnit_DynamicHierarchy.cpp
+++ b/Engine/Plugins/Animation/ControlRig/Source/ControlRig/Private/Units/Execution/RigUnit_DynamicHierarchy.cpp
@@ -74,6 +74,34 @@ FRigUnit_AddParent_Execute()
 	}
 }
 
+FRigUnit_SetDefaultParent_Execute()
+{
+	if(!FRigUnit_DynamicHierarchyBase::IsValidToRunInContext(Context, ExecuteContext, true))
+	{
+		return;
+	}
+	
+	FRigTransformElement* ChildElement = ExecuteContext.Hierarchy->Find<FRigTransformElement>(Child);
+	if(ChildElement == nullptr)
+	{
+		UE_CONTROLRIG_RIGUNIT_REPORT_WARNING(TEXT("Child item %s does not exist."), *Child.ToString())
+		return;
+	}
+
+	FRigTransformElement* ParentElement = ExecuteContext.Hierarchy->Find<FRigTransformElement>(Parent);
+	if(ParentElement == nullptr)
+	{
+		UE_CONTROLRIG_RIGUNIT_REPORT_WARNING(TEXT("Parent item %s does not exist."), *Parent.ToString())
+		return;
+	}
+
+	FRigHierarchyEnableControllerBracket EnableController(ExecuteContext.Hierarchy, true);
+	if(URigHierarchyController* Controller = ExecuteContext.Hierarchy->GetController(true))
+	{
+		Controller->AddParent(ChildElement, ParentElement, 1.0f, true, true);
+	}
+}
+
 FRigUnit_SwitchParent_Execute()
 {
 	if(!FRigUnit_DynamicHierarchyBase::IsValidToRunInContext(Context, ExecuteContext, false))
diff --git a/Engine/Plugins/Animation/ControlRig/Source/ControlRig/Private/Units/Execution/RigUnit_DynamicHierarchy.h b/Engine/Plugins/Animation/ControlRig/Source/ControlRig/Private/Units/Execution/RigUnit_DynamicHierarchy.h
--- a/Engine/Plugins/Animation/ControlRig/Source/ControlRig/Private/Units/Execution/RigUnit_DynamicHierarchy.h
+++ b/Engine/Plugins/Animation/ControlRig/Source/ControlRig/Private/Units/Execution/RigUnit_DynamicHierarchy.h
@@ -54,6 +54,35 @@ struct CONTROLRIG_API FRigUnit_AddParent : public FRigUnit_DynamicHierarchyBaseM
 	FRigElementKey Parent;
 };
 
+/**
+ * Changes the default parent for an item - this removes all other current parents.
+ */
+USTRUCT(meta=(DisplayName="Set Default Parent", Keywords="Children,Parent,Constraint,Space,SetParent,AddParent", Varying))
+struct CONTROLRIG_API FRigUnit_SetDefaultParent : public FRigUnit_DynamicHierarchyBaseMutable
+{
+	GENERATED_BODY()
+
+	FRigUnit_SetDefaultParent()
+	{
+		Child = Parent = FRigElementKey(NAME_None, ERigElementType::Control);
+	}
+
+	RIGVM_METHOD()
+	virtual void Execute(const FRigUnitContext& Context) override;
+
+	/*
+	 * The child to be parented under the new default parent
+	 */
+	UPROPERTY(meta = (Input, ExpandByDefault))
+	FRigElementKey Child;
+
+	/*
+	 * The default parent to be used for the child
+	 */
+	UPROPERTY(meta = (Input, ExpandByDefault))
+	FRigElementKey Parent;
+};
+
 UENUM()
 enum class ERigSwitchParentMode : uint8
 {
