Control Rig: Fix crash in FilterItemsBytag

diff --git a/Engine/Plugins/Animation/ControlRig/Source/ControlRig/Private/Units/Hierarchy/RigUnit_Metadata.cpp b/Engine/Plugins/Animation/ControlRig/Source/ControlRig/Private/Units/Hierarchy/RigUnit_Metadata.cpp
--- a/Engine/Plugins/Animation/ControlRig/Source/ControlRig/Private/Units/Hierarchy/RigUnit_Metadata.cpp
+++ b/Engine/Plugins/Animation/ControlRig/Source/ControlRig/Private/Units/Hierarchy/RigUnit_Metadata.cpp
@@ -633,39 +633,41 @@ FRigUnit_FilterItemsByMetadataTags_Execute()
 	{
 		if(CachedIndices[Index].UpdateCache(Items[Index], Hierarchy))
 		{
-			const FRigBaseElement* Element = CachedIndices[Index].GetElement();
-			if(const FRigNameArrayMetadata* Md = Cast<FRigNameArrayMetadata>(Element->GetMetadata(URigHierarchy::TagMetadataName, ERigMetadataType::NameArray)))
+			if(const FRigBaseElement* Element = Hierarchy->Get(CachedIndices[Index]))
 			{
-				if(Inclusive)
+				if(const FRigNameArrayMetadata* Md = Cast<FRigNameArrayMetadata>(Element->GetMetadata(URigHierarchy::TagMetadataName, ERigMetadataType::NameArray)))
 				{
-					bool bFoundAll = true;
-					for(const FName& Tag : Tags)
+					if(Inclusive)
 					{
-						if(!Md->GetValue().Contains(Tag))
+						bool bFoundAll = true;
+						for(const FName& Tag : Tags)
 						{
-							bFoundAll = false;
-							break;
+							if(!Md->GetValue().Contains(Tag))
+							{
+								bFoundAll = false;
+								break;
+							}
 						}
-					}
-					if(bFoundAll)
-					{
-						Result.Add(Element->GetKey());
-					}
-				}
-				else
-				{
-					bool bFoundAny = false;
-					for(const FName& Tag : Tags)
-					{
-						if(Md->GetValue().Contains(Tag))
+						if(bFoundAll)
 						{
-							bFoundAny = true;
-							break;
+							Result.Add(Element->GetKey());
 						}
 					}
-					if(bFoundAny)
+					else
 					{
-						Result.Add(Element->GetKey());
+						bool bFoundAny = false;
+						for(const FName& Tag : Tags)
+						{
+							if(Md->GetValue().Contains(Tag))
+							{
+								bFoundAny = true;
+								break;
+							}
+						}
+						if(bFoundAny)
+						{
+							Result.Add(Element->GetKey());
+						}
 					}
 				}
 			}
