Control Rig: Fix core equals / not equals dispatch

diff --git a/Engine/Plugins/Animation/ControlRig/Source/ControlRig/Private/Units/Core/RigUnit_CoreDispatch.cpp b/Engine/Plugins/Animation/ControlRig/Source/ControlRig/Private/Units/Core/RigUnit_CoreDispatch.cpp
--- a/Engine/Plugins/Animation/ControlRig/Source/ControlRig/Private/Units/Core/RigUnit_CoreDispatch.cpp
+++ b/Engine/Plugins/Animation/ControlRig/Source/ControlRig/Private/Units/Core/RigUnit_CoreDispatch.cpp
@@ -81,8 +81,7 @@ FRigVMFunctionPtr FRigDispatch_CoreEquals::GetDispatchFunctionImpl(const FRigVMT
 	{
 		return &FRigDispatch_CoreEquals::MathTypeEquals<FLinearColor>;
 	}
-
-	return FRigDispatch_CoreBase::GetDispatchFunctionImpl(InTypes);
+	return &FRigDispatch_CoreEquals::Execute;
 }
 
 bool FRigDispatch_CoreEquals::AdaptResult(bool bResult, const FRigVMExtendedExecuteContext& InContext)
