Control Rig: Fix queued events for remove hierarchy elements

diff --git a/Engine/Plugins/Animation/ControlRig/Source/ControlRig/Private/Rigs/RigHierarchy.cpp b/Engine/Plugins/Animation/ControlRig/Source/ControlRig/Private/Rigs/RigHierarchy.cpp
--- a/Engine/Plugins/Animation/ControlRig/Source/ControlRig/Private/Rigs/RigHierarchy.cpp
+++ b/Engine/Plugins/Animation/ControlRig/Source/ControlRig/Private/Rigs/RigHierarchy.cpp
@@ -2606,10 +2606,11 @@ void URigHierarchy::SendQueuedNotifications()
 	for(int32 Index = FilteredNotifications.Num() - 1; Index >= 0; Index--)
 	{
 		const FQueuedNotification& Entry = FilteredNotifications[Index];
-		// we expect all elements to still exist, removal /rename
-		// // is not allowed during the run of a VM
-		const FRigBaseElement* Element = FindChecked(Entry.Key); 
-		ModifiedEvent.Broadcast(Entry.Type, this, Element);
+		const FRigBaseElement* Element = Find(Entry.Key);
+		if(Element)
+		{
+			ModifiedEvent.Broadcast(Entry.Type, this, Element);
+		}
 	}
 
 	ModifiedEvent.Broadcast(ERigHierarchyNotification::InteractionBracketClosed, this, nullptr);
