Control Rig: Fix FilterItemsByTags node exclusive mode

diff --git a/Engine/Plugins/Animation/ControlRig/Source/ControlRig/Private/Units/Hierarchy/RigUnit_Metadata.cpp b/Engine/Plugins/Animation/ControlRig/Source/ControlRig/Private/Units/Hierarchy/RigUnit_Metadata.cpp
--- a/Engine/Plugins/Animation/ControlRig/Source/ControlRig/Private/Units/Hierarchy/RigUnit_Metadata.cpp
+++ b/Engine/Plugins/Animation/ControlRig/Source/ControlRig/Private/Units/Hierarchy/RigUnit_Metadata.cpp
@@ -664,12 +664,16 @@ FRigUnit_FilterItemsByMetadataTags_Execute()
 								break;
 							}
 						}
-						if(bFoundAny)
+						if(!bFoundAny)
 						{
 							Result.Add(Element->GetKey());
 						}
 					}
 				}
+				else if(!Inclusive)
+				{
+					Result.Add(Element->GetKey());
+				}
 			}
 		}
 		else
