Control Rig: New proxy controller icon

diff --git a/Engine/Plugins/Animation/ControlRig/Source/ControlRigEditor/Private/ControlRigEditorStyle.h b/Engine/Plugins/Animation/ControlRig/Source/ControlRigEditor/Private/ControlRigEditorStyle.h
--- a/Engine/Plugins/Animation/ControlRig/Source/ControlRigEditor/Private/ControlRigEditorStyle.h
+++ b/Engine/Plugins/Animation/ControlRig/Source/ControlRigEditor/Private/ControlRigEditorStyle.h
@@ -149,6 +149,7 @@ public:
 			Set("ControlRig.Tree.BoneUser", new IMAGE_BRUSH("Slate/BoneNonWeighted_16x", Icon16x16));
 			Set("ControlRig.Tree.BoneImported", new IMAGE_BRUSH("Slate/Bone_16x", Icon16x16));
 			Set("ControlRig.Tree.Control", new IMAGE_BRUSH("Slate/RigControlCircle_16x", Icon16x16));
+			Set("ControlRig.Tree.ProxyControl", new IMAGE_BRUSH("Slate/ProxyControl1_16x", Icon16x16));
 			Set("ControlRig.Tree.Null", new IMAGE_BRUSH("Slate/Null_16x", Icon16x16));
 			Set("ControlRig.Tree.RigidBody", new IMAGE_BRUSH("Slate/RigidBody_16x", Icon16x16));
 			Set("ControlRig.Tree.Socket", new IMAGE_BRUSH("Slate/Socket_16x", Icon16x16));
diff --git a/Engine/Plugins/Animation/ControlRig/Source/ControlRigEditor/Private/Editor/SRigHierarchyTreeView.cpp b/Engine/Plugins/Animation/ControlRig/Source/ControlRigEditor/Private/Editor/SRigHierarchyTreeView.cpp
--- a/Engine/Plugins/Animation/ControlRig/Source/ControlRigEditor/Private/Editor/SRigHierarchyTreeView.cpp
+++ b/Engine/Plugins/Animation/ControlRig/Source/ControlRigEditor/Private/Editor/SRigHierarchyTreeView.cpp
@@ -722,13 +722,20 @@ TPair<const FSlateBrush*, FSlateColor> SRigHierarchyItem::GetBrushForElementType
 	{
 		case ERigElementType::Control:
 		{
-			Brush = FControlRigEditorStyle::Get().GetBrush("ControlRig.Tree.Control");
 			if(const FRigControlElement* Control = InHierarchy->Find<FRigControlElement>(InKey))
 			{
 				FLinearColor ShapeColor = FLinearColor::White;
 				
 				if(Control->Settings.SupportsShape())
 				{
+					if(Control->Settings.AnimationType == ERigControlAnimationType::ProxyControl)
+					{
+						Brush = FControlRigEditorStyle::Get().GetBrush("ControlRig.Tree.ProxyControl");
+					}
+					else
+					{
+						Brush = FControlRigEditorStyle::Get().GetBrush("ControlRig.Tree.Control");
+					}
 					ShapeColor = Control->Settings.ShapeColor;
 				}
 				else
@@ -742,6 +749,10 @@ TPair<const FSlateBrush*, FSlateColor> SRigHierarchyItem::GetBrushForElementType
 				ShapeColor.A = 1.f;
 				Color = FSlateColor(ShapeColor);
 			}
+			else
+			{
+				Brush = FControlRigEditorStyle::Get().GetBrush("ControlRig.Tree.Control");
+			}
 			break;
 		}
 		case ERigElementType::Null:
