Control Rig: Improve metadata customwidget

diff --git a/Engine/Plugins/Animation/ControlRig/Source/ControlRig/Private/Rigs/RigHierarchy.cpp b/Engine/Plugins/Animation/ControlRig/Source/ControlRig/Private/Rigs/RigHierarchy.cpp
--- a/Engine/Plugins/Animation/ControlRig/Source/ControlRig/Private/Rigs/RigHierarchy.cpp
+++ b/Engine/Plugins/Animation/ControlRig/Source/ControlRig/Private/Rigs/RigHierarchy.cpp
@@ -95,6 +95,7 @@ bool URigHierarchy::bEnableValidityCheckbyDefault = false;
 URigHierarchy::URigHierarchy()
 : TopologyVersion(0)
 , MetadataVersion(0)
+, MetadataTagVersion(0)
 , bEnableDirtyPropagation(true)
 , Elements()
 , IndexLookup()
@@ -4466,7 +4467,7 @@ void URigHierarchy::CleanupInvalidCaches()
 
 void URigHierarchy::OnMetadataChanged(const FRigElementKey& InKey, const FName& InName)
 {
-	IncrementMetadataVersion();
+	IncrementMetadataVersion(InKey, InName);
 	if(MetadataChangedDelegate.IsBound())
 	{
 		MetadataChangedDelegate.Broadcast(InKey, InName);
@@ -4475,6 +4476,7 @@ void URigHierarchy::OnMetadataChanged(const FRigElementKey& InKey, const FName&
 
 void URigHierarchy::OnMetadataTagChanged(const FRigElementKey& InKey, const FName& InTag, bool bAdded)
 {
+	IncrementMetadataTagVersion(InKey, InTag, bAdded);
 	if(MetadataTagChangedDelegate.IsBound())
 	{
 		MetadataTagChangedDelegate.Broadcast(InKey, InTag, bAdded);
diff --git a/Engine/Plugins/Animation/ControlRig/Source/ControlRig/Private/Units/Hierarchy/RigUnit_Metadata.cpp b/Engine/Plugins/Animation/ControlRig/Source/ControlRig/Private/Units/Hierarchy/RigUnit_Metadata.cpp
--- a/Engine/Plugins/Animation/ControlRig/Source/ControlRig/Private/Units/Hierarchy/RigUnit_Metadata.cpp
+++ b/Engine/Plugins/Animation/ControlRig/Source/ControlRig/Private/Units/Hierarchy/RigUnit_Metadata.cpp
@@ -109,6 +109,18 @@ FText FRigDispatch_MetadataBase::GetArgumentTooltip(const FName& InArgumentName,
 	return FRigDispatchFactory::GetArgumentTooltip(InArgumentName, InTypeIndex);
 }
 
+FString FRigDispatch_MetadataBase::GetArgumentMetaData(const FName& InArgumentName, const FName& InMetaDataKey) const
+{
+	if(InArgumentName == TEXT("Name"))
+	{
+		if(InMetaDataKey == FRigVMStruct::CustomWidgetMetaName)
+		{
+			return TEXT("MetadataName");
+		}
+	}
+	return Super::GetArgumentMetaData(InArgumentName, InMetaDataKey);
+}
+
 #endif
 
 const TArray<TRigVMTypeIndex>& FRigDispatch_MetadataBase::GetValueTypes() const
@@ -157,22 +169,6 @@ TArray<FRigVMTemplateArgument> FRigDispatch_GetMetadata::GetArguments() const
 	return Arguments;
 }
 
-#if WITH_EDITOR
-
-FString FRigDispatch_GetMetadata::GetArgumentMetaData(const FName& InArgumentName, const FName& InMetaDataKey) const
-{
-	if(InArgumentName == TEXT("Name"))
-	{
-		if(InMetaDataKey == FRigVMStruct::CustomWidgetMetaName)
-		{
-			return TEXT("MetadataName");
-		}
-	}
-	return FRigDispatch_MetadataBase::GetArgumentMetaData(InArgumentName, InMetaDataKey);
-}
-
-#endif
-
 FRigBaseMetadata* FRigDispatch_GetMetadata::FindMetadata(const FRigVMExtendedExecuteContext& InContext,
                                                          const FRigElementKey& InKey, const FName& InName, ERigMetadataType InType, FCachedRigElement& Cache)
 {
diff --git a/Engine/Plugins/Animation/ControlRig/Source/ControlRig/Private/Units/Hierarchy/RigUnit_Metadata.h b/Engine/Plugins/Animation/ControlRig/Source/ControlRig/Private/Units/Hierarchy/RigUnit_Metadata.h
--- a/Engine/Plugins/Animation/ControlRig/Source/ControlRig/Private/Units/Hierarchy/RigUnit_Metadata.h
+++ b/Engine/Plugins/Animation/ControlRig/Source/ControlRig/Private/Units/Hierarchy/RigUnit_Metadata.h
@@ -17,6 +17,7 @@ struct CONTROLRIG_API FRigDispatch_MetadataBase : public FRigDispatchFactory
 	virtual bool IsSetMetadata() const { return false; }
 
 #if WITH_EDITOR
+	virtual FString GetArgumentMetaData(const FName& InArgumentName, const FName& InMetaDataKey) const override;
 	virtual FText GetArgumentTooltip(const FName& InArgumentName, TRigVMTypeIndex InTypeIndex) const override;
 #endif
 
@@ -53,10 +54,6 @@ struct CONTROLRIG_API FRigDispatch_GetMetadata : public FRigDispatch_MetadataBas
 
 	virtual TArray<FRigVMTemplateArgument> GetArguments() const override;
 
-#if WITH_EDITOR
-	virtual FString GetArgumentMetaData(const FName& InArgumentName, const FName& InMetaDataKey) const override;
-#endif
-
 protected:
 
 	static FRigBaseMetadata* FindMetadata(const FRigVMExtendedExecuteContext& InContext, const FRigElementKey& InKey, const FName& InName, ERigMetadataType InType, FCachedRigElement& Cache);
@@ -308,7 +305,7 @@ struct CONTROLRIG_API FRigUnit_FindItemsWithMetadata : public FRigUnit
 	/**
 	 * The name of the metadata to find
 	 */ 
-	UPROPERTY(meta = (Input, CustomWidget="MetadataNameAll"))
+	UPROPERTY(meta = (Input, CustomWidget="MetadataName"))
 	FName Name;
 
 	/**
@@ -382,7 +379,7 @@ struct CONTROLRIG_API FRigUnit_SetMetadataTag : public FRigUnitMutable
 	/**
 	 * The name of the tag to set
 	 */ 
-	UPROPERTY(meta = (Input))
+	UPROPERTY(meta = (Input, CustomWidget="MetadataTagName"))
 	FName Tag;
 
 	// Used to cache the internally
@@ -416,7 +413,7 @@ struct CONTROLRIG_API FRigUnit_SetMetadataTagArray : public FRigUnitMutable
 	/**
 	 * The tags to set for the item
 	 */ 
-	UPROPERTY(meta = (Input))
+	UPROPERTY(meta = (Input, CustomWidget="MetadataTagName"))
 	TArray<FName> Tags;
 
 	// Used to cache the internally
@@ -450,7 +447,7 @@ struct CONTROLRIG_API FRigUnit_RemoveMetadataTag : public FRigUnitMutable
 	/**
 	 * The name of the tag to set
 	 */ 
-	UPROPERTY(meta = (Input))
+	UPROPERTY(meta = (Input, CustomWidget="MetadataTagName"))
 	FName Tag;
 
 	/**
@@ -561,7 +558,7 @@ struct CONTROLRIG_API FRigUnit_FindItemsWithMetadataTag : public FRigUnit
 	/**
 	 * The name of the tag to find
 	 */ 
-	UPROPERTY(meta = (Input, CustomWidget="MetadataTagNameAll"))
+	UPROPERTY(meta = (Input, CustomWidget="MetadataTagName"))
 	FName Tag;
 
 	// The items containing the metadata
@@ -588,7 +585,7 @@ struct CONTROLRIG_API FRigUnit_FindItemsWithMetadataTagArray : public FRigUnit
 	/**
 	 * The tags to find
 	 */ 
-	UPROPERTY(meta = (Input, CustomWidget="MetadataTagNameAll"))
+	UPROPERTY(meta = (Input, CustomWidget="MetadataTagName"))
 	TArray<FName> Tags;
 
 	// The items containing the metadata
@@ -621,7 +618,7 @@ struct CONTROLRIG_API FRigUnit_FilterItemsByMetadataTags : public FRigUnit
 	/**
 	 * The tags to find
 	 */ 
-	UPROPERTY(meta = (Input, CustomWidget="MetadataTagNameAll"))
+	UPROPERTY(meta = (Input, CustomWidget="MetadataTagName"))
 	TArray<FName> Tags;
 
 	/**
diff --git a/Engine/Plugins/Animation/ControlRig/Source/ControlRig/Public/Rigs/RigHierarchy.h b/Engine/Plugins/Animation/ControlRig/Source/ControlRig/Public/Rigs/RigHierarchy.h
--- a/Engine/Plugins/Animation/ControlRig/Source/ControlRig/Public/Rigs/RigHierarchy.h
+++ b/Engine/Plugins/Animation/ControlRig/Source/ControlRig/Public/Rigs/RigHierarchy.h
@@ -2888,7 +2888,23 @@ public:
 	/**
 	 * Increments the metadata version
 	 */
-	FORCEINLINE void IncrementMetadataVersion() { MetadataVersion++; }
+	FORCEINLINE void IncrementMetadataVersion(const FRigElementKey& InKey, const FName& InName)
+	{
+		MetadataVersion += 1 + (int32)GetTypeHash(InName);
+	}
+
+	/**
+     * Returns the metadata tag version of this hierarchy
+	 */
+	uint16 GetMetadataTagVersion() const { return MetadataTagVersion; }
+
+	/**
+	 * Increments the metadataTag version
+	 */
+	FORCEINLINE void IncrementMetadataTagVersion(const FRigElementKey& InKey, const FName& InTag, bool bAdded)
+	{
+		MetadataTagVersion += 1 + (int32)GetTypeHash(InTag);
+	}
 
 	/**
 	 * Returns the current / initial pose of the hierarchy
@@ -3519,6 +3535,13 @@ private:
 	UPROPERTY(transient)
 	uint16 MetadataVersion;
 
+	/**
+	 * The metadata version of the hierarchy changes when metadata is being
+	 * created or removed (not when the metadata values changes)
+	 */
+	UPROPERTY(transient)
+	uint16 MetadataTagVersion;
+
 	/**
 	 * If set to false the dirty flag propagation will be disabled
 	 */
diff --git a/Engine/Plugins/Animation/ControlRig/Source/ControlRigEditor/Private/Graph/ControlRigGraphPanelPinFactory.cpp b/Engine/Plugins/Animation/ControlRig/Source/ControlRigEditor/Private/Graph/ControlRigGraphPanelPinFactory.cpp
--- a/Engine/Plugins/Animation/ControlRig/Source/ControlRigEditor/Private/Graph/ControlRigGraphPanelPinFactory.cpp
+++ b/Engine/Plugins/Animation/ControlRig/Source/ControlRigEditor/Private/Graph/ControlRigGraphPanelPinFactory.cpp
@@ -204,7 +204,7 @@ TSharedPtr<SGraphPin> FControlRigGraphPanelPinFactory::CreatePin(UEdGraphPin* In
 							return &EmptyNameList;
 						});
 				}
-				else if (CustomWidgetName == TEXT("MetadataName") || CustomWidgetName == TEXT("MetadataNameAll"))
+				else if (CustomWidgetName == TEXT("MetadataName"))
 				{
 					struct FCachedMetadataNames
 					{
@@ -216,61 +216,19 @@ TSharedPtr<SGraphPin> FControlRigGraphPanelPinFactory::CreatePin(UEdGraphPin* In
 						{}
 					};
 
-					const bool bCompleteList = CustomWidgetName == TEXT("MetadataNameAll");
-					
 					return SNew(SControlRigGraphPinNameList, InPin)
 						.ModelPin(ModelPin)
-						.OnGetNameListContent_Lambda([RigGraph, bCompleteList](const URigVMPin* InPin)
+						.SearchHintText(NSLOCTEXT("FControlRigGraphPanelPinFactory", "MetadataName", "Metadata Name"))
+						.AllowUserProvidedText(true)
+						.EnableNameListCache(false)
+						.OnGetNameListContent_Lambda([RigGraph](const URigVMPin* InPin)
 						{
 							if (const UControlRigBlueprint* Blueprint = RigGraph->GetTypedOuter<UControlRigBlueprint>())
 							{
 								if(UControlRig* ControlRig = Cast<UControlRig>(Blueprint->GetObjectBeingDebugged()))
 								{
-									FRigElementKey ItemKey;
-									bool bHasSourceLinks = false;
-
-									// find the pin that holds the item
-									for(URigVMPin* Pin : InPin->GetRootPin()->GetNode()->GetPins())
-									{
-										if(Pin->GetCPPType() == FRigElementKey::StaticStruct()->GetStructCPPName())
-										{
-											if(!bCompleteList)
-											{
-												if(Pin->GetSourceLinks(true).Num() > 0)
-												{
-													bHasSourceLinks = true;
-												}
-											}
-											if(!bHasSourceLinks)
-											{
-												const FString DefaultValue = Pin->GetDefaultValue();
-												if(!DefaultValue.IsEmpty())
-												{
-													FRigElementKey::StaticStruct()->ImportText(*DefaultValue, &ItemKey, nullptr, EPropertyPortFlags::PPF_None, nullptr, FRigElementKey::StaticStruct()->GetName(), true);
-												}
-											}
-											break;
-										}
-									}
-
-									const FRigBaseElement* Element = nullptr;
-									if(ItemKey.IsValid())
-									{
-										if(!bCompleteList && !bHasSourceLinks)
-										{
-											Element = ControlRig->GetHierarchy()->Find(ItemKey);
-											if(Element == nullptr)
-											{
-												static TArray<TSharedPtr<FString>> EmptyNameList;
-												return &EmptyNameList;
-											}
-										}
-									}
-
-									const FString MapHash = Element ?
-										Blueprint->GetPathName() + TEXT("|") + ItemKey.ToString() :
-										Blueprint->GetPathName();
-									const int32 MetadataVersion = Element ? Element->GetMetadataVersion() : ControlRig->GetHierarchy()->GetMetadataVersion(); 
+									const FString MapHash = Blueprint->GetPathName();
+									const int32 MetadataVersion = ControlRig->GetHierarchy()->GetMetadataVersion(); 
 
 									static TMap<FString, FCachedMetadataNames> MetadataNameLists;
 									FCachedMetadataNames& MetadataNames = MetadataNameLists.FindOrAdd(MapHash);
@@ -278,22 +236,12 @@ TSharedPtr<SGraphPin> FControlRigGraphPanelPinFactory::CreatePin(UEdGraphPin* In
 									if(MetadataNames.MetadataVersion != MetadataVersion)
 									{
 										TArray<FName> Names;
-										if(bCompleteList || bHasSourceLinks || (Element == nullptr))
-										{
-											for(int32 ElementIndex=0; ElementIndex < ControlRig->GetHierarchy()->Num(); ElementIndex++)
-											{
-												const FRigBaseElement* OtherElement = ControlRig->GetHierarchy()->Get(ElementIndex);
-												for(int32 MetadataIndex = 0; MetadataIndex < OtherElement->NumMetadata(); MetadataIndex++)
-												{
-													Names.AddUnique(OtherElement->GetMetadata(MetadataIndex)->GetName());
-												}
-											}
-										}
-										else if(Element)
+										for(int32 ElementIndex=0; ElementIndex < ControlRig->GetHierarchy()->Num(); ElementIndex++)
 										{
-											for(int32 MetadataIndex = 0; MetadataIndex < Element->NumMetadata(); MetadataIndex++)
+											const FRigBaseElement* OtherElement = ControlRig->GetHierarchy()->Get(ElementIndex);
+											for(int32 MetadataIndex = 0; MetadataIndex < OtherElement->NumMetadata(); MetadataIndex++)
 											{
-												Names.Add(Element->GetMetadata(MetadataIndex)->GetName());
+												Names.AddUnique(OtherElement->GetMetadata(MetadataIndex)->GetName());
 											}
 										}
 
@@ -324,62 +272,37 @@ TSharedPtr<SGraphPin> FControlRigGraphPanelPinFactory::CreatePin(UEdGraphPin* In
 							return &EmptyNameList;
 						});
 				}
-				else if (CustomWidgetName == TEXT("MetadataTagName") || CustomWidgetName == TEXT("MetadataTagNameAll"))
+				else if (CustomWidgetName == TEXT("MetadataTagName"))
 				{
-					bool bCompleteList = CustomWidgetName == TEXT("MetadataTagNameAll");
-
 					struct FCachedMetadataTagNames
 					{
+						int32 MetadataTagVersion;
 						TSharedPtr<TArray<TSharedPtr<FString>>> Names;
 						FCachedMetadataTagNames()
+						: MetadataTagVersion(INDEX_NONE)
 						{}
 					};
 
 					return SNew(SControlRigGraphPinNameList, InPin)
 						.ModelPin(ModelPin)
-						.OnGetNameListContent_Lambda([RigGraph, bCompleteList](const URigVMPin* InPin)
+						.SearchHintText(NSLOCTEXT("FControlRigGraphPanelPinFactory", "TagName", "Tag Name"))
+						.AllowUserProvidedText(true)
+						.EnableNameListCache(false)
+						.OnGetNameListContent_Lambda([RigGraph](const URigVMPin* InPin)
 						{
 							if (const UControlRigBlueprint* Blueprint = RigGraph->GetTypedOuter<UControlRigBlueprint>())
 							{
 								if(UControlRig* ControlRig = Cast<UControlRig>(Blueprint->GetObjectBeingDebugged()))
 								{
-									FRigElementKey ItemKey;
-									bool bHasSourceLinks = false;
-
-									// find the pin that holds the item
-									for(URigVMPin* Pin : InPin->GetRootPin()->GetNode()->GetPins())
-									{
-										if(Pin->GetCPPType() == FRigElementKey::StaticStruct()->GetStructCPPName())
-										{
-											if(!bCompleteList)
-											{
-												if(Pin->GetSourceLinks(true).Num() > 0)
-												{
-													bHasSourceLinks = true;
-												}
-											}
-											if(!bHasSourceLinks)
-											{
-												const FString DefaultValue = Pin->GetDefaultValue();
-												if(!DefaultValue.IsEmpty())
-												{
-													FRigElementKey::StaticStruct()->ImportText(*DefaultValue, &ItemKey, nullptr, EPropertyPortFlags::PPF_None, nullptr, FRigElementKey::StaticStruct()->GetName(), true);
-												}
-											}
-											break;
-										}
-									}
-
-									const FString MapHash = ItemKey.IsValid() ?
-										Blueprint->GetPathName() + TEXT("|") + ItemKey.ToString() :
-										Blueprint->GetPathName();
+									const FString MapHash = Blueprint->GetPathName();
+									const int32 MetadataTagVersion = ControlRig->GetHierarchy()->GetMetadataTagVersion(); 
 
 									static TMap<FString, FCachedMetadataTagNames> MetadataTagNameLists;
 									FCachedMetadataTagNames& MetadataTagNames = MetadataTagNameLists.FindOrAdd(MapHash);
 
-									TArray<FName> Tags; 
-									if(bCompleteList || !ItemKey.IsValid())
+									if(MetadataTagNames.MetadataTagVersion != MetadataTagVersion)
 									{
+										TArray<FName> Tags; 
 										for(int32 ElementIndex=0; ElementIndex < ControlRig->GetHierarchy()->Num(); ElementIndex++)
 										{
 											const FRigBaseElement* Element = ControlRig->GetHierarchy()->Get(ElementIndex);
@@ -391,33 +314,25 @@ TSharedPtr<SGraphPin> FControlRigGraphPanelPinFactory::CreatePin(UEdGraphPin* In
 												}
 											}
 										}
-									}
-									else
-									{
-										if(const FRigBaseElement* Element = ControlRig->GetHierarchy()->Find(ItemKey))
+
+										if(!MetadataTagNames.Names.IsValid())
 										{
-											if(const FRigNameArrayMetadata* Md = Cast<FRigNameArrayMetadata>(Element->GetMetadata(URigHierarchy::TagMetadataName, ERigMetadataType::NameArray)))
-											{
-												Tags = Md->GetValue();
-											}
+											MetadataTagNames.Names = MakeShareable(new TArray<TSharedPtr<FString>>());
 										}
-									}
+										MetadataTagNames.Names->Reset();
 
-									if(!MetadataTagNames.Names.IsValid())
-									{
-										MetadataTagNames.Names = MakeShareable(new TArray<TSharedPtr<FString>>());
-									}
-									MetadataTagNames.Names->Reset();
+										for(const FName& Tag : Tags)
+										{
+											MetadataTagNames.Names->Add(MakeShareable(new FString(Tag.ToString())));
+										}
+										MetadataTagNames.Names->Sort([](const TSharedPtr<FString>& A, const TSharedPtr<FString>& B)
+										{
+											return A.Get() > B.Get();
+										});
+										MetadataTagNames.Names->Insert(MakeShareable(new FString(FName(NAME_None).ToString())), 0);
 
-									for(const FName& Tag : Tags)
-									{
-										MetadataTagNames.Names->Add(MakeShareable(new FString(Tag.ToString())));
+										MetadataTagNames.MetadataTagVersion = MetadataTagVersion;
 									}
-									MetadataTagNames.Names->Sort([](const TSharedPtr<FString>& A, const TSharedPtr<FString>& B)
-									{
-										return A.Get() > B.Get();
-									});
-									MetadataTagNames.Names->Insert(MakeShareable(new FString(FName(NAME_None).ToString())), 0);
 
 									return MetadataTagNames.Names.Get();
 								}
diff --git a/Engine/Plugins/Animation/ControlRig/Source/ControlRigEditor/Private/Graph/SControlRigGraphPinNameList.cpp b/Engine/Plugins/Animation/ControlRig/Source/ControlRigEditor/Private/Graph/SControlRigGraphPinNameList.cpp
--- a/Engine/Plugins/Animation/ControlRig/Source/ControlRigEditor/Private/Graph/SControlRigGraphPinNameList.cpp
+++ b/Engine/Plugins/Animation/ControlRig/Source/ControlRigEditor/Private/Graph/SControlRigGraphPinNameList.cpp
@@ -31,6 +31,9 @@ void SControlRigGraphPinNameList::Construct(const FArguments& InArgs, UEdGraphPi
 	this->OnGetNameListContentForValidation = InArgs._OnGetNameListContentForValidation;
 	this->OnGetNameFromSelection = InArgs._OnGetNameFromSelection;
 	this->bMarkupInvalidItems = InArgs._MarkupInvalidItems;
+	this->EnableNameListCache = InArgs._EnableNameListCache;
+	this->SearchHintText = InArgs._SearchHintText;
+	this->AllowUserProvidedText = InArgs._AllowUserProvidedText;
 
 	UpdateNameLists();
 	SGraphPin::Construct(SGraphPin::FArguments(), InGraphPinObj);
@@ -61,6 +64,8 @@ TSharedRef<SWidget>	SControlRigGraphPinNameList::GetDefaultValueWidget()
 					.OnSelectionChanged(this, &SControlRigGraphPinNameList::OnNameListChanged)
 					.OnComboBoxOpening(this, &SControlRigGraphPinNameList::OnNameListComboBox)
 					.InitiallySelectedItem(InitialSelected)
+					.SearchHintText(SearchHintText)
+					.AllowUserProvidedText(AllowUserProvidedText)
 					.Content()
 					[
 						SNew(STextBlock)
@@ -159,10 +164,16 @@ void SControlRigGraphPinNameList::SetNameListText(const FText& NewTypeInValue, E
 	}
 }
 
-void SControlRigGraphPinNameList::UpdateNameLists()
+void SControlRigGraphPinNameList::UpdateNameLists(bool bUpdateCurrent, bool bUpdateValidation)
 {
-	CurrentList = GetNameList(true);
-	ValidationList = GetNameList(false);
+	if(CurrentList == nullptr || bUpdateCurrent)
+	{
+		CurrentList = GetNameList(true);
+	}
+	if(ValidationList == nullptr || bUpdateValidation)
+	{
+		ValidationList = GetNameList(false);
+	}
 }
 
 FSlateColor SControlRigGraphPinNameList::GetNameColor() const
@@ -170,7 +181,12 @@ FSlateColor SControlRigGraphPinNameList::GetNameColor() const
 	if(bMarkupInvalidItems)
 	{
 		FString CurrentItem = GetNameListText().ToString();
-		
+
+		if(!EnableNameListCache)
+		{
+			((SControlRigGraphPinNameList*)this)->UpdateNameLists(false, true);
+		}
+
 		bool bFound = false;
 		for (TSharedPtr<FString> Item : (*ValidationList))
 		{
@@ -221,7 +237,10 @@ void SControlRigGraphPinNameList::OnNameListComboBox()
 		}
 	}
 	NameListComboBox->SetOptionsSource(CurrentList);
-	NameListComboBox->SetSelectedItem(CurrentlySelected);
+	if(CurrentlySelected.IsValid())
+	{
+		NameListComboBox->SetSelectedItem(CurrentlySelected);
+	}
 }
 
 FSlateColor SControlRigGraphPinNameList::OnGetWidgetForeground() const
diff --git a/Engine/Plugins/Animation/ControlRig/Source/ControlRigEditor/Private/Graph/SControlRigGraphPinNameList.h b/Engine/Plugins/Animation/ControlRig/Source/ControlRigEditor/Private/Graph/SControlRigGraphPinNameList.h
--- a/Engine/Plugins/Animation/ControlRig/Source/ControlRigEditor/Private/Graph/SControlRigGraphPinNameList.h
+++ b/Engine/Plugins/Animation/ControlRig/Source/ControlRigEditor/Private/Graph/SControlRigGraphPinNameList.h
@@ -18,6 +18,9 @@ public:
 
 	SLATE_BEGIN_ARGS(SControlRigGraphPinNameList)
 		: _MarkupInvalidItems(true)
+		, _EnableNameListCache(true)
+		, _SearchHintText(NSLOCTEXT("SControlRigGraphPinNameList", "Search", "Search"))
+		, _AllowUserProvidedText(false)
 	{}
 
 		SLATE_ARGUMENT(URigVMPin*, ModelPin)
@@ -25,6 +28,9 @@ public:
 		SLATE_EVENT(FOnGetNameListContent, OnGetNameListContentForValidation)
 		SLATE_EVENT(FOnGetNameFromSelection, OnGetNameFromSelection)
 		SLATE_ARGUMENT(bool, MarkupInvalidItems)
+		SLATE_ARGUMENT(bool, EnableNameListCache)
+		SLATE_ARGUMENT(FText, SearchHintText)
+		SLATE_ARGUMENT(bool, AllowUserProvidedText)
 
 	SLATE_END_ARGS()
 
@@ -39,7 +45,7 @@ protected:
 	FText GetNameListText() const;
 	FSlateColor GetNameColor() const;
 	virtual void SetNameListText(const FText& NewTypeInValue, ETextCommit::Type CommitInfo);
-	void UpdateNameLists();
+	void UpdateNameLists(bool bUpdateCurrent = true, bool bUpdateValidation = true);
 
 	TSharedRef<SWidget> MakeNameListItemWidget(TSharedPtr<FString> InItem);
 	void OnNameListChanged(TSharedPtr<FString> NewSelection, ESelectInfo::Type SelectInfo);
@@ -53,6 +59,9 @@ protected:
 	const TArray<TSharedPtr<FString>>* CurrentList;
 	const TArray<TSharedPtr<FString>>* ValidationList;
 	bool bMarkupInvalidItems;
+	bool EnableNameListCache;
+	FText SearchHintText;
+	bool AllowUserProvidedText;
 
 	/** Helper buttons. */
 	FSlateColor OnGetWidgetForeground() const;
diff --git a/Engine/Plugins/Animation/ControlRig/Source/ControlRigEditor/Private/Graph/SControlRigGraphPinNameListValueWidget.cpp b/Engine/Plugins/Animation/ControlRig/Source/ControlRigEditor/Private/Graph/SControlRigGraphPinNameListValueWidget.cpp
--- a/Engine/Plugins/Animation/ControlRig/Source/ControlRigEditor/Private/Graph/SControlRigGraphPinNameListValueWidget.cpp
+++ b/Engine/Plugins/Animation/ControlRig/Source/ControlRigEditor/Private/Graph/SControlRigGraphPinNameListValueWidget.cpp
@@ -12,6 +12,7 @@ void SControlRigGraphPinNameListValueWidget::Construct(const FArguments& InArgs)
 	this->OnComboBoxOpening = InArgs._OnComboBoxOpening;
 	this->OnSelectionChanged = InArgs._OnSelectionChanged;
 	this->OnGenerateWidget = InArgs._OnGenerateWidget;
+	this->AllowUserProvidedText = InArgs._AllowUserProvidedText;
 
 	OptionsSource = InArgs._OptionsSource;
 	CustomScrollbar = InArgs._CustomScrollbar;
@@ -27,7 +28,7 @@ void SControlRigGraphPinNameListValueWidget::Construct(const FArguments& InArgs)
 			.AutoHeight()
 			[
 				SAssignNew(this->SearchField, SEditableTextBox)
-				.HintText(LOCTEXT("Search", "Search"))
+				.HintText(InArgs._SearchHintText)
 				.OnTextChanged(this, &SControlRigGraphPinNameListValueWidget::OnSearchTextChanged)
 				.OnTextCommitted(this, &SControlRigGraphPinNameListValueWidget::OnSearchTextCommitted)
 				.OnKeyDownHandler(this, &SControlRigGraphPinNameListValueWidget::OnSearchTextKeyDown)
@@ -201,7 +202,7 @@ void SControlRigGraphPinNameListValueWidget::OnSelectionChanged_Internal(TShared
 	if (SelectInfo != ESelectInfo::OnNavigation || bForce)
 	{
 		// Ensure that the proposed selection is different from selected
-		if (ProposedSelection != SelectedItem || bForce)
+		if ((ProposedSelection.IsValid() && (ProposedSelection != SelectedItem)) || bForce)
 		{
 			SelectedItem = ProposedSelection;
 			OnSelectionChanged.ExecuteIfBound(ProposedSelection, ESelectInfo::OnMouseClick); // SelectInfo);
@@ -254,15 +255,22 @@ void SControlRigGraphPinNameListValueWidget::OnSearchTextCommitted(const FText&
 {
 	if (CommitType == ETextCommit::OnEnter)
 	{
-		FString ProposedSelection = ChangedText.ToString().ToLower();
+		const FString LowerCaseProposedSelection = ChangedText.ToString().ToLower();
 		for (int32 i = 0; i < OptionsSource->Num(); i++)
 		{
 			const TSharedPtr<FString>& Option = (*OptionsSource)[i];
-			if (Option->ToLower() == ProposedSelection)
+			if (Option->ToLower() == LowerCaseProposedSelection)
 			{
 				OnSelectionChanged_Internal(Option, ESelectInfo::OnKeyPress, true);
+				return;
 			}
 		}
+
+		if(AllowUserProvidedText)
+		{
+			const TSharedPtr<FString> UserDefinedOption = MakeShareable(new FString(ChangedText.ToString()));
+			OnSelectionChanged_Internal(UserDefinedOption, ESelectInfo::OnKeyPress, true);
+		}
 	}
 }
 
diff --git a/Engine/Plugins/Animation/ControlRig/Source/ControlRigEditor/Private/Graph/SControlRigGraphPinNameListValueWidget.h b/Engine/Plugins/Animation/ControlRig/Source/ControlRigEditor/Private/Graph/SControlRigGraphPinNameListValueWidget.h
--- a/Engine/Plugins/Animation/ControlRig/Source/ControlRigEditor/Private/Graph/SControlRigGraphPinNameListValueWidget.h
+++ b/Engine/Plugins/Animation/ControlRig/Source/ControlRigEditor/Private/Graph/SControlRigGraphPinNameListValueWidget.h
@@ -32,6 +32,8 @@ public:
 		, _Method()
 		, _MaxListHeight(450.0f)
 		, _HasDownArrow(true)
+		, _SearchHintText(NSLOCTEXT("SControlRigGraphPinNameListValueWidget", "Search", "Search"))
+		, _AllowUserProvidedText(false)
 	{}
 
 	/** Slot for this button's content (optional) */
@@ -63,6 +65,16 @@ public:
 		 */
 		SLATE_ARGUMENT(bool, HasDownArrow)
 
+		/*
+		 * The visible text in the search / editable text field
+		 */
+		SLATE_ARGUMENT(FText, SearchHintText)
+
+		/*
+		 * If set to true the user is allowed to enter custom text here
+		 */
+		SLATE_ARGUMENT(bool, AllowUserProvidedText)
+
 	SLATE_END_ARGS()
 
 	/**
@@ -133,6 +145,7 @@ private:
 	FOnGenerateWidget OnGenerateWidget;
 
 	const TArray< TSharedPtr<FString> >* OptionsSource;
+	bool AllowUserProvidedText;
 
 	/** Used to focus the name box immediately following construction */
 	EActiveTimerReturnType SetFocusPostConstruct(double InCurrentTime, float InDeltaTime);
