Sequencer: Control Rig: When finding section to key space channel just search all bindings since the CR track may not live under the skel mesh binding

diff --git a/Engine/Plugins/Animation/ControlRig/Source/ControlRigEditor/Private/Sequencer/ControlRigSpaceChannelEditors.cpp b/Engine/Plugins/Animation/ControlRig/Source/ControlRigEditor/Private/Sequencer/ControlRigSpaceChannelEditors.cpp
--- a/Engine/Plugins/Animation/ControlRig/Source/ControlRigEditor/Private/Sequencer/ControlRigSpaceChannelEditors.cpp
+++ b/Engine/Plugins/Animation/ControlRig/Source/ControlRigEditor/Private/Sequencer/ControlRigSpaceChannelEditors.cpp
@@ -190,86 +190,62 @@ FSpaceChannelAndSection FControlRigSpaceChannelHelpers::FindSpaceChannelAndSecti
 	{
 		return SpaceChannelAndSection;
 	}
-	if (TSharedPtr<IControlRigObjectBinding> ObjectBinding = ControlRig->GetObjectBinding())
+	UMovieScene* MovieScene = Sequencer->GetFocusedMovieSceneSequence()->GetMovieScene();
+	if (!MovieScene)
 	{
-		USceneComponent* Component = Cast<USceneComponent>(ObjectBinding->GetBoundObject());
-		if (!Component)
-		{
-			return SpaceChannelAndSection;
-		}
-		const bool bCreateHandleIfMissing = false;
-		FName CreatedFolderName = NAME_None;
-		FGuid ObjectHandle = Sequencer->GetHandleToObject(Component, bCreateHandleIfMissing);
-		if (!ObjectHandle.IsValid())
-		{
-			UObject* ActorObject = Component->GetOwner();
-			ObjectHandle = Sequencer->GetHandleToObject(ActorObject, bCreateHandleIfMissing);
-			if (!ObjectHandle.IsValid())
-			{
-				return SpaceChannelAndSection;
-			}
-		}
-		bool bCreateTrack = false;
-		UMovieScene* MovieScene = Sequencer->GetFocusedMovieSceneSequence()->GetMovieScene();
-		if (!MovieScene)
-		{
-			return SpaceChannelAndSection;
-		}
-		bool bRecreateCurves = false;
-		TArray<TPair<UControlRig*, FName>> ControlRigPairsToReselect;
-		if (FMovieSceneBinding* Binding = MovieScene->FindBinding(ObjectHandle))
+		return SpaceChannelAndSection;
+	}
+	const TArray<FMovieSceneBinding>& Bindings = MovieScene->GetBindings();
+	bool bRecreateCurves = false;
+	TArray<TPair<UControlRig*, FName>> ControlRigPairsToReselect;
+	for (const FMovieSceneBinding& Binding : Bindings)
+	{
+		UMovieSceneControlRigParameterTrack* ControlRigParameterTrack = Cast<UMovieSceneControlRigParameterTrack>(MovieScene->FindTrack(UMovieSceneControlRigParameterTrack::StaticClass(), Binding.GetObjectGuid(), NAME_None));
+		if (ControlRigParameterTrack && ControlRigParameterTrack->GetControlRig() == ControlRig)
 		{
-			for (UMovieSceneTrack* Track : Binding->GetTracks())
+			UMovieSceneControlRigParameterSection* ActiveSection = Cast<UMovieSceneControlRigParameterSection>(ControlRigParameterTrack->GetSectionToKey());
+			if (ActiveSection)
 			{
-				if (UMovieSceneControlRigParameterTrack* ControlRigParameterTrack = Cast<UMovieSceneControlRigParameterTrack>(Track))
+				ActiveSection->Modify();
+				ControlRig->Modify();
+				SpaceChannelAndSection.SectionToKey = ActiveSection;
+				FSpaceControlNameAndChannel* NameAndChannel = ActiveSection->GetSpaceChannel(ControlName);
+				if (NameAndChannel)
 				{
-					if (ControlRigParameterTrack->GetControlRig() == ControlRig)
+					SpaceChannelAndSection.SpaceChannel = &NameAndChannel->SpaceCurve;
+				}
+				else if (bCreateIfNeeded)
+				{
+					if (ControlRig->IsControlSelected(ControlName))
 					{
-						UMovieSceneControlRigParameterSection* ActiveSection = Cast<UMovieSceneControlRigParameterSection>(ControlRigParameterTrack->GetSectionToKey());
-						if (ActiveSection)
-						{
-							ActiveSection->Modify();
-							ControlRig->Modify();
-							SpaceChannelAndSection.SectionToKey = ActiveSection;
-							FSpaceControlNameAndChannel* NameAndChannel = ActiveSection->GetSpaceChannel(ControlName);
-							if (NameAndChannel)
-							{
-								SpaceChannelAndSection.SpaceChannel = &NameAndChannel->SpaceCurve;
-							}
-							else if (bCreateIfNeeded)
-							{
-								if (ControlRig->IsControlSelected(ControlName))
-								{
-									TPair<UControlRig*, FName> Pair;
-									Pair.Key = ControlRig;
-									Pair.Value = ControlName;
-									ControlRigPairsToReselect.Add(Pair);
-								}
-								ActiveSection->AddSpaceChannel(ControlName, true /*ReconstructChannelProxy*/);
-								NameAndChannel = ActiveSection->GetSpaceChannel(ControlName);
-								if (NameAndChannel)
-								{
-									SpaceChannelAndSection.SpaceChannel = &NameAndChannel->SpaceCurve;
-									bRecreateCurves = true;
-								}
-							}
-						}
+						TPair<UControlRig*, FName> Pair;
+						Pair.Key = ControlRig;
+						Pair.Value = ControlName;
+						ControlRigPairsToReselect.Add(Pair);
+					}
+					ActiveSection->AddSpaceChannel(ControlName, true /*ReconstructChannelProxy*/);
+					NameAndChannel = ActiveSection->GetSpaceChannel(ControlName);
+					if (NameAndChannel)
+					{
+						SpaceChannelAndSection.SpaceChannel = &NameAndChannel->SpaceCurve;
+						bRecreateCurves = true;
 					}
 				}
 			}
-			if (bRecreateCurves)
-			{
-				Sequencer->RecreateCurveEditor(); //this will require the curve editor to get recreated so the ordering is correct
-				for (TPair<UControlRig*, FName>& Pair : ControlRigPairsToReselect)
+			break;
+		}
+	}
+	if (bRecreateCurves)
+	{
+		Sequencer->RecreateCurveEditor(); //this will require the curve editor to get recreated so the ordering is correct
+		for (TPair<UControlRig*, FName>& Pair : ControlRigPairsToReselect)
+		{
+
+			GEditor->GetTimerManager()->SetTimerForNextTick([Pair]()
 				{
+					Pair.Key->SelectControl(Pair.Value);
+				});
 
-					GEditor->GetTimerManager()->SetTimerForNextTick([Pair]()
-					{
-						Pair.Key->SelectControl(Pair.Value);
-					});
-					
-				}
-			}
 		}
 	}
 	return SpaceChannelAndSection;
