Sequencer: Constraints: Add support for constraints embedded in transform sections.  Added ECS system for the constraints to be evaluated(thanks AndrewR), support Smart constraint key so we compensate when keying constraints. Still a bit more todo, may move some things around, need to make sure delete delgate is handled correctly with right world, etc..

diff --git a/Engine/Plugins/Animation/ControlRig/Source/ControlRigEditor/Private/Sequencer/ConstraintChannelHelper.cpp b/Engine/Plugins/Animation/ControlRig/Source/ControlRigEditor/Private/Sequencer/ConstraintChannelHelper.cpp
--- a/Engine/Plugins/Animation/ControlRig/Source/ControlRigEditor/Private/Sequencer/ConstraintChannelHelper.cpp
+++ b/Engine/Plugins/Animation/ControlRig/Source/ControlRigEditor/Private/Sequencer/ConstraintChannelHelper.cpp
@@ -18,14 +18,19 @@
 
 #include "Tools/BakingHelper.h"
 #include "Tools/ConstraintBaker.h"
+#include "ScopedTransaction.h"
+#include "Channels/MovieSceneFloatChannel.h"
+#include "Channels/MovieSceneDoubleChannel.h"
 
 bool FConstraintChannelHelper::bDoNotCompensate = false;
 
+#define LOCTEXT_NAMESPACE "Constraints"
+
 namespace
 {
-	bool CanAddKey(const FConstraintAndActiveChannel* Channel, const FFrameNumber& InTime, bool& ActiveValue)
+	bool CanAddKey(const FMovieSceneBoolChannel& ActiveChannel, const FFrameNumber& InTime, bool& ActiveValue)
 	{
-		const TMovieSceneChannelData<const bool> ChannelData = Channel->ActiveChannel.GetData();
+		const TMovieSceneChannelData<const bool> ChannelData = ActiveChannel.GetData();
 		const TArrayView<const FFrameNumber> Times = ChannelData.GetTimes();
 		if (Times.IsEmpty())
 		{
@@ -154,7 +159,7 @@
 			const FFrameNumber Time = FrameTime.GetFrame();
 
 			bool Value = false;
-			CanAddKey(Channel, Time, Value);
+			CanAddKey(Channel->ActiveChannel, Time, Value);
 			
 			TMovieSceneChannelData<bool> ChannelData = Channel->ActiveChannel.GetData();
 			ChannelData.AddKey(Time, InConstraint->Active);
@@ -261,16 +266,78 @@
 		const int32 ChannelStartIndex = pChannelIndex->ChannelIndex;
 		return FloatChannels.Slice(ChannelStartIndex, NumChannels);
 	}
-	
+	return EmptyChannelsView;
+}
+
+TArrayView<FMovieSceneDoubleChannel*> FConstraintChannelHelper::GetTransformDoubleChannels(
+	const UTransformableHandle* InHandle,
+	const TSharedPtr<ISequencer>& InSequencer)
+{
+	// InHandle and InSequencer are assumed to be valid at this stage so no need to check
+
+	static const TArrayView<FMovieSceneDoubleChannel*> EmptyChannelsView;
 	if (const UTransformableComponentHandle* ComponentHandle = Cast<UTransformableComponentHandle>(InHandle))
 	{
 		const UMovieScene3DTransformSection* Section = GetTransformSection(ComponentHandle, InSequencer);
-		return Section ? Section->GetChannelProxy().GetChannels<FMovieSceneFloatChannel>() : EmptyChannelsView;
+		return Section ? Section->GetChannelProxy().GetChannels<FMovieSceneDoubleChannel>() : EmptyChannelsView;
 	}
-	
+
 	return EmptyChannelsView;
 }
 
+template<typename ChannelType>
+void GetFramesToCompensate(
+	const FMovieSceneConstraintChannel& InActiveChannel,
+	const bool InActiveValueToBeSet,
+	const FFrameNumber& InTime,
+	const TArrayView<ChannelType*>& InChannels,
+	TArray<FFrameNumber>& OutFramesAfter)
+{
+	using ChannelValueType = typename ChannelType::ChannelValueType;
+
+	const bool bHasKeys = (InActiveChannel.GetNumKeys() > 0);
+	
+	OutFramesAfter.Reset();
+
+	// add the current frame
+	OutFramesAfter.Add(InTime);
+
+	// add the next frames that need transform compensation 
+	for (const ChannelType* InChannel: InChannels)
+	{
+		const TMovieSceneChannelData<const ChannelValueType> ChannelData = InChannel->GetData();
+		const TArrayView<const FFrameNumber> Times = ChannelData.GetTimes();
+		if (!Times.IsEmpty())
+		{
+			// look for the first next key frame for this channel 
+			const int32 NextTimeIndex = Algo::UpperBound(Times, InTime);
+			if (Times.IsValidIndex(NextTimeIndex))
+			{
+				// store the time while the state is different
+				for (int32 Index = NextTimeIndex; Index < Times.Num(); ++Index)
+				{
+					if (!bHasKeys)
+					{
+						OutFramesAfter.Add(Times[Index]);
+					}
+					else
+					{
+						bool NextValue = false; InActiveChannel.Evaluate(Times[Index], NextValue);
+						if (NextValue == InActiveValueToBeSet)
+						{
+							break;
+						}
+						OutFramesAfter.Add(Times[Index]);
+					}
+				}
+			}
+		}
+	}
+
+	// uniqueness
+	OutFramesAfter.Sort();
+	OutFramesAfter.SetNum(Algo::Unique(OutFramesAfter));
+}
 namespace
 {
 
@@ -513,6 +580,8 @@
 	
 	if (UMovieSceneControlRigParameterSection* Section = GetControlSection(ControlHandle, InSequencer))
 	{
+		FScopedTransaction Transaction(LOCTEXT("KeyConstraintaKehy", "Key Constraint Key"));
+		Section->Modify();
 		// add the channel
 		Section->AddConstraintChannel(InConstraint, true);
 
@@ -524,8 +593,9 @@
 			const FFrameNumber Time = FrameTime.GetFrame();
 
 			bool ActiveValueToBeSet = false;
-			if (CanAddKey(Channel, Time, ActiveValueToBeSet))
+			if (CanAddKey(Channel->ActiveChannel, Time, ActiveValueToBeSet))
 			{
+
 				TGuardValue<bool> CompensateGuard(bDoNotCompensate, true);
 				
 				// set constraint as dynamic
@@ -537,7 +607,7 @@
 				// store the frames to compensate
 				const TArrayView<FMovieSceneFloatChannel*> Channels = GetTransformFloatChannels(ControlHandle, InSequencer);
 				TArray<FFrameNumber> FramesToCompensate;
-				GetFramesToCompensate(Channel->ActiveChannel, ActiveValueToBeSet, Time, Channels, FramesToCompensate);
+				GetFramesToCompensate<FMovieSceneFloatChannel>(Channel->ActiveChannel, ActiveValueToBeSet, Time, Channels, FramesToCompensate);
 
 				// store child and space transforms for these frames
 				FCompensationEvaluator Evaluator(InConstraint);
@@ -546,24 +616,42 @@
 				
 				// store tangents at this time
 				TArray<FMovieSceneTangentData> Tangents;
-				FControlRigSpaceChannelHelpers::EvaluateTangentAtThisTime(ControlRig, Section, ControlName, Time, Tangents);
+				int32 ChannelIndex = 0, NumChannels = 0;
 
+				FChannelMapInfo* pChannelIndex = nullptr;
+				FRigControlElement* ControlElement = nullptr;
+				Tie(ControlElement, pChannelIndex) = FControlRigSpaceChannelHelpers::GetControlAndChannelInfo(ControlRig, Section, ControlName);
+
+				if (pChannelIndex && ControlElement)
+				{
+					// get the number of float channels to treat
+					NumChannels = FControlRigSpaceChannelHelpers::GetNumFloatChannels(ControlElement->Settings.ControlType);
+					if (NumChannels > 0)
+					{
+						ChannelIndex = pChannelIndex->ChannelIndex;
+						FControlRigSpaceChannelHelpers::EvaluateTangentAtThisTime(ChannelIndex, NumChannels, true,Section, Time, Tangents);
+					}
+				}
+				
+
 				const EMovieSceneTransformChannel ChannelsToKey = FConstraintBaker::GetChannelsToKey(InConstraint);
 				
 				// add child's transform key at Time-1 to keep animation
 				{
-					const FFrameNumber TimeMinusOne(Time -1);
-					
-					FBakingHelper::AddTransformKeys(ControlRig, ControlName, {TimeMinusOne},
-					 	{ChildLocals[0]}, ChannelsToKey, TickResolution, true);
-					
+					const FFrameNumber TimeMinusOne(Time - 1);
+
+					FBakingHelper::AddTransformKeys(ControlRig, ControlName, { TimeMinusOne },
+						{ ChildLocals[0] }, ChannelsToKey, TickResolution, true);
+
 					// set tangents at Time-1
-					FControlRigSpaceChannelHelpers::SetTangentsAtThisTime(ControlRig, Section, ControlName, TimeMinusOne, Tangents);
+					if (NumChannels > 0)
+					{
+						FControlRigSpaceChannelHelpers::SetTangentsAtThisTime(ChannelIndex, NumChannels, true, Section, TimeMinusOne, Tangents);
+					}
 				}
 
 				// add active key
 				{
-					Section->Modify();
 					TMovieSceneChannelData<bool> ChannelData = Channel->ActiveChannel.GetData();
 					ChannelData.AddKey(Time, ActiveValueToBeSet);
 				}
@@ -574,11 +662,14 @@
 					ChildLocals.RemoveAt(0);
 
 					// add keys
-					FBakingHelper::AddTransformKeys( ControlRig, ControlHandle->ControlName, FramesToCompensate,
-					 	ChildLocals, ChannelsToKey, TickResolution, true);
-					
+					FBakingHelper::AddTransformKeys(ControlRig, ControlHandle->ControlName, FramesToCompensate,
+						ChildLocals, ChannelsToKey, TickResolution, true);
+
 					// set tangents at Time
-					FControlRigSpaceChannelHelpers::SetTangentsAtThisTime(ControlRig, Section, ControlName, Time, Tangents);
+					if (NumChannels > 0 )
+					{
+						FControlRigSpaceChannelHelpers::SetTangentsAtThisTime(ChannelIndex, NumChannels, true, Section, Time, Tangents);
+					}
 				}
 			}
 		}
@@ -594,57 +685,93 @@
 	{
 		return;
 	}
-}
+	AActor* Actor = ComponentHandle->Component->GetOwner();
+	if (!Actor)
+	{
+		return;
+	}
 
-void FConstraintChannelHelper::GetFramesToCompensate(
-	const FMovieSceneConstraintChannel& InActiveChannel,
-	const bool InActiveValueToBeSet,
-	const FFrameNumber& InTime,
-	const TArrayView<FMovieSceneFloatChannel*>& InChannels,
-	TArray<FFrameNumber>& OutFramesAfter)
-{
-	const bool bHasKeys = (InActiveChannel.GetNumKeys() > 0);
-	
-	OutFramesAfter.Reset();
+	const FTransform LocalTransform = ComponentHandle->GetLocalTransform();
+	const FGuid Guid = InSequencer->GetHandleToObject(Actor, true);
+	if (!Guid.IsValid())
+	{
+		return;
+	}
 
-	// add the current frame
-	OutFramesAfter.Add(InTime);
+	if (UMovieScene3DTransformSection* Section = FBakingHelper::GetTransformSection(InSequencer.Get(), Guid, LocalTransform))
+	{
+		FScopedTransaction Transaction(LOCTEXT("KeyConstraintaKehy", "Key Constraint Key"));
+		Section->Modify();
+		// add the channel
+		Section->AddConstraintChannel(InConstraint);
 
-	// add the next frames that need transform compensation 
-	for (const FMovieSceneFloatChannel* InChannel: InChannels)
-	{
-		const TMovieSceneChannelData<const FMovieSceneFloatValue> ChannelData = InChannel->GetData();
-		const TArrayView<const FFrameNumber> Times = ChannelData.GetTimes();
-		if (!Times.IsEmpty())
+		// add key if needed
+		if (FConstraintAndActiveChannel* Channel = Section->GetConstraintChannel(InConstraint->GetFName()))
 		{
-			// look for the first next key frame for this channel 
-			const int32 NextTimeIndex = Algo::UpperBound(Times, InTime);
-			if (Times.IsValidIndex(NextTimeIndex))
+			const FFrameRate TickResolution = InSequencer->GetFocusedTickResolution();
+			const FFrameTime FrameTime = InSequencer->GetLocalTime().ConvertTo(TickResolution);
+			const FFrameNumber Time = FrameTime.GetFrame();
+
+			bool ActiveValueToBeSet = false;
+			if (CanAddKey(Channel->ActiveChannel, Time, ActiveValueToBeSet))
 			{
-				// store the time while the state is different
-				for (int32 Index = NextTimeIndex; Index < Times.Num(); ++Index)
+	
+				Section->Modify();
+
+				//new for compensatio
+
+				TGuardValue<bool> CompensateGuard(bDoNotCompensate, true);
+
+				// set constraint as dynamic
+				InConstraint->bDynamicOffset = true;
+
+				// store the frames to compensate
+				const TArrayView<FMovieSceneDoubleChannel*> Channels = GetTransformDoubleChannels(ComponentHandle, InSequencer);
+				TArray<FFrameNumber> FramesToCompensate;
+				GetFramesToCompensate<FMovieSceneDoubleChannel>(Channel->ActiveChannel, ActiveValueToBeSet, Time, Channels, FramesToCompensate);
+
+				// store child and space transforms for these frames
+				FCompensationEvaluator Evaluator(InConstraint);
+				Evaluator.ComputeLocalTransforms(Actor->GetWorld(), InSequencer, FramesToCompensate, ActiveValueToBeSet);
+				TArray<FTransform>& ChildLocals = Evaluator.ChildLocals;
+
+				// store tangents at this time
+				TArray<FMovieSceneTangentData> Tangents;
+				const int32 ChannelIndex = 0;
+				const int32 NumChannels = 9;
+				FControlRigSpaceChannelHelpers::EvaluateTangentAtThisTime(ChannelIndex, NumChannels, false, Section, Time, Tangents);
+
+				const EMovieSceneTransformChannel ChannelsToKey = FConstraintBaker::GetChannelsToKey(InConstraint);
+
+				// add child's transform key at Time-1 to keep animation
 				{
-					if (!bHasKeys)
-					{
-						OutFramesAfter.Add(Times[Index]);
-					}
-					else
-					{
-						bool NextValue = false; InActiveChannel.Evaluate(Times[Index], NextValue);
-						if (NextValue == InActiveValueToBeSet)
-						{
-							break;
-						}
-						OutFramesAfter.Add(Times[Index]);
-					}
+					const FFrameNumber TimeMinusOne(Time - 1);
+
+					FBakingHelper::AddTransformKeys(Section, { TimeMinusOne }, { ChildLocals[0] }, ChannelsToKey);
+
+					FControlRigSpaceChannelHelpers::SetTangentsAtThisTime(ChannelIndex, NumChannels, false,Section, TimeMinusOne, Tangents);
 				}
+
+				// add active key
+				{
+					TMovieSceneChannelData<bool> ChannelData = Channel->ActiveChannel.GetData();
+					ChannelData.AddKey(Time, ActiveValueToBeSet);
+				}
+
+				// compensate
+				{
+					// we need to remove the first transforms as we store NumFrames+1 transforms
+					ChildLocals.RemoveAt(0);
+
+					// add keys
+					FBakingHelper::AddTransformKeys(Section, FramesToCompensate, ChildLocals, ChannelsToKey);
+
+					// set tangents at Time
+					FControlRigSpaceChannelHelpers::SetTangentsAtThisTime(ChannelIndex,NumChannels, false, Section, Time, Tangents);
+				}
 			}
 		}
 	}
-
-	// uniqueness
-	OutFramesAfter.Sort();
-	OutFramesAfter.SetNum(Algo::Unique(OutFramesAfter));
 }
 
 void FConstraintChannelHelper::Compensate(UTickableTransformConstraint* InConstraint, const bool bAllTimes)
@@ -754,3 +881,5 @@
 		InSequencer->ForceEvaluate();
 	}
 }
+
+#undef LOCTEXT_NAMESPACE
diff --git a/Engine/Plugins/Animation/ControlRig/Source/ControlRigEditor/Private/Sequencer/ControlRigSpaceChannelEditors.cpp b/Engine/Plugins/Animation/ControlRig/Source/ControlRigEditor/Private/Sequencer/ControlRigSpaceChannelEditors.cpp
--- a/Engine/Plugins/Animation/ControlRig/Source/ControlRigEditor/Private/Sequencer/ControlRigSpaceChannelEditors.cpp
+++ b/Engine/Plugins/Animation/ControlRig/Source/ControlRigEditor/Private/Sequencer/ControlRigSpaceChannelEditors.cpp
@@ -251,6 +251,138 @@ FSpaceChannelAndSection FControlRigSpaceChannelHelpers::FindSpaceChannelAndSecti
 	return SpaceChannelAndSection;
 }
 
+template<typename ChannelType>
+void EvaluateTangentAtThisTimeImpl(int32 ChannelIndex, int32 NumChannels,
+	UMovieSceneSection* Section,
+	FFrameNumber Time,
+	TArray<FMovieSceneTangentData>& OutTangents)
+{
+	using ChannelValueType = typename ChannelType::ChannelValueType;
+
+	// NOTE this might be moved to FMovieSceneFloatChannel
+	auto EvaluateTangent = [](const ChannelType* InChannel, FFrameNumber InTime)
+	{
+		const TMovieSceneChannelData<const ChannelValueType> ChannelInterface = InChannel->GetData();
+		const TArrayView<const FFrameNumber> Times = ChannelInterface.GetTimes();
+
+		// Need at least two keys to evaluate a derivative
+		if (Times.Num() < 2)
+		{
+			static const FMovieSceneTangentData DefaultTangent;
+			return DefaultTangent;
+		}
+
+		const TArrayView<const ChannelValueType> Values = ChannelInterface.GetValues();
+
+		// look around to get the closest key tangent if fairly close
+		const int32 Tolerance = static_cast<int32>(InChannel->GetTickResolution().AsDecimal() * 0.01);
+		// NOTE FindKey might return Times.Num() (see AlgoImpl::LowerBoundInternal)
+		const int32 ExistingIndex = ChannelInterface.FindKey(InTime, Tolerance);
+		if (Times.IsValidIndex(ExistingIndex) && Values.IsValidIndex(ExistingIndex))
+		{
+			const FFrameNumber& Time = Times[ExistingIndex];
+			const int32 DiffToKey = InTime.Value - Time.Value;
+
+			// if the closest key is within a threshold, we use it's tangent directly instead of computing one
+			if (FMath::Abs(DiffToKey) <= Tolerance)
+			{
+				const ChannelValueType Value = Values[ExistingIndex];
+				if (DiffToKey == 0)
+				{
+					return Value.Tangent;
+				}
+
+				FMovieSceneTangentData Tangent = Value.Tangent;
+				if (DiffToKey < 0) // pretty close to the next key
+				{
+					Tangent.LeaveTangent = 0.f;
+				}
+				else // pretty close to the previous key
+				{
+					Tangent.ArriveTangent = 0.f;
+				}
+				return Tangent;
+			}
+		}
+
+		// compute tangent using central difference
+		// NOTE we may wanna compute a backward / forward difference instead
+		const int32 Delta = static_cast<int32>(InChannel->GetTickResolution().AsDecimal() * 0.1);
+
+		float PrevValue = 0.f; InChannel->Evaluate(InTime - Delta, PrevValue);
+		float NextValue = 0.f; InChannel->Evaluate(InTime + Delta, NextValue);
+		const float TangentValue = (NextValue - PrevValue) / (2.f * Delta);
+
+		FMovieSceneTangentData TangentData;
+		TangentData.ArriveTangent = TangentValue;
+		TangentData.LeaveTangent = TangentValue;
+
+		return TangentData;
+	};
+
+
+	// compute and store tangents
+	OutTangents.SetNum(NumChannels);
+
+	const TArrayView<ChannelType*> Channels = Section->GetChannelProxy().GetChannels<ChannelType>();
+	for (int32 Index = 0; Index < NumChannels; ++Index, ++ChannelIndex)
+	{
+		OutTangents[Index] = EvaluateTangent(Channels[ChannelIndex], Time);
+	}
+}
+
+// NOTE we may pass an enum to tell which tangent we wanna set (arrive, leave, both)
+template<typename ChannelType>
+void SetTangentsAtThisTimeImpl(int32 ChannelIndex,
+	int32 NumChannels,
+	UMovieSceneSection* Section,
+	FFrameNumber Time,
+	const TArray<FMovieSceneTangentData>& InTangents)
+{
+	using ChannelValueType = typename ChannelType::ChannelValueType;
+
+	const TArrayView<ChannelType*> Channels = Section->GetChannelProxy().GetChannels<ChannelType>();
+
+	for (int32 Index = 0; Index < NumChannels; ++Index, ++ChannelIndex)
+	{
+		TMovieSceneChannelData<ChannelValueType> ChannelInterface = Channels[ChannelIndex]->GetData();
+		TArrayView<ChannelValueType> Values = ChannelInterface.GetValues();
+		const int32 KeyIndex = ChannelInterface.FindKey(Time);
+		if (KeyIndex != INDEX_NONE)
+		{
+			ChannelValueType& Value = Values[KeyIndex];
+			Value.Tangent.ArriveTangent = InTangents[Index].ArriveTangent;
+			Value.Tangent.LeaveTangent = InTangents[Index].LeaveTangent;
+			Value.TangentMode = RCTM_Break;
+		}
+	}
+}
+
+void FControlRigSpaceChannelHelpers::EvaluateTangentAtThisTime(int32 ChannelIndex, int32 NumChannels, bool bFloatChannels, UMovieSceneSection* Section, FFrameNumber Time, TArray<FMovieSceneTangentData>& OutTangents)
+{
+	if (bFloatChannels)
+	{
+		EvaluateTangentAtThisTimeImpl<FMovieSceneFloatChannel>(ChannelIndex, NumChannels, Section, Time, OutTangents);
+	}
+	else
+	{
+		EvaluateTangentAtThisTimeImpl<FMovieSceneDoubleChannel>(ChannelIndex, NumChannels, Section, Time, OutTangents);
+	}
+}
+
+void FControlRigSpaceChannelHelpers::SetTangentsAtThisTime(int32 ChannelIndex, int32 NumChannels, bool bFloatChannels, UMovieSceneSection* Section, FFrameNumber Time, const TArray<FMovieSceneTangentData>& InTangents)
+{
+	if (bFloatChannels)
+	{
+		SetTangentsAtThisTimeImpl<FMovieSceneFloatChannel>(ChannelIndex, NumChannels, Section, Time, InTangents);
+	}
+	else
+	{
+		SetTangentsAtThisTimeImpl<FMovieSceneDoubleChannel>(ChannelIndex, NumChannels, Section, Time, InTangents);
+	}
+}
+
+
 /*
 *  Situations to handle\
 *  0) New space sames as current space do nothing
@@ -353,10 +485,24 @@ FKeyHandle FControlRigSpaceChannelHelpers::SequencerKeyControlRigSpaceChannel(UC
 
 		// store tangents to keep the animation when switching 
 		TArray<FMovieSceneTangentData> Tangents;
+		int32 ChannelIndex = 0, NumChannels = 0;
 		UMovieSceneControlRigParameterSection* ControlRigSection = Cast<UMovieSceneControlRigParameterSection>(SectionToKey);
 		if (ControlRigSection)
 		{
-			EvaluateTangentAtThisTime(ControlRig, ControlRigSection, ControlKey.Name, Time, Tangents);
+			FChannelMapInfo* pChannelIndex = nullptr;
+			FRigControlElement* ControlElement = nullptr;
+			Tie(ControlElement, pChannelIndex) = GetControlAndChannelInfo(ControlRig, ControlRigSection, ControlKey.Name);
+
+			if (pChannelIndex && ControlElement)
+			{
+				// get the number of float channels to treat
+				NumChannels = GetNumFloatChannels(ControlElement->Settings.ControlType);
+				if (NumChannels > 0)
+				{
+					ChannelIndex = pChannelIndex->ChannelIndex;
+					EvaluateTangentAtThisTime(ChannelIndex, NumChannels, true,ControlRigSection, Time, Tangents);
+				}
+			}
 		}
 		
 		// if current Value not the same as the previous add new space key
@@ -417,10 +563,9 @@ FKeyHandle FControlRigSpaceChannelHelpers::SequencerKeyControlRigSpaceChannel(UC
 			ControlRig->SetControlGlobalTransform(ControlKey.Name, ControlWorldTransforms[0], true, Context);
 
 			// need to update the tangents here to keep the arriving animation
-			if (ControlRigSection)
+			if (ControlRigSection  && NumChannels > 0)
 			{
-				// ControlRigSection->Modify(); ?
-				SetTangentsAtThisTime(ControlRig, ControlRigSection, ControlKey.Name, GlobalTime.GetFrame(), Tangents);
+				SetTangentsAtThisTime(ChannelIndex, NumChannels,true, ControlRigSection,GlobalTime.GetFrame(), Tangents);
 			}
 		}
 
@@ -445,10 +590,9 @@ FKeyHandle FControlRigSpaceChannelHelpers::SequencerKeyControlRigSpaceChannel(UC
 			// need to update the tangents here to keep the leaving animation
 			if (bSetPreviousKey && FramesIndex == 0)
 			{
-				if (ControlRigSection)
+				if (ControlRigSection && NumChannels > 0)
 				{
-					// ControlRigSection->Modify(); ?
-					SetTangentsAtThisTime(ControlRig, ControlRigSection, ControlKey.Name, GlobalTime.GetFrame(), Tangents);
+					SetTangentsAtThisTime(ChannelIndex, NumChannels, true,  ControlRigSection, GlobalTime.GetFrame(), Tangents);
 				}
 			}
 
@@ -1398,141 +1542,6 @@ int32 FControlRigSpaceChannelHelpers::GetNumFloatChannels(const ERigControlType
 	return 0;
 }
 
-void FControlRigSpaceChannelHelpers::EvaluateTangentAtThisTime(	UControlRig* ControlRig,
-																UMovieSceneControlRigParameterSection* Section,
-																FName ControlName,
-																FFrameNumber Time,
-																TArray<FMovieSceneTangentData>& OutTangents)
-{
-	// NOTE this might be moved to FMovieSceneFloatChannel
-	auto EvaluateTangent = [](const FMovieSceneFloatChannel* InChannel, FFrameNumber InTime)
-	{
-		const TMovieSceneChannelData<const FMovieSceneFloatValue> ChannelInterface = InChannel->GetData();
-		const TArrayView<const FFrameNumber> Times = ChannelInterface.GetTimes();
-
-		// Need at least two keys to evaluate a derivative
-		if (Times.Num() < 2)
-		{
-			static const FMovieSceneTangentData DefaultTangent;
-			return DefaultTangent;
-		}
-
-		const TArrayView<const FMovieSceneFloatValue> Values = ChannelInterface.GetValues();
-		
-		// look around to get the closest key tangent if fairly close
-		const int32 Tolerance = static_cast<int32>(InChannel->GetTickResolution().AsDecimal() * 0.01);
-		// NOTE FindKey might return Times.Num() (see AlgoImpl::LowerBoundInternal)
-		const int32 ExistingIndex = ChannelInterface.FindKey(InTime, Tolerance);
-		if (Times.IsValidIndex(ExistingIndex) && Values.IsValidIndex(ExistingIndex))
-		{
-			const FFrameNumber& Time = Times[ExistingIndex];
-			const int32 DiffToKey = InTime.Value - Time.Value;
-
-			// if the closest key is within a threshold, we use it's tangent directly instead of computing one
-			if (FMath::Abs(DiffToKey) <= Tolerance)
-			{
-				const FMovieSceneFloatValue Value = Values[ExistingIndex];
-				if (DiffToKey == 0)
-				{
-					return Value.Tangent;				
-				}
-
-				FMovieSceneTangentData Tangent = Value.Tangent;
-				if (DiffToKey < 0) // pretty close to the next key
-				{
-				 	Tangent.LeaveTangent = 0.f;
-				}
-				else // pretty close to the previous key
-				{
-					Tangent.ArriveTangent = 0.f;
-				}
-				return Tangent;
-			}
-		}
-
-		// compute tangent using central difference
-		// NOTE we may wanna compute a backward / forward difference instead
-		const int32 Delta = static_cast<int32>(InChannel->GetTickResolution().AsDecimal() * 0.1);
-
-		float PrevValue = 0.f; InChannel->Evaluate(InTime-Delta, PrevValue);
-		float NextValue = 0.f; InChannel->Evaluate(InTime+Delta, NextValue);
-		const float TangentValue = (NextValue - PrevValue) / (2.f*Delta);
-		
-		FMovieSceneTangentData TangentData;
-		TangentData.ArriveTangent = TangentValue;
-		TangentData.LeaveTangent = TangentValue;
-	
-		return TangentData;
-	};
-	
-	FChannelMapInfo* pChannelIndex = nullptr;
-	FRigControlElement* ControlElement = nullptr;
-	Tie(ControlElement, pChannelIndex) = GetControlAndChannelInfo(ControlRig, Section, ControlName);
-	
-	if (!pChannelIndex || !ControlElement)
-	{
-		return;
-	}
-
-	// get the number of float channels to treat
-	const int NumChannels = GetNumFloatChannels(ControlElement->Settings.ControlType);
-	if (!NumChannels)
-	{
-		return;
-	}
-
-	// compute and store tangents
-	OutTangents.SetNum(NumChannels);
-
-	const TArrayView<FMovieSceneFloatChannel*> FloatChannels = Section->GetChannelProxy().GetChannels<FMovieSceneFloatChannel>();
-	int32 ChannelIndex = pChannelIndex->ChannelIndex;
-	for (int32 Index = 0; Index < NumChannels; ++Index, ++ChannelIndex)
-	{
-		OutTangents[Index] = EvaluateTangent(FloatChannels[ChannelIndex], Time);
-	}
-}
-
-// NOTE we may pass an enum to tell which tangent we wanna set (arrive, leave, both)
-void FControlRigSpaceChannelHelpers::SetTangentsAtThisTime(	UControlRig* ControlRig,
-															UMovieSceneControlRigParameterSection* Section,
-															FName ControlName,
-															FFrameNumber Time,
-															const TArray<FMovieSceneTangentData>& InTangents)
-{
-	FChannelMapInfo* pChannelIndex = nullptr;
-    FRigControlElement* ControlElement = nullptr;
-    Tie(ControlElement, pChannelIndex) = GetControlAndChannelInfo(ControlRig, Section, ControlName);
-    
-    if (!pChannelIndex || !ControlElement)
-    {
-    	return;
-    }
-
-	// get the number of float channels to treat
-	const int32 NumChannels = GetNumFloatChannels(ControlElement->Settings.ControlType);
-	if (!NumChannels || InTangents.Num() != NumChannels)
-	{
-		return;
-	}
-
-	const TArrayView<FMovieSceneFloatChannel*> FloatChannels = Section->GetChannelProxy().GetChannels<FMovieSceneFloatChannel>();
-	int32 ChannelIndex = pChannelIndex->ChannelIndex;
-
-	for (int32 Index = 0; Index < NumChannels; ++Index, ++ChannelIndex)
-	{
-		TMovieSceneChannelData<FMovieSceneFloatValue> ChannelInterface = FloatChannels[ChannelIndex]->GetData();
-		TArrayView<FMovieSceneFloatValue> Values = ChannelInterface.GetValues();
-		const int32 KeyIndex = ChannelInterface.FindKey(Time);
-		if (KeyIndex != INDEX_NONE)
-		{
-			FMovieSceneFloatValue& Value = Values[KeyIndex];
-			Value.Tangent.ArriveTangent = InTangents[Index].ArriveTangent;
-			Value.Tangent.LeaveTangent = InTangents[Index].LeaveTangent;
-			Value.TangentMode = RCTM_Break;	
-		}
-	}
-}
-
 void DrawExtra(FMovieSceneControlRigSpaceChannel* Channel, const UMovieSceneSection* Owner,const FGeometry& AllottedGeometry, FSequencerSectionPainter& Painter)
 {
 	using namespace UE::Sequencer;
diff --git a/Engine/Plugins/Animation/ControlRig/Source/ControlRigEditor/Public/ConstraintChannelHelper.h b/Engine/Plugins/Animation/ControlRig/Source/ControlRigEditor/Public/ConstraintChannelHelper.h
--- a/Engine/Plugins/Animation/ControlRig/Source/ControlRigEditor/Public/ConstraintChannelHelper.h
+++ b/Engine/Plugins/Animation/ControlRig/Source/ControlRigEditor/Public/ConstraintChannelHelper.h
@@ -12,6 +12,7 @@ class UTickableTransformConstraint;
 class UTransformableControlHandle;
 class UTransformableComponentHandle;
 struct FMovieSceneFloatChannel;
+struct FMovieSceneDoubleChannel;
 struct FMovieSceneConstraintChannel;
 class UMovieSceneControlRigParameterSection;
 struct FFrameNumber;
@@ -75,18 +76,15 @@ private:
 
 	/** END COMPONENT SECTION */
 
-	/** @todo documentation. */
-	static void GetFramesToCompensate(
-		const FMovieSceneConstraintChannel& InActiveChannel,
-		const bool InValue,
-		const FFrameNumber& InTime,
-		const TArrayView<FMovieSceneFloatChannel*>& InChannels,
-		TArray<FFrameNumber>& OutFramesAfter);
-
 	/** @todo documentation. */
 	static TArrayView<FMovieSceneFloatChannel*> GetTransformFloatChannels(
 		const UTransformableHandle* InHandle,
 		const TSharedPtr<ISequencer>& InSequencer);
 
+	/** @todo documentation.  TransformComponents use doubles */
+	static TArrayView<FMovieSceneDoubleChannel*> GetTransformDoubleChannels(
+		const UTransformableHandle* InHandle,
+		const TSharedPtr<ISequencer>& InSequencer);
+
 	static bool bDoNotCompensate;
 };
diff --git a/Engine/Plugins/Animation/ControlRig/Source/ControlRigEditor/Public/ControlRigSpaceChannelEditors.h b/Engine/Plugins/Animation/ControlRig/Source/ControlRigEditor/Public/ControlRigSpaceChannelEditors.h
--- a/Engine/Plugins/Animation/ControlRig/Source/ControlRigEditor/Public/ControlRigSpaceChannelEditors.h
+++ b/Engine/Plugins/Animation/ControlRig/Source/ControlRigEditor/Public/ControlRigSpaceChannelEditors.h
@@ -61,9 +61,9 @@ struct FControlRigSpaceChannelHelpers
 	static TArray<FKeyBarCurveModel::FBarRange> FindRanges(FMovieSceneControlRigSpaceChannel* Channel, const UMovieSceneSection* Section);
 
 	// evaluates the tangent by central difference at a specific time. (return the closest key tangent if within an acceptable tolerance) 
-	static void EvaluateTangentAtThisTime(UControlRig* ControlRig, UMovieSceneControlRigParameterSection* Section, FName ControlName, FFrameNumber Time, TArray<FMovieSceneTangentData>& OutTangents);
+	static void EvaluateTangentAtThisTime(int32 ChannelIndex, int32 NumChannels, bool bFloatChannels, UMovieSceneSection* Section,  FFrameNumber Time, TArray<FMovieSceneTangentData>& OutTangents);
 	// set the tangent at a specific time. (if there's a key to set at this time)
-	static void SetTangentsAtThisTime(UControlRig* ControlRig, UMovieSceneControlRigParameterSection* Section, FName ControlName, FFrameNumber Time, const TArray<FMovieSceneTangentData>& InTangents);
+	static void SetTangentsAtThisTime(int32 ChannelIndex, int32 NumChannels,bool bFloatChannels, UMovieSceneSection* Section, FFrameNumber Time, const TArray<FMovieSceneTangentData>& InTangents);
 	// retrieve the control and the channel infos for that ControlRig/Section.
 	static TPair<FRigControlElement*, FChannelMapInfo*> GetControlAndChannelInfo(UControlRig* ControlRig, UMovieSceneControlRigParameterSection* Section, FName ControlName);
 	// retrieve the number of float channels based on the control type.
diff --git a/Engine/Source/Editor/CurveEditor/Public/KeyBarCurveModel.h b/Engine/Source/Editor/CurveEditor/Public/KeyBarCurveModel.h
--- a/Engine/Source/Editor/CurveEditor/Public/KeyBarCurveModel.h
+++ b/Engine/Source/Editor/CurveEditor/Public/KeyBarCurveModel.h
@@ -44,6 +44,8 @@ public:
 	*/
 	struct CURVEEDITOR_API FBarRange
 	{
+		/* Spaces have ifinite ranges, Constraints don't*/
+		bool bRangeIsInfinite = true;
 		TRange<double> Range;
 		FName  Name;
 		FLinearColor Color;
diff --git a/Engine/Source/Editor/MovieSceneTools/Private/Channels/ConstraintChannelCurveModel.cpp b/Engine/Source/Editor/MovieSceneTools/Private/Channels/ConstraintChannelCurveModel.cpp
new file mode 100644
--- /dev/null
+++ b/Engine/Source/Editor/MovieSceneTools/Private/Channels/ConstraintChannelCurveModel.cpp
@@ -0,0 +1,321 @@
+// Copyright Epic Games, Inc. All Rights Reserved.
+
+#include "Channels/ConstraintChannelCurveModel.h"
+#include "HAL/PlatformMath.h"
+#include "Channels/MovieSceneBoolChannel.h"
+#include "MovieSceneSection.h"
+#include "MovieScene.h"
+#include "CurveDrawInfo.h"
+#include "CurveDataAbstraction.h"
+#include "CurveEditor.h"
+#include "CurveEditorScreenSpace.h"
+#include "CurveEditorSnapMetrics.h"
+#include "Styling/AppStyle.h"
+#include "SequencerChannelTraits.h"
+#include "SequencerSectionPainter.h"
+#include "ISequencer.h"
+#include "MVVM/Views/STrackAreaView.h"
+#include "ConstraintsManager.h"
+#include "Channels/ConstraintChannelEditor.h"
+
+ECurveEditorViewID FConstraintChannelCurveModel::ViewID = ECurveEditorViewID::Invalid;
+
+FConstraintChannelCurveModel::FConstraintChannelCurveModel(TMovieSceneChannelHandle<FMovieSceneConstraintChannel> InChannel, UMovieSceneSection* OwningSection, TWeakPtr<ISequencer> InWeakSequencer)
+{
+	using namespace UE::Sequencer;
+
+	FMovieSceneChannelProxy* NewChannelProxy = &OwningSection->GetChannelProxy();
+	ChannelHandle = NewChannelProxy->MakeHandle<FMovieSceneConstraintChannel>(InChannel.GetChannelIndex());
+	if (FMovieSceneChannelProxy* ChannelProxy = InChannel.GetChannelProxy())
+	{
+		OnDestroyHandle = NewChannelProxy->OnDestroy.AddRaw(this, &FConstraintChannelCurveModel::FixupCurve);
+	}
+
+	WeakSection = OwningSection;
+	WeakSequencer = InWeakSequencer;
+	SupportedViews = ViewID;
+	Constraint = nullptr;
+	Color = STrackAreaView::BlendDefaultTrackColor(OwningSection->GetTypedOuter<UMovieSceneTrack>()->GetColorTint());
+}
+
+FConstraintChannelCurveModel::~FConstraintChannelCurveModel()
+{
+	if (FMovieSceneChannelProxy* ChannelProxy = ChannelHandle.GetChannelProxy())
+	{
+		ChannelProxy->OnDestroy.Remove(OnDestroyHandle);
+	}
+}
+
+void FConstraintChannelCurveModel::FixupCurve()
+{
+	if (UMovieSceneSection* Section = WeakSection.Get())
+	{
+		FMovieSceneChannelProxy* NewChannelProxy = &Section->GetChannelProxy();
+		ChannelHandle = NewChannelProxy->MakeHandle<FMovieSceneConstraintChannel>(ChannelHandle.GetChannelIndex());
+		OnDestroyHandle = NewChannelProxy->OnDestroy.AddRaw(this, &FConstraintChannelCurveModel::FixupCurve);
+	}
+}
+TArray<FKeyBarCurveModel::FBarRange> FConstraintChannelCurveModel::FindRanges()
+{
+	FMovieSceneConstraintChannel* Channel = ChannelHandle.Get();
+	UMovieSceneSection* Section = WeakSection.Get();
+	TArray<FKeyBarCurveModel::FBarRange> Ranges = FConstraintChannelEditor::GetBarRanges(Channel, Section);
+	return Ranges;
+}
+
+const void* FConstraintChannelCurveModel::GetCurve() const
+{
+	return ChannelHandle.Get();
+}
+
+void FConstraintChannelCurveModel::Modify()
+{
+	if (UMovieSceneSection* Section = WeakSection.Get())
+	{
+		Section->Modify();
+	}
+}
+
+void FConstraintChannelCurveModel::AddKeys(TArrayView<const FKeyPosition> InKeyPositions, TArrayView<const FKeyAttributes> InKeyAttributes, TArrayView<TOptional<FKeyHandle>>* OutKeyHandles)
+{
+	check(InKeyPositions.Num() == InKeyAttributes.Num() && (!OutKeyHandles || OutKeyHandles->Num() == InKeyPositions.Num()));
+
+	FMovieSceneConstraintChannel* Channel = ChannelHandle.Get();
+	UMovieSceneSection* Section = WeakSection.Get();
+	if (Channel && Section)
+	{
+		Section->Modify();
+		TMovieSceneChannelData<bool> ChannelData = Channel->GetData();
+		FFrameRate TickResolution = Section->GetTypedOuter<UMovieScene>()->GetTickResolution();
+
+		for (int32 Index = 0; Index < InKeyPositions.Num(); ++Index)
+		{
+			FKeyPosition   Position = InKeyPositions[Index];
+			FKeyAttributes Attributes = InKeyAttributes[Index];
+
+			FFrameNumber Time = (Position.InputValue * TickResolution).RoundToFrame();
+			Section->ExpandToFrame(Time);
+
+			bool Value = false;
+			FKeyHandle NewHandle = ChannelData.UpdateOrAddKey(Time, Value);
+			if (NewHandle != FKeyHandle::Invalid())
+			{
+				if (OutKeyHandles)
+				{
+					(*OutKeyHandles)[Index] = NewHandle;
+				}
+			}
+		}
+	}
+}
+
+bool FConstraintChannelCurveModel::Evaluate(double Time, double& OutValue) const
+{
+	// Constraints don't evaluate into a valid value
+	OutValue = 0.0;
+	return false;
+}
+
+void FConstraintChannelCurveModel::RemoveKeys(TArrayView<const FKeyHandle> InKeys)
+{
+	FMovieSceneConstraintChannel* Channel = ChannelHandle.Get();
+	UMovieSceneSection* Section = WeakSection.Get();
+	if (Channel && Section)
+	{
+		Section->Modify();
+
+		TMovieSceneChannelData<bool> ChannelData = Channel->GetData();
+
+		for (FKeyHandle Handle : InKeys)
+		{
+			int32 KeyIndex = ChannelData.GetIndex(Handle);
+			if (KeyIndex != INDEX_NONE)
+			{
+				ChannelData.RemoveKey(KeyIndex);
+			}
+		}
+	}
+}
+
+void FConstraintChannelCurveModel::DrawCurve(const FCurveEditor& CurveEditor, const FCurveEditorScreenSpace& ScreenSpace, TArray<TTuple<double, double>>& InterpolatingPoints) const
+{
+	// Constraint Channels don't draw any lines so there's no need to fill out the Interpolating Points array.
+}
+
+void FConstraintChannelCurveModel::GetKeys(const FCurveEditor& CurveEditor, double MinTime, double MaxTime, double MinValue, double MaxValue, TArray<FKeyHandle>& OutKeyHandles) const
+{
+	FMovieSceneConstraintChannel* Channel = ChannelHandle.Get();
+	UMovieSceneSection* Section = WeakSection.Get();
+
+	if (Channel && Section)
+	{
+		FFrameRate TickResolution = Section->GetTypedOuter<UMovieScene>()->GetTickResolution();
+
+		TMovieSceneChannelData<bool> ChannelData = Channel->GetData();
+		TArrayView<const FFrameNumber>          Times = ChannelData.GetTimes();
+
+		const FFrameNumber StartFrame = MinTime <= MIN_int32 ? MIN_int32 : (MinTime * TickResolution).CeilToFrame();
+		const FFrameNumber EndFrame = MaxTime >= MAX_int32 ? MAX_int32 : (MaxTime * TickResolution).FloorToFrame();
+
+		const int32 StartingIndex = Algo::LowerBound(Times, StartFrame);
+		const int32 EndingIndex = Algo::UpperBound(Times, EndFrame);
+
+		for (int32 KeyIndex = StartingIndex; KeyIndex < EndingIndex; ++KeyIndex)
+		{
+			// Space Channels don't have Values associated with them so we ignore the Min/Max value and always return the key.
+			OutKeyHandles.Add(ChannelData.GetHandle(KeyIndex));
+		}
+	}
+}
+
+//we actually still draw the 'key' as a the border which you can move and right click, context menu with.
+void FConstraintChannelCurveModel::GetKeyDrawInfo(ECurvePointType PointType, const FKeyHandle InKeyHandle, FKeyDrawInfo& OutDrawInfo) const
+{
+	OutDrawInfo.Brush = FAppStyle::Get().GetBrush("FilledBorder");
+	// NOTE Y is set to SCurveEditorKeyBarView::TrackHeight
+	OutDrawInfo.ScreenSize = FVector2D(10, 24);
+}
+
+void FConstraintChannelCurveModel::GetKeyPositions(TArrayView<const FKeyHandle> InKeys, TArrayView<FKeyPosition> OutKeyPositions) const
+{
+	FMovieSceneConstraintChannel* Channel = ChannelHandle.Get();
+	UMovieSceneSection* Section = WeakSection.Get();
+
+	if (Channel && Section)
+	{
+		FFrameRate TickResolution = Section->GetTypedOuter<UMovieScene>()->GetTickResolution();
+
+		TMovieSceneChannelData<bool> ChannelData = Channel->GetData();
+		TArrayView<const FFrameNumber>          Times = ChannelData.GetTimes();
+
+		for (int32 Index = 0; Index < InKeys.Num(); ++Index)
+		{
+			int32 KeyIndex = ChannelData.GetIndex(InKeys[Index]);
+			if (KeyIndex != INDEX_NONE)
+			{
+				OutKeyPositions[Index].InputValue = Times[KeyIndex] / TickResolution;
+				// Spaces have no values so we just output zero.
+				OutKeyPositions[Index].OutputValue = 0.0;
+			}
+		}
+	}
+}
+
+void FConstraintChannelCurveModel::SetKeyPositions(TArrayView<const FKeyHandle> InKeys, TArrayView<const FKeyPosition> InKeyPositions, EPropertyChangeType::Type ChangeType)
+{
+	FMovieSceneConstraintChannel* Channel = ChannelHandle.Get();
+	UMovieSceneSection* Section = WeakSection.Get();
+
+	if (Channel && Section)
+	{
+		Section->Modify();
+
+		FFrameRate TickResolution = Section->GetTypedOuter<UMovieScene>()->GetTickResolution();
+
+		TMovieSceneChannelData<bool> ChannelData = Channel->GetData();
+		for (int32 Index = 0; Index < InKeys.Num(); ++Index)
+		{
+			int32 KeyIndex = ChannelData.GetIndex(InKeys[Index]);
+			if (KeyIndex != INDEX_NONE)
+			{
+				FFrameNumber NewTime = (InKeyPositions[Index].InputValue * TickResolution).FloorToFrame();
+
+				KeyIndex = ChannelData.MoveKey(KeyIndex, NewTime);
+				Section->ExpandToFrame(NewTime);
+			}
+		}
+	}
+}
+
+void FConstraintChannelCurveModel::GetKeyAttributes(TArrayView<const FKeyHandle> InKeys, TArrayView<FKeyAttributes> OutAttributes) const
+{
+}
+
+void FConstraintChannelCurveModel::SetKeyAttributes(TArrayView<const FKeyHandle> InKeys, TArrayView<const FKeyAttributes> InAttributes, EPropertyChangeType::Type ChangeType)
+{
+}
+
+void FConstraintChannelCurveModel::GetCurveAttributes(FCurveAttributes& OutCurveAttributes) const
+{
+	// Constraint Channels have no Pre/Post Extrapolation
+	OutCurveAttributes.SetPreExtrapolation(RCCE_None);
+	OutCurveAttributes.SetPostExtrapolation(RCCE_None);
+}
+
+void FConstraintChannelCurveModel::SetCurveAttributes(const FCurveAttributes& InCurveAttributes)
+{
+}
+
+void FConstraintChannelCurveModel::CreateKeyProxies(TArrayView<const FKeyHandle> InKeyHandles, TArrayView<UObject*> OutObjects)
+{
+}
+
+void FConstraintChannelCurveModel::GetTimeRange(double& MinTime, double& MaxTime) const
+{
+	FMovieSceneConstraintChannel* Channel = ChannelHandle.Get();
+	UMovieSceneSection* Section = WeakSection.Get();
+
+	if (Channel && Section)
+	{
+		TArrayView<const FFrameNumber> Times = Channel->GetData().GetTimes();
+		if (Times.Num() == 0)
+		{
+			MinTime = 0.f;
+			MaxTime = 0.f;
+		}
+		else
+		{
+			FFrameRate TickResolution = Section->GetTypedOuter<UMovieScene>()->GetTickResolution();
+			double ToTime = TickResolution.AsInterval();
+			MinTime = static_cast<double> (Times[0].Value) * ToTime;
+			MaxTime = static_cast<double>(Times[Times.Num() - 1].Value) * ToTime;
+		}
+	}
+}
+
+void FConstraintChannelCurveModel::GetValueRange(double& MinValue, double& MaxValue) const
+{
+	// Constraint channels have no value, thus their value range is just zero.
+	MinValue = MaxValue = 0.0;
+}
+
+int32 FConstraintChannelCurveModel::GetNumKeys() const
+{
+	FMovieSceneConstraintChannel* Channel = ChannelHandle.Get();
+
+	if (Channel)
+	{
+		return Channel->GetData().GetTimes().Num();
+	}
+
+	return 0;
+}
+
+void FConstraintChannelCurveModel::BuildContextMenu(const FCurveEditor& CurveEditor, FMenuBuilder& MenuBuilder, TOptional<FCurvePointHandle> ClickedPoint)
+{
+	if (IsReadOnly())
+	{
+		return;
+	}
+	if (ClickedPoint.IsSet())
+	{
+		MenuBuilder.BeginSection("SpaceContextMenuSection", NSLOCTEXT("CurveEditor", "Constraint", "Constraint"));
+		{
+			FMovieSceneConstraintChannel* Channel = ChannelHandle.Get();
+			const FKeyHandle KeyHandle = ClickedPoint.GetValue().KeyHandle;
+			if (Channel && KeyHandle != FKeyHandle::Invalid())
+			{
+				TMovieSceneChannelData<bool> ChannelData = Channel->GetData();
+				int32 KeyIndex = ChannelData.GetIndex(KeyHandle);
+				TArrayView<const bool> Values = ChannelData.GetValues();
+				if (KeyIndex >= 0 && KeyIndex < Values.Num())
+				{
+					//need baking dialog for 
+					//FControlRigSpaceChannelHelpers::OpenBakeDialog(WeakSequencer.Pin().Get(), Channel, KeyIndex, WeakSection.Get());
+				}
+
+			}
+		}
+		MenuBuilder.EndSection();
+	}
+}
diff --git a/Engine/Source/Editor/MovieSceneTools/Private/Channels/ConstraintChannelEditor.cpp b/Engine/Source/Editor/MovieSceneTools/Private/Channels/ConstraintChannelEditor.cpp
--- a/Engine/Source/Editor/MovieSceneTools/Private/Channels/ConstraintChannelEditor.cpp
+++ b/Engine/Source/Editor/MovieSceneTools/Private/Channels/ConstraintChannelEditor.cpp
@@ -1,7 +1,7 @@
 ﻿// Copyright Epic Games, Inc. All Rights Reserved.
 
 #include "ConstraintChannelEditor.h"
-
+#include "Channels/ConstraintChannelCurveModel.h"
 #include "ConstraintChannel.h"
 
 #include "KeyBarCurveModel.h"
@@ -165,49 +165,61 @@ private:
 
 }
 
-void DrawExtra(
+TArray<FKeyBarCurveModel::FBarRange> FConstraintChannelEditor::GetBarRanges(
 	FMovieSceneConstraintChannel* Channel,
-	const UMovieSceneSection* Owner,
-	const FGeometry& AllottedGeometry,
-	FSequencerSectionPainter& Painter)
+	const UMovieSceneSection* Owner)
 {
-	if (!Owner || !Channel)
-	{
-		return;
-	}
+	TArray<FKeyBarCurveModel::FBarRange> Ranges;
 
-	// get index range
 	TArray<IndexRange> IndexRanges = GetIndexRanges(Channel);
 	if (IndexRanges.IsEmpty())
 	{
-		return;
+		return Ranges;
 	}
 
 	// convert to bar range
-	TArray<FKeyBarCurveModel::FBarRange> Ranges;
 	const TMovieSceneChannelData<const bool> ChannelData = Channel->GetData();
 	const TArrayView<const FFrameNumber> Times = ChannelData.GetTimes();
-
-	// const bool IsFirstValueTrue = IndexRanges[0].GetLowerBoundValue() == 0;
-	
-	for (const IndexRange& ActiveRange: IndexRanges)
+	const FConstraintChannelLabelBuilder LabelBuilder(Owner, Channel);
+	const FText& Label = LabelBuilder.Label;
+	for (const IndexRange& ActiveRange : IndexRanges)
 	{
 		FKeyBarCurveModel::FBarRange BarRange;
 		FFrameRate TickResolution = Owner->GetTypedOuter<UMovieScene>()->GetTickResolution();
 		double LowerValue = Times[ActiveRange.GetLowerBoundValue()] / TickResolution;
 		double UpperValue = Times[ActiveRange.GetUpperBoundValue()] / TickResolution;
-	
+
 		BarRange.Range.SetLowerBound(TRangeBound<double>(LowerValue));
 		BarRange.Range.SetUpperBound(TRangeBound<double>(UpperValue));
-	
-		BarRange.Name = "Constraint";
+
+		BarRange.Name = FName(*Label.ToString());
 		BarRange.Color = FLinearColor(.2, .5, .1);
 		static FLinearColor ZebraTint = FLinearColor::White.CopyWithNewOpacity(0.01f);
 		BarRange.Color = BarRange.Color * (1.f - ZebraTint.A) + ZebraTint * ZebraTint.A;
-		
+		BarRange.bRangeIsInfinite = false;
 		Ranges.Add(BarRange);
 	}
 
+	return Ranges;
+}
+
+void DrawExtra(
+	FMovieSceneConstraintChannel* Channel,
+	const UMovieSceneSection* Owner,
+	const FGeometry& AllottedGeometry,
+	FSequencerSectionPainter& Painter)
+{
+	if (!Owner || !Channel)
+	{
+		return;
+	}
+	TArray<FKeyBarCurveModel::FBarRange> Ranges = FConstraintChannelEditor::GetBarRanges(Channel,Owner);
+
+	if (Ranges.IsEmpty())
+	{
+		return;
+	}
+	
 	// draw bars + labels
 	const FSlateBrush* WhiteBrush = FAppStyle::GetBrush("WhiteBrush");
 	static constexpr ESlateDrawEffect DrawEffects = ESlateDrawEffect::None;
@@ -233,7 +245,7 @@ void DrawExtra(
 		const FKeyBarCurveModel::FBarRange& Range = Ranges[Index];
 
 		// draw bar
-		double LowerSeconds = /*(Index == 0 && IsFirstValueTrue) ? InputMin : */Range.Range.GetLowerBoundValue();
+		double LowerSeconds = Range.Range.GetLowerBoundValue();
 		double UpperSeconds = Range.Range.GetUpperBoundValue();
 		if (UpperSeconds == Range.Range.GetLowerBoundValue())
 		{
@@ -299,4 +311,9 @@ TSharedRef<SWidget> CreateKeyEditor(
 	TWeakPtr<ISequencer> InSequencer)
 {
 	return SNullWidget::NullWidget;
+}
+
+TUniquePtr<FCurveModel> CreateCurveEditorModel(const TMovieSceneChannelHandle<FMovieSceneConstraintChannel>& Channel, UMovieSceneSection* OwningSection, TSharedRef<ISequencer> InSequencer)
+{
+	return MakeUnique<FConstraintChannelCurveModel>(Channel, OwningSection, InSequencer);
 }
\ No newline at end of file
diff --git a/Engine/Source/Editor/MovieSceneTools/Private/Channels/ConstraintChannelEditor.h b/Engine/Source/Editor/MovieSceneTools/Private/Channels/ConstraintChannelEditor.h
--- a/Engine/Source/Editor/MovieSceneTools/Private/Channels/ConstraintChannelEditor.h
+++ b/Engine/Source/Editor/MovieSceneTools/Private/Channels/ConstraintChannelEditor.h
@@ -4,9 +4,16 @@
 
 #include "SequencerChannelTraits.h"
 #include "SequencerChannelInterface.h"
+#include "KeyBarCurveModel.h"
 
 struct FMovieSceneConstraintChannel;
 
+/** Should concillidate this with the FCosntraintChannelHelper once that moves out of the control rig module*/
+struct FConstraintChannelEditor
+{
+	static TArray<FKeyBarCurveModel::FBarRange> GetBarRanges(FMovieSceneConstraintChannel* Channel, const UMovieSceneSection* Owner);
+};
+
 /** Key drawing overrides */
 void DrawKeys(
 	FMovieSceneConstraintChannel* Channel,
@@ -37,4 +44,8 @@ TSharedRef<SWidget> CreateKeyEditor(
 	UMovieSceneSection* InSection,
 	const FGuid& InObjectBindingID,
 	TWeakPtr<FTrackInstancePropertyBindings> PropertyBindings,
-	TWeakPtr<ISequencer> Sequencer);
\ No newline at end of file
+	TWeakPtr<ISequencer> Sequencer);
+
+/** Curve editor models */
+inline bool SupportsCurveEditorModels(const TMovieSceneChannelHandle<FMovieSceneConstraintChannel>& Channel) { return true; }
+TUniquePtr<FCurveModel> CreateCurveEditorModel(const TMovieSceneChannelHandle<FMovieSceneConstraintChannel>& Channel, UMovieSceneSection* OwningSection, TSharedRef<ISequencer> InSequencer);
diff --git a/Engine/Source/Editor/MovieSceneTools/Private/Channels/SCurveEditorKeyBarView.cpp b/Engine/Source/Editor/MovieSceneTools/Private/Channels/SCurveEditorKeyBarView.cpp
--- a/Engine/Source/Editor/MovieSceneTools/Private/Channels/SCurveEditorKeyBarView.cpp
+++ b/Engine/Source/Editor/MovieSceneTools/Private/Channels/SCurveEditorKeyBarView.cpp
@@ -120,7 +120,11 @@ void SCurveEditorKeyBarView::DrawLabels(const FGeometry& AllottedGeometry, FSlat
 			{
 				const FKeyBarCurveModel::FBarRange& Range = Ranges[Index];
 				const double LowerSeconds = Range.Range.GetLowerBoundValue();
-				const double UpperSeconds = Range.Range.GetUpperBoundValue();
+				double UpperSeconds = Range.Range.GetUpperBoundValue();
+				if (LowerSeconds == UpperSeconds)
+				{
+					UpperSeconds = InputMax;
+				}
 				const bool bOutsideUpper = (Index != Ranges.Num() - 1) && UpperSeconds < InputMin;
 				if (LowerSeconds > InputMax || bOutsideUpper) //out of range
 				{
@@ -146,12 +150,29 @@ void SCurveEditorKeyBarView::DrawLabels(const FGeometry& AllottedGeometry, FSlat
 				// draw background box
 				if (CurveColor != FLinearColor::White)
 				{
-					const double LowerSecondsForBox = Index == 0 ? InputMin : FMath::Max<double>(InputMin, LowerSeconds);
-					const double BoxPos = CurveSpace.SecondsToScreen(LowerSecondsForBox);	
-					const FPaintGeometry BoxGeometry = AllottedGeometry.ToPaintGeometry(
-						LocalSize,
-						FSlateLayoutTransform(FVector2D(BoxPos, LaneTop))
-					);
+					FPaintGeometry BoxGeometry;
+					if (Range.bRangeIsInfinite)
+					{
+						const double LowerSecondsForBox = Index == 0 ? InputMin : FMath::Max<double>(InputMin, LowerSeconds);
+						const double BoxPos = CurveSpace.SecondsToScreen(LowerSecondsForBox);
+						BoxGeometry = AllottedGeometry.ToPaintGeometry(
+							LocalSize,
+							FSlateLayoutTransform(FVector2D(BoxPos, LaneTop))
+						);
+					}
+					else
+					{
+						const double LowerSecondsForBox = FMath::Max<double>(InputMin, LowerSeconds);
+						const double UpperSecondsForBox = FMath::Min<double>(InputMax, UpperSeconds);
+						const double BoxPos = CurveSpace.SecondsToScreen(LowerSecondsForBox);
+						const double BoxEnd = CurveSpace.SecondsToScreen(UpperSecondsForBox);
+						const FVector2D OurLocalSize(BoxEnd - BoxPos, LocalSize.Y);
+
+						BoxGeometry = AllottedGeometry.ToPaintGeometry(
+							OurLocalSize,
+							FSlateLayoutTransform(FVector2D(BoxPos, LaneTop))
+						);
+					}
 				
 					FSlateDrawElement::MakeBox(OutDrawElements, BaseLayerId + CurveViewConstants::ELayerOffset::Background, BoxGeometry, WhiteBrush, DrawEffects, CurveColor);
 				}
diff --git a/Engine/Source/Editor/MovieSceneTools/Private/MovieSceneToolsModule.cpp b/Engine/Source/Editor/MovieSceneTools/Private/MovieSceneToolsModule.cpp
--- a/Engine/Source/Editor/MovieSceneTools/Private/MovieSceneToolsModule.cpp
+++ b/Engine/Source/Editor/MovieSceneTools/Private/MovieSceneToolsModule.cpp
@@ -83,6 +83,8 @@
 #include "Channels/MovieSceneAudioTriggerChannel.h"
 #include "ConstraintChannel.h"
 #include "Channels/ConstraintChannelEditor.h"
+#include "Channels/ConstraintChannelCurveModel.h"
+#include "Channels/SCurveEditorKeyBarView.h"
 #include "Sections/MovieSceneEventSection.h"
 
 #include "MovieSceneEventUtils.h"
@@ -206,6 +208,13 @@ void FMovieSceneToolsModule::StartupModule()
 				return SNew(SCurveEditorEventChannelView, WeakCurveEditor);
 			}
 		));
+
+		FConstraintChannelCurveModel::ViewID = CurveEditorModule.RegisterView(FOnCreateCurveEditorView::CreateStatic(
+			[](TWeakPtr<FCurveEditor> WeakCurveEditor) -> TSharedRef<SCurveEditorView>
+			{
+				return SNew(SCurveEditorKeyBarView, WeakCurveEditor);
+			}
+		));
 	}
 
 	FixupPayloadParameterNameHandle = UMovieSceneEventSectionBase::FixupPayloadParameterNameEvent.AddStatic(FixupPayloadParameterNameForSection);
@@ -250,6 +259,7 @@ void FMovieSceneToolsModule::ShutdownModule()
 	if (ICurveEditorModule* CurveEditorModule = FModuleManager::GetModulePtr<ICurveEditorModule>("CurveEditor"))
 	{
 		CurveEditorModule->UnregisterView(FEventChannelCurveModel::EventView);
+		CurveEditorModule->UnregisterView(FConstraintChannelCurveModel::ViewID);
 	}
 
 	if (ISettingsModule* SettingsModule = FModuleManager::GetModulePtr<ISettingsModule>("Settings"))
diff --git a/Engine/Source/Editor/MovieSceneTools/Public/Channels/ConstraintChannelCurveModel.h b/Engine/Source/Editor/MovieSceneTools/Public/Channels/ConstraintChannelCurveModel.h
new file mode 100644
--- /dev/null
+++ b/Engine/Source/Editor/MovieSceneTools/Public/Channels/ConstraintChannelCurveModel.h
@@ -0,0 +1,65 @@
+// Copyright Epic Games, Inc. All Rights Reserved.
+
+#pragma once
+
+#include "CoreTypes.h"
+#include "UObject/WeakObjectPtr.h"
+#include "UObject/WeakObjectPtrTemplates.h"
+#include "KeyBarCurveModel.h"
+#include "ConstraintChannel.h"
+#include "Channels/MovieSceneChannelHandle.h"
+#include "Framework/MultiBox/MultiBoxBuilder.h"
+
+class UMovieSceneSection;
+class ISequencer;
+class UTickableConstraint;
+
+class FConstraintChannelCurveModel : public FKeyBarCurveModel
+{
+public:
+	static ECurveEditorViewID ViewID;
+
+	FConstraintChannelCurveModel(TMovieSceneChannelHandle<FMovieSceneConstraintChannel> InChannel, UMovieSceneSection* InOwningSection, TWeakPtr<ISequencer> InWeakSequencer);
+	~FConstraintChannelCurveModel();
+	//FCurveModel
+	virtual const void* GetCurve() const override;
+
+	virtual void Modify() override;
+
+	virtual void DrawCurve(const FCurveEditor& CurveEditor, const FCurveEditorScreenSpace& ScreenSpace, TArray<TTuple<double, double>>& InterpolatingPoints) const override;
+	virtual void GetKeys(const FCurveEditor& CurveEditor, double MinTime, double MaxTime, double MinValue, double MaxValue, TArray<FKeyHandle>& OutKeyHandles) const override;
+	virtual void GetKeyDrawInfo(ECurvePointType PointType, const FKeyHandle InKeyHandle, FKeyDrawInfo& OutDrawInfo) const override;
+
+	virtual void GetKeyPositions(TArrayView<const FKeyHandle> InKeys, TArrayView<FKeyPosition> OutKeyPositions) const override;
+	virtual void SetKeyPositions(TArrayView<const FKeyHandle> InKeys, TArrayView<const FKeyPosition> InKeyPositions, EPropertyChangeType::Type ChangeType) override;
+
+	virtual void GetKeyAttributes(TArrayView<const FKeyHandle> InKeys, TArrayView<FKeyAttributes> OutAttributes) const override;
+	virtual void SetKeyAttributes(TArrayView<const FKeyHandle> InKeys, TArrayView<const FKeyAttributes> InAttributes, EPropertyChangeType::Type ChangeType = EPropertyChangeType::Unspecified) override;
+
+	virtual void GetCurveAttributes(FCurveAttributes& OutCurveAttributes) const override;
+	virtual void SetCurveAttributes(const FCurveAttributes& InCurveAttributes) override;
+	virtual void GetTimeRange(double& MinTime, double& MaxTime) const override;
+	virtual void GetValueRange(double& MinValue, double& MaxValue) const override;
+	virtual int32 GetNumKeys() const override;
+	virtual void GetNeighboringKeys(const FKeyHandle InKeyHandle, TOptional<FKeyHandle>& OutPreviousKeyHandle, TOptional<FKeyHandle>& OutNextKeyHandle) const override {}
+	virtual bool Evaluate(double ProspectiveTime, double& OutValue) const override;
+	virtual void AddKeys(TArrayView<const FKeyPosition> InKeyPositions, TArrayView<const FKeyAttributes> InAttributes, TArrayView<TOptional<FKeyHandle>>* OutKeyHandles) override;
+	virtual void RemoveKeys(TArrayView<const FKeyHandle> InKeys) override;
+
+	virtual void CreateKeyProxies(TArrayView<const FKeyHandle> InKeyHandles, TArrayView<UObject*> OutObjects) override;
+
+	//FKeyBarCurveModel
+	virtual void BuildContextMenu(const FCurveEditor& CurveEditor, FMenuBuilder& MenuBuilder, TOptional<FCurvePointHandle> ClickedPoint);
+	virtual TArray<FKeyBarCurveModel::FBarRange> FindRanges();
+
+
+private:
+	void FixupCurve();
+private:
+	TMovieSceneChannelHandle<FMovieSceneConstraintChannel> ChannelHandle;
+	TWeakObjectPtr<UMovieSceneSection> WeakSection;
+	TWeakPtr<ISequencer> WeakSequencer;
+	TWeakObjectPtr<UTickableConstraint> Constraint;
+
+	FDelegateHandle OnDestroyHandle;
+};
\ No newline at end of file
diff --git a/Engine/Source/Runtime/MovieSceneTracks/MovieSceneTracks.Build.cs b/Engine/Source/Runtime/MovieSceneTracks/MovieSceneTracks.Build.cs
--- a/Engine/Source/Runtime/MovieSceneTracks/MovieSceneTracks.Build.cs
+++ b/Engine/Source/Runtime/MovieSceneTracks/MovieSceneTracks.Build.cs
@@ -22,6 +22,8 @@ public class MovieSceneTracks : ModuleRules
 				"TimeManagement",
 				"AnimationCore",
 				"AudioExtensions",
+				"Constraints",
+
 			});
 
 		PrivateDependencyModuleNames.AddRange(new string[]
diff --git a/Engine/Source/Runtime/MovieSceneTracks/Private/MovieSceneTracksComponentTypes.cpp b/Engine/Source/Runtime/MovieSceneTracks/Private/MovieSceneTracksComponentTypes.cpp
--- a/Engine/Source/Runtime/MovieSceneTracks/Private/MovieSceneTracksComponentTypes.cpp
+++ b/Engine/Source/Runtime/MovieSceneTracks/Private/MovieSceneTracksComponentTypes.cpp
@@ -556,6 +556,8 @@ FMovieSceneTracksComponentTypes::FMovieSceneTracksComponentTypes()
 	ComponentRegistry->NewComponentType(&QuaternionRotationChannel[1], TEXT("Quaternion Rotation Channel 1"));
 	ComponentRegistry->NewComponentType(&QuaternionRotationChannel[2], TEXT("Quaternion Rotation Channel 2"));
 
+	ComponentRegistry->NewComponentType(&ConstraintChannel, TEXT("Constraint Channel"));
+
 	ComponentRegistry->NewComponentType(&AttachParent, TEXT("Attach Parent"));
 	ComponentRegistry->NewComponentType(&AttachComponent, TEXT("Attachment Component"));
 	ComponentRegistry->NewComponentType(&AttachParentBinding, TEXT("Attach Parent Binding"));
@@ -736,6 +738,11 @@ FMovieSceneTracksComponentTypes::FMovieSceneTracksComponentTypes()
 		ComponentRegistry->Factories.DefineMutuallyInclusiveComponent(QuaternionRotationChannel[Index], BuiltInComponents->EvalTime);
 	}
 
+	// -------------------------------------------------------------------------------------------
+	// Set up constraint components
+	ComponentRegistry->Factories.DuplicateChildComponent(ConstraintChannel);
+	ComponentRegistry->Factories.DefineMutuallyInclusiveComponent(ConstraintChannel, BuiltInComponents->EvalTime);
+
 	// --------------------------------------------------------------------------------------------
 	// Set up attachment components
 	ComponentRegistry->Factories.DefineChildComponent(AttachParentBinding, AttachParent);
diff --git a/Engine/Source/Runtime/MovieSceneTracks/Private/Sections/MovieScene3DTransformSection.cpp b/Engine/Source/Runtime/MovieSceneTracks/Private/Sections/MovieScene3DTransformSection.cpp
--- a/Engine/Source/Runtime/MovieSceneTracks/Private/Sections/MovieScene3DTransformSection.cpp
+++ b/Engine/Source/Runtime/MovieSceneTracks/Private/Sections/MovieScene3DTransformSection.cpp
@@ -14,12 +14,14 @@
 #include "EntitySystem/MovieSceneEntityBuilder.h"
 #include "MovieSceneTracksComponentTypes.h"
 #include "Algo/AnyOf.h"
+#include "TransformConstraint.h"
+#include "TransformableHandle.h"
 
 #if WITH_EDITOR
 
 struct F3DTransformChannelEditorData
 {
-	F3DTransformChannelEditorData(EMovieSceneTransformChannel Mask)
+	F3DTransformChannelEditorData(EMovieSceneTransformChannel Mask,int SortOrderStart)
 	{
 		FText LocationGroup = NSLOCTEXT("MovieSceneTransformSection", "Location", "Location");
 		FText RotationGroup = NSLOCTEXT("MovieSceneTransformSection", "Rotation", "Rotation");
@@ -28,57 +30,57 @@ struct F3DTransformChannelEditorData
 			MetaData[0].SetIdentifiers("Location.X", FCommonChannelData::ChannelX, LocationGroup);
 			MetaData[0].bEnabled = EnumHasAllFlags(Mask, EMovieSceneTransformChannel::TranslationX);
 			MetaData[0].Color = FCommonChannelData::RedChannelColor;
-			MetaData[0].SortOrder = 0;
+			MetaData[0].SortOrder = SortOrderStart;
 			MetaData[0].bCanCollapseToTrack = false;
 
 			MetaData[1].SetIdentifiers("Location.Y", FCommonChannelData::ChannelY, LocationGroup);
 			MetaData[1].bEnabled = EnumHasAllFlags(Mask, EMovieSceneTransformChannel::TranslationY);
 			MetaData[1].Color = FCommonChannelData::GreenChannelColor;
-			MetaData[1].SortOrder = 1;
+			MetaData[1].SortOrder = SortOrderStart + 1;
 			MetaData[1].bCanCollapseToTrack = false;
 
 			MetaData[2].SetIdentifiers("Location.Z", FCommonChannelData::ChannelZ, LocationGroup);
 			MetaData[2].bEnabled = EnumHasAllFlags(Mask, EMovieSceneTransformChannel::TranslationZ);
 			MetaData[2].Color = FCommonChannelData::BlueChannelColor;
-			MetaData[2].SortOrder = 2;
+			MetaData[2].SortOrder = SortOrderStart + 2;
 			MetaData[2].bCanCollapseToTrack = false;
 		}
 		{
 			MetaData[3].SetIdentifiers("Rotation.X", NSLOCTEXT("MovieSceneTransformSection", "RotationX", "Roll"), RotationGroup);
 			MetaData[3].bEnabled = EnumHasAllFlags(Mask, EMovieSceneTransformChannel::RotationX);
 			MetaData[3].Color = FCommonChannelData::RedChannelColor;
-			MetaData[3].SortOrder = 3;
+			MetaData[3].SortOrder = SortOrderStart + 3;
 			MetaData[3].bCanCollapseToTrack = false;
 
 			MetaData[4].SetIdentifiers("Rotation.Y", NSLOCTEXT("MovieSceneTransformSection", "RotationY", "Pitch"), RotationGroup);
 			MetaData[4].bEnabled = EnumHasAllFlags(Mask, EMovieSceneTransformChannel::RotationY);
 			MetaData[4].Color = FCommonChannelData::GreenChannelColor;
-			MetaData[4].SortOrder = 4;
+			MetaData[4].SortOrder = SortOrderStart + 4;
 			MetaData[4].bCanCollapseToTrack = false;
 
 			MetaData[5].SetIdentifiers("Rotation.Z", NSLOCTEXT("MovieSceneTransformSection", "RotationZ", "Yaw"), RotationGroup);
 			MetaData[5].bEnabled = EnumHasAllFlags(Mask, EMovieSceneTransformChannel::RotationZ);
 			MetaData[5].Color = FCommonChannelData::BlueChannelColor;
-			MetaData[5].SortOrder = 5;
+			MetaData[5].SortOrder = SortOrderStart + 5;
 			MetaData[5].bCanCollapseToTrack = false;
 		}
 		{
 			MetaData[6].SetIdentifiers("Scale.X", FCommonChannelData::ChannelX, ScaleGroup);
 			MetaData[6].bEnabled = EnumHasAllFlags(Mask, EMovieSceneTransformChannel::ScaleX);
 			MetaData[6].Color = FCommonChannelData::RedChannelColor;
-			MetaData[6].SortOrder = 6;
+			MetaData[6].SortOrder = SortOrderStart + 6;
 			MetaData[6].bCanCollapseToTrack = false;
 
 			MetaData[7].SetIdentifiers("Scale.Y", FCommonChannelData::ChannelY, ScaleGroup);
 			MetaData[7].bEnabled = EnumHasAllFlags(Mask, EMovieSceneTransformChannel::ScaleY);
 			MetaData[7].Color = FCommonChannelData::GreenChannelColor;
-			MetaData[7].SortOrder = 7;
+			MetaData[7].SortOrder = SortOrderStart + 7;
 			MetaData[7].bCanCollapseToTrack = false;
 
 			MetaData[8].SetIdentifiers("Scale.Z", FCommonChannelData::ChannelZ, ScaleGroup);
 			MetaData[8].bEnabled = EnumHasAllFlags(Mask, EMovieSceneTransformChannel::ScaleZ);
 			MetaData[8].Color = FCommonChannelData::BlueChannelColor;
-			MetaData[8].SortOrder = 8;
+			MetaData[8].SortOrder = SortOrderStart + 8;
 			MetaData[8].bCanCollapseToTrack = false;
 		}
 		{
@@ -297,6 +299,11 @@ void FMovieScene3DTransformKeyStruct::PropagateChanges(const FPropertyChangedEve
 	KeyStructInterop.Apply(Time);
 }
 
+namespace UE::MovieScene
+{
+	static constexpr uint32 ConstraintTypeMask = 1 << 31;
+}
+
 
 /* UMovieScene3DTransformSection interface
  *****************************************************************************/
@@ -419,8 +426,32 @@ void UMovieScene3DTransformSection::BuildEntity(BaseBuilderType& InBaseBuilder,
 	);
 }
 
+void UMovieScene3DTransformSection::PopulateConstraintEntities(const TRange<FFrameNumber>& EffectiveRange, const FMovieSceneEvaluationFieldEntityMetaData& InMetaData, FMovieSceneEntityComponentFieldBuilder* OutFieldBuilder)
+{
+	check(Constraints);
+
+	using namespace UE::MovieScene;
+
+	const int32 MetaDataIndex = OutFieldBuilder->AddMetaData(InMetaData);
+
+	// Add explicitly typed entities for each constraint
+	// Encode the top bit of the entity ID to mean it is a constraint
+	// We can check this in ImportEntityImpl to import the correct constraint
+	for (int32 ConstraintIndex = 0; ConstraintIndex < Constraints->ConstraintsChannels.Num(); ++ConstraintIndex)
+	{
+		// Entity IDs for constraints are their index within the array, masked with the top bit set
+		const uint32 EntityID = ConstraintIndex | ConstraintTypeMask;
+		OutFieldBuilder->AddPersistentEntity(EffectiveRange, this, EntityID, MetaDataIndex);
+	}
+}
+
+
 bool UMovieScene3DTransformSection::PopulateEvaluationFieldImpl(const TRange<FFrameNumber>& EffectiveRange, const FMovieSceneEvaluationFieldEntityMetaData& InMetaData, FMovieSceneEntityComponentFieldBuilder* OutFieldBuilder)
 {
+	if (Constraints)
+	{
+		PopulateConstraintEntities(EffectiveRange, InMetaData, OutFieldBuilder);
+	}
 	if (OverrideRegistry)
 	{
 		return OverrideRegistry->PopulateEvaluationFieldImpl(EffectiveRange, InMetaData, OutFieldBuilder, *this);
@@ -429,10 +460,49 @@ bool UMovieScene3DTransformSection::PopulateEvaluationFieldImpl(const TRange<FFr
 	return false;
 }
 
+void UMovieScene3DTransformSection::ImportConstraintEntity(UMovieSceneEntitySystemLinker* EntityLinker, const FEntityImportParams& Params, FImportedEntity* OutImportedEntity)
+{
+	using namespace UE::MovieScene;
+
+	FBuiltInComponentTypes*          BuiltInComponentTypes = FBuiltInComponentTypes::Get();
+	FMovieSceneTracksComponentTypes* TrackComponents       = FMovieSceneTracksComponentTypes::Get();
+
+	// Constraints must always operate on a USceneComponent. Putting one on a generic FTransform property has no effect (or is not possible).
+	FGuid ObjectBindingID = Params.GetObjectBindingID();
+	if (ObjectBindingID.IsValid())
+	{
+		// Mask out the top bit to get the constraint index we encoded into the EntityID.
+		const int32 ConstraintIndex = static_cast<int32>(Params.EntityID & ~ConstraintTypeMask);
+
+		checkf(Constraints->ConstraintsChannels.IsValidIndex(ConstraintIndex), TEXT("Encoded constraint (%d) index is not valid within array size %d. Data must have been manipulated without re-compilaition."), ConstraintIndex, Constraints->ConstraintsChannels.Num());
+
+		if (Constraints->ConstraintsChannels[ConstraintIndex].Constraint.IsValid())
+		{
+			FName ConstraintName = Constraints->ConstraintsChannels[ConstraintIndex].Constraint->GetFName();
+
+			FConstraintComponentData ComponentData;
+			ComponentData.ConstraintName = ConstraintName;
+			ComponentData.Channel = &(Constraints->ConstraintsChannels[ConstraintIndex].ActiveChannel);
+			OutImportedEntity->AddBuilder(
+				FEntityBuilder()
+				.Add(BuiltInComponentTypes->SceneComponentBinding, ObjectBindingID)
+				//.Add(TrackComponents->ConstraintName, ConstraintName)
+				.Add(TrackComponents->ConstraintChannel,ComponentData)
+			);
+		}
+	}
+}
+
 void UMovieScene3DTransformSection::ImportEntityImpl(UMovieSceneEntitySystemLinker* EntityLinker, const FEntityImportParams& Params, FImportedEntity* OutImportedEntity)
 {
 	using namespace UE::MovieScene;
 
+	if ((Params.EntityID & ConstraintTypeMask) != 0  && Constraints) 
+	{
+		ImportConstraintEntity(EntityLinker, Params, OutImportedEntity);
+		return;
+	}
+
 	FBuiltInComponentTypes*   BuiltInComponentTypes = FBuiltInComponentTypes::Get();
 	UMovieScenePropertyTrack* Track                 = GetTypedOuter<UMovieScenePropertyTrack>();
 
@@ -531,9 +601,46 @@ EMovieSceneChannelProxyType UMovieScene3DTransformSection::CacheChannelProxy()
 {
 	FMovieSceneChannelProxyData Channels;
 
+	//Constraints go on top
+	int32 NumConstraints = 0;
+	if(Constraints)
+	{ 
+		for (FConstraintAndActiveChannel& ConstraintChannel : Constraints->ConstraintsChannels)
+		{
 #if WITH_EDITOR
+			if (ConstraintChannel.Constraint.IsValid())
+			{
+				FText ConstraintGroup = NSLOCTEXT("MovieSceneTransformSection", "Constraints", "Constraints");
+				const FName& ConstraintName = ConstraintChannel.Constraint->GetFName();
+				ConstraintChannel.ActiveChannel.ExtraLabel = [WeakConstraint = MakeWeakObjectPtr(ConstraintChannel.Constraint.Get())]
+				{
+					if (WeakConstraint.IsValid())
+					{
+						FString ParentStr; WeakConstraint->GetLabel().Split(TEXT("."), &ParentStr, nullptr);
+						if (!ParentStr.IsEmpty())
+						{
+							return ParentStr;
+						}
+					}
+					static const FString DummyStr;
+					return DummyStr;
+				};
+				const FText DisplayText = FText::FromString(ConstraintChannel.Constraint->GetTypeLabel());
+				FMovieSceneChannelMetaData MetaData(ConstraintName, DisplayText, ConstraintGroup, true /*bInEnabled*/);
+				MetaData.SortOrder = NumConstraints++;
+				MetaData.bCanCollapseToTrack = false;
+				Channels.Add(ConstraintChannel.ActiveChannel, MetaData, TMovieSceneExternalValue<bool>());
+			}
+#else
+			Channels.Add(ConstraintChannel.ActiveChannel);
+#endif
+			
+		}
+	}
 
-	F3DTransformChannelEditorData EditorData(TransformMask.GetChannels());
+#if WITH_EDITOR
+
+	F3DTransformChannelEditorData EditorData(TransformMask.GetChannels(), NumConstraints);
 
 	Channels.Add(Translation[0], EditorData.MetaData[0], EditorData.ExternalValues[0]);
 	Channels.Add(Translation[1], EditorData.MetaData[1], EditorData.ExternalValues[1]);
@@ -705,3 +812,72 @@ void UMovieScene3DTransformSection::SetBlendType(EMovieSceneBlendType InBlendTyp
 		}
 	}
 }
+
+bool UMovieScene3DTransformSection::HasConstraintChannel(const FName& InConstraintName) const
+{
+	if (Constraints)
+	{
+		return Constraints->ConstraintsChannels.ContainsByPredicate([InConstraintName](const FConstraintAndActiveChannel& InChannel)
+			{
+				return InChannel.Constraint.IsValid() ? InChannel.Constraint->GetFName() == InConstraintName : false;
+			});
+	}
+	return false;
+}
+
+FConstraintAndActiveChannel* UMovieScene3DTransformSection::GetConstraintChannel(const FName& InConstraintName)
+{
+	if (Constraints)
+	{
+		const int32 Index = Constraints->ConstraintsChannels.IndexOfByPredicate([InConstraintName](const FConstraintAndActiveChannel& InChannel)
+			{
+				return InChannel.Constraint.IsValid() ? InChannel.Constraint->GetFName() == InConstraintName : false;
+			});
+		return (Index != INDEX_NONE) ? &(Constraints->ConstraintsChannels[Index]) : nullptr;
+	}
+	return nullptr;
+}
+
+void UMovieScene3DTransformSection::SetUpConstraintRemovedHandle()
+{
+	if (Constraints)
+	{
+		if (Constraints->ConstraintsChannels.Num() > 0 && !Constraints->OnConstraintRemovedHandle.IsValid())
+		{
+			FConstraintsManagerController& Controller = FConstraintsManagerController::Get(GetWorld());
+			Constraints->OnConstraintRemovedHandle = Controller.OnConstraintRemoved().AddLambda([this](FName InConstraintName)
+			{
+				Modify();
+				const int32 NumRemoved = Constraints->ConstraintsChannels.RemoveAll([InConstraintName](const FConstraintAndActiveChannel& InChannel)
+					{
+						return InChannel.Constraint.IsValid() ? InChannel.Constraint->GetFName() == InConstraintName : false;
+					});
+				if (NumRemoved > 0)
+				{
+					CacheChannelProxy();
+				}
+			});
+		}
+	}
+}
+
+void UMovieScene3DTransformSection::AddConstraintChannel(UTickableConstraint* InConstraint)
+{
+	if (!Constraints)
+	{
+		Constraints = NewObject<UMovieScene3dTransformSectionConstraints>(this, NAME_None, RF_Transactional);
+	}
+	if (!HasConstraintChannel(InConstraint->GetFName()))
+	{
+		Modify();
+
+		const int32 NewIndex = Constraints->ConstraintsChannels.Add(FConstraintAndActiveChannel(InConstraint));
+
+		FMovieSceneConstraintChannel* ExistingChannel = &Constraints->ConstraintsChannels[NewIndex].ActiveChannel;
+		ExistingChannel->SetDefault(false);
+
+		SetUpConstraintRemovedHandle();
+		CacheChannelProxy();
+	}
+}
+
diff --git a/Engine/Source/Runtime/MovieSceneTracks/Private/Systems/MovieSceneConstraintSystem.cpp b/Engine/Source/Runtime/MovieSceneTracks/Private/Systems/MovieSceneConstraintSystem.cpp
new file mode 100644
--- /dev/null
+++ b/Engine/Source/Runtime/MovieSceneTracks/Private/Systems/MovieSceneConstraintSystem.cpp
@@ -0,0 +1,136 @@
+// Copyright Epic Games, Inc. All Rights Reserved.
+
+#include "Systems/MovieSceneConstraintSystem.h"
+#include "ConstraintsManager.h"
+#include "MovieSceneTracksComponentTypes.h"
+#include "Sections/MovieScene3DTransformSection.h"
+#include "EntitySystem/MovieSceneBoundObjectInstantiator.h"
+#include "EntitySystem/MovieSceneBoundSceneComponentInstantiator.h"
+
+#include "PreAnimatedState/MovieScenePreAnimatedComponentTransformStorage.h"
+#include "Evaluation/PreAnimatedState/MovieScenePreAnimatedStorageID.inl"
+
+#include "EntitySystem/BuiltInComponentTypes.h"
+#include "MovieSceneTracksComponentTypes.h"
+#include "Systems/MovieSceneComponentTransformSystem.h"
+
+class UTickableConstraint;
+
+namespace UE::MovieScene
+{
+
+struct FPreAnimatedConstraint
+{
+	TWeakObjectPtr<UTickableConstraint> WeakConstraint;
+	bool bPreviouslyEnabled;
+};
+
+struct FPreAnimatedConstraintTraits : FBoundObjectPreAnimatedStateTraits
+{
+	using KeyType = TTuple<FObjectKey, FName>;
+	using StorageType = FPreAnimatedConstraint;
+
+	static FPreAnimatedConstraint CachePreAnimatedValue(UObject* InBoundObject, const FName& ConstraintName)
+	{
+		USceneComponent* SceneComponent = CastChecked<USceneComponent>(InBoundObject);
+		FConstraintsManagerController& Controller = FConstraintsManagerController::Get(SceneComponent->GetWorld());
+
+		// @todo: I don't know how to store pre-animated values - do we need to enable/disable constraints or add/remove them?
+		//        Maybe we just need to also cache some other state other than Constraint->Active
+		UTickableConstraint* Constraint = Controller.GetConstraint(ConstraintName);
+		return FPreAnimatedConstraint{ Constraint, Constraint ? Constraint->Active : false };
+	}
+
+	static void RestorePreAnimatedValue(const TTuple<FObjectKey, FName>& InKey, const FPreAnimatedConstraint& OldValue, const FRestoreStateParams& Params)
+	{
+		if (UTickableConstraint* Constraint = OldValue.WeakConstraint.Get())
+		{
+			Constraint->Active = OldValue.bPreviouslyEnabled;
+		}
+	}
+};
+
+struct FPreAnimatedConstraintStorage
+	: public TPreAnimatedStateStorage<FPreAnimatedConstraintTraits>
+{
+	static TAutoRegisterPreAnimatedStorageID<FPreAnimatedConstraintStorage> StorageID;
+};
+
+TAutoRegisterPreAnimatedStorageID<FPreAnimatedConstraintStorage> FPreAnimatedConstraintStorage::StorageID;
+
+} // namespace UE::MovieScene
+
+UMovieSceneConstraintSystem::UMovieSceneConstraintSystem(const FObjectInitializer& ObjInit)
+	: Super(ObjInit)
+{
+	using namespace UE::MovieScene;
+
+	FMovieSceneTracksComponentTypes* TracksComponents = FMovieSceneTracksComponentTypes::Get();
+
+	//RelevantComponent = TracksComponents->ConstraintName;
+	RelevantComponent = TracksComponents->ConstraintChannel;
+
+	// Run constraints during instantiation or evaluation
+	Phase = ESystemPhase::Instantiation | ESystemPhase::Evaluation;
+
+	if (HasAnyFlags(RF_ClassDefaultObject) ) 
+	{
+		// Constraints must be evaluated before their transforms are evaluated.
+		// This is only really necessary if they are in the same phase (which they are not), but I've
+		// defined the prerequisite for safety if its phase changes in future
+		DefineImplicitPrerequisite(GetClass(), UMovieSceneComponentTransformSystem::StaticClass());
+		DefineComponentConsumer(GetClass(), UE::MovieScene::FBuiltInComponentTypes::Get()->EvalTime);
+	}
+}
+
+void UMovieSceneConstraintSystem::OnRun(FSystemTaskPrerequisites& InPrerequisites, FSystemSubsequentTasks& Subsequents)
+{
+	using namespace UE::MovieScene;
+
+	FMovieSceneEntitySystemRunner* ActiveRunner = Linker->GetActiveRunner();
+	ESystemPhase CurrentPhase = ActiveRunner->GetCurrentPhase();
+
+
+	if (CurrentPhase == ESystemPhase::Instantiation)
+	{
+		// Save pre-animated state
+		//TSharedPtr<FPreAnimatedConstraintStorage> PreAnimatedStorage = Linker->PreAnimatedState.GetOrCreateStorage<FPreAnimatedConstraintStorage>();
+
+		//andrew, had to remove this, last parameter not the correct variadic type
+		//PreAnimatedStorage->BeginTrackingAndCachePreAnimatedValues(Linker, BuiltInComponents->BoundObject, TracksComponents->ConstraintName);	
+	}
+	else if (CurrentPhase == ESystemPhase::Evaluation)
+	{
+		FConstraintsManagerController& Controller = FConstraintsManagerController::Get(GetWorld());
+		FMovieSceneTracksComponentTypes* TracksComponents = FMovieSceneTracksComponentTypes::Get();
+		FBuiltInComponentTypes* BuiltInComponents = FBuiltInComponentTypes::Get();
+
+		struct FEvaluateConstraintChannels
+		{
+			FEvaluateConstraintChannels(FConstraintsManagerController* InController) :Controller(InController) { check(Controller); };
+			void ForEachEntity(UObject* BoundObject, const FConstraintComponentData& ConstraintChannel, FFrameTime FrameTime)
+			{
+				UTickableConstraint* Constraint = Controller->GetConstraint(ConstraintChannel.ConstraintName);
+				if (Constraint)
+				{
+					bool Result;
+					if (ConstraintChannel.Channel->Evaluate(FrameTime, Result))
+					{
+						Constraint->Active = Result;
+					}
+				}
+			}
+			FConstraintsManagerController* Controller;
+		};
+
+		// Set up new constraints
+		FEntityTaskBuilder()
+			.SetDesiredThread(Linker->EntityManager.GetGatherThread())
+			.Read(BuiltInComponents->BoundObject)
+			.Read(TracksComponents->ConstraintChannel)
+			.Read(BuiltInComponents->EvalTime)
+			.Dispatch_PerEntity<FEvaluateConstraintChannels>(&Linker->EntityManager, InPrerequisites, &Subsequents, &Controller);
+	}
+}
+
+
diff --git a/Engine/Source/Runtime/MovieSceneTracks/Public/MovieSceneTracksComponentTypes.h b/Engine/Source/Runtime/MovieSceneTracks/Public/MovieSceneTracksComponentTypes.h
--- a/Engine/Source/Runtime/MovieSceneTracks/Public/MovieSceneTracksComponentTypes.h
+++ b/Engine/Source/Runtime/MovieSceneTracks/Public/MovieSceneTracksComponentTypes.h
@@ -73,6 +73,18 @@ struct FMovieSceneDataLayerComponentData
 	TObjectPtr<const UMovieSceneDataLayerSection> Section = nullptr;
 };
 
+/** Component data for the constraint system */
+USTRUCT()
+struct FConstraintComponentData
+{
+	GENERATED_BODY()
+
+	UPROPERTY()
+	FName ConstraintName;
+
+	FMovieSceneBoolChannel* Channel;
+};
+
 namespace UE
 {
 namespace MovieScene
@@ -483,6 +495,8 @@ struct MOVIESCENETRACKS_API FMovieSceneTracksComponentTypes
 	TPropertyComponents<FComponentTransformPropertyTraits> ComponentTransform;
 	TComponentTypeID<FSourceDoubleChannel> QuaternionRotationChannel[3];
 
+	TComponentTypeID<FConstraintComponentData> ConstraintChannel;
+
 	TComponentTypeID<USceneComponent*> AttachParent;
 	TComponentTypeID<FAttachmentComponent> AttachComponent;
 	TComponentTypeID<FMovieSceneObjectBindingID> AttachParentBinding;
diff --git a/Engine/Source/Runtime/MovieSceneTracks/Public/Sections/MovieScene3DTransformSection.h b/Engine/Source/Runtime/MovieSceneTracks/Public/Sections/MovieScene3DTransformSection.h
--- a/Engine/Source/Runtime/MovieSceneTracks/Public/Sections/MovieScene3DTransformSection.h
+++ b/Engine/Source/Runtime/MovieSceneTracks/Public/Sections/MovieScene3DTransformSection.h
@@ -14,6 +14,8 @@
 #include "EntitySystem/IMovieSceneEntityProvider.h"
 #include "TransformData.h"
 #include "Misc/LargeWorldCoordinates.h"
+#include "ConstraintsManager.h"
+#include "ConstraintChannel.h"
 #include "MovieScene3DTransformSection.generated.h"
 
 #if WITH_EDITORONLY_DATA
@@ -208,6 +210,24 @@ private:
 	uint32 Mask;
 };
 
+/**
+* This object contains information needed for constraint channels on the transform section
+*/
+UCLASS()
+class MOVIESCENETRACKS_API UMovieScene3dTransformSectionConstraints : public UObject
+{
+	GENERATED_BODY()
+
+public:
+
+	/** Constraint Channels*/
+	UPROPERTY()
+	TArray<FConstraintAndActiveChannel> ConstraintsChannels;
+
+	FDelegateHandle OnConstraintRemovedHandle;
+};
+
+
 /**
  * A 3D transform section
  */
@@ -265,6 +285,8 @@ private:
 
 	template<typename BaseBuilderType>
 	void BuildEntity(BaseBuilderType& InBaseBuilder, UMovieSceneEntitySystemLinker* Linker, const FEntityImportParams& Params, FImportedEntity* OutImportedEntity);
+	void PopulateConstraintEntities(const TRange<FFrameNumber>& EffectiveRange, const FMovieSceneEvaluationFieldEntityMetaData& InMetaData, FMovieSceneEntityComponentFieldBuilder* OutFieldBuilder);
+	void ImportConstraintEntity(UMovieSceneEntitySystemLinker* EntityLinker, const FEntityImportParams& Params, FImportedEntity* OutImportedEntity);
 
 private:
 
@@ -291,12 +313,35 @@ private:
 	UPROPERTY()
 	TObjectPtr<UMovieSceneSectionChannelOverrideRegistry> OverrideRegistry;
 
+	/** Optional pointer to constraint channels*/
+	UPROPERTY()
+	TObjectPtr<UMovieScene3dTransformSectionConstraints> Constraints;
+
 	/** Whether to use a quaternion linear interpolation between keys. This finds the 'shortest' rotation between keyed orientations. */
 	UPROPERTY(EditAnywhere, DisplayName = "Use Quaternion Interpolation", Category = "Rotation")
 	bool bUseQuaternionInterpolation;
 
 public:
 
+	/*
+	* If it has that constraint with that Name
+	*/
+	MOVIESCENETRACKS_API bool HasConstraintChannel(const FName& InConstraintName) const;
+
+	/*
+	* Get constraint with that name
+	*/
+	MOVIESCENETRACKS_API FConstraintAndActiveChannel* GetConstraintChannel(const FName& InConstraintName);
+
+	/*
+	*  Add Cosntraint channel
+	*/
+	MOVIESCENETRACKS_API void AddConstraintChannel(UTickableConstraint* InConstraint);
+
+private:
+
+	void SetUpConstraintRemovedHandle();
+
 #if WITH_EDITORONLY_DATA
 
 public:
diff --git a/Engine/Source/Runtime/MovieSceneTracks/Public/Systems/MovieSceneConstraintSystem.h b/Engine/Source/Runtime/MovieSceneTracks/Public/Systems/MovieSceneConstraintSystem.h
new file mode 100644
--- /dev/null
+++ b/Engine/Source/Runtime/MovieSceneTracks/Public/Systems/MovieSceneConstraintSystem.h
@@ -0,0 +1,21 @@
+// Copyright Epic Games, Inc. All Rights Reserved.
+
+#pragma once
+
+#include "EntitySystem/MovieSceneEntitySystem.h"
+#include "MovieSceneConstraintSystem.generated.h"
+
+/**
+ * System that is responsible for propagating constraints to a bound object's FConstraintsManagerController.
+ */
+UCLASS()
+class MOVIESCENETRACKS_API UMovieSceneConstraintSystem : public UMovieSceneEntitySystem
+{
+public:
+
+	GENERATED_BODY()
+
+	UMovieSceneConstraintSystem(const FObjectInitializer& ObjInit);
+
+	virtual void OnRun(FSystemTaskPrerequisites& InPrerequisites, FSystemSubsequentTasks& Subsequents) override;
+};
