Control Rig: Fix crash when opening CR_Quinn_Example in ControlRigExample

diff --git a/Engine/Plugins/Animation/ControlRig/Source/ControlRigDeveloper/Private/ControlRigBlueprint.cpp b/Engine/Plugins/Animation/ControlRig/Source/ControlRigDeveloper/Private/ControlRigBlueprint.cpp
--- a/Engine/Plugins/Animation/ControlRig/Source/ControlRigDeveloper/Private/ControlRigBlueprint.cpp
+++ b/Engine/Plugins/Animation/ControlRig/Source/ControlRigDeveloper/Private/ControlRigBlueprint.cpp
@@ -1388,18 +1388,25 @@ void UControlRigBlueprint::RefreshAllModels(EControlRigBlueprintLoadType InLoadT
 				{
 					if (URigVMFunctionReferenceNode* FunctionReferenceNode = Cast<URigVMFunctionReferenceNode>(LibraryNode))
 					{
-						if (FunctionReferenceNode->GetLibrary() != GetLocalFunctionLibrary())
+						if (URigVMLibraryNode* ReferencedNode = Cast<URigVMLibraryNode>(FunctionReferenceNode->GetReferencedNode()))
 						{
-							UPackage* Package = FunctionReferenceNode->GetLibrary()->GetPackage();
-							if (!Package->IsFullyLoaded())
+							if (ReferencedNode->GetLibrary() != GetLocalFunctionLibrary())
 							{
-								Package->FullyLoad();
+								ReferencedNode->GetPackage()->FullyLoad();
+							}
+							else
+							{
+								ContainedGraphs.Add(ReferencedNode->GetContainedGraph());
 							}
-							continue;
 						}
+						
+						continue;
 					}
-					
-					ContainedGraphs.Add(LibraryNode->GetContainedGraph());					
+
+					if (URigVMGraph* ContainedGraph = LibraryNode->GetContainedGraph())
+					{
+						ContainedGraphs.Add(ContainedGraph);
+					}					
 				}
 			}
 			
diff --git a/Engine/Source/Developer/RigVMDeveloper/Private/RigVMModel/Nodes/RigVMFunctionReferenceNode.cpp b/Engine/Source/Developer/RigVMDeveloper/Private/RigVMModel/Nodes/RigVMFunctionReferenceNode.cpp
--- a/Engine/Source/Developer/RigVMDeveloper/Private/RigVMModel/Nodes/RigVMFunctionReferenceNode.cpp
+++ b/Engine/Source/Developer/RigVMDeveloper/Private/RigVMModel/Nodes/RigVMFunctionReferenceNode.cpp
@@ -169,7 +169,11 @@ TArray<FRigVMExternalVariable> URigVMFunctionReferenceNode::GetExternalVariables
 
 const FRigVMTemplate* URigVMFunctionReferenceNode::GetTemplate() const
 {
-	return GetReferencedNode()->GetTemplate();
+	if (URigVMLibraryNode* ReferencedNode = GetReferencedNode())
+	{
+		return ReferencedNode->GetTemplate();
+	}
+	return nullptr;
 }
 
 TArray<FRigVMExternalVariable> URigVMFunctionReferenceNode::GetExternalVariables(bool bRemapped) const
