Control Rig: Avoid Hierarchy ResetToDefault for FK Control Rig

diff --git a/Engine/Plugins/Animation/ControlRig/Source/ControlRig/Private/ControlRig.cpp b/Engine/Plugins/Animation/ControlRig/Source/ControlRig/Private/ControlRig.cpp
--- a/Engine/Plugins/Animation/ControlRig/Source/ControlRig/Private/ControlRig.cpp
+++ b/Engine/Plugins/Animation/ControlRig/Source/ControlRig/Private/ControlRig.cpp
@@ -956,8 +956,11 @@ bool UControlRig::Execute(const EControlRigState InState, const FName& InEventNa
 						{
 							InitialPose = GetHierarchy()->GetPose(true, ERigElementType::ToResetAfterConstructionEvent, FRigElementKeyCollection());
 						}
-						
-						GetHierarchy()->ResetToDefault();
+
+						if(bCopyHierarchyBeforeConstruction)
+						{
+							GetHierarchy()->ResetToDefault();
+						}
 
 						if(InitialPose.Num() > 0)
 						{
