Control Rig: When dragging a Control Rig to sequencer ignore the mesh's default rig

diff --git a/Engine/Plugins/Animation/ControlRig/Source/ControlRigEditor/Private/ControlRigBlueprintActions.cpp b/Engine/Plugins/Animation/ControlRig/Source/ControlRigEditor/Private/ControlRigBlueprintActions.cpp
--- a/Engine/Plugins/Animation/ControlRig/Source/ControlRigEditor/Private/ControlRigBlueprintActions.cpp
+++ b/Engine/Plugins/Animation/ControlRig/Source/ControlRigEditor/Private/ControlRigBlueprintActions.cpp
@@ -32,6 +32,7 @@
 #include "SBlueprintDiff.h"
 #include "Misc/MessageDialog.h"
 #include "LevelSequenceEditorBlueprintLibrary.h"
+#include "Sequencer/ControlRigParameterTrackEditor.h"
 
 #define LOCTEXT_NAMESPACE "ControlRigBlueprintActions"
 
@@ -290,6 +291,8 @@ void FControlRigBlueprintActions::OnSpawnedSkeletalMeshActorChanged(UObject* InO
 	}
 	UClass* ControlRigClass = RigBlueprint->GeneratedClass;
 
+	TGuardValue<bool> DisableTrackCreation(FControlRigParameterTrackEditor::bAutoGenerateControlRigTrack, false);
+
 	// find a level sequence in the world, if can't find that, create one
 	ULevelSequence* Sequence = ULevelSequenceEditorBlueprintLibrary::GetFocusedLevelSequence();
 	if (Sequence == nullptr)
diff --git a/Engine/Plugins/Animation/ControlRig/Source/ControlRigEditor/Private/Sequencer/ControlRigParameterTrackEditor.cpp b/Engine/Plugins/Animation/ControlRig/Source/ControlRigEditor/Private/Sequencer/ControlRigParameterTrackEditor.cpp
--- a/Engine/Plugins/Animation/ControlRig/Source/ControlRigEditor/Private/Sequencer/ControlRigParameterTrackEditor.cpp
+++ b/Engine/Plugins/Animation/ControlRig/Source/ControlRigEditor/Private/Sequencer/ControlRigParameterTrackEditor.cpp
@@ -193,6 +193,8 @@ static USkeleton* AcquireSkeletonFromObjectGuid(const FGuid& Guid, UObject** Obj
 	return nullptr;
 }
 
+bool FControlRigParameterTrackEditor::bAutoGenerateControlRigTrack = true;
+
 FControlRigParameterTrackEditor::FControlRigParameterTrackEditor(TSharedRef<ISequencer> InSequencer)
 	: FKeyframeTrackEditor<UMovieSceneControlRigParameterTrack>(InSequencer)
 	, bCurveDisplayTickIsPending(false)
@@ -1734,7 +1736,7 @@ void FControlRigParameterTrackEditor::AddTrackForComponent(USceneComponent* InCo
 {
 	if (USkeletalMeshComponent* SkelMeshComp = Cast<USkeletalMeshComponent>(InComponent))
 	{
-		if (!SkelMeshComp->GetDefaultAnimatingRig().IsNull())
+		if(bAutoGenerateControlRigTrack && !SkelMeshComp->GetDefaultAnimatingRig().IsNull())
 		{
 			UObject* Object = SkelMeshComp->GetDefaultAnimatingRig().LoadSynchronous();
 			if (Object != nullptr && (Object->IsA<UControlRigBlueprint>() || Object->IsA<UControlRigComponent>()))
diff --git a/Engine/Plugins/Animation/ControlRig/Source/ControlRigEditor/Private/Sequencer/ControlRigParameterTrackEditor.h b/Engine/Plugins/Animation/ControlRig/Source/ControlRigEditor/Private/Sequencer/ControlRigParameterTrackEditor.h
--- a/Engine/Plugins/Animation/ControlRig/Source/ControlRigEditor/Private/Sequencer/ControlRigParameterTrackEditor.h
+++ b/Engine/Plugins/Animation/ControlRig/Source/ControlRigEditor/Private/Sequencer/ControlRigParameterTrackEditor.h
@@ -242,6 +242,11 @@ private:
 
 	/** A transaction used to group multiple key events */
 	TSharedPtr<FScopedTransaction> ControlUndoTransaction;
+
+	/** Controls if the control rig track for the default animating rig should be created */
+	static bool bAutoGenerateControlRigTrack;
+
+	friend class FControlRigBlueprintActions;
 };
 
 
