Control Rig: Update Name Lists after Construction Event

diff --git a/Engine/Plugins/Animation/ControlRig/Source/ControlRig/Public/Rigs/RigHierarchyController.h b/Engine/Plugins/Animation/ControlRig/Source/ControlRig/Public/Rigs/RigHierarchyController.h
--- a/Engine/Plugins/Animation/ControlRig/Source/ControlRig/Public/Rigs/RigHierarchyController.h
+++ b/Engine/Plugins/Animation/ControlRig/Source/ControlRig/Public/Rigs/RigHierarchyController.h
@@ -697,6 +697,8 @@ private:
 	{
 		T* Element = URigHierarchy::NewElement<T>();
 		Element->CreatedAtInstructionIndex = CurrentInstructionIndex;
+		Element->MetadataChangedDelegate.BindUObject(GetHierarchy(), &URigHierarchy::OnMetadataChanged);
+		Element->MetadataTagChangedDelegate.BindUObject(GetHierarchy(), &URigHierarchy::OnMetadataTagChanged);
 		return Element;
 	}
 	
