Sequencer: Control Rig: Selected proxy controls will include driven controls in selected control rig filter

diff --git a/Engine/Plugins/Animation/ControlRig/Source/ControlRigEditor/Private/Sequencer/ControlRigSequencerFilter.cpp b/Engine/Plugins/Animation/ControlRig/Source/ControlRigEditor/Private/Sequencer/ControlRigSequencerFilter.cpp
--- a/Engine/Plugins/Animation/ControlRig/Source/ControlRigEditor/Private/Sequencer/ControlRigSequencerFilter.cpp
+++ b/Engine/Plugins/Animation/ControlRig/Source/ControlRigEditor/Private/Sequencer/ControlRigSequencerFilter.cpp
@@ -168,19 +168,32 @@ public:
 		if (Track)
 		{
 			UControlRig *ControlRig = Track->GetControlRig();
-			if (ControlRig)
+			if (ControlRig && ControlRig->GetHierarchy())
 			{
 				FName Name(*InText.ToString());
-				TArray<FName> ControlNames = ControlRig->CurrentControlSelection();
-				for (const FName& ControlName : ControlNames)
+				TArray<const FRigBaseElement*> SelectedControls = ControlRig->GetHierarchy()->GetSelectedElements(ERigElementType::Control);
+				for (const FRigBaseElement* SelectedControl : SelectedControls)
 				{
-					if (Name == ControlName)
+					if (Name == SelectedControl->GetName())
 					{
 						return true;
 					}
+					if (const FRigControlElement* ControlElement = Cast<FRigControlElement>(SelectedControl))
+					{
+						if (ControlElement->Settings.AnimationType == ERigControlAnimationType::ProxyControl)
+						{
+							const TArray<FRigElementKey>& DrivenControls = ControlElement->Settings.DrivenControls;
+							for (const FRigElementKey& DrivenKey : DrivenControls)
+							{
+								if (Name == DrivenKey.Name)
+								{
+									return true;
+								}
+							}
+						}
+					}
 				}
 			}
-
 		}
 		return false;
 	}
