fixed space / constraints popping when playing animation vs scrubbing/stepping thru frames

diff --git a/Engine/Plugins/Animation/ControlRig/Source/ControlRig/Private/Sequencer/MovieSceneControlRigParameterTemplate.cpp b/Engine/Plugins/Animation/ControlRig/Source/ControlRig/Private/Sequencer/MovieSceneControlRigParameterTemplate.cpp
--- a/Engine/Plugins/Animation/ControlRig/Source/ControlRig/Private/Sequencer/MovieSceneControlRigParameterTemplate.cpp
+++ b/Engine/Plugins/Animation/ControlRig/Source/ControlRig/Private/Sequencer/MovieSceneControlRigParameterTemplate.cpp
@@ -1771,11 +1771,17 @@ void FMovieSceneControlRigParameterTemplate::EvaluateCurvesWithMasks(const FMovi
 			Values.ScalarValues.Emplace(Scalar.ParameterName, Value);
 		}
 
+		// when playing animation, instead of scrubbing/stepping thru frames, InTime might have a subframe of 0.999928
+		// leading to a decimal value of 24399.999928 (for example). This results in evaluating one frame less than expected
+		// (24399 instead of 24400) and leads to spaces and constraints switching parents/state after the control changes
+		// its transform. (float/double channels will interpolate to a value pretty close to the one at 24400 as its based
+		// on that 0.999928 subframe value.
+		const FFrameTime RoundTime = Time.RoundToFrame();
 		for (int32 Index = 0; Index < Spaces.Num(); ++Index)
 		{
 			FMovieSceneControlRigSpaceBaseKey Value;
 			const FSpaceControlNameAndChannel& Space = Spaces[Index];
-			Space.SpaceCurve.Evaluate(Time, Value);
+			Space.SpaceCurve.Evaluate(RoundTime, Value);
 			
 			Values.SpaceValues.Emplace(Space.ControlName, Value);
 		}
@@ -1784,7 +1790,7 @@ void FMovieSceneControlRigParameterTemplate::EvaluateCurvesWithMasks(const FMovi
 		{
 			bool Value = false;
 			const FConstraintAndActiveChannel& ConstraintAndActiveChannel = Constraints[Index];
-			ConstraintAndActiveChannel.ActiveChannel.Evaluate(Time, Value);
+			ConstraintAndActiveChannel.ActiveChannel.Evaluate(RoundTime, Value);
 			
 			Values.ConstraintsValues.Emplace(ConstraintAndActiveChannel.Constraint.Get(), Value);
 		}
