Control Rig: Centralize hierarchy delegate setup

diff --git a/Engine/Plugins/Animation/ControlRig/Source/ControlRig/Private/Rigs/RigHierarchy.cpp b/Engine/Plugins/Animation/ControlRig/Source/ControlRig/Private/Rigs/RigHierarchy.cpp
--- a/Engine/Plugins/Animation/ControlRig/Source/ControlRig/Private/Rigs/RigHierarchy.cpp
+++ b/Engine/Plugins/Animation/ControlRig/Source/ControlRig/Private/Rigs/RigHierarchy.cpp
@@ -4098,13 +4098,6 @@ FRigBaseElement* URigHierarchy::MakeElement(ERigElementType InElementType, int32
 		}
 	}
 
-	if (Element)
-	{
-		Element->OwnedInstances = InCount;
-		Element->MetadataChangedDelegate.BindUObject(this, &URigHierarchy::OnMetadataChanged);
-		Element->MetadataTagChangedDelegate.BindUObject(this, &URigHierarchy::OnMetadataTagChanged);
-	}
-
 	return Element;
 }
 
diff --git a/Engine/Plugins/Animation/ControlRig/Source/ControlRig/Public/Rigs/RigHierarchy.h b/Engine/Plugins/Animation/ControlRig/Source/ControlRig/Public/Rigs/RigHierarchy.h
--- a/Engine/Plugins/Animation/ControlRig/Source/ControlRig/Public/Rigs/RigHierarchy.h
+++ b/Engine/Plugins/Animation/ControlRig/Source/ControlRig/Public/Rigs/RigHierarchy.h
@@ -2890,7 +2890,7 @@ public:
 	 */
 	FORCEINLINE void IncrementMetadataVersion(const FRigElementKey& InKey, const FName& InName)
 	{
-		MetadataVersion += 1 + (int32)GetTypeHash(InName);
+		MetadataVersion += 1 + (int32)HashCombine(GetTypeHash(InKey), GetTypeHash(InName));
 	}
 
 	/**
@@ -2903,7 +2903,7 @@ public:
 	 */
 	FORCEINLINE void IncrementMetadataTagVersion(const FRigElementKey& InKey, const FName& InTag, bool bAdded)
 	{
-		MetadataTagVersion += 1 + (int32)GetTypeHash(InTag);
+		MetadataTagVersion += 1 + (int32)HashCombine(GetTypeHash(InKey), GetTypeHash(InTag));
 	}
 
 	/**
@@ -3479,15 +3479,17 @@ private:
 	 * Templated helper function to create an element
 	 */
 	template<typename ElementType = FRigBaseElement>
-	FORCEINLINE static ElementType* NewElement(int32 Num = 1)
+	FORCEINLINE ElementType* NewElement(int32 Num = 1)
 	{
-		ElementType* Elements = (ElementType*)FMemory::Malloc(sizeof(ElementType) * Num);
+		ElementType* NewElements = (ElementType*)FMemory::Malloc(sizeof(ElementType) * Num);
 		for(int32 Index=0;Index<Num;Index++)
 		{
-			new(&Elements[Index]) ElementType();
+			new(&NewElements[Index]) ElementType();
+			NewElements[Index].MetadataChangedDelegate.BindStatic(&URigHierarchy::OnMetadataChanged_Static, this);
+			NewElements[Index].MetadataTagChangedDelegate.BindStatic(&URigHierarchy::OnMetadataTagChanged_Static, this);
 		}
-		Elements[0].OwnedInstances = Num;
-		return Elements;
+		NewElements[0].OwnedInstances = Num;
+		return NewElements;
 	}
 
 	/**
@@ -3964,6 +3966,19 @@ private:
 	void OnMetadataTagChanged(const FRigElementKey& InKey, const FName& InTag, bool bAdded);
 
 protected:
+
+	FORCEINLINE static void OnMetadataChanged_Static(const FRigElementKey& InKey, const FName& InName, URigHierarchy* InHierarchy)
+	{
+		check(InHierarchy);
+		check(IsValid(InHierarchy));
+		InHierarchy->OnMetadataChanged(InKey, InName);
+	}
+	FORCEINLINE static void OnMetadataTagChanged_Static(const FRigElementKey& InKey, const FName& InTag, bool bAdded, URigHierarchy* InHierarchy)
+	{
+		check(InHierarchy);
+		check(IsValid(InHierarchy));
+		InHierarchy->OnMetadataTagChanged(InKey, InTag, bAdded);
+	}
 	
 	bool bEnableCacheValidityCheck;
 
diff --git a/Engine/Plugins/Animation/ControlRig/Source/ControlRig/Public/Rigs/RigHierarchyController.h b/Engine/Plugins/Animation/ControlRig/Source/ControlRig/Public/Rigs/RigHierarchyController.h
--- a/Engine/Plugins/Animation/ControlRig/Source/ControlRig/Public/Rigs/RigHierarchyController.h
+++ b/Engine/Plugins/Animation/ControlRig/Source/ControlRig/Public/Rigs/RigHierarchyController.h
@@ -695,10 +695,8 @@ private:
 	template<typename T>
 	T* MakeElement()
 	{
-		T* Element = URigHierarchy::NewElement<T>();
+		T* Element = GetHierarchy()->NewElement<T>();
 		Element->CreatedAtInstructionIndex = CurrentInstructionIndex;
-		Element->MetadataChangedDelegate.BindUObject(GetHierarchy(), &URigHierarchy::OnMetadataChanged);
-		Element->MetadataTagChangedDelegate.BindUObject(GetHierarchy(), &URigHierarchy::OnMetadataTagChanged);
 		return Element;
 	}
 	
