Control Rig: Fix crash when running Run CitySampleWin64EditorPerf Incr1

diff --git a/Engine/Plugins/Animation/ControlRig/Source/ControlRigDeveloper/Private/ControlRigBlueprint.cpp b/Engine/Plugins/Animation/ControlRig/Source/ControlRigDeveloper/Private/ControlRigBlueprint.cpp
--- a/Engine/Plugins/Animation/ControlRig/Source/ControlRigDeveloper/Private/ControlRigBlueprint.cpp
+++ b/Engine/Plugins/Animation/ControlRig/Source/ControlRigDeveloper/Private/ControlRigBlueprint.cpp
@@ -1433,6 +1433,7 @@ void UControlRigBlueprint::RefreshAllModels(EControlRigBlueprintLoadType InLoadT
 		URigVMGraph* GraphToClean = SortedGraphsToClean[GraphIndex];
 		URigVMController* Controller = GetOrCreateController(GraphToClean);
 		TGuardValue<bool> RecomputeGuard(Controller->bSuspendRecomputingOuterTemplateFilters, true);
+		Controller->SuspendNotifications(true);
 		
 		for(URigVMNode* ModelNode : GraphToClean->GetNodes())
 		{
@@ -1462,6 +1463,7 @@ void UControlRigBlueprint::RefreshAllModels(EControlRigBlueprintLoadType InLoadT
 				Controller->UpdateLibraryTemplate(LibraryNode, false);
 			}
 		}
+
 #if WITH_EDITOR
 
 		if(bIsPostLoad)
@@ -1472,6 +1474,8 @@ void UControlRigBlueprint::RefreshAllModels(EControlRigBlueprintLoadType InLoadT
 			}
 		}
 #endif
+
+		Controller->SuspendNotifications(false);
 	}
 }
 
