Sequencer:Control Rig. Need to clear out delegates for the all space channels when a rig is unbound.

diff --git a/Engine/Plugins/Animation/ControlRig/Source/ControlRigEditor/Private/Sequencer/ControlRigParameterTrackEditor.cpp b/Engine/Plugins/Animation/ControlRig/Source/ControlRigEditor/Private/Sequencer/ControlRigParameterTrackEditor.cpp
--- a/Engine/Plugins/Animation/ControlRig/Source/ControlRigEditor/Private/Sequencer/ControlRigParameterTrackEditor.cpp
+++ b/Engine/Plugins/Animation/ControlRig/Source/ControlRigEditor/Private/Sequencer/ControlRigParameterTrackEditor.cpp
@@ -437,15 +437,22 @@ void FControlRigParameterTrackEditor::UnbindControlRig(UControlRig* ControlRig)
 		ControlRig->OnInitialized_AnyThread().RemoveAll(this);
 		ControlRig->ControlSelected().RemoveAll(this);
 		BoundControlRigs.Remove(ControlRig);
+		ClearOutAllSpaceDelegates(ControlRig);
 	}
 }
 void FControlRigParameterTrackEditor::UnbindAllControlRigs()
 {
+	ClearOutAllSpaceDelegates();
 	for (TWeakObjectPtr<UControlRig>& ObjectPtr : BoundControlRigs)
 	{
 		if (ObjectPtr.IsValid())
 		{
 			UControlRig* ControlRig = ObjectPtr.Get();
+			UMovieSceneControlRigParameterTrack* Track = FindTrack(ControlRig);
+			if (Track)
+			{
+				Track->SpaceChannelAdded().RemoveAll(this);
+			}
 			ControlRig->ControlModified().RemoveAll(this);
 			ControlRig->OnInitialized_AnyThread().RemoveAll(this);
 			ControlRig->ControlSelected().RemoveAll(this);
@@ -2282,6 +2289,40 @@ void FControlRigParameterTrackEditor::HandleSpaceKeyMoved(UMovieSceneControlRigP
 	}
 }
 
+void FControlRigParameterTrackEditor::ClearOutAllSpaceDelegates(UControlRig* InOptionalControlRig)
+{
+	if (GetSequencer().IsValid())
+	{
+		UMovieScene* MovieScene = GetSequencer()->GetFocusedMovieSceneSequence()->GetMovieScene();
+		const TArray<FMovieSceneBinding>& Bindings = MovieScene->GetBindings();
+		for (const FMovieSceneBinding& Binding : Bindings)
+		{
+			if (UMovieSceneControlRigParameterTrack* Track = Cast<UMovieSceneControlRigParameterTrack>(MovieScene->FindTrack(UMovieSceneControlRigParameterTrack::StaticClass(), Binding.GetObjectGuid(), NAME_None)))
+			{
+				if (InOptionalControlRig && Track->GetControlRig() != InOptionalControlRig)
+				{
+					continue;
+				}
+				for (UMovieSceneSection* Section : Track->GetAllSections())
+				{
+					if (Section)
+					{
+						if (UMovieSceneControlRigParameterSection* CRSection = Cast<UMovieSceneControlRigParameterSection>(Section))
+						{
+							TArray<FSpaceControlNameAndChannel>& Channels = CRSection->GetSpaceChannels();
+							for (FSpaceControlNameAndChannel& SpaceAndChannel : Channels)
+							{
+								SpaceAndChannel.SpaceCurve.OnKeyMovedEvent().Clear();
+								SpaceAndChannel.SpaceCurve.OnKeyDeletedEvent().Clear();
+							}
+						}
+					}
+				}
+			}
+		}
+	}
+}
+
 void FControlRigParameterTrackEditor::SetUpEditModeIfNeeded(UControlRig* ControlRig)
 {
 	if (ControlRig)
diff --git a/Engine/Plugins/Animation/ControlRig/Source/ControlRigEditor/Private/Sequencer/ControlRigParameterTrackEditor.h b/Engine/Plugins/Animation/ControlRig/Source/ControlRigEditor/Private/Sequencer/ControlRigParameterTrackEditor.h
--- a/Engine/Plugins/Animation/ControlRig/Source/ControlRigEditor/Private/Sequencer/ControlRigParameterTrackEditor.h
+++ b/Engine/Plugins/Animation/ControlRig/Source/ControlRigEditor/Private/Sequencer/ControlRigParameterTrackEditor.h
@@ -132,6 +132,8 @@ private:
 	void HandleOnSpaceAdded(UMovieSceneControlRigParameterSection* Section, const FName& ControlName, FMovieSceneControlRigSpaceChannel* Channel);
 	void HandleSpaceKeyDeleted(UMovieSceneControlRigParameterSection* Section, FMovieSceneControlRigSpaceChannel* Channel, const TArray<FKeyAddOrDeleteEventItem>& DeletedItems);
 	void HandleSpaceKeyMoved(UMovieSceneControlRigParameterSection* Section, FMovieSceneControlRigSpaceChannel* Channel, const  TArray<FKeyMoveEventItem>& MovedItems);
+	//if rig not set then we clear delegates for everyone
+	void ClearOutAllSpaceDelegates(UControlRig* InOptionalControlRig = nullptr);
 
 	/** Select control rig if not selected, select controls from key areas */
 	void SelectRigsAndControls(UControlRig* Subject, const TArray<const IKeyArea*>& KeyAreas);
