Control Rig: Fix nand implementation

diff --git a/Engine/Plugins/Animation/ControlRig/Source/ControlRig/Private/Units/Math/RigUnit_MathBool.cpp b/Engine/Plugins/Animation/ControlRig/Source/ControlRig/Private/Units/Math/RigUnit_MathBool.cpp
--- a/Engine/Plugins/Animation/ControlRig/Source/ControlRig/Private/Units/Math/RigUnit_MathBool.cpp
+++ b/Engine/Plugins/Animation/ControlRig/Source/ControlRig/Private/Units/Math/RigUnit_MathBool.cpp
@@ -31,6 +31,21 @@ FRigUnit_MathBoolNand_Execute()
 	Result = (!A) && (!B);
 }
 
+FRigUnit_MathBoolNand2_Execute()
+{
+	DECLARE_SCOPE_HIERARCHICAL_COUNTER_RIGUNIT()
+	Result = !(A && B);
+}
+
+FRigVMStructUpgradeInfo FRigUnit_MathBoolNand::GetUpgradeInfo() const
+{
+	FRigUnit_MathBoolNand2 NewNode;
+	NewNode.A = A;
+	NewNode.B = B;
+
+	return FRigVMStructUpgradeInfo(*this, NewNode);
+}
+
 FRigUnit_MathBoolOr_Execute()
 {
     DECLARE_SCOPE_HIERARCHICAL_COUNTER_RIGUNIT()
diff --git a/Engine/Plugins/Animation/ControlRig/Source/ControlRig/Private/Units/Math/RigUnit_MathBool.h b/Engine/Plugins/Animation/ControlRig/Source/ControlRig/Private/Units/Math/RigUnit_MathBool.h
--- a/Engine/Plugins/Animation/ControlRig/Source/ControlRig/Private/Units/Math/RigUnit_MathBool.h
+++ b/Engine/Plugins/Animation/ControlRig/Source/ControlRig/Private/Units/Math/RigUnit_MathBool.h
@@ -144,13 +144,28 @@ struct CONTROLRIG_API FRigUnit_MathBoolAnd : public FRigUnit_MathBoolBinaryAggre
 /**
  * Returns true if both conditions are false
  */
-USTRUCT(meta=(DisplayName="Nand", TemplateName="Nand"))
+USTRUCT(meta=(DisplayName="Nand", Deprecated="5.1"))
 struct CONTROLRIG_API FRigUnit_MathBoolNand : public FRigUnit_MathBoolBinaryOp
 {
 	GENERATED_BODY()
 
 	RIGVM_METHOD()
 	virtual void Execute(const FRigUnitContext& Context) override;
+
+	RIGVM_METHOD()
+	virtual FRigVMStructUpgradeInfo GetUpgradeInfo() const override;
+};
+
+/**
+ * Returns false if both conditions are true
+ */
+USTRUCT(meta=(DisplayName="Nand", TemplateName="Nand"))
+struct CONTROLRIG_API FRigUnit_MathBoolNand2 : public FRigUnit_MathBoolBinaryOp
+{
+	GENERATED_BODY()
+
+	RIGVM_METHOD()
+	virtual void Execute(const FRigUnitContext& Context) override;
 };
 
 /**
