Sequencer: Control Rig: Push compensation to next tick

diff --git a/Engine/Plugins/Animation/ControlRig/Source/ControlRigEditor/Private/Sequencer/ControlRigParameterTrackEditor.cpp b/Engine/Plugins/Animation/ControlRig/Source/ControlRigEditor/Private/Sequencer/ControlRigParameterTrackEditor.cpp
--- a/Engine/Plugins/Animation/ControlRig/Source/ControlRigEditor/Private/Sequencer/ControlRigParameterTrackEditor.cpp
+++ b/Engine/Plugins/Animation/ControlRig/Source/ControlRigEditor/Private/Sequencer/ControlRigParameterTrackEditor.cpp
@@ -1708,31 +1708,78 @@ static void EvaluateThisControl(UMovieSceneControlRigParameterSection* Section,
 		//note we don't need to evaluate the control rig, setting the value is enough
 	}
 }
-		
+
 //When a channel is changed via Sequencer we need to call SetControlValue on it so that Control Rig can handle seeing that this is a change, but just on this value
-//and then send back a key even if needed, which happens with IK/FK switches. Hopefully new IK/FK system will remove need for this at some point.
+//and then it send back a key if needed, which happens with IK/FK switches. Hopefully new IK/FK system will remove need for this at some point.
+//We also compensate since the changed control could happen at a space switch boundary.
+//Finally, since they can happen thousands of times interactively when moving a bunch of keys on a control rig we move to doing this into the next tick
+struct FChannelChangedStruct
+{
+	FTimerHandle TimerHandle;
+	bool bWasSetAlready = false;
+	TMap<UMovieSceneControlRigParameterSection*, TSet<FName>> SectionControlNames;
+};
+
 void FControlRigParameterTrackEditor::OnChannelChanged(const FMovieSceneChannelMetaData* MetaData, UMovieSceneSection* InSection)
 {
+	static FChannelChangedStruct ChannelChanged;
+
 	UMovieSceneControlRigParameterSection* Section = Cast<UMovieSceneControlRigParameterSection>(InSection);
 	TSharedPtr<ISequencer> SequencerPtr = GetSequencer();
 
 	if (Section && Section->GetControlRig() && MetaData && SequencerPtr.IsValid())
 	{
-		Section->ControlsToSet.Empty();
 		TArray<FString> StringArray;
 		FString String = MetaData->Name.ToString();
 		String.ParseIntoArray(StringArray, TEXT("."));
 		if (StringArray.Num() > 0)
 		{
 			FName ControlName(*StringArray[0]);
-			Section->ControlsToSet.Add(ControlName);
-			FFrameTime Time = SequencerPtr->GetLocalTime().Time;
-			EvaluateThisControl(Section,ControlName, Time);
-			Section->ControlsToSet.Empty();
-
-			TOptional<FFrameNumber> Optional;
-			FControlRigSpaceChannelHelpers::CompensateIfNeeded(Section->GetControlRig(), GetSequencer().Get(), Section,
-				ControlName, Optional);
+			if (ChannelChanged.bWasSetAlready == false)
+			{
+				ChannelChanged.bWasSetAlready = true;
+				ChannelChanged.SectionControlNames.Reset();
+				auto ChannelChangedDelegate = [this]()
+				{
+					if (!(FSlateApplication::Get().HasAnyMouseCaptor() || GUnrealEd->IsUserInteracting()))
+					{
+						if (ChannelChanged.TimerHandle.IsValid())
+						{
+							TSharedPtr<ISequencer> SequencerPtr = GetSequencer();
+							FFrameTime Time = SequencerPtr->GetLocalTime().Time;
+							GEditor->GetTimerManager()->ClearTimer(ChannelChanged.TimerHandle);
+							ChannelChanged.TimerHandle.Invalidate();
+							ChannelChanged.bWasSetAlready = false;
+							TOptional<FFrameNumber> Optional;
+							ISequencer* SequencerRaw = SequencerPtr.Get();
+							for (TPair <UMovieSceneControlRigParameterSection*, TSet<FName>>& Pair : ChannelChanged.SectionControlNames)
+							{
+								for (FName& ControlName : Pair.Value)
+								{
+									Pair.Key->ControlsToSet.Empty();
+									Pair.Key->ControlsToSet.Add(ControlName);
+									EvaluateThisControl(Pair.Key, ControlName, Time);
+									FControlRigSpaceChannelHelpers::CompensateIfNeeded(Pair.Key->GetControlRig(), SequencerRaw, Pair.Key,
+										ControlName, Optional);
+									Pair.Key->ControlsToSet.Empty();
+								}
+							}
+							ChannelChanged.SectionControlNames.Reset();
+						}
+					}
+				};
+				GEditor->GetTimerManager()->SetTimer(ChannelChanged.TimerHandle, ChannelChangedDelegate, 0.01f, true);			
+			}
+			if (TSet<FName>* SetOfNames = ChannelChanged.SectionControlNames.Find(Section))
+			{
+				SetOfNames->Add(ControlName);
+			}
+			else
+			{
+				TSet<FName> Names;
+				Names.Add(ControlName);
+				ChannelChanged.SectionControlNames.Add(Section, Names);
+			}
 		}
 	}
 }
