Removed suspect use of GenerateDiff

#rnx

diff --git a/Engine/Plugins/Animation/ControlRig/Source/ControlRigDeveloper/Private/ControlRigBlueprint.cpp b/Engine/Plugins/Animation/ControlRig/Source/ControlRigDeveloper/Private/ControlRigBlueprint.cpp
--- a/Engine/Plugins/Animation/ControlRig/Source/ControlRigDeveloper/Private/ControlRigBlueprint.cpp
+++ b/Engine/Plugins/Animation/ControlRig/Source/ControlRigDeveloper/Private/ControlRigBlueprint.cpp
@@ -2474,7 +2474,7 @@ void UControlRigBlueprint::PostTransacted(const FTransactionObjectEvent& Transac
 			int32 TransactionIndex = GEditor->Trans->FindTransactionIndex(TransactionEvent.GetTransactionId());
 			const FTransaction* Transaction = GEditor->Trans->GetTransaction(TransactionIndex);
 
-			if (Transaction->GenerateDiff().TransactionTitle == TEXT("Transform Gizmo"))
+			if (Transaction->GetTitle().BuildSourceString() == TEXT("Transform Gizmo"))
 			{
 				PropagatePoseFromBPToInstances();
 				return;
