delegates implementation for active keys move/deletion and channel deletion so that child's transform keys are compensated

diff --git a/Engine/Plugins/Animation/ControlRig/Source/ControlRig/Private/Sequencer/MovieSceneControlRigParameterSection.cpp b/Engine/Plugins/Animation/ControlRig/Source/ControlRig/Private/Sequencer/MovieSceneControlRigParameterSection.cpp
--- a/Engine/Plugins/Animation/ControlRig/Source/ControlRig/Private/Sequencer/MovieSceneControlRigParameterSection.cpp
+++ b/Engine/Plugins/Animation/ControlRig/Source/ControlRig/Private/Sequencer/MovieSceneControlRigParameterSection.cpp
@@ -884,7 +884,6 @@ void UMovieSceneControlRigParameterSection::PostEditImport()
 void UMovieSceneControlRigParameterSection::PostLoad()
 {
 	Super::PostLoad();
-	RegisterConstraintsHandles();
 }
 
 bool UMovieSceneControlRigParameterSection::HasScalarParameter(FName InParameterName) const
@@ -1277,39 +1276,6 @@ void UMovieSceneControlRigParameterSection::AddSpaceChannel(FName InControlName,
 	}
 }
 
-void UMovieSceneControlRigParameterSection::RegisterConstraintsHandles()
-{
-	if (OnConstraintRemovedHandle.IsValid())
-	{ 
-		return;
-	}
-
-	if (!ControlRig)
-	{
-		return;
-	}
-
-	if (ConstraintsChannels.IsEmpty())
-	{
-		return;
-	}
-	
-	// register constraint deletion handle
-	FConstraintsManagerController& Controller = FConstraintsManagerController::Get(ControlRig->GetWorld());
-	OnConstraintRemovedHandle = Controller.OnConstraintRemoved().AddLambda([this](FName InConstraintName)
-	{
-		Modify();
-		const int32 NumRemoved = ConstraintsChannels.RemoveAll( [InConstraintName](const FConstraintAndActiveChannel& InChannel)
-		{
-			return InChannel.Constraint.IsValid() ? InChannel.Constraint->GetFName() == InConstraintName : false;
-		});
-		if (NumRemoved)
-		{
-			ReconstructChannelProxy();
-		}
-	});
-}
-
 bool UMovieSceneControlRigParameterSection::HasConstraintChannel(const FName& InConstraintName) const
 {
 	return ConstraintsChannels.ContainsByPredicate( [InConstraintName](const FConstraintAndActiveChannel& InChannel)
@@ -1337,12 +1303,16 @@ void UMovieSceneControlRigParameterSection::AddConstraintChannel(UTickableConstr
 
 		FMovieSceneConstraintChannel* ExistingChannel = &ConstraintsChannels[NewIndex].ActiveChannel;
 		ExistingChannel->SetDefault(false);
+
+		if (OnConstraintChannelAdded.IsBound())
+		{
+			OnConstraintChannelAdded.Broadcast(this, ExistingChannel);
+		}
 		
 		if (bReconstructChannel)
 		{
 			ReconstructChannelProxy();
 		}
-		RegisterConstraintsHandles();
 	}
 }
 
@@ -1355,6 +1325,7 @@ void UMovieSceneControlRigParameterSection::RemoveConstraintChannel(const FName&
 
 	if (ConstraintsChannels.IsValidIndex(Index))
 	{
+		Modify();
 		ConstraintsChannels.RemoveAt(Index);
 		ReconstructChannelProxy();
 	}
@@ -1365,6 +1336,37 @@ const TArray<FConstraintAndActiveChannel>& UMovieSceneControlRigParameterSection
 	return ConstraintsChannels;
 }
 
+TArray<FConstraintAndActiveChannel>& UMovieSceneControlRigParameterSection::GetConstraintsChannels()
+{
+	return ConstraintsChannels;
+}
+
+const FName& UMovieSceneControlRigParameterSection::FindControlNameFromConstraintChannel(
+	const FMovieSceneConstraintChannel* InConstraintChannel) const
+{
+	const int32 Index = ConstraintsChannels.IndexOfByPredicate([InConstraintChannel](const FConstraintAndActiveChannel& InChannel)
+	{
+		return &(InChannel.ActiveChannel) == InConstraintChannel;
+	});
+	
+	if (Index != INDEX_NONE)
+	{
+		// look for info referencing that constraint index
+		using NameInfoIterator = TMap<FName, FChannelMapInfo>::TRangedForConstIterator;
+		for (NameInfoIterator It = ControlChannelMap.begin(); It; ++It)
+		{
+			const FChannelMapInfo& Info = It->Value;
+			if (Info.ConstraintsIndex.Contains(Index))
+			{
+				return It->Key;
+			}
+		}
+	}
+
+	static const FName DummyName = NAME_None;
+	return DummyName;
+}
+
 TArray<FSpaceControlNameAndChannel>& UMovieSceneControlRigParameterSection::GetSpaceChannels()
 {
 	return SpaceChannels;
diff --git a/Engine/Plugins/Animation/ControlRig/Source/ControlRig/Private/Sequencer/MovieSceneControlRigParameterTrack.cpp b/Engine/Plugins/Animation/ControlRig/Source/ControlRig/Private/Sequencer/MovieSceneControlRigParameterTrack.cpp
--- a/Engine/Plugins/Animation/ControlRig/Source/ControlRig/Private/Sequencer/MovieSceneControlRigParameterTrack.cpp
+++ b/Engine/Plugins/Animation/ControlRig/Source/ControlRig/Private/Sequencer/MovieSceneControlRigParameterTrack.cpp
@@ -49,7 +49,10 @@ UMovieSceneSection* UMovieSceneControlRigParameterTrack::CreateNewSection()
 	{
 		NewSection->SetBlendType(EMovieSceneBlendType::Additive);
 	}
+
 	NewSection->SpaceChannelAdded().AddUObject(this, &UMovieSceneControlRigParameterTrack::HandleOnSpaceAdded);
+	NewSection->ConstraintChannelAdded().AddUObject(this, &UMovieSceneControlRigParameterTrack::HandleOnConstraintAdded);
+
 	if (ControlRig)
 	{
 		NewSection->RecreateWithThisControlRig(ControlRig,bSetDefault);
@@ -90,6 +93,13 @@ void UMovieSceneControlRigParameterTrack::HandleOnSpaceNoLongerUsed(FMovieSceneC
 	}
 }
 
+void UMovieSceneControlRigParameterTrack::HandleOnConstraintAdded(
+	UMovieSceneControlRigParameterSection* InSection,
+	FMovieSceneConstraintChannel* InChannel) const
+{
+	OnConstraintChannelAdded.Broadcast(InSection, InChannel);
+}
+
 void UMovieSceneControlRigParameterTrack::RemoveAllAnimationData()
 {
 	Sections.Empty();
@@ -381,6 +391,12 @@ void UMovieSceneControlRigParameterTrack::ReconstructControlRig()
 					{
 						CRSection->SpaceChannelAdded().AddUObject(this, &UMovieSceneControlRigParameterTrack::HandleOnSpaceAdded);
 					}
+
+					if (!CRSection->ConstraintChannelAdded().IsBoundToObject(this))
+					{
+						CRSection->ConstraintChannelAdded().AddUObject(this, &UMovieSceneControlRigParameterTrack::HandleOnConstraintAdded);
+					}
+					
 					CRSection->RecreateWithThisControlRig(ControlRig, CRSection->GetBlendType() == EMovieSceneBlendType::Absolute);
 				}
 			}
diff --git a/Engine/Plugins/Animation/ControlRig/Source/ControlRig/Public/Sequencer/MovieSceneControlRigParameterSection.h b/Engine/Plugins/Animation/ControlRig/Source/ControlRig/Public/Sequencer/MovieSceneControlRigParameterSection.h
--- a/Engine/Plugins/Animation/ControlRig/Source/ControlRig/Public/Sequencer/MovieSceneControlRigParameterSection.h
+++ b/Engine/Plugins/Animation/ControlRig/Source/ControlRig/Public/Sequencer/MovieSceneControlRigParameterSection.h
@@ -181,8 +181,9 @@ class CONTROLRIG_API UMovieSceneControlRigParameterSection : public UMovieSceneP
 
 public:
 
-	/** Bindable event for when we add a space channel*/
+	/** Bindable events for when we add space or constraint channels. */
 	DECLARE_MULTICAST_DELEGATE_ThreeParams(FSpaceChannelAddedEvent, UMovieSceneControlRigParameterSection*, const FName&, FMovieSceneControlRigSpaceChannel*);
+	DECLARE_MULTICAST_DELEGATE_TwoParams(FConstraintChannelAddedEvent, UMovieSceneControlRigParameterSection*, FMovieSceneConstraintChannel*);
 
 	void AddEnumParameterKey(FName InParameterName, FFrameNumber InTime, uint8 InValue);
 	void AddIntegerParameterKey(FName InParameterName, FFrameNumber InTime, int32 InValue);
@@ -204,12 +205,17 @@ public:
 	
 	FSpaceChannelAddedEvent& SpaceChannelAdded() { return OnSpaceChannelAdded; }
 
+	TArray<FConstraintAndActiveChannel>& GetConstraintsChannels();
 	const TArray<FConstraintAndActiveChannel>& GetConstraintsChannels() const;
+	const FName& FindControlNameFromConstraintChannel(const FMovieSceneConstraintChannel* InConstraintChannel) const;
+	
+	FConstraintChannelAddedEvent& ConstraintChannelAdded() { return OnConstraintChannelAdded; }
 	
 	bool RenameParameterName(const FName& OldParameterName, const FName& NewParameterName);
 private:
 
 	FSpaceChannelAddedEvent OnSpaceChannelAdded;
+	FConstraintChannelAddedEvent OnConstraintChannelAdded;
 
 	/** Control Rig that controls us*/
 	UPROPERTY()
@@ -400,7 +406,6 @@ public:
 
 	/** todo */
 	FDelegateHandle OnConstraintRemovedHandle;
-	void RegisterConstraintsHandles();
 	bool HasConstraintChannel(const FName& InConstraintName) const;
 	FConstraintAndActiveChannel* GetConstraintChannel(const FName& InConstraintName);
 	void AddConstraintChannel(UTickableConstraint* InConstraint, bool bReconstructChannel);
diff --git a/Engine/Plugins/Animation/ControlRig/Source/ControlRig/Public/Sequencer/MovieSceneControlRigParameterTrack.h b/Engine/Plugins/Animation/ControlRig/Source/ControlRig/Public/Sequencer/MovieSceneControlRigParameterTrack.h
--- a/Engine/Plugins/Animation/ControlRig/Source/ControlRig/Public/Sequencer/MovieSceneControlRigParameterTrack.h
+++ b/Engine/Plugins/Animation/ControlRig/Source/ControlRig/Public/Sequencer/MovieSceneControlRigParameterTrack.h
@@ -119,6 +119,7 @@ public:
 	CONTROLRIG_API void SetTrackName(FName InName) { TrackName = InName; }
 
 	UMovieSceneControlRigParameterSection::FSpaceChannelAddedEvent& SpaceChannelAdded() { return OnSpaceChannelAdded; }
+	UMovieSceneControlRigParameterSection::FConstraintChannelAddedEvent& ConstraintChannelAdded() { return OnConstraintChannelAdded; }
 
 private:
 
@@ -129,6 +130,11 @@ private:
 	//then send this event out to the track editor.
 	UMovieSceneControlRigParameterSection::FSpaceChannelAddedEvent OnSpaceChannelAdded;
 
+	void HandleOnConstraintAdded(
+		UMovieSceneControlRigParameterSection* InSection,
+		FMovieSceneConstraintChannel* InChannel) const;
+	UMovieSceneControlRigParameterSection::FConstraintChannelAddedEvent OnConstraintChannelAdded;
+
 
 	void ReconstructControlRig();
 
diff --git a/Engine/Plugins/Animation/ControlRig/Source/ControlRigEditor/Private/Editor/Constraints/SConstraintsWidget.cpp b/Engine/Plugins/Animation/ControlRig/Source/ControlRigEditor/Private/Editor/Constraints/SConstraintsWidget.cpp
--- a/Engine/Plugins/Animation/ControlRig/Source/ControlRigEditor/Private/Editor/Constraints/SConstraintsWidget.cpp
+++ b/Engine/Plugins/Animation/ControlRig/Source/ControlRigEditor/Private/Editor/Constraints/SConstraintsWidget.cpp
@@ -285,8 +285,11 @@ void SDroppableConstraintItem::CreateConstraint(
 	{
 		if (Child != InParent)
 		{
-			if (FTransformConstraintUtils::CreateAndAddFromActors(World, InParent, Child, InConstraintType) != nullptr)
+			UTickableTransformConstraint* Constraint =
+				FTransformConstraintUtils::CreateAndAddFromActors(World, InParent, Child, InConstraintType);
+			if (Constraint)
 			{
+				FConstraintChannelHelper::SmartConstraintKey(Constraint);
 				bCreated = true;
 			}
 		}
@@ -481,15 +484,18 @@ void SEditableConstraintItem::Construct(
 			{
 				if (UTickableTransformConstraint* TransformConstraint = Cast<UTickableTransformConstraint>(Constraint))
 				{
-					// FConstraintChannelHelper::AddConstraintKey(TransformConstraint);
 					FConstraintChannelHelper::SmartConstraintKey(TransformConstraint);
 				}
 				return FReply::Handled();
 			})
-			.IsEnabled_Lambda([this]()
+			.IsEnabled_Lambda([]()
+			{
+				return FConstraintChannelHelper::IsKeyframingAvailable();
+			})
+			.Visibility_Lambda([]()
 			{
-				// TODO check if we have a focussed sequencer
-				return true;
+				const bool bIsKeyframingAvailable = FConstraintChannelHelper::IsKeyframingAvailable();
+				return bIsKeyframingAvailable ? EVisibility::Visible : EVisibility::Hidden;
 			})
 			.ToolTipText(LOCTEXT("KeyConstraintToolTip", "Add an active keyframe for that constraint."))
 			[
diff --git a/Engine/Plugins/Animation/ControlRig/Source/ControlRigEditor/Private/Sequencer/ConstraintChannelHelper.cpp b/Engine/Plugins/Animation/ControlRig/Source/ControlRigEditor/Private/Sequencer/ConstraintChannelHelper.cpp
--- a/Engine/Plugins/Animation/ControlRig/Source/ControlRigEditor/Private/Sequencer/ConstraintChannelHelper.cpp
+++ b/Engine/Plugins/Animation/ControlRig/Source/ControlRigEditor/Private/Sequencer/ConstraintChannelHelper.cpp
@@ -1,8 +1,7 @@
 ﻿// Copyright Epic Games, Inc. All Rights Reserved.
 
-#include "ConstraintChannelHelper.h"
+#include "ConstraintChannelHelper.inl"
 
-#include "ActorForWorldTransforms.h"
 #include "ControlRigSpaceChannelEditors.h"
 #include "ISequencer.h"
 
@@ -13,8 +12,6 @@
 
 #include "TransformConstraint.h"
 #include "Algo/Copy.h"
-#include "Algo/Unique.h"
-#include "ConstraintChannel.h"
 
 #include "Tools/BakingHelper.h"
 #include "Tools/ConstraintBaker.h"
@@ -59,6 +56,22 @@ namespace
 	}
 }
 
+bool FConstraintChannelHelper::IsKeyframingAvailable()
+{
+	const TWeakPtr<ISequencer> WeakSequencer = FBakingHelper::GetSequencer();
+	if (!WeakSequencer.IsValid())
+	{
+		return false;
+	}
+
+	if (!WeakSequencer.Pin()->GetFocusedMovieSceneSequence())
+	{
+		return false;
+	}
+
+	return true;
+}
+
 void FConstraintChannelHelper::SmartConstraintKey(UTickableTransformConstraint* InConstraint)
 {
 	const TWeakPtr<ISequencer> WeakSequencer = FBakingHelper::GetSequencer();
@@ -236,37 +249,42 @@ TArrayView<FMovieSceneFloatChannel*> FConstraintChannelHelper::GetTransformFloat
 	const TSharedPtr<ISequencer>& InSequencer)
 {
 	// InHandle and InSequencer are assumed to be valid at this stage so no need to check
-	
+	if (const UTransformableControlHandle* ControlHandle = Cast<UTransformableControlHandle>(InHandle))
+	{
+		const UMovieSceneControlRigParameterSection* Section = GetControlSection(ControlHandle, InSequencer);
+		return GetTransformFloatChannels(InHandle, Section);
+	}
+
+	static const TArrayView<FMovieSceneFloatChannel*> EmptyChannelsView;
+	return EmptyChannelsView;
+}
+
+TArrayView<FMovieSceneFloatChannel*> FConstraintChannelHelper::GetTransformFloatChannels(
+	const UTransformableHandle* InHandle,
+	const UMovieSceneControlRigParameterSection* InSection)
+{
 	static const TArrayView<FMovieSceneFloatChannel*> EmptyChannelsView;
 	
-	if (const UTransformableControlHandle* ControlHandle = Cast<UTransformableControlHandle>(InHandle))
+	const UTransformableControlHandle* ControlHandle = Cast<UTransformableControlHandle>(InHandle);
+	if (!ControlHandle || !InSection)
 	{
-		UMovieSceneControlRigParameterSection* Section = GetControlSection(ControlHandle, InSequencer);
-		FChannelMapInfo* pChannelIndex = nullptr;
-		FRigControlElement* ControlElement = nullptr;
-		Tie(ControlElement, pChannelIndex) = FControlRigSpaceChannelHelpers::GetControlAndChannelInfo(
-			ControlHandle->ControlRig.Get(),
-			Section,
-			ControlHandle->ControlName);
-    
-		if (!pChannelIndex || !ControlElement)
-		{
-			return EmptyChannelsView;
-		}
+		return EmptyChannelsView;	
+	}
 
-		// get the number of float channels to treat
-		const int32 NumChannels = FControlRigSpaceChannelHelpers::GetNumFloatChannels(ControlElement->Settings.ControlType);
-		if (!NumChannels)
-		{
-			return EmptyChannelsView;
-		}
+	const FChannelMapInfo* ChannelInfo = nullptr;
+	int32 NumChannels = 0;
+	Tie(ChannelInfo, NumChannels) =
+		GetInfoAndNumFloatChannels(ControlHandle->ControlRig.Get(), ControlHandle->ControlName, InSection);
 
-		// return a sub view that just represents the control's channels
-		const TArrayView<FMovieSceneFloatChannel*> FloatChannels = Section->GetChannelProxy().GetChannels<FMovieSceneFloatChannel>();
-		const int32 ChannelStartIndex = pChannelIndex->ChannelIndex;
-		return FloatChannels.Slice(ChannelStartIndex, NumChannels);
+	if (ChannelInfo == nullptr || NumChannels == 0)
+	{
+		return EmptyChannelsView;
 	}
-	return EmptyChannelsView;
+
+	// return a sub view that just represents the control's channels
+	const TArrayView<FMovieSceneFloatChannel*> FloatChannels = InSection->GetChannelProxy().GetChannels<FMovieSceneFloatChannel>();
+	const int32 ChannelStartIndex = ChannelInfo->ChannelIndex;
+	return FloatChannels.Slice(ChannelStartIndex, NumChannels);
 }
 
 TArrayView<FMovieSceneDoubleChannel*> FConstraintChannelHelper::GetTransformDoubleChannels(
@@ -285,59 +303,6 @@ TArrayView<FMovieSceneDoubleChannel*> FConstraintChannelHelper::GetTransformDoub
 	return EmptyChannelsView;
 }
 
-template<typename ChannelType>
-void GetFramesToCompensate(
-	const FMovieSceneConstraintChannel& InActiveChannel,
-	const bool InActiveValueToBeSet,
-	const FFrameNumber& InTime,
-	const TArrayView<ChannelType*>& InChannels,
-	TArray<FFrameNumber>& OutFramesAfter)
-{
-	using ChannelValueType = typename ChannelType::ChannelValueType;
-
-	const bool bHasKeys = (InActiveChannel.GetNumKeys() > 0);
-	
-	OutFramesAfter.Reset();
-
-	// add the current frame
-	OutFramesAfter.Add(InTime);
-
-	// add the next frames that need transform compensation 
-	for (const ChannelType* InChannel: InChannels)
-	{
-		const TMovieSceneChannelData<const ChannelValueType> ChannelData = InChannel->GetData();
-		const TArrayView<const FFrameNumber> Times = ChannelData.GetTimes();
-		if (!Times.IsEmpty())
-		{
-			// look for the first next key frame for this channel 
-			const int32 NextTimeIndex = Algo::UpperBound(Times, InTime);
-			if (Times.IsValidIndex(NextTimeIndex))
-			{
-				// store the time while the state is different
-				for (int32 Index = NextTimeIndex; Index < Times.Num(); ++Index)
-				{
-					if (!bHasKeys)
-					{
-						OutFramesAfter.Add(Times[Index]);
-					}
-					else
-					{
-						bool NextValue = false; InActiveChannel.Evaluate(Times[Index], NextValue);
-						if (NextValue == InActiveValueToBeSet)
-						{
-							break;
-						}
-						OutFramesAfter.Add(Times[Index]);
-					}
-				}
-			}
-		}
-	}
-
-	// uniqueness
-	OutFramesAfter.Sort();
-	OutFramesAfter.SetNum(Algo::Unique(OutFramesAfter));
-}
 namespace
 {
 
@@ -469,6 +434,85 @@ public:
         }
 	}
 
+	void ComputeLocalTransformsBeforeDeletion(
+		UWorld* InWorld,
+		const TSharedPtr<ISequencer>& InSequencer,
+		const TArray<FFrameNumber>& InFrames)
+	{
+		if (InFrames.IsEmpty())
+		{
+			return;
+		}
+		
+		const TArray<UTickableTransformConstraint*> Constraints = GetHandleTransformConstraints(InWorld);
+		// if (Constraints.IsEmpty())
+		// {
+		// 	return;
+		// }
+
+		// find last active constraint in the list that is different than the on we want to compensate for
+		auto GetLastActiveConstraint = [this, Constraints]()
+		{
+			// find last active constraint in the list that is different than the on we want to compensate for
+			const int32 LastActiveIndex = Constraints.FindLastByPredicate([this](const UTickableTransformConstraint* InConstraint)
+			{
+				return InConstraint->Active && InConstraint->bDynamicOffset && InConstraint != Constraint;
+			});
+
+			// if found, return its parent global transform
+			return LastActiveIndex > INDEX_NONE ?Constraints[LastActiveIndex] : nullptr;
+		};
+		
+		const UMovieScene* MovieScene = InSequencer->GetFocusedMovieSceneSequence()->GetMovieScene();
+        const FFrameRate TickResolution = MovieScene->GetTickResolution();
+        const EMovieScenePlayerStatus::Type PlaybackStatus = InSequencer->GetPlaybackStatus();
+
+		const int32 NumFrames = InFrames.Num();
+
+		ChildLocals.SetNum(NumFrames);
+        ChildGlobals.SetNum(NumFrames);
+		SpaceGlobals.SetNum(NumFrames);
+
+        for (int32 Index = 0; Index < NumFrames; ++Index)
+        {
+        	const FFrameNumber& FrameNumber = InFrames[Index];
+        
+        	// evaluate animation
+        	const FMovieSceneEvaluationRange EvaluationRange = FMovieSceneEvaluationRange(FFrameTime(FrameNumber), TickResolution);
+        	const FMovieSceneContext Context = FMovieSceneContext(EvaluationRange, PlaybackStatus).SetHasJumped(true);
+    
+        	InSequencer->GetEvaluationTemplate().Evaluate(Context, *InSequencer);
+        
+        	// evaluate constraints
+        	for (const UTickableTransformConstraint* InConstraint: Constraints)
+        	{
+        		InConstraint->Evaluate();
+        	}
+        
+        	// evaluate ControlRig?
+        	// ControlRig->Evaluate_AnyThread();
+
+        	FTransform& ChildLocal = ChildLocals[Index];
+        	FTransform& ChildGlobal = ChildGlobals[Index];
+        	FTransform& SpaceGlobal = SpaceGlobals[Index];
+
+        	// store child transforms        	
+        	ChildLocal = Handle->GetLocalTransform();
+        	ChildGlobal = Handle->GetGlobalTransform();
+
+        	// store constraint/parent space global transform
+            if (const UTickableTransformConstraint* LastConstraint = GetLastActiveConstraint())
+            {
+            	const ETransformConstraintType LastConstraintType =
+					static_cast<ETransformConstraintType>(LastConstraint->GetType());
+            	
+                SpaceGlobal = LastConstraint->GetParentGlobalTransform();
+                ChildLocal = FTransformConstraintUtils::ComputeRelativeTransform(
+                	ChildLocal, ChildGlobal, SpaceGlobal, LastConstraintType);
+            }
+        }
+	}
+	
 	void ComputeCompensation(UWorld* InWorld, const TSharedPtr<ISequencer>& InSequencer, const FFrameNumber& InTime)
 	{
 		const TArray<UTickableTransformConstraint*> Constraints = GetHandleTransformConstraints(InWorld);
@@ -882,4 +926,177 @@ void FConstraintChannelHelper::CompensateIfNeeded(
 	}
 }
 
+void FConstraintChannelHelper::HandleConstraintRemoved(
+	UTickableConstraint* InConstraint,
+	const FMovieSceneConstraintChannel* InConstraintChannel,
+	const TSharedPtr<ISequencer>& InSequencer,
+	const UMovieSceneSection* InSection)
+{
+	UTickableTransformConstraint* Constraint = Cast<UTickableTransformConstraint>(InConstraint);
+	if (!Constraint || !InConstraintChannel || !InSection)
+	{
+		return;
+	}
+	
+	TGuardValue<bool> CompensateGuard(bDoNotCompensate, true);
+	if (const UTransformableControlHandle* ControlHandle = Cast<UTransformableControlHandle>(Constraint->ChildTRSHandle))
+	{
+		const UMovieSceneControlRigParameterSection* Section = Cast<UMovieSceneControlRigParameterSection>(InSection);
+		if (!Section)
+		{
+			return;
+		}
+
+		const TArrayView<const FFrameNumber> Times = InConstraintChannel->GetData().GetTimes();
+		if (Times.IsEmpty())
+		{
+			return;
+		}
+		
+		UControlRig* ControlRig = ControlHandle->ControlRig.Get();
+		const FName& ControlName = ControlHandle->ControlName;
+		
+		// get transform channels
+		const TArrayView<FMovieSceneFloatChannel*> TransformChannels = GetTransformFloatChannels(ControlHandle, Section);
+		
+		// get frames after this time
+		TArray<FFrameNumber> FramesToCompensate;
+		GetFramesWithinActiveState(*InConstraintChannel, TransformChannels, FramesToCompensate);
+
+		// do the compensation
+		FCompensationEvaluator Evaluator(Constraint);
+		Evaluator.ComputeLocalTransformsBeforeDeletion(ControlRig->GetWorld(), InSequencer, FramesToCompensate);
+		const TArray<FTransform>& ChildLocals = Evaluator.ChildLocals;
+
+		const EMovieSceneTransformChannel ChannelsToKey = FConstraintBaker::GetChannelsToKey(Constraint);
+		const FFrameRate TickResolution = InSequencer->GetFocusedTickResolution();
+		
+		FBakingHelper::AddTransformKeys(ControlRig, ControlName, FramesToCompensate,
+			ChildLocals, ChannelsToKey, TickResolution);
+	}
+}
+
+void FConstraintChannelHelper::HandleConstraintKeyDeleted(
+	UTickableTransformConstraint* InConstraint,
+	const FMovieSceneConstraintChannel* InConstraintChannel,
+	const TSharedPtr<ISequencer>& InSequencer,
+	const UMovieSceneSection* InSection,
+	const FFrameNumber& InTime)
+{
+	const FFrameNumber TimeMinusOne(InTime-1);
+	
+	bool CurrentValue = false, PreviousValue = false;
+	InConstraintChannel->Evaluate(TimeMinusOne, PreviousValue);
+	InConstraintChannel->Evaluate(InTime, CurrentValue);
+
+	if (CurrentValue == PreviousValue)
+	{
+		const int32 NumKeys = InConstraintChannel->GetNumKeys();
+		if (NumKeys > 1)
+		{
+			return;
+		}
+	}
+
+	TGuardValue<bool> CompensateGuard(bDoNotCompensate, true);
+	
+	if (const UTransformableControlHandle* ControlHandle = Cast<UTransformableControlHandle>(InConstraint->ChildTRSHandle))
+	{
+		const UMovieSceneControlRigParameterSection* Section = Cast<UMovieSceneControlRigParameterSection>(InSection);
+		if (!Section)
+		{
+			return;
+		}
+		
+		UControlRig* ControlRig = ControlHandle->ControlRig.Get();
+		const FName& ControlName = ControlHandle->ControlName;
+		
+		// get transform channels
+		const TArrayView<FMovieSceneFloatChannel*> TransformChannels = GetTransformFloatChannels(ControlHandle, Section);
+		
+		// get frames after this time
+		TArray<FFrameNumber> FramesToCompensate;
+		GetFramesAfter(*InConstraintChannel, InTime, TransformChannels, FramesToCompensate);
+
+		// do the compensation
+		FCompensationEvaluator Evaluator(InConstraint);
+		Evaluator.ComputeLocalTransforms(ControlRig->GetWorld(), InSequencer, FramesToCompensate, PreviousValue);
+		TArray<FTransform>& ChildLocals = Evaluator.ChildLocals;
+		ChildLocals.RemoveAt(0);
+
+		const EMovieSceneTransformChannel ChannelsToKey = FConstraintBaker::GetChannelsToKey(InConstraint);
+		const FFrameRate TickResolution = InSequencer->GetFocusedTickResolution();
+
+		FBakingHelper::AddTransformKeys( ControlRig, ControlName, FramesToCompensate,
+			ChildLocals, ChannelsToKey, TickResolution);
+		
+		// now delete any extra TimeMinusOne
+		DeleteTransformKeys(TransformChannels, TimeMinusOne);
+	}
+}
+
+void FConstraintChannelHelper::HandleConstraintKeyMoved(
+	const UTickableTransformConstraint* InConstraint,
+	const FMovieSceneConstraintChannel* InConstraintChannel,
+	const UMovieSceneSection* InSection,
+	const FFrameNumber& InCurrentFrame, const FFrameNumber& InNextFrame)
+{
+	const FFrameNumber Delta = InNextFrame - InCurrentFrame;
+	if (Delta == 0)
+	{
+		return;
+	}
+	
+	if (!InConstraint || !InConstraintChannel || !InSection)
+	{
+		return;
+	}
+
+	if (const UTransformableControlHandle* ControlHandle = Cast<UTransformableControlHandle>(InConstraint->ChildTRSHandle))
+	{
+		const UMovieSceneControlRigParameterSection* Section = Cast<UMovieSceneControlRigParameterSection>(InSection);
+		if (!Section)
+		{
+			return;
+		}
+		
+		// get transform channels
+		const TArrayView<FMovieSceneFloatChannel*> TransformChannels = GetTransformFloatChannels(ControlHandle, Section);
+
+		// move them
+		MoveTransformKeys(TransformChannels, InCurrentFrame, InNextFrame);
+	}
+}
+
+TPair<const FChannelMapInfo*, int32> FConstraintChannelHelper::GetInfoAndNumFloatChannels(
+	const UControlRig* InControlRig,
+	const FName& InControlName,
+	const UMovieSceneControlRigParameterSection* InSection)
+{
+	const FRigControlElement* ControlElement = InControlRig ? InControlRig->FindControl(InControlName) : nullptr;
+	auto GetNumFloatChannels = [](const ERigControlType& InControlType)
+	{
+		switch (InControlType)
+		{
+		case ERigControlType::Position:
+		case ERigControlType::Scale:
+		case ERigControlType::Rotator:
+			return 3;
+		case ERigControlType::TransformNoScale:
+			return 6;
+		case ERigControlType::Transform:
+		case ERigControlType::EulerTransform:
+			return 9;
+		default:
+			break;
+		}
+		return 0;		
+	};
+	
+	const int32 NumFloatChannels = ControlElement ? GetNumFloatChannels(ControlElement->Settings.ControlType) : 0;
+	const FChannelMapInfo* ChannelInfo = InSection ? InSection->ControlChannelMap.Find(InControlName) : nullptr;
+	
+	return {ChannelInfo, NumFloatChannels};
+}
+
 #undef LOCTEXT_NAMESPACE
diff --git a/Engine/Plugins/Animation/ControlRig/Source/ControlRigEditor/Private/Sequencer/ConstraintChannelHelper.inl b/Engine/Plugins/Animation/ControlRig/Source/ControlRigEditor/Private/Sequencer/ConstraintChannelHelper.inl
new file mode 100644
--- /dev/null
+++ b/Engine/Plugins/Animation/ControlRig/Source/ControlRigEditor/Private/Sequencer/ConstraintChannelHelper.inl
@@ -0,0 +1,258 @@
+﻿// Copyright Epic Games, Inc. All Rights Reserved.
+
+#pragma once
+
+#include "ConstraintChannel.h"
+#include "ConstraintChannelHelper.h"
+
+#include "Channels/MovieSceneChannelData.h"
+
+#include "Algo/Unique.h"
+
+template<typename ChannelType>
+void FConstraintChannelHelper::GetFramesToCompensate(
+	const FMovieSceneConstraintChannel& InActiveChannel,
+	const bool InActiveValueToBeSet,
+	const FFrameNumber& InTime,
+	const TArrayView<ChannelType*>& InChannels,
+	TArray<FFrameNumber>& OutFramesAfter)
+{
+	using ChannelValueType = typename ChannelType::ChannelValueType;
+
+	const bool bHasKeys = (InActiveChannel.GetNumKeys() > 0);
+	
+	OutFramesAfter.Reset();
+
+	// add the current frame
+	OutFramesAfter.Add(InTime);
+
+	// add the next frames that need transform compensation 
+	for (const ChannelType* InChannel: InChannels)
+	{
+		const TMovieSceneChannelData<const ChannelValueType> ChannelData = InChannel->GetData();
+		const TArrayView<const FFrameNumber> Times = ChannelData.GetTimes();
+		if (!Times.IsEmpty())
+		{
+			// look for the first next key frame for this channel 
+			const int32 NextTimeIndex = Algo::UpperBound(Times, InTime);
+			if (Times.IsValidIndex(NextTimeIndex))
+			{
+				// store the time while the state is different
+				for (int32 Index = NextTimeIndex; Index < Times.Num(); ++Index)
+				{
+					if (!bHasKeys)
+					{
+						OutFramesAfter.Add(Times[Index]);
+					}
+					else
+					{
+						bool NextValue = false; InActiveChannel.Evaluate(Times[Index], NextValue);
+						if (NextValue == InActiveValueToBeSet)
+						{
+							break;
+						}
+						OutFramesAfter.Add(Times[Index]);
+					}
+				}
+			}
+		}
+	}
+
+	// uniqueness
+	OutFramesAfter.Sort();
+	OutFramesAfter.SetNum(Algo::Unique(OutFramesAfter));
+}
+
+template< typename ChannelType >
+void FConstraintChannelHelper::GetFramesAfter(
+	const FMovieSceneConstraintChannel& InActiveChannel,
+	const FFrameNumber& InTime,
+	const TArrayView<ChannelType*>& InChannels,
+	TArray<FFrameNumber>& OutFrames)
+{
+	using ChannelValueType = typename ChannelType::ChannelValueType;
+	
+	OutFrames.Reset();
+	
+	const TMovieSceneChannelData<const bool> ConstraintChannelData = InActiveChannel.GetData();
+	const int32 KeyIndex = ConstraintChannelData.FindKey(InTime);
+	if (!ConstraintChannelData.GetTimes().IsValidIndex(KeyIndex))
+	{
+		return;
+	}
+
+	const bool CurrentValue = ConstraintChannelData.GetValues()[KeyIndex];
+	
+	// compute last frame to compensate
+	auto GetEndOfCompensationTime = [KeyIndex](const TMovieSceneChannelData<const bool>& InData)
+	{
+		const TArrayView<const bool> Values = InData.GetValues();
+		const TArrayView<const FFrameNumber> Times = InData.GetTimes();
+
+		const bool CurrentValue = Values[KeyIndex];
+		for (int32 NextIndex = KeyIndex+1; NextIndex < Times.Num(); ++NextIndex)
+		{
+			if (Values[NextIndex] != CurrentValue)
+			{
+				return TOptional<FFrameNumber>(Times[NextIndex]);
+			}
+		}
+		return TOptional<FFrameNumber>();
+	};
+	const TOptional<FFrameNumber> EndOfCompensationTime = GetEndOfCompensationTime(ConstraintChannelData);
+
+	const bool bHasEndTime = EndOfCompensationTime.IsSet();
+	
+	// add the current frame
+	OutFrames.Add(InTime);
+
+	// add the next frames that need transform compensation 
+	for (const ChannelType* InChannel: InChannels)
+	{
+		const TMovieSceneChannelData<const ChannelValueType> ChannelData = InChannel->GetData();
+		const TArrayView<const FFrameNumber> Times = ChannelData.GetTimes();
+		if (!Times.IsEmpty())
+		{
+			// look for the first next key frame for this channel 
+			const int32 NextTimeIndex = Algo::UpperBound(Times, InTime);
+			if (Times.IsValidIndex(NextTimeIndex))
+			{
+				// store the time while the state is different
+				for (int32 Index = NextTimeIndex; Index < Times.Num(); ++Index)
+				{
+					if (!bHasEndTime || Times[Index] < EndOfCompensationTime.GetValue() )
+					{
+						OutFrames.Add(Times[Index]);
+					}
+				}
+			}
+		}
+	}
+
+	// uniqueness
+	OutFrames.Sort();
+	OutFrames.SetNum(Algo::Unique(OutFrames));
+}
+
+template< typename ChannelType >
+void FConstraintChannelHelper::GetFramesWithinActiveState(
+	const FMovieSceneConstraintChannel& InActiveChannel,
+	const TArrayView<ChannelType*>& InChannels,
+	TArray<FFrameNumber>& OutFrames)
+{
+	using ChannelValueType = typename ChannelType::ChannelValueType;
+	
+	OutFrames.Reset();
+	
+	const TMovieSceneChannelData<const bool> ConstraintChannelData = InActiveChannel.GetData();
+	const TArrayView<const FFrameNumber>& ActiveTimes = ConstraintChannelData.GetTimes();
+	if (ActiveTimes.IsEmpty())
+	{
+		return;
+	}
+
+	const FFrameNumber& FirstTime = ActiveTimes[0];
+	const FFrameNumber& LastTime = ActiveTimes.Last();
+	
+	// add active times
+	OutFrames.Append(ActiveTimes);
+
+	const bool bIsLastStateInactive = ConstraintChannelData.GetValues().Last() == false; 
+
+	// add frames where the constraint is active 
+	for (const ChannelType* InChannel: InChannels)
+	{
+		const TMovieSceneChannelData<const ChannelValueType> ChannelData = InChannel->GetData();
+		const TArrayView<const FFrameNumber> Times = ChannelData.GetTimes();
+		if (!Times.IsEmpty())
+		{
+			// look for the first next key frame for this channel 
+			const int32 NextTimeIndex = Algo::UpperBound(Times,FirstTime);
+			if (Times.IsValidIndex(NextTimeIndex))
+			{
+				// store the time is the state is active
+				for (int32 Index = NextTimeIndex; Index < Times.Num(); ++Index)
+				{
+					bool bIsActive = false;
+					InActiveChannel.Evaluate(Times[Index], bIsActive);
+					if (bIsActive)
+					{
+						OutFrames.Add(Times[Index]);
+					}
+
+					if (bIsLastStateInactive && Times[Index] > LastTime)
+					{
+						break;
+					}
+				}
+			}
+		}
+	}
+
+	// uniqueness
+	OutFrames.Sort();
+	OutFrames.SetNum(Algo::Unique(OutFrames));
+}
+
+template< typename ChannelType >
+void FConstraintChannelHelper::MoveTransformKeys(
+	const TArrayView<ChannelType*>& InChannels,
+	const FFrameNumber& InCurrentTime,
+	const FFrameNumber& InNextTime)
+{
+	const FFrameNumber Delta = InNextTime - InCurrentTime;
+	if (Delta == 0)
+	{
+		return;
+	}
+	
+	for (ChannelType* Channel: InChannels)
+	{
+		TMovieSceneChannelData<typename ChannelType::ChannelValueType> Data = Channel->GetData();
+		const TArrayView<const FFrameNumber> Times = Data.GetTimes();
+		const int32 NumTimes = Times.Num();
+
+		if (Delta > 0) //if we are moving keys positively in time we start from end frames and move them so we can use indices
+		{
+			for (int32 KeyIndex = NumTimes - 1; KeyIndex >= 0; --KeyIndex)
+			{
+				const FFrameNumber& Frame = Times[KeyIndex];
+				const FFrameNumber AbsDiff = FMath::Abs(Frame - InCurrentTime);
+				if (AbsDiff<= 1)
+				{
+					Data.MoveKey(KeyIndex, Frame + Delta);
+				}
+			}
+		}
+		else
+		{
+			for (int32 KeyIndex = 0; KeyIndex < NumTimes; ++KeyIndex)
+			{
+				const FFrameNumber& Frame = Times[KeyIndex];
+				const FFrameNumber AbsDiff = FMath::Abs( Frame - InCurrentTime);
+				if (AbsDiff <= 1)
+				{
+					Data.MoveKey(KeyIndex, Frame + Delta);
+				}
+			}
+		}
+	}
+}
+
+template< typename ChannelType >
+void FConstraintChannelHelper::DeleteTransformKeys(
+	const TArrayView<ChannelType*>& InChannels,
+	const FFrameNumber& InTime)
+{
+	for (ChannelType* Channel: InChannels)
+	{
+		TMovieSceneChannelData<typename ChannelType::ChannelValueType> Data = Channel->GetData();
+		const TArrayView<const FFrameNumber> Times = Data.GetTimes();
+		
+		const int32 KeyIndex = Algo::UpperBound(Times, InTime);
+		if (Times.IsValidIndex(KeyIndex))
+		{
+			Data.RemoveKey(KeyIndex);
+		}
+	}
+}
\ No newline at end of file
diff --git a/Engine/Plugins/Animation/ControlRig/Source/ControlRigEditor/Private/Sequencer/ControlRigParameterTrackEditor.cpp b/Engine/Plugins/Animation/ControlRig/Source/ControlRigEditor/Private/Sequencer/ControlRigParameterTrackEditor.cpp
--- a/Engine/Plugins/Animation/ControlRig/Source/ControlRigEditor/Private/Sequencer/ControlRigParameterTrackEditor.cpp
+++ b/Engine/Plugins/Animation/ControlRig/Source/ControlRigEditor/Private/Sequencer/ControlRigParameterTrackEditor.cpp
@@ -76,6 +76,7 @@
 #include "Toolkits/IToolkitHost.h"
 #include "ControlRigEditModeSettings.h"
 #include "ControlRigSpaceChannelEditors.h"
+#include "TransformConstraint.h"
 #include "Misc/ScopedSlowTask.h"
 
 #define LOCTEXT_NAMESPACE "FControlRigParameterTrackEditor"
@@ -421,6 +422,7 @@ void FControlRigParameterTrackEditor::BindControlRig(UControlRig* ControlRig)
 				}
 			}
 			Track->SpaceChannelAdded().AddRaw(this, &FControlRigParameterTrackEditor::HandleOnSpaceAdded);
+			Track->ConstraintChannelAdded().AddRaw(this, &FControlRigParameterTrackEditor::HandleOnConstraintAdded);
 		}
 	}
 }
@@ -432,17 +434,18 @@ void FControlRigParameterTrackEditor::UnbindControlRig(UControlRig* ControlRig)
 		if (Track)
 		{
 			Track->SpaceChannelAdded().RemoveAll(this);
+			Track->ConstraintChannelAdded().RemoveAll(this);
 		}
 		ControlRig->ControlModified().RemoveAll(this);
 		ControlRig->OnInitialized_AnyThread().RemoveAll(this);
 		ControlRig->ControlSelected().RemoveAll(this);
 		BoundControlRigs.Remove(ControlRig);
-		ClearOutAllSpaceDelegates(ControlRig);
+		ClearOutAllSpaceAndConstraintDelegates(ControlRig);
 	}
 }
 void FControlRigParameterTrackEditor::UnbindAllControlRigs()
 {
-	ClearOutAllSpaceDelegates();
+	ClearOutAllSpaceAndConstraintDelegates();
 	for (TWeakObjectPtr<UControlRig>& ObjectPtr : BoundControlRigs)
 	{
 		if (ObjectPtr.IsValid())
@@ -452,6 +455,7 @@ void FControlRigParameterTrackEditor::UnbindAllControlRigs()
 			if (Track)
 			{
 				Track->SpaceChannelAdded().RemoveAll(this);
+				Track->ConstraintChannelAdded().RemoveAll(this);
 			}
 			ControlRig->ControlModified().RemoveAll(this);
 			ControlRig->OnInitialized_AnyThread().RemoveAll(this);
@@ -2304,64 +2308,87 @@ void FControlRigParameterTrackEditor::PostUndo(bool bSuccess)
 			{ 
 				HandleOnSpaceAdded(Section, Channel.ControlName, &(Channel.SpaceCurve));
 			}
+
+			TArray<FConstraintAndActiveChannel>& ConstraintChannels = Section->GetConstraintsChannels();
+			for (FConstraintAndActiveChannel& Channel: ConstraintChannels)
+			{ 
+				HandleOnConstraintAdded(Section, &(Channel.ActiveChannel));
+			}
 		}
 	}
 }
 
 
-void FControlRigParameterTrackEditor::HandleSpaceKeyDeleted(UMovieSceneControlRigParameterSection* Section, FMovieSceneControlRigSpaceChannel* Channel, const TArray<FKeyAddOrDeleteEventItem>& DeletedItems)
+void FControlRigParameterTrackEditor::HandleSpaceKeyDeleted(
+	UMovieSceneControlRigParameterSection* Section,
+	FMovieSceneControlRigSpaceChannel* Channel,
+	const TArray<FKeyAddOrDeleteEventItem>& DeletedItems) const
 {
 	const TSharedPtr<ISequencer> ParentSequencer = GetSequencer();
 
 	if (Section && Section->GetControlRig() && Channel && ParentSequencer.IsValid())
 	{
-		FName ControlName = Section->FindControlNameFromSpaceChannel(Channel);
+		const FName ControlName = Section->FindControlNameFromSpaceChannel(Channel);
 		for (const FKeyAddOrDeleteEventItem& EventItem : DeletedItems)
 		{
-			FControlRigSpaceChannelHelpers::SequencerSpaceChannelKeyDeleted(Section->GetControlRig(), ParentSequencer.Get(), ControlName, Channel, Section,EventItem.Frame);
+			FControlRigSpaceChannelHelpers::SequencerSpaceChannelKeyDeleted(
+				Section->GetControlRig(), ParentSequencer.Get(), ControlName, Channel, Section,EventItem.Frame);
 		}
 	}
 }
 
-void FControlRigParameterTrackEditor::HandleSpaceKeyMoved(UMovieSceneControlRigParameterSection* Section, FMovieSceneControlRigSpaceChannel* SpaceChannel, const  TArray<FKeyMoveEventItem>& MovedItems)
+void FControlRigParameterTrackEditor::HandleSpaceKeyMoved(
+	UMovieSceneControlRigParameterSection* Section,
+	FMovieSceneControlRigSpaceChannel* SpaceChannel,
+	const  TArray<FKeyMoveEventItem>& MovedItems)
 {
 	if (Section && Section->GetControlRig() && SpaceChannel)
 	{
-		FName ControlName = Section->FindControlNameFromSpaceChannel(SpaceChannel);
+		const FName ControlName = Section->FindControlNameFromSpaceChannel(SpaceChannel);
 		for (const FKeyMoveEventItem& MoveEventItem : MovedItems)
 		{
-			FControlRigSpaceChannelHelpers::HandleSpaceKeyTimeChanged(Section->GetControlRig(), ControlName, SpaceChannel, Section,
+			FControlRigSpaceChannelHelpers::HandleSpaceKeyTimeChanged(
+				Section->GetControlRig(), ControlName, SpaceChannel, Section,
 				MoveEventItem.Frame, MoveEventItem.NewFrame);
 		}
 	}
 }
 
-void FControlRigParameterTrackEditor::ClearOutAllSpaceDelegates(UControlRig* InOptionalControlRig)
+void FControlRigParameterTrackEditor::ClearOutAllSpaceAndConstraintDelegates(const UControlRig* InOptionalControlRig) const
 {
 	if (GetSequencer().IsValid())
 	{
-		UMovieScene* MovieScene = GetSequencer()->GetFocusedMovieSceneSequence()->GetMovieScene();
+		const UMovieScene* MovieScene = GetSequencer()->GetFocusedMovieSceneSequence()->GetMovieScene();
 		const TArray<FMovieSceneBinding>& Bindings = MovieScene->GetBindings();
 		for (const FMovieSceneBinding& Binding : Bindings)
 		{
-			if (UMovieSceneControlRigParameterTrack* Track = Cast<UMovieSceneControlRigParameterTrack>(MovieScene->FindTrack(UMovieSceneControlRigParameterTrack::StaticClass(), Binding.GetObjectGuid(), NAME_None)))
+			const UMovieSceneTrack* Track = MovieScene->FindTrack(
+				UMovieSceneControlRigParameterTrack::StaticClass(), Binding.GetObjectGuid(), NAME_None);
+			if (const UMovieSceneControlRigParameterTrack* CRTrack = Cast<UMovieSceneControlRigParameterTrack>(Track))
 			{
-				if (InOptionalControlRig && Track->GetControlRig() != InOptionalControlRig)
+				if (InOptionalControlRig && CRTrack->GetControlRig() != InOptionalControlRig)
 				{
 					continue;
 				}
+				
 				for (UMovieSceneSection* Section : Track->GetAllSections())
 				{
-					if (Section)
+					if (UMovieSceneControlRigParameterSection* CRSection = Cast<UMovieSceneControlRigParameterSection>(Section))
 					{
-						if (UMovieSceneControlRigParameterSection* CRSection = Cast<UMovieSceneControlRigParameterSection>(Section))
+						// clear space channels
+						TArray<FSpaceControlNameAndChannel>& Channels = CRSection->GetSpaceChannels();
+						for (FSpaceControlNameAndChannel& SpaceAndChannel : Channels)
 						{
-							TArray<FSpaceControlNameAndChannel>& Channels = CRSection->GetSpaceChannels();
-							for (FSpaceControlNameAndChannel& SpaceAndChannel : Channels)
-							{
-								SpaceAndChannel.SpaceCurve.OnKeyMovedEvent().Clear();
-								SpaceAndChannel.SpaceCurve.OnKeyDeletedEvent().Clear();
-							}
+							SpaceAndChannel.SpaceCurve.OnKeyMovedEvent().Clear();
+							SpaceAndChannel.SpaceCurve.OnKeyDeletedEvent().Clear();
+						}
+
+						// clear constraint channels
+						TArray<FConstraintAndActiveChannel>& ConstraintChannels = CRSection->GetConstraintsChannels();
+						for (FConstraintAndActiveChannel& Channel: ConstraintChannels)
+						{
+							Channel.ActiveChannel.OnKeyMovedEvent().Clear();
+							Channel.ActiveChannel.OnKeyDeletedEvent().Clear();							
 						}
 					}
 				}
@@ -2370,6 +2397,171 @@ void FControlRigParameterTrackEditor::ClearOutAllSpaceDelegates(UControlRig* InO
 	}
 }
 
+namespace
+{
+	struct FConstraintAndControlData
+	{
+		static FConstraintAndControlData CreateFromSection(
+			const UMovieSceneControlRigParameterSection* InSection,
+			const FMovieSceneConstraintChannel* InConstraintChannel)
+		{
+			FConstraintAndControlData Data;
+			
+			// get constraint channel
+			const TArray<FConstraintAndActiveChannel>& ConstraintChannels = InSection->GetConstraintsChannels();
+			const int32 Index = ConstraintChannels.IndexOfByPredicate([InConstraintChannel](const FConstraintAndActiveChannel& InChannel)
+			{
+				return &(InChannel.ActiveChannel) == InConstraintChannel;
+			});
+	
+			if (Index == INDEX_NONE)
+			{
+				return Data;
+			}
+
+			Data.Constraint = Cast<UTickableTransformConstraint>(ConstraintChannels[Index].Constraint.Get());
+
+			// get constraint name
+			auto GetControlName = [InSection, Index]()
+			{
+				using NameInfoIterator = TMap<FName, FChannelMapInfo>::TRangedForConstIterator;
+				for (NameInfoIterator It = InSection->ControlChannelMap.begin(); It; ++It)
+				{
+					const FChannelMapInfo& Info = It->Value;
+					if (Info.ConstraintsIndex.Contains(Index))
+					{
+						return It->Key;
+					}
+				}
+
+				static const FName DummyName = NAME_None;
+				return DummyName;
+			};
+	
+			Data.ControlName = GetControlName();
+			
+			return Data;
+		}
+
+		bool IsValid() const
+		{
+			return Constraint.IsValid() && ControlName != NAME_None; 
+		}
+		
+		TWeakObjectPtr<UTickableTransformConstraint> Constraint = nullptr;
+		FName ControlName = NAME_None;
+	};
+}
+
+void FControlRigParameterTrackEditor::HandleOnConstraintAdded(
+	UMovieSceneControlRigParameterSection* InSection,
+	FMovieSceneConstraintChannel* InConstraintChannel)
+{
+	if (!InConstraintChannel)
+	{
+		return;
+	}
+
+	// handle key moved
+	InConstraintChannel->OnKeyMovedEvent().AddLambda([this,InSection](
+		FMovieSceneChannel* InChannel, const TArray<FKeyMoveEventItem>& InMovedItems)
+	{
+		const FMovieSceneConstraintChannel* ConstraintChannel = static_cast<FMovieSceneConstraintChannel*>(InChannel);
+		HandleConstraintKeyMoved(InSection, ConstraintChannel, InMovedItems);
+	});
+
+	// handle key deleted
+	InConstraintChannel->OnKeyDeletedEvent().AddLambda([this, InSection](
+		FMovieSceneChannel* InChannel, const TArray<FKeyAddOrDeleteEventItem>& InDeletedItems)
+	{
+		const FMovieSceneConstraintChannel* ConstraintChannel = static_cast<FMovieSceneConstraintChannel*>(InChannel);
+		HandleConstraintKeyDeleted(InSection, ConstraintChannel, InDeletedItems);
+	});
+
+	// handle constraint deleted
+	if (InSection)
+	{
+		HandleConstraintRemoved(InSection);
+	}
+}
+
+void FControlRigParameterTrackEditor::HandleConstraintKeyDeleted(
+	const UMovieSceneControlRigParameterSection* InSection,
+	const FMovieSceneConstraintChannel* InConstraintChannel,
+	const TArray<FKeyAddOrDeleteEventItem>& InDeletedItems) const
+{
+	const UControlRig* ControlRig = InSection ? InSection->GetControlRig() : nullptr;
+	if (!ControlRig || !InConstraintChannel)
+	{
+		return;
+	}
+	
+	const FConstraintAndControlData ConstraintAndControlData =
+		FConstraintAndControlData::CreateFromSection(InSection, InConstraintChannel);
+	if (ConstraintAndControlData.IsValid())
+	{
+		UTickableTransformConstraint* Constraint = ConstraintAndControlData.Constraint.Get();
+		for (const FKeyAddOrDeleteEventItem& EventItem: InDeletedItems)
+		{
+			FConstraintChannelHelper::HandleConstraintKeyDeleted(
+				Constraint, InConstraintChannel,
+				GetSequencer(), InSection,
+				EventItem.Frame);
+		}
+	}
+}
+
+void FControlRigParameterTrackEditor::HandleConstraintKeyMoved(
+	const UMovieSceneControlRigParameterSection* InSection,
+	const FMovieSceneConstraintChannel* InConstraintChannel,
+	const TArray<FKeyMoveEventItem>& InMovedItems)
+{
+	const FConstraintAndControlData ConstraintAndControlData =
+	FConstraintAndControlData::CreateFromSection(InSection, InConstraintChannel);
+
+	if (ConstraintAndControlData.IsValid())
+	{
+		const UTickableTransformConstraint* Constraint = ConstraintAndControlData.Constraint.Get();
+		for (const FKeyMoveEventItem& MoveEventItem : InMovedItems)
+		{
+			FConstraintChannelHelper::HandleConstraintKeyMoved(
+				Constraint, InConstraintChannel, InSection,
+				MoveEventItem.Frame, MoveEventItem.NewFrame);
+		}
+	}
+}
+
+void FControlRigParameterTrackEditor::HandleConstraintRemoved(UMovieSceneControlRigParameterSection* InSection) const
+{
+	if (const UControlRig* ControlRig = InSection->GetControlRig())
+	{
+		FConstraintsManagerController& Controller = FConstraintsManagerController::Get(ControlRig->GetWorld());
+		if (!Controller.OnConstraintRemoved().IsBoundToObject(InSection))
+		{
+			InSection->OnConstraintRemovedHandle =
+			Controller.OnConstraintRemoved().AddLambda([InSection, this](FName InConstraintName)
+			{
+				const FConstraintAndActiveChannel* ConstraintChannel = InSection->GetConstraintChannel(InConstraintName);
+				if (!ConstraintChannel)
+				{
+					return;
+				}
+				
+				if (ConstraintChannel->Constraint.IsValid())
+				{
+					FConstraintChannelHelper::HandleConstraintRemoved(
+						ConstraintChannel->Constraint.Get(),
+						&ConstraintChannel->ActiveChannel,
+						GetSequencer(),
+						InSection);
+				}
+
+				InSection->RemoveConstraintChannel(InConstraintName);
+		   });
+		}
+	}
+}
+
 void FControlRigParameterTrackEditor::SetUpEditModeIfNeeded(UControlRig* ControlRig)
 {
 	if (ControlRig)
@@ -4074,6 +4266,24 @@ bool FControlRigParameterSection::RequestDeleteCategory(const TArray<FName>& Cat
 {
 	UMovieSceneControlRigParameterSection* ParameterSection = CastChecked<UMovieSceneControlRigParameterSection>(WeakSection.Get());
 	TSharedPtr<ISequencer> SequencerPtr = WeakSequencer.Pin();
+
+	if (ParameterSection && SequencerPtr)
+	{
+		const FName& Channel = CategoryNamePaths.Last();
+
+		// remove constraint channel if there are no keys
+		const FConstraintAndActiveChannel* ConstraintChannel = ParameterSection->GetConstraintChannel(Channel);
+		if (ConstraintChannel && ConstraintChannel->ActiveChannel.GetNumKeys() == 0)
+		{
+			if (ParameterSection->TryModify())
+			{
+				ParameterSection->RemoveConstraintChannel(Channel);
+				SequencerPtr->NotifyMovieSceneDataChanged(EMovieSceneDataChangeType::MovieSceneStructureItemsChanged);
+				return true;
+			}
+		}
+	}
+	
 	/*
 	const FScopedTransaction Transaction(LOCTEXT("DeleteTransformCategory", "Delete transform category"));
 
diff --git a/Engine/Plugins/Animation/ControlRig/Source/ControlRigEditor/Private/Sequencer/ControlRigParameterTrackEditor.h b/Engine/Plugins/Animation/ControlRig/Source/ControlRigEditor/Private/Sequencer/ControlRigParameterTrackEditor.h
--- a/Engine/Plugins/Animation/ControlRig/Source/ControlRigEditor/Private/Sequencer/ControlRigParameterTrackEditor.h
+++ b/Engine/Plugins/Animation/ControlRig/Source/ControlRigEditor/Private/Sequencer/ControlRigParameterTrackEditor.h
@@ -128,12 +128,36 @@ private:
 	void HandleControlUndoBracket(UControlRig* Subject, bool bOpenUndoBracket);
 	void HandleOnInitialized(UControlRig* Subject, const EControlRigState InState, const FName& InEventName);
 
-	/** SpaceChannel Delegates*/
-	void HandleOnSpaceAdded(UMovieSceneControlRigParameterSection* Section, const FName& ControlName, FMovieSceneControlRigSpaceChannel* Channel);
-	void HandleSpaceKeyDeleted(UMovieSceneControlRigParameterSection* Section, FMovieSceneControlRigSpaceChannel* Channel, const TArray<FKeyAddOrDeleteEventItem>& DeletedItems);
-	void HandleSpaceKeyMoved(UMovieSceneControlRigParameterSection* Section, FMovieSceneControlRigSpaceChannel* Channel, const  TArray<FKeyMoveEventItem>& MovedItems);
 	//if rig not set then we clear delegates for everyone
-	void ClearOutAllSpaceDelegates(UControlRig* InOptionalControlRig = nullptr);
+	void ClearOutAllSpaceAndConstraintDelegates(const UControlRig* InOptionalControlRig = nullptr) const;
+
+	/** SpaceChannel Delegates*/
+	void HandleOnSpaceAdded(
+		UMovieSceneControlRigParameterSection* Section,
+		const FName& ControlName,
+		FMovieSceneControlRigSpaceChannel* Channel);
+	void HandleSpaceKeyDeleted(
+		UMovieSceneControlRigParameterSection* Section,
+		FMovieSceneControlRigSpaceChannel* Channel,
+		const TArray<FKeyAddOrDeleteEventItem>& DeletedItems) const;
+	static void HandleSpaceKeyMoved(
+		UMovieSceneControlRigParameterSection* Section,
+		FMovieSceneControlRigSpaceChannel* Channel,
+		const TArray<FKeyMoveEventItem>& MovedItems);
+
+	/** ConstraintChannel Delegates*/
+	void HandleOnConstraintAdded(
+		UMovieSceneControlRigParameterSection* InSection,
+		FMovieSceneConstraintChannel* InConstraintChannel);
+	void HandleConstraintKeyDeleted(
+		const UMovieSceneControlRigParameterSection* InSection,
+		const FMovieSceneConstraintChannel* InConstraintChannel,
+		const TArray<FKeyAddOrDeleteEventItem>& InDeletedItems) const;
+	static void HandleConstraintKeyMoved(
+		const UMovieSceneControlRigParameterSection* InSection,
+		const FMovieSceneConstraintChannel* InConstraintChannel,
+		const TArray<FKeyMoveEventItem>& InMovedItems);
+	void HandleConstraintRemoved(UMovieSceneControlRigParameterSection* InSection) const;
 
 	/** Select control rig if not selected, select controls from key areas */
 	void SelectRigsAndControls(UControlRig* Subject, const TArray<const IKeyArea*>& KeyAreas);
diff --git a/Engine/Plugins/Animation/ControlRig/Source/ControlRigEditor/Public/ConstraintChannelHelper.h b/Engine/Plugins/Animation/ControlRig/Source/ControlRigEditor/Public/ConstraintChannelHelper.h
--- a/Engine/Plugins/Animation/ControlRig/Source/ControlRigEditor/Public/ConstraintChannelHelper.h
+++ b/Engine/Plugins/Animation/ControlRig/Source/ControlRigEditor/Public/ConstraintChannelHelper.h
@@ -4,9 +4,13 @@
 
 #include "Templates/SharedPointer.h"
 
+struct FChannelMapInfo;
+struct FConstraintAndActiveChannel;
 class UControlRig;
 class ISequencer;
 class UMovieScene3DTransformSection;
+class UMovieSceneControlRigParameterSection;
+class UMovieSceneSection;
 class UTransformableHandle;
 class UTickableTransformConstraint;
 class UTransformableControlHandle;
@@ -14,7 +18,6 @@ class UTransformableComponentHandle;
 struct FMovieSceneFloatChannel;
 struct FMovieSceneDoubleChannel;
 struct FMovieSceneConstraintChannel;
-class UMovieSceneControlRigParameterSection;
 struct FFrameNumber;
 
 /**
@@ -24,6 +27,9 @@ struct FFrameNumber;
 struct FConstraintChannelHelper
 {
 public:
+	/** @todo documentation. */
+	static bool IsKeyframingAvailable();
+	
 	/** @todo documentation. */
 	static void AddConstraintKey(UTickableTransformConstraint* InConstraint);
 	
@@ -45,6 +51,28 @@ public:
 		const TSharedPtr<ISequencer>& InSequencer,
 		const UMovieSceneControlRigParameterSection* Section,
 		const TOptional<FFrameNumber>& OptionalTime);
+
+	/** @todo documentation. */
+	static void HandleConstraintRemoved(
+		UTickableConstraint* InConstraint,
+		const FMovieSceneConstraintChannel* InChannel,
+		const TSharedPtr<ISequencer>& InSequencer,
+		const UMovieSceneSection* InSection);
+
+	/** @todo documentation. */
+	static void HandleConstraintKeyDeleted(
+		UTickableTransformConstraint* InConstraint,
+		const FMovieSceneConstraintChannel* InConstraintChannel,
+		const TSharedPtr<ISequencer>& InSequencer,
+		const UMovieSceneSection* InSection,
+		const FFrameNumber& InTime);
+
+	/** @todo documentation. */
+	static void HandleConstraintKeyMoved(
+		const UTickableTransformConstraint* InConstraint,
+		const FMovieSceneConstraintChannel* InConstraintChannel,
+		const UMovieSceneSection* InSection,
+		const FFrameNumber& InCurrentFrame, const FFrameNumber& InNextFrame);
 	
 private:
 
@@ -80,11 +108,57 @@ private:
 	static TArrayView<FMovieSceneFloatChannel*> GetTransformFloatChannels(
 		const UTransformableHandle* InHandle,
 		const TSharedPtr<ISequencer>& InSequencer);
+	static TArrayView<FMovieSceneFloatChannel*> GetTransformFloatChannels(
+		const UTransformableHandle* InHandle,
+		const UMovieSceneControlRigParameterSection* InSection);
 
 	/** @todo documentation.  TransformComponents use doubles */
 	static TArrayView<FMovieSceneDoubleChannel*> GetTransformDoubleChannels(
 		const UTransformableHandle* InHandle,
 		const TSharedPtr<ISequencer>& InSequencer);
 
+	/** @todo factorize. */
+	static TPair<const FChannelMapInfo*, int32> GetInfoAndNumFloatChannels(
+		const UControlRig* InControlRig,
+		const FName& InControlName,
+		const UMovieSceneControlRigParameterSection* InSection);
+
+	/** @todo documentation. */
+	template<typename ChannelType>
+	static void GetFramesToCompensate(
+		const FMovieSceneConstraintChannel& InActiveChannel,
+		const bool InActiveValueToBeSet,
+		const FFrameNumber& InTime,
+		const TArrayView<ChannelType*>& InChannels,
+		TArray<FFrameNumber>& OutFramesAfter);
+
+	/** @todo documentation. */
+	template< typename ChannelType >
+	static void GetFramesAfter(
+		const FMovieSceneConstraintChannel& InActiveChannel,
+		const FFrameNumber& InTime,
+		const TArrayView<ChannelType*>& InChannels,
+		TArray<FFrameNumber>& OutFrames);
+		
+	/** @todo documentation. */
+	template< typename ChannelType >
+    static void GetFramesWithinActiveState(
+	    const FMovieSceneConstraintChannel& InActiveChannel,
+	    const TArrayView<ChannelType*>& InChannels,
+	    TArray<FFrameNumber>& OutFrames);
+	
+	/** @todo documentation. */
+	template< typename ChannelType >
+	static void MoveTransformKeys(
+		const TArrayView<ChannelType*>& InChannels,
+		const FFrameNumber& InCurrentTime,
+		const FFrameNumber& InNextTime);
+
+	/** @todo documentation. */
+	template< typename ChannelType >
+	static void DeleteTransformKeys(
+		const TArrayView<ChannelType*>& InChannels,
+		const FFrameNumber& InTime);
+	
 	static bool bDoNotCompensate;
 };
diff --git a/Engine/Source/Runtime/Experimental/Animation/Constraints/Private/ConstraintChannel.cpp b/Engine/Source/Runtime/Experimental/Animation/Constraints/Private/ConstraintChannel.cpp
--- a/Engine/Source/Runtime/Experimental/Animation/Constraints/Private/ConstraintChannel.cpp
+++ b/Engine/Source/Runtime/Experimental/Animation/Constraints/Private/ConstraintChannel.cpp
@@ -4,13 +4,12 @@
 
 bool FMovieSceneConstraintChannel::Evaluate(FFrameTime InTime, bool& OutValue) const
 {
-	const TArrayView<const FFrameNumber> TimeArray = GetTimes();
-	if (TimeArray.IsEmpty())
+	if (Times.IsEmpty())
 	{
 		return false;
 	}
 
-	if (InTime.FrameNumber < TimeArray[0])
+	if (InTime.FrameNumber < Times[0])
 	{
 		return false;
 	}
diff --git a/Engine/Source/Runtime/Experimental/Animation/Constraints/Private/ConstraintsManager.cpp b/Engine/Source/Runtime/Experimental/Animation/Constraints/Private/ConstraintsManager.cpp
--- a/Engine/Source/Runtime/Experimental/Animation/Constraints/Private/ConstraintsManager.cpp
+++ b/Engine/Source/Runtime/Experimental/Animation/Constraints/Private/ConstraintsManager.cpp
@@ -419,16 +419,17 @@ bool FConstraintsManagerController::RemoveConstraint(const int32 InConstraintInd
 		return false;
 	}
 
-	// TODO handle transaction
-	const FName ConstraintName = Manager->Constraints[InConstraintIndex]->GetFName(); 
+	const FName ConstraintName = Manager->Constraints[InConstraintIndex]->GetFName();
 	UTickableConstraint* Constraint = Manager->Constraints[InConstraintIndex];
-	Manager->Constraints[InConstraintIndex]->ConstraintTick.UnRegisterTickFunction();
-	Manager->Constraints.RemoveAt(InConstraintIndex);
-
+	
 	// notify deletion
 	ConstraintRemoved.Broadcast(ConstraintName);
 	Manager->OnConstraintRemoved_BP.Broadcast(Manager, Constraint);
 
+	// TODO handle transaction
+	Manager->Constraints[InConstraintIndex]->ConstraintTick.UnRegisterTickFunction();
+	Manager->Constraints.RemoveAt(InConstraintIndex);
+
 	// destroy constraints actor if no constraints left
 	if (Manager->Constraints.IsEmpty())
 	{
diff --git a/Engine/Source/Runtime/Experimental/Animation/Constraints/Public/ConstraintChannel.h b/Engine/Source/Runtime/Experimental/Animation/Constraints/Public/ConstraintChannel.h
--- a/Engine/Source/Runtime/Experimental/Animation/Constraints/Public/ConstraintChannel.h
+++ b/Engine/Source/Runtime/Experimental/Animation/Constraints/Public/ConstraintChannel.h
@@ -25,8 +25,34 @@ struct CONSTRAINTS_API FMovieSceneConstraintChannel : public FMovieSceneBoolChan
 
 	FMovieSceneConstraintChannel() {};
 
-	// @todo 
-	bool Evaluate(FFrameTime InTime, bool& OutValue) const;
+	/**
+	 * Access a mutable interface for this channel's data
+	 *
+	 * @return An object that is able to manipulate this channel's data
+	 */
+	virtual FORCEINLINE TMovieSceneChannelData<bool> GetData() override
+	{
+		return TMovieSceneChannelData<bool>(&Times, &Values, &KeyHandles, this);
+	}
+
+	/**
+	 * Access a constant interface for this channel's data
+	 *
+	 * @return An object that is able to interrogate this channel's data
+	 */
+	virtual FORCEINLINE TMovieSceneChannelData<const bool> GetData() const override
+	{
+		return FMovieSceneBoolChannel::GetData();
+	}
+	
+	/**
+	 * Evaluate this channel
+	 *
+	 * @param InTime     The time to evaluate at
+	 * @param OutValue   A value to receive the result
+	 * @return true if the channel was evaluated successfully, false otherwise
+	 */
+	virtual bool Evaluate(FFrameTime InTime, bool& OutValue) const override;
 
 #if WITH_EDITOR
 	using ExtraLabelFunction = TFunction< FString() >;
diff --git a/Engine/Source/Runtime/MovieScene/Public/Channels/MovieSceneBoolChannel.h b/Engine/Source/Runtime/MovieScene/Public/Channels/MovieSceneBoolChannel.h
--- a/Engine/Source/Runtime/MovieScene/Public/Channels/MovieSceneBoolChannel.h
+++ b/Engine/Source/Runtime/MovieScene/Public/Channels/MovieSceneBoolChannel.h
@@ -45,7 +45,7 @@ struct MOVIESCENE_API FMovieSceneBoolChannel : public FMovieSceneChannel
 	 *
 	 * @return An object that is able to manipulate this channel's data
 	 */
-	FORCEINLINE TMovieSceneChannelData<bool> GetData()
+	virtual FORCEINLINE TMovieSceneChannelData<bool> GetData()
 	{
 		return TMovieSceneChannelData<bool>(&Times, &Values, &KeyHandles);
 	}
@@ -55,7 +55,7 @@ struct MOVIESCENE_API FMovieSceneBoolChannel : public FMovieSceneChannel
 	 *
 	 * @return An object that is able to interrogate this channel's data
 	 */
-	FORCEINLINE TMovieSceneChannelData<const bool> GetData() const
+	virtual FORCEINLINE TMovieSceneChannelData<const bool> GetData() const
 	{
 		return TMovieSceneChannelData<const bool>(&Times, &Values);
 	}
@@ -91,7 +91,7 @@ struct MOVIESCENE_API FMovieSceneBoolChannel : public FMovieSceneChannel
 	 * @param OutValue   A value to receive the result
 	 * @return true if the channel was evaluated successfully, false otherwise
 	 */
-	bool Evaluate(FFrameTime InTime, bool& OutValue) const;
+	virtual bool Evaluate(FFrameTime InTime, bool& OutValue) const;
 
 public:
 
@@ -141,7 +141,7 @@ public:
 		bHasDefaultValue = false;
 	}
 
-private:
+protected:
 
 	UPROPERTY(meta=(KeyTimes))
 	TArray<FFrameNumber> Times;
