fix compilation issue

diff --git a/Engine/Plugins/Animation/ControlRig/Source/ControlRigEditor/Private/Sequencer/ConstraintChannelHelper.inl b/Engine/Plugins/Animation/ControlRig/Source/ControlRigEditor/Private/Sequencer/ConstraintChannelHelper.inl
--- a/Engine/Plugins/Animation/ControlRig/Source/ControlRigEditor/Private/Sequencer/ConstraintChannelHelper.inl
+++ b/Engine/Plugins/Animation/ControlRig/Source/ControlRigEditor/Private/Sequencer/ConstraintChannelHelper.inl
@@ -155,7 +155,10 @@ void FConstraintChannelHelper::GetFramesWithinActiveState(
 	const FFrameNumber& LastTime = ActiveTimes.Last();
 	
 	// add active times
-	OutFrames.Append(ActiveTimes);
+	for (int32 Index = 0; Index < ActiveTimes.Num(); ++Index)
+	{
+		OutFrames.Add(ActiveTimes[Index]);
+	}
 
 	const bool bIsLastStateInactive = ConstraintChannelData.GetValues().Last() == false; 
 
