Control Rig: Remove an unnecessary check during reattach link. Previously it was checking if template is initialized on library node, which by design should happen after rettach link. So there is no need to check for it during reattach link. Invalid links will be broken by a call to RecomputeAllTemplateFilteredPermutations() later in RefreshAllModels()

diff --git a/Engine/Plugins/Animation/ControlRig/Source/ControlRigDeveloper/Private/ControlRigBlueprint.cpp b/Engine/Plugins/Animation/ControlRig/Source/ControlRigDeveloper/Private/ControlRigBlueprint.cpp
--- a/Engine/Plugins/Animation/ControlRig/Source/ControlRigDeveloper/Private/ControlRigBlueprint.cpp
+++ b/Engine/Plugins/Animation/ControlRig/Source/ControlRigDeveloper/Private/ControlRigBlueprint.cpp
@@ -1357,8 +1357,9 @@ void UControlRigBlueprint::RefreshAllModels(EControlRigBlueprintLoadType InLoadT
 	for (URigVMGraph* GraphToDetach : GraphsToDetach)
 	{
 		URigVMController* Controller = GetOrCreateController(GraphToDetach);
-		bool bAllowNonArgumentLinks = GetLinkerCustomVersion(FControlRigObjectVersion::GUID) < FControlRigObjectVersion::LibraryNodeTemplates;
-		Controller->ReattachLinksToPinObjects(true /* follow redirectors */, nullptr, false, true, bAllowNonArgumentLinks);
+		// at this stage, allow all links to be reattached,
+		// RecomputeAllTemplateFilteredPermutations() later should break any invalid links
+		Controller->ReattachLinksToPinObjects(true /* follow redirectors */, nullptr, false, true, true);
 	}
 
 	if(bIsPostLoad)
