Sequencer: Control Rig: Anim Slider Fixes. Added option to spin boxes so that enter is required to get focus, otherwise after using slider you can't use a hotkey,it would be grabbed by the spin box. Renamed the tools based upon feedback.  Added hotkeys (u and shift-u) to drag the slider interactively and change the active slider. Also changed the Hide All Control Shapes hotkey to avoid Sequencer conflict.

diff --git a/Engine/Plugins/Animation/ControlRig/Source/ControlRigEditor/Private/EditMode/ControlRigEditMode.cpp b/Engine/Plugins/Animation/ControlRig/Source/ControlRigEditor/Private/EditMode/ControlRigEditMode.cpp
--- a/Engine/Plugins/Animation/ControlRig/Source/ControlRigEditor/Private/EditMode/ControlRigEditMode.cpp
+++ b/Engine/Plugins/Animation/ControlRig/Source/ControlRigEditor/Private/EditMode/ControlRigEditMode.cpp
@@ -941,13 +941,70 @@ bool FControlRigEditMode::InputKey(FEditorViewportClient* InViewportClient, FVie
 		{
 			return true;
 		}
+		if (IsDragAnimSliderToolPressed(InViewport)) //this is needed to make sure we get all of the processed mouse events, for some reason the above may not return true
+		{
+			return true;
+		}
 	}
 
 	return FEdMode::InputKey(InViewportClient, InViewport, InKey, InEvent);
 }
 
+bool FControlRigEditMode::ProcessCapturedMouseMoves(FEditorViewportClient* InViewportClient, FViewport* InViewport, const TArrayView<FIntPoint>& CapturedMouseMoves)
+{
+	const bool bChange = IsDragAnimSliderToolPressed(InViewport);
+	if (CapturedMouseMoves.Num() > 0)
+	{
+		if (bChange)
+		{
+			for (int32 Index = 0; Index < CapturedMouseMoves.Num(); ++Index)
+			{
+				int32 X = CapturedMouseMoves[Index].X;
+				if (StartXValue.IsSet() == false)
+				{
+					bisTrackingAnimToolDrag = true;
+					GEditor->BeginTransaction(LOCTEXT("AnimSliderBlend", "AnimSlider Blend"));
+					StartXValue = X;
+				}
+				else
+				{
+					int32 Diff = X - StartXValue.GetValue();
+					if (Diff != 0)
+					{
+						FIntPoint Origin, Size;
+						InViewportClient->GetViewportDimensions(Origin, Size);
+						const double ViewPortSize = (double)Size.X * 0.25; // 1/4 screen drag should do full blend, todo perhaps expose this as a sensitivity
+						const double ViewDiff = (double)(Diff) / ViewPortSize;
+						DragAnimSliderTool(ViewDiff);
+					}
+					StartXValue = X;
+				}
+			}
+		}
+		else if (bisTrackingAnimToolDrag && bChange == false)
+		{
+			ResetAnimSlider();
+		}
+		return bisTrackingAnimToolDrag;
+	}
+	else if(bisTrackingAnimToolDrag && bChange == false)
+	{
+		ResetAnimSlider();
+	}
+	return false;
+}
+
 bool FControlRigEditMode::EndTracking(FEditorViewportClient* InViewportClient, FViewport* InViewport)
 {
+	if (bisTrackingAnimToolDrag)
+	{
+		ResetAnimSlider();
+	}
+	if (IsDragAnimSliderToolPressed(InViewport))
+	{
+		return true;
+	}
+
 	const bool bWasInteracting = bManipulatorMadeChange && InteractionType != (uint8)EControlRigInteractionType::None;
 	
 	InteractionType = (uint8)EControlRigInteractionType::None;
@@ -991,6 +1048,13 @@ bool FControlRigEditMode::EndTracking(FEditorViewportClient* InViewportClient, F
 
 bool FControlRigEditMode::StartTracking(FEditorViewportClient* InViewportClient, FViewport* InViewport)
 {
+	//break right out if doing anim slider scrub
+
+	if (IsDragAnimSliderToolPressed(InViewport))
+	{
+		return true;
+	}
+
 	InteractionType = GetInteractionType(InViewportClient);
 	
 	if (InteractionScopes.Num() == 0)
@@ -1678,6 +1742,11 @@ void FControlRigEditMode::SelectNone()
 
 bool FControlRigEditMode::InputDelta(FEditorViewportClient* InViewportClient, FViewport* InViewport, FVector& InDrag, FRotator& InRot, FVector& InScale)
 {
+	if (IsDragAnimSliderToolPressed(InViewport)) //this is needed to make sure we get all of the processed mouse events, for some reason the above may not return true
+	{
+		//handled by processed mouse clicks
+		return true;
+	}
 	FVector Drag = InDrag;
 	FRotator Rot = InRot;
 	FVector Scale = InScale;
@@ -2251,6 +2320,12 @@ void FControlRigEditMode::BindCommands()
 		Commands.ClearSelection,
 		FExecuteAction::CreateRaw(this, &FControlRigEditMode::ClearSelection));
 
+	CommandBindings->MapAction(
+		Commands.ChangeAnimSliderTool,
+		FExecuteAction::CreateRaw(this, &FControlRigEditMode::ChangeAnimSliderTool),
+		FCanExecuteAction::CreateRaw(this, &FControlRigEditMode::CanChangeAnimSliderTool)
+	);
+
 	CommandBindings->MapAction(
 		Commands.FrameSelection,
 		FExecuteAction::CreateRaw(this, &FControlRigEditMode::FrameSelection),
@@ -2474,6 +2549,52 @@ void FControlRigEditMode::GetAllSelectedControls(TMap<UControlRig*, TArray<FRigE
 	}
 }
 
+/** If Anim Slider is open, got to the next tool*/
+void FControlRigEditMode::ChangeAnimSliderTool()
+{
+	StaticCastSharedPtr<FControlRigEditModeToolkit>(Toolkit)->GetToNextActiveSlider();
+}
+
+/** If Anim Slider is open, then can drag*/
+bool FControlRigEditMode::CanChangeAnimSliderTool() const
+{
+	return (StaticCastSharedPtr<FControlRigEditModeToolkit>(Toolkit)->CanChangeAnimSliderTool());
+}
+
+/** If Anim Slider is open, drag the tool*/
+void FControlRigEditMode::DragAnimSliderTool(double IncrementVal)
+{
+	StaticCastSharedPtr<FControlRigEditModeToolkit>(Toolkit)->DragAnimSliderTool(IncrementVal);
+}
+
+/** Reset and stop user the anim slider tool*/
+void FControlRigEditMode::ResetAnimSlider()
+{
+	GEditor->EndTransaction();
+	bisTrackingAnimToolDrag = false;
+	StartXValue.Reset();
+	StaticCastSharedPtr<FControlRigEditModeToolkit>(Toolkit)->ResetAnimSlider();
+}
+
+/** If the Drag Anim Slider Tool is pressed*/
+bool FControlRigEditMode::IsDragAnimSliderToolPressed(FViewport* InViewport)
+{
+	if (IsInLevelEditor() && CanChangeAnimSliderTool())
+	{
+		bool bIsMovingSlider = false;
+		const FControlRigEditModeCommands& Commands = FControlRigEditModeCommands::Get();
+		// Need to iterate through primary and secondary to make sure they are all pressed.
+		for (uint32 i = 0; i < static_cast<uint8>(EMultipleKeyBindingIndex::NumChords); ++i)
+		{
+			EMultipleKeyBindingIndex ChordIndex = static_cast<EMultipleKeyBindingIndex>(i);
+			const FInputChord& Chord = *Commands.DragAnimSliderTool->GetActiveChord(ChordIndex);
+			bIsMovingSlider |= Chord.IsValidChord() && InViewport->KeyState(Chord.Key);
+		}
+		return bIsMovingSlider;
+	}
+	return false;
+}
+
 void FControlRigEditMode::OpenSpacePickerWidget()
 {
 	TMap<UControlRig*, TArray<FRigElementKey>> SelectedControlRigsAndControls;
diff --git a/Engine/Plugins/Animation/ControlRig/Source/ControlRigEditor/Private/EditMode/ControlRigEditMode.h b/Engine/Plugins/Animation/ControlRig/Source/ControlRigEditor/Private/EditMode/ControlRigEditMode.h
--- a/Engine/Plugins/Animation/ControlRig/Source/ControlRigEditor/Private/EditMode/ControlRigEditMode.h
+++ b/Engine/Plugins/Animation/ControlRig/Source/ControlRigEditor/Private/EditMode/ControlRigEditMode.h
@@ -113,6 +113,7 @@ public:
 	virtual bool InputKey(FEditorViewportClient* InViewportClient, FViewport* InViewport, FKey InKey, EInputEvent InEvent) override;
 	virtual bool EndTracking(FEditorViewportClient* InViewportClient, FViewport* InViewport) override;
 	virtual bool StartTracking(FEditorViewportClient* InViewportClient, FViewport* InViewport) override;
+	virtual bool ProcessCapturedMouseMoves(FEditorViewportClient* InViewportClient, FViewport* InViewport, const TArrayView<FIntPoint>& CapturedMouseMoves) override;
 	virtual bool HandleClick(FEditorViewportClient* InViewportClient, HHitProxy *HitProxy, const FViewportClick &Click) override;
 	virtual bool BoxSelect(FBox& InBox, bool InSelect = true) override;
 	virtual bool FrustumSelect(const FConvexVolume& InFrustum, FEditorViewportClient* InViewportClient, bool InSelect = true) override;
@@ -209,6 +210,21 @@ protected:
 	/** Toggles visibility of all  control rig shapes in the viewport */
 	void ToggleAllManipulators();
 
+	/** If Anim Slider is open, got to the next tool*/
+	void ChangeAnimSliderTool();
+
+	/** If Anim Slider is open, then can drag*/
+	bool CanChangeAnimSliderTool() const;
+
+	/** If Anim Slider is open, drag the tool*/
+	void DragAnimSliderTool(double IncrementVal);
+
+	/** Reset and stop user the anim slider tool*/
+	void ResetAnimSlider();
+
+	/** If the Drag Anim Slider Tool is pressed*/
+	bool IsDragAnimSliderToolPressed(FViewport* InViewport);
+
 public:
 	
 	/** Clear Selection*/
@@ -323,6 +339,11 @@ protected:
 	
 	/** Set initial transform handlers */
 	void OpenContextMenu(FEditorViewportClient* InViewportClient);
+
+	/* params to handle mouse move for changing anim tools sliders*/
+	bool bisTrackingAnimToolDrag = false;
+	/** Starting X Value*/
+	TOptional<int32> StartXValue;
 	
 public: 
 	/** Clear all selected RigElements */
diff --git a/Engine/Plugins/Animation/ControlRig/Source/ControlRigEditor/Private/EditMode/ControlRigEditModeCommands.cpp b/Engine/Plugins/Animation/ControlRig/Source/ControlRigEditor/Private/EditMode/ControlRigEditModeCommands.cpp
--- a/Engine/Plugins/Animation/ControlRig/Source/ControlRigEditor/Private/EditMode/ControlRigEditModeCommands.cpp
+++ b/Engine/Plugins/Animation/ControlRig/Source/ControlRigEditor/Private/EditMode/ControlRigEditModeCommands.cpp
@@ -9,7 +9,7 @@ void FControlRigEditModeCommands::RegisterCommands()
 	UI_COMMAND(ResetTransforms, "Reset Transform", "Reset the Controls Transforms", EUserInterfaceActionType::Button, FInputChord(EKeys::G, EModifierKey::Control));
 	UI_COMMAND(ResetAllTransforms, "Reset All Transform", "Reset all of the Controls Transforms", EUserInterfaceActionType::Button, FInputChord(EKeys::G, EModifierKey::Control | EModifierKey::Shift));
 	UI_COMMAND(ToggleManipulators, "Toggle Shapes", "Toggles visibility of active control rig shapes in the viewport", EUserInterfaceActionType::Button, FInputChord(EKeys::T));
-	UI_COMMAND(ToggleAllManipulators, "Toggle All Shapes", "Toggles visibility of all control rig shapes in the viewport", EUserInterfaceActionType::Button, FInputChord(EKeys::T, EModifierKey::Control | EModifierKey::Shift));
+	UI_COMMAND(ToggleAllManipulators, "Toggle All Shapes", "Toggles visibility of all control rig shapes in the viewport", EUserInterfaceActionType::Button, FInputChord(EKeys::T, EModifierKey::Control | EModifierKey::Alt));
 	UI_COMMAND(FrameSelection, "Frame Selection", "Focus the viewport on the current selection", EUserInterfaceActionType::Button, FInputChord(EKeys::F));
 	UI_COMMAND(ClearSelection, "Clear Selection", "Clear Selection", EUserInterfaceActionType::Button, FInputChord(EKeys::Escape));
 
@@ -17,6 +17,10 @@ void FControlRigEditModeCommands::RegisterCommands()
 	UI_COMMAND(DecreaseControlShapeSize, "Decrease Shape Size", "Decrease Size of the Shapes In The Viewport", EUserInterfaceActionType::Button, FInputChord(EKeys::Hyphen, EModifierKey::Shift));
 	UI_COMMAND(ResetControlShapeSize, "Reset Shape Size", "Resize Shape Size To Default", EUserInterfaceActionType::Button, FInputChord(EKeys::Equals));
 
+	UI_COMMAND(DragAnimSliderTool, "Drag Anim Slider Tool", "Drag existing anim slider", EUserInterfaceActionType::Button, FInputChord(EKeys::U));
+	UI_COMMAND(ChangeAnimSliderTool, "Change Anim Slider Tool", "Go to the next anim slider", EUserInterfaceActionType::Button, FInputChord(EKeys::U, EModifierKey::Shift));
+
+
 	UI_COMMAND(ToggleControlShapeTransformEdit, "Toggle Shape Transform Edit", "Toggle Editing Selected Control's Shape Transform", EUserInterfaceActionType::ToggleButton, FInputChord(EKeys::Period, EModifierKey::Control)); 
 	UI_COMMAND(OpenSpacePickerWidget, "Open the Space Picker", "Allows space switching on the control", EUserInterfaceActionType::Button, FInputChord(EKeys::Tab)); 
 }
diff --git a/Engine/Plugins/Animation/ControlRig/Source/ControlRigEditor/Private/EditMode/ControlRigEditModeCommands.h b/Engine/Plugins/Animation/ControlRig/Source/ControlRigEditor/Private/EditMode/ControlRigEditModeCommands.h
--- a/Engine/Plugins/Animation/ControlRig/Source/ControlRigEditor/Private/EditMode/ControlRigEditModeCommands.h
+++ b/Engine/Plugins/Animation/ControlRig/Source/ControlRigEditor/Private/EditMode/ControlRigEditModeCommands.h
@@ -23,9 +23,15 @@ public:
 	/** Toggles hiding all manipulators on active control rig in the viewport */
 	TSharedPtr< FUICommandInfo > ToggleManipulators;
 
-	/** Toggles hiding all manipulators on all control rigsin the viewport */
+	/** Toggles hiding all manipulators on all control rigs in the viewport */
 	TSharedPtr< FUICommandInfo > ToggleAllManipulators;
 
+	/** Drag Anim Slider Tool*/
+	TSharedPtr< FUICommandInfo > DragAnimSliderTool;
+
+	/**  Change current Anim Slider Tool */
+	TSharedPtr< FUICommandInfo > ChangeAnimSliderTool;
+
 	/** Reset Transforms for Controls */
 	TSharedPtr< FUICommandInfo > ResetTransforms;
 
diff --git a/Engine/Plugins/Animation/ControlRig/Source/ControlRigEditor/Private/EditMode/ControlRigEditModeToolkit.cpp b/Engine/Plugins/Animation/ControlRig/Source/ControlRigEditor/Private/EditMode/ControlRigEditModeToolkit.cpp
--- a/Engine/Plugins/Animation/ControlRig/Source/ControlRigEditor/Private/EditMode/ControlRigEditModeToolkit.cpp
+++ b/Engine/Plugins/Animation/ControlRig/Source/ControlRigEditor/Private/EditMode/ControlRigEditModeToolkit.cpp
@@ -121,7 +121,7 @@ void FControlRigEditModeToolkit::TryInvokeToolkitUI(const FName InName)
 	}
 	else if (InName == TweenOverlayName)
 	{
-		if(TweenWidget)
+		if(TweenWidgetParent)
 		{ 
 			RemoveAndDestroyTweenOverlay();
 		}
@@ -136,7 +136,7 @@ bool FControlRigEditModeToolkit::IsToolkitUIActive(const FName InName) const
 {
 	if (InName == TweenOverlayName)
 	{
-		return TweenWidget.IsValid();
+		return TweenWidgetParent.IsValid();
 	}
 
 	TSharedPtr<FAssetEditorModeUILayer> ModeUILayerPtr = ModeUILayer.Pin();
@@ -226,7 +226,7 @@ void FControlRigEditModeToolkit::CreateAndShowTweenOverlay()
 	}
 	UpdateTweenWidgetLocation(NewTweenWidgetLocation);
 
-	SAssignNew(TweenWidget, SHorizontalBox)
+	SAssignNew(TweenWidgetParent, SHorizontalBox)
 
 		+ SHorizontalBox::Slot()
 		.FillWidth(1.0f)
@@ -234,35 +234,62 @@ void FControlRigEditModeToolkit::CreateAndShowTweenOverlay()
 		.HAlign(HAlign_Left)
 		.Padding(TAttribute<FMargin>(this, &FControlRigEditModeToolkit::GetTweenWidgetPadding))
 		[
-			SNew(SControlRigTweenWidget)
+			SAssignNew(TweenWidget, SControlRigTweenWidget)
 			.InOwningToolkit(SharedThis(this))
 		];
 
 	TryShowTweenOverlay();
 }
 
+void FControlRigEditModeToolkit::GetToNextActiveSlider()
+{
+	if (TweenWidgetParent && TweenWidget)
+	{
+		TweenWidget->GetToNextActiveSlider();
+	}
+}
+bool FControlRigEditModeToolkit::CanChangeAnimSliderTool() const
+{
+	return (TweenWidgetParent && TweenWidget);
+}
+
+void FControlRigEditModeToolkit::DragAnimSliderTool(double IncrementVal)
+{
+	if (TweenWidgetParent && TweenWidget)
+	{
+		TweenWidget->DragAnimSliderTool(IncrementVal);
+	}
+}
+void FControlRigEditModeToolkit::ResetAnimSlider()
+{
+	if (TweenWidgetParent && TweenWidget)
+	{
+		TweenWidget->ResetAnimSlider();
+	}
+}
 void FControlRigEditModeToolkit::TryShowTweenOverlay()
 {
-	if (TweenWidget)
+	if (TweenWidgetParent)
 	{
-		GetToolkitHost()->AddViewportOverlayWidget(TweenWidget.ToSharedRef());
+		GetToolkitHost()->AddViewportOverlayWidget(TweenWidgetParent.ToSharedRef());
 	}
 }
 
 void FControlRigEditModeToolkit::RemoveAndDestroyTweenOverlay()
 {
 	TryRemoveTweenOverlay();
-	if (TweenWidget)
+	if (TweenWidgetParent)
 	{
+		TweenWidgetParent.Reset();
 		TweenWidget.Reset();
 	}
 }
 
 void FControlRigEditModeToolkit::TryRemoveTweenOverlay()
 {
-	if (IsHosted() && TweenWidget)
+	if (IsHosted() && TweenWidgetParent)
 	{
-		GetToolkitHost()->RemoveViewportOverlayWidget(TweenWidget.ToSharedRef());
+		GetToolkitHost()->RemoveViewportOverlayWidget(TweenWidgetParent.ToSharedRef());
 	}
 }
 
diff --git a/Engine/Plugins/Animation/ControlRig/Source/ControlRigEditor/Private/EditMode/ControlRigEditModeToolkit.h b/Engine/Plugins/Animation/ControlRig/Source/ControlRigEditor/Private/EditMode/ControlRigEditModeToolkit.h
--- a/Engine/Plugins/Animation/ControlRig/Source/ControlRigEditor/Private/EditMode/ControlRigEditModeToolkit.h
+++ b/Engine/Plugins/Animation/ControlRig/Source/ControlRigEditor/Private/EditMode/ControlRigEditModeToolkit.h
@@ -11,6 +11,8 @@
 #include "SControlRigEditModeTools.h"
 #include "ControlRigEditMode.h"
 
+class SControlRigTweenWidget;
+
 class FControlRigEditModeToolkit : public FModeToolkit
 {
 public:
@@ -50,6 +52,12 @@ public:
 	void TryInvokeToolkitUI(const FName InName);
 	bool IsToolkitUIActive(const FName InName) const;
 
+	/** For updating tween values via hot keys*/
+	void GetToNextActiveSlider();
+	bool CanChangeAnimSliderTool() const;
+	void DragAnimSliderTool(double IncrementVal);
+	void ResetAnimSlider();
+
 public:
 	static const FName PoseTabName;
 	static const FName MotionTrailTabName;
@@ -82,7 +90,10 @@ private:
 private:
 	/** The edit mode we are bound to */
 	FControlRigEditMode& EditMode;
-	TSharedPtr<SWidget> TweenWidget;
+	TSharedPtr<SWidget> TweenWidgetParent;
+	TSharedPtr<SControlRigTweenWidget> TweenWidget;
+
+
 	FVector2D InViewportTweenWidgetLocation;
 	/** The tools widget */
 	TSharedPtr<SControlRigEditModeTools> ModeTools;
diff --git a/Engine/Plugins/Animation/ControlRig/Source/ControlRigEditor/Private/EditMode/SControlRigTweenWidget.cpp b/Engine/Plugins/Animation/ControlRig/Source/ControlRigEditor/Private/EditMode/SControlRigTweenWidget.cpp
--- a/Engine/Plugins/Animation/ControlRig/Source/ControlRigEditor/Private/EditMode/SControlRigTweenWidget.cpp
+++ b/Engine/Plugins/Animation/ControlRig/Source/ControlRigEditor/Private/EditMode/SControlRigTweenWidget.cpp
@@ -43,8 +43,8 @@ void SControlRigTweenSlider::Construct(const FArguments& InArgs)
 			+ SHorizontalBox::Slot()
 				.AutoWidth()
 				.VAlign(VAlign_Center)
-				[
-					SNew(SSpinBox<double>)
+				[	
+					SAssignNew(SpinBox,SSpinBox<double>)
 					.PreventThrottling(true)
 					.Value(this, &SControlRigTweenSlider::OnGetPoseBlendValue)
 					.ToolTipText(LOCTEXT("TweenTooltip", "Key at current frame between previous(-1.0) and next(1.0) poses. Use Ctrl drag for under and over shoot."))
@@ -58,13 +58,14 @@ void SControlRigTweenSlider::Construct(const FArguments& InArgs)
 					.SupportDynamicSliderMinValue(true)
 					.SupportDynamicSliderMaxValue(true)
 					.ClearKeyboardFocusOnCommit(true)
+					.SelectAllTextOnCommit(false)
 					.OnValueChanged(this, &SControlRigTweenSlider::OnPoseBlendChanged)
 					.OnValueCommitted(this, &SControlRigTweenSlider::OnPoseBlendCommited)
 					.OnBeginSliderMovement(this, &SControlRigTweenSlider::OnBeginSliderMovement)
 					.OnEndSliderMovement(this, &SControlRigTweenSlider::OnEndSliderMovement)
+					
 				]
 			]
-		
 	];	
 }
 
@@ -79,6 +80,39 @@ void SControlRigTweenSlider::OnPoseBlendChanged(double ChangedVal)
 
 }
 
+void SControlRigTweenSlider::ResetAnimSlider()
+{
+	if (SpinBox.IsValid())
+	{
+		PoseBlendValue = 0.0;
+	}
+}
+
+void SControlRigTweenSlider::DragAnimSliderTool(double IncrementVal)
+{
+	if (SpinBox.IsValid())
+	{
+		//if control is down then act like we are overriding the slider
+		const bool bCtrlDown = FSlateApplication::Get().GetModifierKeys().IsControlDown();
+		const double MinSliderVal = bCtrlDown ? SpinBox->GetMinValue() : SpinBox->GetMinSliderValue();
+		const double MaxSliderVal = bCtrlDown ? SpinBox->GetMaxValue() : SpinBox->GetMaxSliderValue();
+
+		double NewVal = SpinBox->GetValue() + IncrementVal;
+		if (NewVal > MaxSliderVal)
+		{
+			NewVal = MaxSliderVal;
+		}
+		else if (NewVal < MinSliderVal)
+		{
+			NewVal = MinSliderVal;
+		}
+		Setup();
+		bIsBlending = true;
+		OnPoseBlendChanged(NewVal);
+		bIsBlending = false;
+	}
+}
+
 void SControlRigTweenSlider::OnBeginSliderMovement()
 {
 	if (bSliderStartedTransaction == false)
@@ -87,7 +121,6 @@ void SControlRigTweenSlider::OnBeginSliderMovement()
 		if (bIsBlending)
 		{
 			GEditor->BeginTransaction(LOCTEXT("AnimSliderBlend", "AnimSlider Blend"));
-
 		}
 	}
 }
@@ -110,11 +143,21 @@ void SControlRigTweenSlider::OnEndSliderMovement(double NewValue)
 	{
 		GEditor->EndTransaction();
 		bSliderStartedTransaction = false;
-
+		bIsBlending = false;
+		PoseBlendValue = 0.0f;
 	}
+	// Set focus back to the parent widget for users focusing the slider
+	TSharedRef<SWidget> ThisRef = AsShared();
+	FSlateApplication::Get().ForEachUser([&ThisRef](FSlateUser& User) {
+		if (User.HasFocusedDescendants(ThisRef) && ThisRef->IsParentValid())
+		{
+			User.SetFocus(ThisRef->GetParentWidget().ToSharedRef());
+		}
+	});
 	WeakSequencer = nullptr;
 }
 
+
 FReply SControlRigTweenWidget::OnDragDetected(const FGeometry& MyGeometry, const FPointerEvent& MouseEvent)
 {
 	// Need to remember where within a tab we grabbed
@@ -197,9 +240,9 @@ void SControlRigTweenWidget::Construct(const FArguments& InArgs)
 {
 	TSharedPtr<FBaseAnimSlider> TweenPtr = MakeShareable(new FControlsToTween());
 	AnimBlendTools.RegisterAnimSlider(TweenPtr);
-	TSharedPtr<FBaseAnimSlider> AmplifyPtr = MakeShareable(new FAmplifySlider());
-	AnimBlendTools.RegisterAnimSlider(AmplifyPtr);
-	TSharedPtr<FBaseAnimSlider> BlendNeighborPtr = MakeShareable(new FBlendToPreviousNextSlider());
+	TSharedPtr<FBaseAnimSlider> PushPullPtr = MakeShareable(new FPushPullSlider());
+	AnimBlendTools.RegisterAnimSlider(PushPullPtr);
+	TSharedPtr<FBaseAnimSlider> BlendNeighborPtr = MakeShareable(new FBlendNeighborSlider());
 	AnimBlendTools.RegisterAnimSlider(BlendNeighborPtr);
 
 	OwningToolkit = InArgs._InOwningToolkit;
@@ -286,7 +329,29 @@ void SControlRigTweenWidget::Construct(const FArguments& InArgs)
 			MainBox
 		]
 	];
+}
+
+void SControlRigTweenWidget::GetToNextActiveSlider()
+{
+	int32 Index = ActiveSlider < (AnimBlendTools.GetAnimSliders().Num() - 1) ? ActiveSlider + 1 : 0;
+	OnSelectSliderTool(Index);
+}
 
+void SControlRigTweenWidget::DragAnimSliderTool(double IncrementVal)
+{
+	if (SliderWidget.IsValid())
+	{
+		SliderWidget->DragAnimSliderTool(IncrementVal);
+	}
 }
 
+void SControlRigTweenWidget::ResetAnimSlider()
+{
+	if (SliderWidget.IsValid())
+	{
+		SliderWidget->ResetAnimSlider();
+	}
+}
+
+
 #undef LOCTEXT_NAMESPACE
diff --git a/Engine/Plugins/Animation/ControlRig/Source/ControlRigEditor/Private/EditMode/SControlRigTweenWidget.h b/Engine/Plugins/Animation/ControlRig/Source/ControlRigEditor/Private/EditMode/SControlRigTweenWidget.h
--- a/Engine/Plugins/Animation/ControlRig/Source/ControlRigEditor/Private/EditMode/SControlRigTweenWidget.h
+++ b/Engine/Plugins/Animation/ControlRig/Source/ControlRigEditor/Private/EditMode/SControlRigTweenWidget.h
@@ -15,6 +15,7 @@
 #include "Widgets/Views/STableRow.h"
 #include "Tools/ControlRigPose.h"
 #include "Tools/ControlRigTweener.h"
+#include "Widgets/Input/SSpinBox.h"
 
 class UControlRig;
 class ISequencer;
@@ -31,7 +32,8 @@ class SControlRigTweenSlider : public SCompoundWidget
 
 	void Construct(const FArguments& InArgs);
 	void SetAnimSlider(TSharedPtr<FBaseAnimSlider>& InAnimSlider) { AnimSlider = InAnimSlider; }
-
+	void DragAnimSliderTool(double IncrementVal);
+	void ResetAnimSlider();
 private:
 
 
@@ -52,6 +54,7 @@ private:
 	
 	TWeakPtr<ISequencer> WeakSequencer;
 	TSharedPtr<FBaseAnimSlider> AnimSlider;
+	TSharedPtr<SSpinBox<double>> SpinBox;
 
 };
 
@@ -67,14 +70,16 @@ class SControlRigTweenWidget : public SCompoundWidget
 
 	void Construct(const FArguments& InArgs);
 
+	void GetToNextActiveSlider();
+	void DragAnimSliderTool(double IncrementVal);
+	void ResetAnimSlider();
+
 private:
 
 	void OnSelectSliderTool(int32 Index);
 	FText GetActiveSliderName() const;
 	FText GetActiveSliderTooltip() const;
 
-
-
 	FReply OnMouseButtonDown(const FGeometry& MyGeometry, const FPointerEvent& MouseEvent) override
 	{
 		return FReply::Handled().DetectDrag(SharedThis(this), EKeys::LeftMouseButton);
diff --git a/Engine/Plugins/Animation/ControlRig/Source/ControlRigEditor/Private/Tools/ControlRigTweener.cpp b/Engine/Plugins/Animation/ControlRig/Source/ControlRigEditor/Private/Tools/ControlRigTweener.cpp
--- a/Engine/Plugins/Animation/ControlRig/Source/ControlRigEditor/Private/Tools/ControlRigTweener.cpp
+++ b/Engine/Plugins/Animation/ControlRig/Source/ControlRigEditor/Private/Tools/ControlRigTweener.cpp
@@ -578,7 +578,7 @@ double FControlsToTween::DoBlend(const double PreviousTime, const double Previou
 
 FText FControlsToTween::GetText()
 {
-	return LOCTEXT("TC", "TC");
+	return LOCTEXT("TW", "TW");
 }
 
 FText FControlsToTween::GetTooltipText()
@@ -598,12 +598,12 @@ bool FControlsToTween::Setup(TWeakPtr<ISequencer>& InSequencer)
 
 /*
 *
-*  Amplify Slider
+*  PushPull Slider
 *
 */
 
 
-double FAmplifySlider::DoBlend(const double PreviousTime, const double PreviousValue, const double CurrentTime, const double CurrentValue,
+double FPushPullSlider::DoBlend(const double PreviousTime, const double PreviousValue, const double CurrentTime, const double CurrentValue,
 	const double NextTime, const double NextValue, const double BlendValue)
 {
 	const double T = (CurrentTime - PreviousTime) / (NextTime - PreviousTime);
@@ -621,14 +621,14 @@ double FAmplifySlider::DoBlend(const double PreviousTime, const double PreviousV
 	return NewValue;
 }
 
-FText FAmplifySlider::GetText()
+FText FPushPullSlider::GetText()
 {
-	return LOCTEXT("AA", "AA");
+	return LOCTEXT("PP", "PP");
 }
 
-FText FAmplifySlider::GetTooltipText()
+FText FPushPullSlider::GetTooltipText()
 {
-	return LOCTEXT("AmplifySliderTooltip", "Amplify or diminish between the next and previous keys");
+	return LOCTEXT("PushPullTooltip", "Push or pull the values to the interpolation between the previous and next keys");
 }
 
 
@@ -638,7 +638,7 @@ FText FAmplifySlider::GetTooltipText()
 *
 */
 
-double FBlendToPreviousNextSlider::DoBlend(const double PreviousTime, const double PreviousValue, const double CurrentTime, const double CurrentValue,
+double FBlendNeighborSlider::DoBlend(const double PreviousTime, const double PreviousValue, const double CurrentTime, const double CurrentValue,
 	const double NextTime, const double NextValue, const double BlendValue)
 {
 	double NewValue;
@@ -653,12 +653,12 @@ double FBlendToPreviousNextSlider::DoBlend(const double PreviousTime, const doub
 	return NewValue;
 }
 
-FText FBlendToPreviousNextSlider::GetText()
+FText FBlendNeighborSlider::GetText()
 {
-	return LOCTEXT("PN", "PN");
+	return LOCTEXT("BN", "BN");
 }
 
-FText FBlendToPreviousNextSlider::GetTooltipText()
+FText FBlendNeighborSlider::GetTooltipText()
 {
 	return LOCTEXT("BlendToNeighborsSliderTooltip", "Blend to the next or previous values for selected keys or objects");
 }
diff --git a/Engine/Plugins/Animation/ControlRig/Source/ControlRigEditor/Public/Tools/ControlRigTweener.h b/Engine/Plugins/Animation/ControlRig/Source/ControlRigEditor/Public/Tools/ControlRigTweener.h
--- a/Engine/Plugins/Animation/ControlRig/Source/ControlRigEditor/Public/Tools/ControlRigTweener.h
+++ b/Engine/Plugins/Animation/ControlRig/Source/ControlRigEditor/Public/Tools/ControlRigTweener.h
@@ -197,7 +197,7 @@ struct FControlsToTween :public FBasicBlendSlider
 *   Slider that will amplify or reduce the differences from the slope between the next or previous
 *
 **/
-struct FAmplifySlider :public FBasicBlendSlider
+struct FPushPullSlider :public FBasicBlendSlider
 {
 	//BaseAnimSlider overrides
 	virtual FText GetText();
@@ -213,7 +213,7 @@ struct FAmplifySlider :public FBasicBlendSlider
 *   Similar to tween but will blend off of the current position when blending to the next or previous keys
 *
 **/
-struct FBlendToPreviousNextSlider : public FBasicBlendSlider
+struct FBlendNeighborSlider : public FBasicBlendSlider
 {
 	//BaseAnimSlider overrides
 	virtual FText GetText();
