Control Rig: Fixed crash when renaming rig element in details panel

diff --git a/Engine/Plugins/Animation/ControlRig/Source/ControlRigEditor/Private/ControlRigElementDetails.cpp b/Engine/Plugins/Animation/ControlRig/Source/ControlRigEditor/Private/ControlRigElementDetails.cpp
--- a/Engine/Plugins/Animation/ControlRig/Source/ControlRigEditor/Private/ControlRigElementDetails.cpp
+++ b/Engine/Plugins/Animation/ControlRig/Source/ControlRigEditor/Private/ControlRigElementDetails.cpp
@@ -648,7 +648,11 @@ void FRigBaseElementDetails::CustomizeDetails(IDetailLayoutBuilder& DetailBuilde
 FRigElementKey FRigBaseElementDetails::GetElementKey() const
 {
 	check(PerElementInfos.Num() == 1);
-	return PerElementInfos[0].GetElement()->GetKey();
+	if (FRigBaseElement* Element = PerElementInfos[0].GetElement())
+	{
+		return Element->GetKey();
+	}
+	return FRigElementKey();
 }
 
 FText FRigBaseElementDetails::GetName() const
