Update URigHierarchy to use PostEditUndo

#rnx

diff --git a/Engine/Plugins/Animation/ControlRig/Source/ControlRig/Private/Rigs/RigHierarchy.cpp b/Engine/Plugins/Animation/ControlRig/Source/ControlRig/Private/Rigs/RigHierarchy.cpp
--- a/Engine/Plugins/Animation/ControlRig/Source/ControlRig/Private/Rigs/RigHierarchy.cpp
+++ b/Engine/Plugins/Animation/ControlRig/Source/ControlRig/Private/Rigs/RigHierarchy.cpp
@@ -2205,16 +2205,14 @@ bool URigHierarchy::SetTransformStackIndex(int32 InTransformStackIndex)
 
 #if WITH_EDITOR
 
-void URigHierarchy::PostTransacted(const FTransactionObjectEvent& TransactionEvent)
+void URigHierarchy::PostEditUndo()
 {
-	if (TransactionEvent.GetEventType() == ETransactionObjectEventType::UndoRedo)
+	Super::PostEditUndo();
+
+	const int32 DesiredStackIndex = TransformStackIndex;
+	TransformStackIndex = TransformUndoStack.Num();
+	if (DesiredStackIndex != TransformStackIndex)
 	{
-		const int32 DesiredStackIndex = TransformStackIndex;
-		TransformStackIndex = TransformUndoStack.Num();
-		if (DesiredStackIndex == TransformStackIndex)
-		{
-			return;
-		}
 		SetTransformStackIndex(DesiredStackIndex);
 	}
 }
diff --git a/Engine/Plugins/Animation/ControlRig/Source/ControlRig/Public/Rigs/RigHierarchy.h b/Engine/Plugins/Animation/ControlRig/Source/ControlRig/Public/Rigs/RigHierarchy.h
--- a/Engine/Plugins/Animation/ControlRig/Source/ControlRig/Public/Rigs/RigHierarchy.h
+++ b/Engine/Plugins/Animation/ControlRig/Source/ControlRig/Public/Rigs/RigHierarchy.h
@@ -108,7 +108,7 @@ public:
 	void Load(FArchive& Ar);
 	virtual void PostLoad() override;
 #if WITH_EDITOR
-	virtual void PostTransacted(const FTransactionObjectEvent& TransactionEvent) override;
+	virtual void PostEditUndo() override;
 #endif
 
 	/**
