AssetRegistry updates for Niagara/ControlRig: Call AssetTagsFinalized after asynchronous operations that modify the data that contributes to tags.
#rnx

diff --git a/Engine/Plugins/Animation/ControlRig/Source/ControlRigDeveloper/Private/ControlRigBlueprint.cpp b/Engine/Plugins/Animation/ControlRig/Source/ControlRigDeveloper/Private/ControlRigBlueprint.cpp
--- a/Engine/Plugins/Animation/ControlRig/Source/ControlRigDeveloper/Private/ControlRigBlueprint.cpp
+++ b/Engine/Plugins/Animation/ControlRig/Source/ControlRigDeveloper/Private/ControlRigBlueprint.cpp
@@ -796,6 +796,7 @@ void UControlRigBlueprint::PreSave(FObjectPreSaveContext ObjectSaveContext)
 	}
 
 	FunctionReferenceNodeData = GetReferenceNodeData();
+	IAssetRegistry::GetChecked().AssetTagsFinalized(*this);
 }
 
 void UControlRigBlueprint::PostLoad()
@@ -1086,7 +1087,7 @@ void UControlRigBlueprint::HandlePackageDone()
 			if(!FunctionLibrary->FunctionReferences_DEPRECATED.IsEmpty())
 			{
 				// let's also update the asset data of the dependents
-				const FAssetRegistryModule& AssetRegistryModule = FModuleManager::LoadModuleChecked<FAssetRegistryModule>("AssetRegistry");
+				IAssetRegistry& AssetRegistry = IAssetRegistry::GetChecked();
 				
 				for(const TTuple< TObjectPtr<URigVMLibraryNode>, FRigVMFunctionReferenceArray >& Pair :
 					FunctionLibrary->FunctionReferences_DEPRECATED)
@@ -1099,7 +1100,7 @@ void UControlRigBlueprint::HandlePackageDone()
 						BuildData->RegisterFunctionReference(FunctionKey, Pair.Value[ReferenceIndex]);
 
 						// find all control rigs matching the reference node
-						FAssetData AssetData = AssetRegistryModule.Get().GetAssetByObjectPath(
+						FAssetData AssetData = AssetRegistry.GetAssetByObjectPath(
 							Pair.Value[ReferenceIndex].ToSoftObjectPath().GetAssetPathName());
 
 						// if the asset has never been loaded - make sure to load it once and mark as dirty
@@ -2364,7 +2365,7 @@ TArray<FAssetData> UControlRigBlueprint::GetDependentAssets() const
 
 	if(URigVMFunctionLibrary* FunctionLibrary = RigVMClient.GetFunctionLibrary())
 	{
-		const FAssetRegistryModule& AssetRegistryModule = FModuleManager::LoadModuleChecked<FAssetRegistryModule>(TEXT("AssetRegistry"));
+		IAssetRegistry& AssetRegistry = IAssetRegistry::GetChecked();
 		TArray<URigVMLibraryNode*> Functions = FunctionLibrary->GetFunctions();
 		for(URigVMLibraryNode* Function : Functions)
 		{
@@ -2384,7 +2385,7 @@ TArray<FAssetData> UControlRigBlueprint::GetDependentAssets() const
 					{
 						AssetPaths.Add(AssetPath);
 
-						const FAssetData AssetData = AssetRegistryModule.Get().GetAssetByObjectPath(*AssetPath.ToString());
+						const FAssetData AssetData = AssetRegistry.GetAssetByObjectPath(*AssetPath.ToString());
 						if(AssetData.IsValid())
 						{
 							Dependents.Add(AssetData);
diff --git a/Engine/Plugins/FX/Niagara/Source/Niagara/Classes/NiagaraSystem.h b/Engine/Plugins/FX/Niagara/Source/Niagara/Classes/NiagaraSystem.h
--- a/Engine/Plugins/FX/Niagara/Source/Niagara/Classes/NiagaraSystem.h
+++ b/Engine/Plugins/FX/Niagara/Source/Niagara/Classes/NiagaraSystem.h
@@ -734,6 +734,8 @@ private:
 	*/
 	bool QueryCompileComplete(bool bWait, bool bDoPost, bool bDoNotApply = false);
 
+	void BroadcastOnSystemCompiled();
+
 	void PreProcessWaitingDDCTasks(bool bProcessForWait);
 
 	bool CompilationResultsValid(FNiagaraSystemCompileRequest& CompileRequest) const;
diff --git a/Engine/Plugins/FX/Niagara/Source/Niagara/Private/NiagaraSystem.cpp b/Engine/Plugins/FX/Niagara/Source/Niagara/Private/NiagaraSystem.cpp
--- a/Engine/Plugins/FX/Niagara/Source/Niagara/Private/NiagaraSystem.cpp
+++ b/Engine/Plugins/FX/Niagara/Source/Niagara/Private/NiagaraSystem.cpp
@@ -3,6 +3,7 @@
 
 #include "NiagaraSystem.h"
 
+#include "AssetRegistry/IAssetRegistry.h"
 #include "INiagaraEditorOnlyDataUtlities.h"
 #include "NiagaraAsyncCompile.h"
 #include "NiagaraConstants.h"
@@ -2723,7 +2724,7 @@ bool UNiagaraSystem::QueryCompileComplete(bool bWait, bool bDoPost, bool bDoNotA
 		{
 			SCOPE_CYCLE_COUNTER(STAT_Niagara_System_CompileScriptResetAfter);
 
-			OnSystemCompiled().Broadcast(this);
+			BroadcastOnSystemCompiled();
 		}
 
 		return true;
@@ -2732,6 +2733,14 @@ bool UNiagaraSystem::QueryCompileComplete(bool bWait, bool bDoPost, bool bDoNotA
 	return false;
 }
 
+void UNiagaraSystem::BroadcastOnSystemCompiled()
+{
+	OnSystemCompiled().Broadcast(this);
+#if WITH_EDITOR
+	IAssetRegistry::GetChecked().AssetTagsFinalized(*this);
+#endif
+}
+
 void UNiagaraSystem::InitEmitterVariableAliasNames(FNiagaraEmitterCompiledData& EmitterCompiledDataToInit, const UNiagaraEmitter* InAssociatedEmitter)
 {
 	EmitterCompiledDataToInit.EmitterSpawnIntervalVar.SetName(GetEmitterVariableAliasName(SYS_PARAM_EMITTER_SPAWN_INTERVAL, InAssociatedEmitter));
