Control Rig: Remove the output flag for the SimulatedResult on the quaternion version of SpringInterp.

#rnx

diff --git a/Engine/Plugins/Animation/ControlRig/Source/ControlRig/Private/Units/Simulation/RigUnit_SpringInterp.h b/Engine/Plugins/Animation/ControlRig/Source/ControlRig/Private/Units/Simulation/RigUnit_SpringInterp.h
--- a/Engine/Plugins/Animation/ControlRig/Source/ControlRig/Private/Units/Simulation/RigUnit_SpringInterp.h
+++ b/Engine/Plugins/Animation/ControlRig/Source/ControlRig/Private/Units/Simulation/RigUnit_SpringInterp.h
@@ -361,7 +361,7 @@ struct CONTROLRIG_API FRigUnit_SpringInterpQuaternionV2 : public FRigUnit_SimBas
 	UPROPERTY(meta = (Output))
 	FVector AngularVelocity;
 
-	UPROPERTY(meta=(Output))
+	UPROPERTY()
 	FQuat SimulatedResult;
 
 	UPROPERTY()
