Control Rig: Properly handle breakpoints with multiple events

diff --git a/Engine/Plugins/Animation/ControlRig/Source/ControlRig/Private/ControlRig.cpp b/Engine/Plugins/Animation/ControlRig/Source/ControlRig/Private/ControlRig.cpp
--- a/Engine/Plugins/Animation/ControlRig/Source/ControlRig/Private/ControlRig.cpp
+++ b/Engine/Plugins/Animation/ControlRig/Source/ControlRig/Private/ControlRig.cpp
@@ -324,9 +324,35 @@ void UControlRig::Evaluate_AnyThread()
 #if WITH_EDITOR
 		if (VM)
 		{
-			if (VM->GetHaltedAtBreakpoint().IsValid())
+			const FRigVMBreakpoint& Breakpoint = VM->GetHaltedAtBreakpoint(); 
+			if (Breakpoint.IsValid())
 			{
-				break;
+				// make sure that the instruction index for the breakpoint is part
+				// of the current entry.
+				const FRigVMByteCode& ByteCode = VM->GetByteCode();
+				const int32 CurrentEntryIndex = ByteCode.FindEntryIndex(EventName);
+				if(CurrentEntryIndex != INDEX_NONE)
+				{
+					const FRigVMByteCodeEntry& CurrentEntry = ByteCode.GetEntry(CurrentEntryIndex);
+					if(Breakpoint.InstructionIndex >= CurrentEntry.InstructionIndex)
+					{
+						int32 LastInstructionIndexInEntry;
+						if(CurrentEntryIndex == ByteCode.NumEntries() - 1)
+						{
+							LastInstructionIndexInEntry = ByteCode.GetNumInstructions();
+						}
+						else
+						{
+							const int32 NextEntryIndex = CurrentEntryIndex + 1;
+							LastInstructionIndexInEntry = ByteCode.GetEntry(NextEntryIndex).InstructionIndex - 1;
+						}
+
+						if(Breakpoint.InstructionIndex <= LastInstructionIndexInEntry)
+						{
+							break;
+						}
+					}
+				}
 			}
 		}
 #endif
@@ -1495,11 +1521,17 @@ bool UControlRig::ExecuteUnits(FRigUnitContext& InOutContext, const FName& InEve
 			{
 				if(URigVM* SnapShotVM = GetSnapshotVM(false)) // don't create it for normal runs
 				{
+					const bool bIsEventFirstInQueue = !EventQueueToRun.IsEmpty() && EventQueueToRun[0] == InEventName; 
+					const bool bIsEventLastInQueue = !EventQueueToRun.IsEmpty() && EventQueueToRun.Last() == InEventName;
+
 					if (VM->GetHaltedAtBreakpoint().IsValid())
 					{
-						VM->CopyFrom(SnapShotVM, false, false, false, true, true);	
+						if(bIsEventFirstInQueue)
+						{
+							VM->CopyFrom(SnapShotVM, false, false, false, true, true);
+						}
 					}
-					else
+					else if(bIsEventLastInQueue)
 					{
 						SnapShotVM->CopyFrom(VM, false, false, false, true, true);
 					}
