Control Rig: Fix crash on Packaged Third Person application during GC

diff --git a/Engine/Plugins/Animation/ControlRig/Source/ControlRig/Private/ControlRig.cpp b/Engine/Plugins/Animation/ControlRig/Source/ControlRig/Private/ControlRig.cpp
--- a/Engine/Plugins/Animation/ControlRig/Source/ControlRig/Private/ControlRig.cpp
+++ b/Engine/Plugins/Animation/ControlRig/Source/ControlRig/Private/ControlRig.cpp
@@ -195,7 +195,7 @@ void UControlRig::Initialize(bool bInitRigUnits)
 	// recompute the hash used to differentiate VMs based on their memory layout 
 	if (HasAnyFlags(RF_ClassDefaultObject))
 	{
-		CachedMemoryHash = INDEX_NONE;
+		CachedMemoryHash = 0;
 
 		if(VM)
 		{
@@ -636,6 +636,14 @@ void UControlRig::InstantiateVMFromCDO()
 		}
 	}
 
+	if(VM)
+	{
+		CachedMemoryHash = HashCombine(
+			VM->GetLiteralMemory()->GetMemoryHash(),
+			VM->GetWorkMemory()->GetMemoryHash()
+		);
+	}
+
 	RequestInit();
 }
 
@@ -1473,7 +1481,7 @@ bool UControlRig::ExecuteUnits(FRigUnitContext& InOutContext, const FName& InEve
 				UControlRig* CDO = Cast<UControlRig>(GetClass()->GetDefaultObject());
 
 				bool bIsValidSnapshot = false;
-				if(SnapshotHash != INDEX_NONE)
+				if(SnapshotHash != 0)
 				{
 					TObjectPtr<URigVM>* InitializedVMSnapshotPtr = CDO->InitializedVMSnapshots.Find(SnapshotHash);
 					if(InitializedVMSnapshotPtr && !InitializedVMSnapshotPtr->IsNull())
@@ -1507,6 +1515,12 @@ bool UControlRig::ExecuteUnits(FRigUnitContext& InOutContext, const FName& InEve
 					InitializedVMSnapshot->WorkMemoryStorageObject->CopyFrom(VM->WorkMemoryStorageObject);
 
 					CDO->InitializedVMSnapshots.Add(SnapshotHash, InitializedVMSnapshot);
+
+					// GC won't consider some subobjects that are created after the constructor as part of the CDO,
+					// so even if CDO is rooted and references these sub objects, 
+					// it is not enough to keep them alive.
+					// Hence, we have to add them to root here.
+					InitializedVMSnapshot->AddToRoot();
 				}
 			}
 			else
@@ -3386,9 +3400,9 @@ void UControlRig::PostInitInstance(UControlRig* InCDO)
 
 uint32 UControlRig::GetHashForInitializeVMSnapShot()
 {
-	if(CachedMemoryHash == INDEX_NONE)
+	if(CachedMemoryHash == 0)
 	{
-		return INDEX_NONE;
+		return 0;
 	}
 	
 	uint32 Hash = GetHierarchy()->GetNameHash();
diff --git a/Engine/Plugins/Animation/ControlRig/Source/ControlRig/Public/ControlRig.h b/Engine/Plugins/Animation/ControlRig/Source/ControlRig/Public/ControlRig.h
--- a/Engine/Plugins/Animation/ControlRig/Source/ControlRig/Public/ControlRig.h
+++ b/Engine/Plugins/Animation/ControlRig/Source/ControlRig/Public/ControlRig.h
@@ -500,7 +500,7 @@ protected:
 	// and stores an initialized VM per hash.
 	UPROPERTY(transient)
 	TMap<uint32, TObjectPtr<URigVM>> InitializedVMSnapshots;
-	uint32 CachedMemoryHash = INDEX_NONE;
+	uint32 CachedMemoryHash = 0;
 
 	// computes the hash used to store / find VM snapshots
 	uint32 GetHashForInitializeVMSnapShot();
