Control Rig: Fix crash during basic interactions with Rig Hierarchy elements

diff --git a/Engine/Plugins/Animation/ControlRig/Source/ControlRigEditor/Private/EditMode/SControlRigEditModeTools.cpp b/Engine/Plugins/Animation/ControlRig/Source/ControlRigEditor/Private/EditMode/SControlRigEditModeTools.cpp
--- a/Engine/Plugins/Animation/ControlRig/Source/ControlRigEditor/Private/EditMode/SControlRigEditModeTools.cpp
+++ b/Engine/Plugins/Animation/ControlRig/Source/ControlRigEditor/Private/EditMode/SControlRigEditModeTools.cpp
@@ -921,7 +921,6 @@ FReply SControlRigEditModeTools::OnBakeControlsToNewSpaceButtonClicked()
 				FFrameNumber CurrentTime = FrameTime.GetFrame();
 				FMovieSceneControlRigSpaceBaseKey Value;
 				using namespace UE::MovieScene;
-				URigHierarchy* RigHierarchy = SpacePickerWidget->GetHierarchy();
 				Settings.TargetSpace = URigHierarchy::GetDefaultParentKey();
 
 				TRange<FFrameNumber> Range = Sequencer->GetFocusedMovieSceneSequence()->GetMovieScene()->GetPlaybackRange();
diff --git a/Engine/Plugins/Animation/ControlRig/Source/ControlRigEditor/Private/Editor/SRigSpacePickerWidget.cpp b/Engine/Plugins/Animation/ControlRig/Source/ControlRigEditor/Private/Editor/SRigSpacePickerWidget.cpp
--- a/Engine/Plugins/Animation/ControlRig/Source/ControlRigEditor/Private/Editor/SRigSpacePickerWidget.cpp
+++ b/Engine/Plugins/Animation/ControlRig/Source/ControlRigEditor/Private/Editor/SRigSpacePickerWidget.cpp
@@ -196,7 +196,7 @@ SRigSpacePickerWidget::~SRigSpacePickerWidget()
 
 	if(HierarchyModifiedHandle.IsValid())
 	{
-		if(Hierarchy)
+		if(Hierarchy.IsValid())
 		{
 			Hierarchy->OnModified().Remove(HierarchyModifiedHandle);
 			HierarchyModifiedHandle.Reset();
@@ -208,23 +208,24 @@ void SRigSpacePickerWidget::SetControls(
 	URigHierarchy* InHierarchy,
 	const TArray<FRigElementKey>& InControls)
 {
-	if(Hierarchy && Hierarchy != InHierarchy)
+	if(Hierarchy.IsValid())
 	{
-		if(HierarchyModifiedHandle.IsValid())
+		URigHierarchy* StrongHierarchy = Hierarchy.Get();
+		if (StrongHierarchy != InHierarchy)
 		{
-			if(IsValid(Hierarchy) && !Hierarchy->HasAnyFlags(RF_BeginDestroyed) && !Hierarchy->HasAnyFlags(RF_TagGarbageTemp))
+			if(HierarchyModifiedHandle.IsValid())
 			{
-				Hierarchy->OnModified().Remove(HierarchyModifiedHandle);
+				StrongHierarchy->OnModified().Remove(HierarchyModifiedHandle);
+				HierarchyModifiedHandle.Reset();
 			}
-			HierarchyModifiedHandle.Reset();
 		}
 	}
 	
 	Hierarchy = InHierarchy;
 	ControlKeys = InControls;
-	if (Hierarchy && HierarchyModifiedHandle.IsValid() == false)
+	if (Hierarchy.IsValid() && HierarchyModifiedHandle.IsValid() == false)
 	{
-		HierarchyModifiedHandle = Hierarchy->OnModified().AddSP(this, &SRigSpacePickerWidget::OnHierarchyModified);
+		HierarchyModifiedHandle = InHierarchy->OnModified().AddSP(this, &SRigSpacePickerWidget::OnHierarchyModified);
 	}
 	UpdateActiveSpaces();
 	RepopulateItemSpaces();
@@ -328,15 +329,19 @@ void SRigSpacePickerWidget::Tick(const FGeometry& AllottedGeometry, const double
 	}
 	else if(GetAdditionalSpacesDelegate.IsBound())
 	{
-		TArray<FRigElementKey> CurrentAdditionalSpaces;
-		for(const FRigElementKey& ControlKey: ControlKeys)
+		if (Hierarchy.IsValid())
 		{
-			CurrentAdditionalSpaces.Append(GetAdditionalSpacesDelegate.Execute(Hierarchy, ControlKey));
-		}
+			URigHierarchy* StrongHierarchy = Hierarchy.Get();
+			TArray<FRigElementKey> CurrentAdditionalSpaces;
+			for(const FRigElementKey& ControlKey: ControlKeys)
+			{
+				CurrentAdditionalSpaces.Append(GetAdditionalSpacesDelegate.Execute(StrongHierarchy, ControlKey));
+			}
 		
-		if(CurrentAdditionalSpaces != AdditionalSpaces)
-		{
-			RepopulateItemSpaces();
+			if(CurrentAdditionalSpaces != AdditionalSpaces)
+			{
+				RepopulateItemSpaces();
+			}
 		}
 	}
 }
@@ -519,11 +524,16 @@ FReply SRigSpacePickerWidget::HandleWorldSpaceClicked()
 
 FReply SRigSpacePickerWidget::HandleElementSpaceClicked(FRigElementKey InKey)
 {
-	//need to make copy since array may get shrunk during the event broadcast
-	TArray<FRigElementKey> ControlKeysCopy = ControlKeys;
-	for (const FRigElementKey& ControlKey : ControlKeysCopy)
+	if (Hierarchy.IsValid())
 	{
-		ActiveSpaceChangedEvent.Broadcast(Hierarchy, ControlKey, InKey);
+		URigHierarchy* StrongHierarchy = Hierarchy.Get();
+		
+		//need to make copy since array may get shrunk during the event broadcast
+		TArray<FRigElementKey> ControlKeysCopy = ControlKeys;
+		for (const FRigElementKey& ControlKey : ControlKeysCopy)
+		{
+			ActiveSpaceChangedEvent.Broadcast(StrongHierarchy, ControlKey, InKey);
+		}
 	}
 
 	if(DialogWindow.IsValid())
@@ -546,9 +556,13 @@ FReply SRigSpacePickerWidget::HandleSpaceMoveUp(FRigElementKey InKey)
 				TArray<FRigElementKey> ChangedSpaceKeys = CurrentSpaceKeys;
 				ChangedSpaceKeys.Swap(Index, Index - 1);
 
-				for(const FRigElementKey& ControlKey : ControlKeys)
+				if (Hierarchy.IsValid())
 				{
-					SpaceListChangedEvent.Broadcast(Hierarchy, ControlKey, ChangedSpaceKeys);
+					URigHierarchy* StrongHierarchy = Hierarchy.Get();
+					for(const FRigElementKey& ControlKey : ControlKeys)
+					{
+						SpaceListChangedEvent.Broadcast(StrongHierarchy, ControlKey, ChangedSpaceKeys);
+					}
 				}
 
 				return FReply::Handled();
@@ -570,9 +584,13 @@ FReply SRigSpacePickerWidget::HandleSpaceMoveDown(FRigElementKey InKey)
 				TArray<FRigElementKey> ChangedSpaceKeys = CurrentSpaceKeys;
 				ChangedSpaceKeys.Swap(Index, Index + 1);
 
-				for(const FRigElementKey& ControlKey : ControlKeys)
+				if (Hierarchy.IsValid())
 				{
-					SpaceListChangedEvent.Broadcast(Hierarchy, ControlKey, ChangedSpaceKeys);
+					URigHierarchy* StrongHierarchy = Hierarchy.Get();
+					for(const FRigElementKey& ControlKey : ControlKeys)
+					{
+						SpaceListChangedEvent.Broadcast(StrongHierarchy, ControlKey, ChangedSpaceKeys);
+					}
 				}
 				
 				return FReply::Handled();
@@ -587,9 +605,13 @@ void SRigSpacePickerWidget::HandleSpaceDelete(FRigElementKey InKey)
 	TArray<FRigElementKey> ChangedSpaceKeys = CurrentSpaceKeys;
 	if(ChangedSpaceKeys.Remove(InKey) > 0)
 	{
-		for(const FRigElementKey& ControlKey : ControlKeys)
+		if (Hierarchy.IsValid())
 		{
-			SpaceListChangedEvent.Broadcast(Hierarchy, ControlKey, ChangedSpaceKeys);
+			URigHierarchy* StrongHierarchy = Hierarchy.Get();
+			for(const FRigElementKey& ControlKey : ControlKeys)
+			{
+				SpaceListChangedEvent.Broadcast(StrongHierarchy, ControlKey, ChangedSpaceKeys);
+			}
 		}
 	}
 }
@@ -605,44 +627,48 @@ FReply SRigSpacePickerWidget::HandleAddElementClicked()
 			const FRigElementKey Key = InItem->Key;
 			if(!IsDefaultSpace(Key) && IsValidKey(Key))
 			{
-				for(const FRigElementKey& ControlKey : ControlKeys)
+				if (Hierarchy.IsValid())
 				{
-					URigHierarchy::TElementDependencyMap DependencyMap;
-					FString FailureReason;
-
-					if(UControlRig* ControlRig = Hierarchy->GetTypedOuter<UControlRig>())
+					URigHierarchy* StrongHierarchy = Hierarchy.Get();
+					for(const FRigElementKey& ControlKey : ControlKeys)
 					{
-						DependencyMap = Hierarchy->GetDependenciesForVM(ControlRig->GetVM()); 
-					}
-					else if(UControlRigBlueprint* RigBlueprint = Hierarchy->GetTypedOuter<UControlRigBlueprint>())
-					{
-						if(UControlRig* CDO = Cast<UControlRig>(RigBlueprint->GetControlRigClass()->GetDefaultObject()))
+						URigHierarchy::TElementDependencyMap DependencyMap;
+						FString FailureReason;
+
+						if(UControlRig* ControlRig = StrongHierarchy->GetTypedOuter<UControlRig>())
+						{
+							DependencyMap = StrongHierarchy->GetDependenciesForVM(ControlRig->GetVM()); 
+						}
+						else if(UControlRigBlueprint* RigBlueprint = StrongHierarchy->GetTypedOuter<UControlRigBlueprint>())
+						{
+							if(UControlRig* CDO = Cast<UControlRig>(RigBlueprint->GetControlRigClass()->GetDefaultObject()))
+							{
+								DependencyMap = StrongHierarchy->GetDependenciesForVM(CDO->GetVM()); 
+							}
+						}
+						
+						if(!StrongHierarchy->CanSwitchToParent(ControlKey, Key, DependencyMap, &FailureReason))
 						{
-							DependencyMap = Hierarchy->GetDependenciesForVM(CDO->GetVM()); 
+							// notification
+							FNotificationInfo Info(FText::FromString(FailureReason));
+							Info.bFireAndForget = true;
+							Info.FadeOutDuration = 2.0f;
+							Info.ExpireDuration = 8.0f;
+
+							const TSharedPtr<SNotificationItem> NotificationPtr = FSlateNotificationManager::Get().AddNotification(Info);
+							NotificationPtr->SetCompletionState(SNotificationItem::CS_Fail);
+							return;
 						}
 					}
 					
-					if(!Hierarchy->CanSwitchToParent(ControlKey, Key, DependencyMap, &FailureReason))
+					TArray<FRigElementKey> ChangedSpaceKeys = CurrentSpaceKeys;
+					ChangedSpaceKeys.AddUnique(Key);
+
+					for(const FRigElementKey& ControlKey : ControlKeys)
 					{
-						// notification
-						FNotificationInfo Info(FText::FromString(FailureReason));
-						Info.bFireAndForget = true;
-						Info.FadeOutDuration = 2.0f;
-						Info.ExpireDuration = 8.0f;
-
-						const TSharedPtr<SNotificationItem> NotificationPtr = FSlateNotificationManager::Get().AddNotification(Info);
-						NotificationPtr->SetCompletionState(SNotificationItem::CS_Fail);
-						return;
+						SpaceListChangedEvent.Broadcast(StrongHierarchy, ControlKey, ChangedSpaceKeys);
 					}
 				}
-				
-				TArray<FRigElementKey> ChangedSpaceKeys = CurrentSpaceKeys;
-				ChangedSpaceKeys.AddUnique(Key);
-
-				for(const FRigElementKey& ControlKey : ControlKeys)
-				{
-					SpaceListChangedEvent.Broadcast(Hierarchy, ControlKey, ChangedSpaceKeys);
-				}
 			}
 		}
 
@@ -825,10 +851,12 @@ void SRigSpacePickerWidget::RepopulateItemSpaces()
 	{
 		return;
 	}
-	if(Hierarchy == nullptr)
+	if(!Hierarchy.IsValid())
 	{
 		return;
 	}
+
+	URigHierarchy* StrongHierarchy = Hierarchy.Get();
 	
 	TArray<FRigElementKey> FavoriteKeys, SpacesFromDelegate;
 
@@ -839,7 +867,7 @@ void SRigSpacePickerWidget::RepopulateItemSpaces()
 			const FRigControlElementCustomization* Customization = nullptr;
 			if(GetControlCustomizationDelegate.IsBound())
 			{
-				Customization = GetControlCustomizationDelegate.Execute(Hierarchy, ControlKey);
+				Customization = GetControlCustomizationDelegate.Execute(StrongHierarchy, ControlKey);
 			}
 
 			if(Customization)
@@ -855,7 +883,7 @@ void SRigSpacePickerWidget::RepopulateItemSpaces()
 			}
 			
 			// check if the customization is different from the base one in the asset
-			if(const FRigControlElement* ControlElement = Hierarchy->Find<FRigControlElement>(ControlKey))
+			if(const FRigControlElement* ControlElement = StrongHierarchy->Find<FRigControlElement>(ControlKey))
 			{
 				if(Customization != &ControlElement->Settings.Customization)
 				{
@@ -890,7 +918,7 @@ void SRigSpacePickerWidget::RepopulateItemSpaces()
 		AdditionalSpaces.Reset();
 		for(const FRigElementKey& ControlKey: ControlKeys)
 		{
-			AdditionalSpaces.Append(GetAdditionalSpacesDelegate.Execute(Hierarchy, ControlKey));
+			AdditionalSpaces.Append(GetAdditionalSpacesDelegate.Execute(StrongHierarchy, ControlKey));
 		}
 		
 		for(const FRigElementKey& Key : AdditionalSpaces)
@@ -946,7 +974,7 @@ void SRigSpacePickerWidget::RepopulateItemSpaces()
 
 	for(const FRigElementKey& Key : Keys)
 	{
-		TPair<const FSlateBrush*, FSlateColor> IconAndColor = SRigHierarchyItem::GetBrushForElementType(Hierarchy, Key);
+		TPair<const FSlateBrush*, FSlateColor> IconAndColor = SRigHierarchyItem::GetBrushForElementType(StrongHierarchy, Key);
 		
 		AddSpacePickerRow(
 			ItemSpacesListBox,
@@ -971,18 +999,19 @@ void SRigSpacePickerWidget::UpdateActiveSpaces()
 {
 	ActiveSpaceKeys.Reset();
 
-	if(Hierarchy == nullptr)
+	if(!Hierarchy.IsValid())
 	{
 		return;
 	}
-	
+
+	URigHierarchy* StrongHierarchy = Hierarchy.Get();	
 	for(int32 ControlIndex=0;ControlIndex<ControlKeys.Num();ControlIndex++)
 	{
 		ActiveSpaceKeys.Add(URigHierarchy::GetDefaultParentKey());
 
 		if(GetActiveSpaceDelegate.IsBound())
 		{
-			ActiveSpaceKeys[ControlIndex] = GetActiveSpaceDelegate.Execute(Hierarchy, ControlKeys[ControlIndex]);
+			ActiveSpaceKeys[ControlIndex] = GetActiveSpaceDelegate.Execute(StrongHierarchy, ControlKeys[ControlIndex]);
 		}
 	}
 }
diff --git a/Engine/Plugins/Animation/ControlRig/Source/ControlRigEditor/Private/Editor/SRigSpacePickerWidget.h b/Engine/Plugins/Animation/ControlRig/Source/ControlRigEditor/Private/Editor/SRigSpacePickerWidget.h
--- a/Engine/Plugins/Animation/ControlRig/Source/ControlRigEditor/Private/Editor/SRigSpacePickerWidget.h
+++ b/Engine/Plugins/Animation/ControlRig/Source/ControlRigEditor/Private/Editor/SRigSpacePickerWidget.h
@@ -71,8 +71,16 @@ public:
 	virtual bool SupportsKeyboardFocus() const override;
 	virtual void Tick(const FGeometry& AllottedGeometry, const double InCurrentTime, const float InDeltaTime) override;
 
-	URigHierarchy* GetHierarchy() const { return Hierarchy; }
-	const URigHierarchy* GetHierarchyConst() const { return Hierarchy; }
+	URigHierarchy* GetHierarchy() const
+	{
+		if (Hierarchy.IsValid())
+		{
+			return Hierarchy.Get();
+		}
+		return nullptr;
+	}
+	const URigHierarchy* GetHierarchyConst() const { return GetHierarchy(); }
+	
 	const TArray<FRigElementKey>& GetControls() const { return ControlKeys; }
 	const TArray<FRigElementKey>& GetActiveSpaces() const;
 	TArray<FRigElementKey> GetDefaultSpaces() const;
@@ -135,7 +143,7 @@ private:
 	FRigSpacePickerActiveSpaceChanged ActiveSpaceChangedEvent;
 	FRigSpacePickerSpaceListChanged SpaceListChangedEvent;
 
-	URigHierarchy* Hierarchy;
+	TWeakObjectPtr<URigHierarchy> Hierarchy;
 	TArray<FRigElementKey> ControlKeys;
 	TArray<FRigElementKey> CurrentSpaceKeys;
 	TArray<FRigElementKey> ActiveSpaceKeys;
