Control Rig: Fix crash during refresh hierarchy

diff --git a/Engine/Plugins/Animation/ControlRig/Source/ControlRigEditor/Private/Editor/SRigHierarchy.cpp b/Engine/Plugins/Animation/ControlRig/Source/ControlRigEditor/Private/Editor/SRigHierarchy.cpp
--- a/Engine/Plugins/Animation/ControlRig/Source/ControlRigEditor/Private/Editor/SRigHierarchy.cpp
+++ b/Engine/Plugins/Animation/ControlRig/Source/ControlRigEditor/Private/Editor/SRigHierarchy.cpp
@@ -1336,6 +1336,8 @@ void SRigHierarchy::RefreshHierarchy(const FAssetData& InAssetData)
 		return;
 	}
 
+	ControlRigEditor.Pin()->ClearDetailObject();
+
 	URigHierarchy* Hierarchy = GetHierarchy();
 	USkeletalMesh* Mesh = Cast<USkeletalMesh>(InAssetData.GetAsset());
 	if (Mesh && Hierarchy)
