Sequencer: Anim Sliders.  Changed sliders to work on absolute and not incremental, should help performance and fixes issues with indirect manipulation. Added some checks to handle channels being empty on additive sections.

diff --git a/Engine/Plugins/Animation/ControlRig/Source/ControlRigEditor/Private/EditMode/ControlRigEditMode.cpp b/Engine/Plugins/Animation/ControlRig/Source/ControlRigEditor/Private/EditMode/ControlRigEditMode.cpp
--- a/Engine/Plugins/Animation/ControlRig/Source/ControlRigEditor/Private/EditMode/ControlRigEditMode.cpp
+++ b/Engine/Plugins/Animation/ControlRig/Source/ControlRigEditor/Private/EditMode/ControlRigEditMode.cpp
@@ -977,7 +977,6 @@ bool FControlRigEditMode::ProcessCapturedMouseMoves(FEditorViewportClient* InVie
 						const double ViewDiff = (double)(Diff) / ViewPortSize;
 						DragAnimSliderTool(ViewDiff);
 					}
-					StartXValue = X;
 				}
 			}
 		}
diff --git a/Engine/Plugins/Animation/ControlRig/Source/ControlRigEditor/Private/EditMode/ControlRigEditMode.h b/Engine/Plugins/Animation/ControlRig/Source/ControlRigEditor/Private/EditMode/ControlRigEditMode.h
--- a/Engine/Plugins/Animation/ControlRig/Source/ControlRigEditor/Private/EditMode/ControlRigEditMode.h
+++ b/Engine/Plugins/Animation/ControlRig/Source/ControlRigEditor/Private/EditMode/ControlRigEditMode.h
@@ -217,7 +217,7 @@ protected:
 	bool CanChangeAnimSliderTool() const;
 
 	/** If Anim Slider is open, drag the tool*/
-	void DragAnimSliderTool(double IncrementVal);
+	void DragAnimSliderTool(double Val);
 
 	/** Reset and stop user the anim slider tool*/
 	void ResetAnimSlider();
@@ -344,6 +344,7 @@ protected:
 	bool bisTrackingAnimToolDrag = false;
 	/** Starting X Value*/
 	TOptional<int32> StartXValue;
+
 	
 public: 
 	/** Clear all selected RigElements */
diff --git a/Engine/Plugins/Animation/ControlRig/Source/ControlRigEditor/Private/EditMode/ControlRigEditModeToolkit.cpp b/Engine/Plugins/Animation/ControlRig/Source/ControlRigEditor/Private/EditMode/ControlRigEditModeToolkit.cpp
--- a/Engine/Plugins/Animation/ControlRig/Source/ControlRigEditor/Private/EditMode/ControlRigEditModeToolkit.cpp
+++ b/Engine/Plugins/Animation/ControlRig/Source/ControlRigEditor/Private/EditMode/ControlRigEditModeToolkit.cpp
@@ -253,11 +253,11 @@ bool FControlRigEditModeToolkit::CanChangeAnimSliderTool() const
 	return (TweenWidgetParent && TweenWidget);
 }
 
-void FControlRigEditModeToolkit::DragAnimSliderTool(double IncrementVal)
+void FControlRigEditModeToolkit::DragAnimSliderTool(double Val)
 {
 	if (TweenWidgetParent && TweenWidget)
 	{
-		TweenWidget->DragAnimSliderTool(IncrementVal);
+		TweenWidget->DragAnimSliderTool(Val);
 	}
 }
 void FControlRigEditModeToolkit::ResetAnimSlider()
diff --git a/Engine/Plugins/Animation/ControlRig/Source/ControlRigEditor/Private/EditMode/ControlRigEditModeToolkit.h b/Engine/Plugins/Animation/ControlRig/Source/ControlRigEditor/Private/EditMode/ControlRigEditModeToolkit.h
--- a/Engine/Plugins/Animation/ControlRig/Source/ControlRigEditor/Private/EditMode/ControlRigEditModeToolkit.h
+++ b/Engine/Plugins/Animation/ControlRig/Source/ControlRigEditor/Private/EditMode/ControlRigEditModeToolkit.h
@@ -55,7 +55,7 @@ public:
 	/** For updating tween values via hot keys*/
 	void GetToNextActiveSlider();
 	bool CanChangeAnimSliderTool() const;
-	void DragAnimSliderTool(double IncrementVal);
+	void DragAnimSliderTool(double Val);
 	void ResetAnimSlider();
 
 public:
diff --git a/Engine/Plugins/Animation/ControlRig/Source/ControlRigEditor/Private/EditMode/SControlRigTweenWidget.cpp b/Engine/Plugins/Animation/ControlRig/Source/ControlRigEditor/Private/EditMode/SControlRigTweenWidget.cpp
--- a/Engine/Plugins/Animation/ControlRig/Source/ControlRigEditor/Private/EditMode/SControlRigTweenWidget.cpp
+++ b/Engine/Plugins/Animation/ControlRig/Source/ControlRigEditor/Private/EditMode/SControlRigTweenWidget.cpp
@@ -29,10 +29,8 @@ void SControlRigTweenSlider::Construct(const FArguments& InArgs)
 	bIsBlending = false;
 	bSliderStartedTransaction = false;
 	AnimSlider = InArgs._InAnimSlider;
-	FText Text = AnimSlider->GetText();
-	FText TooltipText = AnimSlider->GetTooltipText();
 	ChildSlot
-[
+	[
 		SNew(SVerticalBox)
 		+ SVerticalBox::Slot()
 		.AutoHeight()
@@ -47,7 +45,11 @@ void SControlRigTweenSlider::Construct(const FArguments& InArgs)
 					SAssignNew(SpinBox,SSpinBox<double>)
 					.PreventThrottling(true)
 					.Value(this, &SControlRigTweenSlider::OnGetPoseBlendValue)
-					.ToolTipText(LOCTEXT("TweenTooltip", "Key at current frame between previous(-1.0) and next(1.0) poses. Use Ctrl drag for under and over shoot."))
+					.ToolTipText_Lambda([this]()
+						{
+							FText TooltipText = AnimSlider->GetTooltipText();
+							return TooltipText;
+						})
 					.MinValue(-2.0)
 					.MaxValue(2.0)
 					.MinSliderValue(-1.0)
@@ -88,7 +90,7 @@ void SControlRigTweenSlider::ResetAnimSlider()
 	}
 }
 
-void SControlRigTweenSlider::DragAnimSliderTool(double IncrementVal)
+void SControlRigTweenSlider::DragAnimSliderTool(double Val)
 {
 	if (SpinBox.IsValid())
 	{
@@ -97,7 +99,7 @@ void SControlRigTweenSlider::DragAnimSliderTool(double IncrementVal)
 		const double MinSliderVal = bCtrlDown ? SpinBox->GetMinValue() : SpinBox->GetMinSliderValue();
 		const double MaxSliderVal = bCtrlDown ? SpinBox->GetMaxValue() : SpinBox->GetMaxSliderValue();
 
-		double NewVal = SpinBox->GetValue() + IncrementVal;
+		double NewVal = Val;
 		if (NewVal > MaxSliderVal)
 		{
 			NewVal = MaxSliderVal;
@@ -191,21 +193,25 @@ void SControlRigTweenWidget::FinishDraggingWidget(const FVector2D InLocation)
 
 void SControlRigTweenSlider::OnPoseBlendCommited(double ChangedVal, ETextCommit::Type Type)
 {
-	if (bIsBlending == false)
+	if (SpinBox.IsValid() && SpinBox->HasKeyboardFocus())
 	{
-		bIsBlending = Setup();
+		if (bIsBlending == false)
+		{
+			bIsBlending = Setup();
 
-	}
-	if(bIsBlending)
-	{
-		FScopedTransaction ScopedTransaction(LOCTEXT("TweenTransaction", "Tween"));
-		PoseBlendValue = ChangedVal;
-		OnPoseBlendChanged(ChangedVal);
-		bIsBlending = false;
-		PoseBlendValue = 0.0f;
+		}
+		if (bIsBlending)
+		{
+			FScopedTransaction ScopedTransaction(LOCTEXT("TweenTransaction", "Tween"));
+			PoseBlendValue = ChangedVal;
+			OnPoseBlendChanged(ChangedVal);
+			bIsBlending = false;
+			PoseBlendValue = 0.0f;
+		}
 	}
 }
 
+
 void SControlRigTweenWidget::OnSelectSliderTool(int32 Index)
 {
 	ActiveSlider = Index;
@@ -337,11 +343,11 @@ void SControlRigTweenWidget::GetToNextActiveSlider()
 	OnSelectSliderTool(Index);
 }
 
-void SControlRigTweenWidget::DragAnimSliderTool(double IncrementVal)
+void SControlRigTweenWidget::DragAnimSliderTool(double Val)
 {
 	if (SliderWidget.IsValid())
 	{
-		SliderWidget->DragAnimSliderTool(IncrementVal);
+		SliderWidget->DragAnimSliderTool(Val);
 	}
 }
 
diff --git a/Engine/Plugins/Animation/ControlRig/Source/ControlRigEditor/Private/EditMode/SControlRigTweenWidget.h b/Engine/Plugins/Animation/ControlRig/Source/ControlRigEditor/Private/EditMode/SControlRigTweenWidget.h
--- a/Engine/Plugins/Animation/ControlRig/Source/ControlRigEditor/Private/EditMode/SControlRigTweenWidget.h
+++ b/Engine/Plugins/Animation/ControlRig/Source/ControlRigEditor/Private/EditMode/SControlRigTweenWidget.h
@@ -32,7 +32,7 @@ class SControlRigTweenSlider : public SCompoundWidget
 
 	void Construct(const FArguments& InArgs);
 	void SetAnimSlider(TSharedPtr<FBaseAnimSlider>& InAnimSlider) { AnimSlider = InAnimSlider; }
-	void DragAnimSliderTool(double IncrementVal);
+	void DragAnimSliderTool(double Val);
 	void ResetAnimSlider();
 private:
 
@@ -71,7 +71,7 @@ class SControlRigTweenWidget : public SCompoundWidget
 	void Construct(const FArguments& InArgs);
 
 	void GetToNextActiveSlider();
-	void DragAnimSliderTool(double IncrementVal);
+	void DragAnimSliderTool(double Val);
 	void ResetAnimSlider();
 
 private:
diff --git a/Engine/Plugins/Animation/ControlRig/Source/ControlRigEditor/Private/Tools/ControlRigTweener.cpp b/Engine/Plugins/Animation/ControlRig/Source/ControlRigEditor/Private/Tools/ControlRigTweener.cpp
--- a/Engine/Plugins/Animation/ControlRig/Source/ControlRigEditor/Private/Tools/ControlRigTweener.cpp
+++ b/Engine/Plugins/Animation/ControlRig/Source/ControlRigEditor/Private/Tools/ControlRigTweener.cpp
@@ -387,6 +387,10 @@ bool FAnimSliderObjectSelection::Setup(const TArray<UControlRig*>& SelectedContr
 							KeyTimes.SetNum(0);
 							Handles.SetNum(0);
 							FChannelMapInfo* pChannelIndex = Section->ControlChannelMap.Find(ControlElement->GetName());
+							if (pChannelIndex == nullptr)
+							{
+								continue;
+							}
 							int NumChannels = 0;
 							switch (ControlElement->Settings.ControlType)
 							{
