Sequencer: Consraints: Add/Move/Delete constraint keys for transform tracks
Main change was to move code out of ConstraintChannelHelper to a new MovieSceneConstraintChannelHelper so that we can use the same functions for add/delete/move for both TransformableHandle types. It would have been better perhaps to keep this in control rig plugin, it would have made things much easier, but the Sequencer architecture limits us to only one track editor per track type so we had to move code into the transform track editor, which meant moving functionality to MovieSceneTools module. but do to this I had to add some modularity which means less code duplication. In particularly
Added an Interface IMovieSceneConstrainedSection, to handle accessing contraint channels. Also moved getting channels and setting transforms to the TransfomableHandles. Both of these changes will make moving more things like compensentation to runtime if needed later. Also added a snapper fix for sub sections

diff --git a/Engine/Plugins/Animation/ControlRig/Source/ControlRig/Private/Constraints/ControlRigTransformableHandle.cpp b/Engine/Plugins/Animation/ControlRig/Source/ControlRig/Private/Constraints/ControlRigTransformableHandle.cpp
--- a/Engine/Plugins/Animation/ControlRig/Source/ControlRig/Private/Constraints/ControlRigTransformableHandle.cpp
+++ b/Engine/Plugins/Animation/ControlRig/Source/ControlRig/Private/Constraints/ControlRigTransformableHandle.cpp
@@ -6,6 +6,9 @@
 #include "ControlRig.h"
 #include "IControlRigObjectBinding.h"
 #include "Rigs/RigHierarchyElements.h"
+#include "Channels/MovieSceneChannelProxy.h"
+#include "Sequencer/MovieSceneControlRigParameterSection.h"
+#include "Sections/MovieScene3DTransformSection.h"
 
 /**
  * UTransformableControlHandle
@@ -260,6 +263,111 @@ void UTransformableControlHandle::OnControlModified(
 	}
 }
 
+static TPair<const FChannelMapInfo*, int32> GetInfoAndNumFloatChannels(
+	const UControlRig* InControlRig,
+	const FName& InControlName,
+	const UMovieSceneControlRigParameterSection* InSection)
+{
+	const FRigControlElement* ControlElement = InControlRig ? InControlRig->FindControl(InControlName) : nullptr;
+	auto GetNumFloatChannels = [](const ERigControlType& InControlType)
+	{
+		switch (InControlType)
+		{
+		case ERigControlType::Position:
+		case ERigControlType::Scale:
+		case ERigControlType::Rotator:
+			return 3;
+		case ERigControlType::TransformNoScale:
+			return 6;
+		case ERigControlType::Transform:
+		case ERigControlType::EulerTransform:
+			return 9;
+		default:
+			break;
+		}
+		return 0;
+	};
+
+	const int32 NumFloatChannels = ControlElement ? GetNumFloatChannels(ControlElement->Settings.ControlType) : 0;
+	const FChannelMapInfo* ChannelInfo = InSection ? InSection->ControlChannelMap.Find(InControlName) : nullptr;
+
+	return { ChannelInfo, NumFloatChannels };
+}
+TArrayView<FMovieSceneFloatChannel*>  UTransformableControlHandle::GetFloatChannels(const UMovieSceneSection* InSection) const
+{
+	// no floats for transform sections
+	static const TArrayView<FMovieSceneFloatChannel*> EmptyChannelsView;
+
+	const FChannelMapInfo* ChannelInfo = nullptr;
+	int32 NumChannels = 0;
+	const UMovieSceneControlRigParameterSection* CRSection = Cast<UMovieSceneControlRigParameterSection>(InSection);
+	if (CRSection == nullptr)
+	{
+		return EmptyChannelsView;
+	}
+
+	Tie(ChannelInfo, NumChannels) = GetInfoAndNumFloatChannels(ControlRig.Get(),ControlName, CRSection);
+
+	if (ChannelInfo == nullptr || NumChannels == 0)
+	{
+		return EmptyChannelsView;
+	}
+
+	// return a sub view that just represents the control's channels
+	const TArrayView<FMovieSceneFloatChannel*> FloatChannels = InSection->GetChannelProxy().GetChannels<FMovieSceneFloatChannel>();
+	const int32 ChannelStartIndex = ChannelInfo->ChannelIndex;
+	return FloatChannels.Slice(ChannelStartIndex, NumChannels);
+}
+
+TArrayView<FMovieSceneDoubleChannel*>  UTransformableControlHandle::GetDoubleChannels(const UMovieSceneSection* InSection) const
+{
+	static const TArrayView<FMovieSceneDoubleChannel*> EmptyChannelsView;
+	return EmptyChannelsView;
+}
+
+bool UTransformableControlHandle::AddTransformKeys(const TArray<FFrameNumber>& InFrames,
+	const TArray<FTransform>& InTransforms,
+	const EMovieSceneTransformChannel& InChannels,
+	const FFrameRate& InTickResolution,
+	UMovieSceneSection*,
+	const bool bLocal) const
+{
+	if (!ControlRig.IsValid() || ControlName == NAME_None || InFrames.IsEmpty() || InFrames.Num() != InTransforms.Num())
+	{
+		return false;
+	}
+	auto KeyframeFunc = [this, bLocal](const FTransform& InTransform, const FRigControlModifiedContext& InKeyframeContext)
+	{
+		UControlRig* InControlRig = ControlRig.Get();
+		static constexpr bool bNotify = true;
+		static constexpr bool bUndo = false;
+		static constexpr bool bFixEuler = true;
+
+		if (bLocal)
+		{
+			return InControlRig->SetControlLocalTransform(ControlName, InTransform, bNotify, InKeyframeContext, bUndo, bFixEuler);
+		}
+		InControlRig->SetControlGlobalTransform(ControlName, InTransform, bNotify, InKeyframeContext, bUndo, bFixEuler);
+	};
+
+	FRigControlModifiedContext KeyframeContext;
+	KeyframeContext.SetKey = EControlRigSetKey::Always;
+	KeyframeContext.KeyMask = static_cast<uint32>(InChannels);
+
+	for (int32 Index = 0; Index < InFrames.Num(); ++Index)
+	{
+		const FFrameNumber& Frame = InFrames[Index];
+		KeyframeContext.LocalTime = InTickResolution.AsSeconds(FFrameTime(Frame));
+
+		KeyframeFunc(InTransforms[Index], KeyframeContext);
+	}
+
+	return true;
+}
+
+
+
+
 #if WITH_EDITOR
 
 FString UTransformableControlHandle::GetLabel() const
diff --git a/Engine/Plugins/Animation/ControlRig/Source/ControlRig/Private/Sequencer/MovieSceneControlRigParameterSection.cpp b/Engine/Plugins/Animation/ControlRig/Source/ControlRig/Private/Sequencer/MovieSceneControlRigParameterSection.cpp
--- a/Engine/Plugins/Animation/ControlRig/Source/ControlRig/Private/Sequencer/MovieSceneControlRigParameterSection.cpp
+++ b/Engine/Plugins/Animation/ControlRig/Source/ControlRig/Private/Sequencer/MovieSceneControlRigParameterSection.cpp
@@ -1293,7 +1293,7 @@ FConstraintAndActiveChannel* UMovieSceneControlRigParameterSection::GetConstrain
 	return (Index != INDEX_NONE) ? &ConstraintsChannels[Index] : nullptr;	
 }
 
-void UMovieSceneControlRigParameterSection::AddConstraintChannel(UTickableConstraint* InConstraint, bool bReconstructChannel)
+void UMovieSceneControlRigParameterSection::AddConstraintChannel(UTickableConstraint* InConstraint)
 {
 	if (!HasConstraintChannel(InConstraint->GetFName()))
 	{
@@ -1308,11 +1308,9 @@ void UMovieSceneControlRigParameterSection::AddConstraintChannel(UTickableConstr
 		{
 			OnConstraintChannelAdded.Broadcast(this, ExistingChannel);
 		}
-		
-		if (bReconstructChannel)
-		{
-			ReconstructChannelProxy();
-		}
+		//todo got rid of the if(bReconstructChannel) flag since it was always true but it may need to be false from undo, in which case we need to change
+		//change this virtual functions signature
+		ReconstructChannelProxy();
 	}
 }
 
@@ -1331,12 +1329,12 @@ void UMovieSceneControlRigParameterSection::RemoveConstraintChannel(const FName&
 	}
 }
 
-const TArray<FConstraintAndActiveChannel>& UMovieSceneControlRigParameterSection::GetConstraintsChannels() const
+TArray<FConstraintAndActiveChannel>& UMovieSceneControlRigParameterSection::GetConstraintsChannels() 
 {
 	return ConstraintsChannels;
 }
 
-TArray<FConstraintAndActiveChannel>& UMovieSceneControlRigParameterSection::GetConstraintsChannels()
+const TArray<FConstraintAndActiveChannel>& UMovieSceneControlRigParameterSection::GetConstraintsChannels() const
 {
 	return ConstraintsChannels;
 }
diff --git a/Engine/Plugins/Animation/ControlRig/Source/ControlRig/Private/Sequencer/MovieSceneControlRigParameterTrack.cpp b/Engine/Plugins/Animation/ControlRig/Source/ControlRig/Private/Sequencer/MovieSceneControlRigParameterTrack.cpp
--- a/Engine/Plugins/Animation/ControlRig/Source/ControlRig/Private/Sequencer/MovieSceneControlRigParameterTrack.cpp
+++ b/Engine/Plugins/Animation/ControlRig/Source/ControlRig/Private/Sequencer/MovieSceneControlRigParameterTrack.cpp
@@ -94,7 +94,7 @@ void UMovieSceneControlRigParameterTrack::HandleOnSpaceNoLongerUsed(FMovieSceneC
 }
 
 void UMovieSceneControlRigParameterTrack::HandleOnConstraintAdded(
-	UMovieSceneControlRigParameterSection* InSection,
+	IMovieSceneConstrainedSection* InSection,
 	FMovieSceneConstraintChannel* InChannel) const
 {
 	OnConstraintChannelAdded.Broadcast(InSection, InChannel);
diff --git a/Engine/Plugins/Animation/ControlRig/Source/ControlRig/Public/Constraints/ControlRigTransformableHandle.h b/Engine/Plugins/Animation/ControlRig/Source/ControlRig/Public/Constraints/ControlRigTransformableHandle.h
--- a/Engine/Plugins/Animation/ControlRig/Source/ControlRig/Public/Constraints/ControlRigTransformableHandle.h
+++ b/Engine/Plugins/Animation/ControlRig/Source/ControlRig/Public/Constraints/ControlRigTransformableHandle.h
@@ -51,6 +51,17 @@ public:
 	/** @todo document */
 	virtual TWeakObjectPtr<UObject> GetTarget() const override;
 
+	/** Get the array of float channels for the specified section*/
+	virtual TArrayView<FMovieSceneFloatChannel*>  GetFloatChannels(const UMovieSceneSection* InSection) const override;
+	/** Get the array of double channels for the specified section*/
+	virtual TArrayView<FMovieSceneDoubleChannel*>  GetDoubleChannels(const UMovieSceneSection* InSection) const override;
+	virtual bool AddTransformKeys(const TArray<FFrameNumber>& InFrames,
+		const TArray<FTransform>& InTransforms,
+		const EMovieSceneTransformChannel& InChannels,
+		const FFrameRate& InTickResolution,
+		UMovieSceneSection* InSection,
+		const bool bLocal = true) const override;
+
 	/** Returns the skeletal mesh bound to ControlRig. */
 	USkeletalMeshComponent* GetSkeletalMesh() const;
 
diff --git a/Engine/Plugins/Animation/ControlRig/Source/ControlRig/Public/Sequencer/MovieSceneControlRigParameterSection.h b/Engine/Plugins/Animation/ControlRig/Source/ControlRig/Public/Sequencer/MovieSceneControlRigParameterSection.h
--- a/Engine/Plugins/Animation/ControlRig/Source/ControlRig/Public/Sequencer/MovieSceneControlRigParameterSection.h
+++ b/Engine/Plugins/Animation/ControlRig/Source/ControlRig/Public/Sequencer/MovieSceneControlRigParameterSection.h
@@ -5,6 +5,7 @@
 #include "CoreMinimal.h"
 #include "ConstraintsManager.h"
 #include "Sections/MovieSceneParameterSection.h"
+#include "Sections/MovieSceneConstrainedSection.h"
 #include "UObject/ObjectMacros.h"
 #include "Channels/MovieSceneFloatChannel.h"
 #include "Sections/MovieSceneSubSection.h"
@@ -175,7 +176,7 @@ struct FMovieSceneControlRigSpaceChannel;
  * Movie scene section that controls animation controller animation
  */
 UCLASS()
-class CONTROLRIG_API UMovieSceneControlRigParameterSection : public UMovieSceneParameterSection
+class CONTROLRIG_API UMovieSceneControlRigParameterSection : public UMovieSceneParameterSection, public IMovieSceneConstrainedSection
 {
 	GENERATED_BODY()
 
@@ -183,7 +184,6 @@ public:
 
 	/** Bindable events for when we add space or constraint channels. */
 	DECLARE_MULTICAST_DELEGATE_ThreeParams(FSpaceChannelAddedEvent, UMovieSceneControlRigParameterSection*, const FName&, FMovieSceneControlRigSpaceChannel*);
-	DECLARE_MULTICAST_DELEGATE_TwoParams(FConstraintChannelAddedEvent, UMovieSceneControlRigParameterSection*, FMovieSceneConstraintChannel*);
 
 	void AddEnumParameterKey(FName InParameterName, FFrameNumber InTime, uint8 InValue);
 	void AddIntegerParameterKey(FName InParameterName, FFrameNumber InTime, int32 InValue);
@@ -205,18 +205,12 @@ public:
 	
 	FSpaceChannelAddedEvent& SpaceChannelAdded() { return OnSpaceChannelAdded; }
 
-	TArray<FConstraintAndActiveChannel>& GetConstraintsChannels();
-	const TArray<FConstraintAndActiveChannel>& GetConstraintsChannels() const;
 	const FName& FindControlNameFromConstraintChannel(const FMovieSceneConstraintChannel* InConstraintChannel) const;
-	
-	FConstraintChannelAddedEvent& ConstraintChannelAdded() { return OnConstraintChannelAdded; }
-	
+		
 	bool RenameParameterName(const FName& OldParameterName, const FName& NewParameterName);
 private:
 
 	FSpaceChannelAddedEvent OnSpaceChannelAdded;
-	FConstraintChannelAddedEvent OnConstraintChannelAdded;
-
 	/** Control Rig that controls us*/
 	UPROPERTY()
 	TObjectPtr<UControlRig> ControlRig;
@@ -267,6 +261,34 @@ public:
 	//UMovieSceneSection virtuals
 	virtual void SetBlendType(EMovieSceneBlendType InBlendType) override;
 	virtual UObject* GetImplicitObjectOwner() override;
+	// IMovieSceneConstrainedSection overrides
+	/*
+	* Whether it has that channel
+	*/
+	virtual bool HasConstraintChannel(const FName& InConstraintName) const override;
+
+	/*
+	* Get constraint with that name
+	*/
+	virtual FConstraintAndActiveChannel* GetConstraintChannel(const FName& InConstraintName) override;
+
+	/*
+	*  Add Constraint channel
+	*/
+	virtual void AddConstraintChannel(UTickableConstraint* InConstraint) override;
+
+	/*
+	*  Remove Constraint channel
+	*/
+	virtual void RemoveConstraintChannel(const FName& InConstraintName) override;
+
+	/*
+	*  Get The channels
+	*/
+	virtual TArray<FConstraintAndActiveChannel>& GetConstraintsChannels()  override;
+
+	//not override but needed
+	const TArray<FConstraintAndActiveChannel>& GetConstraintsChannels() const;
 
 #if WITH_EDITOR
 	//Function to save control rig key when recording.
@@ -404,12 +426,6 @@ public:
 	/** Add Space Parameter for a specified Control, no Default since that is Parent space*/
 	void AddSpaceChannel(FName InControlName, bool bReconstructChannel);
 
-	/** todo */
-	FDelegateHandle OnConstraintRemovedHandle;
-	bool HasConstraintChannel(const FName& InConstraintName) const;
-	FConstraintAndActiveChannel* GetConstraintChannel(const FName& InConstraintName);
-	void AddConstraintChannel(UTickableConstraint* InConstraint, bool bReconstructChannel);
-	void RemoveConstraintChannel(const FName& InConstraintName);
 
 	/** Clear Everything Out*/
 	void ClearAllParameters();
diff --git a/Engine/Plugins/Animation/ControlRig/Source/ControlRig/Public/Sequencer/MovieSceneControlRigParameterTrack.h b/Engine/Plugins/Animation/ControlRig/Source/ControlRig/Public/Sequencer/MovieSceneControlRigParameterTrack.h
--- a/Engine/Plugins/Animation/ControlRig/Source/ControlRig/Public/Sequencer/MovieSceneControlRigParameterTrack.h
+++ b/Engine/Plugins/Animation/ControlRig/Source/ControlRig/Public/Sequencer/MovieSceneControlRigParameterTrack.h
@@ -119,7 +119,7 @@ public:
 	CONTROLRIG_API void SetTrackName(FName InName) { TrackName = InName; }
 
 	UMovieSceneControlRigParameterSection::FSpaceChannelAddedEvent& SpaceChannelAdded() { return OnSpaceChannelAdded; }
-	UMovieSceneControlRigParameterSection::FConstraintChannelAddedEvent& ConstraintChannelAdded() { return OnConstraintChannelAdded; }
+	IMovieSceneConstrainedSection::FConstraintChannelAddedEvent& ConstraintChannelAdded() { return OnConstraintChannelAdded; }
 
 private:
 
@@ -131,9 +131,9 @@ private:
 	UMovieSceneControlRigParameterSection::FSpaceChannelAddedEvent OnSpaceChannelAdded;
 
 	void HandleOnConstraintAdded(
-		UMovieSceneControlRigParameterSection* InSection,
+		IMovieSceneConstrainedSection* InSection,
 		FMovieSceneConstraintChannel* InChannel) const;
-	UMovieSceneControlRigParameterSection::FConstraintChannelAddedEvent OnConstraintChannelAdded;
+	IMovieSceneConstrainedSection::FConstraintChannelAddedEvent OnConstraintChannelAdded;
 
 
 	void ReconstructControlRig();
diff --git a/Engine/Plugins/Animation/ControlRig/Source/ControlRigEditor/Private/Editor/Constraints/SConstraintsWidget.cpp b/Engine/Plugins/Animation/ControlRig/Source/ControlRigEditor/Private/Editor/Constraints/SConstraintsWidget.cpp
--- a/Engine/Plugins/Animation/ControlRig/Source/ControlRigEditor/Private/Editor/Constraints/SConstraintsWidget.cpp
+++ b/Engine/Plugins/Animation/ControlRig/Source/ControlRigEditor/Private/Editor/Constraints/SConstraintsWidget.cpp
@@ -20,6 +20,7 @@
 #include "PropertyEditorModule.h"
 #include "DetailsViewArgs.h"
 #include "Tools/ConstraintBaker.h"
+#include "ScopedTransaction.h"
 
 #define LOCTEXT_NAMESPACE "SConstraintsWidget"
 
@@ -285,6 +286,7 @@ void SDroppableConstraintItem::CreateConstraint(
 	{
 		if (Child != InParent)
 		{
+			FScopedTransaction Transaction(LOCTEXT("CreateConstraintKey", "Create Constraint Key"));
 			UTickableTransformConstraint* Constraint =
 				FTransformConstraintUtils::CreateAndAddFromActors(World, InParent, Child, InConstraintType);
 			if (Constraint)
@@ -484,6 +486,7 @@ void SEditableConstraintItem::Construct(
 			{
 				if (UTickableTransformConstraint* TransformConstraint = Cast<UTickableTransformConstraint>(Constraint))
 				{
+					FScopedTransaction Transaction(LOCTEXT("CreateConstraintKey", "Create Constraint Key"));
 					FConstraintChannelHelper::SmartConstraintKey(TransformConstraint);
 				}
 				return FReply::Handled();
@@ -729,6 +732,7 @@ void SConstraintsEditionWidget::MoveItemDown(const TSharedPtr<FEditableConstrain
 
 void SConstraintsEditionWidget::RemoveItem(const TSharedPtr<FEditableConstraintItem>& Item)
 {
+	FScopedTransaction Transaction(LOCTEXT("RemoveConstraint", "Remove Constraint"));
 	UWorld* World = GCurrentLevelEditingViewportClient->GetWorld();
 	const FConstraintsManagerController& Controller = FConstraintsManagerController::Get(World);
 		
@@ -873,7 +877,8 @@ TSharedPtr<SWidget> SConstraintsEditionWidget::CreateContextMenu()
 		{
 			if (UTickableTransformConstraint* TransformConstraint = Cast<UTickableTransformConstraint>(Constraint))
 			{
-				FConstraintChannelHelper::AddConstraintKey(TransformConstraint);
+				FScopedTransaction Transaction(LOCTEXT("CreateConstraintKey", "Create Constraint Key"));
+				FConstraintChannelHelper::SmartConstraintKey(TransformConstraint);
 			}
 		})),
 		NAME_None,
diff --git a/Engine/Plugins/Animation/ControlRig/Source/ControlRigEditor/Private/Sequencer/ConstraintChannelHelper.cpp b/Engine/Plugins/Animation/ControlRig/Source/ControlRigEditor/Private/Sequencer/ConstraintChannelHelper.cpp
--- a/Engine/Plugins/Animation/ControlRig/Source/ControlRigEditor/Private/Sequencer/ConstraintChannelHelper.cpp
+++ b/Engine/Plugins/Animation/ControlRig/Source/ControlRigEditor/Private/Sequencer/ConstraintChannelHelper.cpp
@@ -1,7 +1,7 @@
 ﻿// Copyright Epic Games, Inc. All Rights Reserved.
 
 #include "ConstraintChannelHelper.inl"
-
+#include "ConstraintChannelHelper.h"
 #include "ControlRigSpaceChannelEditors.h"
 #include "ISequencer.h"
 
@@ -91,111 +91,6 @@ void FConstraintChannelHelper::SmartConstraintKey(UTickableTransformConstraint*
 	}
 }
 
-void FConstraintChannelHelper::AddChildTransformKey(
-	const UTransformableHandle* InHandle,
-	const FFrameNumber& InTime,
-	const TSharedPtr<ISequencer>& InSequencer)
-{
-	if (const UTransformableComponentHandle* ComponentHandle = Cast<UTransformableComponentHandle>(InHandle))
-	{
-		AActor* Actor = ComponentHandle->Component->GetOwner();
-		if (!Actor)
-		{
-			return;
-		}
-
-		const FTransform LocalTransform = ComponentHandle->GetLocalTransform();
-		const FGuid Guid = InSequencer->GetHandleToObject(Actor,true);
-		if (!Guid.IsValid())
-		{
-			return;
-		}
-
-		const UMovieScene3DTransformSection* TransformSection = FBakingHelper::GetTransformSection(InSequencer.Get(), Guid, LocalTransform);
-		if (!TransformSection)
-		{
-			return;
-		}
-
-		FBakingHelper::AddTransformKeys(TransformSection, {InTime}, {LocalTransform}, EMovieSceneTransformChannel::AllTransform);
-		return; 
-	}
-	
-	const UTransformableControlHandle* ControlHandle = Cast<UTransformableControlHandle>(InHandle);
-	if (ControlHandle && ControlHandle->IsValid())
-	{
-		const FTransform LocalTransform = ControlHandle->GetLocalTransform();
-		
-		static constexpr bool bNotify = true, bUndo = false, bFixEuler = true;
-	
-		FRigControlModifiedContext KeyframeContext;
-		KeyframeContext.SetKey = EControlRigSetKey::Always;
-		KeyframeContext.KeyMask = static_cast<uint32>(EControlRigContextChannelToKey::AllTransform);
-
-		const FFrameRate TickResolution = InSequencer->GetFocusedTickResolution();
-		KeyframeContext.LocalTime = TickResolution.AsSeconds(FFrameTime(InTime));
-	
-		ControlHandle->ControlRig->SetControlLocalTransform(ControlHandle->ControlName, LocalTransform, bNotify, KeyframeContext, bUndo, bFixEuler);
-	}
-}
-
-void FConstraintChannelHelper::AddConstraintKey(UTickableTransformConstraint* InConstraint)
-{
-	const TWeakPtr<ISequencer> WeakSequencer = FBakingHelper::GetSequencer();
-	if (!WeakSequencer.IsValid() || !WeakSequencer.Pin()->GetFocusedMovieSceneSequence())
-	{
-		return;
-	}
-	const TSharedPtr<ISequencer> Sequencer = WeakSequencer.Pin();
-	const UMovieScene* MovieScene = Sequencer->GetFocusedMovieSceneSequence()->GetMovieScene();
-	if (!MovieScene)
-	{
-		return;
-	}
-
-	const UTransformableControlHandle* ControlHandle = Cast<UTransformableControlHandle>(InConstraint->ChildTRSHandle);
-	if (!ControlHandle)
-	{
-		return;
-	}
-
-	if (UMovieSceneControlRigParameterSection* ParamSection = GetControlSection(ControlHandle, Sequencer))
-	{
-		ParamSection->AddConstraintChannel(InConstraint, true);
-
-		if (FConstraintAndActiveChannel* Channel = ParamSection->GetConstraintChannel(InConstraint->GetFName()))
-		{
-			ParamSection->Modify();
-
-			const FFrameRate TickResolution = Sequencer->GetFocusedTickResolution();
-			const FFrameTime FrameTime = Sequencer->GetLocalTime().ConvertTo(TickResolution);
-			const FFrameNumber Time = FrameTime.GetFrame();
-
-			bool Value = false;
-			CanAddKey(Channel->ActiveChannel, Time, Value);
-			
-			TMovieSceneChannelData<bool> ChannelData = Channel->ActiveChannel.GetData();
-			ChannelData.AddKey(Time, InConstraint->Active);
-
-			// bool PrevActiveValue = false; Channel->ActiveChannel.Evaluate(Time, PrevActiveValue);
-			// bool NextActiveValue = false; Channel->ActiveChannel.Evaluate(Time, PrevActiveValue);
-
-					
-			// int32 ExistingIndex = ChannelData.FindKey(Time);
-			// if (ExistingIndex != INDEX_NONE)
-			// {
-			// 	Handle = ChannelInterface.GetHandle(ExistingIndex);
-			// 	using namespace UE::MovieScene;
-			// 	AssignValue(Channel, Handle, Forward<FMovieSceneControlRigSpaceBaseKey>(Value));
-			// }
-			// else
-			// {
-			// 	ExistingIndex = ChannelInterface.AddKey(Time, Forward<FMovieSceneControlRigSpaceBaseKey>(Value));
-			// 	Handle = ChannelInterface.GetHandle(ExistingIndex);
-			// }
-		}
-	}
-}
 
 UMovieSceneControlRigParameterSection* FConstraintChannelHelper::GetControlSection(
 	const UTransformableControlHandle* InHandle,
@@ -241,375 +136,7 @@ UMovieScene3DTransformSection* FConstraintChannelHelper::GetTransformSection(
 		return nullptr;
 	}
 	
-	return FBakingHelper::GetTransformSection(InSequencer.Get(), Guid);
-}
-
-TArrayView<FMovieSceneFloatChannel*> FConstraintChannelHelper::GetTransformFloatChannels(
-	const UTransformableHandle* InHandle,
-	const TSharedPtr<ISequencer>& InSequencer)
-{
-	// InHandle and InSequencer are assumed to be valid at this stage so no need to check
-	if (const UTransformableControlHandle* ControlHandle = Cast<UTransformableControlHandle>(InHandle))
-	{
-		const UMovieSceneControlRigParameterSection* Section = GetControlSection(ControlHandle, InSequencer);
-		return GetTransformFloatChannels(InHandle, Section);
-	}
-
-	static const TArrayView<FMovieSceneFloatChannel*> EmptyChannelsView;
-	return EmptyChannelsView;
-}
-
-TArrayView<FMovieSceneFloatChannel*> FConstraintChannelHelper::GetTransformFloatChannels(
-	const UTransformableHandle* InHandle,
-	const UMovieSceneControlRigParameterSection* InSection)
-{
-	static const TArrayView<FMovieSceneFloatChannel*> EmptyChannelsView;
-	
-	const UTransformableControlHandle* ControlHandle = Cast<UTransformableControlHandle>(InHandle);
-	if (!ControlHandle || !InSection)
-	{
-		return EmptyChannelsView;	
-	}
-
-	const FChannelMapInfo* ChannelInfo = nullptr;
-	int32 NumChannels = 0;
-	Tie(ChannelInfo, NumChannels) =
-		GetInfoAndNumFloatChannels(ControlHandle->ControlRig.Get(), ControlHandle->ControlName, InSection);
-
-	if (ChannelInfo == nullptr || NumChannels == 0)
-	{
-		return EmptyChannelsView;
-	}
-
-	// return a sub view that just represents the control's channels
-	const TArrayView<FMovieSceneFloatChannel*> FloatChannels = InSection->GetChannelProxy().GetChannels<FMovieSceneFloatChannel>();
-	const int32 ChannelStartIndex = ChannelInfo->ChannelIndex;
-	return FloatChannels.Slice(ChannelStartIndex, NumChannels);
-}
-
-TArrayView<FMovieSceneDoubleChannel*> FConstraintChannelHelper::GetTransformDoubleChannels(
-	const UTransformableHandle* InHandle,
-	const TSharedPtr<ISequencer>& InSequencer)
-{
-	// InHandle and InSequencer are assumed to be valid at this stage so no need to check
-
-	static const TArrayView<FMovieSceneDoubleChannel*> EmptyChannelsView;
-	if (const UTransformableComponentHandle* ComponentHandle = Cast<UTransformableComponentHandle>(InHandle))
-	{
-		const UMovieScene3DTransformSection* Section = GetTransformSection(ComponentHandle, InSequencer);
-		return Section ? Section->GetChannelProxy().GetChannels<FMovieSceneDoubleChannel>() : EmptyChannelsView;
-	}
-
-	return EmptyChannelsView;
-}
-
-namespace
-{
-
-struct FCompensationEvaluator
-{
-public:
-	TArray<FTransform> ChildLocals;
-	TArray<FTransform> ChildGlobals;
-	TArray<FTransform> SpaceGlobals;
-	
-	FCompensationEvaluator(UTickableTransformConstraint* InConstraint)
-		: Constraint(InConstraint)
-		, Handle(InConstraint ? InConstraint->ChildTRSHandle : nullptr)
-	{}
-
-	void ComputeLocalTransforms(
-		UWorld* InWorld, const TSharedPtr<ISequencer>& InSequencer,
-		const TArray<FFrameNumber>& InFrames, const bool bToActive)
-	{
-		if (InFrames.IsEmpty())
-		{
-			return;
-		}
-		
-		const TArray<UTickableTransformConstraint*> Constraints = GetHandleTransformConstraints(InWorld);
-		if (Constraints.IsEmpty())
-		{
-			return;
-		}
-
-		// find last active constraint in the list that is different than the on we want to compensate for
-		auto GetLastActiveConstraint = [this, Constraints]()
-		{
-			// find last active constraint in the list that is different than the on we want to compensate for
-			const int32 LastActiveIndex = Constraints.FindLastByPredicate([this](const UTickableTransformConstraint* InConstraint)
-			{
-				return InConstraint->Active && InConstraint->bDynamicOffset && InConstraint != Constraint;
-			});
-
-			// if found, return its parent global transform
-			return LastActiveIndex > INDEX_NONE ?Constraints[LastActiveIndex] : nullptr;
-		};
-		
-		const UMovieScene* MovieScene = InSequencer->GetFocusedMovieSceneSequence()->GetMovieScene();
-        const FFrameRate TickResolution = MovieScene->GetTickResolution();
-        const EMovieScenePlayerStatus::Type PlaybackStatus = InSequencer->GetPlaybackStatus();
-
-		const int32 NumFrames = InFrames.Num();
-
-		// resize arrays to num frames + 1 as we also evaluate at InFrames[0]-1
-		ChildLocals.SetNum(NumFrames+1);
-        ChildGlobals.SetNum(NumFrames+1);
-		SpaceGlobals.SetNum(NumFrames+1);
-
-		const ETransformConstraintType ConstraintType = static_cast<ETransformConstraintType>(Constraint->GetType());
-		
-        for (int32 Index = 0; Index < NumFrames+1; ++Index)
-        {
-        	const FFrameNumber FrameNumber = (Index == 0) ? InFrames[0]-1 : InFrames[Index-1];
-        
-        	// evaluate animation
-        	const FMovieSceneEvaluationRange EvaluationRange = FMovieSceneEvaluationRange(FFrameTime(FrameNumber), TickResolution);
-        	const FMovieSceneContext Context = FMovieSceneContext(EvaluationRange, PlaybackStatus).SetHasJumped(true);
-    
-        	InSequencer->GetEvaluationTemplate().Evaluate(Context, *InSequencer);
-        
-        	// evaluate constraints
-        	for (const UTickableTransformConstraint* InConstraint: Constraints)
-        	{
-        		InConstraint->Evaluate();
-        	}
-        
-        	// evaluate ControlRig?
-        	// ControlRig->Evaluate_AnyThread();
-
-        	FTransform& ChildLocal = ChildLocals[Index];
-        	FTransform& ChildGlobal = ChildGlobals[Index];
-        	FTransform& SpaceGlobal = SpaceGlobals[Index];
-
-        	// store child transforms        	
-        	ChildLocal = Handle->GetLocalTransform();
-        	ChildGlobal = Handle->GetGlobalTransform();
-        	
-        	const UTickableTransformConstraint* LastConstraint = GetLastActiveConstraint();
-        	const ETransformConstraintType LastConstraintType = LastConstraint ?
-				static_cast<ETransformConstraintType>(LastConstraint->GetType()) : ETransformConstraintType::Parent; 
-
-        	// store constraint/parent space global transform
-        	if (bToActive)
-       		{
-        		// if activating the constraint, store last constraint or parent space at T[0]-1
-        		// and constraint space for all other times
-				if (Index == 0)
-				{
-					if (LastConstraint)
-					{
-						SpaceGlobal = LastConstraint->GetParentGlobalTransform();
-						ChildLocal = FTransformConstraintUtils::ComputeRelativeTransform(
-							ChildLocal, ChildGlobal, SpaceGlobal, LastConstraintType);
-					}
-				}
-				else
-				{
-					SpaceGlobal = Constraint->GetParentGlobalTransform();
-					ChildLocal = FTransformConstraintUtils::ComputeRelativeTransform(
-						ChildLocal, ChildGlobal, SpaceGlobal, ConstraintType);
-				}
-       		}
-            else
-            {
-            	// if deactivating the constraint, store constraint space at T[0]-1
-            	// and last constraint or parent space for all other times
-            	if (Index == 0)
-            	{
-            		SpaceGlobal = Constraint->GetParentGlobalTransform();
-            		ChildLocal = FTransformConstraintUtils::ComputeRelativeTransform(
-            			ChildLocal, ChildGlobal, SpaceGlobal, ConstraintType);
-            	}
-                else
-                {
-                	if (LastConstraint)
-                	{
-                		SpaceGlobal = LastConstraint->GetParentGlobalTransform();
-                		ChildLocal = FTransformConstraintUtils::ComputeRelativeTransform(
-                			ChildLocal, ChildGlobal, SpaceGlobal, LastConstraintType);
-                	}
-                }
-            }
-        }
-	}
-
-	void ComputeLocalTransformsBeforeDeletion(
-		UWorld* InWorld,
-		const TSharedPtr<ISequencer>& InSequencer,
-		const TArray<FFrameNumber>& InFrames)
-	{
-		if (InFrames.IsEmpty())
-		{
-			return;
-		}
-		
-		const TArray<UTickableTransformConstraint*> Constraints = GetHandleTransformConstraints(InWorld);
-		// if (Constraints.IsEmpty())
-		// {
-		// 	return;
-		// }
-
-		// find last active constraint in the list that is different than the on we want to compensate for
-		auto GetLastActiveConstraint = [this, Constraints]()
-		{
-			// find last active constraint in the list that is different than the on we want to compensate for
-			const int32 LastActiveIndex = Constraints.FindLastByPredicate([this](const UTickableTransformConstraint* InConstraint)
-			{
-				return InConstraint->Active && InConstraint->bDynamicOffset && InConstraint != Constraint;
-			});
-
-			// if found, return its parent global transform
-			return LastActiveIndex > INDEX_NONE ?Constraints[LastActiveIndex] : nullptr;
-		};
-		
-		const UMovieScene* MovieScene = InSequencer->GetFocusedMovieSceneSequence()->GetMovieScene();
-        const FFrameRate TickResolution = MovieScene->GetTickResolution();
-        const EMovieScenePlayerStatus::Type PlaybackStatus = InSequencer->GetPlaybackStatus();
-
-		const int32 NumFrames = InFrames.Num();
-
-		ChildLocals.SetNum(NumFrames);
-        ChildGlobals.SetNum(NumFrames);
-		SpaceGlobals.SetNum(NumFrames);
-
-        for (int32 Index = 0; Index < NumFrames; ++Index)
-        {
-        	const FFrameNumber& FrameNumber = InFrames[Index];
-        
-        	// evaluate animation
-        	const FMovieSceneEvaluationRange EvaluationRange = FMovieSceneEvaluationRange(FFrameTime(FrameNumber), TickResolution);
-        	const FMovieSceneContext Context = FMovieSceneContext(EvaluationRange, PlaybackStatus).SetHasJumped(true);
-    
-        	InSequencer->GetEvaluationTemplate().Evaluate(Context, *InSequencer);
-        
-        	// evaluate constraints
-        	for (const UTickableTransformConstraint* InConstraint: Constraints)
-        	{
-        		InConstraint->Evaluate();
-        	}
-        
-        	// evaluate ControlRig?
-        	// ControlRig->Evaluate_AnyThread();
-
-        	FTransform& ChildLocal = ChildLocals[Index];
-        	FTransform& ChildGlobal = ChildGlobals[Index];
-        	FTransform& SpaceGlobal = SpaceGlobals[Index];
-
-        	// store child transforms        	
-        	ChildLocal = Handle->GetLocalTransform();
-        	ChildGlobal = Handle->GetGlobalTransform();
-
-        	// store constraint/parent space global transform
-            if (const UTickableTransformConstraint* LastConstraint = GetLastActiveConstraint())
-            {
-            	const ETransformConstraintType LastConstraintType =
-					static_cast<ETransformConstraintType>(LastConstraint->GetType());
-            	
-                SpaceGlobal = LastConstraint->GetParentGlobalTransform();
-                ChildLocal = FTransformConstraintUtils::ComputeRelativeTransform(
-                	ChildLocal, ChildGlobal, SpaceGlobal, LastConstraintType);
-            }
-        }
-	}
-	
-	void ComputeCompensation(UWorld* InWorld, const TSharedPtr<ISequencer>& InSequencer, const FFrameNumber& InTime)
-	{
-		const TArray<UTickableTransformConstraint*> Constraints = GetHandleTransformConstraints(InWorld);
-		if (Constraints.IsEmpty())
-		{
-			return;
-		}
-
-		// find last active constraint in the list that is different than the on we want to compensate for
-		auto GetLastActiveConstraint = [this, Constraints]()
-		{
-			// find last active constraint in the list that is different than the on we want to compensate for
-			const int32 LastActiveIndex = Constraints.FindLastByPredicate([this](const UTickableTransformConstraint* InConstraint)
-			{
-				return InConstraint->Active && InConstraint->bDynamicOffset;
-			});
-
-			// if found, return its parent global transform
-			return LastActiveIndex > INDEX_NONE ? Constraints[LastActiveIndex] : nullptr;
-		};
-
-		auto EvaluateAt = [InSequencer, &Constraints](const FFrameNumber InFrame)
-		{
-			const UMovieScene* MovieScene = InSequencer->GetFocusedMovieSceneSequence()->GetMovieScene();
-			const FFrameRate TickResolution = MovieScene->GetTickResolution();
-			const EMovieScenePlayerStatus::Type PlaybackStatus = InSequencer->GetPlaybackStatus();
-
-			const FMovieSceneEvaluationRange EvaluationRange0 = FMovieSceneEvaluationRange(FFrameTime(InFrame), TickResolution);
-			const FMovieSceneContext Context0 = FMovieSceneContext(EvaluationRange0, PlaybackStatus).SetHasJumped(true);
-
-			InSequencer->GetEvaluationTemplate().Evaluate(Context0, *InSequencer);
-
-			for (const UTickableTransformConstraint* InConstraint: Constraints)
-			{
-				InConstraint->Evaluate();
-			}
-
-			// ControlRig->Evaluate_AnyThread();
-		};
-
-		// allocate
-		ChildLocals.SetNum(1);
-		ChildGlobals.SetNum(1);
-		SpaceGlobals.SetNum(1);
-		
-		// evaluate at InTime and store global
-		EvaluateAt(InTime);
-		ChildGlobals[0] = Handle->GetGlobalTransform();
-
-		// evaluate at InTime-1 and store local
-		EvaluateAt(InTime-1);
-        ChildLocals[0] = Handle->GetLocalTransform();
-
-		// if constraint at T-1 then switch to its space
-        if (const UTickableTransformConstraint* LastConstraint = GetLastActiveConstraint())
-        {
-        	SpaceGlobals[0] = LastConstraint->GetParentGlobalTransform();
-        	
-        	const ETransformConstraintType LastConstraintType = static_cast<ETransformConstraintType>(LastConstraint->GetType());
-        	ChildLocals[0] = FTransformConstraintUtils::ComputeRelativeTransform(ChildLocals[0],
-		ChildGlobals[0], SpaceGlobals[0], LastConstraintType);
-        }
-        else // switch to parent space
-        {
-        	const FTransform ChildLocal = ChildLocals[0];
-        	Handle->SetGlobalTransform(ChildGlobals[0]);
-        	ChildLocals[0] = Handle->GetLocalTransform();
-        	Handle->SetLocalTransform(ChildLocal);
-        }
-	}
-	
-private:
-
-	TArray<UTickableTransformConstraint*> GetHandleTransformConstraints(UWorld* InWorld) const
-	{
-		TArray<UTickableTransformConstraint*> TransformConstraints;
-		if (Handle)
-		{
-			// get sorted transform constraints
-			const FConstraintsManagerController& Controller = FConstraintsManagerController::Get(InWorld);
-			static constexpr bool bSorted = true;
-			TArray<TObjectPtr<UTickableConstraint>> Constraints = Controller.GetParentConstraints(Handle->GetHash(), bSorted);
-			for (const TObjectPtr<UTickableConstraint>& InConstraint: Constraints)
-			{
-				if (UTickableTransformConstraint* TransformConstraint = Cast<UTickableTransformConstraint>(InConstraint.Get()))
-				{
-					TransformConstraints.Add(TransformConstraint);
-				}
-			}
-		}
-		return TransformConstraints;
-	}
-	
-	UTickableTransformConstraint* Constraint = nullptr;
-	UTransformableHandle* Handle = nullptr;
-};
-	
+	return MovieSceneToolHelpers::GetTransformSection(InSequencer.Get(), Guid);
 }
 
 void FConstraintChannelHelper::SmartControlConstraintKey(
@@ -627,7 +154,7 @@ void FConstraintChannelHelper::SmartControlConstraintKey(
 		FScopedTransaction Transaction(LOCTEXT("KeyConstraintaKehy", "Key Constraint Key"));
 		Section->Modify();
 		// add the channel
-		Section->AddConstraintChannel(InConstraint, true);
+		Section->AddConstraintChannel(InConstraint);
 
 		// add key if needed
 		if (FConstraintAndActiveChannel* Channel = Section->GetConstraintChannel(InConstraint->GetFName()))
@@ -649,9 +176,9 @@ void FConstraintChannelHelper::SmartControlConstraintKey(
 				const FName& ControlName = ControlHandle->ControlName;
 				
 				// store the frames to compensate
-				const TArrayView<FMovieSceneFloatChannel*> Channels = GetTransformFloatChannels(ControlHandle, InSequencer);
+				const TArrayView<FMovieSceneFloatChannel*> Channels = ControlHandle->GetFloatChannels(Section);
 				TArray<FFrameNumber> FramesToCompensate;
-				GetFramesToCompensate<FMovieSceneFloatChannel>(Channel->ActiveChannel, ActiveValueToBeSet, Time, Channels, FramesToCompensate);
+				FMovieSceneConstraintChannelHelper::GetFramesToCompensate<FMovieSceneFloatChannel>(Channel->ActiveChannel, ActiveValueToBeSet, Time, Channels, FramesToCompensate);
 
 				// store child and space transforms for these frames
 				FCompensationEvaluator Evaluator(InConstraint);
@@ -673,24 +200,23 @@ void FConstraintChannelHelper::SmartControlConstraintKey(
 					if (NumChannels > 0)
 					{
 						ChannelIndex = pChannelIndex->ChannelIndex;
-						FControlRigSpaceChannelHelpers::EvaluateTangentAtThisTime(ChannelIndex, NumChannels, true,Section, Time, Tangents);
+						EvaluateTangentAtThisTime<FMovieSceneFloatChannel>(ChannelIndex, NumChannels, Section, Time, Tangents);
 					}
 				}
-				
-
-				const EMovieSceneTransformChannel ChannelsToKey = FConstraintBaker::GetChannelsToKey(InConstraint);
+			
+				const EMovieSceneTransformChannel ChannelsToKey =InConstraint->GetChannelsToKey();
 				
 				// add child's transform key at Time-1 to keep animation
 				{
 					const FFrameNumber TimeMinusOne(Time - 1);
 
-					FBakingHelper::AddTransformKeys(ControlRig, ControlName, { TimeMinusOne },
-						{ ChildLocals[0] }, ChannelsToKey, TickResolution, true);
+					ControlHandle->AddTransformKeys({ TimeMinusOne },
+						{ ChildLocals[0] }, ChannelsToKey, TickResolution, nullptr,true);
 
 					// set tangents at Time-1
 					if (NumChannels > 0)
 					{
-						FControlRigSpaceChannelHelpers::SetTangentsAtThisTime(ChannelIndex, NumChannels, true, Section, TimeMinusOne, Tangents);
+						SetTangentsAtThisTime<FMovieSceneFloatChannel>(ChannelIndex, NumChannels, Section, TimeMinusOne, Tangents);
 					}
 				}
 
@@ -706,13 +232,13 @@ void FConstraintChannelHelper::SmartControlConstraintKey(
 					ChildLocals.RemoveAt(0);
 
 					// add keys
-					FBakingHelper::AddTransformKeys(ControlRig, ControlHandle->ControlName, FramesToCompensate,
-						ChildLocals, ChannelsToKey, TickResolution, true);
+					ControlHandle->AddTransformKeys(FramesToCompensate,
+						ChildLocals, ChannelsToKey, TickResolution, nullptr,true);
 
 					// set tangents at Time
 					if (NumChannels > 0 )
 					{
-						FControlRigSpaceChannelHelpers::SetTangentsAtThisTime(ChannelIndex, NumChannels, true, Section, Time, Tangents);
+						SetTangentsAtThisTime<FMovieSceneFloatChannel>(ChannelIndex, NumChannels, Section, Time, Tangents);
 					}
 				}
 			}
@@ -742,7 +268,7 @@ void FConstraintChannelHelper::SmartComponentConstraintKey(
 		return;
 	}
 
-	if (UMovieScene3DTransformSection* Section = FBakingHelper::GetTransformSection(InSequencer.Get(), Guid, LocalTransform))
+	if (UMovieScene3DTransformSection* Section = MovieSceneToolHelpers::GetTransformSection(InSequencer.Get(), Guid, LocalTransform))
 	{
 		FScopedTransaction Transaction(LOCTEXT("KeyConstraintaKehy", "Key Constraint Key"));
 		Section->Modify();
@@ -770,9 +296,9 @@ void FConstraintChannelHelper::SmartComponentConstraintKey(
 				InConstraint->bDynamicOffset = true;
 
 				// store the frames to compensate
-				const TArrayView<FMovieSceneDoubleChannel*> Channels = GetTransformDoubleChannels(ComponentHandle, InSequencer);
+				const TArrayView<FMovieSceneDoubleChannel*> Channels = ComponentHandle->GetDoubleChannels(Section);
 				TArray<FFrameNumber> FramesToCompensate;
-				GetFramesToCompensate<FMovieSceneDoubleChannel>(Channel->ActiveChannel, ActiveValueToBeSet, Time, Channels, FramesToCompensate);
+				FMovieSceneConstraintChannelHelper::GetFramesToCompensate<FMovieSceneDoubleChannel>(Channel->ActiveChannel, ActiveValueToBeSet, Time, Channels, FramesToCompensate);
 
 				// store child and space transforms for these frames
 				FCompensationEvaluator Evaluator(InConstraint);
@@ -783,17 +309,17 @@ void FConstraintChannelHelper::SmartComponentConstraintKey(
 				TArray<FMovieSceneTangentData> Tangents;
 				const int32 ChannelIndex = 0;
 				const int32 NumChannels = 9;
-				FControlRigSpaceChannelHelpers::EvaluateTangentAtThisTime(ChannelIndex, NumChannels, false, Section, Time, Tangents);
+				EvaluateTangentAtThisTime<FMovieSceneDoubleChannel>(ChannelIndex, NumChannels, Section, Time, Tangents);
 
-				const EMovieSceneTransformChannel ChannelsToKey = FConstraintBaker::GetChannelsToKey(InConstraint);
+				const EMovieSceneTransformChannel ChannelsToKey = InConstraint->GetChannelsToKey();
 
 				// add child's transform key at Time-1 to keep animation
 				{
 					const FFrameNumber TimeMinusOne(Time - 1);
 
-					FBakingHelper::AddTransformKeys(Section, { TimeMinusOne }, { ChildLocals[0] }, ChannelsToKey);
+					MovieSceneToolHelpers::AddTransformKeys(Section, { TimeMinusOne }, { ChildLocals[0] }, ChannelsToKey);
 
-					FControlRigSpaceChannelHelpers::SetTangentsAtThisTime(ChannelIndex, NumChannels, false,Section, TimeMinusOne, Tangents);
+					SetTangentsAtThisTime<FMovieSceneDoubleChannel>(ChannelIndex, NumChannels, Section, TimeMinusOne, Tangents);
 				}
 
 				// add active key
@@ -808,10 +334,10 @@ void FConstraintChannelHelper::SmartComponentConstraintKey(
 					ChildLocals.RemoveAt(0);
 
 					// add keys
-					FBakingHelper::AddTransformKeys(Section, FramesToCompensate, ChildLocals, ChannelsToKey);
+					MovieSceneToolHelpers::AddTransformKeys(Section, FramesToCompensate, ChildLocals, ChannelsToKey);
 
 					// set tangents at Time
-					FControlRigSpaceChannelHelpers::SetTangentsAtThisTime(ChannelIndex,NumChannels, false, Section, Time, Tangents);
+					SetTangentsAtThisTime<FMovieSceneDoubleChannel>(ChannelIndex,NumChannels, Section, Time, Tangents);
 				}
 			}
 		}
@@ -856,14 +382,14 @@ void FConstraintChannelHelper::CompensateIfNeeded(
 	}
 
 	TGuardValue<bool> CompensateGuard(bDoNotCompensate, true);
-	
+
 	// Frames to compensate
 	TArray<FFrameNumber> OptionalTimeArray;
 	if (OptionalTime.IsSet())
 	{
 		OptionalTimeArray.Add(OptionalTime.GetValue());
 	}
-	
+
 	auto GetSpaceTimesToCompensate = [&OptionalTimeArray](const FConstraintAndActiveChannel& Channel)->TArrayView<const FFrameNumber>
 	{
 		if (OptionalTimeArray.IsEmpty())
@@ -872,231 +398,58 @@ void FConstraintChannelHelper::CompensateIfNeeded(
 		}
 		return OptionalTimeArray;
 	};
-	
+
 	// keyframe context
 	FRigControlModifiedContext KeyframeContext;
 	KeyframeContext.SetKey = EControlRigSetKey::Always;
 	const FFrameRate TickResolution = InSequencer->GetFocusedTickResolution();
-	
+
 	bool bNeedsEvaluation = false;
 
 	// gather all transform constraints
 	TArray<FConstraintAndActiveChannel> TransformConstraintsChannels;
 	Algo::CopyIf(Section->GetConstraintsChannels(), TransformConstraintsChannels,
-		[](const FConstraintAndActiveChannel& InChannel) 
+		[](const FConstraintAndActiveChannel& InChannel)
 		{
-			return InChannel.Constraint.IsValid() && InChannel.Constraint->IsA<UTickableTransformConstraint>(); 
+			return InChannel.Constraint.IsValid() && InChannel.Constraint->IsA<UTickableTransformConstraint>();
 		}
 	);
-	
+
 	// compensate constraints
-	for (const FConstraintAndActiveChannel& Channel: TransformConstraintsChannels)
+	for (const FConstraintAndActiveChannel& Channel : TransformConstraintsChannels)
 	{
 		const TArrayView<const FFrameNumber> FramesToCompensate = GetSpaceTimesToCompensate(Channel);
 		for (const FFrameNumber& Time : FramesToCompensate)
 		{
-			const FFrameNumber TimeMinusOne(Time-1);
-			
+			const FFrameNumber TimeMinusOne(Time - 1);
+
 			bool CurrentValue = false, PreviousValue = false;
 			Channel.ActiveChannel.Evaluate(TimeMinusOne, PreviousValue);
 			Channel.ActiveChannel.Evaluate(Time, CurrentValue);
 
 			if (CurrentValue != PreviousValue) //if they are the same no need to do anything
 			{
-				UTickableTransformConstraint* Constraint = Cast<UTickableTransformConstraint>(Channel.Constraint.Get()); 
-				
+				UTickableTransformConstraint* Constraint = Cast<UTickableTransformConstraint>(Channel.Constraint.Get());
+
 				// compute transform to set
 				// if switching from active to inactive then we must add a key at T-1 in the constraint space
 				// if switching from inactive to active then we must add a key at T-1 in the previous constraint or parent space
 				FCompensationEvaluator Evaluator(Constraint);
 				Evaluator.ComputeCompensation(ControlRig->GetWorld(), InSequencer, Time);
 				const TArray<FTransform>& LocalTransforms = Evaluator.ChildLocals;
-				
-				const EMovieSceneTransformChannel ChannelsToKey = FConstraintBaker::GetChannelsToKey(Constraint);
+
+				const EMovieSceneTransformChannel ChannelsToKey = Constraint->GetChannelsToKey();
 				FConstraintBaker::AddTransformKeys(
-					InSequencer, Constraint->ChildTRSHandle, {TimeMinusOne}, LocalTransforms, ChannelsToKey);
+					InSequencer, Constraint->ChildTRSHandle, { TimeMinusOne }, LocalTransforms, ChannelsToKey);
 				bNeedsEvaluation = true;
 			}
 		}
 	}
-	
+
 	if (bNeedsEvaluation)
 	{
 		InSequencer->ForceEvaluate();
 	}
 }
 
-void FConstraintChannelHelper::HandleConstraintRemoved(
-	UTickableConstraint* InConstraint,
-	const FMovieSceneConstraintChannel* InConstraintChannel,
-	const TSharedPtr<ISequencer>& InSequencer,
-	const UMovieSceneSection* InSection)
-{
-	UTickableTransformConstraint* Constraint = Cast<UTickableTransformConstraint>(InConstraint);
-	if (!Constraint || !InConstraintChannel || !InSection)
-	{
-		return;
-	}
-	
-	TGuardValue<bool> CompensateGuard(bDoNotCompensate, true);
-	if (const UTransformableControlHandle* ControlHandle = Cast<UTransformableControlHandle>(Constraint->ChildTRSHandle))
-	{
-		const UMovieSceneControlRigParameterSection* Section = Cast<UMovieSceneControlRigParameterSection>(InSection);
-		if (!Section)
-		{
-			return;
-		}
-
-		const TArrayView<const FFrameNumber> Times = InConstraintChannel->GetData().GetTimes();
-		if (Times.IsEmpty())
-		{
-			return;
-		}
-		
-		UControlRig* ControlRig = ControlHandle->ControlRig.Get();
-		const FName& ControlName = ControlHandle->ControlName;
-		
-		// get transform channels
-		const TArrayView<FMovieSceneFloatChannel*> TransformChannels = GetTransformFloatChannels(ControlHandle, Section);
-		
-		// get frames after this time
-		TArray<FFrameNumber> FramesToCompensate;
-		GetFramesWithinActiveState(*InConstraintChannel, TransformChannels, FramesToCompensate);
-
-		// do the compensation
-		FCompensationEvaluator Evaluator(Constraint);
-		Evaluator.ComputeLocalTransformsBeforeDeletion(ControlRig->GetWorld(), InSequencer, FramesToCompensate);
-		const TArray<FTransform>& ChildLocals = Evaluator.ChildLocals;
-
-		const EMovieSceneTransformChannel ChannelsToKey = FConstraintBaker::GetChannelsToKey(Constraint);
-		const FFrameRate TickResolution = InSequencer->GetFocusedTickResolution();
-		
-		FBakingHelper::AddTransformKeys(ControlRig, ControlName, FramesToCompensate,
-			ChildLocals, ChannelsToKey, TickResolution);
-	}
-}
-
-void FConstraintChannelHelper::HandleConstraintKeyDeleted(
-	UTickableTransformConstraint* InConstraint,
-	const FMovieSceneConstraintChannel* InConstraintChannel,
-	const TSharedPtr<ISequencer>& InSequencer,
-	const UMovieSceneSection* InSection,
-	const FFrameNumber& InTime)
-{
-	const FFrameNumber TimeMinusOne(InTime-1);
-	
-	bool CurrentValue = false, PreviousValue = false;
-	InConstraintChannel->Evaluate(TimeMinusOne, PreviousValue);
-	InConstraintChannel->Evaluate(InTime, CurrentValue);
-
-	if (CurrentValue == PreviousValue)
-	{
-		const int32 NumKeys = InConstraintChannel->GetNumKeys();
-		if (NumKeys > 1)
-		{
-			return;
-		}
-	}
-
-	TGuardValue<bool> CompensateGuard(bDoNotCompensate, true);
-	
-	if (const UTransformableControlHandle* ControlHandle = Cast<UTransformableControlHandle>(InConstraint->ChildTRSHandle))
-	{
-		const UMovieSceneControlRigParameterSection* Section = Cast<UMovieSceneControlRigParameterSection>(InSection);
-		if (!Section)
-		{
-			return;
-		}
-		
-		UControlRig* ControlRig = ControlHandle->ControlRig.Get();
-		const FName& ControlName = ControlHandle->ControlName;
-		
-		// get transform channels
-		const TArrayView<FMovieSceneFloatChannel*> TransformChannels = GetTransformFloatChannels(ControlHandle, Section);
-		
-		// get frames after this time
-		TArray<FFrameNumber> FramesToCompensate;
-		GetFramesAfter(*InConstraintChannel, InTime, TransformChannels, FramesToCompensate);
-
-		// do the compensation
-		FCompensationEvaluator Evaluator(InConstraint);
-		Evaluator.ComputeLocalTransforms(ControlRig->GetWorld(), InSequencer, FramesToCompensate, PreviousValue);
-		TArray<FTransform>& ChildLocals = Evaluator.ChildLocals;
-		ChildLocals.RemoveAt(0);
-
-		const EMovieSceneTransformChannel ChannelsToKey = FConstraintBaker::GetChannelsToKey(InConstraint);
-		const FFrameRate TickResolution = InSequencer->GetFocusedTickResolution();
-
-		FBakingHelper::AddTransformKeys( ControlRig, ControlName, FramesToCompensate,
-			ChildLocals, ChannelsToKey, TickResolution);
-		
-		// now delete any extra TimeMinusOne
-		DeleteTransformKeys(TransformChannels, TimeMinusOne);
-	}
-}
-
-void FConstraintChannelHelper::HandleConstraintKeyMoved(
-	const UTickableTransformConstraint* InConstraint,
-	const FMovieSceneConstraintChannel* InConstraintChannel,
-	const UMovieSceneSection* InSection,
-	const FFrameNumber& InCurrentFrame, const FFrameNumber& InNextFrame)
-{
-	const FFrameNumber Delta = InNextFrame - InCurrentFrame;
-	if (Delta == 0)
-	{
-		return;
-	}
-	
-	if (!InConstraint || !InConstraintChannel || !InSection)
-	{
-		return;
-	}
-
-	if (const UTransformableControlHandle* ControlHandle = Cast<UTransformableControlHandle>(InConstraint->ChildTRSHandle))
-	{
-		const UMovieSceneControlRigParameterSection* Section = Cast<UMovieSceneControlRigParameterSection>(InSection);
-		if (!Section)
-		{
-			return;
-		}
-		
-		// get transform channels
-		const TArrayView<FMovieSceneFloatChannel*> TransformChannels = GetTransformFloatChannels(ControlHandle, Section);
-
-		// move them
-		MoveTransformKeys(TransformChannels, InCurrentFrame, InNextFrame);
-	}
-}
-
-TPair<const FChannelMapInfo*, int32> FConstraintChannelHelper::GetInfoAndNumFloatChannels(
-	const UControlRig* InControlRig,
-	const FName& InControlName,
-	const UMovieSceneControlRigParameterSection* InSection)
-{
-	const FRigControlElement* ControlElement = InControlRig ? InControlRig->FindControl(InControlName) : nullptr;
-	auto GetNumFloatChannels = [](const ERigControlType& InControlType)
-	{
-		switch (InControlType)
-		{
-		case ERigControlType::Position:
-		case ERigControlType::Scale:
-		case ERigControlType::Rotator:
-			return 3;
-		case ERigControlType::TransformNoScale:
-			return 6;
-		case ERigControlType::Transform:
-		case ERigControlType::EulerTransform:
-			return 9;
-		default:
-			break;
-		}
-		return 0;		
-	};
-	
-	const int32 NumFloatChannels = ControlElement ? GetNumFloatChannels(ControlElement->Settings.ControlType) : 0;
-	const FChannelMapInfo* ChannelInfo = InSection ? InSection->ControlChannelMap.Find(InControlName) : nullptr;
-	
-	return {ChannelInfo, NumFloatChannels};
-}
-
 #undef LOCTEXT_NAMESPACE
diff --git a/Engine/Plugins/Animation/ControlRig/Source/ControlRigEditor/Private/Sequencer/ConstraintChannelHelper.inl b/Engine/Plugins/Animation/ControlRig/Source/ControlRigEditor/Private/Sequencer/ConstraintChannelHelper.inl
deleted file mode 100644
diff --git a/Engine/Plugins/Animation/ControlRig/Source/ControlRigEditor/Private/Sequencer/ControlRigParameterTrackEditor.cpp b/Engine/Plugins/Animation/ControlRig/Source/ControlRigEditor/Private/Sequencer/ControlRigParameterTrackEditor.cpp
--- a/Engine/Plugins/Animation/ControlRig/Source/ControlRigEditor/Private/Sequencer/ControlRigParameterTrackEditor.cpp
+++ b/Engine/Plugins/Animation/ControlRig/Source/ControlRigEditor/Private/Sequencer/ControlRigParameterTrackEditor.cpp
@@ -70,7 +70,6 @@
 #include "Animation/SkeletalMeshActor.h"
 #include "TimerManager.h"
 #include "BakeToControlRigSettings.h"
-#include "ConstraintChannelHelper.h"
 #include "UnrealEdGlobals.h"
 #include "Editor/UnrealEdEngine.h"
 #include "Toolkits/IToolkitHost.h"
@@ -79,6 +78,8 @@
 #include "TransformConstraint.h"
 #include "Misc/ScopedSlowTask.h"
 #include "Misc/TransactionObjectEvent.h"
+#include "ConstraintChannelHelper.h"
+#include "MovieSceneConstraintChannelHelper.h"
 
 #define LOCTEXT_NAMESPACE "FControlRigParameterTrackEditor"
 
@@ -2455,7 +2456,7 @@ namespace
 }
 
 void FControlRigParameterTrackEditor::HandleOnConstraintAdded(
-	UMovieSceneControlRigParameterSection* InSection,
+	IMovieSceneConstrainedSection* InSection,
 	FMovieSceneConstraintChannel* InConstraintChannel)
 {
 	if (!InConstraintChannel)
@@ -2487,60 +2488,67 @@ void FControlRigParameterTrackEditor::HandleOnConstraintAdded(
 }
 
 void FControlRigParameterTrackEditor::HandleConstraintKeyDeleted(
-	const UMovieSceneControlRigParameterSection* InSection,
+	IMovieSceneConstrainedSection* InSection,
 	const FMovieSceneConstraintChannel* InConstraintChannel,
 	const TArray<FKeyAddOrDeleteEventItem>& InDeletedItems) const
 {
-	const UControlRig* ControlRig = InSection ? InSection->GetControlRig() : nullptr;
+	UMovieSceneControlRigParameterSection* Section = Cast<UMovieSceneControlRigParameterSection>(InSection);
+
+	const UControlRig* ControlRig = Section ? Section->GetControlRig() : nullptr;
 	if (!ControlRig || !InConstraintChannel)
 	{
 		return;
 	}
 	
 	const FConstraintAndControlData ConstraintAndControlData =
-		FConstraintAndControlData::CreateFromSection(InSection, InConstraintChannel);
+		FConstraintAndControlData::CreateFromSection(Section, InConstraintChannel);
 	if (ConstraintAndControlData.IsValid())
 	{
+
 		UTickableTransformConstraint* Constraint = ConstraintAndControlData.Constraint.Get();
 		for (const FKeyAddOrDeleteEventItem& EventItem: InDeletedItems)
 		{
-			FConstraintChannelHelper::HandleConstraintKeyDeleted(
+			FMovieSceneConstraintChannelHelper::HandleConstraintKeyDeleted(
 				Constraint, InConstraintChannel,
-				GetSequencer(), InSection,
+				GetSequencer(), Section,
 				EventItem.Frame);
 		}
 	}
 }
 
 void FControlRigParameterTrackEditor::HandleConstraintKeyMoved(
-	const UMovieSceneControlRigParameterSection* InSection,
+	IMovieSceneConstrainedSection* InSection,
 	const FMovieSceneConstraintChannel* InConstraintChannel,
 	const TArray<FKeyMoveEventItem>& InMovedItems)
 {
+	UMovieSceneControlRigParameterSection* Section = Cast<UMovieSceneControlRigParameterSection>(InSection);
+
 	const FConstraintAndControlData ConstraintAndControlData =
-	FConstraintAndControlData::CreateFromSection(InSection, InConstraintChannel);
+	FConstraintAndControlData::CreateFromSection(Section, InConstraintChannel);
 
 	if (ConstraintAndControlData.IsValid())
 	{
 		const UTickableTransformConstraint* Constraint = ConstraintAndControlData.Constraint.Get();
 		for (const FKeyMoveEventItem& MoveEventItem : InMovedItems)
 		{
-			FConstraintChannelHelper::HandleConstraintKeyMoved(
-				Constraint, InConstraintChannel, InSection,
+			FMovieSceneConstraintChannelHelper::HandleConstraintKeyMoved(
+				Constraint, InConstraintChannel, Section,
 				MoveEventItem.Frame, MoveEventItem.NewFrame);
 		}
 	}
 }
 
-void FControlRigParameterTrackEditor::HandleConstraintRemoved(UMovieSceneControlRigParameterSection* InSection) const
+void FControlRigParameterTrackEditor::HandleConstraintRemoved(IMovieSceneConstrainedSection* InSection) const
 {
-	if (const UControlRig* ControlRig = InSection->GetControlRig())
+	UMovieSceneControlRigParameterSection* Section = Cast<UMovieSceneControlRigParameterSection>(InSection);
+
+	if (const UControlRig* ControlRig = Section->GetControlRig())
 	{
 		FConstraintsManagerController& Controller = FConstraintsManagerController::Get(ControlRig->GetWorld());
 		if (!Controller.OnConstraintRemoved().IsBoundToObject(InSection))
 		{
 			InSection->OnConstraintRemovedHandle =
-			Controller.OnConstraintRemoved().AddLambda([InSection, this](FName InConstraintName)
+			Controller.OnConstraintRemoved().AddLambda([InSection, Section, this](FName InConstraintName)
 			{
 				const FConstraintAndActiveChannel* ConstraintChannel = InSection->GetConstraintChannel(InConstraintName);
 				if (!ConstraintChannel)
@@ -2550,11 +2558,11 @@ void FControlRigParameterTrackEditor::HandleConstraintRemoved(UMovieSceneControl
 				
 				if (ConstraintChannel->Constraint.IsValid())
 				{
-					FConstraintChannelHelper::HandleConstraintRemoved(
+					FMovieSceneConstraintChannelHelper::HandleConstraintRemoved(
 						ConstraintChannel->Constraint.Get(),
 						&ConstraintChannel->ActiveChannel,
 						GetSequencer(),
-						InSection);
+						Section);
 				}
 
 				InSection->RemoveConstraintChannel(InConstraintName);
diff --git a/Engine/Plugins/Animation/ControlRig/Source/ControlRigEditor/Private/Sequencer/ControlRigParameterTrackEditor.h b/Engine/Plugins/Animation/ControlRig/Source/ControlRigEditor/Private/Sequencer/ControlRigParameterTrackEditor.h
--- a/Engine/Plugins/Animation/ControlRig/Source/ControlRigEditor/Private/Sequencer/ControlRigParameterTrackEditor.h
+++ b/Engine/Plugins/Animation/ControlRig/Source/ControlRigEditor/Private/Sequencer/ControlRigParameterTrackEditor.h
@@ -147,17 +147,17 @@ private:
 
 	/** ConstraintChannel Delegates*/
 	void HandleOnConstraintAdded(
-		UMovieSceneControlRigParameterSection* InSection,
+		IMovieSceneConstrainedSection* InSection,
 		FMovieSceneConstraintChannel* InConstraintChannel);
 	void HandleConstraintKeyDeleted(
-		const UMovieSceneControlRigParameterSection* InSection,
+		IMovieSceneConstrainedSection* InSection,
 		const FMovieSceneConstraintChannel* InConstraintChannel,
 		const TArray<FKeyAddOrDeleteEventItem>& InDeletedItems) const;
 	static void HandleConstraintKeyMoved(
-		const UMovieSceneControlRigParameterSection* InSection,
+		IMovieSceneConstrainedSection* InSection,
 		const FMovieSceneConstraintChannel* InConstraintChannel,
 		const TArray<FKeyMoveEventItem>& InMovedItems);
-	void HandleConstraintRemoved(UMovieSceneControlRigParameterSection* InSection) const;
+	void HandleConstraintRemoved(IMovieSceneConstrainedSection* InSection) const;
 
 	/** Select control rig if not selected, select controls from key areas */
 	void SelectRigsAndControls(UControlRig* Subject, const TArray<const IKeyArea*>& KeyAreas);
diff --git a/Engine/Plugins/Animation/ControlRig/Source/ControlRigEditor/Private/Sequencer/ControlRigSpaceChannelEditors.cpp b/Engine/Plugins/Animation/ControlRig/Source/ControlRigEditor/Private/Sequencer/ControlRigSpaceChannelEditors.cpp
--- a/Engine/Plugins/Animation/ControlRig/Source/ControlRigEditor/Private/Sequencer/ControlRigSpaceChannelEditors.cpp
+++ b/Engine/Plugins/Animation/ControlRig/Source/ControlRigEditor/Private/Sequencer/ControlRigSpaceChannelEditors.cpp
@@ -1,6 +1,7 @@
 // Copyright Epic Games, Inc. All Rights Reserved.
 
 #include "ControlRigSpaceChannelEditors.h"
+#include "ConstraintChannelHelper.inl"
 #include "MovieSceneSequence.h"
 #include "MovieSceneSequenceEditor.h"
 #include "MovieSceneEventUtils.h"
@@ -174,6 +175,8 @@ TSharedRef<SWidget> CreateKeyEditor(const TMovieSceneChannelHandle<FMovieSceneCo
 	return SNullWidget::NullWidget;
 }
 
+
+
 /*******************************************************************
 *
 * FControlRigSpaceChannelHelpers
@@ -251,136 +254,6 @@ FSpaceChannelAndSection FControlRigSpaceChannelHelpers::FindSpaceChannelAndSecti
 	return SpaceChannelAndSection;
 }
 
-template<typename ChannelType>
-void EvaluateTangentAtThisTimeImpl(int32 ChannelIndex, int32 NumChannels,
-	UMovieSceneSection* Section,
-	FFrameNumber Time,
-	TArray<FMovieSceneTangentData>& OutTangents)
-{
-	using ChannelValueType = typename ChannelType::ChannelValueType;
-
-	// NOTE this might be moved to FMovieSceneFloatChannel
-	auto EvaluateTangent = [](const ChannelType* InChannel, FFrameNumber InTime)
-	{
-		const TMovieSceneChannelData<const ChannelValueType> ChannelInterface = InChannel->GetData();
-		const TArrayView<const FFrameNumber> Times = ChannelInterface.GetTimes();
-
-		// Need at least two keys to evaluate a derivative
-		if (Times.Num() < 2)
-		{
-			static const FMovieSceneTangentData DefaultTangent;
-			return DefaultTangent;
-		}
-
-		const TArrayView<const ChannelValueType> Values = ChannelInterface.GetValues();
-
-		// look around to get the closest key tangent if fairly close
-		const int32 Tolerance = static_cast<int32>(InChannel->GetTickResolution().AsDecimal() * 0.01);
-		// NOTE FindKey might return Times.Num() (see AlgoImpl::LowerBoundInternal)
-		const int32 ExistingIndex = ChannelInterface.FindKey(InTime, Tolerance);
-		if (Times.IsValidIndex(ExistingIndex) && Values.IsValidIndex(ExistingIndex))
-		{
-			const FFrameNumber& Time = Times[ExistingIndex];
-			const int32 DiffToKey = InTime.Value - Time.Value;
-
-			// if the closest key is within a threshold, we use it's tangent directly instead of computing one
-			if (FMath::Abs(DiffToKey) <= Tolerance)
-			{
-				const ChannelValueType Value = Values[ExistingIndex];
-				if (DiffToKey == 0)
-				{
-					return Value.Tangent;
-				}
-
-				FMovieSceneTangentData Tangent = Value.Tangent;
-				if (DiffToKey < 0) // pretty close to the next key
-				{
-					Tangent.LeaveTangent = 0.f;
-				}
-				else // pretty close to the previous key
-				{
-					Tangent.ArriveTangent = 0.f;
-				}
-				return Tangent;
-			}
-		}
-
-		// compute tangent using central difference
-		// NOTE we may wanna compute a backward / forward difference instead
-		const int32 Delta = static_cast<int32>(InChannel->GetTickResolution().AsDecimal() * 0.1);
-
-		float PrevValue = 0.f; InChannel->Evaluate(InTime - Delta, PrevValue);
-		float NextValue = 0.f; InChannel->Evaluate(InTime + Delta, NextValue);
-		const float TangentValue = (NextValue - PrevValue) / (2.f * Delta);
-
-		FMovieSceneTangentData TangentData;
-		TangentData.ArriveTangent = TangentValue;
-		TangentData.LeaveTangent = TangentValue;
-
-		return TangentData;
-	};
-
-
-	// compute and store tangents
-	OutTangents.SetNum(NumChannels);
-
-	const TArrayView<ChannelType*> Channels = Section->GetChannelProxy().GetChannels<ChannelType>();
-	for (int32 Index = 0; Index < NumChannels; ++Index, ++ChannelIndex)
-	{
-		OutTangents[Index] = EvaluateTangent(Channels[ChannelIndex], Time);
-	}
-}
-
-// NOTE we may pass an enum to tell which tangent we wanna set (arrive, leave, both)
-template<typename ChannelType>
-void SetTangentsAtThisTimeImpl(int32 ChannelIndex,
-	int32 NumChannels,
-	UMovieSceneSection* Section,
-	FFrameNumber Time,
-	const TArray<FMovieSceneTangentData>& InTangents)
-{
-	using ChannelValueType = typename ChannelType::ChannelValueType;
-
-	const TArrayView<ChannelType*> Channels = Section->GetChannelProxy().GetChannels<ChannelType>();
-
-	for (int32 Index = 0; Index < NumChannels; ++Index, ++ChannelIndex)
-	{
-		TMovieSceneChannelData<ChannelValueType> ChannelInterface = Channels[ChannelIndex]->GetData();
-		TArrayView<ChannelValueType> Values = ChannelInterface.GetValues();
-		const int32 KeyIndex = ChannelInterface.FindKey(Time);
-		if (KeyIndex != INDEX_NONE)
-		{
-			ChannelValueType& Value = Values[KeyIndex];
-			Value.Tangent.ArriveTangent = InTangents[Index].ArriveTangent;
-			Value.Tangent.LeaveTangent = InTangents[Index].LeaveTangent;
-			Value.TangentMode = RCTM_Break;
-		}
-	}
-}
-
-void FControlRigSpaceChannelHelpers::EvaluateTangentAtThisTime(int32 ChannelIndex, int32 NumChannels, bool bFloatChannels, UMovieSceneSection* Section, FFrameNumber Time, TArray<FMovieSceneTangentData>& OutTangents)
-{
-	if (bFloatChannels)
-	{
-		EvaluateTangentAtThisTimeImpl<FMovieSceneFloatChannel>(ChannelIndex, NumChannels, Section, Time, OutTangents);
-	}
-	else
-	{
-		EvaluateTangentAtThisTimeImpl<FMovieSceneDoubleChannel>(ChannelIndex, NumChannels, Section, Time, OutTangents);
-	}
-}
-
-void FControlRigSpaceChannelHelpers::SetTangentsAtThisTime(int32 ChannelIndex, int32 NumChannels, bool bFloatChannels, UMovieSceneSection* Section, FFrameNumber Time, const TArray<FMovieSceneTangentData>& InTangents)
-{
-	if (bFloatChannels)
-	{
-		SetTangentsAtThisTimeImpl<FMovieSceneFloatChannel>(ChannelIndex, NumChannels, Section, Time, InTangents);
-	}
-	else
-	{
-		SetTangentsAtThisTimeImpl<FMovieSceneDoubleChannel>(ChannelIndex, NumChannels, Section, Time, InTangents);
-	}
-}
 
 
 /*
@@ -500,7 +373,7 @@ FKeyHandle FControlRigSpaceChannelHelpers::SequencerKeyControlRigSpaceChannel(UC
 				if (NumChannels > 0)
 				{
 					ChannelIndex = pChannelIndex->ChannelIndex;
-					EvaluateTangentAtThisTime(ChannelIndex, NumChannels, true,ControlRigSection, Time, Tangents);
+					EvaluateTangentAtThisTime<FMovieSceneFloatChannel>(ChannelIndex, NumChannels,ControlRigSection, Time, Tangents);
 				}
 			}
 		}
@@ -565,7 +438,7 @@ FKeyHandle FControlRigSpaceChannelHelpers::SequencerKeyControlRigSpaceChannel(UC
 			// need to update the tangents here to keep the arriving animation
 			if (ControlRigSection  && NumChannels > 0)
 			{
-				SetTangentsAtThisTime(ChannelIndex, NumChannels,true, ControlRigSection,GlobalTime.GetFrame(), Tangents);
+				SetTangentsAtThisTime<FMovieSceneFloatChannel>(ChannelIndex, NumChannels, ControlRigSection,GlobalTime.GetFrame(), Tangents);
 			}
 		}
 
@@ -592,7 +465,7 @@ FKeyHandle FControlRigSpaceChannelHelpers::SequencerKeyControlRigSpaceChannel(UC
 			{
 				if (ControlRigSection && NumChannels > 0)
 				{
-					SetTangentsAtThisTime(ChannelIndex, NumChannels, true,  ControlRigSection, GlobalTime.GetFrame(), Tangents);
+					SetTangentsAtThisTime<FMovieSceneFloatChannel>(ChannelIndex, NumChannels, ControlRigSection, GlobalTime.GetFrame(), Tangents);
 				}
 			}
 
diff --git a/Engine/Plugins/Animation/ControlRig/Source/ControlRigEditor/Private/Tools/BakingHelper.cpp b/Engine/Plugins/Animation/ControlRig/Source/ControlRigEditor/Private/Tools/BakingHelper.cpp
--- a/Engine/Plugins/Animation/ControlRig/Source/ControlRigEditor/Private/Tools/BakingHelper.cpp
+++ b/Engine/Plugins/Animation/ControlRig/Source/ControlRigEditor/Private/Tools/BakingHelper.cpp
@@ -5,8 +5,6 @@
 
 #include "ControlRig.h"
 #include "ILevelSequenceEditorToolkit.h"
-#include "Tracks/MovieScene3DTransformTrack.h"
-#include "Sections/MovieScene3DTransformSection.h"
 #include "ISequencer.h"
 #include "LevelSequence.h"
 #include "LevelSequenceEditorBlueprintLibrary.h"
@@ -29,236 +27,3 @@ TWeakPtr<ISequencer> FBakingHelper::GetSequencer()
 	}
 	return nullptr;
 }
-
-
-void FBakingHelper::CalculateFramesBetween(
-	const UMovieScene* MovieScene,
-	FFrameNumber StartFrame,
-	FFrameNumber EndFrame,
-	TArray<FFrameNumber>& OutFrames)
-{
-	if(StartFrame > EndFrame)
-	{
-		Swap(StartFrame, EndFrame);
-	}
-	
-	const FFrameRate TickResolution = MovieScene->GetTickResolution();
-	const FFrameRate DisplayResolution = MovieScene->GetDisplayRate();
-
-	const FFrameNumber StartTimeInDisplay = FFrameRate::TransformTime(FFrameTime(StartFrame), TickResolution, DisplayResolution).FloorToFrame();
-	const FFrameNumber EndTimeInDisplay = FFrameRate::TransformTime(FFrameTime(EndFrame), TickResolution, DisplayResolution).CeilToFrame();
-
-	OutFrames.Reserve(EndTimeInDisplay.Value - StartTimeInDisplay.Value + 1);
-	for (FFrameNumber DisplayFrameNumber = StartTimeInDisplay; DisplayFrameNumber <= EndTimeInDisplay; ++DisplayFrameNumber)
-	{
-		FFrameNumber TickFrameNumber = FFrameRate::TransformTime(FFrameTime(DisplayFrameNumber), DisplayResolution, TickResolution).FrameNumber;
-		OutFrames.Add(TickFrameNumber);
-	}
-}
-
-UMovieScene3DTransformSection* FBakingHelper::GetTransformSection(
-	const ISequencer* InSequencer,
-	const FGuid& InGuid,
-	const FTransform& InDefaultTransform)
-{
-	if (!InSequencer || !InSequencer->GetFocusedMovieSceneSequence())
-	{
-		return nullptr;
-	}
-	
-	UMovieScene* MovieScene = InSequencer->GetFocusedMovieSceneSequence()->GetMovieScene();
-	if (!MovieScene)
-	{
-		return nullptr;
-	}
-
-	UMovieScene3DTransformTrack* TransformTrack = MovieScene->FindTrack<UMovieScene3DTransformTrack>(InGuid);
-	if (!TransformTrack)
-	{
-		MovieScene->Modify();
-		TransformTrack = MovieScene->AddTrack<UMovieScene3DTransformTrack>(InGuid);
-	}
-	TransformTrack->Modify();
-
-	static constexpr FFrameNumber Frame0;
-	bool bSectionAdded = false;
-	UMovieScene3DTransformSection* TransformSection =
-				Cast<UMovieScene3DTransformSection>(TransformTrack->FindOrAddSection(Frame0, bSectionAdded));
-	if (!TransformSection)
-	{
-		return nullptr;
-	}
-
-	TransformSection->Modify();
-	if (bSectionAdded)
-	{
-		TransformSection->SetRange(TRange<FFrameNumber>::All());
-
-		const FVector Location0 = InDefaultTransform.GetLocation();
-		const FRotator Rotation0 = InDefaultTransform.GetRotation().Rotator();
-		const FVector Scale3D0 = InDefaultTransform.GetScale3D();
-
-		const TArrayView<FMovieSceneDoubleChannel*> Channels =
-			TransformSection->GetChannelProxy().GetChannels<FMovieSceneDoubleChannel>();
-		Channels[0]->SetDefault(Location0.X);
-		Channels[1]->SetDefault(Location0.Y);
-		Channels[2]->SetDefault(Location0.Z);
-		Channels[3]->SetDefault(Rotation0.Roll);
-		Channels[4]->SetDefault(Rotation0.Pitch);
-		Channels[5]->SetDefault(Rotation0.Yaw);
-		Channels[6]->SetDefault(Scale3D0.X);
-		Channels[7]->SetDefault(Scale3D0.Y);
-		Channels[8]->SetDefault(Scale3D0.Z);
-	}
-	
-	return TransformSection;
-}
-
-bool FBakingHelper::AddTransformKeys(
-	const UMovieScene3DTransformSection* InTransformSection,
-	const TArray<FFrameNumber>& Frames,
-	const TArray<FTransform>& InLocalTransforms,
-	const EMovieSceneTransformChannel& InChannels)
-{
-	if (!InTransformSection)
-	{
-		return false;	
-	}
-
-	if (Frames.IsEmpty() || Frames.Num() != InLocalTransforms.Num())
-	{
-		return false;
-	}
-
-	auto GetValue = [](const uint32 Index,const FVector& InLocation, const FRotator& InRotation, const FVector& InScale)
-	{
-		switch (Index)
-		{
-		case 0:
-			return InLocation.X;
-		case 1:
-			return InLocation.Y;
-		case 2:
-			return InLocation.Z;
-		case 3:
-			return InRotation.Roll;
-		case 4:
-			return InRotation.Pitch;
-		case 5:
-			return InRotation.Yaw;
-		case 6:
-			return InScale.X;
-		case 7:
-			return InScale.Y;
-		case 8:
-			return InScale.Z;
-		default:
-			ensure(false);
-			break;
-		}
-		return 0.0;
-	};
-
-	const bool bKeyTranslation = EnumHasAllFlags(InChannels, EMovieSceneTransformChannel::Translation);
-	const bool bKeyRotation = EnumHasAllFlags(InChannels, EMovieSceneTransformChannel::Rotation);
-	const bool bKeyScale = EnumHasAllFlags(InChannels, EMovieSceneTransformChannel::Scale);
-
-	TArray<uint32> ChannelsIndexToKey;
-	if (bKeyTranslation)
-	{
-		ChannelsIndexToKey.Append({0,1,2});
-	}
-	if (bKeyRotation)
-	{
-		ChannelsIndexToKey.Append({3,4,5});
-	}
-	if (bKeyScale)
-	{
-		ChannelsIndexToKey.Append({6,7,8});
-	}
-
-	const TArrayView<FMovieSceneDoubleChannel*> Channels =
-		InTransformSection->GetChannelProxy().GetChannels<FMovieSceneDoubleChannel>();
-
-	// set default
-	const FTransform& LocalTransform0 = InLocalTransforms[0];
-	const FVector Location0 = LocalTransform0.GetLocation();
-	const FRotator Rotation0 = LocalTransform0.GetRotation().Rotator();
-	const FVector Scale3D0 = LocalTransform0.GetScale3D();
-	
-	for (int32 ChannelIndex = 0; ChannelIndex < 9; ChannelIndex++)
-	{
-		if (!Channels[ChannelIndex]->GetDefault().IsSet())
-		{
-			const double Value = GetValue(ChannelIndex, Location0, Rotation0, Scale3D0);
-			Channels[ChannelIndex]->SetDefault(Value);
-		}
-	}
-
-	// add keys
-	for (int32 Index = 0; Index < Frames.Num(); ++Index)
-	{
-		const FFrameNumber& Frame = Frames[Index];
-		const FTransform& LocalTransform = InLocalTransforms[Index];
-		
-		const FVector Location = LocalTransform.GetLocation();
-		const FRotator Rotation = LocalTransform.GetRotation().Rotator();
-		const FVector Scale3D = LocalTransform.GetScale3D();
-
-		for (const int32 ChannelIndex: ChannelsIndexToKey)
-		{
-			const double Value = GetValue(ChannelIndex, Location, Rotation, Scale3D);
-			TMovieSceneChannelData<FMovieSceneDoubleValue> ChannelData = Channels[ChannelIndex]->GetData();
-			MovieSceneToolHelpers::SetOrAddKey(ChannelData, Frame, Value);
-		}
-	}
-
-	//now we need to set auto tangents
-	for (const int32 ChannelIndex: ChannelsIndexToKey)
-	{
-		Channels[ChannelIndex]->AutoSetTangents();
-	}
-
-	return true;
-}
-
-bool FBakingHelper::AddTransformKeys(
-	UControlRig* InControlRig, const FName& InControlName,
-	const TArray<FFrameNumber>& InFrames,
-	const TArray<FTransform>& InTransforms,
-	const EMovieSceneTransformChannel& InChannels,
-	const FFrameRate& InTickResolution,
-	const bool bLocal)
-{
-	if (InFrames.IsEmpty() || InFrames.Num() != InTransforms.Num())
-	{
-		return false;
-	}
-
-	auto KeyframeFunc = [InControlRig, InControlName, bLocal](const FTransform& InTransform, const FRigControlModifiedContext& InKeyframeContext)
-	{
-		static constexpr bool bNotify = true;
-		static constexpr bool bUndo = false;
-		static constexpr bool bFixEuler = true;
-
-		if (bLocal)
-		{
-			return InControlRig->SetControlLocalTransform(InControlName, InTransform, bNotify, InKeyframeContext, bUndo, bFixEuler);
-		}
-		InControlRig->SetControlGlobalTransform(InControlName, InTransform, bNotify, InKeyframeContext, bUndo, bFixEuler);
-	};
-
-	FRigControlModifiedContext KeyframeContext;
-	KeyframeContext.SetKey = EControlRigSetKey::Always;
-	KeyframeContext.KeyMask = static_cast<uint32>(InChannels);
-	
-	for (int32 Index = 0; Index < InFrames.Num(); ++Index)
-	{
-		const FFrameNumber& Frame = InFrames[Index];
-		KeyframeContext.LocalTime = InTickResolution.AsSeconds(FFrameTime(Frame));
-		
-		KeyframeFunc(InTransforms[Index], KeyframeContext);
-	}
-
-	return true;
-}
\ No newline at end of file
diff --git a/Engine/Plugins/Animation/ControlRig/Source/ControlRigEditor/Private/Tools/ConstraintBaker.cpp b/Engine/Plugins/Animation/ControlRig/Source/ControlRigEditor/Private/Tools/ConstraintBaker.cpp
--- a/Engine/Plugins/Animation/ControlRig/Source/ControlRigEditor/Private/Tools/ConstraintBaker.cpp
+++ b/Engine/Plugins/Animation/ControlRig/Source/ControlRigEditor/Private/Tools/ConstraintBaker.cpp
@@ -21,78 +21,75 @@
 
 namespace
 {
-	
-UMovieScene3DTransformSection* GetTransformSection(
-	const TSharedPtr<ISequencer>& InSequencer,
-	AActor* InActor,
-	const FTransform& InTransform0)
-{
-	if (!InSequencer || !InSequencer->GetFocusedMovieSceneSequence())
-	{
-		return nullptr;
-	}
-	
-	const FGuid Guid = InSequencer->GetHandleToObject(InActor,true);
-	if (!Guid.IsValid())
-	{
-		return nullptr;
-	}
-	
-	return FBakingHelper::GetTransformSection(InSequencer.Get(), Guid, InTransform0);
-}
 
-void BakeComponent(
-	const TSharedPtr<ISequencer>& InSequencer,
-	const UTransformableComponentHandle* InComponentHandle,
-	const TArray<FFrameNumber>& InFrames,
-	const TArray<FTransform>& InTransforms,
-	const EMovieSceneTransformChannel& InChannels)
-{
-	ensure(InTransforms.Num());
-	
-	if (!InComponentHandle->IsValid())
+	UMovieScene3DTransformSection* GetTransformSection(
+		const TSharedPtr<ISequencer>& InSequencer,
+		AActor* InActor,
+		const FTransform& InTransform0)
 	{
-		return;
-	}
-	AActor* Actor = InComponentHandle->Component->GetOwner();
-	if (!Actor)
-	{
-		return;
+		if (!InSequencer || !InSequencer->GetFocusedMovieSceneSequence())
+		{
+			return nullptr;
+		}
+
+		const FGuid Guid = InSequencer->GetHandleToObject(InActor, true);
+		if (!Guid.IsValid())
+		{
+			return nullptr;
+		}
+
+		return MovieSceneToolHelpers::GetTransformSection(InSequencer.Get(), Guid, InTransform0);
 	}
-	const UMovieScene3DTransformSection* TransformSection = GetTransformSection(InSequencer, Actor, InTransforms[0]);
-	if (!TransformSection)
+
+	void BakeComponent(
+		const TSharedPtr<ISequencer>& InSequencer,
+		const UTransformableComponentHandle* InComponentHandle,
+		const TArray<FFrameNumber>& InFrames,
+		const TArray<FTransform>& InTransforms,
+		const EMovieSceneTransformChannel& InChannels)
 	{
-		return;
-	}
+		ensure(InTransforms.Num());
 
-	FBakingHelper::AddTransformKeys(TransformSection, InFrames, InTransforms, InChannels);
-}
+		if (!InComponentHandle->IsValid())
+		{
+			return;
+		}
+		AActor* Actor = InComponentHandle->Component->GetOwner();
+		if (!Actor)
+		{
+			return;
+		}
+		UMovieScene3DTransformSection* TransformSection = GetTransformSection(InSequencer, Actor, InTransforms[0]);
+		if (!TransformSection)
+		{
+			return;
+		}
+		TransformSection->Modify();
+		const UMovieScene* MovieScene = InSequencer->GetFocusedMovieSceneSequence()->GetMovieScene();
+		InComponentHandle->AddTransformKeys(InFrames, InTransforms, InChannels, MovieScene->GetTickResolution(), TransformSection, true);
 
-void BakeControl(
-	const TSharedPtr<ISequencer>& InSequencer,
-	const UTransformableControlHandle* InControlHandle,
-	const TArray<FFrameNumber>& InFrames,
-	const TArray<FTransform>& InLocalTransforms,
-	const EMovieSceneTransformChannel& InChannels)
-{
-	if (!InControlHandle->IsValid())
-	{
-		return;
 	}
 
-	const UMovieScene* MovieScene = InSequencer->GetFocusedMovieSceneSequence()->GetMovieScene();
+	void BakeControl(
+		const TSharedPtr<ISequencer>& InSequencer,
+		const UTransformableControlHandle* InControlHandle,
+		const TArray<FFrameNumber>& InFrames,
+		const TArray<FTransform>& InLocalTransforms,
+		const EMovieSceneTransformChannel& InChannels)
+	{
+		ensure(InLocalTransforms.Num());
 
-	FBakingHelper::AddTransformKeys(
-		InControlHandle->ControlRig.Get(),
-		InControlHandle->ControlName,
-		InFrames,
-		InLocalTransforms,
-		InChannels,
-		MovieScene->GetTickResolution());
-}
+		if (!InControlHandle->IsValid())
+		{
+			return;
+		}
 
-}
+		const UMovieScene* MovieScene = InSequencer->GetFocusedMovieSceneSequence()->GetMovieScene();
+		UMovieSceneSection* Section = nullptr; //cnntrol rig doesn't need section it instead
+		InControlHandle->AddTransformKeys(InFrames, InLocalTransforms, InChannels, MovieScene->GetTickResolution(), Section, true);
 
+	}
+}
 void FConstraintBaker::DoIt(UTickableTransformConstraint* InConstraint)
 {
 	const TWeakPtr<ISequencer> WeakSequencer = FBakingHelper::GetSequencer();
@@ -111,7 +108,7 @@ void FConstraintBaker::DoIt(UTickableTransformConstraint* InConstraint)
 	const FFrameNumber StartFrame = MovieScene->GetPlaybackRange().GetLowerBoundValue();
 	const FFrameNumber EndFrame = MovieScene->GetPlaybackRange().GetUpperBoundValue();
 	TArray<FFrameNumber> Frames;
-	FBakingHelper::CalculateFramesBetween(MovieScene, StartFrame, EndFrame, Frames);
+	MovieSceneToolHelpers::CalculateFramesBetween(MovieScene, StartFrame, EndFrame, Frames);
 
 	if (Frames.IsEmpty())
 	{
@@ -128,7 +125,7 @@ void FConstraintBaker::DoIt(UTickableTransformConstraint* InConstraint)
 	}
 
 	// bake to channel curves
-	const EMovieSceneTransformChannel Channels = GetChannelsToKey(InConstraint);
+	const EMovieSceneTransformChannel Channels = InConstraint->GetChannelsToKey();
 	AddTransformKeys(Sequencer, InConstraint->ChildTRSHandle, Frames, Transforms, Channels);
 
 	// disable constraint
@@ -222,25 +219,6 @@ void FConstraintBaker::AddTransformKeys(
 	}
 }
 
-EMovieSceneTransformChannel FConstraintBaker::GetChannelsToKey(const UTickableTransformConstraint* InConstraint)
-{
-	static const TMap< ETransformConstraintType, EMovieSceneTransformChannel > ConstraintToChannels({
-	{ETransformConstraintType::Translation, EMovieSceneTransformChannel::Translation},
-	{ETransformConstraintType::Rotation, EMovieSceneTransformChannel::Rotation},
-	{ETransformConstraintType::Scale, EMovieSceneTransformChannel::Scale},
-	{ETransformConstraintType::Parent, EMovieSceneTransformChannel::AllTransform},
-	{ETransformConstraintType::LookAt, EMovieSceneTransformChannel::Rotation}
-	});
-
-	const ETransformConstraintType ConstType = static_cast<ETransformConstraintType>(InConstraint->GetType());
-	if (const EMovieSceneTransformChannel* Channel = ConstraintToChannels.Find(ConstType))
-	{
-		return *Channel;; 
-	}
-	
-	return EMovieSceneTransformChannel::AllTransform;
-}
-
 #undef LOCTEXT_NAMESPACE
 
 
diff --git a/Engine/Plugins/Animation/ControlRig/Source/ControlRigEditor/Private/Tools/ControlRigSnapper.cpp b/Engine/Plugins/Animation/ControlRig/Source/ControlRigEditor/Private/Tools/ControlRigSnapper.cpp
--- a/Engine/Plugins/Animation/ControlRig/Source/ControlRigEditor/Private/Tools/ControlRigSnapper.cpp
+++ b/Engine/Plugins/Animation/ControlRig/Source/ControlRigEditor/Private/Tools/ControlRigSnapper.cpp
@@ -190,7 +190,7 @@ struct FGuidAndActor
 		const TArray<FTransform>& WorldTransformsToSnapTo, const UControlRigSnapSettings* SnapSettings) const
 	{
 		// get section
-		const UMovieScene3DTransformSection* TransformSection = FBakingHelper::GetTransformSection(Sequencer, Guid);
+		const UMovieScene3DTransformSection* TransformSection = MovieSceneToolHelpers::GetTransformSection(Sequencer, Guid);
 		if (!TransformSection)
 		{
 			return false;
@@ -240,7 +240,7 @@ struct FGuidAndActor
 		}
 
 		// add keys
-		FBakingHelper::AddTransformKeys(TransformSection, Frames, LocalTransforms, Channels);
+		MovieSceneToolHelpers::AddTransformKeys(TransformSection, Frames, LocalTransforms, Channels);
 
 		// notify
 		Sequencer->NotifyMovieSceneDataChanged(EMovieSceneDataChangeType::MovieSceneStructureItemsChanged);
@@ -320,7 +320,7 @@ bool FControlRigSnapper::SnapIt(FFrameNumber StartFrame, FFrameNumber EndFrame,c
 		MovieScene->Modify();
 		
 		TArray<FFrameNumber> Frames;
-		FBakingHelper::CalculateFramesBetween(MovieScene, StartFrame, EndFrame, Frames);
+		MovieSceneToolHelpers::CalculateFramesBetween(MovieScene, StartFrame, EndFrame, Frames);
 
 		TArray<FTransform> WorldTransformToSnap;
 		bool bSnapToFirstFrameNotParents = !CalculateWorldTransformsFromParents(Sequencer, ParentToSnap, Frames, WorldTransformToSnap);
@@ -329,11 +329,10 @@ bool FControlRigSnapper::SnapIt(FFrameNumber StartFrame, FFrameNumber EndFrame,c
 			UE_LOG(LogControlRig, Error, TEXT("Number of Frames %d to Snap different than Parent Frames %d"), Frames.Num(),WorldTransformToSnap.Num());
 			return false;
 		}
-		// If we are snapping from the current time there are some assumptions to not re-calcuate values so we need to make sure sequencer is up to date
-		if (StartFrame == Sequencer->GetLocalTime().Time.RoundToFrame())
-		{
-			Sequencer->ForceEvaluate();
-		}
+
+		//need to make sure binding is setup
+		Sequencer->ForceEvaluate();
+		
 		TArray<FGuidAndActor > ActorsToSnap;
 		//There may be Actors here not in Sequencer so we add them to sequencer also
 		for (const FActorForWorldTransforms& ActorSelection : ActorToSnap.Actors)
diff --git a/Engine/Plugins/Animation/ControlRig/Source/ControlRigEditor/Public/ConstraintChannelHelper.h b/Engine/Plugins/Animation/ControlRig/Source/ControlRigEditor/Public/ConstraintChannelHelper.h
--- a/Engine/Plugins/Animation/ControlRig/Source/ControlRigEditor/Public/ConstraintChannelHelper.h
+++ b/Engine/Plugins/Animation/ControlRig/Source/ControlRigEditor/Public/ConstraintChannelHelper.h
@@ -30,17 +30,9 @@ public:
 	/** @todo documentation. */
 	static bool IsKeyframingAvailable();
 	
-	/** @todo documentation. */
-	static void AddConstraintKey(UTickableTransformConstraint* InConstraint);
-	
 	/** Adds an active key if needed and does the compensation when switching. */
 	static void SmartConstraintKey(UTickableTransformConstraint* InConstraint);
 
-	/** Adds a transform key on InHandle at InTime. */
-	static void AddChildTransformKey(
-		const UTransformableHandle* InHandle,
-		const FFrameNumber& InTime,
-		const TSharedPtr<ISequencer>& InSequencer);
 
 	/** Compensate transform on handles when a constraint switches state. */
 	static void Compensate(
@@ -52,27 +44,8 @@ public:
 		const UMovieSceneControlRigParameterSection* Section,
 		const TOptional<FFrameNumber>& OptionalTime);
 
-	/** @todo documentation. */
-	static void HandleConstraintRemoved(
-		UTickableConstraint* InConstraint,
-		const FMovieSceneConstraintChannel* InChannel,
-		const TSharedPtr<ISequencer>& InSequencer,
-		const UMovieSceneSection* InSection);
 
-	/** @todo documentation. */
-	static void HandleConstraintKeyDeleted(
-		UTickableTransformConstraint* InConstraint,
-		const FMovieSceneConstraintChannel* InConstraintChannel,
-		const TSharedPtr<ISequencer>& InSequencer,
-		const UMovieSceneSection* InSection,
-		const FFrameNumber& InTime);
 
-	/** @todo documentation. */
-	static void HandleConstraintKeyMoved(
-		const UTickableTransformConstraint* InConstraint,
-		const FMovieSceneConstraintChannel* InConstraintChannel,
-		const UMovieSceneSection* InSection,
-		const FFrameNumber& InCurrentFrame, const FFrameNumber& InNextFrame);
 	
 private:
 
@@ -104,61 +77,6 @@ private:
 
 	/** END COMPONENT SECTION */
 
-	/** @todo documentation. */
-	static TArrayView<FMovieSceneFloatChannel*> GetTransformFloatChannels(
-		const UTransformableHandle* InHandle,
-		const TSharedPtr<ISequencer>& InSequencer);
-	static TArrayView<FMovieSceneFloatChannel*> GetTransformFloatChannels(
-		const UTransformableHandle* InHandle,
-		const UMovieSceneControlRigParameterSection* InSection);
-
-	/** @todo documentation.  TransformComponents use doubles */
-	static TArrayView<FMovieSceneDoubleChannel*> GetTransformDoubleChannels(
-		const UTransformableHandle* InHandle,
-		const TSharedPtr<ISequencer>& InSequencer);
-
-	/** @todo factorize. */
-	static TPair<const FChannelMapInfo*, int32> GetInfoAndNumFloatChannels(
-		const UControlRig* InControlRig,
-		const FName& InControlName,
-		const UMovieSceneControlRigParameterSection* InSection);
-
-	/** @todo documentation. */
-	template<typename ChannelType>
-	static void GetFramesToCompensate(
-		const FMovieSceneConstraintChannel& InActiveChannel,
-		const bool InActiveValueToBeSet,
-		const FFrameNumber& InTime,
-		const TArrayView<ChannelType*>& InChannels,
-		TArray<FFrameNumber>& OutFramesAfter);
-
-	/** @todo documentation. */
-	template< typename ChannelType >
-	static void GetFramesAfter(
-		const FMovieSceneConstraintChannel& InActiveChannel,
-		const FFrameNumber& InTime,
-		const TArrayView<ChannelType*>& InChannels,
-		TArray<FFrameNumber>& OutFrames);
-		
-	/** @todo documentation. */
-	template< typename ChannelType >
-    static void GetFramesWithinActiveState(
-	    const FMovieSceneConstraintChannel& InActiveChannel,
-	    const TArrayView<ChannelType*>& InChannels,
-	    TArray<FFrameNumber>& OutFrames);
-	
-	/** @todo documentation. */
-	template< typename ChannelType >
-	static void MoveTransformKeys(
-		const TArrayView<ChannelType*>& InChannels,
-		const FFrameNumber& InCurrentTime,
-		const FFrameNumber& InNextTime);
-
-	/** @todo documentation. */
-	template< typename ChannelType >
-	static void DeleteTransformKeys(
-		const TArrayView<ChannelType*>& InChannels,
-		const FFrameNumber& InTime);
 	
 	static bool bDoNotCompensate;
 };
diff --git a/Engine/Plugins/Animation/ControlRig/Source/ControlRigEditor/Public/ControlRigSpaceChannelEditors.h b/Engine/Plugins/Animation/ControlRig/Source/ControlRigEditor/Public/ControlRigSpaceChannelEditors.h
--- a/Engine/Plugins/Animation/ControlRig/Source/ControlRigEditor/Public/ControlRigSpaceChannelEditors.h
+++ b/Engine/Plugins/Animation/ControlRig/Source/ControlRigEditor/Public/ControlRigSpaceChannelEditors.h
@@ -59,11 +59,6 @@ struct FControlRigSpaceChannelHelpers
 	static FLinearColor GetColor(const FMovieSceneControlRigSpaceBaseKey& Key);
 	static FReply OpenBakeDialog(ISequencer* Sequencer, FMovieSceneControlRigSpaceChannel* Channel, int32 KeyIndex, UMovieSceneSection* SectionToKey);
 	static TArray<FKeyBarCurveModel::FBarRange> FindRanges(FMovieSceneControlRigSpaceChannel* Channel, const UMovieSceneSection* Section);
-
-	// evaluates the tangent by central difference at a specific time. (return the closest key tangent if within an acceptable tolerance) 
-	static void EvaluateTangentAtThisTime(int32 ChannelIndex, int32 NumChannels, bool bFloatChannels, UMovieSceneSection* Section,  FFrameNumber Time, TArray<FMovieSceneTangentData>& OutTangents);
-	// set the tangent at a specific time. (if there's a key to set at this time)
-	static void SetTangentsAtThisTime(int32 ChannelIndex, int32 NumChannels,bool bFloatChannels, UMovieSceneSection* Section, FFrameNumber Time, const TArray<FMovieSceneTangentData>& InTangents);
 	// retrieve the control and the channel infos for that ControlRig/Section.
 	static TPair<FRigControlElement*, FChannelMapInfo*> GetControlAndChannelInfo(UControlRig* ControlRig, UMovieSceneControlRigParameterSection* Section, FName ControlName);
 	// retrieve the number of float channels based on the control type.
diff --git a/Engine/Plugins/Animation/ControlRig/Source/ControlRigEditor/Public/Tools/BakingHelper.h b/Engine/Plugins/Animation/ControlRig/Source/ControlRigEditor/Public/Tools/BakingHelper.h
--- a/Engine/Plugins/Animation/ControlRig/Source/ControlRigEditor/Public/Tools/BakingHelper.h
+++ b/Engine/Plugins/Animation/ControlRig/Source/ControlRigEditor/Public/Tools/BakingHelper.h
@@ -6,11 +6,8 @@
 #include "Templates/SharedPointer.h"
 #include "Math/Transform.h"
 
-class UMovieScene3DTransformSection;
-class AActor;
+
 class ISequencer;
-struct FGuid;
-class UMovieScene;
 class UControlRig;
 struct FFrameRate;
 enum class EMovieSceneTransformChannel : uint32;
@@ -23,34 +20,4 @@ struct FBakingHelper
 {
 	/** Returns the current sequencer. */
 	static TWeakPtr<ISequencer> GetSequencer();
-	
-	/** Returns the frame numbers between start and end. */
-	static void CalculateFramesBetween(
-		const UMovieScene* MovieScene,
-		FFrameNumber StartFrame,
-		FFrameNumber EndFrame,
-		TArray<FFrameNumber>& OutFrames);
-	
-	/** Returns the transform section for that guid. */
-	static UMovieScene3DTransformSection* GetTransformSection(
-		const ISequencer* InSequencer,
-		const FGuid& InGuid,
-		const FTransform& InDefaultTransform = FTransform::Identity);
-
-	/** Adds transform keys to the section based on the channels filters. */
-	static bool AddTransformKeys(
-		const UMovieScene3DTransformSection* InTransformSection,
-		const TArray<FFrameNumber>& Frames,
-		const TArray<FTransform>& InLocalTransforms,
-		const EMovieSceneTransformChannel& InChannels);
-
-	/** Adds transform keys to control based on the channels filters. */
-	static bool AddTransformKeys(
-		UControlRig* InControlRig,
-		const FName& InControlName,
-		const TArray<FFrameNumber>& InFrames,
-		const TArray<FTransform>& InTransforms,
-		const EMovieSceneTransformChannel& InChannels,
-		const FFrameRate& InTickResolution,
-		const bool bLocal = true);
 };
diff --git a/Engine/Plugins/Animation/ControlRig/Source/ControlRigEditor/Public/Tools/ConstraintBaker.h b/Engine/Plugins/Animation/ControlRig/Source/ControlRigEditor/Public/Tools/ConstraintBaker.h
--- a/Engine/Plugins/Animation/ControlRig/Source/ControlRigEditor/Public/Tools/ConstraintBaker.h
+++ b/Engine/Plugins/Animation/ControlRig/Source/ControlRigEditor/Public/Tools/ConstraintBaker.h
@@ -25,9 +25,6 @@ public:
 		const bool bLocal,
 		TArray<FTransform>& OutTransforms);
 	
-	/** Returns the channels to key based on the constraint's type. */
-	static EMovieSceneTransformChannel GetChannelsToKey(const UTickableTransformConstraint* InConstraint);
-
 	/** Add InTransforms keys at InFrames into the InHandle transform animation channels. */
 	static void AddTransformKeys(
 		const TSharedPtr<ISequencer>& InSequencer,
diff --git a/Engine/Source/Editor/MovieSceneTools/Private/MovieSceneConstraintChannelHelper.cpp b/Engine/Source/Editor/MovieSceneTools/Private/MovieSceneConstraintChannelHelper.cpp
new file mode 100644
--- /dev/null
+++ b/Engine/Source/Editor/MovieSceneTools/Private/MovieSceneConstraintChannelHelper.cpp
@@ -0,0 +1,561 @@
+﻿#include "MovieSceneConstraintChannelHelper.h"
+#include "ConstraintChannelHelper.inl"
+
+#include "ISequencer.h"
+
+#include "TransformableHandle.h"
+#include "LevelSequence.h"
+#include "MovieSceneToolHelpers.h"
+
+#include "TransformConstraint.h"
+#include "Algo/Copy.h"
+
+//#include "Tools/BakingHelper.h"
+//#include "Tools/ConstraintBaker.h"
+#include "ScopedTransaction.h"
+#include "Channels/MovieSceneFloatChannel.h"
+#include "Channels/MovieSceneDoubleChannel.h"
+#include "LevelEditorViewport.h"
+#include "Sections/MovieSceneConstrainedSection.h"
+
+
+/*
+*
+*  FCompensationEvaluator
+*
+*/
+
+FCompensationEvaluator::FCompensationEvaluator(UTickableTransformConstraint* InConstraint)
+	: Constraint(InConstraint)
+	, Handle(InConstraint ? InConstraint->ChildTRSHandle : nullptr)
+{}
+void FCompensationEvaluator::ComputeLocalTransforms(
+	UWorld* InWorld, const TSharedPtr<ISequencer>& InSequencer,
+	const TArray<FFrameNumber>& InFrames, const bool bToActive)
+{
+	if (InFrames.IsEmpty())
+	{
+		return;
+	}
+
+	const TArray<UTickableTransformConstraint*> Constraints = GetHandleTransformConstraints(InWorld);
+	if (Constraints.IsEmpty())
+	{
+		return;
+	}
+
+	// find last active constraint in the list that is different than the on we want to compensate for
+	auto GetLastActiveConstraint = [this, Constraints]()
+	{
+		// find last active constraint in the list that is different than the on we want to compensate for
+		const int32 LastActiveIndex = Constraints.FindLastByPredicate([this](const UTickableTransformConstraint* InConstraint)
+			{
+				return InConstraint->Active && InConstraint->bDynamicOffset && InConstraint != Constraint;
+			});
+
+		// if found, return its parent global transform
+		return LastActiveIndex > INDEX_NONE ? Constraints[LastActiveIndex] : nullptr;
+	};
+
+	const UMovieScene* MovieScene = InSequencer->GetFocusedMovieSceneSequence()->GetMovieScene();
+	const FFrameRate TickResolution = MovieScene->GetTickResolution();
+	const EMovieScenePlayerStatus::Type PlaybackStatus = InSequencer->GetPlaybackStatus();
+
+	const int32 NumFrames = InFrames.Num();
+
+	// resize arrays to num frames + 1 as we also evaluate at InFrames[0]-1
+	ChildLocals.SetNum(NumFrames + 1);
+	ChildGlobals.SetNum(NumFrames + 1);
+	SpaceGlobals.SetNum(NumFrames + 1);
+
+	const ETransformConstraintType ConstraintType = static_cast<ETransformConstraintType>(Constraint->GetType());
+
+	for (int32 Index = 0; Index < NumFrames + 1; ++Index)
+	{
+		const FFrameNumber FrameNumber = (Index == 0) ? InFrames[0] - 1 : InFrames[Index - 1];
+
+		// evaluate animation
+		const FMovieSceneEvaluationRange EvaluationRange = FMovieSceneEvaluationRange(FFrameTime(FrameNumber), TickResolution);
+		const FMovieSceneContext Context = FMovieSceneContext(EvaluationRange, PlaybackStatus).SetHasJumped(true);
+
+		InSequencer->GetEvaluationTemplate().Evaluate(Context, *InSequencer);
+
+		// evaluate constraints
+		for (const UTickableTransformConstraint* InConstraint : Constraints)
+		{
+			InConstraint->Evaluate();
+		}
+
+		// evaluate ControlRig?
+		// ControlRig->Evaluate_AnyThread();
+
+		FTransform& ChildLocal = ChildLocals[Index];
+		FTransform& ChildGlobal = ChildGlobals[Index];
+		FTransform& SpaceGlobal = SpaceGlobals[Index];
+
+		// store child transforms        	
+		ChildLocal = Handle->GetLocalTransform();
+		ChildGlobal = Handle->GetGlobalTransform();
+
+		const UTickableTransformConstraint* LastConstraint = GetLastActiveConstraint();
+		const ETransformConstraintType LastConstraintType = LastConstraint ?
+			static_cast<ETransformConstraintType>(LastConstraint->GetType()) : ETransformConstraintType::Parent;
+
+		// store constraint/parent space global transform
+		if (bToActive)
+		{
+			// if activating the constraint, store last constraint or parent space at T[0]-1
+			// and constraint space for all other times
+			if (Index == 0)
+			{
+				if (LastConstraint)
+				{
+					SpaceGlobal = LastConstraint->GetParentGlobalTransform();
+					ChildLocal = FTransformConstraintUtils::ComputeRelativeTransform(
+						ChildLocal, ChildGlobal, SpaceGlobal, LastConstraintType);
+				}
+			}
+			else
+			{
+				SpaceGlobal = Constraint->GetParentGlobalTransform();
+				ChildLocal = FTransformConstraintUtils::ComputeRelativeTransform(
+					ChildLocal, ChildGlobal, SpaceGlobal, ConstraintType);
+			}
+		}
+		else
+		{
+			// if deactivating the constraint, store constraint space at T[0]-1
+			// and last constraint or parent space for all other times
+			if (Index == 0)
+			{
+				SpaceGlobal = Constraint->GetParentGlobalTransform();
+				ChildLocal = FTransformConstraintUtils::ComputeRelativeTransform(
+					ChildLocal, ChildGlobal, SpaceGlobal, ConstraintType);
+			}
+			else
+			{
+				if (LastConstraint)
+				{
+					SpaceGlobal = LastConstraint->GetParentGlobalTransform();
+					ChildLocal = FTransformConstraintUtils::ComputeRelativeTransform(
+						ChildLocal, ChildGlobal, SpaceGlobal, LastConstraintType);
+				}
+			}
+		}
+	}
+}
+
+void FCompensationEvaluator::ComputeLocalTransformsBeforeDeletion(
+	UWorld* InWorld,
+	const TSharedPtr<ISequencer>& InSequencer,
+	const TArray<FFrameNumber>& InFrames)
+{
+	if (InFrames.IsEmpty())
+	{
+		return;
+	}
+
+	const TArray<UTickableTransformConstraint*> Constraints = GetHandleTransformConstraints(InWorld);
+	// if (Constraints.IsEmpty())
+	// {
+	// 	return;
+	// }
+
+	// find last active constraint in the list that is different than the on we want to compensate for
+	auto GetLastActiveConstraint = [this, Constraints]()
+	{
+		// find last active constraint in the list that is different than the on we want to compensate for
+		const int32 LastActiveIndex = Constraints.FindLastByPredicate([this](const UTickableTransformConstraint* InConstraint)
+			{
+				return InConstraint->Active && InConstraint->bDynamicOffset && InConstraint != Constraint;
+			});
+
+		// if found, return its parent global transform
+		return LastActiveIndex > INDEX_NONE ? Constraints[LastActiveIndex] : nullptr;
+	};
+
+	const UMovieScene* MovieScene = InSequencer->GetFocusedMovieSceneSequence()->GetMovieScene();
+	const FFrameRate TickResolution = MovieScene->GetTickResolution();
+	const EMovieScenePlayerStatus::Type PlaybackStatus = InSequencer->GetPlaybackStatus();
+
+	const int32 NumFrames = InFrames.Num();
+
+	ChildLocals.SetNum(NumFrames);
+	ChildGlobals.SetNum(NumFrames);
+	SpaceGlobals.SetNum(NumFrames);
+
+	for (int32 Index = 0; Index < NumFrames; ++Index)
+	{
+		const FFrameNumber& FrameNumber = InFrames[Index];
+
+		// evaluate animation
+		const FMovieSceneEvaluationRange EvaluationRange = FMovieSceneEvaluationRange(FFrameTime(FrameNumber), TickResolution);
+		const FMovieSceneContext Context = FMovieSceneContext(EvaluationRange, PlaybackStatus).SetHasJumped(true);
+
+		InSequencer->GetEvaluationTemplate().Evaluate(Context, *InSequencer);
+
+		// evaluate constraints
+		for (const UTickableTransformConstraint* InConstraint : Constraints)
+		{
+			InConstraint->Evaluate();
+		}
+
+		// evaluate ControlRig?
+		// ControlRig->Evaluate_AnyThread();
+
+		FTransform& ChildLocal = ChildLocals[Index];
+		FTransform& ChildGlobal = ChildGlobals[Index];
+		FTransform& SpaceGlobal = SpaceGlobals[Index];
+
+		// store child transforms        	
+		ChildLocal = Handle->GetLocalTransform();
+		ChildGlobal = Handle->GetGlobalTransform();
+
+		// store constraint/parent space global transform
+		if (const UTickableTransformConstraint* LastConstraint = GetLastActiveConstraint())
+		{
+			const ETransformConstraintType LastConstraintType =
+				static_cast<ETransformConstraintType>(LastConstraint->GetType());
+
+			SpaceGlobal = LastConstraint->GetParentGlobalTransform();
+			ChildLocal = FTransformConstraintUtils::ComputeRelativeTransform(
+				ChildLocal, ChildGlobal, SpaceGlobal, LastConstraintType);
+		}
+	}
+}
+
+void FCompensationEvaluator::ComputeCompensation(UWorld* InWorld, const TSharedPtr<ISequencer>& InSequencer, const FFrameNumber& InTime)
+{
+	const TArray<UTickableTransformConstraint*> Constraints = GetHandleTransformConstraints(InWorld);
+	if (Constraints.IsEmpty())
+	{
+		return;
+	}
+
+	// find last active constraint in the list that is different than the on we want to compensate for
+	auto GetLastActiveConstraint = [this, Constraints]()
+	{
+		// find last active constraint in the list that is different than the on we want to compensate for
+		const int32 LastActiveIndex = Constraints.FindLastByPredicate([this](const UTickableTransformConstraint* InConstraint)
+			{
+				return InConstraint->Active && InConstraint->bDynamicOffset;
+			});
+
+		// if found, return its parent global transform
+		return LastActiveIndex > INDEX_NONE ? Constraints[LastActiveIndex] : nullptr;
+	};
+
+	auto EvaluateAt = [InSequencer, &Constraints](const FFrameNumber InFrame)
+	{
+		const UMovieScene* MovieScene = InSequencer->GetFocusedMovieSceneSequence()->GetMovieScene();
+		const FFrameRate TickResolution = MovieScene->GetTickResolution();
+		const EMovieScenePlayerStatus::Type PlaybackStatus = InSequencer->GetPlaybackStatus();
+
+		const FMovieSceneEvaluationRange EvaluationRange0 = FMovieSceneEvaluationRange(FFrameTime(InFrame), TickResolution);
+		const FMovieSceneContext Context0 = FMovieSceneContext(EvaluationRange0, PlaybackStatus).SetHasJumped(true);
+
+		InSequencer->GetEvaluationTemplate().Evaluate(Context0, *InSequencer);
+
+		for (const UTickableTransformConstraint* InConstraint : Constraints)
+		{
+			InConstraint->Evaluate();
+		}
+
+		// ControlRig->Evaluate_AnyThread();
+	};
+
+	// allocate
+	ChildLocals.SetNum(1);
+	ChildGlobals.SetNum(1);
+	SpaceGlobals.SetNum(1);
+
+	// evaluate at InTime and store global
+	EvaluateAt(InTime);
+	ChildGlobals[0] = Handle->GetGlobalTransform();
+
+	// evaluate at InTime-1 and store local
+	EvaluateAt(InTime - 1);
+	ChildLocals[0] = Handle->GetLocalTransform();
+
+	// if constraint at T-1 then switch to its space
+	if (const UTickableTransformConstraint* LastConstraint = GetLastActiveConstraint())
+	{
+		SpaceGlobals[0] = LastConstraint->GetParentGlobalTransform();
+
+		const ETransformConstraintType LastConstraintType = static_cast<ETransformConstraintType>(LastConstraint->GetType());
+		ChildLocals[0] = FTransformConstraintUtils::ComputeRelativeTransform(ChildLocals[0],
+			ChildGlobals[0], SpaceGlobals[0], LastConstraintType);
+	}
+	else // switch to parent space
+	{
+		const FTransform ChildLocal = ChildLocals[0];
+		Handle->SetGlobalTransform(ChildGlobals[0]);
+		ChildLocals[0] = Handle->GetLocalTransform();
+		Handle->SetLocalTransform(ChildLocal);
+	}
+}
+
+TArray<UTickableTransformConstraint*> FCompensationEvaluator::GetHandleTransformConstraints(UWorld* InWorld) const
+{
+	TArray<UTickableTransformConstraint*> TransformConstraints;
+	if (Handle)
+	{
+		// get sorted transform constraints
+		const FConstraintsManagerController& Controller = FConstraintsManagerController::Get(InWorld);
+		static constexpr bool bSorted = true;
+		TArray<TObjectPtr<UTickableConstraint>> Constraints = Controller.GetParentConstraints(Handle->GetHash(), bSorted);
+		for (const TObjectPtr<UTickableConstraint>& InConstraint : Constraints)
+		{
+			if (UTickableTransformConstraint* TransformConstraint = Cast<UTickableTransformConstraint>(InConstraint.Get()))
+			{
+				TransformConstraints.Add(TransformConstraint);
+			}
+		}
+	}
+	return TransformConstraints;
+}
+
+bool FMovieSceneConstraintChannelHelper::bDoNotCompensate = false;
+
+
+/*
+*
+*  Constraint Channel Helpers
+*
+*/
+void FMovieSceneConstraintChannelHelper::HandleConstraintRemoved(
+	UTickableConstraint* InConstraint,
+	const FMovieSceneConstraintChannel* InConstraintChannel,
+	const TSharedPtr<ISequencer>& InSequencer,
+	UMovieSceneSection* InSection)
+{
+	UTickableTransformConstraint* Constraint = Cast<UTickableTransformConstraint>(InConstraint);
+	if (!Constraint || !InConstraintChannel || !InSection)
+	{
+		return;
+	}
+
+	TGuardValue<bool> CompensateGuard(bDoNotCompensate, true);
+	if (const UTransformableHandle* ControlHandle =Constraint->ChildTRSHandle)
+	{
+		const TArrayView<const FFrameNumber> Times = InConstraintChannel->GetData().GetTimes();
+		if (Times.IsEmpty())
+		{
+			return;
+		}
+
+		// get transform channels
+		const TArrayView<FMovieSceneFloatChannel*> FloatTransformChannels = ControlHandle->GetFloatChannels(InSection);
+		const TArrayView<FMovieSceneDoubleChannel*> DoubleTransformChannels = ControlHandle->GetDoubleChannels(InSection);
+
+		// get frames after this time
+		TArray<FFrameNumber> FramesToCompensate;
+		if (FloatTransformChannels.Num() > 0)
+		{
+			FMovieSceneConstraintChannelHelper::GetFramesWithinActiveState(*InConstraintChannel, FloatTransformChannels, FramesToCompensate);
+		}
+		else
+		{
+			FMovieSceneConstraintChannelHelper::GetFramesWithinActiveState(*InConstraintChannel, DoubleTransformChannels, FramesToCompensate);
+		}
+		// do the compensation
+		UWorld* World = GCurrentLevelEditingViewportClient ? GCurrentLevelEditingViewportClient->GetWorld() : nullptr;
+
+		FCompensationEvaluator Evaluator(Constraint);
+		Evaluator.ComputeLocalTransformsBeforeDeletion(World, InSequencer, FramesToCompensate);
+		const TArray<FTransform>& ChildLocals = Evaluator.ChildLocals;
+
+		const EMovieSceneTransformChannel ChannelsToKey = Constraint->GetChannelsToKey();
+		const FFrameRate TickResolution = InSequencer->GetFocusedTickResolution();
+
+		ControlHandle->AddTransformKeys(FramesToCompensate,
+			ChildLocals, ChannelsToKey, TickResolution, InSection);
+	}
+}
+
+void FMovieSceneConstraintChannelHelper::HandleConstraintKeyDeleted(
+	UTickableTransformConstraint* InConstraint,
+	const FMovieSceneConstraintChannel* InConstraintChannel,
+	const TSharedPtr<ISequencer>& InSequencer,
+	UMovieSceneSection* InSection,
+	const FFrameNumber& InTime)
+{
+	const FFrameNumber TimeMinusOne(InTime - 1);
+
+	bool CurrentValue = false, PreviousValue = false;
+	InConstraintChannel->Evaluate(TimeMinusOne, PreviousValue);
+	InConstraintChannel->Evaluate(InTime, CurrentValue);
+
+	if (CurrentValue == PreviousValue)
+	{
+		const int32 NumKeys = InConstraintChannel->GetNumKeys();
+		if (NumKeys > 1)
+		{
+			return;
+		}
+	}
+
+	TGuardValue<bool> CompensateGuard(bDoNotCompensate, true);
+
+	if (const UTransformableHandle* ControlHandle = InConstraint->ChildTRSHandle)
+	{
+		// get transform channels
+		const TArrayView<FMovieSceneFloatChannel*> FloatTransformChannels = ControlHandle->GetFloatChannels(InSection);
+		const TArrayView<FMovieSceneDoubleChannel*> DoubleTransformChannels = ControlHandle->GetDoubleChannels(InSection);
+
+		// get frames after this time
+		TArray<FFrameNumber> FramesToCompensate;
+		if (FloatTransformChannels.Num() > 0)
+		{
+			FMovieSceneConstraintChannelHelper::GetFramesWithinActiveState(*InConstraintChannel, FloatTransformChannels, FramesToCompensate);
+		}
+		else
+		{
+			FMovieSceneConstraintChannelHelper::GetFramesWithinActiveState(*InConstraintChannel, DoubleTransformChannels, FramesToCompensate);
+		}
+		// do the compensation
+		FCompensationEvaluator Evaluator(InConstraint);
+		UWorld* World = GCurrentLevelEditingViewportClient ? GCurrentLevelEditingViewportClient->GetWorld() : nullptr;
+		Evaluator.ComputeLocalTransforms(World, InSequencer, FramesToCompensate, PreviousValue);
+		TArray<FTransform>& ChildLocals = Evaluator.ChildLocals;
+		ChildLocals.RemoveAt(0);
+
+		const EMovieSceneTransformChannel ChannelsToKey = InConstraint->GetChannelsToKey();
+		const FFrameRate TickResolution = InSequencer->GetFocusedTickResolution();
+
+		ControlHandle->AddTransformKeys(FramesToCompensate,
+			ChildLocals, ChannelsToKey, TickResolution, InSection);
+
+		// now delete any extra TimeMinusOne
+		if (FloatTransformChannels.Num() > 0)
+		{
+			FMovieSceneConstraintChannelHelper::DeleteTransformKeys(FloatTransformChannels, TimeMinusOne);
+		}
+		else
+		{
+			FMovieSceneConstraintChannelHelper::DeleteTransformKeys(DoubleTransformChannels, TimeMinusOne);
+		}
+	}
+}
+
+void FMovieSceneConstraintChannelHelper::HandleConstraintKeyMoved(
+	const UTickableTransformConstraint* InConstraint,
+	const FMovieSceneConstraintChannel* InConstraintChannel,
+	UMovieSceneSection* InSection,
+	const FFrameNumber& InCurrentFrame, const FFrameNumber& InNextFrame)
+{
+	const FFrameNumber Delta = InNextFrame - InCurrentFrame;
+	if (Delta == 0)
+	{
+		return;
+	}
+
+	if (!InConstraint || !InConstraintChannel || !InSection)
+	{
+		return;
+	}
+
+	if (const UTransformableHandle* ControlHandle = InConstraint->ChildTRSHandle)
+	{
+
+		// get transform channels
+		const TArrayView<FMovieSceneFloatChannel*> FloatTransformChannels = ControlHandle->GetFloatChannels(InSection);
+		const TArrayView<FMovieSceneDoubleChannel*> DoubleTransformChannels = ControlHandle->GetDoubleChannels(InSection);
+
+		// move them
+		if (FloatTransformChannels.Num() > 0)
+		{
+			FMovieSceneConstraintChannelHelper::MoveTransformKeys(FloatTransformChannels, InCurrentFrame, InNextFrame);
+		}
+		else
+		{
+			FMovieSceneConstraintChannelHelper::MoveTransformKeys(DoubleTransformChannels, InCurrentFrame, InNextFrame);
+		}
+	}
+
+}
+
+void FMovieSceneConstraintChannelHelper::CompensateIfNeeded(
+	const TSharedPtr<ISequencer>& InSequencer,
+	UMovieSceneSection* InSection,
+	const TOptional<FFrameNumber>& OptionalTime)
+{
+	if (bDoNotCompensate)
+	{
+		return;
+	}
+	IMovieSceneConstrainedSection* ConstrainedSection = Cast<IMovieSceneConstrainedSection>(InSection);
+
+	TGuardValue<bool> CompensateGuard(bDoNotCompensate, true);
+
+	// Frames to compensate
+	TArray<FFrameNumber> OptionalTimeArray;
+	if (OptionalTime.IsSet())
+	{
+		OptionalTimeArray.Add(OptionalTime.GetValue());
+	}
+
+	auto GetSpaceTimesToCompensate = [&OptionalTimeArray](const FConstraintAndActiveChannel& Channel)->TArrayView<const FFrameNumber>
+	{
+		if (OptionalTimeArray.IsEmpty())
+		{
+			return Channel.ActiveChannel.GetData().GetTimes();
+		}
+		return OptionalTimeArray;
+	};
+
+	const FFrameRate TickResolution = InSequencer->GetFocusedTickResolution();
+
+	bool bNeedsEvaluation = false;
+
+	// gather all transform constraints
+	TArray<FConstraintAndActiveChannel> TransformConstraintsChannels;
+	Algo::CopyIf(ConstrainedSection->GetConstraintsChannels(), TransformConstraintsChannels,
+		[](const FConstraintAndActiveChannel& InChannel)
+		{
+			return InChannel.Constraint.IsValid() && InChannel.Constraint->IsA<UTickableTransformConstraint>();
+		}
+	);
+
+	// compensate constraints
+	for (const FConstraintAndActiveChannel& Channel : TransformConstraintsChannels)
+	{
+		const TArrayView<const FFrameNumber> FramesToCompensate = GetSpaceTimesToCompensate(Channel);
+		for (const FFrameNumber& Time : FramesToCompensate)
+		{
+			const FFrameNumber TimeMinusOne(Time - 1);
+
+			bool CurrentValue = false, PreviousValue = false;
+			Channel.ActiveChannel.Evaluate(TimeMinusOne, PreviousValue);
+			Channel.ActiveChannel.Evaluate(Time, CurrentValue);
+
+			if (CurrentValue != PreviousValue) //if they are the same no need to do anything
+			{
+				UTickableTransformConstraint* Constraint = Cast<UTickableTransformConstraint>(Channel.Constraint.Get());
+
+				UWorld* World = GCurrentLevelEditingViewportClient ? GCurrentLevelEditingViewportClient->GetWorld() : nullptr;
+
+				// compute transform to set
+				// if switching from active to inactive then we must add a key at T-1 in the constraint space
+				// if switching from inactive to active then we must add a key at T-1 in the previous constraint or parent space
+				FCompensationEvaluator Evaluator(Constraint);
+				Evaluator.ComputeCompensation(World, InSequencer, Time);
+				const TArray<FTransform>& LocalTransforms = Evaluator.ChildLocals;
+
+				const EMovieSceneTransformChannel ChannelsToKey = Constraint->GetChannelsToKey();
+				if (const UTransformableHandle* ControlHandle = Constraint->ChildTRSHandle)
+				{
+					ControlHandle->AddTransformKeys({ TimeMinusOne }, LocalTransforms, ChannelsToKey, TickResolution, InSection);
+					bNeedsEvaluation = true;
+				}
+			}
+		}
+	}
+
+	if (bNeedsEvaluation)
+	{
+		InSequencer->ForceEvaluate();
+	}
+}
+
+
diff --git a/Engine/Source/Editor/MovieSceneTools/Private/MovieSceneToolHelpers.cpp b/Engine/Source/Editor/MovieSceneTools/Private/MovieSceneToolHelpers.cpp
--- a/Engine/Source/Editor/MovieSceneTools/Private/MovieSceneToolHelpers.cpp
+++ b/Engine/Source/Editor/MovieSceneTools/Private/MovieSceneToolHelpers.cpp
@@ -4540,4 +4540,197 @@ void MovieSceneToolHelpers::GetActorWorldTransforms(IMovieScenePlayer* Player, U
 	}
 }
 
+void MovieSceneToolHelpers::CalculateFramesBetween(
+	const UMovieScene* MovieScene,
+	FFrameNumber StartFrame,
+	FFrameNumber EndFrame,
+	TArray<FFrameNumber>& OutFrames)
+{
+	if (StartFrame > EndFrame)
+	{
+		Swap(StartFrame, EndFrame);
+	}
+	//if the end frame is the last frame, we move it back one tick otherwise it's possible when evaluating it won't evaluate.
+	if (EndFrame == MovieScene->GetPlaybackRange().GetUpperBoundValue())
+	{
+		--EndFrame;
+	}
+	const FFrameRate TickResolution = MovieScene->GetTickResolution();
+	const FFrameRate DisplayResolution = MovieScene->GetDisplayRate();
+
+	const FFrameNumber StartTimeInDisplay = FFrameRate::TransformTime(FFrameTime(StartFrame), TickResolution, DisplayResolution).FloorToFrame();
+	const FFrameNumber EndTimeInDisplay = FFrameRate::TransformTime(FFrameTime(EndFrame), TickResolution, DisplayResolution).CeilToFrame();
+
+	OutFrames.Reserve(EndTimeInDisplay.Value - StartTimeInDisplay.Value + 1);
+	for (FFrameNumber DisplayFrameNumber = StartTimeInDisplay; DisplayFrameNumber <= EndTimeInDisplay; ++DisplayFrameNumber)
+	{
+		FFrameNumber TickFrameNumber = FFrameRate::TransformTime(FFrameTime(DisplayFrameNumber), DisplayResolution, TickResolution).FrameNumber;
+		OutFrames.Add(TickFrameNumber);
+	}
+}
+
+UMovieScene3DTransformSection* MovieSceneToolHelpers::GetTransformSection(
+	const ISequencer* InSequencer,
+	const FGuid& InGuid,
+	const FTransform& InDefaultTransform)
+{
+	if (!InSequencer || !InSequencer->GetFocusedMovieSceneSequence())
+	{
+		return nullptr;
+	}
+
+	UMovieScene* MovieScene = InSequencer->GetFocusedMovieSceneSequence()->GetMovieScene();
+	if (!MovieScene)
+	{
+		return nullptr;
+	}
+
+	UMovieScene3DTransformTrack* TransformTrack = MovieScene->FindTrack<UMovieScene3DTransformTrack>(InGuid);
+	if (!TransformTrack)
+	{
+		MovieScene->Modify();
+		TransformTrack = MovieScene->AddTrack<UMovieScene3DTransformTrack>(InGuid);
+	}
+	TransformTrack->Modify();
+
+	static constexpr FFrameNumber Frame0;
+	bool bSectionAdded = false;
+	UMovieScene3DTransformSection* TransformSection =
+		Cast<UMovieScene3DTransformSection>(TransformTrack->FindOrAddSection(Frame0, bSectionAdded));
+	if (!TransformSection)
+	{
+		return nullptr;
+	}
+
+	TransformSection->Modify();
+	if (bSectionAdded)
+	{
+		TransformSection->SetRange(TRange<FFrameNumber>::All());
+
+		const FVector Location0 = InDefaultTransform.GetLocation();
+		const FRotator Rotation0 = InDefaultTransform.GetRotation().Rotator();
+		const FVector Scale3D0 = InDefaultTransform.GetScale3D();
+
+		const TArrayView<FMovieSceneDoubleChannel*> Channels =
+			TransformSection->GetChannelProxy().GetChannels<FMovieSceneDoubleChannel>();
+		Channels[0]->SetDefault(Location0.X);
+		Channels[1]->SetDefault(Location0.Y);
+		Channels[2]->SetDefault(Location0.Z);
+		Channels[3]->SetDefault(Rotation0.Roll);
+		Channels[4]->SetDefault(Rotation0.Pitch);
+		Channels[5]->SetDefault(Rotation0.Yaw);
+		Channels[6]->SetDefault(Scale3D0.X);
+		Channels[7]->SetDefault(Scale3D0.Y);
+		Channels[8]->SetDefault(Scale3D0.Z);
+	}
+
+	return TransformSection;
+}
+
+bool MovieSceneToolHelpers::AddTransformKeys(
+	const UMovieScene3DTransformSection* InTransformSection,
+	const TArray<FFrameNumber>& Frames,
+	const TArray<FTransform>& InLocalTransforms,
+	const EMovieSceneTransformChannel& InChannels)
+{
+	if (!InTransformSection)
+	{
+		return false;
+	}
+
+	if (Frames.IsEmpty() || Frames.Num() != InLocalTransforms.Num())
+	{
+		return false;
+	}
+
+	auto GetValue = [](const uint32 Index, const FVector& InLocation, const FRotator& InRotation, const FVector& InScale)
+	{
+		switch (Index)
+		{
+		case 0:
+			return InLocation.X;
+		case 1:
+			return InLocation.Y;
+		case 2:
+			return InLocation.Z;
+		case 3:
+			return InRotation.Roll;
+		case 4:
+			return InRotation.Pitch;
+		case 5:
+			return InRotation.Yaw;
+		case 6:
+			return InScale.X;
+		case 7:
+			return InScale.Y;
+		case 8:
+			return InScale.Z;
+		default:
+			ensure(false);
+			break;
+		}
+		return 0.0;
+	};
+
+	const bool bKeyTranslation = EnumHasAllFlags(InChannels, EMovieSceneTransformChannel::Translation);
+	const bool bKeyRotation = EnumHasAllFlags(InChannels, EMovieSceneTransformChannel::Rotation);
+	const bool bKeyScale = EnumHasAllFlags(InChannels, EMovieSceneTransformChannel::Scale);
+
+	TArray<uint32> ChannelsIndexToKey;
+	if (bKeyTranslation)
+	{
+		ChannelsIndexToKey.Append({ 0,1,2 });
+	}
+	if (bKeyRotation)
+	{
+		ChannelsIndexToKey.Append({ 3,4,5 });
+	}
+	if (bKeyScale)
+	{
+		ChannelsIndexToKey.Append({ 6,7,8 });
+	}
+
+	const TArrayView<FMovieSceneDoubleChannel*> Channels =
+		InTransformSection->GetChannelProxy().GetChannels<FMovieSceneDoubleChannel>();
+
+	// set default
+	const FTransform& LocalTransform0 = InLocalTransforms[0];
+	const FVector Location0 = LocalTransform0.GetLocation();
+	const FRotator Rotation0 = LocalTransform0.GetRotation().Rotator();
+	const FVector Scale3D0 = LocalTransform0.GetScale3D();
+
+	for (int32 ChannelIndex = 0; ChannelIndex < 9; ChannelIndex++)
+	{
+		if (!Channels[ChannelIndex]->GetDefault().IsSet())
+		{
+			const double Value = GetValue(ChannelIndex, Location0, Rotation0, Scale3D0);
+			Channels[ChannelIndex]->SetDefault(Value);
+		}
+	}
+
+	// add keys
+	for (int32 Index = 0; Index < Frames.Num(); ++Index)
+	{
+		const FFrameNumber& Frame = Frames[Index];
+		const FTransform& LocalTransform = InLocalTransforms[Index];
+
+		const FVector Location = LocalTransform.GetLocation();
+		const FRotator Rotation = LocalTransform.GetRotation().Rotator();
+		const FVector Scale3D = LocalTransform.GetScale3D();
+
+		for (const int32 ChannelIndex : ChannelsIndexToKey)
+		{
+			const double Value = GetValue(ChannelIndex, Location, Rotation, Scale3D);
+			TMovieSceneChannelData<FMovieSceneDoubleValue> ChannelData = Channels[ChannelIndex]->GetData();
+			MovieSceneToolHelpers::SetOrAddKey(ChannelData, Frame, Value);
+		}
+	}
 
+	//now we need to set auto tangents
+	for (const int32 ChannelIndex : ChannelsIndexToKey)
+	{
+		Channels[ChannelIndex]->AutoSetTangents();
+	}
+
+	return true;
+}
\ No newline at end of file
diff --git a/Engine/Source/Editor/MovieSceneTools/Private/TrackEditors/TransformTrackEditor.cpp b/Engine/Source/Editor/MovieSceneTools/Private/TrackEditors/TransformTrackEditor.cpp
--- a/Engine/Source/Editor/MovieSceneTools/Private/TrackEditors/TransformTrackEditor.cpp
+++ b/Engine/Source/Editor/MovieSceneTools/Private/TrackEditors/TransformTrackEditor.cpp
@@ -34,6 +34,13 @@
 #include "Tracks/IMovieSceneTransformOrigin.h"
 #include "IMovieScenePlaybackClient.h"
 
+#include "Editor.h"
+#include "LevelEditorViewport.h"
+#include "TransformConstraint.h"
+#include "MovieSceneConstraintChannelHelper.h"
+#include "Misc/TransactionObjectEvent.h"
+#include "Engine/Selection.h"
+
 #define LOCTEXT_NAMESPACE "MovieScene_TransformTrack"
 
 void GetActorAndSceneComponentFromObject( UObject* Object, AActor*& OutActor, USceneComponent*& OutSceneComponent )
@@ -72,9 +79,18 @@ F3DTransformTrackEditor::~F3DTransformTrackEditor()
 
 void F3DTransformTrackEditor::OnRelease()
 {
+	ClearOutConstraintDelegates();
 	FCoreUObjectDelegates::OnPreObjectPropertyChanged.RemoveAll(this);
 	FCoreUObjectDelegates::OnObjectPropertyChanged.RemoveAll(this);
 
+	if (OnSceneComponentConstrainedHandle.IsValid())
+	{
+		UWorld* World = GCurrentLevelEditingViewportClient ? GCurrentLevelEditingViewportClient->GetWorld() : nullptr;
+		FConstraintsManagerController& Controller = FConstraintsManagerController::Get(World);
+		Controller.OnSceneComponentConstrained().Remove(OnSceneComponentConstrainedHandle);
+		OnSceneComponentConstrainedHandle.Reset();
+	}
+
 	for(FLevelEditorViewportClient* LevelVC : GEditor->GetLevelViewportClients())
 	{
 		if (LevelVC && LevelVC->IsPerspective() && LevelVC->GetViewMode() != VMI_Unknown)
@@ -140,6 +156,11 @@ void F3DTransformTrackEditor::BuildTrackContextMenu( FMenuBuilder& MenuBuilder,
 TSharedRef<ISequencerSection> F3DTransformTrackEditor::MakeSectionInterface(UMovieSceneSection& SectionObject, UMovieSceneTrack& Track, FGuid ObjectBinding)
 {
 	check(SupportsType(SectionObject.GetOuter()->GetClass()));
+	UMovieScene3DTransformSection* Section = Cast<UMovieScene3DTransformSection>(&SectionObject);
+	if (Section)
+	{
+		Section->ConstraintChannelAdded().AddRaw(this, &F3DTransformTrackEditor::HandleOnConstraintAdded);
+	}
 	return MakeShared<FTransformSection>(SectionObject, GetSequencer());
 }
 
@@ -830,7 +851,20 @@ void F3DTransformTrackEditor::AddTransformKeys( UObject* ObjectToKey, const TOpt
 	};
 	auto OnKeyProperty = [=](FFrameNumber Time) -> FKeyPropertyResult
 	{
-		return this->AddKeysToObjects(MakeArrayView(&ObjectToKey, 1), Time,  KeyMode, UMovieScene3DTransformTrack::StaticClass(), TransformPropertyName, InitializeNewTrack, GenerateKeys);
+		FKeyPropertyResult KeyPropertyResult = this->AddKeysToObjects(MakeArrayView(&ObjectToKey, 1), Time,  KeyMode, UMovieScene3DTransformTrack::StaticClass(), TransformPropertyName, InitializeNewTrack, GenerateKeys);
+		if (KeyPropertyResult.SectionsKeyed.Num() > 0)
+		{
+			for (TWeakObjectPtr<UMovieSceneSection>& WeakSection : KeyPropertyResult.SectionsKeyed)
+			{
+				
+				if (UMovieScene3DTransformSection* Section = Cast< UMovieScene3DTransformSection>(WeakSection.Get()))
+				{
+					FMovieSceneConstraintChannelHelper::CompensateIfNeeded(GetSequencer(), Section, Time);
+				}
+			}
+		}
+		
+		return KeyPropertyResult;
 	};
 
 	AnimatablePropertyChanged( FOnKeyProperty::CreateLambda(OnKeyProperty) );
@@ -1210,4 +1244,239 @@ void F3DTransformTrackEditor::ImportAnimSequenceTransformsEnterPressed(const TAr
 	}
 }
 
+bool F3DTransformTrackEditor::MatchesContext(const FTransactionContext& InContext, const TArray<TPair<UObject*, FTransactionObjectEvent>>& TransactionObjects) const
+{
+	SectionsGettingUndone.SetNum(0);
+	// Check if we care about the undo/redo
+	bool bGettingUndone = false;
+	for (const TPair<UObject*, FTransactionObjectEvent>& TransactionObjectPair : TransactionObjects)
+	{
+
+		UObject* Object = TransactionObjectPair.Key;
+		while (Object != nullptr)
+		{
+			const UClass* ObjectClass = Object->GetClass();
+			if (ObjectClass && ObjectClass->IsChildOf(UMovieScene3DTransformSection::StaticClass()))
+			{
+				UMovieScene3DTransformSection* Section = Cast< UMovieScene3DTransformSection>(Object);
+				if (Section)
+				{
+					SectionsGettingUndone.Add(Section);
+				}
+				bGettingUndone = true;
+				break;
+			}
+			Object = Object->GetOuter();
+		}
+	}
+
+	return bGettingUndone;
+}
+
+void F3DTransformTrackEditor::PostUndo(bool bSuccess)
+{
+	for (UMovieScene3DTransformSection* Section : SectionsGettingUndone)
+	{
+		if (Section)
+		{
+			TArray<FConstraintAndActiveChannel>& ConstraintChannels = Section->GetConstraintsChannels();
+			for (FConstraintAndActiveChannel& Channel : ConstraintChannels)
+			{
+				HandleOnConstraintAdded(Section, &(Channel.ActiveChannel));
+			}
+		}
+	}
+}
+
+
+void F3DTransformTrackEditor::HandleOnConstraintAdded(IMovieSceneConstrainedSection* InSection, FMovieSceneConstraintChannel* InConstraintChannel)
+{
+	if (!InConstraintChannel)
+	{
+		return;
+	}
+	// handle scene component changes so we can update the gizmo location when it changes
+	UWorld* World = GCurrentLevelEditingViewportClient ? GCurrentLevelEditingViewportClient->GetWorld() : nullptr;
+	FConstraintsManagerController& Controller = FConstraintsManagerController::Get(World);
+	if (!Controller.OnSceneComponentConstrained().IsBoundToObject(this))
+	{
+		OnSceneComponentConstrainedHandle =
+			Controller.OnSceneComponentConstrained().AddLambda([](USceneComponent* InSceneComponent)
+				{
+					//only update gizmo if it or it's actor is seleced, otherwise can be wasteful
+					GUnrealEd->UpdatePivotLocationForSelection();
+					AActor* Actor = InSceneComponent->GetTypedOuter<AActor>();
+					if (USelection* SelectedActors = GEditor->GetSelectedActors())
+					{
+						if (SelectedActors->IsSelected(Actor))
+						{
+							GUnrealEd->UpdatePivotLocationForSelection();
+							return;
+						}
+					}
+					else if (USelection* SelectedComponents = GEditor->GetSelectedComponents())
+					{
+						if (SelectedComponents->IsSelected(InSceneComponent))
+						{
+							GUnrealEd->UpdatePivotLocationForSelection();
+						}
+					}
+				});
+	}
+	// handle key moved
+	InConstraintChannel->OnKeyMovedEvent().AddLambda([this, InSection](
+		FMovieSceneChannel* InChannel, const TArray<FKeyMoveEventItem>& InMovedItems)
+		{
+			const FMovieSceneConstraintChannel* ConstraintChannel = static_cast<FMovieSceneConstraintChannel*>(InChannel);
+			HandleConstraintKeyMoved(InSection, ConstraintChannel, InMovedItems);
+		});
+
+	// handle key deleted
+	InConstraintChannel->OnKeyDeletedEvent().AddLambda([this, InSection](
+		FMovieSceneChannel* InChannel, const TArray<FKeyAddOrDeleteEventItem>& InDeletedItems)
+		{
+			const FMovieSceneConstraintChannel* ConstraintChannel = static_cast<FMovieSceneConstraintChannel*>(InChannel);
+			HandleConstraintKeyDeleted(InSection, ConstraintChannel, InDeletedItems);
+		});
+
+	// handle constraint deleted
+	if (InSection)
+	{
+		HandleConstraintRemoved(InSection);
+	}
+}
+
+static UTickableTransformConstraint* GetTickableTransformConstraint(IMovieSceneConstrainedSection* InSection, const FMovieSceneConstraintChannel* InConstraintChannel)
+{
+	UTickableTransformConstraint* Constraint = nullptr;
+	// get constraint channel
+	TArray<FConstraintAndActiveChannel>& ConstraintChannels = InSection->GetConstraintsChannels();
+	const int32 Index = ConstraintChannels.IndexOfByPredicate([InConstraintChannel](const FConstraintAndActiveChannel& InChannel)
+		{
+			return &(InChannel.ActiveChannel) == InConstraintChannel;
+		});
+
+	if (Index != INDEX_NONE)
+	{
+		Constraint = Cast<UTickableTransformConstraint>(ConstraintChannels[Index].Constraint.Get());
+	}
+	return Constraint;
+}
+
+void F3DTransformTrackEditor::HandleConstraintKeyDeleted(IMovieSceneConstrainedSection* InSection, const FMovieSceneConstraintChannel* InConstraintChannel,
+	const TArray<FKeyAddOrDeleteEventItem>& InDeletedItems) const
+{
+	if (!InConstraintChannel)
+	{
+		return;
+	}
+	UTickableTransformConstraint* Constraint = GetTickableTransformConstraint(InSection, InConstraintChannel);
+	if (!Constraint)
+	{
+		return;
+	}
+    if(UTickableTransformConstraint* Constrain = GetTickableTransformConstraint(InSection,InConstraintChannel))
+	{
+		for (const FKeyAddOrDeleteEventItem& EventItem : InDeletedItems)
+		{
+			UMovieSceneSection* Section = Cast<UMovieSceneSection>(InSection);
+			FMovieSceneConstraintChannelHelper::HandleConstraintKeyDeleted(
+				Constraint, InConstraintChannel,
+				GetSequencer(), Section,
+				EventItem.Frame);
+			
+		}
+	}
+}
+void F3DTransformTrackEditor::HandleConstraintKeyMoved(IMovieSceneConstrainedSection* InSection, const FMovieSceneConstraintChannel* InConstraintChannel,
+	const TArray<FKeyMoveEventItem>& InMovedItems)
+{
+
+	if (!InConstraintChannel)
+	{
+		return;
+	}
+	UTickableTransformConstraint* Constraint = GetTickableTransformConstraint(InSection, InConstraintChannel);
+	if (!Constraint)
+	{
+		return;
+	}
+	UMovieSceneSection* Section = Cast<UMovieSceneSection>(InSection);
+
+	for (const FKeyMoveEventItem& MoveEventItem : InMovedItems)
+	{
+		FMovieSceneConstraintChannelHelper::HandleConstraintKeyMoved(
+			Constraint, InConstraintChannel, Section,
+			MoveEventItem.Frame, MoveEventItem.NewFrame);
+	}
+}
+
+void F3DTransformTrackEditor::HandleConstraintRemoved(IMovieSceneConstrainedSection* InSection) const
+{
+	UWorld* World = GCurrentLevelEditingViewportClient ? GCurrentLevelEditingViewportClient->GetWorld() : nullptr;
+	UMovieSceneSection* Section = Cast<UMovieSceneSection>(InSection);
+
+	FConstraintsManagerController& Controller = FConstraintsManagerController::Get(World);
+	if (!Controller.OnConstraintRemoved().IsBoundToObject(InSection))
+	{
+		InSection->OnConstraintRemovedHandle =
+			Controller.OnConstraintRemoved().AddLambda([InSection,Section, this](FName InConstraintName)
+				{
+					const FConstraintAndActiveChannel* ConstraintChannel = InSection->GetConstraintChannel(InConstraintName);
+					if (!ConstraintChannel)
+					{
+						return;
+					}
+
+					if (ConstraintChannel->Constraint.IsValid())
+					{
+						FMovieSceneConstraintChannelHelper::HandleConstraintRemoved(
+							ConstraintChannel->Constraint.Get(),
+							&ConstraintChannel->ActiveChannel,
+							GetSequencer(),
+							Section);
+					}
+
+					InSection->RemoveConstraintChannel(InConstraintName);
+				});
+	}
+	
+}
+
+
+void F3DTransformTrackEditor::ClearOutConstraintDelegates() const
+{
+	if (GetSequencer().IsValid())
+	{
+		const UMovieScene* MovieScene = GetSequencer()->GetFocusedMovieSceneSequence()->GetMovieScene();
+		const TArray<FMovieSceneBinding>& Bindings = MovieScene->GetBindings();
+		for (const FMovieSceneBinding& Binding : Bindings)
+		{
+			const UMovieSceneTrack* Track = MovieScene->FindTrack(
+				UMovieScene3DTransformTrack::StaticClass(), Binding.GetObjectGuid(), NAME_None);
+			if (const UMovieScene3DTransformTrack* CRTrack = Cast<UMovieScene3DTransformTrack>(Track))
+			{
+
+				for (UMovieSceneSection* Section : Track->GetAllSections())
+				{
+					if (UMovieScene3DTransformSection* CRSection = Cast<UMovieScene3DTransformSection>(Section))
+					{
+						// clear constraint channels
+						TArray<FConstraintAndActiveChannel>& ConstraintChannels = CRSection->GetConstraintsChannels();
+						for (FConstraintAndActiveChannel& Channel : ConstraintChannels)
+						{
+							Channel.ActiveChannel.OnKeyMovedEvent().Clear();
+							Channel.ActiveChannel.OnKeyDeletedEvent().Clear();
+						}
+					}
+				}
+			}
+		}
+	}
+}
+
+
+
+	
+
 #undef LOCTEXT_NAMESPACE
diff --git a/Engine/Source/Editor/MovieSceneTools/Private/TrackEditors/TransformTrackEditor.h b/Engine/Source/Editor/MovieSceneTools/Private/TrackEditors/TransformTrackEditor.h
--- a/Engine/Source/Editor/MovieSceneTools/Private/TrackEditors/TransformTrackEditor.h
+++ b/Engine/Source/Editor/MovieSceneTools/Private/TrackEditors/TransformTrackEditor.h
@@ -15,6 +15,7 @@
 #include "Tracks/MovieScene3DTransformTrack.h"
 #include "Sections/MovieScene3DTransformSection.h"
 #include "SequencerKeyParams.h"
+#include "EditorUndoClient.h"
 
 class AActor;
 struct FAssetData;
@@ -26,7 +27,7 @@ namespace UE { namespace MovieScene { struct FIntermediate3DTransform; } }
  * Tools for animatable transforms
  */
 class F3DTransformTrackEditor
-	: public FKeyframeTrackEditor<UMovieScene3DTransformTrack>
+	: public FKeyframeTrackEditor<UMovieScene3DTransformTrack>, public FEditorUndoClient
 {
 public:
 
@@ -64,6 +65,11 @@ public:
 	virtual void OnPreSaveWorld(UWorld* World) override;
 	virtual void OnPostSaveWorld(UWorld* World) override;
 
+	//FEditorUndoClient Interface
+	virtual bool MatchesContext(const FTransactionContext& InContext, const TArray<TPair<UObject*, FTransactionObjectEvent>>& TransactionObjects) const override;
+	virtual void PostUndo(bool bSuccess) override;
+	virtual void PostRedo(bool bSuccess) override { PostUndo(bSuccess); }
+
 private:
 
 	void ProcessKeyOperation(UObject* ObjectToKey, TArrayView<const UE::Sequencer::FKeySectionOperation> SectionsToKey, ISequencer& InSequencer, FFrameNumber KeyTime);
@@ -173,6 +179,15 @@ private:
 	/** Import an animation sequence's root transforms into a transform section */
 	static void ImportAnimSequenceTransformsEnterPressed(const TArray<FAssetData>& Asset, TSharedRef<class ISequencer> Sequencer, UMovieScene3DTransformTrack* TransformTrack);
 
+	/** ConstraintChannel Delegates*/
+	FDelegateHandle OnSceneComponentConstrainedHandle;
+	void HandleOnConstraintAdded(IMovieSceneConstrainedSection* InSection, FMovieSceneConstraintChannel* InConstraintChannel);
+	void HandleConstraintKeyDeleted(IMovieSceneConstrainedSection* InSection, const FMovieSceneConstraintChannel* InConstraintChannel,
+		const TArray<FKeyAddOrDeleteEventItem>& InDeletedItems) const;
+	void HandleConstraintKeyMoved(IMovieSceneConstrainedSection* InSection, const FMovieSceneConstraintChannel* InConstraintChannel,
+		const TArray<FKeyMoveEventItem>& InMovedItems);
+	void HandleConstraintRemoved(IMovieSceneConstrainedSection* InSection) const;
+	void ClearOutConstraintDelegates() const;
 private:
 
 	static FName TransformPropertyName;
@@ -185,4 +200,7 @@ private:
 
 	/** List of locked cameras to restore after save */
 	TArray<FGuid> LockedCameraBindings;
+
+	/** Array of sections that are getting undone, we need to recreate any constraint channel add, move key delegates to them*/
+	mutable TArray<UMovieScene3DTransformSection*> SectionsGettingUndone;
 };
diff --git a/Engine/Source/Editor/MovieSceneTools/Public/ConstraintChannelHelper.inl b/Engine/Source/Editor/MovieSceneTools/Public/ConstraintChannelHelper.inl
new file mode 100644
--- /dev/null
+++ b/Engine/Source/Editor/MovieSceneTools/Public/ConstraintChannelHelper.inl
@@ -0,0 +1,369 @@
+﻿// Copyright Epic Games, Inc. All Rights Reserved.
+
+#pragma once
+
+#include "ConstraintChannel.h"
+#include "MovieSceneConstraintChannelHelper.h"
+#include "Channels/MovieSceneCurveChannelCommon.h"
+#include "Channels/MovieSceneChannelData.h"
+#include "Channels/MovieSceneChannelProxy.h"
+#include "MovieSceneSection.h"
+#include "Algo/Unique.h"
+
+template<typename ChannelType>
+void FMovieSceneConstraintChannelHelper::GetFramesToCompensate(
+	const FMovieSceneConstraintChannel& InActiveChannel,
+	const bool InActiveValueToBeSet,
+	const FFrameNumber& InTime,
+	const TArrayView<ChannelType*>& InChannels,
+	TArray<FFrameNumber>& OutFramesAfter)
+{
+	using ChannelValueType = typename ChannelType::ChannelValueType;
+
+	const bool bHasKeys = (InActiveChannel.GetNumKeys() > 0);
+	
+	OutFramesAfter.Reset();
+
+	// add the current frame
+	OutFramesAfter.Add(InTime);
+
+	// add the next frames that need transform compensation 
+	for (const ChannelType* InChannel: InChannels)
+	{
+		const TMovieSceneChannelData<const ChannelValueType> ChannelData = InChannel->GetData();
+		const TArrayView<const FFrameNumber> Times = ChannelData.GetTimes();
+		if (!Times.IsEmpty())
+		{
+			// look for the first next key frame for this channel 
+			const int32 NextTimeIndex = Algo::UpperBound(Times, InTime);
+			if (Times.IsValidIndex(NextTimeIndex))
+			{
+				// store the time while the state is different
+				for (int32 Index = NextTimeIndex; Index < Times.Num(); ++Index)
+				{
+					if (!bHasKeys)
+					{
+						OutFramesAfter.Add(Times[Index]);
+					}
+					else
+					{
+						bool NextValue = false; InActiveChannel.Evaluate(Times[Index], NextValue);
+						if (NextValue == InActiveValueToBeSet)
+						{
+							break;
+						}
+						OutFramesAfter.Add(Times[Index]);
+					}
+				}
+			}
+		}
+	}
+
+	// uniqueness
+	OutFramesAfter.Sort();
+	OutFramesAfter.SetNum(Algo::Unique(OutFramesAfter));
+}
+
+template< typename ChannelType >
+void FMovieSceneConstraintChannelHelper::GetFramesAfter(
+	const FMovieSceneConstraintChannel& InActiveChannel,
+	const FFrameNumber& InTime,
+	const TArrayView<ChannelType*>& InChannels,
+	TArray<FFrameNumber>& OutFrames)
+{
+	using ChannelValueType = typename ChannelType::ChannelValueType;
+	
+	OutFrames.Reset();
+	
+	const TMovieSceneChannelData<const bool> ConstraintChannelData = InActiveChannel.GetData();
+	const int32 KeyIndex = ConstraintChannelData.FindKey(InTime);
+	if (!ConstraintChannelData.GetTimes().IsValidIndex(KeyIndex))
+	{
+		return;
+	}
+
+	const bool CurrentValue = ConstraintChannelData.GetValues()[KeyIndex];
+	
+	// compute last frame to compensate
+	auto GetEndOfCompensationTime = [KeyIndex](const TMovieSceneChannelData<const bool>& InData)
+	{
+		const TArrayView<const bool> Values = InData.GetValues();
+		const TArrayView<const FFrameNumber> Times = InData.GetTimes();
+
+		const bool CurrentValue = Values[KeyIndex];
+		for (int32 NextIndex = KeyIndex+1; NextIndex < Times.Num(); ++NextIndex)
+		{
+			if (Values[NextIndex] != CurrentValue)
+			{
+				return TOptional<FFrameNumber>(Times[NextIndex]);
+			}
+		}
+		return TOptional<FFrameNumber>();
+	};
+	const TOptional<FFrameNumber> EndOfCompensationTime = GetEndOfCompensationTime(ConstraintChannelData);
+
+	const bool bHasEndTime = EndOfCompensationTime.IsSet();
+	
+	// add the current frame
+	OutFrames.Add(InTime);
+
+	// add the next frames that need transform compensation 
+	for (const ChannelType* InChannel: InChannels)
+	{
+		const TMovieSceneChannelData<const ChannelValueType> ChannelData = InChannel->GetData();
+		const TArrayView<const FFrameNumber> Times = ChannelData.GetTimes();
+		if (!Times.IsEmpty())
+		{
+			// look for the first next key frame for this channel 
+			const int32 NextTimeIndex = Algo::UpperBound(Times, InTime);
+			if (Times.IsValidIndex(NextTimeIndex))
+			{
+				// store the time while the state is different
+				for (int32 Index = NextTimeIndex; Index < Times.Num(); ++Index)
+				{
+					if (!bHasEndTime || Times[Index] < EndOfCompensationTime.GetValue() )
+					{
+						OutFrames.Add(Times[Index]);
+					}
+				}
+			}
+		}
+	}
+
+	// uniqueness
+	OutFrames.Sort();
+	OutFrames.SetNum(Algo::Unique(OutFrames));
+}
+
+template< typename ChannelType >
+void FMovieSceneConstraintChannelHelper::GetFramesWithinActiveState(
+	const FMovieSceneConstraintChannel& InActiveChannel,
+	const TArrayView<ChannelType*>& InChannels,
+	TArray<FFrameNumber>& OutFrames)
+{
+	using ChannelValueType = typename ChannelType::ChannelValueType;
+	
+	OutFrames.Reset();
+	
+	const TMovieSceneChannelData<const bool> ConstraintChannelData = InActiveChannel.GetData();
+	const TArrayView<const FFrameNumber>& ActiveTimes = ConstraintChannelData.GetTimes();
+	if (ActiveTimes.IsEmpty())
+	{
+		return;
+	}
+
+	const FFrameNumber& FirstTime = ActiveTimes[0];
+	const FFrameNumber& LastTime = ActiveTimes.Last();
+	
+	// add active times
+	for (int32 Index = 0; Index < ActiveTimes.Num(); ++Index)
+	{
+		OutFrames.Add(ActiveTimes[Index]);
+	}
+
+	const bool bIsLastStateInactive = ConstraintChannelData.GetValues().Last() == false; 
+
+	// add frames where the constraint is active 
+	for (const ChannelType* InChannel: InChannels)
+	{
+		const TMovieSceneChannelData<const ChannelValueType> ChannelData = InChannel->GetData();
+		const TArrayView<const FFrameNumber> Times = ChannelData.GetTimes();
+		if (!Times.IsEmpty())
+		{
+			// look for the first next key frame for this channel 
+			const int32 NextTimeIndex = Algo::UpperBound(Times,FirstTime);
+			if (Times.IsValidIndex(NextTimeIndex))
+			{
+				// store the time is the state is active
+				for (int32 Index = NextTimeIndex; Index < Times.Num(); ++Index)
+				{
+					bool bIsActive = false;
+					InActiveChannel.Evaluate(Times[Index], bIsActive);
+					if (bIsActive)
+					{
+						OutFrames.Add(Times[Index]);
+					}
+
+					if (bIsLastStateInactive && Times[Index] > LastTime)
+					{
+						break;
+					}
+				}
+			}
+		}
+	}
+
+	// uniqueness
+	OutFrames.Sort();
+	OutFrames.SetNum(Algo::Unique(OutFrames));
+}
+
+template< typename ChannelType >
+void FMovieSceneConstraintChannelHelper::MoveTransformKeys(
+	const TArrayView<ChannelType*>& InChannels,
+	const FFrameNumber& InCurrentTime,
+	const FFrameNumber& InNextTime)
+{
+	const FFrameNumber Delta = InNextTime - InCurrentTime;
+	if (Delta == 0)
+	{
+		return;
+	}
+	
+	for (ChannelType* Channel: InChannels)
+	{
+		TMovieSceneChannelData<typename ChannelType::ChannelValueType> Data = Channel->GetData();
+		const TArrayView<const FFrameNumber> Times = Data.GetTimes();
+		const int32 NumTimes = Times.Num();
+
+		if (Delta > 0) //if we are moving keys positively in time we start from end frames and move them so we can use indices
+		{
+			for (int32 KeyIndex = NumTimes - 1; KeyIndex >= 0; --KeyIndex)
+			{
+				const FFrameNumber& Frame = Times[KeyIndex];
+				const FFrameNumber AbsDiff = FMath::Abs(Frame - InCurrentTime);
+				if (AbsDiff<= 1)
+				{
+					Data.MoveKey(KeyIndex, Frame + Delta);
+				}
+			}
+		}
+		else
+		{
+			for (int32 KeyIndex = 0; KeyIndex < NumTimes; ++KeyIndex)
+			{
+				const FFrameNumber& Frame = Times[KeyIndex];
+				const FFrameNumber AbsDiff = FMath::Abs( Frame - InCurrentTime);
+				if (AbsDiff <= 1)
+				{
+					Data.MoveKey(KeyIndex, Frame + Delta);
+				}
+			}
+		}
+	}
+}
+
+template< typename ChannelType >
+void FMovieSceneConstraintChannelHelper::DeleteTransformKeys(
+	const TArrayView<ChannelType*>& InChannels,
+	const FFrameNumber& InTime)
+{
+	for (ChannelType* Channel: InChannels)
+	{
+		TMovieSceneChannelData<typename ChannelType::ChannelValueType> Data = Channel->GetData();
+		const TArrayView<const FFrameNumber> Times = Data.GetTimes();
+		
+		const int32 KeyIndex = Algo::UpperBound(Times, InTime);
+		if (Times.IsValidIndex(KeyIndex))
+		{
+			Data.RemoveKey(KeyIndex);
+		}
+	}
+}
+
+template<typename ChannelType>
+void EvaluateTangentAtThisTime(int32 ChannelIndex, int32 NumChannels,
+	UMovieSceneSection* Section,
+	FFrameNumber Time,
+	TArray<FMovieSceneTangentData>& OutTangents)
+{
+	using ChannelValueType = typename ChannelType::ChannelValueType;
+
+	// NOTE this might be moved to FMovieSceneFloatChannel
+	auto EvaluateTangent = [](const ChannelType* InChannel, FFrameNumber InTime)
+	{
+		const TMovieSceneChannelData<const ChannelValueType> ChannelInterface = InChannel->GetData();
+		const TArrayView<const FFrameNumber> Times = ChannelInterface.GetTimes();
+
+		// Need at least two keys to evaluate a derivative
+		if (Times.Num() < 2)
+		{
+			static const FMovieSceneTangentData DefaultTangent;
+			return DefaultTangent;
+		}
+
+		const TArrayView<const ChannelValueType> Values = ChannelInterface.GetValues();
+
+		// look around to get the closest key tangent if fairly close
+		const int32 Tolerance = static_cast<int32>(InChannel->GetTickResolution().AsDecimal() * 0.01);
+		// NOTE FindKey might return Times.Num() (see AlgoImpl::LowerBoundInternal)
+		const int32 ExistingIndex = ChannelInterface.FindKey(InTime, Tolerance);
+		if (Times.IsValidIndex(ExistingIndex) && Values.IsValidIndex(ExistingIndex))
+		{
+			const FFrameNumber& Time = Times[ExistingIndex];
+			const int32 DiffToKey = InTime.Value - Time.Value;
+
+			// if the closest key is within a threshold, we use it's tangent directly instead of computing one
+			if (FMath::Abs(DiffToKey) <= Tolerance)
+			{
+				const ChannelValueType Value = Values[ExistingIndex];
+				if (DiffToKey == 0)
+				{
+					return Value.Tangent;
+				}
+
+				FMovieSceneTangentData Tangent = Value.Tangent;
+				if (DiffToKey < 0) // pretty close to the next key
+				{
+					Tangent.LeaveTangent = 0.f;
+				}
+				else // pretty close to the previous key
+				{
+					Tangent.ArriveTangent = 0.f;
+				}
+				return Tangent;
+			}
+		}
+
+		// compute tangent using central difference
+		// NOTE we may wanna compute a backward / forward difference instead
+		const int32 Delta = static_cast<int32>(InChannel->GetTickResolution().AsDecimal() * 0.1);
+
+		float PrevValue = 0.f; InChannel->Evaluate(InTime - Delta, PrevValue);
+		float NextValue = 0.f; InChannel->Evaluate(InTime + Delta, NextValue);
+		const float TangentValue = (NextValue - PrevValue) / (2.f * Delta);
+
+		FMovieSceneTangentData TangentData;
+		TangentData.ArriveTangent = TangentValue;
+		TangentData.LeaveTangent = TangentValue;
+
+		return TangentData;
+	};
+
+
+	// compute and store tangents
+	OutTangents.SetNum(NumChannels);
+
+	const TArrayView<ChannelType*> Channels = Section->GetChannelProxy().GetChannels<ChannelType>();
+	for (int32 Index = 0; Index < NumChannels; ++Index, ++ChannelIndex)
+	{
+		OutTangents[Index] = EvaluateTangent(Channels[ChannelIndex], Time);
+	}
+}
+
+// NOTE we may pass an enum to tell which tangent we wanna set (arrive, leave, both)
+template<typename ChannelType>
+void SetTangentsAtThisTime(int32 ChannelIndex,
+	int32 NumChannels,
+	UMovieSceneSection* Section,
+	FFrameNumber Time,
+	const TArray<FMovieSceneTangentData>& InTangents)
+{
+	using ChannelValueType = typename ChannelType::ChannelValueType;
+
+	const TArrayView<ChannelType*> Channels = Section->GetChannelProxy().GetChannels<ChannelType>();
+
+	for (int32 Index = 0; Index < NumChannels; ++Index, ++ChannelIndex)
+	{
+		TMovieSceneChannelData<ChannelValueType> ChannelInterface = Channels[ChannelIndex]->GetData();
+		TArrayView<ChannelValueType> Values = ChannelInterface.GetValues();
+		const int32 KeyIndex = ChannelInterface.FindKey(Time);
+		if (KeyIndex != INDEX_NONE)
+		{
+			ChannelValueType& Value = Values[KeyIndex];
+			Value.Tangent.ArriveTangent = InTangents[Index].ArriveTangent;
+			Value.Tangent.LeaveTangent = InTangents[Index].LeaveTangent;
+			Value.TangentMode = RCTM_Break;
+		}
+	}
+}
\ No newline at end of file
diff --git a/Engine/Source/Editor/MovieSceneTools/Public/KeyframeTrackEditor.h b/Engine/Source/Editor/MovieSceneTools/Public/KeyframeTrackEditor.h
--- a/Engine/Source/Editor/MovieSceneTools/Public/KeyframeTrackEditor.h
+++ b/Engine/Source/Editor/MovieSceneTools/Public/KeyframeTrackEditor.h
@@ -664,6 +664,10 @@ protected:
 			for (const FMovieSceneChannelValueSetter& GeneratedKey : Keys)
 			{
 				KeyPropertyResult.bKeyCreated |= GeneratedKey->Apply(Section, Proxy, KeyTime, InterpolationMode, bKeyEvenIfUnchanged, bKeyEvenIfEmpty);
+				if (KeyPropertyResult.bKeyCreated)
+				{
+					KeyPropertyResult.SectionsKeyed.Add(Section);
+				}
 			}
 		}
 			
diff --git a/Engine/Source/Editor/MovieSceneTools/Public/MovieSceneConstraintChannelHelper.h b/Engine/Source/Editor/MovieSceneTools/Public/MovieSceneConstraintChannelHelper.h
new file mode 100644
--- /dev/null
+++ b/Engine/Source/Editor/MovieSceneTools/Public/MovieSceneConstraintChannelHelper.h
@@ -0,0 +1,111 @@
+﻿// Copyright Epic Games, Inc. All Rights Reserved.
+
+#pragma once
+
+#include "Templates/SharedPointer.h"
+
+class UWorld;
+class UTickableConstraint;
+class ISequencer;
+class UMovieSceneSection;
+class UTickableTransformConstraint;
+class UTransformableHandle;
+struct FMovieSceneConstraintChannel;
+struct FFrameNumber;
+
+
+struct MOVIESCENETOOLS_API FCompensationEvaluator
+{
+public:
+	TArray<FTransform> ChildLocals;
+	TArray<FTransform> ChildGlobals;
+	TArray<FTransform> SpaceGlobals;
+
+	FCompensationEvaluator(UTickableTransformConstraint* InConstraint);
+
+	void ComputeLocalTransforms(UWorld* InWorld, const TSharedPtr<ISequencer>& InSequencer, const TArray<FFrameNumber>& InFrames, const bool bToActive);
+	void ComputeLocalTransformsBeforeDeletion(UWorld* InWorld, const TSharedPtr<ISequencer>& InSequencer, const TArray<FFrameNumber>& InFrames);
+	void ComputeCompensation(UWorld* InWorld, const TSharedPtr<ISequencer>& InSequencer, const FFrameNumber& InTime);
+
+private:
+
+	TArray<UTickableTransformConstraint*> GetHandleTransformConstraints(UWorld* InWorld) const;
+
+	UTickableTransformConstraint* Constraint = nullptr;
+	UTransformableHandle* Handle = nullptr;
+};
+
+struct MOVIESCENETOOLS_API FMovieSceneConstraintChannelHelper
+{
+public:
+
+	/** Compensate transform on handles when a constraint switches state. 
+	static void Compensate(
+		UTickableTransformConstraint* InConstraint,
+		const bool bAllTimes = false);
+		*/
+	static void CompensateIfNeeded(
+		const TSharedPtr<ISequencer>& InSequencer,
+		UMovieSceneSection* Section,
+		const TOptional<FFrameNumber>& OptionalTime);
+		
+	static void HandleConstraintRemoved(
+		UTickableConstraint* InConstraint,
+		const FMovieSceneConstraintChannel* InChannel,
+		const TSharedPtr<ISequencer>& InSequencer,
+		UMovieSceneSection* InSection);
+
+	static void HandleConstraintKeyDeleted(
+		UTickableTransformConstraint* InConstraint,
+		const FMovieSceneConstraintChannel* InConstraintChannel,
+		const TSharedPtr<ISequencer>& InSequencer,
+		UMovieSceneSection* InSection,
+		const FFrameNumber& InTime);
+
+	static void HandleConstraintKeyMoved(
+		const UTickableTransformConstraint* InConstraint,
+		const FMovieSceneConstraintChannel* InConstraintChannel,
+		UMovieSceneSection* InSection,
+		const FFrameNumber& InCurrentFrame, const FFrameNumber& InNextFrame);
+
+	/** @todo documentation. */
+	template<typename ChannelType>
+	static void GetFramesToCompensate(
+		const FMovieSceneConstraintChannel& InActiveChannel,
+		const bool InActiveValueToBeSet,
+		const FFrameNumber& InTime,
+		const TArrayView<ChannelType*>& InChannels,
+		TArray<FFrameNumber>& OutFramesAfter);
+
+	/** @todo documentation. */
+	template< typename ChannelType >
+	static void GetFramesAfter(
+		const FMovieSceneConstraintChannel& InActiveChannel,
+		const FFrameNumber& InTime,
+		const TArrayView<ChannelType*>& InChannels,
+		TArray<FFrameNumber>& OutFrames);
+
+	/** @todo documentation. */
+	template< typename ChannelType >
+	static void GetFramesWithinActiveState(
+		const FMovieSceneConstraintChannel& InActiveChannel,
+		const TArrayView<ChannelType*>& InChannels,
+		TArray<FFrameNumber>& OutFrames);
+
+	/** @todo documentation. */
+	template< typename ChannelType >
+	static void MoveTransformKeys(
+		const TArrayView<ChannelType*>& InChannels,
+		const FFrameNumber& InCurrentTime,
+		const FFrameNumber& InNextTime);
+
+	/** @todo documentation. */
+	template< typename ChannelType >
+	static void DeleteTransformKeys(
+		const TArrayView<ChannelType*>& InChannels,
+		const FFrameNumber& InTime);
+
+private:
+	static bool bDoNotCompensate;
+
+};
diff --git a/Engine/Source/Editor/MovieSceneTools/Public/MovieSceneToolHelpers.h b/Engine/Source/Editor/MovieSceneTools/Public/MovieSceneToolHelpers.h
--- a/Engine/Source/Editor/MovieSceneTools/Public/MovieSceneToolHelpers.h
+++ b/Engine/Source/Editor/MovieSceneTools/Public/MovieSceneToolHelpers.h
@@ -37,6 +37,8 @@ template<typename ChannelType> struct TMovieSceneChannelData;
 enum class EVisibilityBasedAnimTickOption : uint8;
 class ACameraActor;
 struct FActorForWorldTransforms;
+class UMovieScene3DTransformSection;
+enum class EMovieSceneTransformChannel : uint32;
 
 namespace fbxsdk
 {
@@ -557,6 +559,26 @@ public:
 	 * Return whether this asset is valid for the given sequence
 	 */
 	static bool IsValidAsset(UMovieSceneSequence* Sequence, const FAssetData& InAssetData);
+
+	/** Returns the frame numbers between start and end. */
+	static void CalculateFramesBetween(
+		const UMovieScene* MovieScene,
+		FFrameNumber StartFrame,
+		FFrameNumber EndFrame,
+		TArray<FFrameNumber>& OutFrames);
+
+	/** Returns the transform section for that guid. */
+	static UMovieScene3DTransformSection* GetTransformSection(
+		const ISequencer* InSequencer,
+		const FGuid& InGuid,
+		const FTransform& InDefaultTransform = FTransform::Identity);
+
+	/** Adds transform keys to the section based on the channels filters. */
+	static bool AddTransformKeys(
+		const UMovieScene3DTransformSection* InTransformSection,
+		const TArray<FFrameNumber>& Frames,
+		const TArray<FTransform>& InLocalTransforms,
+		const EMovieSceneTransformChannel& InChannels);
 };
 
 // Helper to make spawnables persist throughout the export process and then restore properly afterwards
diff --git a/Engine/Source/Editor/Sequencer/Public/MovieSceneTrackEditor.h b/Engine/Source/Editor/Sequencer/Public/MovieSceneTrackEditor.h
--- a/Engine/Source/Editor/Sequencer/Public/MovieSceneTrackEditor.h
+++ b/Engine/Source/Editor/Sequencer/Public/MovieSceneTrackEditor.h
@@ -51,6 +51,9 @@ struct FKeyPropertyResult
 
 	/* Was a section created */
 	TArray<TWeakObjectPtr<UMovieSceneSection> > SectionsCreated;
+
+	/* Was a section keyed */
+	TArray<TWeakObjectPtr<UMovieSceneSection> > SectionsKeyed;
 };
 
 /** Delegate for adding keys for a property
diff --git a/Engine/Source/Runtime/Experimental/Animation/Constraints/Private/TransformConstraint.cpp b/Engine/Source/Runtime/Experimental/Animation/Constraints/Private/TransformConstraint.cpp
--- a/Engine/Source/Runtime/Experimental/Animation/Constraints/Private/TransformConstraint.cpp
+++ b/Engine/Source/Runtime/Experimental/Animation/Constraints/Private/TransformConstraint.cpp
@@ -5,6 +5,7 @@
 #include "ConstraintsManager.h"
 #include "TransformableRegistry.h"
 #include "GameFramework/Actor.h"
+#include "MovieSceneSection.h"
 
 /** 
  * UTickableTransformConstraint
@@ -14,6 +15,24 @@ int64 UTickableTransformConstraint::GetType() const
 {
 	return static_cast<int64>(Type);
 }
+EMovieSceneTransformChannel UTickableTransformConstraint::GetChannelsToKey() const
+{
+	static const TMap< ETransformConstraintType, EMovieSceneTransformChannel > ConstraintToChannels({
+	{ETransformConstraintType::Translation, EMovieSceneTransformChannel::Translation},
+	{ETransformConstraintType::Rotation, EMovieSceneTransformChannel::Rotation},
+	{ETransformConstraintType::Scale, EMovieSceneTransformChannel::Scale},
+	{ETransformConstraintType::Parent, EMovieSceneTransformChannel::AllTransform},
+	{ETransformConstraintType::LookAt, EMovieSceneTransformChannel::Rotation}
+		});
+
+	const ETransformConstraintType ConstType = static_cast<ETransformConstraintType>(GetType());
+	if (const EMovieSceneTransformChannel* Channel = ConstraintToChannels.Find(ConstType))
+	{
+		return *Channel;;
+	}
+
+	return EMovieSceneTransformChannel::AllTransform;
+}
 
 #if WITH_EDITOR
 
diff --git a/Engine/Source/Runtime/Experimental/Animation/Constraints/Private/TransformableHandle.cpp b/Engine/Source/Runtime/Experimental/Animation/Constraints/Private/TransformableHandle.cpp
--- a/Engine/Source/Runtime/Experimental/Animation/Constraints/Private/TransformableHandle.cpp
+++ b/Engine/Source/Runtime/Experimental/Animation/Constraints/Private/TransformableHandle.cpp
@@ -2,11 +2,15 @@
 
 
 #include "TransformableHandle.h"
-
+#include "ConstraintsManager.h"
 #include "Components/SceneComponent.h"
 #include "GameFramework/Actor.h"
 #include "Engine/Engine.h"
-
+#include "MovieSceneSection.h"
+#include "Channels/MovieSceneFloatChannel.h"
+#include "Channels/MovieSceneDoubleChannel.h"
+#include "Channels/MovieSceneChannelProxy.h"
+#include "Channels/MovieSceneDoubleChannel.h"
 /**
  * UTransformableHandle
  */
@@ -40,9 +44,15 @@ bool UTransformableComponentHandle::IsValid() const
 
 void UTransformableComponentHandle::SetGlobalTransform(const FTransform& InGlobal) const
 {
-	if(Component.IsValid())
+	if (Component.IsValid())
 	{
 		Component->SetWorldTransform(InGlobal);
+		UWorld* World = Component->GetWorld();
+		if (World)
+		{
+			FConstraintsManagerController& Controller = FConstraintsManagerController::Get(World);
+			Controller.OnSceneComponentConstrained().Broadcast(Component.Get());
+		}
 	}
 }
 
@@ -51,6 +61,12 @@ void UTransformableComponentHandle::SetLocalTransform(const FTransform& InLocal)
 	if(Component.IsValid())
 	{
 		Component->SetRelativeTransform(InLocal);
+		UWorld* World = Component->GetWorld();
+		if (World)
+		{
+			FConstraintsManagerController& Controller = FConstraintsManagerController::Get(World);
+			Controller.OnSceneComponentConstrained().Broadcast(Component.Get());
+		}
 	}
 }
 
@@ -166,6 +182,154 @@ void UTransformableComponentHandle::OnPostPropertyChanged(
 	}
 }
 
+TArrayView<FMovieSceneFloatChannel*>  UTransformableComponentHandle::GetFloatChannels(const UMovieSceneSection* InSection) const
+{
+	// no floats for transform sections
+	static const TArrayView<FMovieSceneFloatChannel*> EmptyChannelsView;
+	return EmptyChannelsView;
+}
+
+TArrayView<FMovieSceneDoubleChannel*>  UTransformableComponentHandle::GetDoubleChannels(const UMovieSceneSection* InSection) const
+{
+	static const TArrayView<FMovieSceneDoubleChannel*> EmptyChannelsView;
+	if (InSection)
+	{
+		return InSection->GetChannelProxy().GetChannels<FMovieSceneDoubleChannel>();
+	}
+	return EmptyChannelsView;
+}
+
+bool UTransformableComponentHandle::AddTransformKeys(const TArray<FFrameNumber>& Frames,
+	const TArray<FTransform>& InLocalTransforms,
+	const EMovieSceneTransformChannel& InChannels,
+	const FFrameRate& InTickResolution,
+	UMovieSceneSection* InTransformSection,
+	const bool bLocal) const
+{
+	//todo 
+	//note this is the same as MOvieSceneToolHelpers::AddTransformKeys but that's in the editor code, need to find a place in the runtime someewhere
+	if (!InTransformSection)
+	{
+		return false;
+	}
+
+	if (Frames.IsEmpty() || Frames.Num() != InLocalTransforms.Num())
+	{
+		return false;
+	}
+
+	auto GetValue = [](const uint32 Index, const FVector& InLocation, const FRotator& InRotation, const FVector& InScale)
+	{
+		switch (Index)
+		{
+		case 0:
+			return InLocation.X;
+		case 1:
+			return InLocation.Y;
+		case 2:
+			return InLocation.Z;
+		case 3:
+			return InRotation.Roll;
+		case 4:
+			return InRotation.Pitch;
+		case 5:
+			return InRotation.Yaw;
+		case 6:
+			return InScale.X;
+		case 7:
+			return InScale.Y;
+		case 8:
+			return InScale.Z;
+		default:
+			ensure(false);
+			break;
+		}
+		return 0.0;
+	};
+
+	const bool bKeyTranslation = EnumHasAllFlags(InChannels, EMovieSceneTransformChannel::Translation);
+	const bool bKeyRotation = EnumHasAllFlags(InChannels, EMovieSceneTransformChannel::Rotation);
+	const bool bKeyScale = EnumHasAllFlags(InChannels, EMovieSceneTransformChannel::Scale);
+
+	TArray<uint32> ChannelsIndexToKey;
+	if (bKeyTranslation)
+	{
+		ChannelsIndexToKey.Append({ 0,1,2 });
+	}
+	if (bKeyRotation)
+	{
+		ChannelsIndexToKey.Append({ 3,4,5 });
+	}
+	if (bKeyScale)
+	{
+		ChannelsIndexToKey.Append({ 6,7,8 });
+	}
+
+	const TArrayView<FMovieSceneDoubleChannel*> Channels =
+		InTransformSection->GetChannelProxy().GetChannels<FMovieSceneDoubleChannel>();
+
+	// set default
+	const FTransform& LocalTransform0 = InLocalTransforms[0];
+	const FVector Location0 = LocalTransform0.GetLocation();
+	const FRotator Rotation0 = LocalTransform0.GetRotation().Rotator();
+	const FVector Scale3D0 = LocalTransform0.GetScale3D();
+
+	for (int32 ChannelIndex = 0; ChannelIndex < 9; ChannelIndex++)
+	{
+		if (!Channels[ChannelIndex]->GetDefault().IsSet())
+		{
+			const double Value = GetValue(ChannelIndex, Location0, Rotation0, Scale3D0);
+			Channels[ChannelIndex]->SetDefault(Value);
+		}
+	}
+
+	// add keys
+	for (int32 Index = 0; Index < Frames.Num(); ++Index)
+	{
+		const FFrameNumber& Frame = Frames[Index];
+		const FTransform& LocalTransform = InLocalTransforms[Index];
+
+		const FVector Location = LocalTransform.GetLocation();
+		const FRotator Rotation = LocalTransform.GetRotation().Rotator();
+		const FVector Scale3D = LocalTransform.GetScale3D();
+
+		for (const int32 ChannelIndex : ChannelsIndexToKey)
+		{
+			const double Value = GetValue(ChannelIndex, Location, Rotation, Scale3D);
+			TMovieSceneChannelData<FMovieSceneDoubleValue> ChannelData = Channels[ChannelIndex]->GetData();
+			{
+				int32 ExistingIndex = ChannelData.FindKey(Frame);
+				if (ExistingIndex != INDEX_NONE)
+				{
+					FMovieSceneDoubleValue& DoubleValue = ChannelData.GetValues()[ExistingIndex]; //-V758
+					DoubleValue.Value = Value;
+				}
+				else
+				{
+					FMovieSceneDoubleValue NewKey(Value);
+					ERichCurveTangentWeightMode WeightedMode = RCTWM_WeightedNone;
+					NewKey.InterpMode = ERichCurveInterpMode::RCIM_Cubic;
+					NewKey.TangentMode = ERichCurveTangentMode::RCTM_Auto;
+					NewKey.Tangent.ArriveTangent = 0.0f;
+					NewKey.Tangent.LeaveTangent = 0.0f;
+					NewKey.Tangent.TangentWeightMode = WeightedMode;
+					NewKey.Tangent.ArriveTangentWeight = 0.0f;
+					NewKey.Tangent.LeaveTangentWeight = 0.0f;
+					ChannelData.AddKey(Frame, NewKey);
+				}
+			}
+		}
+	}
+
+	//now we need to set auto tangents
+	for (const int32 ChannelIndex : ChannelsIndexToKey)
+	{
+		Channels[ChannelIndex]->AutoSetTangents();
+	}
+	return true;
+}
+
+
 #if WITH_EDITOR
 FString UTransformableComponentHandle::GetLabel() const
 {
diff --git a/Engine/Source/Runtime/Experimental/Animation/Constraints/Public/ConstraintsManager.h b/Engine/Source/Runtime/Experimental/Animation/Constraints/Public/ConstraintsManager.h
--- a/Engine/Source/Runtime/Experimental/Animation/Constraints/Public/ConstraintsManager.h
+++ b/Engine/Source/Runtime/Experimental/Animation/Constraints/Public/ConstraintsManager.h
@@ -240,6 +240,9 @@ public:
 		const FName& InNameToTickAfter) const;
 	
 private:
+	/** Delegeate that's fired when a scene component is constrained, this is needed to make sure things like gizmo's get updated after the constraint tick happens*/
+	DECLARE_MULTICAST_DELEGATE_OneParam(FOnSceneComponentConstrained, USceneComponent* /*InSceneComponent*/);
+	FOnSceneComponentConstrained SceneComponentConstrained;
 
 	/** @todo document */
 	DECLARE_MULTICAST_DELEGATE_OneParam(FOnConstraintRemoved, FName /*InConstraintName*/);
@@ -258,7 +261,8 @@ private:
 	UWorld* World = nullptr;
 
 public:
-
+	/** Delegate that's fired when a scene component is constrained, this is needed to make sure things like gizmo's get updated after the constraint tick happens*/
+	FOnSceneComponentConstrained& OnSceneComponentConstrained() { return SceneComponentConstrained; }
 	/** @todo document */
 	FOnConstraintRemoved& OnConstraintRemoved(){ return ConstraintRemoved; };
 };
\ No newline at end of file
diff --git a/Engine/Source/Runtime/Experimental/Animation/Constraints/Public/TransformConstraint.h b/Engine/Source/Runtime/Experimental/Animation/Constraints/Public/TransformConstraint.h
--- a/Engine/Source/Runtime/Experimental/Animation/Constraints/Public/TransformConstraint.h
+++ b/Engine/Source/Runtime/Experimental/Animation/Constraints/Public/TransformConstraint.h
@@ -11,6 +11,7 @@
 class UTransformableHandle;
 class UTransformableComponentHandle;
 class USceneComponent;
+enum class EMovieSceneTransformChannel : uint32;
 
 using SetTransformFunc = TFunction<void(const FTransform&)>;
 using GetTransformFunc = TFunction<FTransform()>;
@@ -75,6 +76,9 @@ public:
 	
 	/** Get the current parent's local transform. */
 	FTransform GetParentLocalTransform() const;
+
+	/** Returns the channels to key based on the constraint's type. */
+	EMovieSceneTransformChannel GetChannelsToKey() const;
 	
 protected:
 
diff --git a/Engine/Source/Runtime/Experimental/Animation/Constraints/Public/TransformableHandle.h b/Engine/Source/Runtime/Experimental/Animation/Constraints/Public/TransformableHandle.h
--- a/Engine/Source/Runtime/Experimental/Animation/Constraints/Public/TransformableHandle.h
+++ b/Engine/Source/Runtime/Experimental/Animation/Constraints/Public/TransformableHandle.h
@@ -8,6 +8,12 @@
 
 struct FTickFunction;
 class USceneComponent;
+struct FMovieSceneFloatChannel;
+struct FMovieSceneDoubleChannel;
+class UMovieSceneSection;
+struct FFRameNumber;
+struct FFrameRate;
+enum class EMovieSceneTransformChannel : uint32;
 
 /**
  * UTransformableHandle
@@ -38,6 +44,18 @@ public:
 	/** Gets the local transform of the underlying transformable object in it's parent space. */
 	virtual FTransform GetLocalTransform() const PURE_VIRTUAL(GetLocalTransform, return FTransform::Identity;);
 
+	/** Get the array of float channels for the specified section*/
+	virtual TArrayView<FMovieSceneFloatChannel*>  GetFloatChannels(const UMovieSceneSection* InSection) const PURE_VIRTUAL(GetFloatChannels, return TArrayView<FMovieSceneFloatChannel*>(); );
+	/** Get the array of double channels for the specified section*/
+	virtual TArrayView<FMovieSceneDoubleChannel*>  GetDoubleChannels(const UMovieSceneSection * InSection) const PURE_VIRTUAL(GetDoubleChannels, return TArrayView<FMovieSceneDoubleChannel*>(); );
+	/** Add Transform Keys at the specified times*/
+	virtual bool AddTransformKeys(const TArray<FFrameNumber>& InFrames,
+		const TArray<FTransform>& InTransforms,
+		const EMovieSceneTransformChannel& InChannels,
+		const FFrameRate& InTickResolution,
+		UMovieSceneSection* InSection,
+		const bool bLocal = true) const PURE_VIRTUAL(AddTransformKeys, return false;);
+
 	/**
 	 * Returns the target object containing the tick function (returned in GetTickFunction).
 	 * See FTickFunction::AddPrerequisite for details.
@@ -102,6 +120,17 @@ public:
 	/** @todo document */
 	virtual TWeakObjectPtr<UObject> GetTarget() const override;
 
+	/** Get the array of float channels for the specified section*/
+	virtual TArrayView<FMovieSceneFloatChannel*>  GetFloatChannels(const UMovieSceneSection* InSection) const override;
+	/** Get the array of double channels for the specified section*/
+	virtual TArrayView<FMovieSceneDoubleChannel*>  GetDoubleChannels(const UMovieSceneSection* InSection) const override;
+	/** Add Transform Keys at the specified times*/
+	virtual bool AddTransformKeys(const TArray<FFrameNumber>& InFrames,
+		const TArray<FTransform>& InTransforms,
+		const EMovieSceneTransformChannel& InChannels,
+		const FFrameRate& InTickResolution,
+		UMovieSceneSection* InSection,
+		const bool bLocal = true) const override;
 #if WITH_EDITOR
 	virtual FString GetLabel() const override;
 	virtual FString GetFullLabel() const override;
diff --git a/Engine/Source/Runtime/MovieScene/Public/MovieSceneSection.h b/Engine/Source/Runtime/MovieScene/Public/MovieSceneSection.h
--- a/Engine/Source/Runtime/MovieScene/Public/MovieSceneSection.h
+++ b/Engine/Source/Runtime/MovieScene/Public/MovieSceneSection.h
@@ -213,6 +213,34 @@ public:
 	UPROPERTY(EditAnywhere, BlueprintReadOnly, Category="Timecode")
 	FTimecode Timecode;
 };
+/**
+* Defines for common transform 'type' sections. Moved here to avoid extra module dependencies
+*/
+enum class EMovieSceneTransformChannel : uint32
+{
+	None = 0x000,
+
+	TranslationX = 0x001,
+	TranslationY = 0x002,
+	TranslationZ = 0x004,
+	Translation = TranslationX | TranslationY | TranslationZ,
+
+	RotationX = 0x008,
+	RotationY = 0x010,
+	RotationZ = 0x020,
+	Rotation = RotationX | RotationY | RotationZ,
+
+	ScaleX = 0x040,
+	ScaleY = 0x080,
+	ScaleZ = 0x100,
+	Scale = ScaleX | ScaleY | ScaleZ,
+
+	AllTransform = Translation | Rotation | Scale,
+
+	Weight = 0x200,
+
+	All = Translation | Rotation | Scale | Weight,
+};
 
 /**
  * Base class for movie scene sections
diff --git a/Engine/Source/Runtime/MovieSceneTracks/Private/Sections/MovieScene3DTransformSection.cpp b/Engine/Source/Runtime/MovieSceneTracks/Private/Sections/MovieScene3DTransformSection.cpp
--- a/Engine/Source/Runtime/MovieSceneTracks/Private/Sections/MovieScene3DTransformSection.cpp
+++ b/Engine/Source/Runtime/MovieSceneTracks/Private/Sections/MovieScene3DTransformSection.cpp
@@ -838,46 +838,53 @@ FConstraintAndActiveChannel* UMovieScene3DTransformSection::GetConstraintChannel
 	return nullptr;
 }
 
-void UMovieScene3DTransformSection::SetUpConstraintRemovedHandle()
+TArray<FConstraintAndActiveChannel>& UMovieScene3DTransformSection::GetConstraintsChannels() 
 {
+	static TArray< FConstraintAndActiveChannel> EmptyChannels;
 	if (Constraints)
 	{
-		if (Constraints->ConstraintsChannels.Num() > 0 && !Constraints->OnConstraintRemovedHandle.IsValid())
-		{
-			FConstraintsManagerController& Controller = FConstraintsManagerController::Get(GetWorld());
-			Constraints->OnConstraintRemovedHandle = Controller.OnConstraintRemoved().AddLambda([this](FName InConstraintName)
-			{
-				Modify();
-				const int32 NumRemoved = Constraints->ConstraintsChannels.RemoveAll([InConstraintName](const FConstraintAndActiveChannel& InChannel)
-					{
-						return InChannel.Constraint.IsValid() ? InChannel.Constraint->GetFName() == InConstraintName : false;
-					});
-				if (NumRemoved > 0)
-				{
-					CacheChannelProxy();
-				}
-			});
-		}
+		return Constraints->ConstraintsChannels;
 	}
+	return EmptyChannels;
 }
 
 void UMovieScene3DTransformSection::AddConstraintChannel(UTickableConstraint* InConstraint)
 {
+	Modify();
 	if (!Constraints)
 	{
 		Constraints = NewObject<UMovieScene3dTransformSectionConstraints>(this, NAME_None, RF_Transactional);
 	}
 	if (!HasConstraintChannel(InConstraint->GetFName()))
 	{
-		Modify();
-
 		const int32 NewIndex = Constraints->ConstraintsChannels.Add(FConstraintAndActiveChannel(InConstraint));
 
 		FMovieSceneConstraintChannel* ExistingChannel = &Constraints->ConstraintsChannels[NewIndex].ActiveChannel;
 		ExistingChannel->SetDefault(false);
 
-		SetUpConstraintRemovedHandle();
 		CacheChannelProxy();
+
+		if (OnConstraintChannelAdded.IsBound())
+		{
+			OnConstraintChannelAdded.Broadcast(this, ExistingChannel);
+		}
 	}
 }
 
+void UMovieScene3DTransformSection::RemoveConstraintChannel(const FName& InConstraintName)
+{
+	if (Constraints)
+	{
+		const int32 Index = Constraints->ConstraintsChannels.IndexOfByPredicate([InConstraintName](const FConstraintAndActiveChannel& InChannel)
+			{
+				return InChannel.Constraint.IsValid() ? InChannel.Constraint->GetFName() == InConstraintName : false;
+			});
+
+		if (Constraints->ConstraintsChannels.IsValidIndex(Index))
+		{
+			Modify();
+			Constraints->ConstraintsChannels.RemoveAt(Index);
+			CacheChannelProxy();
+		}
+	}
+}
diff --git a/Engine/Source/Runtime/MovieSceneTracks/Private/Sections/MovieSceneConstrainedSection.cpp b/Engine/Source/Runtime/MovieSceneTracks/Private/Sections/MovieSceneConstrainedSection.cpp
new file mode 100644
--- /dev/null
+++ b/Engine/Source/Runtime/MovieSceneTracks/Private/Sections/MovieSceneConstrainedSection.cpp
@@ -0,0 +1,8 @@
+// Copyright Epic Games, Inc. All Rights Reserved.
+
+#include "Sections/MovieSceneConstrainedSection.h"
+
+
+IMovieSceneConstrainedSection::IMovieSceneConstrainedSection()
+{
+}
diff --git a/Engine/Source/Runtime/MovieSceneTracks/Public/Sections/MovieScene3DTransformSection.h b/Engine/Source/Runtime/MovieSceneTracks/Public/Sections/MovieScene3DTransformSection.h
--- a/Engine/Source/Runtime/MovieSceneTracks/Public/Sections/MovieScene3DTransformSection.h
+++ b/Engine/Source/Runtime/MovieSceneTracks/Public/Sections/MovieScene3DTransformSection.h
@@ -7,6 +7,7 @@
 #include "Components/SceneComponent.h"
 #include "Curves/KeyHandle.h"
 #include "MovieSceneSection.h"
+#include "MovieSceneConstrainedSection.h"
 #include "MovieSceneKeyStruct.h"
 #include "Channels/MovieSceneDoubleChannel.h"
 #include "Channels/MovieSceneFloatChannel.h"
@@ -132,31 +133,7 @@ struct FMovieScene3DTransformKeyStruct
 };
 template<> struct TStructOpsTypeTraits<FMovieScene3DTransformKeyStruct> : public TStructOpsTypeTraitsBase2<FMovieScene3DTransformKeyStruct> { enum { WithCopy = false }; };
 
-enum class EMovieSceneTransformChannel : uint32
-{
-	None			= 0x000,
-
-	TranslationX 	= 0x001,
-	TranslationY 	= 0x002,
-	TranslationZ 	= 0x004,
-	Translation 	= TranslationX | TranslationY | TranslationZ,
-
-	RotationX 		= 0x008,
-	RotationY 		= 0x010,
-	RotationZ 		= 0x020,
-	Rotation 		= RotationX | RotationY | RotationZ,
-
-	ScaleX 			= 0x040,
-	ScaleY 			= 0x080,
-	ScaleZ 			= 0x100,
-	Scale 			= ScaleX | ScaleY | ScaleZ,
 
-	AllTransform	= Translation | Rotation | Scale,
-
-	Weight 			= 0x200,
-
-	All				= Translation | Rotation | Scale | Weight,
-};
 ENUM_CLASS_FLAGS(EMovieSceneTransformChannel)
 
 USTRUCT()
@@ -224,16 +201,15 @@ public:
 	UPROPERTY()
 	TArray<FConstraintAndActiveChannel> ConstraintsChannels;
 
-	FDelegateHandle OnConstraintRemovedHandle;
 };
 
-
 /**
  * A 3D transform section
  */
 UCLASS(MinimalAPI)
 class UMovieScene3DTransformSection
 	: public UMovieSceneSection
+	, public IMovieSceneConstrainedSection
 	, public IMovieSceneEntityProvider
 {
 	GENERATED_UCLASS_BODY()
@@ -322,25 +298,35 @@ private:
 	bool bUseQuaternionInterpolation;
 
 public:
+	//IMovieSceneConstrainedSection overrides
 
 	/*
 	* If it has that constraint with that Name
 	*/
-	MOVIESCENETRACKS_API bool HasConstraintChannel(const FName& InConstraintName) const;
+	virtual  bool HasConstraintChannel(const FName& InConstraintName) const override;
 
 	/*
 	* Get constraint with that name
 	*/
-	MOVIESCENETRACKS_API FConstraintAndActiveChannel* GetConstraintChannel(const FName& InConstraintName);
+	virtual FConstraintAndActiveChannel* GetConstraintChannel(const FName& InConstraintName) override;
 
 	/*
-	*  Add Cosntraint channel
+	*  Add Constraint channel
 	*/
-	MOVIESCENETRACKS_API void AddConstraintChannel(UTickableConstraint* InConstraint);
+	virtual void AddConstraintChannel(UTickableConstraint* InConstraint) override;
+	
+	/*
+	*  Remove Constraint channel
+	*/
+	virtual void RemoveConstraintChannel(const FName& InConstraintName) override;
 
-private:
+	/*
+	*  Get The channels by value
+	*/
+	virtual TArray<FConstraintAndActiveChannel>& GetConstraintsChannels()  override;
 
-	void SetUpConstraintRemovedHandle();
+
+private:
 
 #if WITH_EDITORONLY_DATA
 
diff --git a/Engine/Source/Runtime/MovieSceneTracks/Public/Sections/MovieSceneConstrainedSection.h b/Engine/Source/Runtime/MovieSceneTracks/Public/Sections/MovieSceneConstrainedSection.h
new file mode 100644
--- /dev/null
+++ b/Engine/Source/Runtime/MovieSceneTracks/Public/Sections/MovieSceneConstrainedSection.h
@@ -0,0 +1,80 @@
+// Copyright Epic Games, Inc. All Rights Reserved.
+
+#pragma once
+
+#include "CoreMinimal.h"
+#include "UObject/ObjectMacros.h"
+#include "UObject/NameTypes.h"
+#include "MovieSceneSection.h"
+#include "ConstraintsManager.h"
+#include "ConstraintChannel.h"
+#include "MovieSceneConstrainedSection.generated.h"
+
+
+
+/**
+ * Functionality to add to sections that contain constraints
+ */
+
+UINTERFACE()
+class MOVIESCENETRACKS_API UMovieSceneConstrainedSection : public UInterface
+{
+public:
+	GENERATED_BODY()
+};
+
+/**
+ * Interface to be added to UMovieSceneSection types when they contain entity data
+ */
+class MOVIESCENETRACKS_API IMovieSceneConstrainedSection
+{
+public:
+
+	GENERATED_BODY()
+	IMovieSceneConstrainedSection();
+
+	DECLARE_MULTICAST_DELEGATE_TwoParams(FConstraintChannelAddedEvent, IMovieSceneConstrainedSection*, FMovieSceneConstraintChannel*);
+
+
+	/*
+	* If it has that constraint with that Name
+	*/
+	virtual bool HasConstraintChannel(const FName& InConstraintName) const = 0;
+
+	/*
+	* Get constraint with that name
+	*/
+	virtual FConstraintAndActiveChannel* GetConstraintChannel(const FName& InConstraintName) = 0;
+
+	/*
+	*  Add Constraint channel
+	*/
+	virtual void AddConstraintChannel(UTickableConstraint* InConstraint) = 0;
+	
+	/*
+	*  Remove Constraint channel
+	*/
+	virtual void RemoveConstraintChannel(const FName& InConstraintName) = 0;
+
+	/*
+	*  Get The channels
+	*/
+	virtual TArray<FConstraintAndActiveChannel>& GetConstraintsChannels() = 0;
+
+	/*
+	* Added Delegate
+	*/
+	 FConstraintChannelAddedEvent& ConstraintChannelAdded() { return OnConstraintChannelAdded; }
+
+	/*
+	*  Removed delegate that get's added by the track editor
+	*/
+	FDelegateHandle OnConstraintRemovedHandle;
+
+
+protected:
+
+	FConstraintChannelAddedEvent OnConstraintChannelAdded;
+
+
+};
