Fix control rig display actor unstabilities, bugs and missing features

diff --git a/Engine/Plugins/Animation/ControlRig/Source/ControlRig/Private/ControlRigControlActor.cpp b/Engine/Plugins/Animation/ControlRig/Source/ControlRig/Private/ControlRigControlActor.cpp
--- a/Engine/Plugins/Animation/ControlRig/Source/ControlRig/Private/ControlRigControlActor.cpp
+++ b/Engine/Plugins/Animation/ControlRig/Source/ControlRig/Private/ControlRigControlActor.cpp
@@ -15,6 +15,7 @@ AControlRigControlActor::AControlRigControlActor(const FObjectInitializer& Objec
 	, bCastShadows(false)
 {
 	ActorRootComponent = CreateDefaultSubobject<USceneComponent>(TEXT("SceneComponent0"));
+	SetRootComponent(ActorRootComponent);
 
 	PrimaryActorTick.bCanEverTick = true;
 	PrimaryActorTick.bTickEvenWhenPaused = true;
@@ -37,21 +38,23 @@ AControlRigControlActor::~AControlRigControlActor()
 	RemoveUnbindDelegate();
 }
 
-void AControlRigControlActor::RemoveUnbindDelegate()
+void AControlRigControlActor::RemoveUnbindDelegate() const
 {
 	if (ControlRig)
 	{
 		if (TSharedPtr<IControlRigObjectBinding> Binding = ControlRig->GetObjectBinding())
 		{
-			if (OnUnbindDelegate.IsValid())
-			{
-				Binding->OnControlRigUnbind().Remove(OnUnbindDelegate);
-				OnUnbindDelegate.Reset();
-			}
+			Binding->OnControlRigUnbind().RemoveAll(this);
 		}
 	}
 }
 
+void AControlRigControlActor::HandleControlRigUnbind()
+{
+	Clear();
+	Refresh();
+}
+
 #if WITH_EDITOR
 
 void AControlRigControlActor::PostEditChangeProperty(FPropertyChangedEvent& PropertyChangedEvent)
@@ -65,9 +68,7 @@ void AControlRigControlActor::PostEditChangeProperty(FPropertyChangedEvent& Prop
 			PropertyChangedEvent.Property->GetFName() == GET_MEMBER_NAME_CHECKED(AControlRigControlActor, ColorParameter) ||
 			PropertyChangedEvent.Property->GetFName() == GET_MEMBER_NAME_CHECKED(AControlRigControlActor, bCastShadows)))
 	{
-		RemoveUnbindDelegate();
-		ControlRig = nullptr;
-		Clear();
+		ResetControlActor();
 		Refresh();
 	}
 }
@@ -86,9 +87,9 @@ void AControlRigControlActor::Tick(float DeltaSeconds)
 
 void AControlRigControlActor::Clear()
 {
-	TArray<USceneComponent*> ChildComponents;
 	if (ActorRootComponent)
 	{
+		TArray<USceneComponent*> ChildComponents;
 		ActorRootComponent->GetChildrenComponents(true, ChildComponents);
 		for (USceneComponent* Child : ChildComponents)
 		{
@@ -112,10 +113,16 @@ void AControlRigControlActor::Clear()
 	Materials.Reset();
 }
 
-void AControlRigControlActor::EndPlay(const EEndPlayReason::Type EndPlayReason)
+void AControlRigControlActor::ResetControlActor()
 {
 	RemoveUnbindDelegate();
 	ControlRig = nullptr;
+	Clear();
+}
+
+void AControlRigControlActor::EndPlay(const EEndPlayReason::Type EndPlayReason)
+{
+	ResetControlActor();
 	Super::EndPlay(EndPlayReason);
 }
 
@@ -123,8 +130,7 @@ void AControlRigControlActor::BeginDestroy()
 {
 	// since end play is not always called, we have to clear the delegate here
 	// clearing it at destructor might be too late as in some cases, the control rig was already GCed
-	RemoveUnbindDelegate();
-	ControlRig = nullptr;
+	ResetControlActor();
 	Super::BeginDestroy();
 }
 
@@ -135,6 +141,12 @@ void AControlRigControlActor::Refresh()
 		return;
 	}
 
+	// workaround to make sure we get a clean reset when the object was initialized from an existing template instance but some arrays didn't get populated somehow...
+	if (ControlNames.IsEmpty() || ControlNames.Num() != Components.Num() || ControlNames.Num() != Materials.Num() || ControlNames.Num() != ShapeTransforms.Num())
+	{
+		ResetControlActor();
+	}
+
 	if (ControlRig == nullptr)
 	{
 		TArray<UControlRig*> Rigs = UControlRig::FindControlRigs(ActorToTrack, ControlRigClass);
@@ -149,9 +161,9 @@ void AControlRigControlActor::Refresh()
 		}
 
 		RemoveUnbindDelegate();
-		if (TSharedPtr<IControlRigObjectBinding> Binding = ControlRig->GetObjectBinding())
+		if (const TSharedPtr<IControlRigObjectBinding> Binding = ControlRig->GetObjectBinding())
 		{
-			OnUnbindDelegate = Binding->OnControlRigUnbind().AddLambda([ this ]( ) { this->Clear(); this->Refresh(); });
+			Binding->OnControlRigUnbind().AddUObject(this, &AControlRigControlActor::HandleControlRigUnbind);
 		}
 
 		const TArray<TSoftObjectPtr<UControlRigShapeLibrary>>& ShapeLibraries = ControlRig->GetShapeLibraries();
@@ -246,7 +258,7 @@ void AControlRigControlActor::Refresh()
 		URigHierarchy* Hierarchy = ControlRig->GetHierarchy();
 
 		const FRigElementKey ControlKey(ControlNames[GizmoIndex], ERigElementType::Control);
-		FRigControlElement* ControlElement = Hierarchy->Find<FRigControlElement>(ControlKey);
+		const FRigControlElement* ControlElement = Hierarchy->Find<FRigControlElement>(ControlKey);
 		if(ControlElement == nullptr)
 		{
 			Components[GizmoIndex]->SetVisibility(false);
diff --git a/Engine/Plugins/Animation/ControlRig/Source/ControlRig/Public/ControlRigControlActor.h b/Engine/Plugins/Animation/ControlRig/Source/ControlRig/Public/ControlRigControlActor.h
--- a/Engine/Plugins/Animation/ControlRig/Source/ControlRig/Public/ControlRigControlActor.h
+++ b/Engine/Plugins/Animation/ControlRig/Source/ControlRig/Public/ControlRigControlActor.h
@@ -35,27 +35,30 @@ public:
 	virtual bool IsSelectable() const override { return bIsSelectable; }
 #endif
 
-	UPROPERTY(BlueprintReadWrite, EditAnywhere, Category = "Control Actor")
+	UPROPERTY(BlueprintReadWrite, EditAnywhere, Interp, Category = "Control Actor")
 	TObjectPtr<class AActor> ActorToTrack;
 
-	UPROPERTY(BlueprintReadWrite, EditAnywhere, Category = "Control Actor")
+	UPROPERTY(BlueprintReadWrite, EditAnywhere, Interp, Category = "Control Actor")
 	TSubclassOf<UControlRig> ControlRigClass;
 
-	UPROPERTY(BlueprintReadWrite, EditAnywhere, Category = "Control Actor")
+	UPROPERTY(BlueprintReadWrite, EditAnywhere, Interp, Category = "Control Actor")
 	bool bRefreshOnTick;
 
-	UPROPERTY(BlueprintReadWrite, EditAnywhere, Category = "Control Actor")
+	UPROPERTY(BlueprintReadWrite, EditAnywhere, Interp, Category = "Control Actor")
 	bool bIsSelectable;
 
-	UPROPERTY(BlueprintReadWrite, EditAnywhere, Category = "Materials")
+	UPROPERTY(BlueprintReadWrite, EditAnywhere, Interp, Category = "Materials")
 	TObjectPtr<UMaterialInterface> MaterialOverride;
 
-	UPROPERTY(BlueprintReadWrite, EditAnywhere, Category = "Materials")
+	UPROPERTY(BlueprintReadWrite, EditAnywhere, Interp, Category = "Materials")
 	FString ColorParameter;
 
-	UPROPERTY(BlueprintReadWrite, EditAnywhere, Category = "Materials")
+	UPROPERTY(BlueprintReadWrite, EditAnywhere, Interp, Category = "Materials")
 	bool bCastShadows;
 
+	UFUNCTION(BlueprintCallable, Category = "Control Actor", DisplayName="Reset")
+	void ResetControlActor();
+	
 	UFUNCTION(BlueprintCallable, Category = "Control Actor")
 	void Clear();
 
@@ -86,6 +89,7 @@ private:
 	FName ColorParameterName;
 
 private:
-	void RemoveUnbindDelegate();
-	FDelegateHandle OnUnbindDelegate;
+	void RemoveUnbindDelegate() const;
+	
+	void HandleControlRigUnbind();
 };
\ No newline at end of file
