Control Rig: Revert initialization of CachedMemoryHash

diff --git a/Engine/Plugins/Animation/ControlRig/Source/ControlRig/Private/ControlRig.cpp b/Engine/Plugins/Animation/ControlRig/Source/ControlRig/Private/ControlRig.cpp
--- a/Engine/Plugins/Animation/ControlRig/Source/ControlRig/Private/ControlRig.cpp
+++ b/Engine/Plugins/Animation/ControlRig/Source/ControlRig/Private/ControlRig.cpp
@@ -195,7 +195,7 @@ void UControlRig::Initialize(bool bInitRigUnits)
 	// recompute the hash used to differentiate VMs based on their memory layout 
 	if (HasAnyFlags(RF_ClassDefaultObject))
 	{
-		CachedMemoryHash = 0;
+		CachedMemoryHash = INDEX_NONE;
 
 		if(VM)
 		{
@@ -636,14 +636,6 @@ void UControlRig::InstantiateVMFromCDO()
 		}
 	}
 
-	if(VM)
-	{
-		CachedMemoryHash = HashCombine(
-			VM->GetLiteralMemory()->GetMemoryHash(),
-			VM->GetWorkMemory()->GetMemoryHash()
-		);
-	}
-
 	RequestInit();
 }
 
@@ -1481,7 +1473,7 @@ bool UControlRig::ExecuteUnits(FRigUnitContext& InOutContext, const FName& InEve
 				UControlRig* CDO = Cast<UControlRig>(GetClass()->GetDefaultObject());
 
 				bool bIsValidSnapshot = false;
-				if(SnapshotHash != 0)
+				if(SnapshotHash != INDEX_NONE)
 				{
 					TObjectPtr<URigVM>* InitializedVMSnapshotPtr = CDO->InitializedVMSnapshots.Find(SnapshotHash);
 					if(InitializedVMSnapshotPtr && !InitializedVMSnapshotPtr->IsNull())
@@ -3400,9 +3392,9 @@ void UControlRig::PostInitInstance(UControlRig* InCDO)
 
 uint32 UControlRig::GetHashForInitializeVMSnapShot()
 {
-	if(CachedMemoryHash == 0)
+	if(CachedMemoryHash == INDEX_NONE)
 	{
-		return 0;
+		return INDEX_NONE;
 	}
 	
 	uint32 Hash = GetHierarchy()->GetNameHash();
diff --git a/Engine/Plugins/Animation/ControlRig/Source/ControlRig/Public/ControlRig.h b/Engine/Plugins/Animation/ControlRig/Source/ControlRig/Public/ControlRig.h
--- a/Engine/Plugins/Animation/ControlRig/Source/ControlRig/Public/ControlRig.h
+++ b/Engine/Plugins/Animation/ControlRig/Source/ControlRig/Public/ControlRig.h
@@ -500,7 +500,7 @@ protected:
 	// and stores an initialized VM per hash.
 	UPROPERTY(transient)
 	TMap<uint32, TObjectPtr<URigVM>> InitializedVMSnapshots;
-	uint32 CachedMemoryHash = 0;
+	uint32 CachedMemoryHash = INDEX_NONE;
 
 	// computes the hash used to store / find VM snapshots
 	uint32 GetHashForInitializeVMSnapShot();
