localizable text

diff --git a/Engine/Plugins/Animation/ControlRig/Source/ControlRigDeveloper/Private/Graph/ControlRigGraphSchema.cpp b/Engine/Plugins/Animation/ControlRig/Source/ControlRigDeveloper/Private/Graph/ControlRigGraphSchema.cpp
--- a/Engine/Plugins/Animation/ControlRig/Source/ControlRigDeveloper/Private/Graph/ControlRigGraphSchema.cpp
+++ b/Engine/Plugins/Animation/ControlRig/Source/ControlRigDeveloper/Private/Graph/ControlRigGraphSchema.cpp
@@ -636,6 +636,8 @@ bool UControlRigGraphSchema::TryCreateConnection(UEdGraphPin* PinA, UEdGraphPin*
 
 #if WITH_EDITOR
 
+			static const FText LoopMessage = LOCTEXT("LinkingLoopNotRecommended", "Linking a function return within a loop is not recommended.\nAre you sure?");
+			
 			// check if we are trying to connect a loop iteration pin to a return
 			if(URigVMGraph* Graph = Controller->GetGraph())
 			{
@@ -698,7 +700,7 @@ bool UControlRigGraphSchema::TryCreateConnection(UEdGraphPin* PinA, UEdGraphPin*
 
 						if(bIsInLoopIteration)
 						{
-							const EAppReturnType::Type Answer = FMessageDialog::Open( EAppMsgType::YesNo, FText::FromString( TEXT("Linking a function return within a loop is not recommended.\nAre you sure?") ) );
+							const EAppReturnType::Type Answer = FMessageDialog::Open( EAppMsgType::YesNo, LoopMessage );
 							if(Answer == EAppReturnType::No)
 							{
 								return false;
diff --git a/Engine/Plugins/Animation/ControlRig/Source/ControlRigEditor/Private/ControlRigBlueprintActions.cpp b/Engine/Plugins/Animation/ControlRig/Source/ControlRigEditor/Private/ControlRigBlueprintActions.cpp
--- a/Engine/Plugins/Animation/ControlRig/Source/ControlRigEditor/Private/ControlRigBlueprintActions.cpp
+++ b/Engine/Plugins/Animation/ControlRig/Source/ControlRigEditor/Private/ControlRigBlueprintActions.cpp
@@ -93,11 +93,13 @@ void FControlRigBlueprintActions::PerformAssetDiff(UObject* OldAsset, UObject* N
 	UBlueprint* OldBlueprint = CastChecked<UBlueprint>(OldAsset);
 	UBlueprint* NewBlueprint = CastChecked<UBlueprint>(NewAsset);
 
+	static const FText DiffWindowMessage = LOCTEXT("ControlRigDiffWindow", "Opening a diff window will close the control rig editor. {0}.\nAre you sure?");
+	
 	UAssetEditorSubsystem* AssetEditorSubsystem = GEditor->GetEditorSubsystem<UAssetEditorSubsystem>();
 	for (IAssetEditorInstance* Editor : AssetEditorSubsystem->FindEditorsForAsset(OldAsset))
 	{
 		const EAppReturnType::Type Answer = FMessageDialog::Open( EAppMsgType::YesNo,
-				FText::FromString(FString::Printf(TEXT("Opening a diff window will close the control rig editor. %s.\nAre you sure?"),  *OldBlueprint->GetName() )));
+				FText::Format(DiffWindowMessage, FText::FromString(OldBlueprint->GetName())));
 		if(Answer == EAppReturnType::No)
 		{
 		   return;
@@ -106,7 +108,7 @@ void FControlRigBlueprintActions::PerformAssetDiff(UObject* OldAsset, UObject* N
 	for (IAssetEditorInstance* Editor : AssetEditorSubsystem->FindEditorsForAsset(NewAsset))
 	{
 		const EAppReturnType::Type Answer = FMessageDialog::Open( EAppMsgType::YesNo,
-				FText::FromString(FString::Printf(TEXT("Opening a diff window will close the control rig editor. %s.\nAre you sure?"),  *NewBlueprint->GetName() )));
+				FText::Format(DiffWindowMessage, FText::FromString(NewBlueprint->GetName())));
 		if(Answer == EAppReturnType::No)
 		{
 			return;
