do not add interaction event if we're in construction mode

diff --git a/Engine/Plugins/Animation/ControlRig/Source/ControlRig/Private/ControlRig.cpp b/Engine/Plugins/Animation/ControlRig/Source/ControlRig/Private/ControlRig.cpp
--- a/Engine/Plugins/Animation/ControlRig/Source/ControlRig/Private/ControlRig.cpp
+++ b/Engine/Plugins/Animation/ControlRig/Source/ControlRig/Private/ControlRig.cpp
@@ -300,7 +300,7 @@ void UControlRig::Evaluate_AnyThread()
 		{
 			EventQueueToRun.Add(FRigUnit_InteractionExecution::EventName);
 		}
-		else
+		else if(!EventQueueToRun.Contains(FRigUnit_PrepareForExecution::EventName))
 		{
 			// insert just before the last event so the interaction runs prior to
 			// forward solve or backwards solve.
