Control Rig: Add note on setting control transforms in in correct order

diff --git a/Engine/Plugins/Animation/ControlRig/Source/ControlRig/Private/Rigs/RigHierarchy.cpp b/Engine/Plugins/Animation/ControlRig/Source/ControlRig/Private/Rigs/RigHierarchy.cpp
--- a/Engine/Plugins/Animation/ControlRig/Source/ControlRig/Private/Rigs/RigHierarchy.cpp
+++ b/Engine/Plugins/Animation/ControlRig/Source/ControlRig/Private/Rigs/RigHierarchy.cpp
@@ -2755,6 +2755,44 @@ void URigHierarchy::SetTransform(FRigTransformElement* InTransformElement, const
 			InTransformElement->GetIndex(),
 			InTransformType
 		);
+
+		// if we are setting a control / null parent after a child here - let's let the user know
+		if(InTransformElement->IsA<FRigControlElement>() || InTransformElement->IsA<FRigNullElement>())
+		{
+			if(const UWorld* World = GetWorld())
+			{
+				// only fire these notes if we are inside the asset editor
+				if(GetWorld()->WorldType == EWorldType::EditorPreview)
+				{
+					const FRigBaseElementChildrenArray& Children = GetChildren(InTransformElement);
+					for (FRigBaseElement* Child : Children)
+					{
+						const bool bChildFound = WrittenTransformsAtRuntime.ContainsByPredicate([Child](const TInstructionSliceElement& Entry) -> bool
+						{
+							return Entry.Get<2>() == Child->GetIndex();
+						});
+
+						if(bChildFound)
+						{
+							const FRigUnitContext* UnitContext = (const FRigUnitContext*)ExecuteContext->OpaqueArguments[0];
+							if(UnitContext && UnitContext->Log)
+							{
+								static constexpr TCHAR MessageFormat[] = TEXT("Setting transform of parent (%s) after setting child (%s).\nThis may lead to unexpected results.");
+								const FString& Message = FString::Printf(
+									MessageFormat,
+									*InTransformElement->GetName().ToString(),
+									*Child->GetName().ToString());
+								UnitContext->Log->Report(
+									EMessageSeverity::Info,
+									ExecuteContext->PublicData.GetFunctionName(),
+									ExecuteContext->PublicData.GetInstructionIndex(),
+									Message);
+							}
+						}
+					}
+				}
+			}
+		}
 	}
 	TGuardValue<bool> RecordTransformsPerInstructionGuard(bRecordTransformsAtRuntime, false);
 	
