Control Rig: Add note on setting control transforms in in correct order (part 2)

Follow up to 21422971

diff --git a/Engine/Plugins/Animation/ControlRig/Source/ControlRig/Private/Rigs/RigHierarchy.cpp b/Engine/Plugins/Animation/ControlRig/Source/ControlRig/Private/Rigs/RigHierarchy.cpp
--- a/Engine/Plugins/Animation/ControlRig/Source/ControlRig/Private/Rigs/RigHierarchy.cpp
+++ b/Engine/Plugins/Animation/ControlRig/Source/ControlRig/Private/Rigs/RigHierarchy.cpp
@@ -2762,7 +2762,7 @@ void URigHierarchy::SetTransform(FRigTransformElement* InTransformElement, const
 			if(const UWorld* World = GetWorld())
 			{
 				// only fire these notes if we are inside the asset editor
-				if(GetWorld()->WorldType == EWorldType::EditorPreview)
+				if(World->WorldType == EWorldType::EditorPreview)
 				{
 					const FRigBaseElementChildrenArray& Children = GetChildren(InTransformElement);
 					for (FRigBaseElement* Child : Children)
