Control Rig: Fix FK Control Rig collapsing on additive mode

diff --git a/Engine/Plugins/Animation/ControlRig/Source/ControlRig/Private/Rigs/FKControlRig.cpp b/Engine/Plugins/Animation/ControlRig/Source/ControlRig/Private/Rigs/FKControlRig.cpp
--- a/Engine/Plugins/Animation/ControlRig/Source/ControlRig/Private/Rigs/FKControlRig.cpp
+++ b/Engine/Plugins/Animation/ControlRig/Source/ControlRig/Private/Rigs/FKControlRig.cpp
@@ -494,15 +494,11 @@ void UFKControlRig::ToggleApplyMode()
 		Context.SetKey = EControlRigSetKey::Never;
 		const bool bSetupUndo = false;
 
-		FTransform ZeroScale = FTransform::Identity;
-		ZeroScale.SetScale3D(FVector::ZeroVector);
-		const FEulerTransform EulerZero(ZeroScale);
-
 		GetHierarchy()->ForEach<FRigControlElement>([&](FRigControlElement* ControlElement) -> bool
         {
             if (ControlElement->Settings.ControlType == ERigControlType::EulerTransform)
 			{
-				SetControlValue<FEulerTransform>(ControlElement->GetName(), EulerZero, true, Context, bSetupUndo);
+				SetControlValue<FEulerTransform>(ControlElement->GetName(), FEulerTransform::Identity, true, Context, bSetupUndo);
 			}
 			else if (ControlElement->Settings.ControlType == ERigControlType::Float)
 			{
