Control Rig: Sync UpdateTemplateNodePinTypes from UE5Main

diff --git a/Engine/Plugins/Animation/ControlRig/Source/ControlRigDeveloper/Private/ControlRigBlueprint.cpp b/Engine/Plugins/Animation/ControlRig/Source/ControlRigDeveloper/Private/ControlRigBlueprint.cpp
--- a/Engine/Plugins/Animation/ControlRig/Source/ControlRigDeveloper/Private/ControlRigBlueprint.cpp
+++ b/Engine/Plugins/Animation/ControlRig/Source/ControlRigDeveloper/Private/ControlRigBlueprint.cpp
@@ -1452,6 +1452,21 @@ void UControlRigBlueprint::RefreshAllModels(EControlRigBlueprintLoadType InLoadT
 					TemplateNode->InvalidateCache();
 					TemplateNode->PostLoad();
 
+					// Remove preferred types that do not match the pin type
+					{
+						TemplateNode->PreferredPermutationPairs.RemoveAll([TemplateNode](const FRigVMTemplatePreferredType& PreferredType)
+						{
+							if (URigVMPin* Pin = TemplateNode->FindPin(PreferredType.GetArgument().ToString()))
+							{
+								if (Pin->GetTypeIndex() == PreferredType.GetTypeIndex())
+								{
+									return false;
+								}
+							}
+							return true;
+						});						
+					}
+
 					if(!TemplateNode->HasWildCardPin() &&
 						!TemplateNode->IsSingleton() &&
 						TemplateNode->PreferredPermutationPairs.IsEmpty())
