ensure that bone transforms are properly updated on the mapped component

diff --git a/Engine/Plugins/Animation/ControlRig/Source/ControlRig/Private/ControlRigComponent.cpp b/Engine/Plugins/Animation/ControlRig/Source/ControlRig/Private/ControlRigComponent.cpp
--- a/Engine/Plugins/Animation/ControlRig/Source/ControlRig/Private/ControlRigComponent.cpp
+++ b/Engine/Plugins/Animation/ControlRig/Source/ControlRig/Private/ControlRigComponent.cpp
@@ -344,13 +344,22 @@ void UControlRigComponent::Initialize()
 	// after the control rig component such that by the time they tick they can
 	// send the latest data for rendering
 	USceneComponent* LastComponent = nullptr;
-	for (FControlRigComponentMappedElement& MappedElement : MappedElements)
+	for (const FControlRigComponentMappedElement& MappedElement : MappedElements)
 	{
-		if (LastComponent != MappedElement.SceneComponent && MappedElement.SceneComponent != nullptr)
+		if (MappedElement.SceneComponent)
 		{
-			
-			MappedElement.SceneComponent->AddTickPrerequisiteComponent(this);
-			LastComponent = MappedElement.SceneComponent;
+			if (LastComponent != MappedElement.SceneComponent)
+			{	
+				MappedElement.SceneComponent->AddTickPrerequisiteComponent(this);
+				LastComponent = MappedElement.SceneComponent;
+			}
+
+			// make sure that the animation is updated so that bone transforms are updated on the mapped component
+			// (otherwise, FControlRigAnimInstanceProxy::Evaluate is never called when moving a control in the editor)
+			if (USkeletalMeshComponent* Component = Cast<USkeletalMeshComponent>(MappedElement.SceneComponent))
+			{
+				Component->SetUpdateAnimationInEditor(bUpdateInEditor);
+			}
 		}
 	}
 }
@@ -1455,7 +1464,7 @@ void UControlRigComponent::TransferOutputs()
 		USceneComponent* LastComponent = nullptr;
 		FControlRigAnimInstanceProxy* Proxy = nullptr;
 
-		for (FControlRigComponentMappedElement& MappedElement : MappedElements)
+		for (const FControlRigComponentMappedElement& MappedElement : MappedElements)
 		{
 			if (LastComponent != MappedElement.SceneComponent || Proxy == nullptr)
 			{
@@ -1469,9 +1478,7 @@ void UControlRigComponent::TransferOutputs()
 			}
 		}
 
-		TArray<USkeletalMeshComponent*> ComponentsToTick;
-
-		for (FControlRigComponentMappedElement& MappedElement : MappedElements)
+		for (const FControlRigComponentMappedElement& MappedElement : MappedElements)
 		{
 			if (MappedElement.ElementIndex == INDEX_NONE || MappedElement.Direction == EControlRigComponentMapDirection::Input)
 			{
@@ -1500,7 +1507,6 @@ void UControlRigComponent::TransferOutputs()
 
 					if (Proxy && (MappedElement.SceneComponent != nullptr))
 					{
-						ComponentsToTick.AddUnique(Cast<USkeletalMeshComponent>(MappedElement.SceneComponent));
 						if(MappedElement.Space == EControlRigComponentSpace::WorldSpace)
 						{
 							Transform = Transform.GetRelativeTransform(MappedElement.SceneComponent->GetComponentToWorld());
@@ -1546,7 +1552,6 @@ void UControlRigComponent::TransferOutputs()
 
 					if (Proxy)
 					{
-						ComponentsToTick.AddUnique(Cast<USkeletalMeshComponent>(MappedElement.SceneComponent));
 						Proxy->StoredCurves.FindOrAdd((SmartName::UID_Type)MappedElement.SubIndex) = ControlRig->GetHierarchy()->GetCurveValue(MappedElement.ElementIndex);
 					}
 				}
