Sequencer: Control Rig:  Fix Redo and Undo issues.  Due to selections changing on next tick with Sequencer, when you undo or redo could be creating a transaction so we only do so if the selection is really changing.  This now makes selection fully undoable/redoable.  Finally for UE-160921 needed to make sure we do the correct selection check on sub Shapes and not Animation channels.

diff --git a/Engine/Plugins/Animation/ControlRig/Source/ControlRigEditor/Private/EditMode/ControlRigEditMode.cpp b/Engine/Plugins/Animation/ControlRig/Source/ControlRigEditor/Private/EditMode/ControlRigEditMode.cpp
--- a/Engine/Plugins/Animation/ControlRig/Source/ControlRigEditor/Private/EditMode/ControlRigEditMode.cpp
+++ b/Engine/Plugins/Animation/ControlRig/Source/ControlRigEditor/Private/EditMode/ControlRigEditMode.cpp
@@ -962,9 +962,7 @@ bool FControlRigEditMode::ProcessCapturedMouseMoves(FEditorViewportClient* InVie
 				int32 X = CapturedMouseMoves[Index].X;
 				if (StartXValue.IsSet() == false)
 				{
-					bisTrackingAnimToolDrag = true;
-					GEditor->BeginTransaction(LOCTEXT("AnimSliderBlend", "AnimSlider Blend"));
-					StartXValue = X;
+					StartAnimSliderTool(X);
 				}
 				else
 				{
@@ -973,7 +971,7 @@ bool FControlRigEditMode::ProcessCapturedMouseMoves(FEditorViewportClient* InVie
 					{
 						FIntPoint Origin, Size;
 						InViewportClient->GetViewportDimensions(Origin, Size);
-						const double ViewPortSize = (double)Size.X * 0.25; // 1/4 screen drag should do full blend, todo perhaps expose this as a sensitivity
+						const double ViewPortSize = (double)Size.X * 0.125; // 1/8 screen drag should do full blend, todo perhaps expose this as a sensitivity
 						const double ViewDiff = (double)(Diff) / ViewPortSize;
 						DragAnimSliderTool(ViewDiff);
 					}
@@ -1283,7 +1281,6 @@ bool FControlRigEditMode::HandleClick(FEditorViewportClient* InViewportClient, H
 						{
 							if (GEditor && GEditor->GetSelectedActorCount())
 							{
-								const FScopedTransaction Transaction(NSLOCTEXT("Sequencer", "UpdatingActorComponentSelectionNone", "Select None"), !GIsTransacting);
 								GEditor->SelectNone(false, true);
 								GEditor->NoteSelectionChange();
 							}
@@ -2566,6 +2563,14 @@ void FControlRigEditMode::DragAnimSliderTool(double IncrementVal)
 	StaticCastSharedPtr<FControlRigEditModeToolkit>(Toolkit)->DragAnimSliderTool(IncrementVal);
 }
 
+void FControlRigEditMode::StartAnimSliderTool(int32 InX)
+{
+	bisTrackingAnimToolDrag = true;
+	GEditor->BeginTransaction(LOCTEXT("AnimSliderBlend", "AnimSlider Blend"));
+	StartXValue = InX;
+	StaticCastSharedPtr<FControlRigEditModeToolkit>(Toolkit)->StartAnimSliderTool();
+}
+
 /** Reset and stop user the anim slider tool*/
 void FControlRigEditMode::ResetAnimSlider()
 {
@@ -3458,7 +3463,6 @@ void FControlRigEditMode::OnHierarchyModified(ERigHierarchyNotification InNotif,
 					// if it's control
 					if (Key.Type == ERigElementType::Control)
 					{
-						FScopedTransaction ScopedTransaction(LOCTEXT("SelectControlTransaction", "Select Control"), IsInLevelEditor() && !GIsTransacting);	
 						if (IsInLevelEditor())
 						{
 							ControlProxy->Modify();
diff --git a/Engine/Plugins/Animation/ControlRig/Source/ControlRigEditor/Private/EditMode/ControlRigEditMode.h b/Engine/Plugins/Animation/ControlRig/Source/ControlRigEditor/Private/EditMode/ControlRigEditMode.h
--- a/Engine/Plugins/Animation/ControlRig/Source/ControlRigEditor/Private/EditMode/ControlRigEditMode.h
+++ b/Engine/Plugins/Animation/ControlRig/Source/ControlRigEditor/Private/EditMode/ControlRigEditMode.h
@@ -215,10 +215,10 @@ protected:
 
 	/** If Anim Slider is open, then can drag*/
 	bool CanChangeAnimSliderTool() const;
-
+	/** Start Slider Tool*/
+	void StartAnimSliderTool(int32 InX);
 	/** If Anim Slider is open, drag the tool*/
 	void DragAnimSliderTool(double Val);
-
 	/** Reset and stop user the anim slider tool*/
 	void ResetAnimSlider();
 
diff --git a/Engine/Plugins/Animation/ControlRig/Source/ControlRigEditor/Private/EditMode/ControlRigEditModeToolkit.cpp b/Engine/Plugins/Animation/ControlRig/Source/ControlRigEditor/Private/EditMode/ControlRigEditModeToolkit.cpp
--- a/Engine/Plugins/Animation/ControlRig/Source/ControlRigEditor/Private/EditMode/ControlRigEditModeToolkit.cpp
+++ b/Engine/Plugins/Animation/ControlRig/Source/ControlRigEditor/Private/EditMode/ControlRigEditModeToolkit.cpp
@@ -267,6 +267,14 @@ void FControlRigEditModeToolkit::ResetAnimSlider()
 		TweenWidget->ResetAnimSlider();
 	}
 }
+
+void FControlRigEditModeToolkit::StartAnimSliderTool()
+{
+	if (TweenWidgetParent && TweenWidget)
+	{
+		TweenWidget->StartAnimSliderTool();
+	}
+}
 void FControlRigEditModeToolkit::TryShowTweenOverlay()
 {
 	if (TweenWidgetParent)
diff --git a/Engine/Plugins/Animation/ControlRig/Source/ControlRigEditor/Private/EditMode/ControlRigEditModeToolkit.h b/Engine/Plugins/Animation/ControlRig/Source/ControlRigEditor/Private/EditMode/ControlRigEditModeToolkit.h
--- a/Engine/Plugins/Animation/ControlRig/Source/ControlRigEditor/Private/EditMode/ControlRigEditModeToolkit.h
+++ b/Engine/Plugins/Animation/ControlRig/Source/ControlRigEditor/Private/EditMode/ControlRigEditModeToolkit.h
@@ -57,7 +57,7 @@ public:
 	bool CanChangeAnimSliderTool() const;
 	void DragAnimSliderTool(double Val);
 	void ResetAnimSlider();
-
+	void StartAnimSliderTool();
 public:
 	static const FName PoseTabName;
 	static const FName MotionTrailTabName;
diff --git a/Engine/Plugins/Animation/ControlRig/Source/ControlRigEditor/Private/EditMode/SControlRigOutliner.cpp b/Engine/Plugins/Animation/ControlRig/Source/ControlRigEditor/Private/EditMode/SControlRigOutliner.cpp
--- a/Engine/Plugins/Animation/ControlRig/Source/ControlRigEditor/Private/EditMode/SControlRigOutliner.cpp
+++ b/Engine/Plugins/Animation/ControlRig/Source/ControlRigEditor/Private/EditMode/SControlRigOutliner.cpp
@@ -1087,11 +1087,18 @@ void SControlRigOutliner::HandleSelectionChanged(TSharedPtr<FMultiRigTreeElement
 		}
 	}
 	TGuardValue<bool> GuardRigHierarchyChanges(bIsChangingRigHierarchy, true);
+	FControlRigEditMode* EditMode = static_cast<FControlRigEditMode*>(ModeTools->GetActiveMode(FControlRigEditMode::ModeName));
+	bool bEndTransaction = false;
+	if (GEditor && EditMode && EditMode->IsInLevelEditor())
+	{
+		GEditor->BeginTransaction(LOCTEXT("SelectControl", "Select Control"));
+		bEndTransaction = true;
+	}
 	//due to how Sequencer Tree View will redo selection on next tick if we aren't keeping or toggling selection we need to clear it out
 	if (FSlateApplication::Get().GetModifierKeys().IsShiftDown() == false || FSlateApplication::Get().GetModifierKeys().IsControlDown() == false)
 	{
-		if (FControlRigEditMode* EditMode = static_cast<FControlRigEditMode*>(ModeTools->GetActiveMode(FControlRigEditMode::ModeName)))
-		{
+		if (EditMode)
+		{		
 			TMap<UControlRig*, TArray<FRigElementKey>> SelectedControls;
 			EditMode->GetAllSelectedControls(SelectedControls);
 			for (TPair<UControlRig*, TArray<FRigElementKey>>& CurrentSelection : SelectedControls)
@@ -1113,10 +1120,18 @@ void SControlRigOutliner::HandleSelectionChanged(TSharedPtr<FMultiRigTreeElement
 			check(Controller);
 			if (!Controller->SetSelection(RigAndKeys.Value))
 			{
+				if (bEndTransaction)
+				{
+					GEditor->EndTransaction();
+				}
 				return;
 			}
 		}
 	}
+	if (bEndTransaction)
+	{
+		GEditor->EndTransaction();
+	}
 }
 
 void SControlRigOutliner::SetEditMode(FControlRigEditMode& InEditMode)
diff --git a/Engine/Plugins/Animation/ControlRig/Source/ControlRigEditor/Private/EditMode/SControlRigTweenWidget.cpp b/Engine/Plugins/Animation/ControlRig/Source/ControlRigEditor/Private/EditMode/SControlRigTweenWidget.cpp
--- a/Engine/Plugins/Animation/ControlRig/Source/ControlRigEditor/Private/EditMode/SControlRigTweenWidget.cpp
+++ b/Engine/Plugins/Animation/ControlRig/Source/ControlRigEditor/Private/EditMode/SControlRigTweenWidget.cpp
@@ -359,5 +359,12 @@ void SControlRigTweenWidget::ResetAnimSlider()
 	}
 }
 
+void SControlRigTweenWidget::StartAnimSliderTool()
+{
+	if (SliderWidget.IsValid())
+	{
+		SliderWidget->Setup();
+	}
+}
 
 #undef LOCTEXT_NAMESPACE
diff --git a/Engine/Plugins/Animation/ControlRig/Source/ControlRigEditor/Private/EditMode/SControlRigTweenWidget.h b/Engine/Plugins/Animation/ControlRig/Source/ControlRigEditor/Private/EditMode/SControlRigTweenWidget.h
--- a/Engine/Plugins/Animation/ControlRig/Source/ControlRigEditor/Private/EditMode/SControlRigTweenWidget.h
+++ b/Engine/Plugins/Animation/ControlRig/Source/ControlRigEditor/Private/EditMode/SControlRigTweenWidget.h
@@ -33,6 +33,8 @@ class SControlRigTweenSlider : public SCompoundWidget
 	void Construct(const FArguments& InArgs);
 	void SetAnimSlider(TSharedPtr<FBaseAnimSlider>& InAnimSlider) { AnimSlider = InAnimSlider; }
 	void DragAnimSliderTool(double Val);
+	bool Setup();
+
 	void ResetAnimSlider();
 private:
 
@@ -46,7 +48,6 @@ private:
 	void OnEndSliderMovement(double NewValue);
 	double OnGetPoseBlendValue() const { return PoseBlendValue; }
 
-	bool Setup();
 	double PoseBlendValue;
 	bool bIsBlending;
 	bool bSliderStartedTransaction;
@@ -73,7 +74,7 @@ class SControlRigTweenWidget : public SCompoundWidget
 	void GetToNextActiveSlider();
 	void DragAnimSliderTool(double Val);
 	void ResetAnimSlider();
-
+	void StartAnimSliderTool();
 private:
 
 	void OnSelectSliderTool(int32 Index);
diff --git a/Engine/Plugins/Animation/ControlRig/Source/ControlRigEditor/Private/Sequencer/ControlRigParameterTrackEditor.cpp b/Engine/Plugins/Animation/ControlRig/Source/ControlRigEditor/Private/Sequencer/ControlRigParameterTrackEditor.cpp
--- a/Engine/Plugins/Animation/ControlRig/Source/ControlRigEditor/Private/Sequencer/ControlRigParameterTrackEditor.cpp
+++ b/Engine/Plugins/Animation/ControlRig/Source/ControlRigEditor/Private/Sequencer/ControlRigParameterTrackEditor.cpp
@@ -1962,8 +1962,6 @@ void FControlRigParameterTrackEditor::OnSelectionChanged(TArray<UMovieSceneTrack
 	TArray<const IKeyArea*> KeyAreas;
 	const bool UseSelectedKeys = CVarSelectedKeysSelectControls.GetValueOnGameThread();
 	GetSequencer()->GetSelectedKeyAreas(KeyAreas, UseSelectedKeys);
-	FScopedTransaction ScopedTransaction(LOCTEXT("SelectControlTransaction", "Select Control"), !GIsTransacting);
-
 	if (KeyAreas.Num() <= 0)
 	{
 		if (ControlRigEditMode)
@@ -1973,8 +1971,9 @@ void FControlRigParameterTrackEditor::OnSelectionChanged(TArray<UMovieSceneTrack
 			for (TPair<UControlRig*, TArray<FRigElementKey>>& SelectedControl : AllSelectedControls)
 			{
 				ControlRig = SelectedControl.Key;
-				if (ControlRig)
+				if (ControlRig && ControlRig->CurrentControlSelection().Num() > 0)
 				{
+					FScopedTransaction ScopedTransaction(LOCTEXT("SelectControlTransaction", "Select Control"), !GIsTransacting);
 					ControlRig->ClearControlSelection();
 				}
 			}
@@ -2008,7 +2007,6 @@ void FControlRigParameterTrackEditor::OnSelectionChanged(TArray<UMovieSceneTrack
 		}
 		return;
 	}
-
 	SelectRigsAndControls(ControlRig, KeyAreas);
 	
 }
@@ -2018,6 +2016,7 @@ void FControlRigParameterTrackEditor::SelectRigsAndControls(UControlRig* Control
 	FControlRigEditMode* ControlRigEditMode = GetEditMode();
 
 	TArray<FString> StringArray;
+	//we have two sets here one to see if selection has really changed that contains the attirbutes, the other to select just the parent
 	TMap<UControlRig*, TSet<FName>> RigsAndControls;
 	for (const IKeyArea* KeyArea : KeyAreas)
 	{
@@ -2062,32 +2061,41 @@ void FControlRigParameterTrackEditor::SelectRigsAndControls(UControlRig* Control
 					// skip nested controls which have the shape enabled flag turned on
 					if(const URigHierarchy* Hierarchy = ControlRig->GetHierarchy())
 					{
+
 						if(const FRigControlElement* ControlElement = Hierarchy->Find<FRigControlElement>(FRigElementKey(ControlName, ERigElementType::Control)))
 						{
-							if(ControlElement->IsAnimationChannel())
+						
+							if (ControlElement->Settings.ControlType == ERigControlType::Bool ||
+								ControlElement->Settings.ControlType == ERigControlType::Float ||
+								ControlElement->Settings.ControlType == ERigControlType::Integer)
 							{
-								if(const FRigControlElement* ParentControlElement = Cast<FRigControlElement>(Hierarchy->GetFirstParent(ControlElement)))
+								if (ControlElement->Settings.SupportsShape() || !Hierarchy->IsAnimatable(ControlElement))
 								{
-									if(const TSet<FName>* Controls = RigsAndControls.Find(ControlRig))
+
+									if (const FRigControlElement* ParentControlElement = Cast<FRigControlElement>(Hierarchy->GetFirstParent(ControlElement)))
 									{
-										if(Controls->Contains(ParentControlElement->GetName()))
+										if (const TSet<FName>* Controls = RigsAndControls.Find(ControlRig))
 										{
-											continue;
+											if (Controls->Contains(ParentControlElement->GetName()))
+											{
+												continue;
+											}
 										}
 									}
 								}
 							}
+							RigsAndControls.FindOrAdd(ControlRig).Add(ControlName);
 						}
 					}
-					
-					RigsAndControls.FindOrAdd(ControlRig).Add(ControlName);
 				}
 			}
 		}
 	}
-
-	ControlRig = nullptr;
-	//Always clear the control rig(s) in the edit mode.
+	//only create transaction if selection is really different.
+	bool bEndTransaction = false;
+	
+	TMap<UControlRig*, TArray<FName>> ControlRigsToClearSelection;
+	//get current selection which we will clear if different
 	if (ControlRigEditMode)
 	{
 		TMap<UControlRig*, TArray<FRigElementKey>> SelectedControls;
@@ -2097,20 +2105,88 @@ void FControlRigParameterTrackEditor::SelectRigsAndControls(UControlRig* Control
 			ControlRig = Selection.Key;
 			if (ControlRig)
 			{
-				ControlRig->ClearControlSelection();
+				TArray<FName> SelectedControlNames = ControlRig->CurrentControlSelection();
+				ControlRigsToClearSelection.Add(ControlRig, SelectedControlNames);
 			}
 		}
 	}
-	for (TPair<UControlRig*, TSet<FName>> Pair : RigsAndControls)
+
+	for (TPair<UControlRig*, TSet<FName>>& Pair : RigsAndControls)
 	{
-		if (Pair.Key != ControlRig)
+		//check to see if new selection is same als old selection
+		bool bIsSame = true;
+		if (TArray<FName>* SelectedNames = ControlRigsToClearSelection.Find(Pair.Key))
+		{
+			TSet<FName>* FullNames = RigsAndControls.Find(Pair.Key);
+			if (!FullNames)
+			{
+				continue; // should never happen
+			}
+			if (SelectedNames->Num() != FullNames->Num())
+			{ 
+				bIsSame = false;
+				if (bEndTransaction == false)
+				{
+					bEndTransaction = true;
+					GEditor->BeginTransaction(LOCTEXT("SelectControl", "Select Control"));
+				}
+				Pair.Key->ClearControlSelection();
+				ControlRigsToClearSelection.Remove(Pair.Key); //remove it
+			}
+			else//okay if same check and see if equal...
+			{
+				for (const FName& Name : (*SelectedNames))
+				{
+					if (FullNames->Contains(Name) == false)
+					{
+						bIsSame = false;
+						if (bEndTransaction == false)
+						{
+							bEndTransaction = true;
+							GEditor->BeginTransaction(LOCTEXT("SelectControl", "Select Control"));
+						}
+						Pair.Key->ClearControlSelection();
+						ControlRigsToClearSelection.Remove(Pair.Key); //remove it
+						break; //break out
+					}
+				}
+			}
+			if (bIsSame == true)
+			{
+				ControlRigsToClearSelection.Remove(Pair.Key); //remove it
+			}
+		}
+		else
+		{
+			bIsSame = false;
+		}
+		if (bIsSame == false)
 		{
-			Pair.Key->ClearControlSelection();
+			for (const FName& Name : Pair.Value)
+			{
+				if (bEndTransaction == false)
+				{
+					bEndTransaction = true;
+					GEditor->BeginTransaction(LOCTEXT("SelectControl", "Select Control"));
+				}
+				Pair.Key->SelectControl(Name, true);
+			}
 		}
-		for (const FName& Name : Pair.Value)
+	}
+	//go through and clear those still not cleared
+	for (TPair<UControlRig*, TArray<FName>>& SelectedPairs : ControlRigsToClearSelection)
+	{
+		if (bEndTransaction == false)
 		{
-			Pair.Key->SelectControl(Name, true);
+			bEndTransaction = true;
+			GEditor->BeginTransaction(LOCTEXT("SelectControl", "Select Control"));
 		}
+		SelectedPairs.Key->ClearControlSelection();
+	}
+
+	if (bEndTransaction)
+	{
+		GEditor->EndTransaction();
 	}
 }
 
