Fix control rig display actor crash

diff --git a/Engine/Plugins/Animation/ControlRig/Source/ControlRig/Private/ControlRigControlActor.cpp b/Engine/Plugins/Animation/ControlRig/Source/ControlRig/Private/ControlRigControlActor.cpp
--- a/Engine/Plugins/Animation/ControlRig/Source/ControlRig/Private/ControlRigControlActor.cpp
+++ b/Engine/Plugins/Animation/ControlRig/Source/ControlRig/Private/ControlRigControlActor.cpp
@@ -95,15 +95,21 @@ void AControlRigControlActor::Clear()
 		{
 			if (UStaticMeshComponent* StaticMeshComponent = Cast<UStaticMeshComponent>(Child))
 			{
-				Components.AddUnique(StaticMeshComponent);
+				if (IsValid(StaticMeshComponent))
+				{
+					Components.AddUnique(StaticMeshComponent);	
+				}
 			}
 		}
 
 		for (UStaticMeshComponent* Component : Components)
 		{
-			Component->DetachFromComponent(FDetachmentTransformRules::KeepRelativeTransform);
-			Component->UnregisterComponent();
-			Component->DestroyComponent();
+			if (IsValid(Component))
+			{
+				Component->DetachFromComponent(FDetachmentTransformRules::KeepRelativeTransform);
+				Component->UnregisterComponent();
+				Component->DestroyComponent();	
+			}
 		}
 	}
 
