Control Rig: Fix unresolved deprecated nodes are not reporting an error if not involved in compilation

diff --git a/Engine/Plugins/Animation/ControlRig/Source/ControlRigDeveloper/Private/ControlRigBlueprint.cpp b/Engine/Plugins/Animation/ControlRig/Source/ControlRigDeveloper/Private/ControlRigBlueprint.cpp
--- a/Engine/Plugins/Animation/ControlRig/Source/ControlRigDeveloper/Private/ControlRigBlueprint.cpp
+++ b/Engine/Plugins/Animation/ControlRig/Source/ControlRigDeveloper/Private/ControlRigBlueprint.cpp
@@ -32,6 +32,7 @@
 #include "RigVMModel/Nodes/RigVMAggregateNode.h"
 #include "Units/ControlRigNodeWorkflow.h"
 #include "Rigs/RigControlHierarchy.h"
+#include "RigVMModel/Nodes/RigVMDispatchNode.h"
 
 #if WITH_EDITOR
 #include "IControlRigEditorModule.h"
@@ -1470,7 +1471,15 @@ void UControlRigBlueprint::RefreshAllModels(EControlRigBlueprintLoadType InLoadT
 								}
 
 								bDirtyDuringLoad = true;
-								Registry.ConvertExecuteContextToBaseType(PreferredType.TypeIndex);								
+#if UE_RIGVM_DEBUG_TYPEINDEX
+								// Create a TRigVMTypeIndex from the int32 index so that it can be modified if needed
+								// and then converted back to int32
+								TRigVMTypeIndex PreferredTypeIndex = PreferredType.TypeIndex;
+								Registry.ConvertExecuteContextToBaseType(PreferredTypeIndex);
+								PreferredType.TypeIndex = PreferredTypeIndex;
+#else
+								Registry.ConvertExecuteContextToBaseType(PreferredType.TypeIndex);
+#endif
 							}
 						}
 					}
@@ -1518,6 +1527,35 @@ void UControlRigBlueprint::RefreshAllModels(EControlRigBlueprintLoadType InLoadT
 			{
 				ensureMsgf(false, TEXT("Pin type changed during load %s"), *GetPackage()->GetPathName());
 			}
+
+			for(URigVMNode* ModelNode : GraphToClean->GetNodes())
+			{
+				if(URigVMUnitNode* UnitNode = Cast<URigVMUnitNode>(ModelNode))
+				{
+					if (!UnitNode->HasWildCardPin())
+					{
+						UScriptStruct* ScriptStruct = UnitNode->GetScriptStruct(); 
+						if(ScriptStruct == nullptr)
+						{
+							Controller->FullyResolveTemplateNode(UnitNode, INDEX_NONE, false);
+						}
+
+						if (UnitNode->GetScriptStruct() == nullptr)
+						{
+							static const TCHAR UnresolvedUnitNodeMessage[] = TEXT("Node %s could not be resolved.");
+							Controller->ReportErrorf(UnresolvedUnitNodeMessage, *ModelNode->GetNodePath(true));
+						}
+					}
+				}
+				if (URigVMDispatchNode* DispatchNode = Cast<URigVMDispatchNode>(ModelNode))
+				{
+					if (DispatchNode->GetFactory() == nullptr)
+					{
+						static const TCHAR UnresolvedDispatchNodeMessage[] = TEXT("Dispatch node %s has no factory..");
+						Controller->ReportErrorf(UnresolvedDispatchNodeMessage, *ModelNode->GetNodePath(true));
+					}
+				}
+			}
 		}
 #endif
 
diff --git a/Engine/Source/Developer/RigVMDeveloper/Private/RigVMModel/RigVMController.cpp b/Engine/Source/Developer/RigVMDeveloper/Private/RigVMModel/RigVMController.cpp
--- a/Engine/Source/Developer/RigVMDeveloper/Private/RigVMModel/RigVMController.cpp
+++ b/Engine/Source/Developer/RigVMDeveloper/Private/RigVMModel/RigVMController.cpp
@@ -16746,6 +16746,13 @@ bool URigVMController::UpdateTemplateNodePinTypes(URigVMTemplateNode* InNode, bo
 	{
 		InNode->ResolvedPermutation = INDEX_NONE;
 	}
+	else if (URigVMUnitNode* UnitNode = Cast<URigVMUnitNode>(InNode))
+	{
+		if (const FRigVMFunction* Function = UnitNode->GetResolvedFunction())
+		{
+			UnitNode->ResolvedFunctionName = Function->GetName();
+		}
+	}
 	
 	return bAnyTypeChanged;
 }
