Sequencer: Readded support for DrawExtra in channel views

DrawExtra now receives a parameter struct that contains the necessary slate rendering parameters, rather than utilizing the section painter and separate geometry. Ideally such rendering code would exist in actual view models or custom views, but for now this API needs to be maintained

diff --git a/Engine/Plugins/Animation/ControlRig/Source/ControlRigEditor/Private/Sequencer/ControlRigSpaceChannelEditors.cpp b/Engine/Plugins/Animation/ControlRig/Source/ControlRigEditor/Private/Sequencer/ControlRigSpaceChannelEditors.cpp
--- a/Engine/Plugins/Animation/ControlRig/Source/ControlRigEditor/Private/Sequencer/ControlRigSpaceChannelEditors.cpp
+++ b/Engine/Plugins/Animation/ControlRig/Source/ControlRigEditor/Private/Sequencer/ControlRigSpaceChannelEditors.cpp
@@ -1415,7 +1415,7 @@ int32 FControlRigSpaceChannelHelpers::GetNumFloatChannels(const ERigControlType
 	return 0;
 }
 
-void DrawExtra(FMovieSceneControlRigSpaceChannel* Channel, const UMovieSceneSection* Owner,const FGeometry& AllottedGeometry, FSequencerSectionPainter& Painter)
+int32 DrawExtra(FMovieSceneControlRigSpaceChannel* Channel, const UMovieSceneSection* Owner, const FSequencerChannelPaintArgs& PaintArgs, int32 LayerId)
 {
 	using namespace UE::Sequencer;
 
@@ -1428,12 +1428,11 @@ void DrawExtra(FMovieSceneControlRigSpaceChannel* Channel, const UMovieSceneSect
 		const FSlateFontInfo FontInfo = FCoreStyle::Get().GetFontStyle("ToolTip.LargerFont");
 		const TSharedRef<FSlateFontMeasure> FontMeasure = FSlateApplication::Get().GetRenderer()->GetFontMeasureService();
 
-		const FVector2D LocalSize = AllottedGeometry.GetLocalSize();
+		const FVector2D LocalSize = PaintArgs.Geometry.GetLocalSize();
 		const float LaneTop = 0;
 
-		const FTimeToPixel& TimeToPixel = Painter.GetTimeConverter();
-		const double InputMin = TimeToPixel.PixelToSeconds(0.f);
-		const double InputMax = TimeToPixel.PixelToSeconds(Painter.SectionGeometry.GetLocalSize().X);
+		const double InputMin = PaintArgs.TimeToPixel.PixelToSeconds(0.f);
+		const double InputMax = PaintArgs.TimeToPixel.PixelToSeconds(LocalSize.X);
 
 		int32 LastLabelPos = -1;
 		for (int32 Index = 0; Index < Ranges.Num(); ++Index)
@@ -1463,17 +1462,17 @@ void DrawExtra(FMovieSceneControlRigSpaceChannel* Channel, const UMovieSceneSect
 			if (CurveColor != FLinearColor::White)
 			{
 				const double LowerSecondsForBox = (Index == 0 && LowerSeconds > InputMin) ? InputMin : LowerSeconds;
-				const double BoxPos = TimeToPixel.SecondsToPixel(LowerSecondsForBox);
+				const double BoxPos = PaintArgs.TimeToPixel.SecondsToPixel(LowerSecondsForBox);
 
-				const FPaintGeometry BoxGeometry = AllottedGeometry.ToPaintGeometry(
-					FVector2D(AllottedGeometry.GetLocalSize().X, AllottedGeometry.GetLocalSize().Y),
+				const FPaintGeometry BoxGeometry = PaintArgs.Geometry.ToPaintGeometry(
+					FVector2D(PaintArgs.Geometry.GetLocalSize().X, PaintArgs.Geometry.GetLocalSize().Y),
 					FSlateLayoutTransform(FVector2D(BoxPos, LaneTop))
 				);
 
-				FSlateDrawElement::MakeBox(Painter.DrawElements, Painter.LayerId, BoxGeometry, WhiteBrush, DrawEffects, CurveColor);
+				FSlateDrawElement::MakeBox(PaintArgs.DrawElements, LayerId, BoxGeometry, WhiteBrush, DrawEffects, CurveColor);
 			}
 			const double LowerSecondsForLabel = (InputMin > LowerSeconds) ? InputMin : LowerSeconds;
-			double LabelPos = TimeToPixel.SecondsToPixel(LowerSecondsForLabel) + 10;
+			double LabelPos = PaintArgs.TimeToPixel.SecondsToPixel(LowerSecondsForLabel) + 10;
 
 			const FText Label = FText::FromName(Range.Name);
 			const FVector2D TextSize = FontMeasure->Measure(Label, FontInfo);
@@ -1482,15 +1481,17 @@ void DrawExtra(FMovieSceneControlRigSpaceChannel* Channel, const UMovieSceneSect
 				LabelPos = (LabelPos < LastLabelPos) ? LastLabelPos + 5 : LabelPos;
 			}
 			LastLabelPos = LabelPos + TextSize.X + 15;
-			const FVector2D Position(LabelPos, LaneTop + (AllottedGeometry.GetLocalSize().Y - TextSize.Y) * .5f);
+			const FVector2D Position(LabelPos, LaneTop + (PaintArgs.Geometry.GetLocalSize().Y - TextSize.Y) * .5f);
 
-			const FPaintGeometry LabelGeometry = AllottedGeometry.ToPaintGeometry(
+			const FPaintGeometry LabelGeometry = PaintArgs.Geometry.ToPaintGeometry(
 				FSlateLayoutTransform(Position)
 			);
 
-			FSlateDrawElement::MakeText(Painter.DrawElements, Painter.LayerId, LabelGeometry, Label, FontInfo, DrawEffects, FLinearColor::White);
-		}	
+			FSlateDrawElement::MakeText(PaintArgs.DrawElements, LayerId, LabelGeometry, Label, FontInfo, DrawEffects, FLinearColor::White);
+		}
 	}
+
+	return LayerId + 1;
 }
 
 void DrawKeys(FMovieSceneControlRigSpaceChannel* Channel, TArrayView<const FKeyHandle> InKeyHandles, const UMovieSceneSection* InOwner, TArrayView<FKeyDrawParams> OutKeyDrawParams)
diff --git a/Engine/Plugins/Animation/ControlRig/Source/ControlRigEditor/Public/ControlRigSpaceChannelEditors.h b/Engine/Plugins/Animation/ControlRig/Source/ControlRigEditor/Public/ControlRigSpaceChannelEditors.h
--- a/Engine/Plugins/Animation/ControlRig/Source/ControlRigEditor/Public/ControlRigSpaceChannelEditors.h
+++ b/Engine/Plugins/Animation/ControlRig/Source/ControlRigEditor/Public/ControlRigSpaceChannelEditors.h
@@ -74,7 +74,7 @@ TSharedRef<SWidget> CreateKeyEditor(const TMovieSceneChannelHandle<FMovieSceneCo
 
 /** Key drawing overrides */
 void DrawKeys(FMovieSceneControlRigSpaceChannel* Channel, TArrayView<const FKeyHandle> InKeyHandles, const UMovieSceneSection* InOwner, TArrayView<FKeyDrawParams> OutKeyDrawParams);
-void DrawExtra(FMovieSceneControlRigSpaceChannel* Channel, const UMovieSceneSection* Owner, const FGeometry& AllottedGeometry,FSequencerSectionPainter& Painter);
+int32 DrawExtra(FMovieSceneControlRigSpaceChannel* Channel, const UMovieSceneSection* Owner, const FSequencerChannelPaintArgs& PaintArgs, int32 LayerId);
 
 //UMovieSceneKeyStructType* IstanceGeneratedStruct(FMovieSceneControlRigSpaceChannel* Channel, FSequencerKeyStructGenerator* Generator);
 
diff --git a/Engine/Source/Editor/MovieSceneTools/Private/Channels/ConstraintChannelEditor.cpp b/Engine/Source/Editor/MovieSceneTools/Private/Channels/ConstraintChannelEditor.cpp
--- a/Engine/Source/Editor/MovieSceneTools/Private/Channels/ConstraintChannelEditor.cpp
+++ b/Engine/Source/Editor/MovieSceneTools/Private/Channels/ConstraintChannelEditor.cpp
@@ -5,7 +5,6 @@
 #include "ConstraintChannel.h"
 
 #include "KeyBarCurveModel.h"
-#include "SequencerSectionPainter.h"
 #include "TimeToPixel.h"
 #include "Fonts/FontMeasure.h"
 #include "Styling/AppStyle.h"
@@ -203,21 +202,21 @@ TArray<FKeyBarCurveModel::FBarRange> FConstraintChannelEditor::GetBarRanges(
 	return Ranges;
 }
 
-void DrawExtra(
+int32 DrawExtra(
 	FMovieSceneConstraintChannel* Channel,
 	const UMovieSceneSection* Owner,
-	const FGeometry& AllottedGeometry,
-	FSequencerSectionPainter& Painter)
+	const FSequencerChannelPaintArgs& PaintArgs,
+	int32 LayerId)
 {
 	if (!Owner || !Channel)
 	{
-		return;
+		return LayerId;
 	}
 	TArray<FKeyBarCurveModel::FBarRange> Ranges = FConstraintChannelEditor::GetBarRanges(Channel,Owner);
 
 	if (Ranges.IsEmpty())
 	{
-		return;
+		return LayerId;
 	}
 	
 	// draw bars + labels
@@ -227,12 +226,11 @@ void DrawExtra(
 	const FSlateFontInfo FontInfo = FCoreStyle::Get().GetFontStyle("ToolTip.LargerFont");
 	const TSharedRef<FSlateFontMeasure> FontMeasure = FSlateApplication::Get().GetRenderer()->GetFontMeasureService();
 
-	const FVector2D& LocalSize = AllottedGeometry.GetLocalSize();
+	const FVector2D LocalSize = PaintArgs.Geometry.GetLocalSize();
 	static constexpr float LaneTop = 0;
 
-	const FTimeToPixel& TimeToPixel = Painter.GetTimeConverter();
-	const double InputMin = TimeToPixel.PixelToSeconds(0.f);
-	const double InputMax = TimeToPixel.PixelToSeconds(Painter.SectionGeometry.GetLocalSize().X);
+	const double InputMin = PaintArgs.TimeToPixel.PixelToSeconds(0.f);
+	const double InputMax = PaintArgs.TimeToPixel.PixelToSeconds(LocalSize.X);
 
 	// label data
 	const FConstraintChannelLabelBuilder LabelBuilder(Owner, Channel);
@@ -252,27 +250,27 @@ void DrawExtra(
 			UpperSeconds = InputMax;
 		}
 
-		const double BoxStart = TimeToPixel.SecondsToPixel(LowerSeconds);
-		const double BoxEnd = TimeToPixel.SecondsToPixel(UpperSeconds);
+		const double BoxStart = PaintArgs.TimeToPixel.SecondsToPixel(LowerSeconds);
+		const double BoxEnd = PaintArgs.TimeToPixel.SecondsToPixel(UpperSeconds);
 		const double BoxSize = BoxEnd - BoxStart;
 
 		const FVector2D Size = FVector2D(BoxSize, LocalSize.Y);
 		const FVector2D Translation = FVector2D(BoxStart, LaneTop);
-		const FPaintGeometry BoxGeometry = AllottedGeometry.ToPaintGeometry(Size, FSlateLayoutTransform(Translation));
+		const FPaintGeometry BoxGeometry = PaintArgs.Geometry.ToPaintGeometry(Size, FSlateLayoutTransform(Translation));
 		
-		FSlateDrawElement::MakeBox(Painter.DrawElements, Painter.LayerId, BoxGeometry, WhiteBrush, DrawEffects, Range.Color);
+		FSlateDrawElement::MakeBox(PaintArgs.DrawElements, LayerId, BoxGeometry, WhiteBrush, DrawEffects, Range.Color);
 
 		// draw label
 		if (LabelBuilder.bValid)
 		{
 			const double LabelPos = BoxStart + LabelPixelOffset;
 			const FVector2D Position(LabelPos, LaneTop + (LocalSize.Y - TextSize.Y) * .5f);
-			const FPaintGeometry LabelGeometry = AllottedGeometry.ToPaintGeometry(FSlateLayoutTransform(Position));
+			const FPaintGeometry LabelGeometry = PaintArgs.Geometry.ToPaintGeometry(FSlateLayoutTransform(Position));
 
 			const double LabelMaxSize = BoxSize-LabelPixelOffset;
 			if (TextSize.X < LabelMaxSize)
 			{
-				FSlateDrawElement::MakeText(Painter.DrawElements, Painter.LayerId, LabelGeometry, Label, FontInfo);				
+				FSlateDrawElement::MakeText(PaintArgs.DrawElements, LayerId, LabelGeometry, Label, FontInfo);				
 			}
 			else
 			{
@@ -280,11 +278,13 @@ void DrawExtra(
 				const int32 End = FontMeasure->FindLastWholeCharacterIndexBeforeOffset(Label, FontInfo, FMath::RoundToInt(LabelMaxSize));
 				if (End >= 0)
 				{
-					FSlateDrawElement::MakeText(Painter.DrawElements, Painter.LayerId, LabelGeometry, Label.ToString(), 0, End, FontInfo);
+					FSlateDrawElement::MakeText(PaintArgs.DrawElements, LayerId, LabelGeometry, Label.ToString(), 0, End, FontInfo);
 				}
 			}
 		}
 	}
+
+	return LayerId + 1;
 }
 
 FKeyHandle AddOrUpdateKey(
diff --git a/Engine/Source/Editor/MovieSceneTools/Private/Channels/ConstraintChannelEditor.h b/Engine/Source/Editor/MovieSceneTools/Private/Channels/ConstraintChannelEditor.h
--- a/Engine/Source/Editor/MovieSceneTools/Private/Channels/ConstraintChannelEditor.h
+++ b/Engine/Source/Editor/MovieSceneTools/Private/Channels/ConstraintChannelEditor.h
@@ -21,11 +21,11 @@ void DrawKeys(
 	const UMovieSceneSection* InOwner,
 	TArrayView<FKeyDrawParams> OutKeyDrawParams);
 
-void DrawExtra(
+int32 DrawExtra(
 	FMovieSceneConstraintChannel* Channel,
 	const UMovieSceneSection* Owner,
-	const FGeometry& AllottedGeometry,
-	FSequencerSectionPainter& Painter);
+	const FSequencerChannelPaintArgs& PaintArgs,
+	int32 LayerId);
 
 /** Overrides for adding or updating a key for non-standard channels */
 FKeyHandle AddOrUpdateKey(
diff --git a/Engine/Source/Editor/MovieSceneTools/Public/Channels/PerlinNoiseChannelInterface.h b/Engine/Source/Editor/MovieSceneTools/Public/Channels/PerlinNoiseChannelInterface.h
--- a/Engine/Source/Editor/MovieSceneTools/Public/Channels/PerlinNoiseChannelInterface.h
+++ b/Engine/Source/Editor/MovieSceneTools/Public/Channels/PerlinNoiseChannelInterface.h
@@ -11,7 +11,6 @@
 #include "Modules/ModuleManager.h"
 #include "MovieSceneSection.h"
 #include "PropertyEditorModule.h"
-#include "SequencerSectionPainter.h"
 #include "TimeToPixel.h"
 #include "ToolMenus.h"
 #include "Widgets/Input/NumericTypeInterface.h"
@@ -195,36 +194,37 @@ struct TPerlinNoiseChannelInterface : ISequencerChannelInterface
 		return nullptr;
 	}
 
-	virtual void DrawExtra_Raw(FMovieSceneChannel* InChannel, const UMovieSceneSection* InOwner, const FGeometry& InKeyGeometry, FSequencerSectionPainter& Painter) const override
+	virtual int32 DrawExtra_Raw(FMovieSceneChannel* InChannel, const UMovieSceneSection* InOwner, const FSequencerChannelPaintArgs& PaintArgs, int32 LayerId) const override
 	{
 		using namespace UE::Sequencer;
 
 		FColor FillColor(0, 0, 0, 51);
-		FTimeToPixel TimeToPixelConverter = Painter.GetTimeConverter();
 		ChannelType* TypedChannel = static_cast<ChannelType*>(InChannel);
 
 		TArray<FVector2D> CurvePoints;
-		CurvePoints.Reserve(InKeyGeometry.Size.X / 2.0);
+		CurvePoints.Reserve(PaintArgs.Geometry.Size.X / 2.0);
 
 		const double Amplitude = TypedChannel->GetParams().Amplitude;
-		const double YOffset = InKeyGeometry.Size.Y / 2.0;
-		const double YScale = (Amplitude != 0) ? (InKeyGeometry.Size.Y / Amplitude / 2.0) : 1;
+		const double YOffset = PaintArgs.Geometry.Size.Y / 2.0;
+		const double YScale = (Amplitude != 0) ? (PaintArgs.Geometry.Size.Y / Amplitude / 2.0) : 1;
 
-		for (double X = 0; X < InKeyGeometry.Size.X; X += 2.f)
+		for (double X = 0; X < PaintArgs.Geometry.Size.X; X += 2.f)
 		{
-			double Seconds = TimeToPixelConverter.PixelToSeconds(X);
+			double Seconds = PaintArgs.TimeToPixel.PixelToSeconds(X);
 			double Y = (double)TypedChannel->Evaluate(Seconds);
 			CurvePoints.Add(FVector2D(X, (Y * YScale) + YOffset));
 		}
 
 		FSlateDrawElement::MakeLines(
-				Painter.DrawElements,
-				Painter.LayerId++,
-				InKeyGeometry.ToPaintGeometry(),
+				PaintArgs.DrawElements,
+				LayerId,
+				PaintArgs.Geometry.ToPaintGeometry(),
 				CurvePoints,
 				ESlateDrawEffect::PreMultipliedAlpha,
 				FillColor,
 				true);
+
+		return LayerId + 1;
 	}
 };
 
diff --git a/Engine/Source/Editor/Sequencer/Private/IKeyArea.cpp b/Engine/Source/Editor/Sequencer/Private/IKeyArea.cpp
--- a/Engine/Source/Editor/Sequencer/Private/IKeyArea.cpp
+++ b/Engine/Source/Editor/Sequencer/Private/IKeyArea.cpp
@@ -143,15 +143,16 @@ TSharedPtr<FStructOnScope> IKeyArea::GetKeyStruct(FKeyHandle KeyHandle) const
 	return nullptr;
 }
 
-void IKeyArea::DrawExtra(FSequencerSectionPainter& Painter, const FGeometry& KeyGeometry) const
+int32 IKeyArea::DrawExtra(const FSequencerChannelPaintArgs& PaintArgs, int32 LayerId) const
 {
 	ISequencerChannelInterface* EditorInterface = FindChannelEditorInterface();
 	if (EditorInterface)
 	{
 		FMovieSceneChannel* Channel = ChannelHandle.Get();
 		const UMovieSceneSection* OwningSection = GetOwningSection();
-		EditorInterface->DrawExtra_Raw(Channel,OwningSection, KeyGeometry,Painter);
+		return EditorInterface->DrawExtra_Raw(Channel, OwningSection, PaintArgs, LayerId);
 	}
+	return LayerId;
 }
 
 void IKeyArea::DrawKeys(TArrayView<const FKeyHandle> InKeyHandles, TArrayView<FKeyDrawParams> OutKeyDrawParams)
diff --git a/Engine/Source/Editor/Sequencer/Private/MVVM/ViewModels/ChannelModel.cpp b/Engine/Source/Editor/Sequencer/Private/MVVM/ViewModels/ChannelModel.cpp
--- a/Engine/Source/Editor/Sequencer/Private/MVVM/ViewModels/ChannelModel.cpp
+++ b/Engine/Source/Editor/Sequencer/Private/MVVM/ViewModels/ChannelModel.cpp
@@ -24,6 +24,7 @@
 #include "SequencerNodeTree.h"
 #include "SKeyAreaEditorSwitcher.h"
 #include "SequencerSectionPainter.h"
+#include "SequencerChannelTraits.h"
 #include "SSequencerSection.h"
 #include "Channels/MovieSceneChannel.h"
 #include "CurveModel.h"
@@ -178,9 +179,9 @@ bool FChannelModel::GetFixedExtents(double& OutFixedMin, double& OutFixedMax) co
 	return false;
 }
 
-int32 FChannelModel::CustomPaint(const FGeometry& KeyGeometry, int32 LayerId) const
+int32 FChannelModel::CustomPaint(const FSequencerChannelPaintArgs& PaintArgs, int32 LayerId) const
 {
-	return LayerId;
+	return KeyArea->DrawExtra(PaintArgs, LayerId);
 }
 
 void FChannelModel::DrawKeys(TArrayView<const FKeyHandle> InKeyHandles, TArrayView<FKeyDrawParams> OutKeyDrawParams)
diff --git a/Engine/Source/Editor/Sequencer/Private/MVVM/Views/SChannelView.cpp b/Engine/Source/Editor/Sequencer/Private/MVVM/Views/SChannelView.cpp
--- a/Engine/Source/Editor/Sequencer/Private/MVVM/Views/SChannelView.cpp
+++ b/Engine/Source/Editor/Sequencer/Private/MVVM/Views/SChannelView.cpp
@@ -13,6 +13,7 @@
 #include "ISequencerSection.h"
 #include "SequencerAddKeyOperation.h"
 #include "SSequencerSection.h"
+#include "SequencerChannelTraits.h"
 
 #include "ScopedTransaction.h"
 
@@ -441,6 +442,21 @@ int32 SChannelView::OnPaint(const FPaintArgs& Args, const FGeometry& AllottedGeo
 		RelativeTimeToPixel.PixelToFrame(BottomRight.X)
 	);
 
+	TViewModelPtr<FChannelModel> Channel = Model.ImplicitCast();
+	if (Channel)
+	{
+		FSequencerChannelPaintArgs ChannelPaintArgs = {
+			OutDrawElements,
+			Args,
+			AllottedGeometry,
+			MyCullingRect,
+			InWidgetStyle,
+			RelativeTimeToPixel,
+			bParentEnabled
+		};
+		LayerId = Channel->CustomPaint(ChannelPaintArgs, LayerId);
+	}
+
 	FSequencerKeyRendererInterface KeyRenderInterface(SectionObject, Sequencer.Get(), TrackArea->GetHotspot());
 	FChannelViewKeyCachedState NewCachedState(VisibleRange, TrackArea->GetHotspot(), Model, Sequencer.Get());
 
@@ -464,7 +480,6 @@ int32 SChannelView::OnPaint(const FPaintArgs& Args, const FGeometry& AllottedGeo
 		PaintArgs.KeyBarColor = KeyBarColor.Get();
 	}
 
-	TViewModelPtr<FChannelModel> Channel = Model.ImplicitCast();
 	if (Sequencer->GetSequencerSettings()->GetShowChannelColors() && Channel)
 	{
 		if (TOptional<FLinearColor> ChannelColor = Channel->GetKeyArea()->GetColor())
diff --git a/Engine/Source/Editor/Sequencer/Private/SequencerChannelTraits.cpp b/Engine/Source/Editor/Sequencer/Private/SequencerChannelTraits.cpp
--- a/Engine/Source/Editor/Sequencer/Private/SequencerChannelTraits.cpp
+++ b/Engine/Source/Editor/Sequencer/Private/SequencerChannelTraits.cpp
@@ -21,8 +21,9 @@ void DrawKeys(FMovieSceneChannel* Channel, TArrayView<const FKeyHandle> InHandle
 	}
 }
 /* Most channels do nothing*/
-void DrawExtra(FMovieSceneChannel* InChannel, const UMovieSceneSection* Owner, const FGeometry& KeyGeometry,FSequencerSectionPainter& Painter)
+int32 DrawExtra(FMovieSceneChannel* InChannel, const UMovieSceneSection* Owner, const FSequencerChannelPaintArgs& PaintArgs, int32 LayerId)
 {
+	return LayerId;
 }
 
 
diff --git a/Engine/Source/Editor/Sequencer/Public/IKeyArea.h b/Engine/Source/Editor/Sequencer/Public/IKeyArea.h
--- a/Engine/Source/Editor/Sequencer/Public/IKeyArea.h
+++ b/Engine/Source/Editor/Sequencer/Public/IKeyArea.h
@@ -15,6 +15,7 @@ struct FRichCurve;
 struct FKeyDrawParams;
 struct FMovieSceneChannel;
 struct FMovieSceneChannelProxy;
+struct FSequencerChannelPaintArgs;
 struct FSequencerPasteEnvironment;
 struct ISequencerChannelInterface;
 struct FMovieSceneChannelMetaData;
@@ -248,16 +249,14 @@ public:
 	 */
 	void CopyKeys(FMovieSceneClipboardBuilder& ClipboardBuilder, TArrayView<const FKeyHandle> KeyMask) const;
 
-
-
 	/**
 	 * Draw additional content in addition to keys for a particular IKeyArea
 	 *
-	 * @param Painter			 The painter to add the created geometry
-	 * @param KeyGeometry		 Allocated geometry to draw in
+	 * @param PaintArgs	       Paint arguments containing the draw element list, time-to-pixel converter and other structures
+	 * @param LayerId          The slate layer to paint onto
+	 * @return The new slate layer ID for subsequent elements to paint onto
 	 */
-	void DrawExtra(FSequencerSectionPainter& Painter, const FGeometry& KeyGeometry) const;
-
+	int32 DrawExtra(const FSequencerChannelPaintArgs& PaintArgs, int32 LayerId) const;
 
 	/**
 	 * Paste the specified key track into this key area
diff --git a/Engine/Source/Editor/Sequencer/Public/ISequencerChannelInterface.h b/Engine/Source/Editor/Sequencer/Public/ISequencerChannelInterface.h
--- a/Engine/Source/Editor/Sequencer/Public/ISequencerChannelInterface.h
+++ b/Engine/Source/Editor/Sequencer/Public/ISequencerChannelInterface.h
@@ -12,8 +12,9 @@
 struct FGeometry;
 struct FKeyDrawParams;
 struct FKeyHandle;
-struct FMovieSceneClipboardEnvironment;
+struct FSequencerChannelPaintArgs;
 struct FSequencerPasteEnvironment;
+struct FMovieSceneClipboardEnvironment;
 
 class SWidget;
 class ISequencer;
@@ -26,7 +27,6 @@ class ISectionLayoutBuilder;
 class FMovieSceneClipboardBuilder;
 class FMovieSceneClipboardKeyTrack;
 class FTrackInstancePropertyBindings;
-class FSequencerSectionPainter;
 
 /** Utility struct representing a number of selected keys on a single channel */
 struct FExtendKeyMenuParams
@@ -171,9 +171,10 @@ struct ISequencerChannelInterface
 	 *
 	 * @param InChannel          The channel to draw extra display information for
 	 * @param InOwner            The owning movie scene section for this channel
-	 * @param InKeyGeometry      Allocated geometry to draw in
-	 * @param Painter			 The painter to add the created geometry to
+	 * @param PaintArgs          Paint arguments containing the draw element list, time-to-pixel converter and other structures
+	 * @param LayerId            The slate layer to paint onto
+	 * @return The new slate layer ID for subsequent elements to paint onto
 	 */
-	virtual void DrawExtra_Raw(FMovieSceneChannel* InChannel, const UMovieSceneSection* InOwner, const FGeometry& InKeyGeometry, FSequencerSectionPainter& Painter) const = 0;
+	virtual int32 DrawExtra_Raw(FMovieSceneChannel* InChannel, const UMovieSceneSection* InOwner, const FSequencerChannelPaintArgs& PaintArgs, int32 LayerId) const = 0;
 
 };
diff --git a/Engine/Source/Editor/Sequencer/Public/MVVM/ViewModels/ChannelModel.h b/Engine/Source/Editor/Sequencer/Public/MVVM/ViewModels/ChannelModel.h
--- a/Engine/Source/Editor/Sequencer/Public/MVVM/ViewModels/ChannelModel.h
+++ b/Engine/Source/Editor/Sequencer/Public/MVVM/ViewModels/ChannelModel.h
@@ -20,6 +20,7 @@
 
 struct FMovieSceneChannel;
 struct FMovieSceneChannelHandle;
+struct FSequencerChannelPaintArgs;
 
 class IKeyArea;
 class ISequencerSection;
@@ -74,6 +75,8 @@ public:
 	/** Returns the desired sizing for the track area row */
 	FOutlinerSizing GetDesiredSizing() const;
 
+	int32 CustomPaint(const FSequencerChannelPaintArgs& CustomPaintArgs, int32 LayerId) const;
+
 	/*~ ITrackLaneExtension */
 	TSharedPtr<ITrackLaneWidget> CreateTrackLaneView(const FCreateTrackLaneViewParams& InParams) override;
 	FTrackLaneVirtualAlignment ArrangeVirtualTrackLaneView() const override;
@@ -81,7 +84,6 @@ public:
 	/*~ IKeyExtension */
 	bool UpdateCachedKeys(TSharedPtr<FCachedKeys>& OutCachedKeys) const override;
 	bool GetFixedExtents(double& OutFixedMin, double& OutFixedMax) const override;
-	int32 CustomPaint(const FGeometry& KeyGeometry, int32 LayerId) const override;
 	void DrawKeys(TArrayView<const FKeyHandle> InKeyHandles, TArrayView<FKeyDrawParams> OutKeyDrawParams) override;
 	TUniquePtr<FCurveModel> CreateCurveModel() override;
 
diff --git a/Engine/Source/Editor/Sequencer/Public/SequencerChannelInterface.h b/Engine/Source/Editor/Sequencer/Public/SequencerChannelInterface.h
--- a/Engine/Source/Editor/Sequencer/Public/SequencerChannelInterface.h
+++ b/Engine/Source/Editor/Sequencer/Public/SequencerChannelInterface.h
@@ -145,13 +145,14 @@ struct TSequencerChannelInterfaceCommon : ISequencerChannelInterface
 	 *
 	 * @param InChannel          The channel to draw extra display information for
 	 * @param InOwner            The owning movie scene section for this channel
-	 * @param InKeyGeometry      Allocated geometry to draw in
-	 * @param Painter			 The painter to add the created geometry to
+	 * @param PaintArgs          Paint arguments containing the draw element list, time-to-pixel converter and other structures
+	 * @param LayerId            The slate layer to paint onto
+	 * @return The new slate layer ID for subsequent elements to paint onto
 	 */
-	virtual void DrawExtra_Raw(FMovieSceneChannel* InChannel,const UMovieSceneSection* InOwner, const FGeometry& InKeyGeometry,FSequencerSectionPainter& Painter) const
+	virtual int32 DrawExtra_Raw(FMovieSceneChannel* InChannel, const UMovieSceneSection* InOwner, const FSequencerChannelPaintArgs& PaintArgs, int32 LayerId) const override
 	{
 		using namespace Sequencer;
-		DrawExtra(static_cast<ChannelType*>(InChannel), InOwner, InKeyGeometry, Painter);
+		return DrawExtra(static_cast<ChannelType*>(InChannel), InOwner, PaintArgs, LayerId);
 	}
 
 
diff --git a/Engine/Source/Editor/Sequencer/Public/SequencerChannelTraits.h b/Engine/Source/Editor/Sequencer/Public/SequencerChannelTraits.h
--- a/Engine/Source/Editor/Sequencer/Public/SequencerChannelTraits.h
+++ b/Engine/Source/Editor/Sequencer/Public/SequencerChannelTraits.h
@@ -12,6 +12,7 @@
 #include "Widgets/SNullWidget.h"
 #include "ISequencer.h"
 #include "MovieSceneSection.h"
+#include "TimeToPixel.h"
 
 /** Utility struct representing a number of selected keys on a single channel */
 template<typename ChannelType>
@@ -27,6 +28,20 @@ struct TExtendKeyMenuParams
 	TArray<FKeyHandle> Handles;
 };
 
+struct FSequencerChannelPaintArgs
+{
+	FSlateWindowElementList& DrawElements;
+
+	const FPaintArgs& WidgetPaintArgs;
+	const FGeometry& Geometry;
+	const FSlateRect& MyCullingRect;
+	const FWidgetStyle& WidgetStyle;
+
+	FTimeToPixel TimeToPixel;
+
+	bool bParentEnabled = true;
+};
+
 /**
  * Stub/default implementations for ISequencerChannelInterface functions.
  * Custom behaviour should be implemented by overloading the relevant function with the necessary channel/data types.
@@ -256,10 +271,11 @@ namespace Sequencer
 	 *
 	 * @param InChannel          The channel to draw extra display information for
 	 * @param InOwner            The owning movie scene section for this channel
-	 * @param InKeyGeometry      Allocated geometry to draw in
-	 * @param Painter			 The painter to add the created geometry to
+	 * @param PaintArgs          Paint arguments containing the draw element list, time-to-pixel converter and other structures
+	 * @param LayerId            The slate layer to paint onto
+	 * @return The new slate layer ID for subsequent elements to paint onto
 	 */
-	SEQUENCER_API void DrawExtra(FMovieSceneChannel* InChannel, const UMovieSceneSection* InOwner, const FGeometry& KeyGeometry, FSequencerSectionPainter& Painter);
+	SEQUENCER_API int32 DrawExtra(FMovieSceneChannel* InChannel, const UMovieSceneSection* InOwner, const FSequencerChannelPaintArgs& PaintArgs, int32 LayerId);
 
 	/**
 	 * Copy the specified keys from a channel
diff --git a/Engine/Source/Editor/SequencerCore/Private/MVVM/Views/KeyRenderer.cpp b/Engine/Source/Editor/SequencerCore/Private/MVVM/Views/KeyRenderer.cpp
--- a/Engine/Source/Editor/SequencerCore/Private/MVVM/Views/KeyRenderer.cpp
+++ b/Engine/Source/Editor/SequencerCore/Private/MVVM/Views/KeyRenderer.cpp
@@ -632,15 +632,8 @@ int32 FKeyRenderer::DrawCurve(const FKeyBatchParameters& Params, const FGeometry
 		CurveThickness
 	);
 
-	++LayerId;
-
 	PaintArgs.DrawElements->PopClip();
-	for (const FCachedKeyDrawInformation& CachedKeyDrawInfo : KeyDrawInfo)
-	{
-		LayerId = CachedKeyDrawInfo.DrawExtra(AllottedGeometry, LayerId);
-	}
-
-	return LayerId;
+	return LayerId + 1;
 }
 
 int32 FKeyRenderer::Draw(const FKeyBatchParameters& Params, const FGeometry& AllottedGeometry, const FSlateRect& MyCullingRect, const FKeyRendererPaintArgs& PaintArgs, int32 LayerId) const
@@ -813,16 +806,6 @@ int32 FKeyRenderer::Draw(const FKeyBatchParameters& Params, const FGeometry& All
 	return LayerId + 2;
 }
 
-int32 FKeyRenderer::FCachedKeyDrawInformation::DrawExtra(const FGeometry& AllottedGeometry, int32 LayerId) const
-{
-	TViewModelPtr<IKeyExtension> KeyExtension = WeakKeyExtension.Pin();
-	if (KeyExtension)
-	{
-		return KeyExtension->CustomPaint(AllottedGeometry, LayerId);
-	}
-	return LayerId;
-}
-
 EViewDependentCacheFlags FKeyRenderer::FCachedKeyDrawInformation::UpdateViewIndependentData()
 {
 	TViewModelPtr<IKeyExtension> KeyExtension = WeakKeyExtension.Pin();
diff --git a/Engine/Source/Editor/SequencerCore/Public/MVVM/Extensions/IKeyExtension.h b/Engine/Source/Editor/SequencerCore/Public/MVVM/Extensions/IKeyExtension.h
--- a/Engine/Source/Editor/SequencerCore/Public/MVVM/Extensions/IKeyExtension.h
+++ b/Engine/Source/Editor/SequencerCore/Public/MVVM/Extensions/IKeyExtension.h
@@ -49,8 +49,6 @@ public:
 
 	virtual bool GetFixedExtents(double& OutFixedMin, double& OutFixedMax) const = 0;
 
-	virtual int32 CustomPaint(const FGeometry& KeyGeometry, int32 LayerId) const = 0;
-
 	virtual void DrawKeys(TArrayView<const FKeyHandle> InKeyHandles, TArrayView<FKeyDrawParams> OutKeyDrawParams) = 0;
 
 	virtual TUniquePtr<FCurveModel> CreateCurveModel() = 0;
diff --git a/Engine/Source/Editor/SequencerCore/Public/MVVM/Views/KeyRenderer.h b/Engine/Source/Editor/SequencerCore/Public/MVVM/Views/KeyRenderer.h
--- a/Engine/Source/Editor/SequencerCore/Public/MVVM/Views/KeyRenderer.h
+++ b/Engine/Source/Editor/SequencerCore/Public/MVVM/Views/KeyRenderer.h
@@ -229,15 +229,6 @@ private:
 		 */
 		void CacheViewDependentData(const FKeyBatchParameters& Params);
 
-
-		/**
-		 * Draw additional content in addition to keys for a particular IKeyArea
-		 *
-		 * @param Painter			 The painter to add the created geometry
-		 * @param KeyGeometry		 Allocated geometry to draw in
-		 */
-		int32 DrawExtra(const FGeometry& KeyGeometry, int32 LayerId) const;
-
 	public:
 
 		/** Index into the array views for the next unhandled key */
