Control Rig: Avoid calls to GetNativizedClass if nativization is disabled

diff --git a/Engine/Plugins/Animation/ControlRig/Source/ControlRig/Private/ControlRig.cpp b/Engine/Plugins/Animation/ControlRig/Source/ControlRig/Private/ControlRig.cpp
--- a/Engine/Plugins/Animation/ControlRig/Source/ControlRig/Private/ControlRig.cpp
+++ b/Engine/Plugins/Animation/ControlRig/Source/ControlRig/Private/ControlRig.cpp
@@ -3272,7 +3272,8 @@ void UControlRig::SwapVMToNativizedIfRequired(UClass* InNativizedClass)
 
 	const bool bNativizedVMDisabled = AreNativizedVMsDisabled();
 
-	if(InNativizedClass == nullptr)
+	// GetNativizedClass can be pretty costly, let's try to skip this if it is not absolutely necessary
+	if(InNativizedClass == nullptr && !bNativizedVMDisabled)
 	{
 		if(!HasAnyFlags(RF_ClassDefaultObject))
 		{
