Cosntraints: Control Rig:  Check IsInGameThread when recieving control modified events to avoid compensation crashes.
#trivial

diff --git a/Engine/Plugins/Animation/ControlRig/Source/ControlRigEditor/Private/Sequencer/ControlRigParameterTrackEditor.cpp b/Engine/Plugins/Animation/ControlRig/Source/ControlRigEditor/Private/Sequencer/ControlRigParameterTrackEditor.cpp
--- a/Engine/Plugins/Animation/ControlRig/Source/ControlRigEditor/Private/Sequencer/ControlRigParameterTrackEditor.cpp
+++ b/Engine/Plugins/Animation/ControlRig/Source/ControlRigEditor/Private/Sequencer/ControlRigParameterTrackEditor.cpp
@@ -2826,6 +2826,10 @@ void FControlRigParameterTrackEditor::HandleControlModified(UControlRig* Control
 	{
 		return;
 	}
+	if (IsInGameThread() == false)
+	{
+		return;
+	}
 	if (!GetSequencer().IsValid() || !GetSequencer()->IsAllowedToChange() || Context.SetKey == EControlRigSetKey::Never 
 		|| ControlElement->Settings.AnimationType == ERigControlAnimationType::ProxyControl
 		|| ControlElement->Settings.AnimationType == ERigControlAnimationType::VisualCue)
