ControlRig: Added evaluate critical section so that CR editor don't poke into CR instances on the anim thread at an inconvenient time

diff --git a/Engine/Plugins/Animation/ControlRig/Source/ControlRig/Private/ControlRig.cpp b/Engine/Plugins/Animation/ControlRig/Source/ControlRig/Private/ControlRig.cpp
--- a/Engine/Plugins/Animation/ControlRig/Source/ControlRig/Private/ControlRig.cpp
+++ b/Engine/Plugins/Animation/ControlRig/Source/ControlRig/Private/ControlRig.cpp
@@ -290,6 +290,15 @@ void UControlRig::Evaluate_AnyThread()
 	DECLARE_SCOPE_HIERARCHICAL_COUNTER_FUNC()
 	QUICK_SCOPE_CYCLE_COUNTER(STAT_ControlRig_Evaluate);
 
+	// we can have other systems trying to poke into running instances of Control Rigs
+	// on the anim thread and query data, such as
+	// UControlRigSkeletalMeshComponent::RebuildDebugDrawSkeleton,
+	// using a lock here to prevent them from having an inconsistent view of the rig at some
+	// intermediate stage of evaluation, for example, during evaluate, we can have a call
+	// to copy hierarchy, which empties the hierarchy for a short period of time
+	// and we don't want other systems to see that.
+	FScopeLock EvaluateLock(&GetEvaluateMutex());
+	
 	// create a copy since we need to change it here temporarily,
 	// and UI / the rig may change the event queue while it is running
 	EventQueueToRun = EventQueue;
diff --git a/Engine/Plugins/Animation/ControlRig/Source/ControlRig/Public/ControlRig.h b/Engine/Plugins/Animation/ControlRig/Source/ControlRig/Public/ControlRig.h
--- a/Engine/Plugins/Animation/ControlRig/Source/ControlRig/Public/ControlRig.h
+++ b/Engine/Plugins/Animation/ControlRig/Source/ControlRig/Public/ControlRig.h
@@ -178,6 +178,9 @@ public:
 	/** Evaluate at Any Thread */
 	virtual void Evaluate_AnyThread();
 
+	/** Locks for the scope of Evaluate_AnyThread */
+	FCriticalSection& GetEvaluateMutex() { return EvaluateMutex; };
+
 	/** Returns the member properties as an external variable array */
 	TArray<FRigVMExternalVariable> GetExternalVariables() const;
 
@@ -941,6 +944,8 @@ public:
 	
 private:
 
+	FCriticalSection EvaluateMutex;
+
 #if UE_CONTROLRIG_PROFILE_EXECUTE_UNITS_NUM
 	int32 ProfilingRunsLeft;
 	uint64 AccumulatedCycles;
diff --git a/Engine/Plugins/Animation/ControlRig/Source/ControlRigEditor/Private/Editor/ControlRigSkeletalMeshComponent.cpp b/Engine/Plugins/Animation/ControlRig/Source/ControlRigEditor/Private/Editor/ControlRigSkeletalMeshComponent.cpp
--- a/Engine/Plugins/Animation/ControlRig/Source/ControlRigEditor/Private/Editor/ControlRigSkeletalMeshComponent.cpp
+++ b/Engine/Plugins/Animation/ControlRig/Source/ControlRigEditor/Private/Editor/ControlRigSkeletalMeshComponent.cpp
@@ -65,6 +65,14 @@ void UControlRigSkeletalMeshComponent::RebuildDebugDrawSkeleton()
 		UControlRig* ControlRig = ControlRigInstance->GetFirstAvailableControlRig();
 		if (ControlRig)
 		{
+			// we are trying to poke into running instances of Control Rigs
+			// on the anim thread and query data, using a lock here to make sure
+			// we don't get an inconsistent view of the rig at some
+			// intermediate stage of evaluation, for example, during evaluate, we can have a call
+			// to copy hierarchy, which empties the hierarchy for a short period of time.
+			// if we did not have this lock and try to grab it, we could get an empty bone array
+			FScopeLock EvaluateLock(&ControlRig->GetEvaluateMutex());
+			
 			// just copy it because this is not thread safe
 			URigHierarchy* Hierarchy = ControlRig->GetHierarchy();
 
