Control Rig: Only Select Rig Controls will deselect if a control isn't clicked.

diff --git a/Engine/Plugins/Animation/ControlRig/Source/ControlRigEditor/Private/EditMode/ControlRigEditMode.cpp b/Engine/Plugins/Animation/ControlRig/Source/ControlRigEditor/Private/EditMode/ControlRigEditMode.cpp
--- a/Engine/Plugins/Animation/ControlRig/Source/ControlRigEditor/Private/EditMode/ControlRigEditMode.cpp
+++ b/Engine/Plugins/Animation/ControlRig/Source/ControlRigEditor/Private/EditMode/ControlRigEditMode.cpp
@@ -1440,20 +1440,20 @@ bool FControlRigEditMode::HandleClick(FEditorViewportClient* InViewportClient, H
 		return true;
 	}
 
-	const UControlRigEditModeSettings* Settings = GetDefault<UControlRigEditModeSettings>();
-
-	if (Settings && Settings->bOnlySelectRigControls)
-	{
-		return true;
-	}
-
-	FScopedTransaction ScopedTransaction(LOCTEXT("SelectControlTransaction", "Select Control"), IsInLevelEditor() &&  !GIsTransacting);
 	
 	// clear selected controls
 	if (Click.IsShiftDown() ==false && Click.IsControlDown() == false)
 	{
+		FScopedTransaction ScopedTransaction(LOCTEXT("SelectControlTransaction", "Select Control"), IsInLevelEditor() && !GIsTransacting);
 		ClearRigElementSelection(FRigElementTypeHelper::ToMask(ERigElementType::All));
 	}
+
+	const UControlRigEditModeSettings* Settings = GetDefault<UControlRigEditModeSettings>();
+
+	if (Settings && Settings->bOnlySelectRigControls)
+	{
+		return true;
+	}
 	/*
 	if(!InViewportClient->IsLevelEditorClient() && !InViewportClient->IsSimulateInEditorViewport())
 	{
