Control Rig: Fix initial value when spawning boolean controls not working

diff --git a/Engine/Plugins/Animation/ControlRig/Source/ControlRig/Private/Units/Execution/RigUnit_DynamicHierarchy.cpp b/Engine/Plugins/Animation/ControlRig/Source/ControlRig/Private/Units/Execution/RigUnit_DynamicHierarchy.cpp
--- a/Engine/Plugins/Animation/ControlRig/Source/ControlRig/Private/Units/Execution/RigUnit_DynamicHierarchy.cpp
+++ b/Engine/Plugins/Animation/ControlRig/Source/ControlRig/Private/Units/Execution/RigUnit_DynamicHierarchy.cpp
@@ -689,9 +689,9 @@ FRigUnit_HierarchyAddAnimationChannelBool_Execute()
 		FRigControlSettings ControlSettings;
 		ControlSettings.ControlType = ERigControlType::Bool;
 		ControlSettings.SetupLimitArrayForType(true, true, true);
-		ControlSettings.MinimumValue = FRigControlValue::Make<float>(MinimumValue ? 1.f : 0.f);
-		ControlSettings.MaximumValue = FRigControlValue::Make<float>(MaximumValue ? 1.f : 0.f);
-		const FRigControlValue Value = FRigControlValue::Make<float>(InitialValue ? 1.f : 0.f);
+		ControlSettings.MinimumValue = FRigControlValue::Make<bool>(MinimumValue);
+		ControlSettings.MaximumValue = FRigControlValue::Make<bool>(MaximumValue);
+		const FRigControlValue Value = FRigControlValue::Make<bool>(InitialValue);
 
 		FRigHierarchyControllerInstructionBracket InstructionBracket(Controller, RigVMExecuteContext.GetInstructionIndex());
 		Item = Controller->AddAnimationChannel(Name, Parent, ControlSettings, false, false);
