keep deprecated properties to guide the user thru a deprecation path

diff --git a/Engine/Plugins/Animation/ControlRig/Source/ControlRig/Private/Rigs/RigHierarchyElements.cpp b/Engine/Plugins/Animation/ControlRig/Source/ControlRig/Private/Rigs/RigHierarchyElements.cpp
--- a/Engine/Plugins/Animation/ControlRig/Source/ControlRig/Private/Rigs/RigHierarchyElements.cpp
+++ b/Engine/Plugins/Animation/ControlRig/Source/ControlRig/Private/Rigs/RigHierarchyElements.cpp
@@ -800,8 +800,8 @@ void FRigControlSettings::Load(FArchive& Ar)
 	bool bLimitTranslation_DEPRECATED = false;
 	bool bLimitRotation_DEPRECATED = false;
 	bool bLimitScale_DEPRECATED = false;
-	bool bAnimatable_DEPRECATED = false;
-	bool bShapeEnabled_DEPRECATED = false;
+	bool bAnimatableDeprecated = false;
+	bool bShapeEnabledDeprecated = false;
 
 	if (Ar.CustomVer(FControlRigObjectVersion::GUID) >= FControlRigObjectVersion::ControlAnimationType)
 	{
@@ -813,7 +813,7 @@ void FRigControlSettings::Load(FArchive& Ar)
 	Ar << bIsCurve;
 	if (Ar.CustomVer(FControlRigObjectVersion::GUID) < FControlRigObjectVersion::ControlAnimationType)
 	{
-		Ar << bAnimatable_DEPRECATED;
+		Ar << bAnimatableDeprecated;
 	}
 	if (Ar.CustomVer(FControlRigObjectVersion::GUID) < FControlRigObjectVersion::PerChannelLimits)
 	{
@@ -843,8 +843,8 @@ void FRigControlSettings::Load(FArchive& Ar)
 	
 	if (Ar.CustomVer(FControlRigObjectVersion::GUID) < FControlRigObjectVersion::ControlAnimationType)
 	{
-		Ar << bShapeEnabled_DEPRECATED;
-		SetAnimationTypeFromDeprecatedData(bAnimatable_DEPRECATED, bShapeEnabled_DEPRECATED);
+		Ar << bShapeEnabledDeprecated;
+		SetAnimationTypeFromDeprecatedData(bAnimatableDeprecated, bShapeEnabledDeprecated);
 		AnimationTypeName = AnimationTypeEnum->GetNameByValue((int64)AnimationType);
 	}
 	
diff --git a/Engine/Plugins/Animation/ControlRig/Source/ControlRig/Public/Rigs/RigHierarchyElements.h b/Engine/Plugins/Animation/ControlRig/Source/ControlRig/Public/Rigs/RigHierarchyElements.h
--- a/Engine/Plugins/Animation/ControlRig/Source/ControlRig/Public/Rigs/RigHierarchyElements.h
+++ b/Engine/Plugins/Animation/ControlRig/Source/ControlRig/Public/Rigs/RigHierarchyElements.h
@@ -1120,6 +1120,15 @@ struct CONTROLRIG_API FRigControlSettings
 	UPROPERTY(BlueprintReadOnly, EditAnywhere, Category = Animation)
 	bool bGroupWithParentControl;
 
+	/**
+	 * Deprecated properties.
+	 */
+	UPROPERTY(meta=(DeprecatedProperty, DeprecationMessage = "Use animation_type instead."))
+	bool bAnimatable_DEPRECATED;
+	
+	UPROPERTY(meta=(DeprecatedProperty, DeprecationMessage = "Use animation_type or shape_visible instead."))
+	bool bShapeEnabled_DEPRECATED;
+	
 	/** Applies the limits expressed by these settings to a value */
 	FORCEINLINE void ApplyLimits(FRigControlValue& InOutValue) const
 	{
