Constraints:
- notification from control rig bound to a new object to ensure that handles and constraints are up to date
- enable keyframe in constraint space
- ensure we exit sequencer silent mode

diff --git a/Engine/Plugins/Animation/ControlRig/Source/ControlRig/Private/Constraints/ControlRigTransformableHandle.cpp b/Engine/Plugins/Animation/ControlRig/Source/ControlRig/Private/Constraints/ControlRigTransformableHandle.cpp
--- a/Engine/Plugins/Animation/ControlRig/Source/ControlRig/Private/Constraints/ControlRigTransformableHandle.cpp
+++ b/Engine/Plugins/Animation/ControlRig/Source/ControlRig/Private/Constraints/ControlRigTransformableHandle.cpp
@@ -4,6 +4,7 @@
 #include "Constraints/ControlRigTransformableHandle.h"
 
 #include "ControlRig.h"
+#include "ControlRigObjectBinding.h"
 #include "IControlRigObjectBinding.h"
 #include "Rigs/RigHierarchyElements.h"
 #include "Channels/MovieSceneChannelProxy.h"
@@ -158,7 +159,7 @@ TWeakObjectPtr<UObject> UTransformableControlHandle::GetTarget() const
 USkeletalMeshComponent* UTransformableControlHandle::GetSkeletalMesh() const
 {
 	const TSharedPtr<IControlRigObjectBinding> ObjectBinding = ControlRig.IsValid() ? ControlRig->GetObjectBinding() : nullptr;
-	return ObjectBinding ? Cast<USkeletalMeshComponent>(ObjectBinding->GetBoundObject()) : nullptr;
+   	return ObjectBinding ? Cast<USkeletalMeshComponent>(ObjectBinding->GetBoundObject()) : nullptr;
 }
 
 bool UTransformableControlHandle::HasDirectDependencyWith(const UTransformableHandle& InOther) const
@@ -221,6 +222,12 @@ void UTransformableControlHandle::UnregisterDelegates() const
 			Hierarchy->OnModified().RemoveAll(this);
 		}
 		ControlRig->ControlModified().RemoveAll(this);
+
+		if (const TSharedPtr<IControlRigObjectBinding> Binding = ControlRig->GetObjectBinding())
+		{
+			Binding->OnControlRigBind().RemoveAll(this);
+		}
+		ControlRig->ControlRigBound().RemoveAll(this);
 	}
 }
 
@@ -240,6 +247,10 @@ void UTransformableControlHandle::RegisterDelegates()
 		}
 		
 		ControlRig->ControlModified().AddUObject(this, &UTransformableControlHandle::OnControlModified);
+		if (!ControlRig->ControlRigBound().IsBoundToObject(this))
+		{
+			ControlRig->ControlRigBound().AddUObject(this, &UTransformableControlHandle::OnControlRigBound);
+		}
 	}
 }
 
@@ -300,11 +311,49 @@ void UTransformableControlHandle::OnControlModified(
 	{
 		if(OnHandleModified.IsBound())
 		{
-			OnHandleModified.Broadcast(this, InContext.bConstraintUpdate);
+			const EHandleEvent Event = InContext.bConstraintUpdate ?
+				EHandleEvent::GlobalTransformUpdated : EHandleEvent::LocalTransformUpdated; 
+			OnHandleModified.Broadcast(this, Event);
 		}
 	}
 }
 
+void UTransformableControlHandle::OnControlRigBound(UControlRig* InControlRig)
+{
+	if (!InControlRig)
+	{
+		return;
+	}
+
+	if (!ControlRig.IsValid() || ControlRig != InControlRig)
+	{
+		return;
+	}
+
+	if (const TSharedPtr<IControlRigObjectBinding> Binding = ControlRig->GetObjectBinding())
+	{
+		if (!Binding->OnControlRigBind().IsBoundToObject(this))
+		{
+			Binding->OnControlRigBind().AddUObject(this, &UTransformableControlHandle::OnObjectBoundToControlRig);
+		}
+	}
+}
+
+void UTransformableControlHandle::OnObjectBoundToControlRig(UObject* InObject)
+{
+	if (!ControlRig.IsValid() || !InObject)
+	{
+		return;
+	}
+
+	const TSharedPtr<IControlRigObjectBinding> Binding = ControlRig->GetObjectBinding();
+	const UObject* CurrentObject = Binding ? Binding->GetBoundObject() : nullptr;
+	if (CurrentObject == InObject)
+	{
+		OnHandleModified.Broadcast(this, EHandleEvent::ComponentUpdated);
+	}
+}
+
 static TPair<const FChannelMapInfo*, int32> GetInfoAndNumFloatChannels(
 	const UControlRig* InControlRig,
 	const FName& InControlName,
diff --git a/Engine/Plugins/Animation/ControlRig/Source/ControlRig/Public/Constraints/ControlRigTransformableHandle.h b/Engine/Plugins/Animation/ControlRig/Source/ControlRig/Public/Constraints/ControlRigTransformableHandle.h
--- a/Engine/Plugins/Animation/ControlRig/Source/ControlRig/Public/Constraints/ControlRigTransformableHandle.h
+++ b/Engine/Plugins/Animation/ControlRig/Source/ControlRig/Public/Constraints/ControlRigTransformableHandle.h
@@ -8,12 +8,12 @@
 
 #include "ControlRigTransformableHandle.generated.h"
 
-class UControlRig;
-class USkeletalMeshComponent;
 struct FRigBaseElement;
-class URigHierarchy;
 struct FRigControlElement;
 
+class UControlRig;
+class USkeletalMeshComponent;
+class URigHierarchy;
 
 /**
  * UTransformableControlHandle
@@ -25,6 +25,7 @@ class CONTROLRIG_API UTransformableControlHandle : public UTransformableHandle
 	GENERATED_BODY()
 	
 public:
+	
 	virtual ~UTransformableControlHandle();
 
 	virtual void PostLoad() override;
@@ -102,6 +103,9 @@ private:
 		URigHierarchy* InHierarchy,
 		const FRigBaseElement* InElement);
 
+	void OnControlRigBound(UControlRig* InControlRig);
+	void OnObjectBoundToControlRig(UObject* InObject);
+
 #if WITH_EDITOR
 	/** @todo document */
 	void OnObjectsReplaced(const TMap<UObject*, UObject*>& InOldToNewInstances);
diff --git a/Engine/Plugins/Animation/ControlRig/Source/ControlRig/Public/ControlRig.h b/Engine/Plugins/Animation/ControlRig/Source/ControlRig/Public/ControlRig.h
--- a/Engine/Plugins/Animation/ControlRig/Source/ControlRig/Public/ControlRig.h
+++ b/Engine/Plugins/Animation/ControlRig/Source/ControlRig/Public/ControlRig.h
@@ -87,6 +87,9 @@ public:
 	DECLARE_EVENT_OneParam(UControlRig, FOnEndLoadPackage, UControlRig*);
 #endif
 
+	/** Bindable event to notify object binding change. */
+	DECLARE_EVENT_OneParam(UControlRig, FControlRigBoundEvent, UControlRig*);
+
 	static const FName OwnerComponent;
 
 	UFUNCTION(BlueprintCallable, Category = ControlRig)
@@ -243,6 +246,7 @@ public:
 	FORCEINLINE void SetObjectBinding(TSharedPtr<IControlRigObjectBinding> InObjectBinding)
 	{
 		ObjectBinding = InObjectBinding;
+		OnControlRigBound.Broadcast(this);
 	}
 
 	FORCEINLINE TSharedPtr<IControlRigObjectBinding> GetObjectBinding() const
@@ -435,6 +439,8 @@ public:
 	// Returns an event that can be used to subscribe to
 	// Undo Bracket requests such as Open and Close.
 	FControlUndoBracketEvent & ControlUndoBracket() { return OnControlUndoBracket; }
+	
+	FControlRigBoundEvent& ControlRigBound() { return OnControlRigBound; };
 
 	bool IsCurveControl(const FRigControlElement* InControlElement) const;
 
@@ -799,6 +805,7 @@ protected:
 	FControlModifiedEvent OnControlModified;
 	FControlSelectedEvent OnControlSelected;
 	FControlUndoBracketEvent OnControlUndoBracket;
+	FControlRigBoundEvent OnControlRigBound;
 
 	UPROPERTY(BlueprintAssignable, Category = ControlRig, meta=(DisplayName="OnControlSelected"))
 	FOnControlSelectedBP OnControlSelected_BP;
diff --git a/Engine/Plugins/Animation/ControlRig/Source/ControlRigEditor/Private/Sequencer/ControlRigParameterTrackEditor.cpp b/Engine/Plugins/Animation/ControlRig/Source/ControlRigEditor/Private/Sequencer/ControlRigParameterTrackEditor.cpp
--- a/Engine/Plugins/Animation/ControlRig/Source/ControlRigEditor/Private/Sequencer/ControlRigParameterTrackEditor.cpp
+++ b/Engine/Plugins/Animation/ControlRig/Source/ControlRigEditor/Private/Sequencer/ControlRigParameterTrackEditor.cpp
@@ -80,6 +80,7 @@
 #include "Misc/TransactionObjectEvent.h"
 #include "ConstraintChannelHelper.h"
 #include "MovieSceneConstraintChannelHelper.h"
+#include "Constraints/ControlRigTransformableHandle.h"
 
 #define LOCTEXT_NAMESPACE "FControlRigParameterTrackEditor"
 
@@ -226,9 +227,10 @@ FControlRigParameterTrackEditor::FControlRigParameterTrackEditor(TSharedRef<ISeq
 		//the other is if bound object on the edit mode is null we request a re-evaluate which will reset it up.
 		FDelegateHandle OnObjectsReplacedHandle = FCoreUObjectDelegates::OnObjectsReplaced.AddLambda([&](const TMap<UObject*, UObject*>& ReplacementMap)
 		{
-			static bool bHasEnteredSilent = false;
 			if (GetSequencer().IsValid())
 			{
+				static bool bHasEnteredSilent = false;
+				
 				bool bRequestEvaluate = false;
 				TMap<UControlRig*, UControlRig*> OldToNewControlRigs;
 				FControlRigEditMode* ControlRigEditMode = GetEditMode();
@@ -352,17 +354,19 @@ FControlRigParameterTrackEditor::FControlRigParameterTrackEditor(TSharedRef<ISeq
 					}
 					if (!ControlRigEditMode)
 					{
-						if (bHasEnteredSilent == true)
-						{
-							GetSequencer()->ExitSilentMode();
-							bHasEnteredSilent = false;
-						}
 						if (bRequestEvaluate)
 						{
 							GetSequencer()->NotifyMovieSceneDataChanged(EMovieSceneDataChangeType::MovieSceneStructureItemsChanged);
 						}
 					}
 				}
+
+				// ensure we exit silent mode if it has been entered
+				if (bHasEnteredSilent)
+				{
+					GetSequencer()->ExitSilentMode();
+					bHasEnteredSilent = false;
+				}
 			}
 		});
 	
@@ -405,6 +409,8 @@ void FControlRigParameterTrackEditor::BindControlRig(UControlRig* ControlRig)
 		ControlRig->OnInitialized_AnyThread().AddRaw(this, &FControlRigParameterTrackEditor::HandleOnInitialized);
 		ControlRig->ControlSelected().AddRaw(this, &FControlRigParameterTrackEditor::HandleControlSelected);
 		ControlRig->ControlUndoBracket().AddRaw(this, &FControlRigParameterTrackEditor::HandleControlUndoBracket);
+		ControlRig->ControlRigBound().AddRaw(this, &FControlRigParameterTrackEditor::HandleOnControlRigBound);
+		
 		BoundControlRigs.Add(ControlRig);
 		UMovieSceneControlRigParameterTrack* Track = FindTrack(ControlRig);
 		if (Track)
@@ -420,6 +426,12 @@ void FControlRigParameterTrackEditor::BindControlRig(UControlRig* ControlRig)
 						{
 							HandleOnSpaceAdded(Section, Channel.ControlName, &(Channel.SpaceCurve));
 						}
+						
+						TArray<FConstraintAndActiveChannel>& ConstraintChannels = Section->GetConstraintsChannels();
+						for (FConstraintAndActiveChannel& Channel: ConstraintChannels)
+						{ 
+							HandleOnConstraintAdded(Section, &(Channel.ActiveChannel));
+						}
 					}
 				}
 			}
@@ -441,6 +453,12 @@ void FControlRigParameterTrackEditor::UnbindControlRig(UControlRig* ControlRig)
 		ControlRig->ControlModified().RemoveAll(this);
 		ControlRig->OnInitialized_AnyThread().RemoveAll(this);
 		ControlRig->ControlSelected().RemoveAll(this);
+		if (const TSharedPtr<IControlRigObjectBinding> Binding = ControlRig->GetObjectBinding())
+		{
+			Binding->OnControlRigBind().RemoveAll(this);
+		}
+		ControlRig->ControlRigBound().RemoveAll(this);
+		
 		BoundControlRigs.Remove(ControlRig);
 		ClearOutAllSpaceAndConstraintDelegates(ControlRig);
 	}
@@ -1555,6 +1573,7 @@ void FControlRigParameterTrackEditor::OnAddTransformKeysForSelectedObjects(EMovi
 	const EControlRigContextChannelToKey ChannelsToKey = static_cast<EControlRigContextChannelToKey>(Channel); 
 	FScopedTransaction KeyTransaction(LOCTEXT("SetKeysOnControls", "Set Keys On Controls"), !GIsTransacting);
 
+	static constexpr bool bInConstraintSpace = true;
 	for (const TPair<UControlRig*, TArray<FRigElementKey>>& Selection : SelectedControls)
 	{
 		UControlRig* ControlRig = Selection.Key;
@@ -1567,7 +1586,8 @@ void FControlRigParameterTrackEditor::OnAddTransformKeysForSelectedObjects(EMovi
 				const TArray<FName> ControlNames = ControlRig->CurrentControlSelection();
 				for (const FName& ControlName : ControlNames)
 				{
-					AddControlKeys(Component, ControlRig, Name, ControlName, ChannelsToKey, ESequencerKeyMode::ManualKeyForced, FLT_MAX);
+					AddControlKeys(Component, ControlRig, Name, ControlName, ChannelsToKey,
+						ESequencerKeyMode::ManualKeyForced, FLT_MAX, bInConstraintSpace);
 				}
 			}
 		}
@@ -2878,13 +2898,80 @@ void FControlRigParameterTrackEditor::HandleControlUndoBracket(UControlRig* Subj
 	}
 }
 
+void FControlRigParameterTrackEditor::HandleOnControlRigBound(UControlRig* InControlRig)
+{
+	if (!InControlRig)
+	{
+		return;
+	}
+	
+	UMovieSceneControlRigParameterTrack* Track = FindTrack(InControlRig);
+	if (!Track)
+	{
+		return;
+	}
+
+	const TSharedPtr<IControlRigObjectBinding> Binding = InControlRig->GetObjectBinding();
+	
+	for (UMovieSceneSection* BaseSection : Track->GetAllSections())
+	{
+		if (UMovieSceneControlRigParameterSection* Section = Cast< UMovieSceneControlRigParameterSection>(BaseSection))
+		{
+			const UControlRig* ControlRig = Section->GetControlRig();
+			if (ControlRig && InControlRig == ControlRig)
+			{
+				if (!Binding->OnControlRigBind().IsBoundToObject(this))
+				{
+					Binding->OnControlRigBind().AddRaw(this, &FControlRigParameterTrackEditor::HandleOnObjectBoundToControlRig);
+				}
+			}
+		}
+	}
+}
+
+void FControlRigParameterTrackEditor::HandleOnObjectBoundToControlRig(UObject* InObject)
+{
+	// look for sections to update
+	TArray<UMovieSceneControlRigParameterSection*> SectionsToUpdate;
+	for (const TWeakObjectPtr<UControlRig>& ControlRigPtr : BoundControlRigs)
+	{
+		const TSharedPtr<IControlRigObjectBinding> Binding =
+			ControlRigPtr.IsValid() ? ControlRigPtr->GetObjectBinding() : nullptr;
+		const UObject* CurrentObject = Binding ? Binding->GetBoundObject() : nullptr;
+		if (CurrentObject == InObject)
+		{
+			if (const UMovieSceneControlRigParameterTrack* Track = FindTrack(ControlRigPtr.Get()))
+			{
+				for (UMovieSceneSection* BaseSection : Track->GetAllSections())
+				{
+					if (UMovieSceneControlRigParameterSection* Section = Cast<UMovieSceneControlRigParameterSection>(BaseSection))
+					{
+						SectionsToUpdate.AddUnique(Section);
+					}
+				}
+			}
+		}
+	}
+
+	// reconstruct proxies
+	if (!SectionsToUpdate.IsEmpty())
+	{
+		for (UMovieSceneControlRigParameterSection* Section: SectionsToUpdate)
+		{
+			Section->ReconstructChannelProxy();
+			Section->MarkAsChanged();
+		}
+	}
+}
+
 void FControlRigParameterTrackEditor::GetControlRigKeys(
 	UControlRig* InControlRig,
 	FName ParameterName,
 	EControlRigContextChannelToKey ChannelsToKey,
 	ESequencerKeyMode KeyMode,
 	UMovieSceneControlRigParameterSection* SectionToKey,
-	FGeneratedTrackKeys& OutGeneratedKeys)
+	FGeneratedTrackKeys& OutGeneratedKeys,
+	const bool bInConstraintSpace)
 {
 	const TArray<bool>& ControlsMask = SectionToKey->GetControlsMask();
 	EMovieSceneTransformChannel TransformMask = SectionToKey->GetTransformMask().GetChannels();
@@ -3021,6 +3108,19 @@ void FControlRigParameterTrackEditor::GetControlRigKeys(
 					Scale = Val.GetScale3D();
 				}
 
+				// switch values to constraint space?
+				if (bInConstraintSpace)
+				{
+					const uint32 ControlHash = UTransformableControlHandle::ComputeHash(InControlRig, ControlElement->GetName());
+					TOptional<FTransform> Transform = FTransformConstraintUtils::GetConstraintRelativeTransform(InControlRig->GetWorld(), ControlHash);
+					if (Transform)
+					{
+						Translation = Transform->GetTranslation();
+						Rotation = Transform->GetRotation().Rotator();
+						Scale = Transform->GetScale3D();
+					}
+				}
+					
 				FVector3f CurrentVector = (FVector3f)Translation;
 				bool bKeyX = bSetKey && EnumHasAnyFlags(ChannelsToKey, EControlRigContextChannelToKey::TranslationX);
 				bool bKeyY = bSetKey && EnumHasAnyFlags(ChannelsToKey, EControlRigContextChannelToKey::TranslationY);
@@ -3229,7 +3329,8 @@ void FControlRigParameterTrackEditor::AddControlKeys(
 	FName RigControlName,
 	EControlRigContextChannelToKey ChannelsToKey,
 	ESequencerKeyMode KeyMode,
-	float InLocalTime)
+	float InLocalTime,
+	const bool bInConstraintSpace)
 {
 	if (KeyMode == ESequencerKeyMode::ManualKey || (GetSequencer().IsValid() && !GetSequencer()->IsAllowedToChange()))
 	{
@@ -3271,8 +3372,8 @@ void FControlRigParameterTrackEditor::AddControlKeys(
 	}
 
 	TSharedRef<FGeneratedTrackKeys> GeneratedKeys = MakeShared<FGeneratedTrackKeys>();
-
-	GetControlRigKeys(InControlRig, RigControlName, ChannelsToKey, KeyMode, ParamSection, *GeneratedKeys);
+	GetControlRigKeys(InControlRig, RigControlName, ChannelsToKey, KeyMode, ParamSection, *GeneratedKeys, bInConstraintSpace);
+	
 	TGuardValue<bool> Guard(bIsDoingSelection, true);
 
 	auto OnKeyProperty = [=](FFrameNumber Time) -> FKeyPropertyResult
diff --git a/Engine/Plugins/Animation/ControlRig/Source/ControlRigEditor/Private/Sequencer/ControlRigParameterTrackEditor.h b/Engine/Plugins/Animation/ControlRig/Source/ControlRigEditor/Private/Sequencer/ControlRigParameterTrackEditor.h
--- a/Engine/Plugins/Animation/ControlRig/Source/ControlRigEditor/Private/Sequencer/ControlRigParameterTrackEditor.h
+++ b/Engine/Plugins/Animation/ControlRig/Source/ControlRigEditor/Private/Sequencer/ControlRigParameterTrackEditor.h
@@ -127,6 +127,8 @@ private:
 	void HandleControlSelected(UControlRig* Subject, FRigControlElement* ControlElement, bool bSelected);
 	void HandleControlUndoBracket(UControlRig* Subject, bool bOpenUndoBracket);
 	void HandleOnInitialized(UControlRig* Subject, const EControlRigState InState, const FName& InEventName);
+	void HandleOnControlRigBound(UControlRig* InControlRig);
+	void HandleOnObjectBoundToControlRig(UObject* InObject);
 
 	//if rig not set then we clear delegates for everyone
 	void ClearOutAllSpaceAndConstraintDelegates(const UControlRig* InOptionalControlRig = nullptr) const;
@@ -200,8 +202,12 @@ private:
 
 public:
 
-	void AddControlKeys(USceneComponent *InSceneComp, UControlRig* InControlRig, FName PropertyName, FName ParameterName, EControlRigContextChannelToKey ChannelsToKey, ESequencerKeyMode KeyMode, float InLocalTime);
-	void GetControlRigKeys(UControlRig* InControlRig, FName ParameterName, EControlRigContextChannelToKey ChannelsToKey, ESequencerKeyMode KeyMode, UMovieSceneControlRigParameterSection* SectionToKey, FGeneratedTrackKeys& OutGeneratedKeys);
+	void AddControlKeys(USceneComponent *InSceneComp, UControlRig* InControlRig, FName PropertyName,
+		FName ParameterName, EControlRigContextChannelToKey ChannelsToKey, ESequencerKeyMode KeyMode,
+		float InLocalTime, const bool bInConstraintSpace = false);
+	void GetControlRigKeys(UControlRig* InControlRig, FName ParameterName, EControlRigContextChannelToKey ChannelsToKey,
+		ESequencerKeyMode KeyMode, UMovieSceneControlRigParameterSection* SectionToKey,	FGeneratedTrackKeys& OutGeneratedKeys,
+		const bool bInConstraintSpace = false);
 	FKeyPropertyResult AddKeysToControlRig(
 		USceneComponent *InSceneComp, UControlRig* InControlRig, FFrameNumber KeyTime, FGeneratedTrackKeys& GeneratedKeys,
 		ESequencerKeyMode KeyMode, TSubclassOf<UMovieSceneTrack> TrackClass, FName ControlRigName, FName RigControlName);
diff --git a/Engine/Source/Runtime/Experimental/Animation/Constraints/Private/TransformConstraint.cpp b/Engine/Source/Runtime/Experimental/Animation/Constraints/Private/TransformConstraint.cpp
--- a/Engine/Source/Runtime/Experimental/Animation/Constraints/Private/TransformConstraint.cpp
+++ b/Engine/Source/Runtime/Experimental/Animation/Constraints/Private/TransformConstraint.cpp
@@ -191,7 +191,7 @@ bool UTickableTransformConstraint::ReferencesObject(TWeakObjectPtr<UObject> InOb
 	const TWeakObjectPtr<UObject> ChildTarget = ChildTRSHandle->IsValid() ? ChildTRSHandle->GetTarget() : nullptr;
 	if (ChildTarget == InObject)
 	{
-		return true;	
+		return true;
 	}
 
 	const TWeakObjectPtr<UObject> ParentTarget = ParentTRSHandle->IsValid() ? ParentTRSHandle->GetTarget() : nullptr;
@@ -205,7 +205,7 @@ bool UTickableTransformConstraint::ReferencesObject(TWeakObjectPtr<UObject> InOb
 
 void UTickableTransformConstraint::SetChildGlobalTransform(const FTransform& InGlobal) const
 {
-	if(ChildTRSHandle->IsValid())
+	if(ChildTRSHandle && ChildTRSHandle->IsValid())
 	{
 		ChildTRSHandle->SetGlobalTransform(InGlobal);
 	}
@@ -213,7 +213,7 @@ void UTickableTransformConstraint::SetChildGlobalTransform(const FTransform& InG
 
 void UTickableTransformConstraint::SetChildLocalTransform(const FTransform& InLocal) const
 {
-	if(ChildTRSHandle->IsValid())
+	if(ChildTRSHandle && ChildTRSHandle->IsValid())
 	{
 		ChildTRSHandle->SetLocalTransform(InLocal);
 	}
@@ -221,26 +221,40 @@ void UTickableTransformConstraint::SetChildLocalTransform(const FTransform& InLo
 
 FTransform UTickableTransformConstraint::GetChildGlobalTransform() const
 {
-	return ChildTRSHandle->IsValid() ? ChildTRSHandle->GetGlobalTransform() : FTransform::Identity;
+	return (ChildTRSHandle && ChildTRSHandle->IsValid()) ? ChildTRSHandle->GetGlobalTransform() : FTransform::Identity;
 }
 
 FTransform UTickableTransformConstraint::GetChildLocalTransform() const
 {
-	return ChildTRSHandle->IsValid() ? ChildTRSHandle->GetLocalTransform() : FTransform::Identity;
+	return (ChildTRSHandle && ChildTRSHandle->IsValid()) ? ChildTRSHandle->GetLocalTransform() : FTransform::Identity;
 }
 
 FTransform UTickableTransformConstraint::GetParentGlobalTransform() const
 {
-	return ParentTRSHandle->IsValid() ? ParentTRSHandle->GetGlobalTransform() : FTransform::Identity;
+	return (ParentTRSHandle && ParentTRSHandle->IsValid()) ? ParentTRSHandle->GetGlobalTransform() : FTransform::Identity;
 }
 
 FTransform UTickableTransformConstraint::GetParentLocalTransform() const
 {
-	return ParentTRSHandle->IsValid() ? ParentTRSHandle->GetLocalTransform() : FTransform::Identity;
+	return (ParentTRSHandle && ParentTRSHandle->IsValid()) ? ParentTRSHandle->GetLocalTransform() : FTransform::Identity;
 }
 
-void UTickableTransformConstraint::OnHandleModified(UTransformableHandle* InHandle, bool bUpdate)
-{}
+void UTickableTransformConstraint::OnHandleModified(UTransformableHandle* InHandle, EHandleEvent InNotification)
+{
+	if (!InHandle)
+	{
+		return;
+	}
+
+	if(InHandle == ChildTRSHandle || InHandle == ParentTRSHandle)
+	{
+		if (InNotification == EHandleEvent::ComponentUpdated)
+		{
+			SetupDependencies();
+			return;
+		}
+	}
+}
 
 /** 
  * UTickableTranslationConstraint
@@ -316,8 +330,10 @@ FConstraintTickFunction::ConstraintFunction UTickableTranslationConstraint::GetF
 	};
 }
 
-void UTickableTranslationConstraint::OnHandleModified(UTransformableHandle* InHandle, bool bUpdate)
+void UTickableTranslationConstraint::OnHandleModified(UTransformableHandle* InHandle, EHandleEvent InEvent)
 {
+	Super::OnHandleModified(InHandle, InEvent);
+	
 	if (!Active || !bDynamicOffset)
 	{
 		return;
@@ -328,6 +344,13 @@ void UTickableTranslationConstraint::OnHandleModified(UTransformableHandle* InHa
 		return;
 	}
 
+	const bool bUpdateFromGlobal = (InEvent == EHandleEvent::GlobalTransformUpdated);
+	const bool bUpdateTransform = InEvent == EHandleEvent::LocalTransformUpdated || bUpdateFromGlobal;
+	if (!bUpdateTransform)
+	{
+		return;
+	}
+
 	const uint32 InputHash = CalculateInputHash();
 	
 	// update dynamic offset
@@ -335,7 +358,7 @@ void UTickableTranslationConstraint::OnHandleModified(UTransformableHandle* InHa
 	{
 		Cache.CachedInputHash = InputHash;
 
-		if (bUpdate)
+		if (bUpdateFromGlobal)
 		{
 			const FTransform ParentWorldTransform = GetParentGlobalTransform();
 			const FTransform ChildGlobalTransform = GetChildGlobalTransform();
@@ -440,8 +463,10 @@ FConstraintTickFunction::ConstraintFunction UTickableRotationConstraint::GetFunc
 	};
 }
 
-void UTickableRotationConstraint::OnHandleModified(UTransformableHandle* InHandle, bool bUpdate)
+void UTickableRotationConstraint::OnHandleModified(UTransformableHandle* InHandle, EHandleEvent InEvent)
 {
+	Super::OnHandleModified(InHandle, InEvent);
+	
 	if (!Active || !bDynamicOffset)
 	{
 		return;
@@ -452,6 +477,13 @@ void UTickableRotationConstraint::OnHandleModified(UTransformableHandle* InHandl
 		return;
 	}
 
+	const bool bUpdateFromGlobal = (InEvent == EHandleEvent::GlobalTransformUpdated);
+	const bool bUpdateTransform = (InEvent == EHandleEvent::LocalTransformUpdated) || bUpdateFromGlobal;
+	if (!bUpdateTransform)
+	{
+		return;
+	}
+
 	const uint32 InputHash = CalculateInputHash();
 	
 	// update dynamic offset
@@ -459,7 +491,7 @@ void UTickableRotationConstraint::OnHandleModified(UTransformableHandle* InHandl
 	{
 		Cache.CachedInputHash = InputHash;
 
-		if (bUpdate)
+		if (bUpdateFromGlobal)
 		{
 			const FTransform ParentWorldTransform = GetParentGlobalTransform();
 			const FTransform ChildGlobalTransform = GetChildGlobalTransform();
@@ -623,8 +655,10 @@ FConstraintTickFunction::ConstraintFunction UTickableParentConstraint::GetFuncti
 	};
 }
 
-void UTickableParentConstraint::OnHandleModified(UTransformableHandle* InHandle, bool bUpdate)
+void UTickableParentConstraint::OnHandleModified(UTransformableHandle* InHandle, EHandleEvent InEvent)
 {
+	Super::OnHandleModified(InHandle, InEvent);
+	
 	if (!Active || !bDynamicOffset)
 	{
 		return;
@@ -635,14 +669,21 @@ void UTickableParentConstraint::OnHandleModified(UTransformableHandle* InHandle,
 		return;
 	}
 
+	const bool bUpdateFromGlobal = (InEvent == EHandleEvent::GlobalTransformUpdated);
+	const bool bUpdateTransform = (InEvent == EHandleEvent::LocalTransformUpdated) || bUpdateFromGlobal;
+	if (!bUpdateTransform)
+	{
+		return;
+	}
+	
 	const uint32 InputHash = CalculateInputHash();
 
 	// update dynamic offset
 	if (InputHash != Cache.CachedInputHash)
 	{
 		Cache.CachedInputHash = InputHash;
-		
-		if (bUpdate)
+
+		if (bUpdateFromGlobal)
 		{
 			const FTransform ParentWorldTransform = GetParentGlobalTransform();
 			const FTransform ChildGlobalTransform = GetChildGlobalTransform();
@@ -1041,4 +1082,44 @@ FTransform FTransformConstraintUtils::ComputeRelativeTransform(
 	}
 	
 	return InChildWorld.GetRelativeTransform(InSpaceWorld);
+}
+
+TOptional<FTransform> FTransformConstraintUtils::GetConstraintRelativeTransform(UWorld* InWorld, const uint32 InHandleHash)
+{
+	if (!InWorld || InHandleHash <= 0)
+	{
+		return TOptional<FTransform>();
+	}
+
+	static constexpr bool bSorted = true;
+	const FConstraintsManagerController& Controller = FConstraintsManagerController::Get(InWorld);
+	const TArray<TObjectPtr<UTickableConstraint>> Constraints = Controller.GetParentConstraints(InHandleHash, bSorted);
+
+	if (Constraints.IsEmpty())
+	{
+		return TOptional<FTransform>();
+	}
+
+	// get current active transform constraint 	
+	const int32 LastActiveIndex = Constraints.FindLastByPredicate([](const TObjectPtr<UTickableConstraint>& InConstraint)
+    {
+    	if (const UTickableTransformConstraint* TransformConstraint = Cast<UTickableTransformConstraint>(InConstraint.Get()))
+    	{
+    		return InConstraint->Active && TransformConstraint->bDynamicOffset;
+    	}
+    	return false;
+    });
+
+	if (!Constraints.IsValidIndex(LastActiveIndex))
+	{
+		return TOptional<FTransform>();
+	}
+
+	// get relative transform
+	const UTickableTransformConstraint* Constraint = Cast<UTickableTransformConstraint>(Constraints[LastActiveIndex]);
+	const FTransform ChildLocal = Constraint->GetChildLocalTransform();
+	const FTransform ChildGlobal = Constraint->GetChildGlobalTransform();
+	const FTransform ParentGlobal = Constraint->GetParentGlobalTransform();
+
+	return ComputeRelativeTransform(ChildLocal, ChildGlobal, ParentGlobal, Constraint);
 }
\ No newline at end of file
diff --git a/Engine/Source/Runtime/Experimental/Animation/Constraints/Private/TransformableHandle.cpp b/Engine/Source/Runtime/Experimental/Animation/Constraints/Private/TransformableHandle.cpp
--- a/Engine/Source/Runtime/Experimental/Animation/Constraints/Private/TransformableHandle.cpp
+++ b/Engine/Source/Runtime/Experimental/Animation/Constraints/Private/TransformableHandle.cpp
@@ -157,7 +157,7 @@ void UTransformableComponentHandle::OnActorMoving(AActor* InActor)
 
 	if(OnHandleModified.IsBound())
 	{
-		OnHandleModified.Broadcast(this, true);
+		OnHandleModified.Broadcast(this, EHandleEvent::GlobalTransformUpdated);
 	}
 }
 
@@ -202,7 +202,7 @@ void UTransformableComponentHandle::OnPostPropertyChanged(
 
 	if(OnHandleModified.IsBound())
 	{
-		OnHandleModified.Broadcast(this, true);
+		OnHandleModified.Broadcast(this, EHandleEvent::GlobalTransformUpdated);
 	}
 }
 
diff --git a/Engine/Source/Runtime/Experimental/Animation/Constraints/Public/TransformConstraint.h b/Engine/Source/Runtime/Experimental/Animation/Constraints/Public/TransformConstraint.h
--- a/Engine/Source/Runtime/Experimental/Animation/Constraints/Public/TransformConstraint.h
+++ b/Engine/Source/Runtime/Experimental/Animation/Constraints/Public/TransformConstraint.h
@@ -8,6 +8,7 @@
 
 #include "TransformConstraint.generated.h"
 
+enum class EHandleEvent : uint8;
 class UTransformableHandle;
 class UTransformableComponentHandle;
 class USceneComponent;
@@ -84,7 +85,7 @@ public:
 	 * Manages changes on the child/parent transformable handle. This can be used to update internal data (i.e. offset)
 	 * when transform changes outside of that system and need to trigger changes within the constraint itself.   
 	*/
-	virtual void OnHandleModified(UTransformableHandle* InHandle, bool bUpdate);
+	virtual void OnHandleModified(UTransformableHandle* InHandle, EHandleEvent InEvent);
 
 protected:
 
@@ -146,7 +147,7 @@ public:
 	virtual FConstraintTickFunction::ConstraintFunction GetFunction() const override;
 
 	/** Updates the dynamic offset based on external child's transform changes. */
-	virtual void OnHandleModified(UTransformableHandle* InHandle, bool bUpdate) override;
+	virtual void OnHandleModified(UTransformableHandle* InHandle, EHandleEvent InEvent) override;
 
 protected:
 
@@ -189,7 +190,7 @@ public:
 	virtual FConstraintTickFunction::ConstraintFunction GetFunction() const override;
 
 	/** Updates the dynamic offset based on external child's transform changes. */
-	virtual void OnHandleModified(UTransformableHandle* InHandle, bool bUpdate) override;
+	virtual void OnHandleModified(UTransformableHandle* InHandle, EHandleEvent InEvent) override;
 
 protected:
 
@@ -262,7 +263,7 @@ public:
 	
 
 	/** Updates the dynamic offset based on external child's transform changes. */
-	virtual void OnHandleModified(UTransformableHandle* InHandle, bool bUpdate) override;
+	virtual void OnHandleModified(UTransformableHandle* InHandle, EHandleEvent InEvent) override;
 
 protected:
 	/** Cache data structure to store last child local/global transform. */
@@ -367,4 +368,7 @@ struct CONSTRAINTS_API FTransformConstraintUtils
 		const FTransform& InChildWorld,
 		const FTransform& InSpaceWorld,
 		const UTickableTransformConstraint* InConstraint);
+
+	/** Computes the current constraint space local transform. */
+	static TOptional<FTransform> GetConstraintRelativeTransform(UWorld* InWorld, const uint32 InHandleHash);
 };
diff --git a/Engine/Source/Runtime/Experimental/Animation/Constraints/Public/TransformableHandle.h b/Engine/Source/Runtime/Experimental/Animation/Constraints/Public/TransformableHandle.h
--- a/Engine/Source/Runtime/Experimental/Animation/Constraints/Public/TransformableHandle.h
+++ b/Engine/Source/Runtime/Experimental/Animation/Constraints/Public/TransformableHandle.h
@@ -15,6 +15,17 @@ struct FFRameNumber;
 struct FFrameRate;
 enum class EMovieSceneTransformChannel : uint32;
 
+UENUM()
+enum class EHandleEvent : uint8
+{
+	LocalTransformUpdated,
+	GlobalTransformUpdated,
+	ComponentUpdated,
+
+	/** MAX - invalid */
+	Max UMETA(Hidden),
+};
+
 /**
  * UTransformableHandle
  */
@@ -26,7 +37,7 @@ class CONSTRAINTS_API UTransformableHandle : public UObject
 	
 public:
 
-	DECLARE_EVENT_TwoParams(UTransformableHandle, FHandleModifiedEvent, UTransformableHandle*, bool);
+	DECLARE_EVENT_TwoParams(UTransformableHandle, FHandleModifiedEvent, UTransformableHandle*, EHandleEvent);
 	
 	virtual ~UTransformableHandle();
 
diff --git a/Engine/Source/Runtime/MovieSceneTracks/Private/Systems/MovieSceneConstraintSystem.cpp b/Engine/Source/Runtime/MovieSceneTracks/Private/Systems/MovieSceneConstraintSystem.cpp
--- a/Engine/Source/Runtime/MovieSceneTracks/Private/Systems/MovieSceneConstraintSystem.cpp
+++ b/Engine/Source/Runtime/MovieSceneTracks/Private/Systems/MovieSceneConstraintSystem.cpp
@@ -149,7 +149,7 @@ void UMovieSceneConstraintSystem::OnRun(FSystemTaskPrerequisites& InPrerequisite
 		{
 			if (UpdateHandle.Constraint.IsValid() && UpdateHandle.TransformHandle.IsValid())
 			{
-				UpdateHandle.Constraint->OnHandleModified(UpdateHandle.TransformHandle.Get(),false);
+				UpdateHandle.Constraint->OnHandleModified(UpdateHandle.TransformHandle.Get(),EHandleEvent::LocalTransformUpdated);
 			}
 		}
 		DynamicOffsets.Reset();
