fix crash when sequencer is invalid

diff --git a/Engine/Plugins/Animation/ControlRig/Source/ControlRigEditor/Private/Sequencer/ControlRigParameterTrackEditor.cpp b/Engine/Plugins/Animation/ControlRig/Source/ControlRigEditor/Private/Sequencer/ControlRigParameterTrackEditor.cpp
--- a/Engine/Plugins/Animation/ControlRig/Source/ControlRigEditor/Private/Sequencer/ControlRigParameterTrackEditor.cpp
+++ b/Engine/Plugins/Animation/ControlRig/Source/ControlRigEditor/Private/Sequencer/ControlRigParameterTrackEditor.cpp
@@ -2327,7 +2327,12 @@ FMovieSceneTrackEditor::FFindOrCreateTrackResult FControlRigParameterTrackEditor
 }
 
 UMovieSceneControlRigParameterTrack* FControlRigParameterTrackEditor::FindTrack(UControlRig* InControlRig)
-{	
+{
+	if (!GetSequencer().IsValid())
+	{
+		return nullptr;
+	}
+	
 	UMovieScene* MovieScene = GetSequencer()->GetFocusedMovieSceneSequence()->GetMovieScene();
 	if (MovieScene)
 	{
