Constraints: Sequencer: Make sure detaching actors work with auto key, need to turn it off while doing so or it keys wrong spot.

diff --git a/Engine/Plugins/Animation/ControlRig/Source/ControlRigEditor/Private/Sequencer/ConstraintChannelHelper.cpp b/Engine/Plugins/Animation/ControlRig/Source/ControlRigEditor/Private/Sequencer/ConstraintChannelHelper.cpp
--- a/Engine/Plugins/Animation/ControlRig/Source/ControlRigEditor/Private/Sequencer/ConstraintChannelHelper.cpp
+++ b/Engine/Plugins/Animation/ControlRig/Source/ControlRigEditor/Private/Sequencer/ConstraintChannelHelper.cpp
@@ -343,7 +343,14 @@ void FConstraintChannelHelper::SmartComponentConstraintKey(
 				//Todo do we need to evaluate all constraints?
 				InConstraint->SetActive(true); //will be false
 				InConstraint->Evaluate();
+
 				//need to fire this event so the transform values set by the constraint propragate to the transform section
+				//first turn off autokey though
+				EAutoChangeMode AutoChangeMode = InSequencer->GetAutoChangeMode();
+				if (AutoChangeMode == EAutoChangeMode::AutoKey || AutoChangeMode == EAutoChangeMode::All)
+				{
+					InSequencer->SetAutoChangeMode(EAutoChangeMode::None);
+				};
 				FProperty* TransformProperty = FindFProperty<FProperty>(USceneComponent::StaticClass(), USceneComponent::GetRelativeLocationPropertyName());
 				FEditPropertyChain PropertyChain;
 				PropertyChain.AddHead(TransformProperty);
@@ -351,6 +358,10 @@ void FConstraintChannelHelper::SmartComponentConstraintKey(
 				FPropertyChangedEvent PropertyChangedEvent(TransformProperty, EPropertyChangeType::ValueSet);
 				FCoreUObjectDelegates::OnObjectPropertyChanged.Broadcast(Actor, PropertyChangedEvent);
 				InSequencer->RequestEvaluate();
+				if (AutoChangeMode == EAutoChangeMode::AutoKey || AutoChangeMode == EAutoChangeMode::All)
+				{
+					InSequencer->SetAutoChangeMode(AutoChangeMode);
+				}
 			}
 		}
 	}
