ControlRig: mark section as changed when the channel proxy changes

This lets the UI known that the control rig parameter track potentially needs rebuilding.

diff --git a/Engine/Plugins/Animation/ControlRig/Source/ControlRig/Private/Sequencer/MovieSceneControlRigParameterSection.cpp b/Engine/Plugins/Animation/ControlRig/Source/ControlRig/Private/Sequencer/MovieSceneControlRigParameterSection.cpp
--- a/Engine/Plugins/Animation/ControlRig/Source/ControlRig/Private/Sequencer/MovieSceneControlRigParameterSection.cpp
+++ b/Engine/Plugins/Animation/ControlRig/Source/ControlRig/Private/Sequencer/MovieSceneControlRigParameterSection.cpp
@@ -2028,6 +2028,8 @@ void UMovieSceneControlRigParameterSection::ReconstructChannelProxy()
 
 
 	ChannelProxy = MakeShared<FMovieSceneChannelProxy>(MoveTemp(Channels));
+
+	MarkAsChanged();
 }
 
 FMovieSceneInterrogationKey UMovieSceneControlRigParameterSection::GetFloatInterrogationKey()
