fixed default initialization

diff --git a/Engine/Plugins/Animation/ControlRig/Source/ControlRig/Public/Rigs/RigHierarchyElements.h b/Engine/Plugins/Animation/ControlRig/Source/ControlRig/Public/Rigs/RigHierarchyElements.h
--- a/Engine/Plugins/Animation/ControlRig/Source/ControlRig/Public/Rigs/RigHierarchyElements.h
+++ b/Engine/Plugins/Animation/ControlRig/Source/ControlRig/Public/Rigs/RigHierarchyElements.h
@@ -1124,10 +1124,10 @@ struct CONTROLRIG_API FRigControlSettings
 	 * Deprecated properties.
 	 */
 	UPROPERTY(meta=(DeprecatedProperty, DeprecationMessage = "Use animation_type instead."))
-	bool bAnimatable_DEPRECATED;
+	bool bAnimatable_DEPRECATED = true;
 	
 	UPROPERTY(meta=(DeprecatedProperty, DeprecationMessage = "Use animation_type or shape_visible instead."))
-	bool bShapeEnabled_DEPRECATED;
+	bool bShapeEnabled_DEPRECATED = true;
 	
 	/** Applies the limits expressed by these settings to a value */
 	FORCEINLINE void ApplyLimits(FRigControlValue& InOutValue) const
