Sequencer - ControlRig Track : Added a eval lock when sequencer wants to do stuff to a controlrig that is getting parallel evaled

diff --git a/Engine/Plugins/Animation/ControlRig/Source/ControlRig/Private/Sequencer/MovieSceneControlRigParameterTemplate.cpp b/Engine/Plugins/Animation/ControlRig/Source/ControlRig/Private/Sequencer/MovieSceneControlRigParameterTemplate.cpp
--- a/Engine/Plugins/Animation/ControlRig/Source/ControlRig/Private/Sequencer/MovieSceneControlRigParameterTemplate.cpp
+++ b/Engine/Plugins/Animation/ControlRig/Source/ControlRig/Private/Sequencer/MovieSceneControlRigParameterTemplate.cpp
@@ -708,6 +708,13 @@ struct FControlRigParameterPreAnimatedTokenProducer : IMovieScenePreAnimatedToke
 				{
 					if (ControlRig->GetObjectBinding())
 					{
+						// control rig evaluate critical section: when restoring the state, we can be poking into running instances of Control Rigs
+						// on the anim thread so using a lock here to avoid this thread and anim thread both touching the rig
+						// at the same time, which can lead to various issues like double-freeing some random array when doing SetControlValue
+						// note: the critical section accepts recursive locking so it is ok that we
+						// are calling evaluate_anythread later within the same scope
+						FScopeLock EvaluateLock(&ControlRig->GetEvaluateMutex());
+						
 						//Restore control rig first
 						const bool bSetupUndo = false;
 						if (URigHierarchy* RigHierarchy = ControlRig->GetHierarchy())
