Constraints:
- make sure to connect binding delegate if the ControlRig is already bound to ensure that tick dependencies are setup properly

diff --git a/Engine/Plugins/Animation/ControlRig/Source/ControlRig/Private/Constraints/ControlRigTransformableHandle.cpp b/Engine/Plugins/Animation/ControlRig/Source/ControlRig/Private/Constraints/ControlRigTransformableHandle.cpp
--- a/Engine/Plugins/Animation/ControlRig/Source/ControlRig/Private/Constraints/ControlRigTransformableHandle.cpp
+++ b/Engine/Plugins/Animation/ControlRig/Source/ControlRig/Private/Constraints/ControlRigTransformableHandle.cpp
@@ -251,6 +251,7 @@ void UTransformableControlHandle::RegisterDelegates()
 		{
 			ControlRig->ControlRigBound().AddUObject(this, &UTransformableControlHandle::OnControlRigBound);
 		}
+		OnControlRigBound(ControlRig.Get());
 	}
 }
 
