Sequencer: Fix crash on close control rig. Make sure to clear the timer.

diff --git a/Engine/Plugins/Animation/ControlRig/Source/ControlRigEditor/Private/Sequencer/ControlRigParameterTrackEditor.cpp b/Engine/Plugins/Animation/ControlRig/Source/ControlRigEditor/Private/Sequencer/ControlRigParameterTrackEditor.cpp
--- a/Engine/Plugins/Animation/ControlRig/Source/ControlRigEditor/Private/Sequencer/ControlRigParameterTrackEditor.cpp
+++ b/Engine/Plugins/Animation/ControlRig/Source/ControlRigEditor/Private/Sequencer/ControlRigParameterTrackEditor.cpp
@@ -1767,11 +1767,17 @@ void FControlRigParameterTrackEditor::OnChannelChanged(const FMovieSceneChannelM
 				ChannelChanged.SectionControlNames.Reset();
 				auto ChannelChangedDelegate = [this]()
 				{
+					TSharedPtr<ISequencer> SequencerPtr = GetSequencer();
+					if (!SequencerPtr.IsValid())
+					{
+						GEditor->GetTimerManager()->ClearTimer(ChannelChanged.TimerHandle);
+						return;
+					}
+
 					if (!(FSlateApplication::Get().HasAnyMouseCaptor() || GUnrealEd->IsUserInteracting()))
 					{
 						if (ChannelChanged.TimerHandle.IsValid())
 						{
-							TSharedPtr<ISequencer> SequencerPtr = GetSequencer();
 							FFrameTime Time = SequencerPtr->GetLocalTime().Time;
 							GEditor->GetTimerManager()->ClearTimer(ChannelChanged.TimerHandle);
 							ChannelChanged.TimerHandle.Invalidate();
