Sequencer:Control Rig: Just check for bool control when setting control value for ik/fk hack.

diff --git a/Engine/Plugins/Animation/ControlRig/Source/ControlRigEditor/Private/Sequencer/ControlRigParameterTrackEditor.cpp b/Engine/Plugins/Animation/ControlRig/Source/ControlRigEditor/Private/Sequencer/ControlRigParameterTrackEditor.cpp
--- a/Engine/Plugins/Animation/ControlRig/Source/ControlRigEditor/Private/Sequencer/ControlRigParameterTrackEditor.cpp
+++ b/Engine/Plugins/Animation/ControlRig/Source/ControlRigEditor/Private/Sequencer/ControlRigParameterTrackEditor.cpp
@@ -1790,7 +1790,15 @@ void FControlRigParameterTrackEditor::OnChannelChanged(const FMovieSceneChannelM
 								{
 									Pair.Key->ControlsToSet.Empty();
 									Pair.Key->ControlsToSet.Add(ControlName);
-									EvaluateThisControl(Pair.Key, ControlName, Time);
+									//only do the fk/ik hack for bool's since that's only where it's needed
+									//otherwise we would incorrectly set values on an addtive section which has scale values of 0.0.
+									if (FRigControlElement* ControlElement = Pair.Key->GetControlRig()->FindControl(ControlName))
+									{
+										if (ControlElement->Settings.ControlType == ERigControlType::Bool)
+										{
+											EvaluateThisControl(Pair.Key, ControlName, Time);
+										}
+									}
 									FControlRigSpaceChannelHelpers::CompensateIfNeeded(Pair.Key->GetControlRig(), SequencerRaw, Pair.Key,
 										ControlName, Optional);
 									Pair.Key->ControlsToSet.Empty();
