Sequencer: Anim Tools: Quick Reordering of Tools, so it goes BN,PP, then TW.

#trivial

diff --git a/Engine/Plugins/Animation/ControlRig/Source/ControlRigEditor/Private/EditMode/SControlRigTweenWidget.cpp b/Engine/Plugins/Animation/ControlRig/Source/ControlRigEditor/Private/EditMode/SControlRigTweenWidget.cpp
--- a/Engine/Plugins/Animation/ControlRig/Source/ControlRigEditor/Private/EditMode/SControlRigTweenWidget.cpp
+++ b/Engine/Plugins/Animation/ControlRig/Source/ControlRigEditor/Private/EditMode/SControlRigTweenWidget.cpp
@@ -244,14 +244,14 @@ int32 SControlRigTweenWidget::ActiveSlider = 0;
 
 void SControlRigTweenWidget::Construct(const FArguments& InArgs)
 {
-	TSharedPtr<FBaseAnimSlider> TweenPtr = MakeShareable(new FControlsToTween());
-	AnimBlendTools.RegisterAnimSlider(TweenPtr);
-	TSharedPtr<FBaseAnimSlider> PushPullPtr = MakeShareable(new FPushPullSlider());
-	AnimBlendTools.RegisterAnimSlider(PushPullPtr);
+
 	TSharedPtr<FBaseAnimSlider> BlendNeighborPtr = MakeShareable(new FBlendNeighborSlider());
 	AnimBlendTools.RegisterAnimSlider(BlendNeighborPtr);
-
+	TSharedPtr<FBaseAnimSlider> PushPullPtr = MakeShareable(new FPushPullSlider());
+	AnimBlendTools.RegisterAnimSlider(PushPullPtr);
 	OwningToolkit = InArgs._InOwningToolkit;
+	TSharedPtr<FBaseAnimSlider> TweenPtr = MakeShareable(new FControlsToTween());
+	AnimBlendTools.RegisterAnimSlider(TweenPtr);
 
 
 	// Combo Button to swap sliders 
