Sequencer: Control Rig: Compatible Skeletons show up when baking to control rig.

diff --git a/Engine/Plugins/Animation/ControlRig/Source/ControlRigEditor/Private/Sequencer/ControlRigParameterTrackEditor.cpp b/Engine/Plugins/Animation/ControlRig/Source/ControlRigEditor/Private/Sequencer/ControlRigParameterTrackEditor.cpp
--- a/Engine/Plugins/Animation/ControlRig/Source/ControlRigEditor/Private/Sequencer/ControlRigParameterTrackEditor.cpp
+++ b/Engine/Plugins/Animation/ControlRig/Source/ControlRigEditor/Private/Sequencer/ControlRigParameterTrackEditor.cpp
@@ -638,16 +638,13 @@ public:
 		bFilterInversion(bInCheckInversion),
 		AssetRegistry(FModuleManager::GetModuleChecked<FAssetRegistryModule>(TEXT("AssetRegistry")).Get())
 	{
-		if (InSkeleton)
-		{
-			SkeletonName = FAssetData(InSkeleton).GetExportTextName();
-		}
+		Skeleton = InSkeleton;
 	}
 	bool bFilterAssetBySkeleton;
 	bool bFilterExposesAnimatableControls;
 	bool bFilterInversion;
 
-	FString SkeletonName;
+	USkeleton* Skeleton;
 	const IAssetRegistry& AssetRegistry;
 
 	bool MatchesFilter(const FAssetData& AssetData)
@@ -684,6 +681,11 @@ public:
 		}
 		if (bFilterAssetBySkeleton)
 		{
+			FString SkeletonName;
+			if (Skeleton)
+			{
+				SkeletonName = FAssetData(Skeleton).GetExportTextName();
+			}
 			FString PreviewSkeletalMesh = AssetData.GetTagValueRef<FString>(TEXT("PreviewSkeletalMesh"));
 			if (PreviewSkeletalMesh.Len() > 0)
 			{
@@ -693,22 +695,50 @@ public:
 				{
 					return true;
 				}
+				else if(Skeleton)
+				{
+					if (Skeleton->IsCompatibleSkeletonByAssetString(PreviewSkeleton))
+					{
+						return true;
+					}
+				}
 			}
 			FString PreviewSkeleton = AssetData.GetTagValueRef<FString>(TEXT("PreviewSkeleton"));
 			if (PreviewSkeleton == SkeletonName)
 			{
 				return true;
 			}
+			else if (Skeleton)
+			{
+				if (Skeleton->IsCompatibleSkeletonByAssetString(PreviewSkeleton))
+				{
+					return true;
+				}
+			}
 			FString SourceHierarchyImport = AssetData.GetTagValueRef<FString>(TEXT("SourceHierarchyImport"));
 			if (SourceHierarchyImport == SkeletonName)
 			{
 				return true;
 			}
+			else if (Skeleton)
+			{
+				if (Skeleton->IsCompatibleSkeletonByAssetString(SourceHierarchyImport))
+				{
+					return true;
+				}
+			}
 			FString SourceCurveImport = AssetData.GetTagValueRef<FString>(TEXT("SourceCurveImport"));
 			if (SourceCurveImport == SkeletonName)
 			{
 				return true;
 			}
+			else if (Skeleton)
+			{
+				if (Skeleton->IsCompatibleSkeletonByAssetString(SourceCurveImport))
+				{
+					return true;
+				}
+			}
 			return false;
 		}
 		return true;
