Control Rig: Fix construction event not executing after VM initialization

diff --git a/Engine/Plugins/Animation/ControlRig/Source/ControlRig/Private/ControlRig.cpp b/Engine/Plugins/Animation/ControlRig/Source/ControlRig/Private/ControlRig.cpp
--- a/Engine/Plugins/Animation/ControlRig/Source/ControlRig/Private/ControlRig.cpp
+++ b/Engine/Plugins/Animation/ControlRig/Source/ControlRig/Private/ControlRig.cpp
@@ -1528,6 +1528,7 @@ bool UControlRig::ExecuteUnits(FRigUnitContext& InOutContext, const FName& InEve
 			{
 				bSuccess = VM->Initialize(LocalMemory, AdditionalArguments);
 			}
+			bRequiresConstructionEvent = true;
 		}
 		else
 		{
