Control Rig: Fix print dispatch node displaying empty strings

diff --git a/Engine/Plugins/Animation/ControlRig/Source/ControlRig/Private/Units/Core/RigUnit_Print.cpp b/Engine/Plugins/Animation/ControlRig/Source/ControlRig/Private/Units/Core/RigUnit_Print.cpp
--- a/Engine/Plugins/Animation/ControlRig/Source/ControlRig/Private/Units/Core/RigUnit_Print.cpp
+++ b/Engine/Plugins/Animation/ControlRig/Source/ControlRig/Private/Units/Core/RigUnit_Print.cpp
@@ -83,6 +83,12 @@ void FRigDispatch_Print::Execute(FRigVMExtendedExecuteContext& InContext, FRigVM
 		return;
 	}
 
+	const FRigUnitContext& Context = GetRigUnitContext(InContext);
+	if (Context.State == EControlRigState::Init)
+	{
+		return;
+	}
+	
 	FString String;
 	ValueProperty->ExportText_Direct(String, Value, Value, nullptr, PPF_None, nullptr);
 
@@ -95,7 +101,6 @@ void FRigDispatch_Print::Execute(FRigVMExtendedExecuteContext& InContext, FRigVM
 	static constexpr TCHAR LogFormat[] = TEXT("%s[%04d] %s%s");
 	UE_LOG(LogControlRig, Display, LogFormat, *ObjectPath, InContext.PublicData.GetInstructionIndex(), *Prefix, *String);
 
-	const FRigUnitContext& Context = GetRigUnitContext(InContext);
 	if(ScreenDuration > SMALL_NUMBER && Context.World)
 	{
 		static constexpr TCHAR PrintStringFormat[] = TEXT("[%04d] %s%s");
