Sequencer: Control Rig:  Minor fix to not start transactions when undoing selections since it can get us into a weird undo/redo state.
#trivial

diff --git a/Engine/Plugins/Animation/ControlRig/Source/ControlRigEditor/Private/Sequencer/ControlRigParameterTrackEditor.cpp b/Engine/Plugins/Animation/ControlRig/Source/ControlRigEditor/Private/Sequencer/ControlRigParameterTrackEditor.cpp
--- a/Engine/Plugins/Animation/ControlRig/Source/ControlRigEditor/Private/Sequencer/ControlRigParameterTrackEditor.cpp
+++ b/Engine/Plugins/Animation/ControlRig/Source/ControlRigEditor/Private/Sequencer/ControlRigParameterTrackEditor.cpp
@@ -2189,7 +2189,7 @@ void FControlRigParameterTrackEditor::SelectRigsAndControls(UControlRig* Control
 			if (SelectedNames->Num() != FullNames->Num())
 			{ 
 				bIsSame = false;
-				if (bEndTransaction == false)
+				if (!GIsTransacting && bEndTransaction == false)
 				{
 					bEndTransaction = true;
 					GEditor->BeginTransaction(LOCTEXT("SelectControl", "Select Control"));
@@ -2204,7 +2204,7 @@ void FControlRigParameterTrackEditor::SelectRigsAndControls(UControlRig* Control
 					if (FullNames->Contains(Name) == false)
 					{
 						bIsSame = false;
-						if (bEndTransaction == false)
+						if (!GIsTransacting && bEndTransaction == false)
 						{
 							bEndTransaction = true;
 							GEditor->BeginTransaction(LOCTEXT("SelectControl", "Select Control"));
@@ -2228,7 +2228,7 @@ void FControlRigParameterTrackEditor::SelectRigsAndControls(UControlRig* Control
 		{
 			for (const FName& Name : Pair.Value)
 			{
-				if (bEndTransaction == false)
+				if (!GIsTransacting && bEndTransaction == false)
 				{
 					bEndTransaction = true;
 					GEditor->BeginTransaction(LOCTEXT("SelectControl", "Select Control"));
@@ -2240,7 +2240,7 @@ void FControlRigParameterTrackEditor::SelectRigsAndControls(UControlRig* Control
 	//go through and clear those still not cleared
 	for (TPair<UControlRig*, TArray<FName>>& SelectedPairs : ControlRigsToClearSelection)
 	{
-		if (bEndTransaction == false)
+		if (!GIsTransacting && bEndTransaction == false)
 		{
 			bEndTransaction = true;
 			GEditor->BeginTransaction(LOCTEXT("SelectControl", "Select Control"));
