Constraints: make sure the ControlRig asset is loaded before registering delegates (follow up to 21703269)

diff --git a/Engine/Plugins/Animation/ControlRig/Source/ControlRig/Private/Constraints/ControlRigTransformableHandle.cpp b/Engine/Plugins/Animation/ControlRig/Source/ControlRig/Private/Constraints/ControlRigTransformableHandle.cpp
--- a/Engine/Plugins/Animation/ControlRig/Source/ControlRig/Private/Constraints/ControlRigTransformableHandle.cpp
+++ b/Engine/Plugins/Animation/ControlRig/Source/ControlRig/Private/Constraints/ControlRigTransformableHandle.cpp
@@ -239,6 +239,12 @@ void UTransformableControlHandle::RegisterDelegates()
 	FCoreUObjectDelegates::OnObjectsReplaced.AddUObject(this, &UTransformableControlHandle::OnObjectsReplaced);
 #endif
 
+	// make sure the CR is loaded so that we can register delegates
+	if (ControlRig.IsPending())
+	{
+		ControlRig.LoadSynchronous();
+	}
+	
 	if (ControlRig.IsValid())
 	{
 		if (URigHierarchy* Hierarchy = ControlRig->GetHierarchy())
