Constraints:
- exit picker when the selection has changed
- auto expand constraint view when creating a new constraint

diff --git a/Engine/Plugins/Animation/ControlRig/Source/ControlRigEditor/Private/EditMode/SControlRigEditModeTools.cpp b/Engine/Plugins/Animation/ControlRig/Source/ControlRigEditor/Private/EditMode/SControlRigEditModeTools.cpp
--- a/Engine/Plugins/Animation/ControlRig/Source/ControlRigEditor/Private/EditMode/SControlRigEditModeTools.cpp
+++ b/Engine/Plugins/Animation/ControlRig/Source/ControlRigEditor/Private/EditMode/SControlRigEditModeTools.cpp
@@ -327,7 +327,7 @@ void SControlRigEditModeTools::Construct(const FArguments& InArgs, TSharedPtr<FC
 			+SVerticalBox::Slot()
 			.AutoHeight()
 			[
-				SAssignNew(PickerExpander, SExpandableArea)
+				SAssignNew(ConstraintPickerExpander, SExpandableArea)
 				.InitiallyCollapsed(true)
 				.AreaTitle(LOCTEXT("ConstraintsWidget", "Constraints"))
 				.AreaTitleFont(FAppStyle::GetFontStyle("DetailsView.CategoryFontStyle"))
@@ -982,14 +982,20 @@ FReply SControlRigEditModeTools::OnBakeControlsToNewSpaceButtonClicked()
 
 FReply SControlRigEditModeTools::HandleAddConstraintClicked()
 {
+	// magic number to auto expand the widget when creating a new constraint. We keep that number below a reasonable
+	// threshold to avoid automatically creating a large number of items (this can be style done by the user) 
+	static constexpr int32 NumAutoExpand = 20;
+
 	const TSharedPtr<SConstraintsCreationWidget> Widget =
 		SNew(SConstraintsCreationWidget)
 		.OnConstraintCreated_Lambda( [this]()
 		{
-			if (ConstraintsEditionWidget)
+			const int32 NumItems = ConstraintsEditionWidget ? ConstraintsEditionWidget->RefreshConstraintList() : 0;
+			
+			if (ConstraintPickerExpander && NumItems < NumAutoExpand)
 			{
-				ConstraintsEditionWidget->RefreshConstraintList();
-			}		
+				ConstraintPickerExpander->SetExpanded(true);
+			}
 		});
 	
 	FMenuBuilder MenuBuilder(true, nullptr);
diff --git a/Engine/Plugins/Animation/ControlRig/Source/ControlRigEditor/Private/EditMode/SControlRigEditModeTools.h b/Engine/Plugins/Animation/ControlRig/Source/ControlRigEditor/Private/EditMode/SControlRigEditModeTools.h
--- a/Engine/Plugins/Animation/ControlRig/Source/ControlRigEditor/Private/EditMode/SControlRigEditModeTools.h
+++ b/Engine/Plugins/Animation/ControlRig/Source/ControlRigEditor/Private/EditMode/SControlRigEditModeTools.h
@@ -93,6 +93,7 @@ private:
 	TArray<TWeakObjectPtr<UControlRig>> ControlRigs;
 
 	/** Constraint edition widget. */
+	TSharedPtr<SExpandableArea> ConstraintPickerExpander = nullptr;
 	TSharedPtr<SConstraintsEditionWidget> ConstraintsEditionWidget = nullptr; 
 
 	/** Display or edit set up for property */
diff --git a/Engine/Plugins/Animation/ControlRig/Source/ControlRigEditor/Private/Editor/Constraints/SConstraintsWidget.cpp b/Engine/Plugins/Animation/ControlRig/Source/ControlRigEditor/Private/Editor/Constraints/SConstraintsWidget.cpp
--- a/Engine/Plugins/Animation/ControlRig/Source/ControlRigEditor/Private/Editor/Constraints/SConstraintsWidget.cpp
+++ b/Engine/Plugins/Animation/ControlRig/Source/ControlRigEditor/Private/Editor/Constraints/SConstraintsWidget.cpp
@@ -837,6 +837,17 @@ void SConstraintsEditionWidget::OnActorSelectionChanged(const TArray<UObject*>&
 	// NOTE we use this delegate to trigger an tree update, however, control actors are not selected as they are 
 	// Temporary Editor Actors so NewSelection won't contain the controls
 	InvalidateConstraintList();
+
+	// NOTE: this is a hack to disable the picker if the selection changes. The picker should trigger the escape key
+	// but this is an editor wide change 
+	if (NewSelection.IsEmpty())
+	{
+		static const FActorPickerModeModule& ActorPickerMode = FModuleManager::Get().GetModuleChecked<FActorPickerModeModule>("ActorPickerMode");
+		if (ActorPickerMode.IsInActorPickingMode())
+		{
+			ActorPickerMode.EndActorPickingMode();
+		}
+	}
 }
 
 TSharedPtr<SWidget> SConstraintsEditionWidget::CreateContextMenu()
