also update default hierarchy when setting offset transform

diff --git a/Engine/Plugins/Animation/ControlRig/Source/ControlRigEditor/Private/Editor/SRigHierarchy.cpp b/Engine/Plugins/Animation/ControlRig/Source/ControlRigEditor/Private/Editor/SRigHierarchy.cpp
--- a/Engine/Plugins/Animation/ControlRig/Source/ControlRigEditor/Private/Editor/SRigHierarchy.cpp
+++ b/Engine/Plugins/Animation/ControlRig/Source/ControlRigEditor/Private/Editor/SRigHierarchy.cpp
@@ -2460,6 +2460,8 @@ void SRigHierarchy::HandleSetInitialTransformFromCurrentTransform()
 					ParentGlobalTransforms.FindOrAdd(SelectedKey) = DebuggedHierarchy->GetParentTransform(SelectedKey);
 				}
 
+				URigHierarchy* DefaultHierarchy = GetDefaultHierarchy();
+				
 				for (const FRigElementKey& SelectedKey : SelectedKeys)
 				{
 					FTransform GlobalTransform = GlobalTransforms[SelectedKey];
@@ -2467,19 +2469,23 @@ void SRigHierarchy::HandleSetInitialTransformFromCurrentTransform()
 
 					if (SelectedKey.Type == ERigElementType::Control)
 					{
-						if(FRigControlElement* ControlElement = GetHierarchy()->Find<FRigControlElement>(SelectedKey))
+						if(FRigControlElement* ControlElement = DebuggedHierarchy->Find<FRigControlElement>(SelectedKey))
 						{
-							GetHierarchy()->SetControlOffsetTransform(ControlElement, LocalTransform, ERigTransformType::InitialLocal, true, true, false, true);
-							GetHierarchy()->SetControlOffsetTransform(ControlElement, LocalTransform, ERigTransformType::CurrentLocal, true, true, false, true);
-							GetHierarchy()->SetTransform(ControlElement, FTransform::Identity, ERigTransformType::InitialLocal, true, true, false, true);
-							GetHierarchy()->SetTransform(ControlElement, FTransform::Identity, ERigTransformType::CurrentLocal, true, true, false, true);
+							DebuggedHierarchy->SetControlOffsetTransform(ControlElement, LocalTransform, ERigTransformType::InitialLocal, true, true, false, true);
+							DebuggedHierarchy->SetControlOffsetTransform(ControlElement, LocalTransform, ERigTransformType::CurrentLocal, true, true, false, true);
+							DebuggedHierarchy->SetTransform(ControlElement, FTransform::Identity, ERigTransformType::InitialLocal, true, true, false, true);
+							DebuggedHierarchy->SetTransform(ControlElement, FTransform::Identity, ERigTransformType::CurrentLocal, true, true, false, true);
 						}
-						if(FRigControlElement* ControlElement = DebuggedHierarchy->Find<FRigControlElement>(SelectedKey))
+
+						if(DefaultHierarchy)
 						{
-							DebuggedHierarchy->SetControlOffsetTransform(ControlElement, LocalTransform, ERigTransformType::InitialLocal, true, true);
-							DebuggedHierarchy->SetControlOffsetTransform(ControlElement, LocalTransform, ERigTransformType::CurrentLocal, true, true);
-							DebuggedHierarchy->SetTransform(ControlElement, FTransform::Identity, ERigTransformType::InitialLocal, true, true);
-							DebuggedHierarchy->SetTransform(ControlElement, FTransform::Identity, ERigTransformType::CurrentLocal, true, true);
+							if(FRigControlElement* ControlElement = DefaultHierarchy->Find<FRigControlElement>(SelectedKey))
+							{
+								DefaultHierarchy->SetControlOffsetTransform(ControlElement, LocalTransform, ERigTransformType::InitialLocal, true, true);
+								DefaultHierarchy->SetControlOffsetTransform(ControlElement, LocalTransform, ERigTransformType::CurrentLocal, true, true);
+								DefaultHierarchy->SetTransform(ControlElement, FTransform::Identity, ERigTransformType::InitialLocal, true, true);
+								DefaultHierarchy->SetTransform(ControlElement, FTransform::Identity, ERigTransformType::CurrentLocal, true, true);
+							}
 						}
 					}
 					else if (SelectedKey.Type == ERigElementType::Null ||
@@ -2496,15 +2502,19 @@ void SRigHierarchy::HandleSetInitialTransformFromCurrentTransform()
 							}
 						}
 
-						if(FRigTransformElement* TransformElement = GetHierarchy()->Find<FRigTransformElement>(SelectedKey))
+						if(FRigTransformElement* TransformElement = DebuggedHierarchy->Find<FRigTransformElement>(SelectedKey))
 						{
-							GetHierarchy()->SetTransform(TransformElement, LocalTransform, ERigTransformType::InitialLocal, true, true, false, true);
-							GetHierarchy()->SetTransform(TransformElement, LocalTransform, ERigTransformType::CurrentLocal, true, true, false, true);
+							DebuggedHierarchy->SetTransform(TransformElement, LocalTransform, ERigTransformType::InitialLocal, true, true, false, true);
+							DebuggedHierarchy->SetTransform(TransformElement, LocalTransform, ERigTransformType::CurrentLocal, true, true, false, true);
 						}
-						if(FRigTransformElement* TransformElement = DebuggedHierarchy->Find<FRigTransformElement>(SelectedKey))
+
+						if(DefaultHierarchy)
 						{
-							DebuggedHierarchy->SetTransform(TransformElement, LocalTransform, ERigTransformType::InitialLocal, true, true);
-							DebuggedHierarchy->SetTransform(TransformElement, LocalTransform, ERigTransformType::CurrentLocal, true, true);
+							if(FRigTransformElement* TransformElement = DefaultHierarchy->Find<FRigTransformElement>(SelectedKey))
+							{
+								DefaultHierarchy->SetTransform(TransformElement, LocalTransform, ERigTransformType::InitialLocal, true, true);
+								DefaultHierarchy->SetTransform(TransformElement, LocalTransform, ERigTransformType::CurrentLocal, true, true);
+							}
 						}
 					}
 				}
