set initial local instead of current local value when setting shape transform
It mimics FRigControlElementDetails::CustomizeShape behavior (see URigHierarchy::SetControlShapeTransform for details about initial/current sync)

diff --git a/Engine/Plugins/Animation/ControlRig/Source/ControlRig/Private/Units/Hierarchy/RigUnit_SetControlOffset.cpp b/Engine/Plugins/Animation/ControlRig/Source/ControlRig/Private/Units/Hierarchy/RigUnit_SetControlOffset.cpp
--- a/Engine/Plugins/Animation/ControlRig/Source/ControlRig/Private/Units/Hierarchy/RigUnit_SetControlOffset.cpp
+++ b/Engine/Plugins/Animation/ControlRig/Source/ControlRig/Private/Units/Hierarchy/RigUnit_SetControlOffset.cpp
@@ -106,7 +106,7 @@ FRigUnit_SetShapeTransform_Execute()
 				}
 
 				FRigControlElement* ControlElement = Hierarchy->Get<FRigControlElement>(CachedControlIndex);
-				Hierarchy->SetControlShapeTransform(ControlElement, Transform, ERigTransformType::CurrentLocal);
+				Hierarchy->SetControlShapeTransform(ControlElement, Transform, ERigTransformType::InitialLocal);
 				break;
 			}
 			default:
