RigVMRegistry: Added support for user defined type addition and removal

Important notes:
1. We should now use TypeUtils::GetUniqueStructName  instead of struct->getstructcppname for type names used by the registry since user defined struct need to have unique names generated for them

2. We should be aware that we can have invalid permutations in templates when types are deleted, so make sure you check for typeindex != INDEX_NONE when looping over permutations

diff --git a/Engine/Plugins/Animation/ControlRig/Source/ControlRig/Private/Units/Animation/RigUnit_AnimAttribute.cpp b/Engine/Plugins/Animation/ControlRig/Source/ControlRig/Private/Units/Animation/RigUnit_AnimAttribute.cpp
--- a/Engine/Plugins/Animation/ControlRig/Source/ControlRig/Private/Units/Animation/RigUnit_AnimAttribute.cpp
+++ b/Engine/Plugins/Animation/ControlRig/Source/ControlRig/Private/Units/Animation/RigUnit_AnimAttribute.cpp
@@ -2,58 +2,96 @@
 
 #include "Units/Animation/RigUnit_AnimAttribute.h"
 
-#include "Animation/BuiltInAttributeTypes.h"
+#include "Animation/AttributeTypes.h"
+
 #include "Units/RigUnitContext.h"
 #include "Components/SkeletalMeshComponent.h"
+#include "Engine/UserDefinedStruct.h"
 
-template<typename T>
-class TAnimAttributeType;
+FName FRigDispatch_AnimAttributeBase::NameArgName = TEXT("Name");
+FName FRigDispatch_AnimAttributeBase::BoneNameArgName = TEXT("BoneName");
+FName FRigDispatch_AnimAttributeBase::CachedBoneNameArgName = TEXT("CachedBoneName");
+FName FRigDispatch_AnimAttributeBase::CachedBoneIndexArgName = TEXT("CachedBoneIndex");
+FName FRigDispatch_AnimAttributeBase::DefaultArgName = TEXT("Default");
+FName FRigDispatch_AnimAttributeBase::ValueArgName = TEXT("Value");
+FName FRigDispatch_AnimAttributeBase::FoundArgName = TEXT("Found");
+FName FRigDispatch_AnimAttributeBase::SuccessArgName = TEXT("Success");
 
-template<>
-class TAnimAttributeType<int>
-{
-public:
-	using Type = FIntegerAnimationAttribute;
-};
 
-template<>
-class TAnimAttributeType<float>
+bool FRigDispatch_AnimAttributeBase::IsTypeSupported(const TRigVMTypeIndex& InTypeIndex)
 {
-public:
-	using Type = FFloatAnimationAttribute;
-};
+	const FRigVMRegistry& Registry = FRigVMRegistry::Get();
+
+	static const TArray<TRigVMTypeIndex> SpecialTypes = {
+		RigVMTypeUtils::TypeIndex::Float,
+		RigVMTypeUtils::TypeIndex::Int32,
+		RigVMTypeUtils::TypeIndex::FString,
+		Registry.GetTypeIndex<FTransform>(false),
+		Registry.GetTypeIndex<FVector>(false),
+		Registry.GetTypeIndex<FQuat>(false),
+	};
+
+	if (SpecialTypes.Contains(InTypeIndex))
+	{
+		return true;
+	}
 
-template<>
-class TAnimAttributeType<FTransform>
-{
-public:
-	using Type = FTransformAnimationAttribute;
-};
+	const FRigVMTemplateArgumentType& InType = Registry.GetType(InTypeIndex);
+
+	// cpp type object can become invalid because users can choose to delete
+	// user defined structs
+	if (IsValid(InType.CPPTypeObject))
+	{
+		if (UScriptStruct* ScriptStruct = Cast<UScriptStruct>(InType.CPPTypeObject))
+		{
+			static const TArray<TWeakObjectPtr<const UScriptStruct>> SpecialAttributeTypes = {
+				FFloatAnimationAttribute::StaticStruct(),
+				FIntegerAnimationAttribute::StaticStruct(),
+				FStringAnimationAttribute::StaticStruct(),
+				FTransformAnimationAttribute::StaticStruct(),
+				FVectorAnimationAttribute::StaticStruct(),
+				FQuaternionAnimationAttribute::StaticStruct()
+			};
+
+			if (SpecialAttributeTypes.Contains(ScriptStruct))
+			{
+				// these type have been added, above, rejecting them here so we don't have duplicated types
+				return false;
+			}
+
+			return UE::Anim::AttributeTypes::IsTypeRegistered(ScriptStruct);
+		}
+	}
+
+	return false;
+}
 
-template<>
-class TAnimAttributeType<FVector>
-{
-public:
-	using Type = FVectorAnimationAttribute;
-};
 
-template<>
-class TAnimAttributeType<FQuat>
+const TArray<FRigVMTemplateArgument::ETypeCategory>& FRigDispatch_AnimAttributeBase::GetValueTypeCategory()
 {
-public:
-	using Type = FQuaternionAnimationAttribute;
-};
+	static const TArray<FRigVMTemplateArgument::ETypeCategory> TypeCategories = {
+		FRigVMTemplateArgument::ETypeCategory_SingleSimpleValue,
+		FRigVMTemplateArgument::ETypeCategory_SingleMathStructValue,
+		FRigVMTemplateArgument::ETypeCategory_SingleScriptStructValue
+	};
+	return TypeCategories;
+}
 
 
-template<typename T>
-FORCEINLINE_DEBUGGABLE static T* GetAnimAttributeValue(
+static uint8* GetAnimAttributeValue(
 	bool bAddIfNotFound,
+	UScriptStruct* InAttributeScriptStruct,
 	const FRigUnitContext& Context,
 	const FName& Name,
 	const FName& BoneName,
 	FName& CachedBoneName,
 	int32& CachedBoneIndex)
 {
+	if (!ensure(InAttributeScriptStruct))
+	{
+		return nullptr;
+	}
+	
 	if (Name.IsNone())
 	{
 		return nullptr;
@@ -91,116 +129,404 @@ FORCEINLINE_DEBUGGABLE static T* GetAnimAttributeValue(
 	if (CachedBoneIndex != INDEX_NONE)
 	{
 		const UE::Anim::FAttributeId Id = {Name, FCompactPoseBoneIndex(CachedBoneIndex)} ;
-		typename TAnimAttributeType<T>::Type* Attribute = bAddIfNotFound ?
-			Context.AnimAttributeContainer->FindOrAdd<typename TAnimAttributeType<T>::Type>(Id) :
-			Context.AnimAttributeContainer->Find<typename TAnimAttributeType<T>::Type>(Id);
+		uint8* Attribute = bAddIfNotFound ?
+			Context.AnimAttributeContainer->FindOrAdd(InAttributeScriptStruct, Id) :
+			Context.AnimAttributeContainer->Find(InAttributeScriptStruct, Id);
+		
 		if (Attribute)
 		{
-			return &Attribute->Value;
+			return Attribute;
 		}
 	}
 
 	return nullptr;
 }
 
-
-FRigUnit_SetAnimAttribute_Integer_Execute()
+void FRigDispatch_GetAnimAttribute::GetAnimAttributeDispatch(FRigVMExtendedExecuteContext& InContext,
+	FRigVMMemoryHandleArray Handles)
 {
-	DECLARE_SCOPE_HIERARCHICAL_COUNTER_RIGUNIT()
+	const FRigDispatch_GetAnimAttribute* Factory = static_cast<const FRigDispatch_GetAnimAttribute*>(InContext.Factory);
+		
+#if WITH_EDITOR
+	if(!Factory->CheckArgumentTypes(Handles))
+	{
+		return;
+	}
+#endif
+
+	// unpack the memory
+	const FName& Name = *(const FName*)Handles[Factory->NameArgIndex].GetData();
+	const FName& BoneName = *(const FName*)Handles[Factory->BoneNameArgIndex].GetData();
+	const uint8* Default = (const uint8*)Handles[Factory->DefaultArgIndex].GetData();
+	
+	uint8* Value = (uint8*)Handles[Factory->ValueArgIndex].GetData();
+	bool& Found = *(bool*)Handles[Factory->FoundArgIndex].GetData();
+	Found = false;
 	
-	int* ValuePtr = GetAnimAttributeValue<int>(true ,Context, Name, BoneName, CachedBoneName, CachedBoneIndex);
-	if (ValuePtr)
+	FName& CachedBoneName = *(FName*)Handles[Factory->CachedBoneNameArgIndex].GetData(false, InContext.GetSlice().GetIndex());
+	int32& CachedBoneIndex = *(int32*)Handles[Factory->CachedBoneIndexArgIndex].GetData(false, InContext.GetSlice().GetIndex());
+	
+	if (const FStructProperty* StructProperty = CastFieldChecked<FStructProperty>(Handles[Factory->ValueArgIndex].GetResolvedProperty()))
 	{
-		*ValuePtr = Value;
+		UScriptStruct* ScriptStruct = StructProperty->Struct;
+		bool Registered = true;
+		const FRigUnitContext& Context = GetRigUnitContext(InContext);
+
+#if WITH_EDITOR
+		{
+			FRigVMExecuteContext& RigVMExecuteContext = InContext.PublicData;
+			if (!UE::Anim::AttributeTypes::IsTypeRegistered(ScriptStruct))
+			{
+				UE_CONTROLRIG_RIGUNIT_REPORT_WARNING(
+					TEXT("Type: '%s' is not registered with the Animation Attribute System. "
+						"Please register the type in Project Settings - Animation - CustomAttributes - User Defined Struct Attributes."),
+					*ScriptStruct->GetAuthoredName());
+				Registered = false;
+			}
+		}
+#endif
+		
+		if (Registered)
+		{
+			// extract the animation attribute
+			const uint8* ValuePtr = GetAnimAttributeValue(false, ScriptStruct, Context, Name, BoneName, CachedBoneName, CachedBoneIndex);
+			Found = ValuePtr ? true : false;
+
+			if (Found)
+			{
+				ScriptStruct->CopyScriptStruct(Value, ValuePtr);		
+			}
+		}
+
+		if (!Found)
+		{
+			ScriptStruct->CopyScriptStruct(Value, Default);		
+		}
 	}
 }
 
-FRigUnit_SetAnimAttribute_Float_Execute()
+void FRigDispatch_SetAnimAttribute::SetAnimAttributeDispatch(FRigVMExtendedExecuteContext& InContext,
+	FRigVMMemoryHandleArray Handles)
 {
-	DECLARE_SCOPE_HIERARCHICAL_COUNTER_RIGUNIT()
-
-	float* ValuePtr = GetAnimAttributeValue<float>(true, Context, Name, BoneName, CachedBoneName, CachedBoneIndex);
-	if (ValuePtr)
+	const FRigDispatch_SetAnimAttribute* Factory = static_cast<const FRigDispatch_SetAnimAttribute*>(InContext.Factory);
+		
+#if WITH_EDITOR
+	if(!Factory->CheckArgumentTypes(Handles))
 	{
-		*ValuePtr = Value;
+		return;
 	}
+#endif
+
+	// unpack the memory
+	const FName& Name = *(const FName*)Handles[Factory->NameArgIndex].GetData();
+	const FName& BoneName = *(const FName*)Handles[Factory->BoneNameArgIndex].GetData();
+	const uint8* Value = (uint8*)Handles[Factory->ValueArgIndex].GetData();
+	
+	bool& Success = *(bool*)Handles[Factory->SuccessArgIndex].GetData();
+	Success = false;
+	
+	FName& CachedBoneName = *(FName*)Handles[Factory->CachedBoneNameArgIndex].GetData(false, InContext.GetSlice().GetIndex());
+	int32& CachedBoneIndex = *(int32*)Handles[Factory->CachedBoneIndexArgIndex].GetData(false, InContext.GetSlice().GetIndex());
+	
+	if (const FStructProperty* StructProperty = CastFieldChecked<FStructProperty>(Handles[Factory->ValueArgIndex].GetResolvedProperty()))
+	{
+		UScriptStruct* ScriptStruct = StructProperty->Struct;
+		bool Registered = true;
+		const FRigUnitContext& Context = GetRigUnitContext(InContext);
+
+#if WITH_EDITOR
+		{
+			FRigVMExecuteContext& RigVMExecuteContext = InContext.PublicData;
+			if (!UE::Anim::AttributeTypes::IsTypeRegistered(ScriptStruct))
+			{
+				UE_CONTROLRIG_RIGUNIT_REPORT_WARNING(
+					TEXT("Type: '%s' is not registered with the Animation Attribute System. "
+						"Please register the type in Project Settings - Animation - CustomAttributes - User Defined Struct Attributes."),
+					*ScriptStruct->GetAuthoredName());
+				Registered = false;
+			}
+		}
+#endif
+		
+		if (Registered)
+		{
+			// extract the animation attribute
+			uint8* ValuePtr = GetAnimAttributeValue(true, ScriptStruct, Context, Name, BoneName, CachedBoneName, CachedBoneIndex);
+
+			if (ValuePtr)
+			{
+				Success = true;
+				ScriptStruct->CopyScriptStruct(ValuePtr, Value);		
+			}
+		}
+	}	
 }
 
-FRigUnit_SetAnimAttribute_Transform_Execute()
+void FRigDispatch_AnimAttributeBase::RegisterDependencyTypes() const
 {
-    DECLARE_SCOPE_HIERARCHICAL_COUNTER_RIGUNIT()
+	Super::RegisterDependencyTypes();
 
-	FTransform* ValuePtr = GetAnimAttributeValue<FTransform>(true, Context, Name, BoneName, CachedBoneName, CachedBoneIndex);
-	if (ValuePtr)
+	FRigVMRegistry& Registry = FRigVMRegistry::Get();
+	
+	TArray<TWeakObjectPtr<const UScriptStruct>> AttributeTypes = UE::Anim::AttributeTypes::GetRegisteredTypes();
+	for (TWeakObjectPtr<const UScriptStruct> Type : AttributeTypes)
 	{
-		*ValuePtr = Value;
+		UScriptStruct* TypePtr = const_cast<UScriptStruct*>(Type.Get());
+		if (TypePtr)
+		{
+			Registry.FindOrAddType(FRigVMTemplateArgumentType(TypePtr));
+		}
 	}
 }
 
-FRigUnit_SetAnimAttribute_Vector_Execute()
+#if WITH_EDITOR
+FString FRigDispatch_AnimAttributeBase::GetNodeTitle(const FRigVMTemplateTypeMap& InTypes) const
 {
-	DECLARE_SCOPE_HIERARCHICAL_COUNTER_RIGUNIT()
-
-	FVector* ValuePtr = GetAnimAttributeValue<FVector>(true, Context, Name, BoneName, CachedBoneName, CachedBoneIndex);
-	if (ValuePtr)
+	if(const TRigVMTypeIndex* ValueTypeIndexPtr = InTypes.Find(ValueArgName))
 	{
-		*ValuePtr = Value;
+		const TRigVMTypeIndex& ValueTypeIndex = *ValueTypeIndexPtr;
+		if(ValueTypeIndex != RigVMTypeUtils::TypeIndex::WildCard &&
+			ValueTypeIndex != RigVMTypeUtils::TypeIndex::WildCardArray)
+		{
+			static constexpr TCHAR GetFormat[] = TEXT("Get %s Animation Attribute");
+			static constexpr TCHAR SetFormat[] = TEXT("Set %s Animation Attribute");
+
+			const FRigVMTemplateArgumentType& ValueType = FRigVMRegistry::Get().GetType(ValueTypeIndex);
+			FString ValueName;
+
+			if(ValueType.CPPTypeObject)
+			{
+				ValueName = ValueType.CPPTypeObject->GetName();
+			}
+			else
+			{
+				ValueName = ValueType.GetBaseCPPType();
+				ValueName = ValueName.Left(1).ToUpper() + ValueName.Mid(1);
+			}
+
+			return FString::Printf(IsSet() ? SetFormat : GetFormat, *ValueName); 
+		}
 	}
+	return FRigDispatchFactory::GetNodeTitle(InTypes);	
 }
+#endif
 
-FRigUnit_SetAnimAttribute_Quaternion_Execute()
+TArray<FRigVMTemplateArgument> FRigDispatch_AnimAttributeBase::GetArguments() const
 {
-	DECLARE_SCOPE_HIERARCHICAL_COUNTER_RIGUNIT()
+	if (Arguments.IsEmpty())
+	{
+		if (IsSet())
+		{
+			ExecuteArgIndex = Arguments.Emplace(FRigVMStruct::ExecuteContextName, ERigVMPinDirection::IO, FRigVMRegistry::Get().GetTypeIndex<FControlRigExecuteContext>());
+		}
+
+		NameArgIndex = Arguments.Emplace(NameArgName, ERigVMPinDirection::Input, RigVMTypeUtils::TypeIndex::FName);
+		BoneNameArgIndex = Arguments.Emplace(BoneNameArgName, ERigVMPinDirection::Input, RigVMTypeUtils::TypeIndex::FName);
+
+		CachedBoneNameArgIndex = Arguments.Emplace(CachedBoneNameArgName, ERigVMPinDirection::Hidden, RigVMTypeUtils::TypeIndex::FName);
+		CachedBoneIndexArgIndex = Arguments.Emplace(CachedBoneIndexArgName, ERigVMPinDirection::Hidden, RigVMTypeUtils::TypeIndex::Int32);
+	}
+	
+	return Arguments;
+}
 
-	FQuat* ValuePtr = GetAnimAttributeValue<FQuat>(true, Context, Name, BoneName, CachedBoneName, CachedBoneIndex);
-	if (ValuePtr)
+#if WITH_EDITOR
+FText FRigDispatch_AnimAttributeBase::GetArgumentTooltip(const FName& InArgumentName, TRigVMTypeIndex InTypeIndex) const
+{
+	if(InArgumentName == NameArgName)
+	{
+		return NSLOCTEXT("FRigDispatch_AnimAttributeBase", "NameArgTooltip", "The name of the animation attribute");
+	}
+	if(InArgumentName == BoneNameArgName)
+	{
+		return NSLOCTEXT("FRigDispatch_AnimAttributeBase", "BoneNameArgTooltip", "The name of the bone that stores the attribute, default to root bone if set to none");
+	}
+	if(InArgumentName == DefaultArgName)
+	{
+		return NSLOCTEXT("FRigDispatch_AnimAttributeBase", "DefaultArgTooltip", "The default value used as a fallback if the animation attribute does not exist");
+	}
+	if(InArgumentName == ValueArgName)
 	{
-		*ValuePtr = Value;
+		return NSLOCTEXT("FRigDispatch_AnimAttributeBase", "ValueArgTooltip", "The value to get / set");
 	}
+	if(InArgumentName == FoundArgName)
+	{
+		return NSLOCTEXT("FRigDispatch_AnimAttributeBase", "FoundArgTooltip", "Returns true if the animation attribute exists with the specific type");
+	}
+	if(InArgumentName == SuccessArgName)
+	{
+		return NSLOCTEXT("FRigDispatch_AnimAttributeBase", "SuccessArgTooltip", "Returns true if the animation attribute was successfully stored");
+	}
+	return FRigDispatchFactory::GetArgumentTooltip(InArgumentName, InTypeIndex);	
 }
 
-FRigUnit_GetAnimAttribute_Integer_Execute()
+
+
+
+#endif
+
+
+
+TArray<FRigVMTemplateArgument> FRigDispatch_GetAnimAttribute::GetArguments() const
 {
-	DECLARE_SCOPE_HIERARCHICAL_COUNTER_RIGUNIT()
+	if (ValueArgIndex == INDEX_NONE)
+	{
+		Arguments = Super::GetArguments(); 
+
+		FRigVMTemplateArgument::FTypeFilter	TypeFilter;
+		TypeFilter.BindStatic(&FRigDispatch_AnimAttributeBase::IsTypeSupported);
+		
+		DefaultArgIndex = Arguments.Emplace(DefaultArgName, ERigVMPinDirection::Input, GetValueTypeCategory(), TypeFilter);
+		ValueArgIndex = Arguments.Emplace(ValueArgName, ERigVMPinDirection::Output, GetValueTypeCategory(), TypeFilter);
+		
+		FoundArgIndex = Arguments.Emplace(FoundArgName, ERigVMPinDirection::Output, RigVMTypeUtils::TypeIndex::Bool);
+	}
+
 	
-	const int* ValuePtr = GetAnimAttributeValue<int>(false ,Context, Name, BoneName, CachedBoneName, CachedBoneIndex);
-	bWasFound = ValuePtr ? true : false;
-	Value = ValuePtr ? *ValuePtr : FallbackValue;
+	return Arguments;
+}
+
+FRigVMTemplateTypeMap FRigDispatch_GetAnimAttribute::OnNewArgumentType(const FName& InArgumentName,
+	TRigVMTypeIndex InTypeIndex) const
+{
+	FRigVMTemplateTypeMap Types;
+
+	// similar pattern to URigVMArrayNode's FRigVMTemplate_NewArgumentTypeDelegate to avoid double registration
+	// this is needed since a single type is called for both DefaultArg and ValueArg but we should only
+	// register one permutation
+	if (InArgumentName == DefaultArgName)
+	{
+		if (IsTypeSupported(InTypeIndex))
+		{
+			Types.Add(NameArgName, RigVMTypeUtils::TypeIndex::FName);
+			Types.Add(BoneNameArgName, RigVMTypeUtils::TypeIndex::FName);
+			Types.Add(DefaultArgName, InTypeIndex);
+			Types.Add(ValueArgName, InTypeIndex);
+			Types.Add(FoundArgName, RigVMTypeUtils::TypeIndex::Bool);
+			Types.Add(CachedBoneNameArgName, RigVMTypeUtils::TypeIndex::FName);
+			Types.Add(CachedBoneIndexArgName, RigVMTypeUtils::TypeIndex::Int32);
+		}
+	}
+	return Types;
 }
 
-FRigUnit_GetAnimAttribute_Float_Execute()
+FRigVMFunctionPtr FRigDispatch_GetAnimAttribute::GetDispatchFunctionImpl(const FRigVMTemplateTypeMap& InTypes) const
 {
-	DECLARE_SCOPE_HIERARCHICAL_COUNTER_RIGUNIT()
+	const FRigVMRegistry& Registry = FRigVMRegistry::Get();
+	const TRigVMTypeIndex& ValueTypeIndex = InTypes.FindChecked(TEXT("Value"));
+
+	if (ValueTypeIndex == RigVMTypeUtils::TypeIndex::Float)
+	{
+		return &FRigDispatch_GetAnimAttribute::GetAnimAttributeDispatch<float>;
+	}
+	if (ValueTypeIndex == RigVMTypeUtils::TypeIndex::Int32)
+	{
+		return &FRigDispatch_GetAnimAttribute::GetAnimAttributeDispatch<int32>;
+	}
+	if (ValueTypeIndex == RigVMTypeUtils::TypeIndex::FString)
+	{
+		return &FRigDispatch_GetAnimAttribute::GetAnimAttributeDispatch<FString>;
+	}
+	if (ValueTypeIndex == Registry.GetTypeIndex<FTransform>(false))
+	{
+		return &FRigDispatch_GetAnimAttribute::GetAnimAttributeDispatch<FTransform>;
+	}
+	if (ValueTypeIndex == Registry.GetTypeIndex<FQuat>(false))
+	{
+		return &FRigDispatch_GetAnimAttribute::GetAnimAttributeDispatch<FQuat>;
+	}
+	if (ValueTypeIndex == Registry.GetTypeIndex<FVector>(false))
+	{
+		return &FRigDispatch_GetAnimAttribute::GetAnimAttributeDispatch<FVector>;
+	}
+	
+	const FRigVMTemplateArgumentType& ValueType = Registry.GetType(ValueTypeIndex);
+	
+	if (UScriptStruct* ScriptStruct = Cast<UScriptStruct>(ValueType.CPPTypeObject))
+	{
+		return &FRigDispatch_GetAnimAttribute::GetAnimAttributeDispatch;
+	}
 
-	const float* ValuePtr = GetAnimAttributeValue<float>(false ,Context, Name, BoneName, CachedBoneName, CachedBoneIndex);
-	bWasFound = ValuePtr ? true : false;
-	Value = ValuePtr ? *ValuePtr : FallbackValue;
+	return nullptr;
 }
 
-FRigUnit_GetAnimAttribute_Transform_Execute()
+TArray<FRigVMTemplateArgument> FRigDispatch_SetAnimAttribute::GetArguments() const
 {
-	DECLARE_SCOPE_HIERARCHICAL_COUNTER_RIGUNIT()
+	if (ValueArgIndex == INDEX_NONE)
+	{
+		Arguments = Super::GetArguments();
 
-	const FTransform* ValuePtr = GetAnimAttributeValue<FTransform>(false, Context, Name, BoneName, CachedBoneName, CachedBoneIndex);
-	bWasFound = ValuePtr ? true : false;
-	Value = ValuePtr ? *ValuePtr : FallbackValue;
+		FRigVMTemplateArgument::FTypeFilter	TypeFilter;
+		TypeFilter.BindStatic(&FRigDispatch_AnimAttributeBase::IsTypeSupported);
+		
+		ValueArgIndex = Arguments.Emplace(ValueArgName, ERigVMPinDirection::Input, GetValueTypeCategory(), TypeFilter);
+		SuccessArgIndex = Arguments.Emplace(SuccessArgName, ERigVMPinDirection::Output, RigVMTypeUtils::TypeIndex::Bool);
+	}
+	
+	return Arguments;
 }
 
-FRigUnit_GetAnimAttribute_Vector_Execute()
+FRigVMTemplateTypeMap FRigDispatch_SetAnimAttribute::OnNewArgumentType(const FName& InArgumentName,
+	TRigVMTypeIndex InTypeIndex) const
 {
-	DECLARE_SCOPE_HIERARCHICAL_COUNTER_RIGUNIT()
+	FRigVMTemplateTypeMap Types;
 
-	const FVector* ValuePtr = GetAnimAttributeValue<FVector>(false, Context, Name, BoneName, CachedBoneName, CachedBoneIndex);
-	bWasFound = ValuePtr ? true : false;
-	Value = ValuePtr ? *ValuePtr : FallbackValue;	
+	// similar pattern to URigVMArrayNode's FRigVMTemplate_NewArgumentTypeDelegate to avoid double registration
+	// likely not strictly needed for this node since it only has a single non-singleton argument
+	if (InArgumentName == ValueArgName)
+	{
+		if (IsTypeSupported(InTypeIndex))
+		{
+			Types.Add(FRigVMStruct::ExecuteContextName, FRigVMRegistry::Get().GetTypeIndex<FControlRigExecuteContext>());
+			Types.Add(NameArgName, RigVMTypeUtils::TypeIndex::FName);
+			Types.Add(BoneNameArgName, RigVMTypeUtils::TypeIndex::FName);
+			Types.Add(ValueArgName, InTypeIndex);
+			Types.Add(SuccessArgName, RigVMTypeUtils::TypeIndex::Bool);
+			Types.Add(CachedBoneNameArgName, RigVMTypeUtils::TypeIndex::FName);
+			Types.Add(CachedBoneIndexArgName, RigVMTypeUtils::TypeIndex::Int32);
+		}
+	}
+	
+	return Types;	
 }
 
-FRigUnit_GetAnimAttribute_Quaternion_Execute()
+FRigVMFunctionPtr FRigDispatch_SetAnimAttribute::GetDispatchFunctionImpl(const FRigVMTemplateTypeMap& InTypes) const
 {
-	DECLARE_SCOPE_HIERARCHICAL_COUNTER_RIGUNIT()
+	const FRigVMRegistry& Registry = FRigVMRegistry::Get();
+	const TRigVMTypeIndex& ValueTypeIndex = InTypes.FindChecked(TEXT("Value"));
+
+	if (ValueTypeIndex == RigVMTypeUtils::TypeIndex::Float)
+	{
+		return &FRigDispatch_SetAnimAttribute::SetAnimAttributeDispatch<float>;
+	}
+	if (ValueTypeIndex == RigVMTypeUtils::TypeIndex::Int32)
+	{
+		return &FRigDispatch_SetAnimAttribute::SetAnimAttributeDispatch<int32>;
+	}
+	if (ValueTypeIndex == RigVMTypeUtils::TypeIndex::FString)
+	{
+		return &FRigDispatch_SetAnimAttribute::SetAnimAttributeDispatch<FString>;
+	}
+	if (ValueTypeIndex == Registry.GetTypeIndex<FTransform>(false))
+	{
+		return &FRigDispatch_SetAnimAttribute::SetAnimAttributeDispatch<FTransform>;
+	}
+	if (ValueTypeIndex == Registry.GetTypeIndex<FQuat>(false))
+	{
+		return &FRigDispatch_SetAnimAttribute::SetAnimAttributeDispatch<FQuat>;
+	}
+	if (ValueTypeIndex == Registry.GetTypeIndex<FVector>(false))
+	{
+		return &FRigDispatch_SetAnimAttribute::SetAnimAttributeDispatch<FVector>;
+	}
+	
+	const FRigVMTemplateArgumentType& ValueType = Registry.GetType(ValueTypeIndex);
+	
+	if (UScriptStruct* ScriptStruct = Cast<UScriptStruct>(ValueType.CPPTypeObject))
+	{
+		return &FRigDispatch_SetAnimAttribute::SetAnimAttributeDispatch;
+	}
 
-	const FQuat* ValuePtr = GetAnimAttributeValue<FQuat>(false, Context, Name, BoneName, CachedBoneName, CachedBoneIndex);
-	bWasFound = ValuePtr ? true : false;
-	Value = ValuePtr ? *ValuePtr : FallbackValue;	
+	return nullptr;
 }
 
diff --git a/Engine/Plugins/Animation/ControlRig/Source/ControlRig/Private/Units/Animation/RigUnit_AnimAttribute.h b/Engine/Plugins/Animation/ControlRig/Source/ControlRig/Private/Units/Animation/RigUnit_AnimAttribute.h
--- a/Engine/Plugins/Animation/ControlRig/Source/ControlRig/Private/Units/Animation/RigUnit_AnimAttribute.h
+++ b/Engine/Plugins/Animation/ControlRig/Source/ControlRig/Private/Units/Animation/RigUnit_AnimAttribute.h
@@ -2,408 +2,324 @@
 
 #pragma once
 #include "Units/RigUnit.h"
+#include "Units/RigDispatchFactory.h"
+#include "Animation/BuiltInAttributeTypes.h"
 #include "RigUnit_AnimAttribute.generated.h"
 
-/**
- * Sets the value of an animation attribute with the matching names.
- * If the attribute was not found, a new attribute is created.
- *
- * Animation Attributes allow dynamically added data to flow from
- * one Anim Node to other Anim Nodes downstream in the Anim Graph.
- */
-USTRUCT(meta=(Abstract, Category="Animation Attribute", TemplateName="Set Anim Attribute", NodeColor = "0.0 0.36470600962638855 1.0", Varying))
-struct FRigUnit_SetAnimAttributeBase : public FRigUnitMutable
+namespace RigUnit_AnimAttribute
 {
-	GENERATED_BODY()
-};
-
+	template<typename T>
+	class TAnimAttributeType;
 
-USTRUCT(meta=(DisplayName="Set Anim Attribute - Integer"))
-struct CONTROLRIG_API FRigUnit_SetAnimAttribute_Integer: public FRigUnit_SetAnimAttributeBase
-{
-	GENERATED_BODY()
-	
-	FRigUnit_SetAnimAttribute_Integer()
+	template<>
+	class TAnimAttributeType<int>
 	{
-		Name = TEXT("NewAttributeName");
-		BoneName = NAME_None;
-		Value = 0;
-		CachedBoneName = NAME_None;
-		CachedBoneIndex = 0;
-	}
+	public:
+		using Type = FIntegerAnimationAttribute;
+	};
 
-	RIGVM_METHOD()
-	virtual void Execute(const FRigUnitContext& Context) override;
-
-	UPROPERTY(meta=(Input))
-	FName Name;
-	
-	UPROPERTY(meta=(Input, CustomWidget="BoneName"))
-	FName BoneName;
-
-	UPROPERTY(meta=(Input))
-	int32 Value;
-
-	UPROPERTY()
-	FName CachedBoneName;
-	
-	UPROPERTY()
-	int32 CachedBoneIndex;
-};
-
-
-USTRUCT(meta=(DisplayName="Set Anim Attribute - Float"))
-struct CONTROLRIG_API FRigUnit_SetAnimAttribute_Float: public FRigUnit_SetAnimAttributeBase
-{
-	GENERATED_BODY()
-	
-	FRigUnit_SetAnimAttribute_Float()
+	template<>
+	class TAnimAttributeType<float>
 	{
-		Name = TEXT("NewAttributeName");
-		BoneName = NAME_None;
-		Value = 0.f;
-		CachedBoneName = NAME_None;
-		CachedBoneIndex = 0;
-	}
-
-	RIGVM_METHOD()
-	virtual void Execute(const FRigUnitContext& Context) override;
-
-	UPROPERTY(meta=(Input))
-	FName Name;
-	
-	UPROPERTY(meta=(Input, CustomWidget="BoneName"))
-	FName BoneName;
-
-	UPROPERTY(meta=(Input))
-	float Value;
-
-	UPROPERTY()
-	FName CachedBoneName;
-	
-	UPROPERTY()
-	int32 CachedBoneIndex;
-};
+	public:
+		using Type = FFloatAnimationAttribute;
+	};
 
-
-USTRUCT(meta=(DisplayName="Set Anim Attribute - Transform"))
-struct CONTROLRIG_API FRigUnit_SetAnimAttribute_Transform: public FRigUnit_SetAnimAttributeBase
-{
-	GENERATED_BODY()
-	
-	FRigUnit_SetAnimAttribute_Transform()
+	template<>
+	class TAnimAttributeType<FString>
 	{
-		Name = TEXT("NewAttributeName");
-		BoneName = NAME_None;
-		Value = FTransform::Identity;
-		CachedBoneName = NAME_None;
-		CachedBoneIndex = 0;
-	}
-
-	RIGVM_METHOD()
-	virtual void Execute(const FRigUnitContext& Context) override;
-
-	UPROPERTY(meta=(Input))
-	FName Name;
-	
-	UPROPERTY(meta=(Input, CustomWidget="BoneName"))
-	FName BoneName;
-
-	UPROPERTY(meta=(Input))
-	FTransform Value;
+	public:
+		using Type = FStringAnimationAttribute;
+	};
 
-	UPROPERTY()
-	FName CachedBoneName;
-	
-	UPROPERTY()
-	int32 CachedBoneIndex;
-};
-
-USTRUCT(meta=(DisplayName="Set Anim Attribute - Vector"))
-struct CONTROLRIG_API FRigUnit_SetAnimAttribute_Vector: public FRigUnit_SetAnimAttributeBase
-{
-	GENERATED_BODY()
-	
-	FRigUnit_SetAnimAttribute_Vector()
+	template<>
+	class TAnimAttributeType<FTransform>
 	{
-		Name = TEXT("NewAttributeName");
-		BoneName = NAME_None;
-		Value = FVector::ZeroVector;
-		CachedBoneName = NAME_None;
-		CachedBoneIndex = 0;
-	}
+	public:
+		using Type = FTransformAnimationAttribute;
+	};
 
-	RIGVM_METHOD()
-	virtual void Execute(const FRigUnitContext& Context) override;
-
-	UPROPERTY(meta=(Input))
-	FName Name;
-	
-	UPROPERTY(meta=(Input, CustomWidget="BoneName"))
-	FName BoneName;
-
-	UPROPERTY(meta=(Input))
-	FVector Value;
-
-	UPROPERTY()
-	FName CachedBoneName;
-	
-	UPROPERTY()
-	int32 CachedBoneIndex;
-};
+	template<>
+	class TAnimAttributeType<FVector>
+	{
+	public:
+		using Type = FVectorAnimationAttribute;
+	};
 
-USTRUCT(meta=(DisplayName="Set Anim Attribute - Quaternion"))
-struct CONTROLRIG_API FRigUnit_SetAnimAttribute_Quaternion: public FRigUnit_SetAnimAttributeBase
-{
-	GENERATED_BODY()
-	
-	FRigUnit_SetAnimAttribute_Quaternion()
+	template<>
+	class TAnimAttributeType<FQuat>
+	{
+	public:
+		using Type = FQuaternionAnimationAttribute;
+	};
+
+
+	template<typename T>
+	T* GetAnimAttributeValue(
+		bool bAddIfNotFound,
+		const FRigUnitContext& Context,
+		const FName& Name,
+		const FName& BoneName,
+		FName& CachedBoneName,
+		int32& CachedBoneIndex)
 	{
-		Name = TEXT("NewAttributeName");
-		BoneName = NAME_None;
-		Value = FQuat::Identity;
-		CachedBoneName = NAME_None;
-		CachedBoneIndex = 0;
+		if (Name.IsNone())
+		{
+			return nullptr;
+		}
+
+		if (!Context.AnimAttributeContainer)
+		{
+			return nullptr;
+		}
+	
+		const USkeletalMeshComponent* OwningComponent = Cast<USkeletalMeshComponent>(Context.OwningComponent);
+
+		if (!OwningComponent ||
+			!OwningComponent->GetSkeletalMesh())
+		{
+			return nullptr;
+		}
+
+		if (BoneName == NAME_None)
+		{
+			// default to use root bone
+			CachedBoneIndex = 0;
+		}
+		else
+		{
+			// Invalidate cache if input changed
+			if (CachedBoneName != BoneName)
+			{
+				CachedBoneIndex = OwningComponent->GetSkeletalMesh()->GetRefSkeleton().FindBoneIndex(BoneName);
+			}
+		}
+	
+		CachedBoneName = BoneName;
+
+		if (CachedBoneIndex != INDEX_NONE)
+		{
+			const UE::Anim::FAttributeId Id = {Name, FCompactPoseBoneIndex(CachedBoneIndex)} ;
+			typename TAnimAttributeType<T>::Type* Attribute = bAddIfNotFound ?
+				Context.AnimAttributeContainer->FindOrAdd<typename TAnimAttributeType<T>::Type>(Id) :
+				Context.AnimAttributeContainer->Find<typename TAnimAttributeType<T>::Type>(Id);
+			if (Attribute)
+			{
+				return &Attribute->Value;
+			}
+		}
+
+		return nullptr;
 	}
+}
 
-	RIGVM_METHOD()
-	virtual void Execute(const FRigUnitContext& Context) override;
 
-	UPROPERTY(meta=(Input))
-	FName Name;
-	
-	UPROPERTY(meta=(Input, CustomWidget="BoneName"))
-	FName BoneName;
 
-	UPROPERTY(meta=(Input))
-	FQuat Value;
-
-	UPROPERTY()
-	FName CachedBoneName;
-	
-	UPROPERTY()
-	int32 CachedBoneIndex;
-};
 
 /**
- * Gets the value of an animation attribute with the matching names.
- * If the attribute was not found, the fallback value is outputted.
- *
  * Animation Attributes allow dynamically added data to flow from
- * one Anim Node to other Anim Nodes downstream in the Anim Graph.
+ * one Anim Node to other Anim Nodes downstream in the Anim Graph
+ * and accessible from deformer graph
  */
-USTRUCT(meta=(Abstract, Category="Animation Attribute", TemplateName="Get Anim Attribute", NodeColor="0.462745, 1,0, 0.329412", Varying))
-struct FRigUnit_GetAnimAttributeBase : public FRigUnit
-{
-	GENERATED_BODY()
-};
-
-
-USTRUCT(meta=(DisplayName="Get Anim Attribute - Integer"))
-struct CONTROLRIG_API FRigUnit_GetAnimAttribute_Integer: public FRigUnit_GetAnimAttributeBase
+USTRUCT(meta=(Abstract, Category="Animation Attribute", NodeColor="0.462745, 1,0, 0.329412", Varying))
+struct CONTROLRIG_API FRigDispatch_AnimAttributeBase : public FRigDispatchFactory
 {
 	GENERATED_BODY()
-	
-	FRigUnit_GetAnimAttribute_Integer()
-	{
-		Name = TEXT("AttributeName");
-		BoneName = NAME_None;
-		FallbackValue = 0;
-		Value = 0;
-		bWasFound = false;
-		CachedBoneName = NAME_None;
-		CachedBoneIndex = 0;
-	}
 
-	RIGVM_METHOD()
-	virtual void Execute(const FRigUnitContext& Context) override;
+	virtual void RegisterDependencyTypes() const override;
 
-	UPROPERTY(meta=(Input))
-	FName Name;
+#if WITH_EDITOR
+	virtual FString GetNodeTitle(const FRigVMTemplateTypeMap& InTypes) const override;;
+#endif
 	
-	UPROPERTY(meta=(Input, CustomWidget="BoneName"))
-	FName BoneName;
+	virtual TArray<FRigVMTemplateArgument> GetArguments() const override;
+	virtual bool IsSet() const { return false; }
 
-	UPROPERTY(meta=(Input))
-	int32 FallbackValue;
-	
-	UPROPERTY(meta=(Output))
-	int32 Value;
-
-	UPROPERTY(meta=(Output))
-	bool bWasFound;
+#if WITH_EDITOR
+	virtual FText GetArgumentTooltip(const FName& InArgumentName, TRigVMTypeIndex InTypeIndex) const override;
+#endif
 
-	UPROPERTY()
-	FName CachedBoneName;
+protected:
+	static bool IsTypeSupported(const TRigVMTypeIndex& InTypeIndex);
+	static const TArray<FRigVMTemplateArgument::ETypeCategory>& GetValueTypeCategory();
 	
-	UPROPERTY()
-	int32 CachedBoneIndex;
-};
+	mutable TArray<FRigVMTemplateArgument> Arguments;
 
-
-USTRUCT(meta=(DisplayName="Get Anim Attribute - Float"))
-struct CONTROLRIG_API FRigUnit_GetAnimAttribute_Float: public FRigUnit_GetAnimAttributeBase
-{
-	GENERATED_BODY()
+	// IO
+	mutable int32 ExecuteArgIndex = INDEX_NONE;
 	
-	FRigUnit_GetAnimAttribute_Float()
-	{
-		Name = TEXT("AttributeName");
-		BoneName = NAME_None;
-		FallbackValue = 0.f;
-		Value = 0.f;
-		bWasFound = false;
-		CachedBoneName = NAME_None;
-		CachedBoneIndex = 0;
-	}
+	// input
+	mutable int32 NameArgIndex = INDEX_NONE;
+	mutable int32 BoneNameArgIndex = INDEX_NONE;
+	mutable int32 DefaultArgIndex = INDEX_NONE;
 
-	RIGVM_METHOD()
-	virtual void Execute(const FRigUnitContext& Context) override;
 
-	UPROPERTY(meta=(Input))
-	FName Name;
-	
-	UPROPERTY(meta=(Input, CustomWidget="BoneName"))
-	FName BoneName;
+	// output
+	mutable int32 ValueArgIndex = INDEX_NONE;
+	mutable int32 FoundArgIndex = INDEX_NONE;
+	mutable int32 SuccessArgIndex = INDEX_NONE;
 
-	UPROPERTY(meta=(Input))
-	float FallbackValue;
+	// hidden
+	mutable int32 CachedBoneNameArgIndex = INDEX_NONE;
+	mutable int32 CachedBoneIndexArgIndex = INDEX_NONE;
 	
-	UPROPERTY(meta=(Output))
-	float Value;
-
-	UPROPERTY(meta=(Output))
-	bool bWasFound;
+	static FName NameArgName;
+	static FName BoneNameArgName;
+	static FName CachedBoneNameArgName;
+	static FName CachedBoneIndexArgName;
+	static FName DefaultArgName;
+	static FName ValueArgName;
+	static FName FoundArgName;
+	static FName SuccessArgName;
 
-	UPROPERTY()
-	FName CachedBoneName;
-	
-	UPROPERTY()
-	int32 CachedBoneIndex;
 };
 
 
-USTRUCT(meta=(DisplayName="Get Anim Attribute - Transform"))
-struct CONTROLRIG_API FRigUnit_GetAnimAttribute_Transform: public FRigUnit_GetAnimAttributeBase
+/*
+ * Get the value of an animation attribute from the skeletal mesh
+ */
+USTRUCT(meta=(DisplayName="Get Animation Attribute"))
+struct CONTROLRIG_API FRigDispatch_GetAnimAttribute: public FRigDispatch_AnimAttributeBase
 {
 	GENERATED_BODY()
-	
-	FRigUnit_GetAnimAttribute_Transform()
-	{
-		Name = TEXT("AttributeName");
-		BoneName = NAME_None;
-		FallbackValue = FTransform::Identity;
-		Value = FTransform::Identity;
-		bWasFound = false;
-		CachedBoneName = NAME_None;
-		CachedBoneIndex = 0;
-	}
-
-	RIGVM_METHOD()
-	virtual void Execute(const FRigUnitContext& Context) override;
-
-	UPROPERTY(meta=(Input))
-	FName Name;
-	
-	UPROPERTY(meta=(Input, CustomWidget="BoneName"))
-	FName BoneName;
-
-	UPROPERTY(meta=(Input))
-	FTransform FallbackValue;
-	
-	UPROPERTY(meta=(Output))
-	FTransform Value;
 
-	UPROPERTY(meta=(Output))
-	bool bWasFound;
+	virtual TArray<FRigVMTemplateArgument> GetArguments() const override;
+	virtual FRigVMTemplateTypeMap OnNewArgumentType(const FName& InArgumentName, TRigVMTypeIndex InTypeIndex) const override;
 
-	UPROPERTY()
-	FName CachedBoneName;
 	
-	UPROPERTY()
-	int32 CachedBoneIndex;
-};
+protected:
+	virtual FRigVMFunctionPtr GetDispatchFunctionImpl(const FRigVMTemplateTypeMap& InTypes) const override;
 
-
-USTRUCT(meta=(DisplayName="Get Anim Attribute - Vector"))
-struct CONTROLRIG_API FRigUnit_GetAnimAttribute_Vector: public FRigUnit_GetAnimAttributeBase
-{
-	GENERATED_BODY()
-	
-	FRigUnit_GetAnimAttribute_Vector()
+#if WITH_EDITOR
+	template<typename ValueType>
+	FORCEINLINE_DEBUGGABLE bool CheckArgumentTypes(FRigVMMemoryHandleArray Handles) const
 	{
-		Name = TEXT("AttributeName");
-		BoneName = NAME_None;
-		FallbackValue = FVector::ZeroVector;
-		Value = FVector::ZeroVector;
-		bWasFound = false;
-		CachedBoneName = NAME_None;
-		CachedBoneIndex = 0;
+		return CheckArgumentType(Handles[NameArgIndex].IsType<FName>(), NameArgName) &&
+			CheckArgumentType(Handles[BoneNameArgIndex].IsType<FName>(), BoneNameArgName) &&
+			CheckArgumentType(Handles[CachedBoneNameArgIndex].IsType<FName>(true), CachedBoneNameArgName) &&
+			CheckArgumentType(Handles[CachedBoneIndexArgIndex].IsType<int32>(true), CachedBoneIndexArgName) &&
+			CheckArgumentType(Handles[DefaultArgIndex].IsType<ValueType>(), DefaultArgName) &&
+			CheckArgumentType(Handles[ValueArgIndex].IsType<ValueType>(), ValueArgName) &&
+			CheckArgumentType(Handles[FoundArgIndex].IsType<bool>(), FoundArgName);
 	}
 
-	RIGVM_METHOD()
-	virtual void Execute(const FRigUnitContext& Context) override;
-
-	UPROPERTY(meta=(Input))
-	FName Name;
-	
-	UPROPERTY(meta=(Input, CustomWidget="BoneName"))
-	FName BoneName;
-
-	UPROPERTY(meta=(Input))
-	FVector FallbackValue;
-	
-	UPROPERTY(meta=(Output))
-	FVector Value;
+	FORCEINLINE_DEBUGGABLE bool CheckArgumentTypes(FRigVMMemoryHandleArray Handles) const
+	{
+		return CheckArgumentType(Handles[NameArgIndex].IsType<FName>(), NameArgName) &&
+			CheckArgumentType(Handles[BoneNameArgIndex].IsType<FName>(), BoneNameArgName) &&
+			CheckArgumentType(Handles[CachedBoneNameArgIndex].IsType<FName>(true), CachedBoneNameArgName) &&
+			CheckArgumentType(Handles[CachedBoneIndexArgIndex].IsType<int32>(true), CachedBoneIndexArgName) &&
+			CheckArgumentType(Handles[FoundArgIndex].IsType<bool>(), FoundArgName);
+	}
+#endif
 
-	UPROPERTY(meta=(Output))
-	bool bWasFound;
+	// dispatch function for built-in types
+	template<typename ValueType>	
+	FORCEINLINE_DEBUGGABLE static void GetAnimAttributeDispatch(FRigVMExtendedExecuteContext& InContext, FRigVMMemoryHandleArray Handles)
+	{
+		const FRigDispatch_GetAnimAttribute* Factory = static_cast<const FRigDispatch_GetAnimAttribute*>(InContext.Factory);
+		
+#if WITH_EDITOR
+		if(!Factory->CheckArgumentTypes<ValueType>(Handles))
+		{
+			return;
+		}
+#endif
+
+		// unpack the memory
+		const FName& Name = *(const FName*)Handles[Factory->NameArgIndex].GetData();
+		const FName& BoneName = *(const FName*)Handles[Factory->BoneNameArgIndex].GetData();
+		const ValueType& Default = *(const ValueType*)Handles[Factory->DefaultArgIndex].GetData();
+		
+		ValueType& Value = *(ValueType*)Handles[Factory->ValueArgIndex].GetData();
+		bool& Found = *(bool*)Handles[Factory->FoundArgIndex].GetData();
+		
+		FName& CachedBoneName = *(FName*)Handles[Factory->CachedBoneNameArgIndex].GetData(false, InContext.GetSlice().GetIndex());
+		int32& CachedBoneIndex = *(int32*)Handles[Factory->CachedBoneIndexArgIndex].GetData(false, InContext.GetSlice().GetIndex());
+
+		// extract the animation attribute
+		const FRigUnitContext& Context = GetRigUnitContext(InContext);
+		const ValueType* ValuePtr = RigUnit_AnimAttribute::GetAnimAttributeValue<ValueType>(false, Context, Name, BoneName, CachedBoneName, CachedBoneIndex);
+		Found = ValuePtr ? true : false;
+		Value = ValuePtr ? *ValuePtr : Default;
+	}
 
-	UPROPERTY()
-	FName CachedBoneName;
+	// dispatch function for user/dev defined types
+	static void GetAnimAttributeDispatch(FRigVMExtendedExecuteContext& InContext, FRigVMMemoryHandleArray Handles);
 	
-	UPROPERTY()
-	int32 CachedBoneIndex;
 };
 
-
-USTRUCT(meta=(DisplayName="Get Anim Attribute - Quaternion"))
-struct CONTROLRIG_API FRigUnit_GetAnimAttribute_Quaternion: public FRigUnit_GetAnimAttributeBase
+/*
+ * Modify an animation attribute if one is found, otherwise add a new animation attribute
+ */
+USTRUCT(meta=(DisplayName="Set Animation Attribute"))
+struct CONTROLRIG_API FRigDispatch_SetAnimAttribute: public FRigDispatch_AnimAttributeBase
 {
 	GENERATED_BODY()
+	virtual bool IsSet() const override { return true; }
+	virtual TArray<FRigVMTemplateArgument> GetArguments() const override;
+	virtual FRigVMTemplateTypeMap OnNewArgumentType(const FName& InArgumentName, TRigVMTypeIndex InTypeIndex) const override;
 	
-	FRigUnit_GetAnimAttribute_Quaternion()
+protected:
+	virtual FRigVMFunctionPtr GetDispatchFunctionImpl(const FRigVMTemplateTypeMap& InTypes) const override;
+
+#if WITH_EDITOR
+	template<typename ValueType>
+	FORCEINLINE_DEBUGGABLE bool CheckArgumentTypes(FRigVMMemoryHandleArray Handles) const
 	{
-		Name = TEXT("AttributeName");
-		BoneName = NAME_None;
-		FallbackValue = FQuat::Identity;
-		Value = FQuat::Identity;
-		bWasFound = false;
-		CachedBoneName = NAME_None;
-		CachedBoneIndex = 0;
+		return CheckArgumentType(Handles[NameArgIndex].IsType<FName>(), NameArgName) &&
+			CheckArgumentType(Handles[BoneNameArgIndex].IsType<FName>(), BoneNameArgName) &&
+			CheckArgumentType(Handles[CachedBoneNameArgIndex].IsType<FName>(true), CachedBoneNameArgName) &&
+			CheckArgumentType(Handles[CachedBoneIndexArgIndex].IsType<int32>(true), CachedBoneIndexArgName) &&
+			CheckArgumentType(Handles[ValueArgIndex].IsType<ValueType>(), ValueArgName) &&
+			CheckArgumentType(Handles[SuccessArgIndex].IsType<bool>(), FoundArgName);
 	}
 
-	RIGVM_METHOD()
-	virtual void Execute(const FRigUnitContext& Context) override;
-
-	UPROPERTY(meta=(Input))
-	FName Name;
-	
-	UPROPERTY(meta=(Input, CustomWidget="BoneName"))
-	FName BoneName;
+	FORCEINLINE_DEBUGGABLE bool CheckArgumentTypes(FRigVMMemoryHandleArray Handles) const
+	{
+		return CheckArgumentType(Handles[NameArgIndex].IsType<FName>(), NameArgName) &&
+			CheckArgumentType(Handles[BoneNameArgIndex].IsType<FName>(), BoneNameArgName) &&
+			CheckArgumentType(Handles[CachedBoneNameArgIndex].IsType<FName>(true), CachedBoneNameArgName) &&
+			CheckArgumentType(Handles[CachedBoneIndexArgIndex].IsType<int32>(true), CachedBoneIndexArgName) &&
+			CheckArgumentType(Handles[SuccessArgIndex].IsType<bool>(), FoundArgName);
+	}
+#endif
 
-	UPROPERTY(meta=(Input))
-	FQuat FallbackValue;
-	
-	UPROPERTY(meta=(Output))
-	FQuat Value;
+	// dispatch function for built-in types
+	template<typename ValueType>	
+	FORCEINLINE_DEBUGGABLE static void SetAnimAttributeDispatch(FRigVMExtendedExecuteContext& InContext, FRigVMMemoryHandleArray Handles)
+	{
+		const FRigDispatch_SetAnimAttribute* Factory = static_cast<const FRigDispatch_SetAnimAttribute*>(InContext.Factory);
+		
+#if WITH_EDITOR
+		if(!Factory->CheckArgumentTypes<ValueType>(Handles))
+		{
+			return;
+		}
+#endif
+
+		// unpack the memory
+		const FName& Name = *(const FName*)Handles[Factory->NameArgIndex].GetData();
+		const FName& BoneName = *(const FName*)Handles[Factory->BoneNameArgIndex].GetData();
+		const ValueType& Value = *(ValueType*)Handles[Factory->ValueArgIndex].GetData();
+		
+		bool& Success = *(bool*)Handles[Factory->SuccessArgIndex].GetData();
+		Success = false;
+		
+		FName& CachedBoneName = *(FName*)Handles[Factory->CachedBoneNameArgIndex].GetData(false, InContext.GetSlice().GetIndex());
+		int32& CachedBoneIndex = *(int32*)Handles[Factory->CachedBoneIndexArgIndex].GetData(false, InContext.GetSlice().GetIndex());
+
+		// extract the animation attribute
+		const FRigUnitContext& Context = GetRigUnitContext(InContext);
+		ValueType* ValuePtr = RigUnit_AnimAttribute::GetAnimAttributeValue<ValueType>(true, Context, Name, BoneName, CachedBoneName, CachedBoneIndex);
+		if (ValuePtr)
+		{
+			Success = true;
+			*ValuePtr = Value;
+		}
+	}
 
-	UPROPERTY(meta=(Output))
-	bool bWasFound;
+	// dispatch function for user/dev defined types
+	static void SetAnimAttributeDispatch(FRigVMExtendedExecuteContext& InContext, FRigVMMemoryHandleArray Handles);
+};
 
-	UPROPERTY()
-	FName CachedBoneName;
-	
-	UPROPERTY()
-	int32 CachedBoneIndex;
-};
\ No newline at end of file
diff --git a/Engine/Plugins/Animation/ControlRig/Source/ControlRig/Private/Units/ControlRigNodeWorkflow.cpp b/Engine/Plugins/Animation/ControlRig/Source/ControlRig/Private/Units/ControlRigNodeWorkflow.cpp
--- a/Engine/Plugins/Animation/ControlRig/Source/ControlRig/Private/Units/ControlRigNodeWorkflow.cpp
+++ b/Engine/Plugins/Animation/ControlRig/Source/ControlRig/Private/Units/ControlRigNodeWorkflow.cpp
@@ -28,7 +28,7 @@ TArray<FRigVMUserWorkflow> UControlRigTransformWorkflowOptions::ProvideWorkflows
 	{
 		if(!Pin->IsArray())
 		{
-			if(Pin->GetCPPType() == TBaseStructure<FTransform>::Get()->GetStructCPPName())
+			if(Pin->GetCPPType() == RigVMTypeUtils::GetUniqueStructTypeName(TBaseStructure<FTransform>::Get()))
 			{
 				Workflows.Emplace(
 					TEXT("Set from hierarchy"),
diff --git a/Engine/Plugins/Animation/ControlRig/Source/ControlRigDeveloper/Private/ControlRigBlueprint.cpp b/Engine/Plugins/Animation/ControlRig/Source/ControlRigDeveloper/Private/ControlRigBlueprint.cpp
--- a/Engine/Plugins/Animation/ControlRig/Source/ControlRigDeveloper/Private/ControlRigBlueprint.cpp
+++ b/Engine/Plugins/Animation/ControlRig/Source/ControlRigDeveloper/Private/ControlRigBlueprint.cpp
@@ -1038,6 +1038,11 @@ void UControlRigBlueprint::PostLoad()
 		HandlePackageDone();
 	}
 #endif
+
+	// RigVMRegistry changes can be triggered when new user defined types(structs/enums) are added/removed
+	// in which case we have to refresh the model
+	FRigVMRegistry::Get().OnRigVMRegistryChanged().RemoveAll(this);
+	FRigVMRegistry::Get().OnRigVMRegistryChanged().AddUObject(this, &UControlRigBlueprint::OnRigVMRegistryChanged);
 }
 
 #if WITH_EDITOR
@@ -1582,6 +1587,14 @@ void UControlRigBlueprint::RefreshAllModels(EControlRigBlueprintLoadType InLoadT
 	}
 }
 
+void UControlRigBlueprint::OnRigVMRegistryChanged()
+{
+	RefreshAllModels();
+	RebuildGraphFromModel();
+	// avoids slate crash
+	FControlRigBlueprintUtils::HandleRefreshAllNodes(this);
+}
+
 void UControlRigBlueprint::HandleReportFromCompiler(EMessageSeverity::Type InSeverity, UObject* InSubject, const FString& InMessage)
 {
 	UObject* SubjectForMessage = InSubject;
@@ -2833,7 +2846,7 @@ bool UControlRigBlueprint::ChangeMemberVariableType(const FName& InName, const F
 	}
 	else if(UScriptStruct* ScriptStruct = URigVMPin::FindObjectFromCPPTypeObjectPath<UScriptStruct>(CPPType))
 	{
-		Variable.TypeName = *ScriptStruct->GetStructCPPName();
+		Variable.TypeName = *RigVMTypeUtils::GetUniqueStructTypeName(ScriptStruct);
 		Variable.TypeObject = ScriptStruct;
 		Variable.Size = ScriptStruct->GetStructureSize();
 	}
diff --git a/Engine/Plugins/Animation/ControlRig/Source/ControlRigDeveloper/Public/ControlRigBlueprint.h b/Engine/Plugins/Animation/ControlRig/Source/ControlRigDeveloper/Public/ControlRigBlueprint.h
--- a/Engine/Plugins/Animation/ControlRig/Source/ControlRigDeveloper/Public/ControlRigBlueprint.h
+++ b/Engine/Plugins/Animation/ControlRig/Source/ControlRigDeveloper/Public/ControlRigBlueprint.h
@@ -299,6 +299,10 @@ public:
 	// any pin that references the struct needs to be regenerated
 	void RefreshAllModels(EControlRigBlueprintLoadType InLoadType = EControlRigBlueprintLoadType::PostLoad);
 
+	// RigVMRegistry changes can be triggered when new user defined types(structs/enums) are added/removed
+	// in which case we have to refresh the model
+	void OnRigVMRegistryChanged();
+	
 	UFUNCTION(BlueprintCallable, Category = "Control Rig Blueprint")
 	void RequestControlRigInit();
 
diff --git a/Engine/Plugins/Animation/ControlRig/Source/ControlRigEditor/Private/Editor/ControlRigEditor.cpp b/Engine/Plugins/Animation/ControlRig/Source/ControlRigEditor/Private/Editor/ControlRigEditor.cpp
--- a/Engine/Plugins/Animation/ControlRig/Source/ControlRigEditor/Private/Editor/ControlRigEditor.cpp
+++ b/Engine/Plugins/Animation/ControlRig/Source/ControlRigEditor/Private/Editor/ControlRigEditor.cpp
@@ -5631,7 +5631,7 @@ void FControlRigEditor::CreateRigHierarchyToGraphDragAndDropMenu() const
 										{
 											Controller->OpenUndoBracket(TEXT("Create Item Array From Selection"));
 
-											static const FString ItemArrayCPPType = FString::Printf(TEXT("TArray<%s>"), *FRigElementKey::StaticStruct()->GetStructCPPName());
+											static const FString ItemArrayCPPType = FString::Printf(TEXT("TArray<%s>"), *RigVMTypeUtils::GetUniqueStructTypeName(FRigElementKey::StaticStruct()));
 											static const FString ItemArrayObjectPath = FRigElementKey::StaticStruct()->GetPathName();
 							
 											if (URigVMNode* ItemsNode = Controller->AddFreeRerouteNode(true, ItemArrayCPPType, *ItemArrayObjectPath, false, NAME_None, FString(), NodePosition))
@@ -5830,7 +5830,7 @@ void FControlRigEditor::HandleMakeElementGetterSetter(ERigElementGetterSetterTyp
 						}
 
 						UScriptStruct* ValueStruct = TBaseStructure<FVector2D>::Get();
-						const FRigVMTemplateArgumentType TypeForStruct(*ValueStruct->GetStructCPPName(), ValueStruct);
+						const FRigVMTemplateArgumentType TypeForStruct(*RigVMTypeUtils::GetUniqueStructTypeName(ValueStruct), ValueStruct);
 						const int32 TypeIndex = FRigVMRegistry::Get().GetTypeIndex(TypeForStruct);
 						PinsToResolve.Add(ValueName, TypeIndex);
 						break;
@@ -5847,7 +5847,7 @@ void FControlRigEditor::HandleMakeElementGetterSetter(ERigElementGetterSetterTyp
 							StructTemplate = FRigUnit_SetVectorAnimationChannel::StaticStruct();
 						}
 						UScriptStruct* ValueStruct = TBaseStructure<FVector>::Get();
-						const FRigVMTemplateArgumentType TypeForStruct(*ValueStruct->GetStructCPPName(), ValueStruct);
+						const FRigVMTemplateArgumentType TypeForStruct(*RigVMTypeUtils::GetUniqueStructTypeName(ValueStruct), ValueStruct);
 						const int32 TypeIndex = FRigVMRegistry::Get().GetTypeIndex(TypeForStruct);
 						PinsToResolve.Add(ValueName, TypeIndex);
 						break;
@@ -5881,7 +5881,7 @@ void FControlRigEditor::HandleMakeElementGetterSetter(ERigElementGetterSetterTyp
 							StructTemplate = FRigUnit_SetTransformAnimationChannel::StaticStruct();
 						}
 						UScriptStruct* ValueStruct = TBaseStructure<FTransform>::Get();
-						const FRigVMTemplateArgumentType TypeForStruct(*ValueStruct->GetStructCPPName(), ValueStruct);
+						const FRigVMTemplateArgumentType TypeForStruct(*RigVMTypeUtils::GetUniqueStructTypeName(ValueStruct), ValueStruct);
 						const int32 TypeIndex = FRigVMRegistry::Get().GetTypeIndex(TypeForStruct);
 						PinsToResolve.Add(ValueName, TypeIndex);
 						break;
diff --git a/Engine/Plugins/Animation/ControlRig/Source/ControlRigEditor/Private/Editor/SRigHierarchyTreeView.cpp b/Engine/Plugins/Animation/ControlRig/Source/ControlRigEditor/Private/Editor/SRigHierarchyTreeView.cpp
--- a/Engine/Plugins/Animation/ControlRig/Source/ControlRigEditor/Private/Editor/SRigHierarchyTreeView.cpp
+++ b/Engine/Plugins/Animation/ControlRig/Source/ControlRigEditor/Private/Editor/SRigHierarchyTreeView.cpp
@@ -839,20 +839,20 @@ FLinearColor SRigHierarchyItem::GetColorForControlType(ERigControlType InControl
 		case ERigControlType::Vector2D:
 		{
 			UScriptStruct* Struct = TBaseStructure<FVector2D>::Get(); 
-			PinType = RigVMTypeUtils::PinTypeFromCPPType(*Struct->GetStructCPPName(), Struct);
+			PinType = RigVMTypeUtils::PinTypeFromCPPType(*RigVMTypeUtils::GetUniqueStructTypeName(Struct), Struct);
 			break;
 		}
 		case ERigControlType::Position:
 		case ERigControlType::Scale:
 		{
 			UScriptStruct* Struct = TBaseStructure<FVector>::Get(); 
-			PinType = RigVMTypeUtils::PinTypeFromCPPType(*Struct->GetStructCPPName(), Struct);
+			PinType = RigVMTypeUtils::PinTypeFromCPPType(*RigVMTypeUtils::GetUniqueStructTypeName(Struct), Struct);
 			break;
 		}
 		case ERigControlType::Rotator:
 		{
 			UScriptStruct* Struct = TBaseStructure<FRotator>::Get(); 
-			PinType = RigVMTypeUtils::PinTypeFromCPPType(*Struct->GetStructCPPName(), Struct);
+			PinType = RigVMTypeUtils::PinTypeFromCPPType(*RigVMTypeUtils::GetUniqueStructTypeName(Struct), Struct);
 			break;
 		}
 		case ERigControlType::Transform:
@@ -861,7 +861,7 @@ FLinearColor SRigHierarchyItem::GetColorForControlType(ERigControlType InControl
 		default:
 		{
 			UScriptStruct* Struct = TBaseStructure<FTransform>::Get(); 
-			PinType = RigVMTypeUtils::PinTypeFromCPPType(*Struct->GetStructCPPName(), Struct);
+			PinType = RigVMTypeUtils::PinTypeFromCPPType(*RigVMTypeUtils::GetUniqueStructTypeName(Struct), Struct);
 			break;
 		}
 	}
diff --git a/Engine/Plugins/Animation/ControlRig/Source/ControlRigEditor/Private/Graph/ControlRigGraphPanelPinFactory.cpp b/Engine/Plugins/Animation/ControlRig/Source/ControlRigEditor/Private/Graph/ControlRigGraphPanelPinFactory.cpp
--- a/Engine/Plugins/Animation/ControlRig/Source/ControlRigEditor/Private/Graph/ControlRigGraphPanelPinFactory.cpp
+++ b/Engine/Plugins/Animation/ControlRig/Source/ControlRigEditor/Private/Graph/ControlRigGraphPanelPinFactory.cpp
@@ -138,7 +138,7 @@ TSharedPtr<SGraphPin> FControlRigGraphPanelPinFactory::CreatePin(UEdGraphPin* In
 										break;
 									}
 
-									if(Pin->GetCPPType() == FRigElementKey::StaticStruct()->GetStructCPPName())
+									if(Pin->GetCPPType() == RigVMTypeUtils::GetUniqueStructTypeName(FRigElementKey::StaticStruct()))
 									{
 										const FString DefaultValue = Pin->GetDefaultValue();
 										if(!DefaultValue.IsEmpty())
diff --git a/Engine/Plugins/Animation/ControlRig/Source/ControlRigEditor/Private/Graph/SControlRigGraphChangePinType.cpp b/Engine/Plugins/Animation/ControlRig/Source/ControlRigEditor/Private/Graph/SControlRigGraphChangePinType.cpp
--- a/Engine/Plugins/Animation/ControlRig/Source/ControlRigEditor/Private/Graph/SControlRigGraphChangePinType.cpp
+++ b/Engine/Plugins/Animation/ControlRig/Source/ControlRigEditor/Private/Graph/SControlRigGraphChangePinType.cpp
@@ -214,7 +214,13 @@ void SControlRigChangePinType::FillPinTypeMenu(FMenuBuilder& MenuBuilder)
 					for(int32 PermutationIndex = 0; PermutationIndex < Template->NumPermutations(); PermutationIndex++)
 					{
 						TRigVMTypeIndex ArgumentTypeIndex = AllArgumentTypes[PermutationIndex];
-					
+
+						if (ArgumentTypeIndex == INDEX_NONE)
+						{
+							// skip invalid permutations
+							continue;
+						}
+
 						bool bIsFilteredOut = false;
 						if(!FilteredArgumentTypes.Contains(ArgumentTypeIndex))
 						{
@@ -326,7 +332,7 @@ void SControlRigChangePinType::FillPinTypeMenu(FMenuBuilder& MenuBuilder)
 
 				if(UScriptStruct* ScriptStruct = Cast<UScriptStruct>(Type.CPPTypeObject))
 				{
-					const FRigVMTemplateArgumentType MathType(*ScriptStruct->GetStructCPPName(), ScriptStruct);
+					const FRigVMTemplateArgumentType MathType(*RigVMTypeUtils::GetUniqueStructTypeName(ScriptStruct), ScriptStruct);
 					const TRigVMTypeIndex MathTypeIndex = FRigVMRegistry::Get().GetTypeIndex(MathType);
 					
 					if(!FRigVMRegistry::Get().GetTypesForCategory(FRigVMTemplateArgument::ETypeCategory_SingleMathStructValue).Contains(MathTypeIndex))
diff --git a/Engine/Source/Developer/RigVMDeveloper/Private/RigVMDeveloperTypeUtils.cpp b/Engine/Source/Developer/RigVMDeveloper/Private/RigVMDeveloperTypeUtils.cpp
--- a/Engine/Source/Developer/RigVMDeveloper/Private/RigVMDeveloperTypeUtils.cpp
+++ b/Engine/Source/Developer/RigVMDeveloper/Private/RigVMDeveloperTypeUtils.cpp
@@ -118,7 +118,7 @@ FRigVMExternalVariable RigVMTypeUtils::ExternalVariableFromPinType(const FName&
 	{
 		if (UScriptStruct* Struct = Cast<UScriptStruct>(InPinType.PinSubCategoryObject))
 		{
-			ExternalVariable.TypeName = *Struct->GetStructCPPName();
+			ExternalVariable.TypeName = *RigVMTypeUtils::GetUniqueStructTypeName(Struct);
 			ExternalVariable.TypeObject = Struct;
 			ExternalVariable.Size = Struct->GetStructureSize();
 		}
@@ -178,7 +178,7 @@ FRigVMExternalVariable RigVMTypeUtils::ExternalVariableFromCPPTypePath(const FNa
 	}
 	else if(UScriptStruct* ScriptStruct = URigVMPin::FindObjectFromCPPTypeObjectPath<UScriptStruct>(CPPTypePath))
 	{
-		Variable.TypeName = *ScriptStruct->GetStructCPPName();
+		Variable.TypeName = *RigVMTypeUtils::GetUniqueStructTypeName(ScriptStruct);
 		Variable.TypeObject = ScriptStruct;
 		Variable.Size = ScriptStruct->GetStructureSize();
 	}
@@ -247,7 +247,7 @@ FRigVMExternalVariable RigVMTypeUtils::ExternalVariableFromCPPType(const FName&
 	}
 	else if(UScriptStruct* ScriptStruct = URigVMPin::FindObjectFromCPPTypeObjectPath<UScriptStruct>(CPPType))
 	{
-		Variable.TypeName = *ScriptStruct->GetStructCPPName();
+		Variable.TypeName = *RigVMTypeUtils::GetUniqueStructTypeName(ScriptStruct);
 		Variable.TypeObject = ScriptStruct;
 		Variable.Size = ScriptStruct->GetStructureSize();
 	}
@@ -637,7 +637,7 @@ bool RigVMTypeUtils::CPPTypeFromPinType(const FEdGraphPinType& InPinType, FStrin
 			OutCPPTypeObjectPath = *CPPTypeObject->GetPathName();
 			if (UScriptStruct* ScriptStruct = Cast<UScriptStruct>(CPPTypeObject))
 			{
-				FString StructName = ScriptStruct->GetStructCPPName();
+				FString StructName = GetUniqueStructTypeName(ScriptStruct);
 				while (IsArrayType(OutCPPType))
 				{
 					OutCPPType = BaseTypeFromArrayType(OutCPPType);
@@ -707,7 +707,7 @@ bool RigVMTypeUtils::CPPTypeFromExternalVariable(const FRigVMExternalVariable& I
 	}
 	else if (UScriptStruct* Struct = Cast<UScriptStruct>(InExternalVariable.TypeObject))
 	{
-		OutCPPType = Prefix + *Struct->GetStructCPPName() + Suffix;
+		OutCPPType = Prefix + *GetUniqueStructTypeName(Struct) + Suffix;
 		*OutCPPTypeObject = Struct;	
 	}
 	else if (UEnum* Enum = Cast<UEnum>(InExternalVariable.TypeObject))
diff --git a/Engine/Source/Developer/RigVMDeveloper/Private/RigVMModel/Nodes/RigVMArrayNode.cpp b/Engine/Source/Developer/RigVMDeveloper/Private/RigVMModel/Nodes/RigVMArrayNode.cpp
--- a/Engine/Source/Developer/RigVMDeveloper/Private/RigVMModel/Nodes/RigVMArrayNode.cpp
+++ b/Engine/Source/Developer/RigVMDeveloper/Private/RigVMModel/Nodes/RigVMArrayNode.cpp
@@ -293,7 +293,7 @@ const FRigVMTemplate* URigVMArrayNode::GetTemplate() const
 			static const TArray<FRigVMTemplateArgument::ETypeCategory> ArrayCategory = {FRigVMTemplateArgument::ETypeCategory_ArrayAnyValue};
 
 			UScriptStruct* ExecuteStruct = GetGraph()->GetExecuteContextStruct();
-			const FRigVMTemplateArgumentType ExecuteType(*ExecuteStruct->GetStructCPPName(), ExecuteStruct);
+			const FRigVMTemplateArgumentType ExecuteType(*RigVMTypeUtils::GetUniqueStructTypeName(ExecuteStruct), ExecuteStruct);
 			const int32 ExecuteTypeIndex = FRigVMRegistry::Get().FindOrAddType(ExecuteType);
 			
 			switch(OpCode)
diff --git a/Engine/Source/Developer/RigVMDeveloper/Private/RigVMModel/Nodes/RigVMTemplateNode.cpp b/Engine/Source/Developer/RigVMDeveloper/Private/RigVMModel/Nodes/RigVMTemplateNode.cpp
--- a/Engine/Source/Developer/RigVMDeveloper/Private/RigVMModel/Nodes/RigVMTemplateNode.cpp
+++ b/Engine/Source/Developer/RigVMDeveloper/Private/RigVMModel/Nodes/RigVMTemplateNode.cpp
@@ -59,6 +59,18 @@ void URigVMTemplateNode::PostLoad()
 {
 	Super::PostLoad();
 
+	for (URigVMPin* Pin : GetPins())
+	{
+		// GetTypeIndex ensures that the pin type is registered
+		// Only need to register for visible pins
+		if (Pin->GetDirection() == ERigVMPinDirection::IO ||
+			Pin->GetDirection() == ERigVMPinDirection::Input ||
+			Pin->GetDirection() == ERigVMPinDirection::Output)
+		{
+			Pin->GetTypeIndex();
+		}
+	}
+	
 	// if there are brackets in the notation remove them
 	const FString OriginalNotation = TemplateNotation.ToString();
 	const FString SanitizedNotation = OriginalNotation.Replace(TEXT("[]"), TEXT(""));
@@ -603,6 +615,12 @@ FString URigVMTemplateNode::GetInitialDefaultValueForPin(const FName& InRootPinN
 
 			const TRigVMTypeIndex TypeIndex = Argument->GetTypeIndices()[PermutationIndex];
 
+			// INDEX_NONE indicates deleted permutation
+			if (TypeIndex == INDEX_NONE)
+			{
+				continue;
+			}
+			
 			if(Factory)
 			{
 				NewDefaultValue = Factory->GetArgumentDefaultValue(Argument->GetName(), TypeIndex);
diff --git a/Engine/Source/Developer/RigVMDeveloper/Private/RigVMModel/RigVMController.cpp b/Engine/Source/Developer/RigVMDeveloper/Private/RigVMModel/RigVMController.cpp
--- a/Engine/Source/Developer/RigVMDeveloper/Private/RigVMModel/RigVMController.cpp
+++ b/Engine/Source/Developer/RigVMDeveloper/Private/RigVMModel/RigVMController.cpp
@@ -12489,7 +12489,7 @@ URigVMIfNode* URigVMController::AddIfNodeFromStruct(UScriptStruct* InScriptStruc
 		return nullptr;
 	}
 
-	return AddIfNode(InScriptStruct->GetStructCPPName(), FName(InScriptStruct->GetPathName()), InPosition, InNodeName, bSetupUndoRedo);
+	return AddIfNode(RigVMTypeUtils::GetUniqueStructTypeName(InScriptStruct), FName(InScriptStruct->GetPathName()), InPosition, InNodeName, bSetupUndoRedo);
 }
 
 URigVMSelectNode* URigVMController::AddSelectNode(const FString& InCPPType, const FName& InCPPTypeObjectPath, const FVector2D& InPosition, const FString& InNodeName, bool bSetupUndoRedo, bool bPrintPythonCommand)
@@ -12611,7 +12611,7 @@ URigVMSelectNode* URigVMController::AddSelectNodeFromStruct(UScriptStruct* InScr
 		return nullptr;
 	}
 
-	return AddSelectNode(InScriptStruct->GetStructCPPName(), FName(InScriptStruct->GetPathName()), InPosition, InNodeName, bSetupUndoRedo);
+	return AddSelectNode(RigVMTypeUtils::GetUniqueStructTypeName(InScriptStruct), FName(InScriptStruct->GetPathName()), InPosition, InNodeName, bSetupUndoRedo);
 }
 
 URigVMTemplateNode* URigVMController::AddTemplateNode(const FName& InNotation, const FVector2D& InPosition, const FString& InNodeName, bool bSetupUndoRedo, bool bPrintPythonCommand)
@@ -14620,6 +14620,7 @@ void URigVMController::RepopulatePinsOnNode(URigVMNode* InNode, bool bFollowCore
 	URigVMIfNode* IfNode = Cast<URigVMIfNode>(InNode);
 	URigVMSelectNode* SelectNode = Cast<URigVMSelectNode>(InNode);
 	URigVMArrayNode* ArrayNode = Cast<URigVMArrayNode>(InNode);
+	URigVMDispatchNode* DispatchNode = Cast<URigVMDispatchNode>(InNode);
 
 	TGuardValue<bool> EventuallySuspendNotifs(bSuspendNotifications, !bNotify);
 	FScopeLock Lock(&PinPathCoreRedirectorsLock);
@@ -14689,6 +14690,52 @@ void URigVMController::RepopulatePinsOnNode(URigVMNode* InNode, bool bFollowCore
 		CreateDefaultValueForStructIfRequired(ScriptStruct, ExportedDefaultValue);
 		AddPinsForStruct(ScriptStruct, UnitNode, nullptr, ERigVMPinDirection::Invalid, ExportedDefaultValue, false, bNotify);
 	}
+	else if (DispatchNode)
+	{
+		TMap<FString, TRigVMTypeIndex> PinTypeMap;
+		for (const URigVMPin* Pin : DispatchNode->Pins)
+		{
+			PinTypeMap.Add(Pin->GetPinPath(), Pin->GetTypeIndex());
+		}
+		
+		RemovePinsDuringRepopulate(DispatchNode, DispatchNode->Pins, bNotify, bSetupOrphanedPins);
+
+		const FRigVMTemplate* Template = DispatchNode->GetTemplate();
+		
+		for (int32 ArgIndex = 0; ArgIndex < Template->NumArguments(); ArgIndex++)
+		{
+			const FRigVMTemplateArgument* Arg = Template->GetArgument(ArgIndex);
+
+			URigVMPin* Pin = NewObject<URigVMPin>(DispatchNode, Arg->GetName());
+			const TRigVMTypeIndex& TypeIndex = PinTypeMap.FindChecked(Pin->GetPinPath());
+			const FRigVMTemplateArgumentType Type = FRigVMRegistry::Get().GetType(TypeIndex);
+			Pin->CPPType = Type.CPPType.ToString();
+			Pin->CPPTypeObject = Type.CPPTypeObject;
+			if (Pin->CPPTypeObject)
+			{
+				Pin->CPPTypeObjectPath = *Pin->CPPTypeObject->GetPathName();
+			}
+			Pin->Direction = Arg->GetDirection();
+
+			AddNodePin(DispatchNode, Pin);
+
+			if(!Pin->IsWildCard())
+			{
+				// any serialize default values will be applied later with ApplyPinStates(...)
+				const FString DefaultValue = DispatchNode->GetInitialDefaultValueForPin(Pin->GetFName());
+				if(UScriptStruct* ScriptStruct = Cast<UScriptStruct>(Pin->CPPTypeObject))
+				{
+					AddPinsForStruct(ScriptStruct, Pin->GetNode(), Pin, Pin->Direction, DefaultValue, false, false);
+				}
+				else if(!DefaultValue.IsEmpty())
+				{
+					SetPinDefaultValue(Pin, DefaultValue, true, false, false, false);
+				}
+			}
+		}
+		
+		ResolveTemplateNodeMetaData(DispatchNode, false);
+	}
 	else if ((RerouteNode != nullptr) || (VariableNode != nullptr))
 	{
 		if (InNode->GetPins().Num() == 0)
@@ -16629,8 +16676,12 @@ bool URigVMController::UpdateTemplateNodePinTypes(URigVMTemplateNode* InNode, bo
 			for (const int32& Permutation : ResolvedPermutations)
 			{
 				if (const FRigVMTemplateArgument* Argument = Template->FindArgument(Pin->GetFName()))
-				{
-					Types.AddUnique(Argument->TypeIndices[Permutation]);
+				{	
+					// INDEX_NONE indicates invalid permutation
+					if (Argument->TypeIndices[Permutation] != INDEX_NONE)
+					{
+						Types.AddUnique(Argument->TypeIndices[Permutation]);
+					}
 				}
 			}
 			
diff --git a/Engine/Source/Developer/RigVMDeveloper/Private/RigVMModel/RigVMPin.cpp b/Engine/Source/Developer/RigVMDeveloper/Private/RigVMModel/RigVMPin.cpp
--- a/Engine/Source/Developer/RigVMDeveloper/Private/RigVMModel/RigVMPin.cpp
+++ b/Engine/Source/Developer/RigVMDeveloper/Private/RigVMModel/RigVMPin.cpp
@@ -625,9 +625,15 @@ TRigVMTypeIndex URigVMPin::GetTypeIndex() const
 	if(LastKnownTypeIndex == INDEX_NONE)
 	{
 		LastKnownCPPType = GetCPPType();
-		const FRigVMTemplateArgumentType Type(*LastKnownCPPType, GetCPPTypeObject());
-		LastKnownTypeIndex = FRigVMRegistry::Get().FindOrAddType(Type);
-		ensure(LastKnownTypeIndex != INDEX_NONE);
+		// cpp type can be empty if it is an unsupported type such as a UObject type
+		if (!LastKnownCPPType.IsEmpty())
+		{
+			const FRigVMTemplateArgumentType Type(*LastKnownCPPType, GetCPPTypeObject());
+			LastKnownTypeIndex = FRigVMRegistry::Get().FindOrAddType(Type);
+			// in rare cases LastKnowTypeIndex can still be NONE here because
+			// we have nodes that has constant pin that references struct type like FRuntimeFloatCurve
+			// which contains a object ptr member and is thus not registered in the registry
+		}
 	}
 	return LastKnownTypeIndex;
 }
diff --git a/Engine/Source/Developer/RigVMDeveloper/Public/RigVMModel/Nodes/RigVMTemplateNode.h b/Engine/Source/Developer/RigVMDeveloper/Public/RigVMModel/Nodes/RigVMTemplateNode.h
--- a/Engine/Source/Developer/RigVMDeveloper/Public/RigVMModel/Nodes/RigVMTemplateNode.h
+++ b/Engine/Source/Developer/RigVMDeveloper/Public/RigVMModel/Nodes/RigVMTemplateNode.h
@@ -140,7 +140,7 @@ public:
 
 	// Tries to reduce the input types to a single type, if all are compatible
 	// Will prioritize the InPreferredType if available
-	bool TryReduceTypesToSingle(TArray<TRigVMTypeIndex>& InOutTypes, const TRigVMTypeIndex InPreferredType = TRigVMTypeIndex()) const;
+	bool TryReduceTypesToSingle(TArray<TRigVMTypeIndex>& InOutTypes, const TRigVMTypeIndex InPreferredType = INDEX_NONE) const;	
 
 	// returns true if updating pin filters with InTypes would result in different filters 
 	bool PinNeedsFilteredTypesUpdate(URigVMPin* InPin, const TArray<TRigVMTypeIndex>& InTypeIndices);
diff --git a/Engine/Source/Runtime/RigVM/Private/RigVMCore/RigVMMemoryStorage.cpp b/Engine/Source/Runtime/RigVM/Private/RigVMCore/RigVMMemoryStorage.cpp
--- a/Engine/Source/Runtime/RigVM/Private/RigVMCore/RigVMMemoryStorage.cpp
+++ b/Engine/Source/Runtime/RigVM/Private/RigVMCore/RigVMMemoryStorage.cpp
@@ -218,8 +218,8 @@ FRigVMPropertyDescription::FRigVMPropertyDescription(const FName& InName, const
 	// sanity check the CPPType matches the provided struct
 	if(const UScriptStruct* ScriptStruct = Cast<UScriptStruct>(CPPTypeObject))
 	{
-		checkf(TailCPPType == ScriptStruct->GetStructCPPName(), TEXT("CPPType '%s' doesn't match provided Struct '%s'"),
-			*TailCPPType, *ScriptStruct->GetStructCPPName());
+		checkf(TailCPPType == RigVMTypeUtils::GetUniqueStructTypeName(ScriptStruct), TEXT("CPPType '%s' doesn't match provided Struct '%s'"),
+			*TailCPPType, *RigVMTypeUtils::GetUniqueStructTypeName(ScriptStruct));
 	}
 }
 
@@ -785,7 +785,7 @@ FProperty* URigVMMemoryStorageGeneratorClass::AddProperty(URigVMMemoryStorageGen
 
 				for(UScriptStruct* BaseStructure : BaseStructures)
 				{
-					if(BaseStructure->GetStructCPPName() == BaseCPPType)
+					if(RigVMTypeUtils::GetUniqueStructTypeName(BaseStructure) == BaseCPPType)
 					{
 						FStructProperty* StructProperty = new FStructProperty(PropertyOwner, InProperty.Name, RF_Public);
 						StructProperty->Struct = BaseStructure;
diff --git a/Engine/Source/Runtime/RigVM/Private/RigVMCore/RigVMRegistry.cpp b/Engine/Source/Runtime/RigVM/Private/RigVMCore/RigVMRegistry.cpp
--- a/Engine/Source/Runtime/RigVM/Private/RigVMCore/RigVMRegistry.cpp
+++ b/Engine/Source/Runtime/RigVM/Private/RigVMCore/RigVMRegistry.cpp
@@ -4,9 +4,12 @@
 #include "RigVMCore/RigVMStruct.h"
 #include "RigVMTypeUtils.h"
 #include "RigVMModule.h"
+#include "Animation/AttributeTypes.h"
+#include "AssetRegistry/AssetRegistryModule.h"
+#include "Engine/UserDefinedEnum.h"
+#include "Engine/UserDefinedStruct.h"
 #include "UObject/UObjectIterator.h"
 #include "Misc/CoreDelegates.h"
-#include "Engine/UserDefinedStruct.h"
 
 FRigVMRegistry FRigVMRegistry::s_RigVMRegistry;
 const FName FRigVMRegistry::TemplateNameMetaName = TEXT("TemplateName");
@@ -124,11 +127,25 @@ void FRigVMRegistry::InitializeIfNeeded()
 	FCoreDelegates::OnExit.AddLambda([]()
 	{
 		s_RigVMRegistry.Reset();
+
+		if (FAssetRegistryModule* AssetRegistryModule = FModuleManager::GetModulePtr<FAssetRegistryModule>(TEXT("AssetRegistry")))
+		{
+			if (AssetRegistryModule->TryGet())
+			{
+				AssetRegistryModule->Get().OnAssetRemoved().RemoveAll(&s_RigVMRegistry);
+			}
+		}
 	});
+	
+	FAssetRegistryModule& AssetRegistryModule = FModuleManager::LoadModuleChecked<FAssetRegistryModule>(TEXT("AssetRegistry"));
+	AssetRegistryModule.Get().OnAssetRemoved().AddRaw(this, &FRigVMRegistry::OnAssetRemoved);
 }
 
 void FRigVMRegistry::Refresh()
 {
+	// refresh only tries to register all loaded types
+	// trying to load all types is very slow for large projects so we don't want to do that
+	
 	// add all user defined structs
 	for (TObjectIterator<UScriptStruct> ScriptIt; ScriptIt; ++ScriptIt)
 	{
@@ -138,7 +155,6 @@ void FRigVMRegistry::Refresh()
 			// if this is a C++ type - skip it
 			if(ScriptStruct->IsA<UUserDefinedStruct>() || ScriptStruct->IsChildOf(FRigVMExecuteContext::StaticStruct()))
 			{
-				const FString CPPType = ScriptStruct->GetStructCPPName();
 				FindOrAddType(FRigVMTemplateArgumentType(ScriptStruct));
 			}
 		}
@@ -175,6 +191,32 @@ void FRigVMRegistry::Refresh()
 	}
 }
 
+void FRigVMRegistry::OnAssetRemoved(const FAssetData& InAssetData)
+{
+	UObject* Asset = InAssetData.GetAsset();
+
+	if (UObjectRedirector* RedirectedStruct = Cast<UObjectRedirector>(Asset))
+	{
+		Asset = RedirectedStruct->DestinationObject;
+	}
+
+	bool bTypeRemoved = false;
+	
+	if (UScriptStruct* ScriptStruct = Cast<UScriptStruct>(Asset))
+	{
+		bTypeRemoved = RemoveType(ScriptStruct);
+	}
+	else if (UEnum* Enum = Cast<UEnum>(Asset))
+	{
+		bTypeRemoved = RemoveType(Enum);
+	}
+
+	if (bTypeRemoved)
+	{
+		OnRigVMRegistryChangedDelegate.Broadcast();
+	}
+}
+
 void FRigVMRegistry::Reset()
 {
 	for(FRigVMDispatchFactory* Factory : Factories)
@@ -460,6 +502,120 @@ void FRigVMRegistry::RegisterTypeInCategory(FRigVMTemplateArgument::ETypeCategor
 	}
 }
 
+bool FRigVMRegistry::RemoveType(const FRigVMTemplateArgumentType& InType)
+{
+	check(!InType.IsArray())
+	
+	TRigVMTypeIndex Index = GetTypeIndex(InType);
+	
+	if(Index == INDEX_NONE)
+	{
+		return false;
+	}
+
+	TArray<TRigVMTypeIndex> Indices;
+	Indices.Init(INDEX_NONE, 3);
+	Indices[0] = Index;
+	Indices[1] = GetArrayTypeFromBaseTypeIndex(Indices[0]);
+
+	// any type that can be removed should have 3 entries in the registry
+	if (ensure(Indices[1] != INDEX_NONE))
+	{
+		Indices[2] = GetArrayTypeFromBaseTypeIndex(Indices[1]);
+	}
+	
+	const UObject* CPPTypeObject = InType.CPPTypeObject;
+
+	for (int32 ArrayDimension=0; ArrayDimension<3; ++ArrayDimension)
+	{
+		Index = Indices[ArrayDimension];
+		
+		if (Index == INDEX_NONE)
+		{
+			break;
+		}
+		
+		if(const UUserDefinedEnum* Enum = Cast<UUserDefinedEnum>(CPPTypeObject))
+		{
+			switch(ArrayDimension)
+			{
+			default:
+			case 0:
+				{
+					RemoveTypeInCategory(FRigVMTemplateArgument::ETypeCategory_SingleEnumValue, Index);
+					RemoveTypeInCategory(FRigVMTemplateArgument::ETypeCategory_SingleAnyValue, Index);
+					break;
+				}
+			case 1:
+				{
+					RemoveTypeInCategory(FRigVMTemplateArgument::ETypeCategory_ArrayEnumValue, Index);
+					RemoveTypeInCategory(FRigVMTemplateArgument::ETypeCategory_ArrayAnyValue, Index);
+					break;
+				}
+			case 2:
+				{
+					RemoveTypeInCategory(FRigVMTemplateArgument::ETypeCategory_ArrayArrayEnumValue, Index);
+					RemoveTypeInCategory(FRigVMTemplateArgument::ETypeCategory_ArrayArrayAnyValue, Index);
+					break;
+				}
+			}
+		}
+		else if(const UUserDefinedStruct* Struct = Cast<UUserDefinedStruct>(CPPTypeObject))
+		{
+			switch(ArrayDimension)
+			{
+			default:
+			case 0:
+				{
+					RemoveTypeInCategory(FRigVMTemplateArgument::ETypeCategory_SingleScriptStructValue, Index);
+					RemoveTypeInCategory(FRigVMTemplateArgument::ETypeCategory_SingleAnyValue, Index);
+					break;
+				}
+			case 1:
+				{
+					RemoveTypeInCategory(FRigVMTemplateArgument::ETypeCategory_ArrayScriptStructValue, Index);
+					RemoveTypeInCategory(FRigVMTemplateArgument::ETypeCategory_ArrayAnyValue, Index);
+					break;
+				}
+			case 2:
+				{
+					RemoveTypeInCategory(FRigVMTemplateArgument::ETypeCategory_ArrayArrayScriptStructValue, Index);
+					RemoveTypeInCategory(FRigVMTemplateArgument::ETypeCategory_ArrayArrayAnyValue, Index);
+					break;
+				}
+			}
+		}
+
+		// remove the type from the registry entirely
+		TypeToIndex.Remove(GetType(Index));
+		Types[Index] = FTypeInfo();
+	}
+
+	return true;
+}
+
+void FRigVMRegistry::RemoveTypeInCategory(FRigVMTemplateArgument::ETypeCategory InCategory, TRigVMTypeIndex InTypeIndex)
+{
+	check(InCategory != FRigVMTemplateArgument::ETypeCategory_Invalid);
+
+	TypesPerCategory.FindChecked(InCategory).Remove(InTypeIndex);
+
+	const TArray<TPair<int32,int32>>& ArgumentsToUseType = ArgumentsPerCategory.FindChecked(InCategory);
+
+	TSet<int32> TemplatesToUseType;
+	
+	for(const TPair<int32,int32>& Pair : ArgumentsToUseType)
+	{
+		TemplatesToUseType.Add(Pair.Key);
+	}
+	
+	for (const int32 TemplateIndex : TemplatesToUseType)
+	{
+		FRigVMTemplate& Template = Templates[TemplateIndex];
+		Template.HandleTypeRemoval(InTypeIndex);
+	}
+}
+
 TRigVMTypeIndex FRigVMRegistry::GetTypeIndex(const FRigVMTemplateArgumentType& InType) const
 {
 	if(const TRigVMTypeIndex* Index = TypeToIndex.Find(InType))
@@ -725,6 +881,13 @@ bool FRigVMRegistry::IsAllowedType(const UStruct* InStruct)
 	{
 		return false;
 	}
+
+	// allow all user defined structs since they can always be changed to be compliant with RigVM restrictions
+	if (const UUserDefinedStruct* UserDefinedStruct = Cast<UUserDefinedStruct>(InStruct))
+	{
+		return true;
+	}
+
 	for (TFieldIterator<FProperty> It(InStruct); It; ++It)
 	{
 		if(!IsAllowedType(*It))
@@ -875,20 +1038,7 @@ const FRigVMFunction* FRigVMRegistry::FindFunction(const TCHAR* InName) const
 		FString FactoryName, ArgumentsString;
 		if(NameString.Split(TEXT("::"), &FactoryName, &ArgumentsString))
 		{
-			// if the factory has never been registered - we should try to look it up
-			if(FindDispatchFactory(*FactoryName) == nullptr)
-			{
-				if(FactoryName.StartsWith(FRigVMDispatchFactory::DispatchPrefix))
-				{
-					const FString ScriptStructName = FactoryName.Mid(FRigVMDispatchFactory::DispatchPrefix.Len());
-					if(UScriptStruct* FactoryStruct = FindFirstObject<UScriptStruct>(*ScriptStructName, EFindFirstObjectOptions::NativeFirst | EFindFirstObjectOptions::EnsureIfAmbiguous))
-					{
-						FRigVMRegistry* MutableThis = (FRigVMRegistry*)this;
-						MutableThis->RegisterFactory(FactoryStruct);
-					}
-				}
-			}
-
+			// if the factory has never been registered - FindDispatchFactory will try to look it up and register
 			if(const FRigVMDispatchFactory* Factory = FindDispatchFactory(*FactoryName))
 			{
 				if(const FRigVMTemplate* Template = Factory->GetTemplate())
@@ -937,23 +1087,20 @@ const FRigVMTemplate* FRigVMRegistry::FindTemplate(const FName& InNotation, bool
 		return &Templates[*TemplateIndexPtr];
 	}
 
-	if(Factories.Num() > 0)
+	static bool IsDispatchingFunction = false;
+	if(IsDispatchingFunction)
 	{
-		static bool IsDispatchingFunction = false;
-		if(IsDispatchingFunction)
-		{
-			return nullptr;
-		}
-		
-		const FString NotationString(InNotation.ToString());
-		FString FactoryName, ArgumentsString;
-		if(NotationString.Split(TEXT("("), &FactoryName, &ArgumentsString))
+		return nullptr;
+	}
+	
+	const FString NotationString(InNotation.ToString());
+	FString FactoryName, ArgumentsString;
+	if(NotationString.Split(TEXT("("), &FactoryName, &ArgumentsString))
+	{
+		if(const FRigVMDispatchFactory* Factory = FindDispatchFactory(*FactoryName))
 		{
-			if(const FRigVMDispatchFactory* Factory = FindDispatchFactory(*FactoryName))
-			{
-				TGuardValue<bool> ReEntryGuard(IsDispatchingFunction, true);
-				return Factory->GetTemplate();
-			}
+			TGuardValue<bool> ReEntryGuard(IsDispatchingFunction, true);
+			return Factory->GetTemplate();
 		}
 	}
 
@@ -1074,6 +1221,20 @@ FRigVMDispatchFactory* FRigVMRegistry::FindDispatchFactory(const FName& InFactor
 	{
 		return *FactoryPtr;
 	}
+
+	FString FactoryName = InFactoryName.ToString();
+	
+	// if the factory has never been registered - we should try to look it up	
+	if(FactoryName.StartsWith(FRigVMDispatchFactory::DispatchPrefix))
+	{
+		const FString ScriptStructName = FactoryName.Mid(FRigVMDispatchFactory::DispatchPrefix.Len());
+		if(UScriptStruct* FactoryStruct = FindFirstObject<UScriptStruct>(*ScriptStructName, EFindFirstObjectOptions::NativeFirst | EFindFirstObjectOptions::EnsureIfAmbiguous))
+		{
+			FRigVMRegistry* MutableThis = (FRigVMRegistry*)this;
+			return (FRigVMDispatchFactory*)MutableThis->RegisterFactory(FactoryStruct);
+		}
+	}	
+	
 	return nullptr;
 }
 
diff --git a/Engine/Source/Runtime/RigVM/Private/RigVMCore/RigVMTemplate.cpp b/Engine/Source/Runtime/RigVM/Private/RigVMCore/RigVMTemplate.cpp
--- a/Engine/Source/Runtime/RigVM/Private/RigVMCore/RigVMTemplate.cpp
+++ b/Engine/Source/Runtime/RigVM/Private/RigVMCore/RigVMTemplate.cpp
@@ -78,15 +78,27 @@ FRigVMTemplateArgument::FRigVMTemplateArgument(const FName& InName, ERigVMPinDir
 	UpdateTypeToPermutations();
 }
 
-FRigVMTemplateArgument::FRigVMTemplateArgument(const FName& InName, ERigVMPinDirection InDirection, const TArray<ETypeCategory>& InTypeCategories)
+FRigVMTemplateArgument::FRigVMTemplateArgument(const FName& InName, ERigVMPinDirection InDirection, const TArray<ETypeCategory>& InTypeCategories, const FTypeFilter& InTypeFilter)
 : Index(INDEX_NONE)
 , Name(InName)
 , Direction(InDirection)
 , TypeCategories(InTypeCategories)
 {
+	TSet<TRigVMTypeIndex> AllTypes;
 	for(ETypeCategory TypeCategory : TypeCategories)
 	{
-		TypeIndices.Append(FRigVMRegistry::Get().GetTypesForCategory(TypeCategory));
+		AllTypes.Append(FRigVMRegistry::Get().GetTypesForCategory(TypeCategory));
+	}
+	
+	TypeIndices.Reserve(AllTypes.Num());
+	for (const TRigVMTypeIndex& Type : AllTypes)
+	{
+		if (InTypeFilter.IsBound() && !InTypeFilter.Execute(Type))
+		{
+			continue;
+		}
+		
+		TypeIndices.Add(Type);
 	}
 
 	EnsureValidExecuteType();
@@ -285,6 +297,12 @@ FRigVMTemplateArgument::EArrayType FRigVMTemplateArgument::GetArrayType() const
 
 		for(int32 PermutationIndex=1;PermutationIndex<TypeIndices.Num();PermutationIndex++)
 		{
+			// INDEX_NONE indicates deleted permutation
+			if (TypeIndices[PermutationIndex] == INDEX_NONE)
+			{
+				continue;
+			}
+			
 			const EArrayType OtherArrayType = Registry.IsArrayType(TypeIndices[PermutationIndex]) ? EArrayType::EArrayType_ArrayValue : EArrayType::EArrayType_SingleValue;
 			if(OtherArrayType != ArrayType)
 			{
@@ -310,14 +328,22 @@ TArray<TRigVMTypeIndex> FRigVMTemplateArgument::GetSupportedTypeIndices(const TA
 	{
 		for (const TRigVMTypeIndex& TypeIndex : TypeIndices)
 		{
-			SupportedTypes.AddUnique(TypeIndex);
+			// INDEX_NONE indicates deleted permutation
+			if (TypeIndex != INDEX_NONE)
+			{
+				SupportedTypes.AddUnique(TypeIndex);
+			}
 		}
 	}
 	else
 	{
 		for(const int32 PermutationIndex : InPermutationIndices)
 		{
-			SupportedTypes.AddUnique(TypeIndices[PermutationIndex]);
+			// INDEX_NONE indicates deleted permutation
+			if (TypeIndices[PermutationIndex] != INDEX_NONE)
+			{
+				SupportedTypes.AddUnique(TypeIndices[PermutationIndex]);
+			}
 		}
 	}
 	return SupportedTypes;
@@ -332,8 +358,12 @@ TArray<FString> FRigVMTemplateArgument::GetSupportedTypeStrings(const TArray<int
 	{
 		for (const TRigVMTypeIndex& TypeIndex : TypeIndices)
 		{
-			const FString TypeString = FRigVMRegistry::Get().GetType(TypeIndex).CPPType.ToString();
-			SupportedTypes.AddUnique(TypeString);
+			// INDEX_NONE indicates deleted permutation
+			if (TypeIndex != INDEX_NONE)
+			{
+				const FString TypeString = FRigVMRegistry::Get().GetType(TypeIndex).CPPType.ToString();
+				SupportedTypes.AddUnique(TypeString);
+			}
 		}
 	}
 	else
@@ -341,8 +371,12 @@ TArray<FString> FRigVMTemplateArgument::GetSupportedTypeStrings(const TArray<int
 		for(const int32 PermutationIndex : InPermutationIndices)
 		{
 			const TRigVMTypeIndex TypeIndex = TypeIndices[PermutationIndex];
-			const FString TypeString = FRigVMRegistry::Get().GetType(TypeIndex).CPPType.ToString();
-			SupportedTypes.AddUnique(TypeString);
+			// INDEX_NONE indicates deleted permutation
+			if (TypeIndex != INDEX_NONE)
+			{
+				const FString TypeString = FRigVMRegistry::Get().GetType(TypeIndex).CPPType.ToString();
+				SupportedTypes.AddUnique(TypeString);
+			}
 		}
 	}
 	return SupportedTypes;
@@ -1146,14 +1180,20 @@ bool FRigVMTemplate::Resolve(FTypeMap& InOutTypes, TArray<int32>& OutPermutation
 			TRigVMTypeIndex MatchedType = *InputType;
 			bool bFoundMatch = false;
 			bool bFoundPerfectMatch = false;
+
+			// Using a map to collect all permutations that we can keep/remove
+			// instead of removing them one by one, which can be costly
+			TMap<int32, bool> PermutationsToKeep;
+			
 			for (int32 PermutationIndex = 0; PermutationIndex < Argument.TypeIndices.Num(); PermutationIndex++)
 			{
 				if(!Registry.CanMatchTypes(Argument.TypeIndices[PermutationIndex], *InputType, bAllowFloatingPointCasts))
 				{
-					OutPermutationIndices.Remove(PermutationIndex);
+					PermutationsToKeep.FindOrAdd(PermutationIndex) = false;
 				}
 				else
 				{
+					PermutationsToKeep.FindOrAdd(PermutationIndex) = true;
 					bFoundMatch = true;
 
 					// if the type matches - but it's not the exact same
@@ -1169,6 +1209,11 @@ bool FRigVMTemplate::Resolve(FTypeMap& InOutTypes, TArray<int32>& OutPermutation
 					}
 				}
 			}
+
+			OutPermutationIndices = OutPermutationIndices.FilterByPredicate([PermutationsToKeep](int32 PermutationIndex) -> bool
+			{
+				return PermutationsToKeep[PermutationIndex];
+			});
 			
 			if(bFoundMatch)
 			{
@@ -1477,3 +1522,39 @@ bool FRigVMTemplate::AddTypeForArgument(const FName& InArgumentName, TRigVMTypeI
 	}
 	return false;
 }
+
+void FRigVMTemplate::HandleTypeRemoval(TRigVMTypeIndex InTypeIndex)
+{
+	TArray<int32> PermutationsToRemove;
+	for (int32 PermutationIndex = 0; PermutationIndex < NumPermutations(); PermutationIndex++)
+	{
+		FRigVMTemplateTypeMap TypeMap = GetTypesForPermutation(PermutationIndex);
+
+		TArray<TRigVMTypeIndex> Types;
+		TypeMap.GenerateValueArray(Types);
+
+		if (Types.Contains(InTypeIndex))
+		{
+			PermutationsToRemove.Add(PermutationIndex);
+		}
+	}
+
+	for (FRigVMTemplateArgument& Argument : Arguments)
+	{
+		TArray<TRigVMTypeIndex> NewTypeIndices; 
+		for (int32 PermutationIndex = 0; PermutationIndex < Argument.TypeIndices.Num(); PermutationIndex++)
+		{
+			if (PermutationsToRemove.Contains(PermutationIndex))
+			{
+				// invalidate the type index for this permutation
+				Argument.TypeIndices[PermutationIndex] = INDEX_NONE;
+			}
+		}
+	}
+	
+	for (FRigVMTemplateArgument& Argument : Arguments)
+	{
+		Argument.TypeToPermutations.Remove(InTypeIndex);
+	}
+}
+
diff --git a/Engine/Source/Runtime/RigVM/Public/RigVMCore/RigVMExternalVariable.h b/Engine/Source/Runtime/RigVM/Public/RigVMCore/RigVMExternalVariable.h
--- a/Engine/Source/Runtime/RigVM/Public/RigVMCore/RigVMExternalVariable.h
+++ b/Engine/Source/Runtime/RigVM/Public/RigVMCore/RigVMExternalVariable.h
@@ -76,7 +76,7 @@ struct RIGVM_API FRigVMExternalVariable
 		}
 		else if (const FStructProperty* StructProperty = CastField<FStructProperty>(InProperty))
 		{
-			OutTypeName = *StructProperty->Struct->GetStructCPPName();
+			OutTypeName = *RigVMTypeUtils::GetUniqueStructTypeName(StructProperty->Struct);
 			OutTypeObject = StructProperty->Struct;
 		}
 		else if (const FObjectProperty* ObjectProperty = CastField<FObjectProperty>(InProperty))
@@ -432,7 +432,7 @@ struct RIGVM_API FRigVMExternalVariable
 	{
 		FRigVMExternalVariable Variable;
 		Variable.Name = InName;
-		Variable.TypeName = *TBaseStructure<T>::Get()->GetStructCPPName();
+		Variable.TypeName = *RigVMTypeUtils::GetUniqueStructTypeName(TBaseStructure<T>::Get());
 		Variable.TypeObject = TBaseStructure<T>::Get();
 		Variable.bIsArray = false;
 		Variable.Size = TBaseStructure<T>::Get()->GetStructureSize();
@@ -448,7 +448,7 @@ struct RIGVM_API FRigVMExternalVariable
 	{
 		FRigVMExternalVariable Variable;
 		Variable.Name = InName;
-		Variable.TypeName = *TBaseStructure<T>::Get()->GetStructCPPName();
+		Variable.TypeName = *RigVMTypeUtils::GetUniqueStructTypeName(TBaseStructure<T>::Get());
 		Variable.TypeObject = TBaseStructure<T>::Get();
 		Variable.bIsArray = true;
 		Variable.Size = TBaseStructure<T>::Get()->GetStructureSize();
@@ -464,7 +464,7 @@ struct RIGVM_API FRigVMExternalVariable
 	{
 		FRigVMExternalVariable Variable;
 		Variable.Name = InName;
-		Variable.TypeName = *T::StaticStruct()->GetStructCPPName();
+		Variable.TypeName = *RigVMTypeUtils::GetUniqueStructTypeName(T::StaticStruct());
 		Variable.TypeObject = T::StaticStruct();
 		Variable.bIsArray = false;
 		Variable.Size = T::StaticStruct()->GetStructureSize();
@@ -480,7 +480,7 @@ struct RIGVM_API FRigVMExternalVariable
 	{
 		FRigVMExternalVariable Variable;
 		Variable.Name = InName;
-		Variable.TypeName = *T::StaticStruct()->GetStructCPPName();
+		Variable.TypeName = *RigVMTypeUtils::GetUniqueStructTypeName(T::StaticStruct());
 		Variable.TypeObject = T::StaticStruct();
 		Variable.bIsArray = true;
 		Variable.Size = T::StaticStruct()->GetStructureSize();
diff --git a/Engine/Source/Runtime/RigVM/Public/RigVMCore/RigVMRegistry.h b/Engine/Source/Runtime/RigVM/Public/RigVMCore/RigVMRegistry.h
--- a/Engine/Source/Runtime/RigVM/Public/RigVMCore/RigVMRegistry.h
+++ b/Engine/Source/Runtime/RigVM/Public/RigVMCore/RigVMRegistry.h
@@ -37,6 +37,8 @@ struct RIGVM_API FRigVMRegistry
 {
 public:
 
+	DECLARE_MULTICAST_DELEGATE(FOnRigVMRegistryChanged);
+	
 	~FRigVMRegistry();
 	
 	// Returns the singleton registry
@@ -56,6 +58,9 @@ public:
 	// Refreshes the list and finds the function pointers
 	// based on the names.
 	void Refresh();
+	
+	// Update the registry when old types are removed
+    void OnAssetRemoved(const FAssetData& InAssetData);
 
 	// Clear the registry
 	void Reset();
@@ -67,6 +72,10 @@ public:
 		return FindOrAddType_Internal(InType, false);
 	}
 
+	// Removes a type from the registry, and updates all dependent templates
+	// which also creates invalid permutations in templates that we should ignore
+	bool RemoveType(const FRigVMTemplateArgumentType& InType);
+
 	// Returns the type index given a type
 	TRigVMTypeIndex GetTypeIndex(const FRigVMTemplateArgumentType& InType) const;
 
@@ -210,6 +219,9 @@ public:
 
 	static const TArray<UScriptStruct*>& GetMathTypes();
 
+	// Notifies other system that types have been added/removed, and template permutations have been updated
+	FOnRigVMRegistryChanged& OnRigVMRegistryChanged() { return OnRigVMRegistryChangedDelegate; }
+
 private:
 
 	static const FName TemplateNameMetaName;
@@ -256,6 +268,8 @@ private:
 	static bool IsAllowedType(const UClass* InClass);
 
 	void RegisterTypeInCategory(FRigVMTemplateArgument::ETypeCategory InCategory, TRigVMTypeIndex InTypeIndex);
+	
+	void RemoveTypeInCategory(FRigVMTemplateArgument::ETypeCategory InCategory, TRigVMTypeIndex InTypeIndex);
 
 	// memory for all (known) types
 	TArray<FTypeInfo> Types;
@@ -289,6 +303,9 @@ private:
 	// Lookup per type category to know which argument to keep in sync
 	TMap<FRigVMTemplateArgument::ETypeCategory, TArray<TPair<int32,int32>>> ArgumentsPerCategory;
 
+	// Notifies other system that types have been added/removed, and template permutations have been updated
+	FOnRigVMRegistryChanged OnRigVMRegistryChangedDelegate;
+	
 	static FRigVMRegistry s_RigVMRegistry;
 	friend struct FRigVMStruct;
 	friend struct FRigVMTemplate;
diff --git a/Engine/Source/Runtime/RigVM/Public/RigVMCore/RigVMTemplate.h b/Engine/Source/Runtime/RigVM/Public/RigVMCore/RigVMTemplate.h
--- a/Engine/Source/Runtime/RigVM/Public/RigVMCore/RigVMTemplate.h
+++ b/Engine/Source/Runtime/RigVM/Public/RigVMCore/RigVMTemplate.h
@@ -65,11 +65,16 @@ struct RIGVM_API FRigVMTemplateArgumentType
 		: CPPType(InCPPType)
 		, CPPTypeObject(InCPPTypeObject)
 	{
+		// InCppType is unreliable because not all caller knows that
+		// we use generated unique names for user defined structs
+		// so here we override the CppType name with the actual name used in the registry
+		CPPType = *RigVMTypeUtils::PostProcessCPPType(CPPType.ToString(), CPPTypeObject);
+		
 		check(!CPPType.IsNone());
 	}
 
 	FRigVMTemplateArgumentType(UScriptStruct* InScriptStruct)
-	: CPPType(*InScriptStruct->GetStructCPPName())
+	: CPPType(*RigVMTypeUtils::GetUniqueStructTypeName(InScriptStruct))
 	, CPPTypeObject(InScriptStruct)
 	{
 	}
@@ -149,6 +154,8 @@ struct RIGVM_API FRigVMTemplateArgumentType
  */
 struct RIGVM_API FRigVMTemplateArgument
 {
+	DECLARE_DELEGATE_RetVal_OneParam(bool, FTypeFilter, const TRigVMTypeIndex&);
+	
 	enum EArrayType
 	{
 		EArrayType_SingleValue,
@@ -187,7 +194,7 @@ struct RIGVM_API FRigVMTemplateArgument
 
 	FRigVMTemplateArgument(const FName& InName, ERigVMPinDirection InDirection, TRigVMTypeIndex InType);
 	FRigVMTemplateArgument(const FName& InName, ERigVMPinDirection InDirection, const TArray<TRigVMTypeIndex>& InTypeIndices);
-	FRigVMTemplateArgument(const FName& InName, ERigVMPinDirection InDirection, const TArray<ETypeCategory>& InTypeCategories);
+	FRigVMTemplateArgument(const FName& InName, ERigVMPinDirection InDirection, const TArray<ETypeCategory>& InTypeCategories, const FTypeFilter& InTypeFilter = {});
 
 	// Serialize
 	void Serialize(FArchive& Ar);
@@ -364,6 +371,9 @@ public:
 	// template factory how to deal with the new type.
 	bool AddTypeForArgument(const FName& InArgumentName, TRigVMTypeIndex InTypeIndex);
 
+	// Invalidates template permutations whenever a type such as a user defined struct is removed
+	void HandleTypeRemoval(TRigVMTypeIndex InTypeIndex);
+	
 	// Returns the delegate to be able to react to type changes dynamically
 	FRigVMTemplate_NewArgumentTypeDelegate& OnNewArgumentType() { return Delegates.NewArgumentTypeDelegate; }
 
diff --git a/Engine/Source/Runtime/RigVM/Public/RigVMCore/RigVMTypeIndex.h b/Engine/Source/Runtime/RigVM/Public/RigVMCore/RigVMTypeIndex.h
--- a/Engine/Source/Runtime/RigVM/Public/RigVMCore/RigVMTypeIndex.h
+++ b/Engine/Source/Runtime/RigVM/Public/RigVMCore/RigVMTypeIndex.h
@@ -9,12 +9,16 @@
 class RIGVM_API TRigVMTypeIndex
 {
 public:
-	
+
+	// default constructor provided just to avoid compilation error
+	// in your code if you need to represent an invalid type, use INDEX_NONE instead of TRigVMTypeIndex()
+	// the reason being that because we typedef int32 to TRigVMTypeIndex,
+	// default constructor of TRigVMTypeIndex should be the same as int32, which defaults to 0 and is not an invalid type
 	TRigVMTypeIndex()
-		: Name(NAME_None)
-		, Index(INDEX_NONE)
+	: Name(NAME_None)
+	, Index(0)
 	{}
-
+	
 	TRigVMTypeIndex(int32 InIndex)
 		: Name(NAME_None)
 		, Index(InIndex)
diff --git a/Engine/Source/Runtime/RigVM/Public/RigVMTypeUtils.h b/Engine/Source/Runtime/RigVM/Public/RigVMTypeUtils.h
--- a/Engine/Source/Runtime/RigVM/Public/RigVMTypeUtils.h
+++ b/Engine/Source/Runtime/RigVM/Public/RigVMTypeUtils.h
@@ -6,6 +6,7 @@
 #include "RigVMCore/RigVMTypeIndex.h"
 #include "RigVMCore/RigVMUnknownType.h"
 #include "UObject/Interface.h"
+#include "Engine/UserDefinedStruct.h"
 
 namespace RigVMTypeUtils
 {
@@ -84,6 +85,16 @@ namespace RigVMTypeUtils
 		return InCPPType.RightChop(7).LeftChop(1).TrimStartAndEnd();
 	}
 
+	FORCEINLINE_DEBUGGABLE FString GetUniqueStructTypeName(const UScriptStruct* InScriptStruct)
+	{
+		if (const UUserDefinedStruct* UserDefinedStruct = Cast<UUserDefinedStruct>(InScriptStruct))
+		{
+			return FString::Printf(TEXT("FUserDefinedStruct_%s"), *UserDefinedStruct->GetCustomGuid().ToString());
+		}
+
+		return InScriptStruct->GetStructCPPName();
+	}
+
 	FORCEINLINE FString CPPTypeFromEnum(UEnum* InEnum)
 	{
 		check(InEnum);
@@ -114,7 +125,7 @@ namespace RigVMTypeUtils
 
 	static const FString& GetWildCardCPPType()
 	{
-		static const FString WildCardCPPType = FRigVMUnknownType::StaticStruct()->GetStructCPPName(); 
+		static const FString WildCardCPPType = GetUniqueStructTypeName(FRigVMUnknownType::StaticStruct()); 
 		return WildCardCPPType;
 	}
 
@@ -136,6 +147,8 @@ namespace RigVMTypeUtils
 		return WildCardArrayCPPTypeName;
 	}
 
+
+
 	FORCEINLINE FString PostProcessCPPType(const FString& InCPPType, UObject* InCPPTypeObject)
 	{
 		FString CPPType = InCPPType;
@@ -153,7 +166,7 @@ namespace RigVMTypeUtils
 		}
 		else if (const UScriptStruct* ScriptStruct = Cast<UScriptStruct>(InCPPTypeObject))
 		{
-			CPPType = ScriptStruct->GetStructCPPName();
+			CPPType = GetUniqueStructTypeName(ScriptStruct);
 		}
 		else if (UEnum* Enum = Cast<UEnum>(InCPPTypeObject))
 		{
