ControlRig: Fixed a compile error

diff --git a/Engine/Plugins/Animation/ControlRig/Source/ControlRig/Private/Units/Animation/RigUnit_AnimAttribute.h b/Engine/Plugins/Animation/ControlRig/Source/ControlRig/Private/Units/Animation/RigUnit_AnimAttribute.h
--- a/Engine/Plugins/Animation/ControlRig/Source/ControlRig/Private/Units/Animation/RigUnit_AnimAttribute.h
+++ b/Engine/Plugins/Animation/ControlRig/Source/ControlRig/Private/Units/Animation/RigUnit_AnimAttribute.h
@@ -76,7 +76,7 @@ namespace RigUnit_AnimAttribute
 		const USkeletalMeshComponent* OwningComponent = Cast<USkeletalMeshComponent>(Context.OwningComponent);
 
 		if (!OwningComponent ||
-			!OwningComponent->GetSkeletalMesh())
+			!OwningComponent->GetSkeletalMeshAsset())
 		{
 			return nullptr;
 		}
@@ -91,7 +91,7 @@ namespace RigUnit_AnimAttribute
 			// Invalidate cache if input changed
 			if (CachedBoneName != BoneName)
 			{
-				CachedBoneIndex = OwningComponent->GetSkeletalMesh()->GetRefSkeleton().FindBoneIndex(BoneName);
+				CachedBoneIndex = OwningComponent->GetSkeletalMeshAsset()->GetRefSkeleton().FindBoneIndex(BoneName);
 			}
 		}
 	
