Constraints: contextual menu update

diff --git a/Engine/Plugins/Animation/ControlRig/Source/ControlRigEditor/Private/Editor/Constraints/SConstraintsWidget.cpp b/Engine/Plugins/Animation/ControlRig/Source/ControlRigEditor/Private/Editor/Constraints/SConstraintsWidget.cpp
--- a/Engine/Plugins/Animation/ControlRig/Source/ControlRigEditor/Private/Editor/Constraints/SConstraintsWidget.cpp
+++ b/Engine/Plugins/Animation/ControlRig/Source/ControlRigEditor/Private/Editor/Constraints/SConstraintsWidget.cpp
@@ -788,6 +788,14 @@ void SConstraintsEditionWidget::RemoveItem(const TSharedPtr<FEditableConstraintI
 void SConstraintsEditionWidget::InvalidateConstraintList()
 {
 	bNeedsRefresh = true;
+
+	// NOTE: this is a hack to disable the picker if the selection changes. The picker should trigger the escape key
+	// but this is an editor wide change
+	static const FActorPickerModeModule& ActorPickerMode = FModuleManager::Get().GetModuleChecked<FActorPickerModeModule>("ActorPickerMode");
+	if (ActorPickerMode.IsInActorPickingMode())
+	{
+		ActorPickerMode.EndActorPickingMode();
+	}
 }
 
 int32 SConstraintsEditionWidget::RefreshConstraintList()
@@ -837,17 +845,6 @@ void SConstraintsEditionWidget::OnActorSelectionChanged(const TArray<UObject*>&
 	// NOTE we use this delegate to trigger an tree update, however, control actors are not selected as they are 
 	// Temporary Editor Actors so NewSelection won't contain the controls
 	InvalidateConstraintList();
-
-	// NOTE: this is a hack to disable the picker if the selection changes. The picker should trigger the escape key
-	// but this is an editor wide change 
-	if (NewSelection.IsEmpty())
-	{
-		static const FActorPickerModeModule& ActorPickerMode = FModuleManager::Get().GetModuleChecked<FActorPickerModeModule>("ActorPickerMode");
-		if (ActorPickerMode.IsInActorPickingMode())
-		{
-			ActorPickerMode.EndActorPickingMode();
-		}
-	}
 }
 
 TSharedPtr<SWidget> SConstraintsEditionWidget::CreateContextMenu()
@@ -871,7 +868,7 @@ TSharedPtr<SWidget> SConstraintsEditionWidget::CreateContextMenu()
 	{
 		return SNullWidget::NullWidget;
 	}
-	
+
 	static constexpr bool CloseAfterSelection = true;
 	FMenuBuilder MenuBuilder(CloseAfterSelection, nullptr);
 
@@ -903,72 +900,59 @@ TSharedPtr<SWidget> SConstraintsEditionWidget::CreateContextMenu()
 	}
 	MenuBuilder.EndSection();
 
-	// baking (note that this will probably be moved)
-	MenuBuilder.BeginSection("BakeConstraint", LOCTEXT("BakeConstraintHeader", "Bake Constraint"));
+	// transform constraint options
+	if (UTickableTransformConstraint* TransformConstraint = Cast<UTickableTransformConstraint>(Constraint))
 	{
-		MenuBuilder.AddMenuEntry(
-		LOCTEXT("BakeConstraintLabel", "Bake"),
-		FText::Format(LOCTEXT("BakeConstraintDoItTooltip", "Bake {0} transforms."), FText::FromName(ListItems[Index]->Name)),
-		FSlateIcon(),
-		FUIAction(FExecuteAction::CreateLambda([Constraint]()
+		const FText ConstraintLabel = FText::FromName(ListItems[Index]->Name);
+	
+		// baking (note that this will probably be moved)
+		MenuBuilder.BeginSection("BakeConstraint", LOCTEXT("BakeConstraintHeader", "Bake Constraint"));
 		{
-			if (UTickableTransformConstraint* TransformConstraint = Cast<UTickableTransformConstraint>(Constraint))
+			MenuBuilder.AddMenuEntry(
+			LOCTEXT("BakeConstraintLabel", "Bake"),
+			FText::Format(LOCTEXT("BakeConstraintDoItTooltip", "Bake {0} transforms."), ConstraintLabel),
+			FSlateIcon(),
+			FUIAction(FExecuteAction::CreateLambda([TransformConstraint]()
 			{
 				FConstraintBaker::DoIt(TransformConstraint);
-			}
-		})),
-		NAME_None,
-		EUserInterfaceActionType::Button);
-	}
-	MenuBuilder.EndSection();
+			})),
+			NAME_None,
+			EUserInterfaceActionType::Button);
+		}
+		MenuBuilder.EndSection();
 
-	// test section
-	MenuBuilder.BeginSection("TestConstraint", LOCTEXT("TestConstraintHeader", "Test"));
-	{
-		MenuBuilder.AddMenuEntry(
-		LOCTEXT("AddConstraintKey", "Add Key"),
-		FText::Format(LOCTEXT("AddKeyTooltip", "Add active key for {0}."), FText::FromName(ListItems[Index]->Name)),
-		FSlateIcon(),
-		FUIAction(FExecuteAction::CreateLambda([Constraint]()
+		// keys section
+		FCanExecuteAction IsCompensationEnabled = FCanExecuteAction::CreateLambda([TransformConstraint]()
 		{
-			if (UTickableTransformConstraint* TransformConstraint = Cast<UTickableTransformConstraint>(Constraint))
-			{
-				FScopedTransaction Transaction(LOCTEXT("CreateConstraintKey", "Create Constraint Key"));
-				FConstraintChannelHelper::SmartConstraintKey(TransformConstraint);
-			}
-		})),
-		NAME_None,
-		
-		EUserInterfaceActionType::Button);
-		MenuBuilder.AddMenuEntry(
-		LOCTEXT("CompensateKeyLabel", "Compensate Key"),
-		FText::Format(LOCTEXT("CompensateKeyTooltip", "Compensate transform key for {0}."), FText::FromName(ListItems[Index]->Name)),
-		FSlateIcon(),
-		FUIAction(FExecuteAction::CreateLambda([Constraint]()
+			return TransformConstraint->bDynamicOffset;
+		});
+	
+		MenuBuilder.BeginSection("KeyConstraint", LOCTEXT("KeyConstraintHeader", "Keys"));
 		{
-			if (UTickableTransformConstraint* TransformConstraint = Cast<UTickableTransformConstraint>(Constraint))
+			MenuBuilder.AddMenuEntry(
+			LOCTEXT("CompensateKeyLabel", "Compensate Key"),
+			FText::Format(LOCTEXT("CompensateKeyTooltip", "Compensate transform key for {0}."), ConstraintLabel),
+			FSlateIcon(),
+			FUIAction(FExecuteAction::CreateLambda([TransformConstraint]()
 			{
 				FConstraintChannelHelper::Compensate(TransformConstraint);
-			}
-		})),
-		NAME_None,
-		EUserInterfaceActionType::Button);
-
-		MenuBuilder.AddMenuEntry(
-		LOCTEXT("CompensateAllKeysLabel", "Compensate All Keys"),
-		FText::Format(LOCTEXT("CompensateAllKeysTooltip", "Compensate all transform keys for {0}."), FText::FromName(ListItems[Index]->Name)),
-		FSlateIcon(),
-		FUIAction(FExecuteAction::CreateLambda([Constraint]()
-		{
-			if (UTickableTransformConstraint* TransformConstraint = Cast<UTickableTransformConstraint>(Constraint))
+			}), IsCompensationEnabled),
+			NAME_None,
+			EUserInterfaceActionType::Button);
+
+			MenuBuilder.AddMenuEntry(
+			LOCTEXT("CompensateAllKeysLabel", "Compensate All Keys"),
+			FText::Format(LOCTEXT("CompensateAllKeysTooltip", "Compensate all transform keys for {0}."), ConstraintLabel),
+			FSlateIcon(),
+			FUIAction(FExecuteAction::CreateLambda([TransformConstraint]()
 			{
 				FConstraintChannelHelper::Compensate(TransformConstraint, true);
-			}
-		})),
-		NAME_None,
-		EUserInterfaceActionType::Button);
+			}), IsCompensationEnabled),
+			NAME_None,
+			EUserInterfaceActionType::Button);
+		}
+		MenuBuilder.EndSection();
 	}
-	MenuBuilder.EndSection();
 
 	return MenuBuilder.MakeWidget();
 }
