Control Rig: Fix control rig nodes in some AnimBPs having multiple alpha pins

diff --git a/Engine/Plugins/Animation/ControlRig/Source/ControlRigDeveloper/Private/AnimGraphNode_ControlRig.cpp b/Engine/Plugins/Animation/ControlRig/Source/ControlRigDeveloper/Private/AnimGraphNode_ControlRig.cpp
--- a/Engine/Plugins/Animation/ControlRig/Source/ControlRigDeveloper/Private/AnimGraphNode_ControlRig.cpp
+++ b/Engine/Plugins/Animation/ControlRig/Source/ControlRigDeveloper/Private/AnimGraphNode_ControlRig.cpp
@@ -80,12 +80,16 @@ void UAnimGraphNode_ControlRig::CreateCustomPins(TArray<UEdGraphPin*>* OldPins)
 						bool bFound = false;
 						for (UEdGraphPin* CurrentPin : Pins)
 						{
-							if (CurrentPin->GetFName() == OldPin->GetFName() &&
-								CurrentPin->PinType == OldPin->PinType)
+							if (CurrentPin->GetFName() == OldPin->GetFName())
 							{
-								bFound = true;
-								break;
-							}
+								if (CurrentPin->PinType == OldPin->PinType ||
+									AnimGraphDefaultSchema->ArePinTypesCompatible(CurrentPin->PinType, OldPin->PinType))
+								{
+									bFound = true;
+									break;
+								}
+								
+							}							
 						}
 
 						if (!bFound)
