Control Rig: Fix reset transforms

Operate on control rigs that have selected controls.

Reset all will reset all controls on the control rigs that have selected controls.
Reset selected will reset just the selected controls.

Move transaction outside of the control rigs loop so that the reset operation gets bundled up into one for all the selected rigs.

diff --git a/Engine/Plugins/Animation/ControlRig/Source/ControlRigEditor/Private/EditMode/ControlRigEditMode.cpp b/Engine/Plugins/Animation/ControlRig/Source/ControlRigEditor/Private/EditMode/ControlRigEditMode.cpp
--- a/Engine/Plugins/Animation/ControlRig/Source/ControlRigEditor/Private/EditMode/ControlRigEditMode.cpp
+++ b/Engine/Plugins/Animation/ControlRig/Source/ControlRigEditor/Private/EditMode/ControlRigEditMode.cpp
@@ -2784,141 +2784,155 @@ void FControlRigEditMode::ToggleAllManipulators()
 
 void FControlRigEditMode::ResetTransforms(bool bSelectionOnly)
 {
+	// Gather up the control rigs for the selected controls
+	TArray<UControlRig*> ControlRigs;
 	for (TWeakObjectPtr<UControlRig>& RuntimeRigPtr : RuntimeControlRigs)
 	{
 		if (UControlRig* ControlRig = RuntimeRigPtr.Get())
 		{
+			if (ControlRig->CurrentControlSelection().Num() > 0)
+			{
+				ControlRigs.Add(ControlRig);
+			}
+		}
+	}
+	if (ControlRigs.Num() == 0)
+	{
+		return;
+	}
 
-			TArray<FRigElementKey> SelectedRigElements = GetSelectedRigElements(ControlRig);
-			TArray<FRigElementKey> ControlsToReset = SelectedRigElements;
-			TArray<FRigElementKey> ControlsInteracting = SelectedRigElements;
-			if (!bSelectionOnly)
+	FScopedTransaction Transaction(LOCTEXT("HierarchyResetTransforms", "Reset Transforms"));
+
+	for (UControlRig* ControlRig : ControlRigs)
+	{
+		TArray<FRigElementKey> SelectedRigElements = GetSelectedRigElements(ControlRig);
+		TArray<FRigElementKey> ControlsToReset = SelectedRigElements;
+		TArray<FRigElementKey> ControlsInteracting = SelectedRigElements;
+		if (!bSelectionOnly)
+		{
+			TArray<FRigControlElement*> Controls;
+			ControlRig->GetControlsInOrder(Controls);
+			ControlsToReset.SetNum(0);
+			ControlsInteracting.SetNum(0);
+			for (const FRigControlElement* Control : Controls)
 			{
-				TArray<FRigControlElement*> Controls;
-				ControlRig->GetControlsInOrder(Controls);
-				ControlsToReset.SetNum(0);
-				ControlsInteracting.SetNum(0);
-				for (const FRigControlElement* Control : Controls)
+				ControlsToReset.Add(Control->GetKey());
+				if(Control->Settings.AnimationType == ERigControlAnimationType::AnimationControl ||
+					Control->IsAnimationChannel())
 				{
-					ControlsToReset.Add(Control->GetKey());
-					if(Control->Settings.AnimationType == ERigControlAnimationType::AnimationControl ||
-						Control->IsAnimationChannel())
-					{
-						ControlsInteracting.Add(Control->GetKey());
-					}
+					ControlsInteracting.Add(Control->GetKey());
 				}
 			}
-			bool bHasNonDefaultParent = false;
-			TArray<FRigElementKey> Parents;
-			for (const FRigElementKey& ControlKey : ControlsToReset)
+		}
+		bool bHasNonDefaultParent = false;
+		TArray<FRigElementKey> Parents;
+		for (const FRigElementKey& ControlKey : ControlsToReset)
+		{
+			FRigElementKey SpaceKey = ControlRig->GetHierarchy()->GetActiveParent(ControlKey);
+			Parents.Add(SpaceKey);
+			if (SpaceKey != ControlRig->GetHierarchy()->GetDefaultParentKey())
 			{
-				FRigElementKey SpaceKey = ControlRig->GetHierarchy()->GetActiveParent(ControlKey);
-				Parents.Add(SpaceKey);
-				if (SpaceKey != ControlRig->GetHierarchy()->GetDefaultParentKey())
-				{
-					bHasNonDefaultParent = true;
-				}
+				bHasNonDefaultParent = true;
 			}
+		}
 
-			FScopedTransaction Transaction(LOCTEXT("HierarchyResetTransforms", "Reset Transforms"));
-			FControlRigInteractionScope InteractionScope(ControlRig, ControlsInteracting);
+		FControlRigInteractionScope InteractionScope(ControlRig, ControlsInteracting);
 
-			for (const FRigElementKey& ControlToReset : ControlsToReset)
+		for (const FRigElementKey& ControlToReset : ControlsToReset)
+		{
+			if (ControlToReset.Type == ERigElementType::Control)
 			{
-				if (ControlToReset.Type == ERigElementType::Control)
+				FRigControlElement* ControlElement = ControlRig->FindControl(ControlToReset.Name);
+				if (ControlElement && !ControlElement->Settings.bIsTransientControl)
 				{
-					FRigControlElement* ControlElement = ControlRig->FindControl(ControlToReset.Name);
-					if (ControlElement && !ControlElement->Settings.bIsTransientControl)
+					const FTransform InitialLocalTransform = ControlRig->GetHierarchy()->GetInitialLocalTransform(ControlToReset);
+					ControlRig->Modify();
+					if (bHasNonDefaultParent == true) //possibly not at default parent so switch to it
 					{
-						const FTransform InitialLocalTransform = ControlRig->GetHierarchy()->GetInitialLocalTransform(ControlToReset);
-						ControlRig->Modify();
-						if (bHasNonDefaultParent == true) //possibly not at default parent so switch to it
-						{
-							ControlRig->GetHierarchy()->SwitchToDefaultParent(ControlElement->GetKey());
-						}
-						ControlRig->SetControlLocalTransform(ControlToReset.Name, InitialLocalTransform, true);
-						NotifyDrivenControls(ControlRig, ControlToReset);
-						if (bHasNonDefaultParent == false)
-						{
-							ControlRig->ControlModified().Broadcast(ControlRig, ControlElement, EControlRigSetKey::DoNotCare);
-						}
+						ControlRig->GetHierarchy()->SwitchToDefaultParent(ControlElement->GetKey());
+					}
+					ControlRig->SetControlLocalTransform(ControlToReset.Name, InitialLocalTransform, true);
+					NotifyDrivenControls(ControlRig, ControlToReset);
+					if (bHasNonDefaultParent == false)
+					{
+						ControlRig->ControlModified().Broadcast(ControlRig, ControlElement, EControlRigSetKey::DoNotCare);
+					}
 
-						//@helge not sure what to do if the non-default parent
-						if (UControlRigBlueprint* Blueprint = Cast<UControlRigBlueprint>(ControlRig->GetClass()->ClassGeneratedBy))
-						{
-							Blueprint->Hierarchy->SetLocalTransform(ControlToReset, InitialLocalTransform);
-						}
+					//@helge not sure what to do if the non-default parent
+					if (UControlRigBlueprint* Blueprint = Cast<UControlRigBlueprint>(ControlRig->GetClass()->ClassGeneratedBy))
+					{
+						Blueprint->Hierarchy->SetLocalTransform(ControlToReset, InitialLocalTransform);
 					}
 				}
 			}
+		}
 
-			if (bHasNonDefaultParent == true) //now we have the initial pose setup we need to get the global transforms as specified now then set them in the current parent space
-			{
-				ControlRig->Evaluate_AnyThread();
+		if (bHasNonDefaultParent == true) //now we have the initial pose setup we need to get the global transforms as specified now then set them in the current parent space
+		{
+			ControlRig->Evaluate_AnyThread();
 
-				//get global transforms
-				TArray<FTransform> GlobalTransforms;
-				for (const FRigElementKey& ControlToReset : ControlsToReset)
+			//get global transforms
+			TArray<FTransform> GlobalTransforms;
+			for (const FRigElementKey& ControlToReset : ControlsToReset)
+			{
+				FRigControlElement* ControlElement = ControlRig->FindControl(ControlToReset.Name);
+				if (ControlElement && !ControlElement->Settings.bIsTransientControl)
 				{
-					FRigControlElement* ControlElement = ControlRig->FindControl(ControlToReset.Name);
-					if (ControlElement && !ControlElement->Settings.bIsTransientControl)
-					{
-						FTransform GlobalTransform = ControlRig->GetHierarchy()->GetGlobalTransform(ControlToReset);
-						GlobalTransforms.Add(GlobalTransform);
-					}
-					NotifyDrivenControls(ControlRig, ControlToReset);
+					FTransform GlobalTransform = ControlRig->GetHierarchy()->GetGlobalTransform(ControlToReset);
+					GlobalTransforms.Add(GlobalTransform);
 				}
-				//switch back to original parent space
-				int32 Index = 0;
-				for (const FRigElementKey& ControlToReset : ControlsToReset)
+				NotifyDrivenControls(ControlRig, ControlToReset);
+			}
+			//switch back to original parent space
+			int32 Index = 0;
+			for (const FRigElementKey& ControlToReset : ControlsToReset)
+			{
+				FRigControlElement* ControlElement = ControlRig->FindControl(ControlToReset.Name);
+				if (ControlElement && !ControlElement->Settings.bIsTransientControl)
 				{
-					FRigControlElement* ControlElement = ControlRig->FindControl(ControlToReset.Name);
-					if (ControlElement && !ControlElement->Settings.bIsTransientControl)
-					{
-						ControlRig->GetHierarchy()->SwitchToParent(ControlToReset, Parents[Index]);
-						++Index;
-					}
+					ControlRig->GetHierarchy()->SwitchToParent(ControlToReset, Parents[Index]);
+					++Index;
 				}
-				//set global transforms in this space // do it twice since ControlsInOrder is not really always in order
-				for (int32 SetHack = 0; SetHack < 2; ++SetHack)
-				{
-					ControlRig->Evaluate_AnyThread();
-					Index = 0;
-					for (const FRigElementKey& ControlToReset : ControlsToReset)
-
-					{
-						FRigControlElement* ControlElement = ControlRig->FindControl(ControlToReset.Name);
-						if (ControlElement && !ControlElement->Settings.bIsTransientControl)
-						{
-							ControlRig->SetControlGlobalTransform(ControlToReset.Name, GlobalTransforms[Index], true);
-							ControlRig->Evaluate_AnyThread();
-							NotifyDrivenControls(ControlRig, ControlToReset);
-							++Index;
-						}
-					}
-				}
-				//send notifies
-
+			}
+			//set global transforms in this space // do it twice since ControlsInOrder is not really always in order
+			for (int32 SetHack = 0; SetHack < 2; ++SetHack)
+			{
+				ControlRig->Evaluate_AnyThread();
+				Index = 0;
 				for (const FRigElementKey& ControlToReset : ControlsToReset)
+
 				{
 					FRigControlElement* ControlElement = ControlRig->FindControl(ControlToReset.Name);
 					if (ControlElement && !ControlElement->Settings.bIsTransientControl)
 					{
-						ControlRig->ControlModified().Broadcast(ControlRig, ControlElement, EControlRigSetKey::DoNotCare);
+						ControlRig->SetControlGlobalTransform(ControlToReset.Name, GlobalTransforms[Index], true);
+						ControlRig->Evaluate_AnyThread();
+						NotifyDrivenControls(ControlRig, ControlToReset);
+						++Index;
 					}
 				}
 			}
-			else
+			//send notifies
+
+			for (const FRigElementKey& ControlToReset : ControlsToReset)
 			{
-				// we at least have to run the interaction event
-				TGuardValue<TArray<FName>> EventGuard(ControlRig->EventQueue, {FRigUnit_InteractionExecution::EventName});
-				ControlRig->Evaluate_AnyThread();
-				for (const FRigElementKey& ControlToReset : ControlsToReset)
+				FRigControlElement* ControlElement = ControlRig->FindControl(ControlToReset.Name);
+				if (ControlElement && !ControlElement->Settings.bIsTransientControl)
 				{
-					NotifyDrivenControls(ControlRig, ControlToReset);
+					ControlRig->ControlModified().Broadcast(ControlRig, ControlElement, EControlRigSetKey::DoNotCare);
 				}
 			}
 		}
+		else
+		{
+			// we at least have to run the interaction event
+			TGuardValue<TArray<FName>> EventGuard(ControlRig->EventQueue, {FRigUnit_InteractionExecution::EventName});
+			ControlRig->Evaluate_AnyThread();
+			for (const FRigElementKey& ControlToReset : ControlsToReset)
+			{
+				NotifyDrivenControls(ControlRig, ControlToReset);
+			}
+		}
 	}
 }
 
