Control Rig: Add missing unbound

diff --git a/Engine/Plugins/Animation/ControlRig/Source/ControlRigEditor/Private/Sequencer/ControlRigParameterTrackEditor.cpp b/Engine/Plugins/Animation/ControlRig/Source/ControlRigEditor/Private/Sequencer/ControlRigParameterTrackEditor.cpp
--- a/Engine/Plugins/Animation/ControlRig/Source/ControlRigEditor/Private/Sequencer/ControlRigParameterTrackEditor.cpp
+++ b/Engine/Plugins/Animation/ControlRig/Source/ControlRigEditor/Private/Sequencer/ControlRigParameterTrackEditor.cpp
@@ -480,6 +480,11 @@ void FControlRigParameterTrackEditor::UnbindAllControlRigs()
 			ControlRig->ControlModified().RemoveAll(this);
 			ControlRig->OnInitialized_AnyThread().RemoveAll(this);
 			ControlRig->ControlSelected().RemoveAll(this);
+			if (const TSharedPtr<IControlRigObjectBinding> Binding = ControlRig->GetObjectBinding())
+			{
+				Binding->OnControlRigBind().RemoveAll(this);
+			}
+			ControlRig->ControlRigBound().RemoveAll(this);
 		}
 	}
 	BoundControlRigs.SetNum(0);
