Cosntraints: Fix Baking so that A) All constraint keys are turned off B) -1 keys are removed C) it's undoable. Also added some Modifies to the manager when adding/removing to help with  undo/redo issues.

diff --git a/Engine/Plugins/Animation/ControlRig/Source/ControlRigEditor/Private/EditMode/SControlRigDetails.cpp b/Engine/Plugins/Animation/ControlRig/Source/ControlRigEditor/Private/EditMode/SControlRigDetails.cpp
--- a/Engine/Plugins/Animation/ControlRig/Source/ControlRigEditor/Private/EditMode/SControlRigDetails.cpp
+++ b/Engine/Plugins/Animation/ControlRig/Source/ControlRigEditor/Private/EditMode/SControlRigDetails.cpp
@@ -426,6 +426,10 @@ void SControlRigDetails::UpdateProxies()
 				const TArray<UControlRigControlsProxy*>& Proxies = ControlProxy->GetSelectedProxies();
 				for (UControlRigControlsProxy* Proxy : Proxies)
 				{
+					if (Proxy == nullptr)
+					{
+						continue;
+					}
 					if (Proxy->GetClass() == UControlRigTransformControlProxy::StaticClass())
 					{
 						if (Proxy->bIsIndividual)
diff --git a/Engine/Plugins/Animation/ControlRig/Source/ControlRigEditor/Private/Editor/Constraints/SConstraintsWidget.cpp b/Engine/Plugins/Animation/ControlRig/Source/ControlRigEditor/Private/Editor/Constraints/SConstraintsWidget.cpp
--- a/Engine/Plugins/Animation/ControlRig/Source/ControlRigEditor/Private/Editor/Constraints/SConstraintsWidget.cpp
+++ b/Engine/Plugins/Animation/ControlRig/Source/ControlRigEditor/Private/Editor/Constraints/SConstraintsWidget.cpp
@@ -22,6 +22,9 @@
 #include "DetailsViewArgs.h"
 #include "Tools/ConstraintBaker.h"
 #include "ScopedTransaction.h"
+#include "ISequencer.h"
+#include "Tools/BakingHelper.h"
+#include "MovieSceneToolHelpers.h"
 
 #define LOCTEXT_NAMESPACE "SConstraintsWidget"
 
@@ -912,10 +915,37 @@ TSharedPtr<SWidget> SConstraintsEditionWidget::CreateContextMenu()
 			LOCTEXT("BakeConstraintLabel", "Bake"),
 			FText::Format(LOCTEXT("BakeConstraintDoItTooltip", "Bake {0} transforms."), ConstraintLabel),
 			FSlateIcon(),
-			FUIAction(FExecuteAction::CreateLambda([TransformConstraint]()
-			{
-				FConstraintBaker::DoIt(TransformConstraint);
-			})),
+			FUIAction(FExecuteAction::CreateLambda([Constraint, World]()
+				{
+					if (UTickableTransformConstraint* TransformConstraint = Cast<UTickableTransformConstraint>(Constraint))
+					{
+						const TWeakPtr<ISequencer> WeakSequencer = FBakingHelper::GetSequencer();
+						if (!WeakSequencer.IsValid() || !WeakSequencer.Pin()->GetFocusedMovieSceneSequence())
+						{
+							return;
+						}
+						const TSharedPtr<ISequencer> Sequencer = WeakSequencer.Pin();
+						const UMovieScene* MovieScene = Sequencer->GetFocusedMovieSceneSequence()->GetMovieScene();
+						if (!MovieScene)
+						{
+							return;
+						}
+
+						// compute frames
+						const FFrameNumber StartFrame = MovieScene->GetPlaybackRange().GetLowerBoundValue();
+						const FFrameNumber EndFrame = MovieScene->GetPlaybackRange().GetUpperBoundValue();
+						TArray<FFrameNumber> Frames;
+						MovieSceneToolHelpers::CalculateFramesBetween(MovieScene, StartFrame, EndFrame, Frames);
+
+						if (Frames.IsEmpty())
+						{
+							return;
+						}
+						FScopedTransaction Transaction(LOCTEXT("BakeConstraint", "Bake Constraint"));
+						FConstraintBaker::Bake(World, TransformConstraint, WeakSequencer.Pin(), Frames);
+
+					}
+				})),
 			NAME_None,
 			EUserInterfaceActionType::Button);
 		}
diff --git a/Engine/Plugins/Animation/ControlRig/Source/ControlRigEditor/Private/Tools/ConstraintBaker.cpp b/Engine/Plugins/Animation/ControlRig/Source/ControlRigEditor/Private/Tools/ConstraintBaker.cpp
--- a/Engine/Plugins/Animation/ControlRig/Source/ControlRigEditor/Private/Tools/ConstraintBaker.cpp
+++ b/Engine/Plugins/Animation/ControlRig/Source/ControlRigEditor/Private/Tools/ConstraintBaker.cpp
@@ -14,8 +14,12 @@
 #include "TransformableHandle.h"
 #include "Channels/MovieSceneChannelProxy.h"
 #include "Constraints/ControlRigTransformableHandle.h"
-#include "Tools/BakingHelper.h"
+#include "MovieSceneConstraintChannelHelper.h"
 #include "Sections/MovieScene3DTransformSection.h"
+#include "Sequencer/MovieSceneControlRigParameterSection.h"
+#include "ConstraintChannelHelper.h"
+#include "ConstraintChannel.h"
+#include "ConstraintChannelHelper.inl"
 
 #define LOCTEXT_NAMESPACE "ConstraintBaker"
 
@@ -64,7 +68,6 @@ namespace
 		{
 			return;
 		}
-		TransformSection->Modify();
 		const UMovieScene* MovieScene = InSequencer->GetFocusedMovieSceneSequence()->GetMovieScene();
 		InComponentHandle->AddTransformKeys(InFrames, InTransforms, InChannels, MovieScene->GetTickResolution(), TransformSection, true);
 
@@ -90,50 +93,6 @@ namespace
 
 	}
 }
-void FConstraintBaker::DoIt(UTickableTransformConstraint* InConstraint)
-{
-	const TWeakPtr<ISequencer> WeakSequencer = FBakingHelper::GetSequencer();
-	if (!WeakSequencer.IsValid() || !WeakSequencer.Pin()->GetFocusedMovieSceneSequence())
-	{
-		return;
-	}
-	const TSharedPtr<ISequencer> Sequencer = WeakSequencer.Pin();
-	const UMovieScene* MovieScene = Sequencer->GetFocusedMovieSceneSequence()->GetMovieScene();
-	if (!MovieScene)
-	{
-		return;
-	}
-
-	// compute frames
-	const FFrameNumber StartFrame = MovieScene->GetPlaybackRange().GetLowerBoundValue();
-	const FFrameNumber EndFrame = MovieScene->GetPlaybackRange().GetUpperBoundValue();
-	TArray<FFrameNumber> Frames;
-	MovieSceneToolHelpers::CalculateFramesBetween(MovieScene, StartFrame, EndFrame, Frames);
-
-	if (Frames.IsEmpty())
-	{
-		return;
-	}
-
-	// compute transforms
-	TArray<FTransform> Transforms;
-	GetHandleTransforms(Sequencer, InConstraint->ChildTRSHandle, {InConstraint}, Frames, true, Transforms);
-	
-	if (Frames.Num() != Transforms.Num())
-	{
-		return;
-	}
-
-	// bake to channel curves
-	const EMovieSceneTransformChannel Channels = InConstraint->GetChannelsToKey();
-	AddTransformKeys(Sequencer, InConstraint->ChildTRSHandle, Frames, Transforms, Channels);
-
-	// disable constraint
-	InConstraint->SetActive(false);
-	
-	// notify
-	Sequencer->NotifyMovieSceneDataChanged(EMovieSceneDataChangeType::MovieSceneStructureItemsChanged);
-}
 
 void FConstraintBaker::GetHandleTransforms(
 	UWorld* InWorld,
@@ -219,6 +178,101 @@ void FConstraintBaker::AddTransformKeys(
 	}
 }
 
+
+void FConstraintBaker::Bake(UWorld* InWorld, 
+	UTickableTransformConstraint* InConstraint,
+	const TSharedPtr<ISequencer>& InSequencer, 
+	const TArray<FFrameNumber>& InFrames)
+{
+	if (InFrames.Num() == 0 || InConstraint == nullptr || InConstraint->ChildTRSHandle == nullptr)
+	{
+		return;
+	}
+	
+	// compute transforms
+	TArray<FTransform> Transforms;
+	GetHandleTransforms(InSequencer, InConstraint->ChildTRSHandle, { InConstraint }, InFrames, true, Transforms);
+
+	if (InFrames.Num() != Transforms.Num())
+	{
+		return;
+	}
+	//get the section to be used later to delete the extra transform keys at the frame -1 times, abort if not there for some reason
+	UMovieSceneSection* Section = nullptr;
+	if (const UTransformableControlHandle* ControlHandle = Cast<UTransformableControlHandle>(InConstraint->ChildTRSHandle))
+	{
+		Section = FConstraintChannelHelper::GetControlSection(ControlHandle, InSequencer);
+	}
+	else if (const UTransformableComponentHandle* ComponentHandle = Cast<UTransformableComponentHandle>(InConstraint->ChildTRSHandle))
+	{ 
+		//todo move to function also used by SmartConstraintKey
+		AActor* Actor = ComponentHandle->Component->GetOwner();
+		if (!Actor)
+		{
+			return;
+		}
+
+		const FTransform LocalTransform = ComponentHandle->GetLocalTransform();
+		const FGuid Guid = InSequencer->GetHandleToObject(Actor, true);
+		if (!Guid.IsValid())
+		{
+			return;
+		}
+		Section = MovieSceneToolHelpers::GetTransformSection(InSequencer.Get(), Guid, LocalTransform);
+	}
+	if (Section == nullptr)
+	{
+		return;
+	}
+	IMovieSceneConstrainedSection* ConstrainedSection = Cast<IMovieSceneConstrainedSection>(Section);
+	if (ConstrainedSection == nullptr)
+	{
+		return;
+	}
+	FConstraintAndActiveChannel* ActiveChannel = ConstrainedSection->GetConstraintChannel(InConstraint->GetFName());
+	if (ActiveChannel == nullptr)
+	{
+		return;
+	}
+
+	Section->Modify();
+
+	// disable constraint and delete extra transform keys
+	TMovieSceneChannelData<bool> ConstraintChannelData = ActiveChannel->ActiveChannel.GetData();
+	TArrayView<const FFrameNumber> ConstraintFrames = ConstraintChannelData.GetTimes();
+
+	// get transform channels
+	const TArrayView<FMovieSceneFloatChannel*> FloatTransformChannels = InConstraint->ChildTRSHandle->GetFloatChannels(Section);
+	const TArrayView<FMovieSceneDoubleChannel*> DoubleTransformChannels = InConstraint->ChildTRSHandle->GetDoubleChannels(Section);
+
+	for (int32 Index = 0; Index < ConstraintFrames.Num(); ++Index)
+	{
+		const FFrameNumber Frame = ConstraintFrames[Index];
+		//todo we need to add a key at the begin/end if there is no frame there
+		if (Frame >= InFrames[0] && Frame <= InFrames[InFrames.Num() - 1])
+		{
+			ConstraintChannelData.UpdateOrAddKey(Frame, false);
+			//delete minus one frames
+			const FFrameNumber FrameMinusOne = Frame - 1;
+
+			if (FloatTransformChannels.Num() > 0)
+			{
+				FMovieSceneConstraintChannelHelper::DeleteTransformKeys(FloatTransformChannels, FrameMinusOne);
+			}
+			else if (DoubleTransformChannels.Num() > 0)
+			{
+				FMovieSceneConstraintChannelHelper::DeleteTransformKeys(DoubleTransformChannels, FrameMinusOne);
+			}
+		}
+	}
+	// now bake to channel curves
+	const EMovieSceneTransformChannel Channels = InConstraint->GetChannelsToKey();
+	AddTransformKeys(InSequencer, InConstraint->ChildTRSHandle, InFrames, Transforms, Channels);
+
+	// notify
+	InSequencer->RequestEvaluate();
+}
+
 #undef LOCTEXT_NAMESPACE
 
 
diff --git a/Engine/Plugins/Animation/ControlRig/Source/ControlRigEditor/Public/Tools/ConstraintBaker.h b/Engine/Plugins/Animation/ControlRig/Source/ControlRigEditor/Public/Tools/ConstraintBaker.h
--- a/Engine/Plugins/Animation/ControlRig/Source/ControlRigEditor/Public/Tools/ConstraintBaker.h
+++ b/Engine/Plugins/Animation/ControlRig/Source/ControlRigEditor/Public/Tools/ConstraintBaker.h
@@ -8,13 +8,15 @@
 class UTransformableHandle;
 class UTickableTransformConstraint;
 class ISequencer;
+class UWorld;
+struct FFrameNumber;
 enum class EMovieSceneTransformChannel : uint32;
 
 struct FConstraintBaker
 {
 public:
-	/** @todo documentation. (subject to changes to handle start / stop as parameters) */
-	static void DoIt(UTickableTransformConstraint* InConstraint);
+	/** Bake constraint over specified frames, frames must be in order*/
+	static void Bake(UWorld* InWorld, UTickableTransformConstraint* InConstraint, const TSharedPtr<ISequencer>& InSequencer, const TArray<FFrameNumber>& InFrames);
 
 	/** Stores InHandle local (or global) transforms at InFrames. */
 	static void GetHandleTransforms(
diff --git a/Engine/Source/Editor/MovieSceneTools/Private/Channels/ConstraintChannelEditor.cpp b/Engine/Source/Editor/MovieSceneTools/Private/Channels/ConstraintChannelEditor.cpp
--- a/Engine/Source/Editor/MovieSceneTools/Private/Channels/ConstraintChannelEditor.cpp
+++ b/Engine/Source/Editor/MovieSceneTools/Private/Channels/ConstraintChannelEditor.cpp
@@ -52,6 +52,7 @@ namespace
 	IndexRange GetRange(const TArrayView<const bool>& Values, const int32 Offset) 
 	{
 		IndexRange Range;
+		Range = Range.Empty();
 			
 		int32 FirstActive = INDEX_NONE;
 		for (int32 Index = Offset; Index < Values.Num(); ++Index)
diff --git a/Engine/Source/Runtime/Experimental/Animation/Constraints/Private/ConstraintsManager.cpp b/Engine/Source/Runtime/Experimental/Animation/Constraints/Private/ConstraintsManager.cpp
--- a/Engine/Source/Runtime/Experimental/Animation/Constraints/Private/ConstraintsManager.cpp
+++ b/Engine/Source/Runtime/Experimental/Animation/Constraints/Private/ConstraintsManager.cpp
@@ -377,7 +377,8 @@ bool FConstraintsManagerController::AddConstraint(UTickableConstraint* InConstra
 		return false;
 	}
 
-	// TODO handle transaction
+	Manager->Modify();
+
 	Manager->Constraints.Emplace(InConstraint);
 
 	InConstraint->ConstraintTick.RegisterFunction(InConstraint->GetFunction());
@@ -426,6 +427,8 @@ bool FConstraintsManagerController::RemoveConstraint(const int32 InConstraintInd
 		return false;
 	}
 
+	Manager->Modify();
+
 	const FName ConstraintName = Manager->Constraints[InConstraintIndex]->GetFName();
 	UTickableConstraint* Constraint = Manager->Constraints[InConstraintIndex];
 	
@@ -433,7 +436,6 @@ bool FConstraintsManagerController::RemoveConstraint(const int32 InConstraintInd
 	ConstraintRemoved.Broadcast(ConstraintName);
 	Manager->OnConstraintRemoved_BP.Broadcast(Manager, Constraint);
 
-	// TODO handle transaction
 	Manager->Constraints[InConstraintIndex]->ConstraintTick.UnRegisterTickFunction();
 	Manager->Constraints.RemoveAt(InConstraintIndex);
 
