ControlRig: Fix spawn new animation channel creating globally unique display names. It should be name scoped by the parent key

diff --git a/Engine/Plugins/Animation/ControlRig/Source/ControlRig/Private/Units/Execution/RigUnit_DynamicHierarchy.cpp b/Engine/Plugins/Animation/ControlRig/Source/ControlRig/Private/Units/Execution/RigUnit_DynamicHierarchy.cpp
--- a/Engine/Plugins/Animation/ControlRig/Source/ControlRig/Private/Units/Execution/RigUnit_DynamicHierarchy.cpp
+++ b/Engine/Plugins/Animation/ControlRig/Source/ControlRig/Private/Units/Execution/RigUnit_DynamicHierarchy.cpp
@@ -691,6 +691,7 @@ FRigUnit_HierarchyAddAnimationChannelBool_Execute()
 		ControlSettings.SetupLimitArrayForType(true, true, true);
 		ControlSettings.MinimumValue = FRigControlValue::Make<bool>(MinimumValue);
 		ControlSettings.MaximumValue = FRigControlValue::Make<bool>(MaximumValue);
+		ControlSettings.DisplayName = Controller->GetHierarchy()->GetSafeNewDisplayName(Parent, Name.ToString());
 		const FRigControlValue Value = FRigControlValue::Make<bool>(InitialValue);
 
 		FRigHierarchyControllerInstructionBracket InstructionBracket(Controller, RigVMExecuteContext.GetInstructionIndex());
@@ -729,6 +730,7 @@ FRigUnit_HierarchyAddAnimationChannelFloat_Execute()
 		ControlSettings.SetupLimitArrayForType(true, true, true);
 		ControlSettings.MinimumValue = FRigControlValue::Make<float>(MinimumValue);
 		ControlSettings.MaximumValue = FRigControlValue::Make<float>(MaximumValue);
+		ControlSettings.DisplayName = Controller->GetHierarchy()->GetSafeNewDisplayName(Parent, Name.ToString());
 		const FRigControlValue Value = FRigControlValue::Make<float>(InitialValue);
 		
 		FRigHierarchyControllerInstructionBracket InstructionBracket(Controller, RigVMExecuteContext.GetInstructionIndex());
@@ -767,6 +769,7 @@ FRigUnit_HierarchyAddAnimationChannelInteger_Execute()
 		ControlSettings.SetupLimitArrayForType(true, true, true);
 		ControlSettings.MinimumValue = FRigControlValue::Make<int32>(MinimumValue);
 		ControlSettings.MaximumValue = FRigControlValue::Make<int32>(MaximumValue);
+		ControlSettings.DisplayName = Controller->GetHierarchy()->GetSafeNewDisplayName(Parent, Name.ToString());
 		const FRigControlValue Value = FRigControlValue::Make<int32>(InitialValue);
 		
 		FRigHierarchyControllerInstructionBracket InstructionBracket(Controller, RigVMExecuteContext.GetInstructionIndex());
@@ -805,6 +808,7 @@ FRigUnit_HierarchyAddAnimationChannelVector2D_Execute()
 		ControlSettings.SetupLimitArrayForType(true, true, true);
 		ControlSettings.MinimumValue = FRigControlValue::Make<FVector3f>(FVector3f(MinimumValue.X, MinimumValue.Y, 0.f));
 		ControlSettings.MaximumValue = FRigControlValue::Make<FVector3f>(FVector3f(MaximumValue.X, MaximumValue.Y, 0.f));
+		ControlSettings.DisplayName = Controller->GetHierarchy()->GetSafeNewDisplayName(Parent, Name.ToString());
 		const FRigControlValue Value = FRigControlValue::Make<FVector3f>(FVector3f(InitialValue.X, InitialValue.Y, 0.f));
 		
 		FRigHierarchyControllerInstructionBracket InstructionBracket(Controller, RigVMExecuteContext.GetInstructionIndex());
@@ -843,6 +847,7 @@ FRigUnit_HierarchyAddAnimationChannelVector_Execute()
 		ControlSettings.SetupLimitArrayForType(true, true, true);
 		ControlSettings.MinimumValue = FRigControlValue::Make<FVector3f>(FVector3f(MinimumValue));
 		ControlSettings.MaximumValue = FRigControlValue::Make<FVector3f>(FVector3f(MaximumValue));
+		ControlSettings.DisplayName = Controller->GetHierarchy()->GetSafeNewDisplayName(Parent, Name.ToString());
 		const FRigControlValue Value = FRigControlValue::Make<FVector3f>(FVector3f(InitialValue));
 		
 		FRigHierarchyControllerInstructionBracket InstructionBracket(Controller, RigVMExecuteContext.GetInstructionIndex());
@@ -881,6 +886,7 @@ FRigUnit_HierarchyAddAnimationChannelRotator_Execute()
 		ControlSettings.SetupLimitArrayForType(true, true, true);
 		ControlSettings.MinimumValue = FRigControlValue::Make<FVector3f>(FVector3f(MinimumValue.Euler()));
 		ControlSettings.MaximumValue = FRigControlValue::Make<FVector3f>(FVector3f(MaximumValue.Euler()));
+		ControlSettings.DisplayName = Controller->GetHierarchy()->GetSafeNewDisplayName(Parent, Name.ToString());
 		const FRigControlValue Value = FRigControlValue::Make<FVector3f>(FVector3f(InitialValue.Euler()));
 		
 		FRigHierarchyControllerInstructionBracket InstructionBracket(Controller, RigVMExecuteContext.GetInstructionIndex());
