set control's name as actor's label

diff --git a/Engine/Plugins/Animation/ControlRig/Source/ControlRig/Private/ControlRigGizmoActor.cpp b/Engine/Plugins/Animation/ControlRig/Source/ControlRig/Private/ControlRigGizmoActor.cpp
--- a/Engine/Plugins/Animation/ControlRig/Source/ControlRig/Private/ControlRigGizmoActor.cpp
+++ b/Engine/Plugins/Animation/ControlRig/Source/ControlRig/Private/ControlRigGizmoActor.cpp
@@ -235,6 +235,9 @@ namespace FControlRigShapeHelper
 			ShapeActor->ShapeName = CreationParam.ShapeName;
 			ShapeActor->SetSelectable(CreationParam.bSelectable);
 			ShapeActor->SetActorTransform(CreationParam.SpawnTransform);
+#if WITH_EDITOR
+			ShapeActor->SetActorLabel(CreationParam.ControlName.ToString(), false);
+#endif // WITH_EDITOR
 
 			UStaticMeshComponent* MeshComponent = ShapeActor->StaticMeshComponent;
 
