redirect focus on parent graph when the node is an aggregate

diff --git a/Engine/Plugins/Animation/ControlRig/Source/ControlRigEditor/Private/Editor/SControlRigStackView.cpp b/Engine/Plugins/Animation/ControlRig/Source/ControlRigEditor/Private/Editor/SControlRigStackView.cpp
--- a/Engine/Plugins/Animation/ControlRig/Source/ControlRigEditor/Private/Editor/SControlRigStackView.cpp
+++ b/Engine/Plugins/Animation/ControlRig/Source/ControlRigEditor/Private/Editor/SControlRigStackView.cpp
@@ -14,6 +14,7 @@
 #include "Framework/MultiBox/MultiBoxBuilder.h"
 #include "Styling/AppStyle.h"
 #include "DetailLayoutBuilder.h"
+#include "RigVMModel/Nodes/RigVMAggregateNode.h"
 #include "Widgets/Input/SSearchBox.h"
 
 #define LOCTEXT_NAMESPACE "SControlRigStackView"
@@ -853,14 +854,23 @@ void SControlRigStackView::HandleFocusOnSelectedGraphNode()
 		}
 
 		const FRigVMByteCode& ByteCode = ControlRig->GetVM()->GetByteCode();
-		URigVMNode* SelectedNode = Cast<URigVMNode>(ByteCode.GetSubjectForInstruction(SelectedItems[0]->InstructionIndex));
-		if (SelectedNode)
+		UObject* Subject = ByteCode.GetSubjectForInstruction(SelectedItems[0]->InstructionIndex);
+		if (URigVMNode* SelectedNode = Cast<URigVMNode>(Subject))
 		{
-			if (UEdGraph* EdGraph = ControlRigBlueprint->GetEdGraph(SelectedNode->GetGraph()))
+			URigVMGraph* GraphToFocus = SelectedNode->GetGraph();
+			if (GraphToFocus && GraphToFocus->GetTypedOuter<URigVMAggregateNode>())
+			{
+				if(URigVMGraph* ParentGraph = GraphToFocus->GetParentGraph())
+				{
+					GraphToFocus = ParentGraph;
+				} 
+			}
+			
+			if (UEdGraph* EdGraph = ControlRigBlueprint->GetEdGraph(GraphToFocus))
 			{
 				ControlRigEditor.Pin()->OpenGraphAndBringToFront(EdGraph, true);
 				ControlRigEditor.Pin()->ZoomToSelection_Clicked();
-				ControlRigEditor.Pin()->HandleModifiedEvent(ERigVMGraphNotifType::NodeSelected, SelectedNode->GetGraph(), SelectedNode);
+				ControlRigEditor.Pin()->HandleModifiedEvent(ERigVMGraphNotifType::NodeSelected, GraphToFocus, SelectedNode);
 			}
 		}
 	}
@@ -990,22 +1000,34 @@ void SControlRigStackView::HandlePreviewControlRigUpdated(FControlRigEditor* InE
 
 void SControlRigStackView::HandleItemMouseDoubleClick(TSharedPtr<FRigStackEntry> InItem)
 {
-	if(ControlRigEditor.IsValid() && ControlRigBlueprint.IsValid())
+	if (!ControlRigEditor.IsValid() || !ControlRigBlueprint.IsValid())
+	{
+		return;
+	}
+	
+	UControlRig* ControlRig = Cast<UControlRig>(ControlRigBlueprint->GetObjectBeingDebugged());
+	if (!ControlRig || !ControlRig->GetVM())
+	{
+		return;
+	}
+
+	const FRigVMByteCode& ByteCode = ControlRig->GetVM()->GetByteCode();
+	if (URigVMNode* Subject = Cast<URigVMNode>(ByteCode.GetSubjectForInstruction(InItem->InstructionIndex)))
 	{
-		if(UControlRig* ControlRig = Cast<UControlRig>(ControlRigBlueprint->GetObjectBeingDebugged()))
+		URigVMGraph* GraphToFocus = Subject->GetGraph();
+		if (GraphToFocus && GraphToFocus->GetTypedOuter<URigVMAggregateNode>())
 		{
-			if(URigVM* VM = ControlRig->GetVM())
+			if(URigVMGraph* ParentGraph = GraphToFocus->GetParentGraph())
 			{
-				if(URigVMNode* Subject = Cast<URigVMNode>(VM->GetByteCode().GetSubjectForInstruction(InItem->InstructionIndex)))
-				{
-					if(UControlRigGraph* EdGraph = Cast<UControlRigGraph>(ControlRigBlueprint->GetEdGraph(Subject->GetGraph())))
-					{
-						if(UEdGraphNode* Node = EdGraph->FindNodeForModelNodeName(Subject->GetFName()))
-						{
-							ControlRigEditor.Pin()->JumpToHyperlink(Node, false);
-						}
-					}
-				}
+				GraphToFocus = ParentGraph;
+			} 
+		}
+		
+		if(UControlRigGraph* EdGraph = Cast<UControlRigGraph>(ControlRigBlueprint->GetEdGraph(GraphToFocus)))
+		{
+			if(const UEdGraphNode* Node = EdGraph->FindNodeForModelNodeName(Subject->GetFName()))
+			{
+				ControlRigEditor.Pin()->JumpToHyperlink(Node, false);
 			}
 		}
 	}
