Removed redundant private paths from controlrig.

diff --git a/Engine/Plugins/Animation/ControlRig/Source/ControlRigEditor/ControlRigEditor.Build.cs b/Engine/Plugins/Animation/ControlRig/Source/ControlRigEditor/ControlRigEditor.Build.cs
--- a/Engine/Plugins/Animation/ControlRig/Source/ControlRigEditor/ControlRigEditor.Build.cs
+++ b/Engine/Plugins/Animation/ControlRig/Source/ControlRigEditor/ControlRigEditor.Build.cs
@@ -9,13 +9,8 @@ namespace UnrealBuildTool.Rules
             PrivateIncludePaths.AddRange(
                 new string[] {
                     "ControlRigEditor/Private",
-                    "ControlRigEditor/Private/Sequencer",
-                    "ControlRigEditor/Private/EditMode",
-                    "ControlRigEditor/Private/Graph",
-                    "ControlRigEditor/Private/Editor",
                     "../../../../Source/Editor/UnrealEd/Private", 
 					"../../../../Source/Developer/MessageLog/Private", //compatibility for FBX importer
-					"ControlRigEditor/Private/Tools"
 				}
 			);
 
diff --git a/Engine/Plugins/Animation/ControlRig/Source/ControlRigEditor/Private/ControlRigEditorModule.cpp b/Engine/Plugins/Animation/ControlRig/Source/ControlRigEditor/Private/ControlRigEditorModule.cpp
--- a/Engine/Plugins/Animation/ControlRig/Source/ControlRigEditor/Private/ControlRigEditorModule.cpp
+++ b/Engine/Plugins/Animation/ControlRig/Source/ControlRigEditor/Private/ControlRigEditorModule.cpp
@@ -7,8 +7,8 @@
 #include "PropertyEditorModule.h"
 #include "ControlRig.h"
 #include "ControlRigComponent.h"
-#include "ControlRigConnectionDrawingPolicy.h"
 #include "ControlRigVariableDetailsCustomization.h"
+#include "Graph/ControlRigConnectionDrawingPolicy.h"
 #include "GraphEditorActions.h"
 #include "Kismet2/KismetEditorUtilities.h"
 #include "ISequencerModule.h"
@@ -23,19 +23,19 @@
 #include "Modules/ModuleManager.h"
 #include "Widgets/Docking/SDockTab.h"
 #include "EditorModeRegistry.h"
-#include "ControlRigEditMode.h"
+#include "EditMode/ControlRigEditMode.h"
 #include "ILevelSequenceModule.h"
 #include "EditorModeManager.h"
-#include "ControlRigEditMode.h"
+#include "EditMode/ControlRigEditMode.h"
 #include "Sequencer/MovieSceneControlRigParameterSection.h"
-#include "MovieSceneControlRigSectionDetailsCustomization.h"
-#include "ControlRigEditModeCommands.h"
+#include "Sequencer/MovieSceneControlRigSectionDetailsCustomization.h"
+#include "EditMode/ControlRigEditModeCommands.h"
 #include "Materials/Material.h"
 #include "Framework/MultiBox/MultiBoxExtender.h"
 #include "Framework/MultiBox/MultiBoxBuilder.h"
 #include "ToolMenus.h"
 #include "ContentBrowserMenuContexts.h"
-#include "ControlRigEditModeSettings.h"
+#include "EditMode/ControlRigEditModeSettings.h"
 #include "Components/SkeletalMeshComponent.h"
 #include "ContentBrowserModule.h"
 #include "AssetRegistry/AssetRegistryModule.h"
@@ -61,14 +61,14 @@
 #include "Kismet2/KismetDebugUtilities.h"
 #include "Graph/ControlRigGraphNode.h"
 #include "EdGraphUtilities.h"
-#include "ControlRigGraphPanelNodeFactory.h"
-#include "ControlRigGraphPanelPinFactory.h"
+#include "Graph/ControlRigGraphPanelNodeFactory.h"
+#include "Graph/ControlRigGraphPanelPinFactory.h"
 #include "ControlRigBlueprintUtils.h"
 #include "ControlRigBlueprintCommands.h"
 #include "ControlRigHierarchyCommands.h"
 #include "ControlRigStackCommands.h"
 #include "Animation/AnimSequence.h"
-#include "ControlRigEditorEditMode.h"
+#include "Editor/ControlRigEditorEditMode.h"
 #include "ControlRigElementDetails.h"
 #include "ControlRigCompilerDetails.h"
 #include "ControlRigDrawingDetails.h"
@@ -78,7 +78,7 @@
 #include "ControlRigInfluenceMapDetails.h"
 #include "Animation/AnimSequence.h"
 #include "Subsystems/AssetEditorSubsystem.h"
-#include "ControlRigParameterTrackEditor.h"
+#include "Sequencer/ControlRigParameterTrackEditor.h"
 #include "ActorFactories/ActorFactorySkeletalMesh.h"
 #include "ControlRigThumbnailRenderer.h"
 #include "RigVMModel/RigVMController.h"
@@ -113,7 +113,7 @@
 #include "ControlRig/Private/Units/Execution/RigUnit_InverseExecution.h"
 #include "Graph/SControlRigGraphPinVariableBinding.h"
 #include "Graph/SControlRigGraphChangePinType.h"
-#include "AssetTypeActions_ControlRigPose.h"
+#include "Tools/AssetTypeActions_ControlRigPose.h"
 #include "ControlRigBlueprintFactory.h"
 #include "ControlRigPythonLogDetails.h"
 #include "Dialog/SCustomDialog.h"
diff --git a/Engine/Plugins/Animation/ControlRig/Source/ControlRigEditor/Private/ControlRigElementDetails.cpp b/Engine/Plugins/Animation/ControlRig/Source/ControlRigEditor/Private/ControlRigElementDetails.cpp
--- a/Engine/Plugins/Animation/ControlRig/Source/ControlRigEditor/Private/ControlRigElementDetails.cpp
+++ b/Engine/Plugins/Animation/ControlRig/Source/ControlRigEditor/Private/ControlRigElementDetails.cpp
@@ -17,7 +17,7 @@
 #include "Graph/SControlRigGraphPinVariableBinding.h"
 #include "HAL/PlatformApplicationMisc.h"
 #include "Styling/AppStyle.h"
-#include "SRigHierarchyTreeView.h"
+#include "Editor/SRigHierarchyTreeView.h"
 
 #define LOCTEXT_NAMESPACE "ControlRigElementDetails"
 
diff --git a/Engine/Plugins/Animation/ControlRig/Source/ControlRigEditor/Private/ControlRigElementDetails.h b/Engine/Plugins/Animation/ControlRig/Source/ControlRigEditor/Private/ControlRigElementDetails.h
--- a/Engine/Plugins/Animation/ControlRig/Source/ControlRigEditor/Private/ControlRigElementDetails.h
+++ b/Engine/Plugins/Animation/ControlRig/Source/ControlRigEditor/Private/ControlRigElementDetails.h
@@ -9,7 +9,7 @@
 #include "Rigs/RigHierarchy.h"
 #include "ControlRig.h"
 #include "ControlRigBlueprint.h"
-#include "DetailsViewWrapperObject.h"
+#include "Editor/DetailsViewWrapperObject.h"
 #include "Graph/SControlRigGraphPinNameListValueWidget.h"
 #include "Editor/SControlRigGizmoNameList.h"
 #include "Styling/SlateTypes.h"
diff --git a/Engine/Plugins/Animation/ControlRig/Source/ControlRigEditor/Private/ControlRigGraphDetails.h b/Engine/Plugins/Animation/ControlRig/Source/ControlRigEditor/Private/ControlRigGraphDetails.h
--- a/Engine/Plugins/Animation/ControlRig/Source/ControlRigEditor/Private/ControlRigGraphDetails.h
+++ b/Engine/Plugins/Animation/ControlRig/Source/ControlRigEditor/Private/ControlRigGraphDetails.h
@@ -8,11 +8,11 @@
 #include "UObject/WeakObjectPtr.h"
 #include "Graph/ControlRigGraph.h"
 #include "ControlRigBlueprint.h"
-#include "ControlRigEditor.h"
+#include "Editor/ControlRigEditor.h"
 #include "SGraphPin.h"
 #include "Graph/SControlRigGraphNode.h"
 #include "Widgets/Colors/SColorBlock.h"
-#include "DetailsViewWrapperObject.h"
+#include "Editor/DetailsViewWrapperObject.h"
 #include "IDetailPropertyExtensionHandler.h"
 #include "Graph/ControlRigGraphSchema.h"
 #include "SAdvancedTransformInputBox.h"
diff --git a/Engine/Plugins/Animation/ControlRig/Source/ControlRigEditor/Private/ControlRigLocalVariableDetails.h b/Engine/Plugins/Animation/ControlRig/Source/ControlRigEditor/Private/ControlRigLocalVariableDetails.h
--- a/Engine/Plugins/Animation/ControlRig/Source/ControlRigEditor/Private/ControlRigLocalVariableDetails.h
+++ b/Engine/Plugins/Animation/ControlRig/Source/ControlRigEditor/Private/ControlRigLocalVariableDetails.h
@@ -7,12 +7,11 @@
 #include "IPropertyTypeCustomization.h"
 #include "ControlRig.h"
 #include "ControlRigBlueprint.h"
-#include "DetailsViewWrapperObject.h"
+#include "Editor/DetailsViewWrapperObject.h"
 #include "Graph/ControlRigGraph.h"
 #include "Graph/SControlRigGraphPinNameListValueWidget.h"
 #include "Styling/SlateTypes.h"
 #include "IPropertyUtilities.h"
-#include "DetailsViewWrapperObject.h"
 #include "Graph/ControlRigGraphSchema.h"
 
 class IPropertyHandle;
diff --git a/Engine/Plugins/Animation/ControlRig/Source/ControlRigEditor/Private/ControlRigSequencerEditorLibrary.cpp b/Engine/Plugins/Animation/ControlRig/Source/ControlRigEditor/Private/ControlRigSequencerEditorLibrary.cpp
--- a/Engine/Plugins/Animation/ControlRig/Source/ControlRigEditor/Private/ControlRigSequencerEditorLibrary.cpp
+++ b/Engine/Plugins/Animation/ControlRig/Source/ControlRigEditor/Private/ControlRigSequencerEditorLibrary.cpp
@@ -17,7 +17,7 @@
 #include "Rigs/FKControlRig.h"
 #include "ControlRig/Private/Units/Execution/RigUnit_InverseExecution.h"
 #include "ControlRig.h"
-#include "ControlRigEditMode.h"
+#include "EditMode/ControlRigEditMode.h"
 #include "EditorModeManager.h"
 #include "Engine/Selection.h"
 #include "ControlRigObjectBinding.h"
@@ -34,7 +34,7 @@
 #include "Exporters/AnimSeqExportOption.h"
 #include "MovieSceneTimeHelpers.h"
 #include "ScopedTransaction.h"
-#include "ControlRigParameterTrackEditor.h"
+#include "Sequencer/ControlRigParameterTrackEditor.h"
 #include "ControlRigSpaceChannelEditors.h"
 #include "LevelSequenceEditorBlueprintLibrary.h"
 
diff --git a/Engine/Plugins/Animation/ControlRig/Source/ControlRigEditor/Private/EditMode/ControlRigBaseDockableView.cpp b/Engine/Plugins/Animation/ControlRig/Source/ControlRigEditor/Private/EditMode/ControlRigBaseDockableView.cpp
--- a/Engine/Plugins/Animation/ControlRig/Source/ControlRigEditor/Private/EditMode/ControlRigBaseDockableView.cpp
+++ b/Engine/Plugins/Animation/ControlRig/Source/ControlRigEditor/Private/EditMode/ControlRigBaseDockableView.cpp
@@ -7,7 +7,7 @@
 #include "ScopedTransaction.h"
 #include "ControlRig.h"
 #include "UnrealEdGlobals.h"
-#include "ControlRigEditMode.h"
+#include "EditMode/ControlRigEditMode.h"
 #include "EditorModeManager.h"
 #include "ISequencer.h"
 
diff --git a/Engine/Plugins/Animation/ControlRig/Source/ControlRigEditor/Private/EditMode/ControlRigControlsProxy.cpp b/Engine/Plugins/Animation/ControlRig/Source/ControlRigEditor/Private/EditMode/ControlRigControlsProxy.cpp
--- a/Engine/Plugins/Animation/ControlRig/Source/ControlRigEditor/Private/EditMode/ControlRigControlsProxy.cpp
+++ b/Engine/Plugins/Animation/ControlRig/Source/ControlRigEditor/Private/EditMode/ControlRigControlsProxy.cpp
@@ -2,7 +2,7 @@
 
 #include "ControlRigControlsProxy.h"
 #include "EditorModeManager.h"
-#include "ControlRigEditMode.h"
+#include "EditMode/ControlRigEditMode.h"
 #include "Sequencer/ControlRigSequence.h"
 #include "Rigs/RigHierarchy.h"
 
diff --git a/Engine/Plugins/Animation/ControlRig/Source/ControlRigEditor/Private/EditMode/ControlRigEditMode.cpp b/Engine/Plugins/Animation/ControlRig/Source/ControlRigEditor/Private/EditMode/ControlRigEditMode.cpp
--- a/Engine/Plugins/Animation/ControlRig/Source/ControlRigEditor/Private/EditMode/ControlRigEditMode.cpp
+++ b/Engine/Plugins/Animation/ControlRig/Source/ControlRigEditor/Private/EditMode/ControlRigEditMode.cpp
@@ -1,13 +1,13 @@
 // Copyright Epic Games, Inc. All Rights Reserved.
 
-#include "ControlRigEditMode.h"
+#include "EditMode/ControlRigEditMode.h"
 #include "ControlRigEditModeToolkit.h"
 #include "Toolkits/ToolkitManager.h"
 #include "SControlRigEditModeTools.h"
 #include "Algo/Transform.h"
 #include "ControlRig.h"
 #include "HitProxies.h"
-#include "ControlRigEditModeSettings.h"
+#include "EditMode/ControlRigEditModeSettings.h"
 #include "ISequencer.h"
 #include "SequencerSettings.h"
 #include "Sequencer/ControlRigSequence.h"
@@ -19,7 +19,7 @@
 #include "Engine/Selection.h"
 #include "LevelEditorViewport.h"
 #include "Components/SkeletalMeshComponent.h"
-#include "ControlRigEditModeCommands.h"
+#include "EditMode/ControlRigEditModeCommands.h"
 #include "Framework/Application/SlateApplication.h"
 #include "Modules/ModuleManager.h"
 #include "ISequencerModule.h"
@@ -53,7 +53,7 @@
 #include "ToolMenus.h"
 #include "Sequencer/MovieSceneControlRigParameterTrack.h"
 #include "Sequencer/MovieSceneControlRigParameterSection.h"
-#include "SRigSpacePickerWidget.h"
+#include "Editor/SRigSpacePickerWidget.h"
 #include "ControlRigSpaceChannelEditors.h"
 #include "ControlRigSequencerEditorLibrary.h"
 #include "LevelSequence.h"
diff --git a/Engine/Plugins/Animation/ControlRig/Source/ControlRigEditor/Private/EditMode/ControlRigEditModeCommands.cpp b/Engine/Plugins/Animation/ControlRig/Source/ControlRigEditor/Private/EditMode/ControlRigEditModeCommands.cpp
--- a/Engine/Plugins/Animation/ControlRig/Source/ControlRigEditor/Private/EditMode/ControlRigEditModeCommands.cpp
+++ b/Engine/Plugins/Animation/ControlRig/Source/ControlRigEditor/Private/EditMode/ControlRigEditModeCommands.cpp
@@ -1,6 +1,6 @@
 // Copyright Epic Games, Inc. All Rights Reserved.
 
-#include "ControlRigEditModeCommands.h"
+#include "EditMode/ControlRigEditModeCommands.h"
 
 #define LOCTEXT_NAMESPACE "ControlRigEditModeCommands"
 
diff --git a/Engine/Plugins/Animation/ControlRig/Source/ControlRigEditor/Private/EditMode/ControlRigEditModeSettings.cpp b/Engine/Plugins/Animation/ControlRig/Source/ControlRigEditor/Private/EditMode/ControlRigEditModeSettings.cpp
--- a/Engine/Plugins/Animation/ControlRig/Source/ControlRigEditor/Private/EditMode/ControlRigEditModeSettings.cpp
+++ b/Engine/Plugins/Animation/ControlRig/Source/ControlRigEditor/Private/EditMode/ControlRigEditModeSettings.cpp
@@ -1,8 +1,8 @@
 // Copyright Epic Games, Inc. All Rights Reserved.
 
-#include "ControlRigEditModeSettings.h"
+#include "EditMode/ControlRigEditModeSettings.h"
 #include "EditorModeManager.h"
-#include "ControlRigEditMode.h"
+#include "EditMode/ControlRigEditMode.h"
 
 void UControlRigEditModeSettings::PreEditChange(FProperty* PropertyAboutToChange)
 {
diff --git a/Engine/Plugins/Animation/ControlRig/Source/ControlRigEditor/Private/EditMode/ControlRigEditModeToolkit.cpp b/Engine/Plugins/Animation/ControlRig/Source/ControlRigEditor/Private/EditMode/ControlRigEditModeToolkit.cpp
--- a/Engine/Plugins/Animation/ControlRig/Source/ControlRigEditor/Private/EditMode/ControlRigEditModeToolkit.cpp
+++ b/Engine/Plugins/Animation/ControlRig/Source/ControlRigEditor/Private/EditMode/ControlRigEditModeToolkit.cpp
@@ -10,7 +10,7 @@
 #include "Toolkits/BaseToolkit.h"
 #include "EditorModeManager.h"
 #include "SControlRigEditModeTools.h"
-#include "ControlRigEditMode.h"
+#include "EditMode/ControlRigEditMode.h"
 #include "Modules/ModuleManager.h"
 #include "EditMode/SControlRigBaseListWidget.h"
 #include "EditMode/SControlRigTweenWidget.h"
@@ -19,7 +19,7 @@
 #include "Editor/SControlRigProfilingView.h"
 #include "Toolkits/AssetEditorModeUILayer.h"
 #include "Widgets/Docking/SDockTab.h"
-#include "ControlRigEditModeSettings.h"
+#include "EditMode/ControlRigEditModeSettings.h"
 #include "SControlRigDetails.h"
 #include "SControlRigOutliner.h"
 #include "SControlRigSpacePicker.h"
diff --git a/Engine/Plugins/Animation/ControlRig/Source/ControlRigEditor/Private/EditMode/ControlRigEditModeToolkit.h b/Engine/Plugins/Animation/ControlRig/Source/ControlRigEditor/Private/EditMode/ControlRigEditModeToolkit.h
--- a/Engine/Plugins/Animation/ControlRig/Source/ControlRigEditor/Private/EditMode/ControlRigEditModeToolkit.h
+++ b/Engine/Plugins/Animation/ControlRig/Source/ControlRigEditor/Private/EditMode/ControlRigEditModeToolkit.h
@@ -9,7 +9,7 @@
 #include "Toolkits/BaseToolkit.h"
 #include "EditorModeManager.h"
 #include "SControlRigEditModeTools.h"
-#include "ControlRigEditMode.h"
+#include "EditMode/ControlRigEditMode.h"
 
 class SControlRigTweenWidget;
 
diff --git a/Engine/Plugins/Animation/ControlRig/Source/ControlRigEditor/Private/EditMode/SControlPicker.cpp b/Engine/Plugins/Animation/ControlRig/Source/ControlRigEditor/Private/EditMode/SControlPicker.cpp
--- a/Engine/Plugins/Animation/ControlRig/Source/ControlRigEditor/Private/EditMode/SControlPicker.cpp
+++ b/Engine/Plugins/Animation/ControlRig/Source/ControlRigEditor/Private/EditMode/SControlPicker.cpp
@@ -12,7 +12,7 @@
 #include "Styling/AppStyle.h"
 #include "Widgets/Input/SButton.h"
 #include "ScopedTransaction.h"
-#include "ControlRigEditMode.h"
+#include "EditMode/ControlRigEditMode.h"
 #include "ControlRig.h"
 #include "EditorModeManager.h"
 #include "SEditorUserWidgetHost.h"
diff --git a/Engine/Plugins/Animation/ControlRig/Source/ControlRigEditor/Private/EditMode/SControlRigBaseListWidget.cpp b/Engine/Plugins/Animation/ControlRig/Source/ControlRigEditor/Private/EditMode/SControlRigBaseListWidget.cpp
--- a/Engine/Plugins/Animation/ControlRig/Source/ControlRigEditor/Private/EditMode/SControlRigBaseListWidget.cpp
+++ b/Engine/Plugins/Animation/ControlRig/Source/ControlRigEditor/Private/EditMode/SControlRigBaseListWidget.cpp
@@ -19,7 +19,7 @@
 
 #include "ControlRig.h"
 #include "UnrealEdGlobals.h"
-#include "ControlRigEditMode.h"
+#include "EditMode/ControlRigEditMode.h"
 #include "Tools/ControlRigPose.h"
 #include "EditorModeManager.h"
 
@@ -31,13 +31,13 @@
 #include "Widgets/Layout/SScrollBox.h"
 #include "Framework/Commands/GenericCommands.h"
 #include "Framework/Commands/UICommandList.h"
-#include "ControlRigEditModeCommands.h"
+#include "EditMode/ControlRigEditModeCommands.h"
 #include "Tools/CreateControlAssetRigSettings.h"
 #include "FileHelpers.h"
 #include "Tools/ControlRigPoseMirrorSettings.h"
 #include "ObjectTools.h"
 #include "SControlRigUpdatePose.h"
-#include "SControlRigRenamePoseControls.h"
+#include "EditMode/SControlRigRenamePoseControls.h"
 #include "Dialogs/Dialogs.h"
 #include "ControlRigSequencerEditorLibrary.h"
 #include "LevelSequence.h"
diff --git a/Engine/Plugins/Animation/ControlRig/Source/ControlRigEditor/Private/EditMode/SControlRigControlViews.cpp b/Engine/Plugins/Animation/ControlRig/Source/ControlRigEditor/Private/EditMode/SControlRigControlViews.cpp
--- a/Engine/Plugins/Animation/ControlRig/Source/ControlRigEditor/Private/EditMode/SControlRigControlViews.cpp
+++ b/Engine/Plugins/Animation/ControlRig/Source/ControlRigEditor/Private/EditMode/SControlRigControlViews.cpp
@@ -15,7 +15,7 @@
 #include "AssetThumbnail.h"
 #include "FileHelpers.h"
 #include "Tools/ControlRigPoseMirrorSettings.h"
-#include "ControlRigEditMode.h"
+#include "EditMode/ControlRigEditMode.h"
 #include "EditorModeManager.h"
 #include "Widgets/Input/SButton.h"
 #include "LevelEditor.h"
diff --git a/Engine/Plugins/Animation/ControlRig/Source/ControlRigEditor/Private/EditMode/SControlRigDetails.cpp b/Engine/Plugins/Animation/ControlRig/Source/ControlRigEditor/Private/EditMode/SControlRigDetails.cpp
--- a/Engine/Plugins/Animation/ControlRig/Source/ControlRigEditor/Private/EditMode/SControlRigDetails.cpp
+++ b/Engine/Plugins/Animation/ControlRig/Source/ControlRigEditor/Private/EditMode/SControlRigDetails.cpp
@@ -10,7 +10,7 @@
 #include "ScopedTransaction.h"
 #include "ControlRig.h"
 #include "UnrealEdGlobals.h"
-#include "ControlRigEditMode.h"
+#include "EditMode/ControlRigEditMode.h"
 #include "EditorModeManager.h"
 #include "ISequencer.h"
 #include "LevelSequence.h"
@@ -22,7 +22,7 @@
 #include "Widgets/Input/SButton.h"
 #include "Widgets/Layout/SScrollBox.h"
 #include "ControlRigControlsProxy.h"
-#include "ControlRigEditModeSettings.h"
+#include "EditMode/ControlRigEditModeSettings.h"
 #include "Modules/ModuleManager.h"
 #include "TimerManager.h"
 
diff --git a/Engine/Plugins/Animation/ControlRig/Source/ControlRigEditor/Private/EditMode/SControlRigEditModeTools.cpp b/Engine/Plugins/Animation/ControlRig/Source/ControlRigEditor/Private/EditMode/SControlRigEditModeTools.cpp
--- a/Engine/Plugins/Animation/ControlRig/Source/ControlRigEditor/Private/EditMode/SControlRigEditModeTools.cpp
+++ b/Engine/Plugins/Animation/ControlRig/Source/ControlRigEditor/Private/EditMode/SControlRigEditModeTools.cpp
@@ -7,10 +7,10 @@
 #include "ISequencer.h"
 #include "PropertyHandle.h"
 #include "ControlRig.h"
-#include "ControlRigEditModeSettings.h"
+#include "EditMode/ControlRigEditModeSettings.h"
 #include "IDetailRootObjectCustomization.h"
 #include "Modules/ModuleManager.h"
-#include "ControlRigEditMode.h"
+#include "EditMode/ControlRigEditMode.h"
 #include "EditorModeManager.h"
 #include "Widgets/Input/SButton.h"
 #include "Widgets/Layout/SExpandableArea.h"
@@ -38,7 +38,7 @@
 #include "ScopedTransaction.h"
 #include "ControlRigEditModeToolkit.h"
 #include "SControlRigDetails.h"
-#include "Constraints/SConstraintsWidget.h"
+#include "Editor/Constraints/SConstraintsWidget.h"
 
 #define LOCTEXT_NAMESPACE "ControlRigEditModeTools"
 
diff --git a/Engine/Plugins/Animation/ControlRig/Source/ControlRigEditor/Private/EditMode/SControlRigEditModeTools.h b/Engine/Plugins/Animation/ControlRig/Source/ControlRigEditor/Private/EditMode/SControlRigEditModeTools.h
--- a/Engine/Plugins/Animation/ControlRig/Source/ControlRigEditor/Private/EditMode/SControlRigEditModeTools.h
+++ b/Engine/Plugins/Animation/ControlRig/Source/ControlRigEditor/Private/EditMode/SControlRigEditModeTools.h
@@ -6,8 +6,8 @@
 #include "Widgets/SCompoundWidget.h"
 #include "IDetailKeyframeHandler.h"
 #include "RigVMModel/RigVMGraph.h"
-#include "SRigHierarchyTreeView.h"
-#include "SRigSpacePickerWidget.h"
+#include "Editor/SRigHierarchyTreeView.h"
+#include "Editor/SRigSpacePickerWidget.h"
 #include "Widgets/DeclarativeSyntaxSupport.h"
 
 class SConstraintsEditionWidget;
diff --git a/Engine/Plugins/Animation/ControlRig/Source/ControlRigEditor/Private/EditMode/SControlRigOutliner.cpp b/Engine/Plugins/Animation/ControlRig/Source/ControlRigEditor/Private/EditMode/SControlRigOutliner.cpp
--- a/Engine/Plugins/Animation/ControlRig/Source/ControlRigEditor/Private/EditMode/SControlRigOutliner.cpp
+++ b/Engine/Plugins/Animation/ControlRig/Source/ControlRigEditor/Private/EditMode/SControlRigOutliner.cpp
@@ -14,7 +14,7 @@
 #include "ScopedTransaction.h"
 #include "ControlRig.h"
 #include "UnrealEdGlobals.h"
-#include "ControlRigEditMode.h"
+#include "EditMode/ControlRigEditMode.h"
 #include "EditorModeManager.h"
 #include "ISequencer.h"
 #include "LevelSequence.h"
@@ -23,9 +23,9 @@
 #include "LevelEditor.h"
 #include "Widgets/Text/STextBlock.h"
 #include "ControlRigControlsProxy.h"
-#include "ControlRigEditModeSettings.h"
+#include "EditMode/ControlRigEditModeSettings.h"
 #include "TimerManager.h"
-#include "SRigHierarchyTreeView.h"
+#include "Editor/SRigHierarchyTreeView.h"
 #include "Rigs/RigHierarchyController.h"
 #include "ControlRigObjectBinding.h"
 #include "ControlRigEditorStyle.h"
diff --git a/Engine/Plugins/Animation/ControlRig/Source/ControlRigEditor/Private/EditMode/SControlRigOutliner.h b/Engine/Plugins/Animation/ControlRig/Source/ControlRigEditor/Private/EditMode/SControlRigOutliner.h
--- a/Engine/Plugins/Animation/ControlRig/Source/ControlRigEditor/Private/EditMode/SControlRigOutliner.h
+++ b/Engine/Plugins/Animation/ControlRig/Source/ControlRigEditor/Private/EditMode/SControlRigOutliner.h
@@ -10,7 +10,7 @@
 #include "Widgets/SCompoundWidget.h"
 #include "ControlRig.h"
 #include "Rigs/RigHierarchy.h"
-#include "SRigHierarchyTreeView.h"
+#include "Editor/SRigHierarchyTreeView.h"
 #include "Widgets/SBoxPanel.h"
 
 
diff --git a/Engine/Plugins/Animation/ControlRig/Source/ControlRigEditor/Private/EditMode/SControlRigRenamePoseControls.cpp b/Engine/Plugins/Animation/ControlRig/Source/ControlRigEditor/Private/EditMode/SControlRigRenamePoseControls.cpp
--- a/Engine/Plugins/Animation/ControlRig/Source/ControlRigEditor/Private/EditMode/SControlRigRenamePoseControls.cpp
+++ b/Engine/Plugins/Animation/ControlRig/Source/ControlRigEditor/Private/EditMode/SControlRigRenamePoseControls.cpp
@@ -1,6 +1,6 @@
 // Copyright Epic Games, Inc. All Rights Reserved.
 
-#include "SControlRigRenamePoseControls.h"
+#include "EditMode/SControlRigRenamePoseControls.h"
 #include "CoreMinimal.h"
 #include "Widgets/SWidget.h"
 #include "Widgets/SCompoundWidget.h"
@@ -9,7 +9,7 @@
 #include "Framework/Application/SlateApplication.h"
 #include "Widgets/SWindow.h"
 #include "Widgets/Input/SButton.h"
-#include "ControlRigEditMode.h"
+#include "EditMode/ControlRigEditMode.h"
 #include "Tools/ControlRigPose.h"
 #include "EditorModeManager.h"
 #include "ScopedTransaction.h"
diff --git a/Engine/Plugins/Animation/ControlRig/Source/ControlRigEditor/Private/EditMode/SControlRigSnapper.cpp b/Engine/Plugins/Animation/ControlRig/Source/ControlRigEditor/Private/EditMode/SControlRigSnapper.cpp
--- a/Engine/Plugins/Animation/ControlRig/Source/ControlRigEditor/Private/EditMode/SControlRigSnapper.cpp
+++ b/Engine/Plugins/Animation/ControlRig/Source/ControlRigEditor/Private/EditMode/SControlRigSnapper.cpp
@@ -12,7 +12,7 @@
 #include "ScopedTransaction.h"
 #include "ControlRig.h"
 #include "UnrealEdGlobals.h"
-#include "ControlRigEditMode.h"
+#include "EditMode/ControlRigEditMode.h"
 #include "Tools/ControlRigPose.h"
 #include "EditorModeManager.h"
 #include "ISequencer.h"
diff --git a/Engine/Plugins/Animation/ControlRig/Source/ControlRigEditor/Private/EditMode/SControlRigSpacePicker.cpp b/Engine/Plugins/Animation/ControlRig/Source/ControlRigEditor/Private/EditMode/SControlRigSpacePicker.cpp
--- a/Engine/Plugins/Animation/ControlRig/Source/ControlRigEditor/Private/EditMode/SControlRigSpacePicker.cpp
+++ b/Engine/Plugins/Animation/ControlRig/Source/ControlRigEditor/Private/EditMode/SControlRigSpacePicker.cpp
@@ -14,7 +14,7 @@
 #include "ISequencer.h"
 #include "ScopedTransaction.h"
 #include "ControlRig.h"
-#include "ControlRigEditMode.h"
+#include "EditMode/ControlRigEditMode.h"
 #include "MovieSceneSequence.h"
 #include "MovieScene.h"
 #include "ControlRigSpaceChannelEditors.h"
diff --git a/Engine/Plugins/Animation/ControlRig/Source/ControlRigEditor/Private/EditMode/SControlRigSpacePicker.h b/Engine/Plugins/Animation/ControlRig/Source/ControlRigEditor/Private/EditMode/SControlRigSpacePicker.h
--- a/Engine/Plugins/Animation/ControlRig/Source/ControlRigEditor/Private/EditMode/SControlRigSpacePicker.h
+++ b/Engine/Plugins/Animation/ControlRig/Source/ControlRigEditor/Private/EditMode/SControlRigSpacePicker.h
@@ -11,7 +11,7 @@
 #include "Widgets/SCompoundWidget.h"
 #include "Misc/FrameNumber.h"
 #include "ControlRig.h"
-#include "SRigSpacePickerWidget.h"
+#include "Editor/SRigSpacePickerWidget.h"
 
 
 class ISequencer;
diff --git a/Engine/Plugins/Animation/ControlRig/Source/ControlRigEditor/Private/EditMode/SControlRigUpdatePose.cpp b/Engine/Plugins/Animation/ControlRig/Source/ControlRigEditor/Private/EditMode/SControlRigUpdatePose.cpp
--- a/Engine/Plugins/Animation/ControlRig/Source/ControlRigEditor/Private/EditMode/SControlRigUpdatePose.cpp
+++ b/Engine/Plugins/Animation/ControlRig/Source/ControlRigEditor/Private/EditMode/SControlRigUpdatePose.cpp
@@ -9,7 +9,7 @@
 #include "Framework/Application/SlateApplication.h"
 #include "Widgets/SWindow.h"
 #include "Widgets/Input/SButton.h"
-#include "ControlRigEditMode.h"
+#include "EditMode/ControlRigEditMode.h"
 #include "Tools/ControlRigPose.h"
 #include "EditorModeManager.h"
 #include "ScopedTransaction.h"
diff --git a/Engine/Plugins/Animation/ControlRig/Source/ControlRigEditor/Private/EditMode/SCurveControlContainer.cpp b/Engine/Plugins/Animation/ControlRig/Source/ControlRigEditor/Private/EditMode/SCurveControlContainer.cpp
--- a/Engine/Plugins/Animation/ControlRig/Source/ControlRigEditor/Private/EditMode/SCurveControlContainer.cpp
+++ b/Engine/Plugins/Animation/ControlRig/Source/ControlRigEditor/Private/EditMode/SCurveControlContainer.cpp
@@ -9,7 +9,7 @@
 #include "Widgets/Input/STextEntryPopup.h"
 #include "PropertyCustomizationHelpers.h"
 #include "Framework/Commands/GenericCommands.h"
-#include "ControlRigEditor.h"
+#include "Editor/ControlRigEditor.h"
 #include "Engine/SkeletalMesh.h"
 #include "Animation/Skeleton.h"
 #include "ControlRig.h"
diff --git a/Engine/Plugins/Animation/ControlRig/Source/ControlRigEditor/Private/Editor/ControlRigEditor.cpp b/Engine/Plugins/Animation/ControlRig/Source/ControlRigEditor/Private/Editor/ControlRigEditor.cpp
--- a/Engine/Plugins/Animation/ControlRig/Source/ControlRigEditor/Private/Editor/ControlRigEditor.cpp
+++ b/Engine/Plugins/Animation/ControlRig/Source/ControlRigEditor/Private/Editor/ControlRigEditor.cpp
@@ -17,8 +17,8 @@
 #include "Kismet2/BlueprintEditorUtils.h"
 #include "IPersonaToolkit.h"
 #include "PersonaModule.h"
-#include "ControlRigEditorEditMode.h"
-#include "ControlRigEditModeSettings.h"
+#include "Editor/ControlRigEditorEditMode.h"
+#include "EditMode/ControlRigEditModeSettings.h"
 #include "EditorModeManager.h"
 #include "ControlRigBlueprintGeneratedClass.h"
 #include "AnimCustomInstanceHelper.h"
diff --git a/Engine/Plugins/Animation/ControlRig/Source/ControlRigEditor/Private/Editor/ControlRigEditor.h b/Engine/Plugins/Animation/ControlRig/Source/ControlRigEditor/Private/Editor/ControlRigEditor.h
--- a/Engine/Plugins/Animation/ControlRig/Source/ControlRigEditor/Private/Editor/ControlRigEditor.h
+++ b/Engine/Plugins/Animation/ControlRig/Source/ControlRigEditor/Private/Editor/ControlRigEditor.h
@@ -3,7 +3,7 @@
 #pragma once
 
 #include "IControlRigEditor.h"
-#include "ControlRigEditorEditMode.h"
+#include "Editor/ControlRigEditorEditMode.h"
 #include "AssetEditorModeManager.h"
 #include "DragAndDrop/GraphNodeDragDropOp.h"
 #include "ControlRigDefines.h"
@@ -20,7 +20,7 @@
 #include "ScopedTransaction.h"
 #include "Graph/ControlRigGraphNode.h"
 #include "RigVMModel/RigVMController.h"
-#include "DetailsViewWrapperObject.h"
+#include "Editor/DetailsViewWrapperObject.h"
 
 class UControlRigBlueprint;
 class IPersonaToolkit;
diff --git a/Engine/Plugins/Animation/ControlRig/Source/ControlRigEditor/Private/Editor/ControlRigEditorEditMode.cpp b/Engine/Plugins/Animation/ControlRig/Source/ControlRigEditor/Private/Editor/ControlRigEditorEditMode.cpp
--- a/Engine/Plugins/Animation/ControlRig/Source/ControlRigEditor/Private/Editor/ControlRigEditorEditMode.cpp
+++ b/Engine/Plugins/Animation/ControlRig/Source/ControlRigEditor/Private/Editor/ControlRigEditorEditMode.cpp
@@ -1,6 +1,6 @@
 // Copyright Epic Games, Inc. All Rights Reserved.
 
-#include "ControlRigEditorEditMode.h"
+#include "Editor/ControlRigEditorEditMode.h"
 #include "IPersonaPreviewScene.h"
 #include "AssetEditorModeManager.h"
 #include "ControlRig.h"
diff --git a/Engine/Plugins/Animation/ControlRig/Source/ControlRigEditor/Private/Editor/ControlRigEditorEditMode.h b/Engine/Plugins/Animation/ControlRig/Source/ControlRigEditor/Private/Editor/ControlRigEditorEditMode.h
--- a/Engine/Plugins/Animation/ControlRig/Source/ControlRigEditor/Private/Editor/ControlRigEditorEditMode.h
+++ b/Engine/Plugins/Animation/ControlRig/Source/ControlRigEditor/Private/Editor/ControlRigEditorEditMode.h
@@ -3,7 +3,7 @@
 #pragma once
 
 #include "CoreMinimal.h"
-#include "ControlRigEditMode.h"
+#include "EditMode/ControlRigEditMode.h"
 #include "IPersonaEditMode.h"
 #include "Preferences/PersonaOptions.h"
 #include "Persona/Private/AnimationEditorViewportClient.h"
diff --git a/Engine/Plugins/Animation/ControlRig/Source/ControlRigEditor/Private/Editor/DetailsViewWrapperObject.cpp b/Engine/Plugins/Animation/ControlRig/Source/ControlRigEditor/Private/Editor/DetailsViewWrapperObject.cpp
--- a/Engine/Plugins/Animation/ControlRig/Source/ControlRigEditor/Private/Editor/DetailsViewWrapperObject.cpp
+++ b/Engine/Plugins/Animation/ControlRig/Source/ControlRigEditor/Private/Editor/DetailsViewWrapperObject.cpp
@@ -1,6 +1,6 @@
 // Copyright Epic Games, Inc. All Rights Reserved.
 
-#include "DetailsViewWrapperObject.h"
+#include "Editor/DetailsViewWrapperObject.h"
 #include "Units/RigUnit.h"
 #include "ControlRigElementDetails.h"
 #include "ControlRigLocalVariableDetails.h"
diff --git a/Engine/Plugins/Animation/ControlRig/Source/ControlRigEditor/Private/Editor/SRigHierarchy.cpp b/Engine/Plugins/Animation/ControlRig/Source/ControlRigEditor/Private/Editor/SRigHierarchy.cpp
--- a/Engine/Plugins/Animation/ControlRig/Source/ControlRigEditor/Private/Editor/SRigHierarchy.cpp
+++ b/Engine/Plugins/Animation/ControlRig/Source/ControlRigEditor/Private/Editor/SRigHierarchy.cpp
@@ -40,7 +40,7 @@
 #include "EditMode/ControlRigEditMode.h"
 #include "ToolMenus.h"
 #include "ControlRigContextMenuContext.h"
-#include "SRigSpacePickerWidget.h"
+#include "Editor/SRigSpacePickerWidget.h"
 #include "Settings/ControlRigSettings.h"
 #include "Slate/Private/Widgets/Views/SListPanel.h"
 #include "Framework/Notifications/NotificationManager.h"
diff --git a/Engine/Plugins/Animation/ControlRig/Source/ControlRigEditor/Private/Editor/SRigHierarchy.h b/Engine/Plugins/Animation/ControlRig/Source/ControlRigEditor/Private/Editor/SRigHierarchy.h
--- a/Engine/Plugins/Animation/ControlRig/Source/ControlRigEditor/Private/Editor/SRigHierarchy.h
+++ b/Engine/Plugins/Animation/ControlRig/Source/ControlRigEditor/Private/Editor/SRigHierarchy.h
@@ -6,7 +6,7 @@
 #include "EditorUndoClient.h"
 #include "DragAndDrop/GraphNodeDragDropOp.h"
 #include "Engine/SkeletalMesh.h"
-#include "SRigHierarchyTreeView.h"
+#include "Editor/SRigHierarchyTreeView.h"
 #include "Units/RigUnitContext.h"
 #include "SRigHierarchy.generated.h"
 
diff --git a/Engine/Plugins/Animation/ControlRig/Source/ControlRigEditor/Private/Editor/SRigHierarchyTreeView.cpp b/Engine/Plugins/Animation/ControlRig/Source/ControlRigEditor/Private/Editor/SRigHierarchyTreeView.cpp
--- a/Engine/Plugins/Animation/ControlRig/Source/ControlRigEditor/Private/Editor/SRigHierarchyTreeView.cpp
+++ b/Engine/Plugins/Animation/ControlRig/Source/ControlRigEditor/Private/Editor/SRigHierarchyTreeView.cpp
@@ -1,6 +1,6 @@
 // Copyright Epic Games, Inc. All Rights Reserved.
 
-#include "SRigHierarchyTreeView.h"
+#include "Editor/SRigHierarchyTreeView.h"
 #include "Styling/AppStyle.h"
 #include "Widgets/Layout/SSpacer.h"
 #include "Widgets/Layout/SScrollBox.h"
diff --git a/Engine/Plugins/Animation/ControlRig/Source/ControlRigEditor/Private/Editor/SRigSpacePickerWidget.cpp b/Engine/Plugins/Animation/ControlRig/Source/ControlRigEditor/Private/Editor/SRigSpacePickerWidget.cpp
--- a/Engine/Plugins/Animation/ControlRig/Source/ControlRigEditor/Private/Editor/SRigSpacePickerWidget.cpp
+++ b/Engine/Plugins/Animation/ControlRig/Source/ControlRigEditor/Private/Editor/SRigSpacePickerWidget.cpp
@@ -1,13 +1,13 @@
 // Copyright Epic Games, Inc. All Rights Reserved.
 
-#include "SRigSpacePickerWidget.h"
+#include "Editor/SRigSpacePickerWidget.h"
 #include "DetailLayoutBuilder.h"
 #include "Editor.h"
 #include "Widgets/Text/STextBlock.h"
 #include "Widgets/Layout/SUniformGridPanel.h"
 #include "Widgets/Input/SButton.h"
 #include "Widgets/Input/SSearchBox.h"
-#include "SRigHierarchyTreeView.h"
+#include "Editor/SRigHierarchyTreeView.h"
 #include "ControlRigEditorStyle.h"
 #include "PropertyCustomizationHelpers.h"
 #include "ISequencer.h"
diff --git a/Engine/Plugins/Animation/ControlRig/Source/ControlRigEditor/Private/Graph/ControlRigGraphPanelNodeFactory.cpp b/Engine/Plugins/Animation/ControlRig/Source/ControlRigEditor/Private/Graph/ControlRigGraphPanelNodeFactory.cpp
--- a/Engine/Plugins/Animation/ControlRig/Source/ControlRigEditor/Private/Graph/ControlRigGraphPanelNodeFactory.cpp
+++ b/Engine/Plugins/Animation/ControlRig/Source/ControlRigEditor/Private/Graph/ControlRigGraphPanelNodeFactory.cpp
@@ -1,6 +1,6 @@
 // Copyright Epic Games, Inc. All Rights Reserved.
 
-#include "ControlRigGraphPanelNodeFactory.h"
+#include "Graph/ControlRigGraphPanelNodeFactory.h"
 #include "Graph/ControlRigGraphNode.h"
 #include "Graph/ControlRigGraph.h"
 #include "EdGraphNode_Comment.h"
diff --git a/Engine/Plugins/Animation/ControlRig/Source/ControlRigEditor/Private/Graph/ControlRigGraphPanelPinFactory.cpp b/Engine/Plugins/Animation/ControlRig/Source/ControlRigEditor/Private/Graph/ControlRigGraphPanelPinFactory.cpp
--- a/Engine/Plugins/Animation/ControlRig/Source/ControlRigEditor/Private/Graph/ControlRigGraphPanelPinFactory.cpp
+++ b/Engine/Plugins/Animation/ControlRig/Source/ControlRigEditor/Private/Graph/ControlRigGraphPanelPinFactory.cpp
@@ -1,6 +1,6 @@
 // Copyright Epic Games, Inc. All Rights Reserved.
 
-#include "ControlRigGraphPanelPinFactory.h"
+#include "Graph/ControlRigGraphPanelPinFactory.h"
 #include "Graph/ControlRigGraphSchema.h"
 #include "Graph/ControlRigGraphNode.h"
 #include "Graph/ControlRigGraph.h"
diff --git a/Engine/Plugins/Animation/ControlRig/Source/ControlRigEditor/Private/Sequencer/ControlRigParameterTrackEditor.cpp b/Engine/Plugins/Animation/ControlRig/Source/ControlRigEditor/Private/Sequencer/ControlRigParameterTrackEditor.cpp
--- a/Engine/Plugins/Animation/ControlRig/Source/ControlRigEditor/Private/Sequencer/ControlRigParameterTrackEditor.cpp
+++ b/Engine/Plugins/Animation/ControlRig/Source/ControlRigEditor/Private/Sequencer/ControlRigParameterTrackEditor.cpp
@@ -1,6 +1,6 @@
 // Copyright Epic Games, Inc. All Rights Reserved.
 
-#include "ControlRigParameterTrackEditor.h"
+#include "Sequencer/ControlRigParameterTrackEditor.h"
 #include "Sequencer/MovieSceneControlRigParameterTrack.h"
 #include "Sequencer/MovieSceneControlRigParameterSection.h"
 #include "Framework/Commands/Commands.h"
@@ -39,7 +39,7 @@
 #include "Engine/SimpleConstructionScript.h"
 #include "Engine/Blueprint.h"
 #include "ControlRig.h"
-#include "ControlRigEditMode.h"
+#include "EditMode/ControlRigEditMode.h"
 #include "EditorModeManager.h"
 #include "Engine/Selection.h"
 #include "ControlRigObjectBinding.h"
@@ -73,7 +73,7 @@
 #include "UnrealEdGlobals.h"
 #include "Editor/UnrealEdEngine.h"
 #include "Toolkits/IToolkitHost.h"
-#include "ControlRigEditModeSettings.h"
+#include "EditMode/ControlRigEditModeSettings.h"
 #include "ControlRigSpaceChannelEditors.h"
 #include "TransformConstraint.h"
 #include "Misc/ScopedSlowTask.h"
diff --git a/Engine/Plugins/Animation/ControlRig/Source/ControlRigEditor/Private/Sequencer/ControlRigSequencerFilter.cpp b/Engine/Plugins/Animation/ControlRig/Source/ControlRigEditor/Private/Sequencer/ControlRigSequencerFilter.cpp
--- a/Engine/Plugins/Animation/ControlRig/Source/ControlRigEditor/Private/Sequencer/ControlRigSequencerFilter.cpp
+++ b/Engine/Plugins/Animation/ControlRig/Source/ControlRigEditor/Private/Sequencer/ControlRigSequencerFilter.cpp
@@ -2,7 +2,7 @@
 
 #include "ControlRigSequencerFilter.h"
 #include "ControlRig.h"
-#include "ControlRigSkeletalMeshComponent.h"
+#include "Editor/ControlRigSkeletalMeshComponent.h"
 #include "Sequencer/MovieSceneControlRigParameterTrack.h"
 #include "Styling/SlateIconFinder.h"
 #include "Framework/Commands/Commands.h"
diff --git a/Engine/Plugins/Animation/ControlRig/Source/ControlRigEditor/Private/Sequencer/ControlRigSpaceChannelEditors.cpp b/Engine/Plugins/Animation/ControlRig/Source/ControlRigEditor/Private/Sequencer/ControlRigSpaceChannelEditors.cpp
--- a/Engine/Plugins/Animation/ControlRig/Source/ControlRigEditor/Private/Sequencer/ControlRigSpaceChannelEditors.cpp
+++ b/Engine/Plugins/Animation/ControlRig/Source/ControlRigEditor/Private/Sequencer/ControlRigSpaceChannelEditors.cpp
@@ -40,7 +40,7 @@
 #include "MovieSceneSpawnableAnnotation.h"
 #include "ISequencerModule.h"
 #include "MovieSceneTracksComponentTypes.h"
-#include "SRigSpacePickerWidget.h"
+#include "Editor/SRigSpacePickerWidget.h"
 #include "Tools/ControlRigSnapper.h"
 #include "Sequencer/MovieSceneControlRigParameterTrack.h"
 #include "Sequencer/MovieSceneControlRigParameterSection.h"
diff --git a/Engine/Plugins/Animation/ControlRig/Source/ControlRigEditor/Private/Sequencer/MovieSceneControlRigSectionDetailsCustomization.cpp b/Engine/Plugins/Animation/ControlRig/Source/ControlRigEditor/Private/Sequencer/MovieSceneControlRigSectionDetailsCustomization.cpp
--- a/Engine/Plugins/Animation/ControlRig/Source/ControlRigEditor/Private/Sequencer/MovieSceneControlRigSectionDetailsCustomization.cpp
+++ b/Engine/Plugins/Animation/ControlRig/Source/ControlRigEditor/Private/Sequencer/MovieSceneControlRigSectionDetailsCustomization.cpp
@@ -1,6 +1,6 @@
 // Copyright Epic Games, Inc. All Rights Reserved.
 
-#include "MovieSceneControlRigSectionDetailsCustomization.h"
+#include "Sequencer/MovieSceneControlRigSectionDetailsCustomization.h"
 #include "DetailLayoutBuilder.h"
 
 TSharedRef<IDetailCustomization> FMovieSceneControlRigSectionDetailsCustomization::MakeInstance()
diff --git a/Engine/Plugins/Animation/ControlRig/Source/ControlRigEditor/Private/Tools/AssetTypeActions_ControlRigPose.cpp b/Engine/Plugins/Animation/ControlRig/Source/ControlRigEditor/Private/Tools/AssetTypeActions_ControlRigPose.cpp
--- a/Engine/Plugins/Animation/ControlRig/Source/ControlRigEditor/Private/Tools/AssetTypeActions_ControlRigPose.cpp
+++ b/Engine/Plugins/Animation/ControlRig/Source/ControlRigEditor/Private/Tools/AssetTypeActions_ControlRigPose.cpp
@@ -1,16 +1,16 @@
 // Copyright Epic Games, Inc. All Rights Reserved.
 
-#include "AssetTypeActions_ControlRigPose.h"
+#include "Tools/AssetTypeActions_ControlRigPose.h"
 #include "EditorFramework/AssetImportData.h"
 #include "ThumbnailRendering/SceneThumbnailInfo.h"
 #include "ToolMenus.h"
 
 #include "ControlRig.h"
 #include "UnrealEdGlobals.h"
-#include "ControlRigEditMode.h"
+#include "EditMode/ControlRigEditMode.h"
 #include "Tools/ControlRigPose.h"
 #include "EditorModeManager.h"
-#include "SControlRigRenamePoseControls.h"
+#include "EditMode/SControlRigRenamePoseControls.h"
 
 #define LOCTEXT_NAMESPACE "AssetTypeActions_ControlRigPose"
 
diff --git a/Engine/Plugins/Animation/ControlRig/Source/ControlRigEditor/Private/Tools/ControlRigTweener.cpp b/Engine/Plugins/Animation/ControlRig/Source/ControlRigEditor/Private/Tools/ControlRigTweener.cpp
--- a/Engine/Plugins/Animation/ControlRig/Source/ControlRigEditor/Private/Tools/ControlRigTweener.cpp
+++ b/Engine/Plugins/Animation/ControlRig/Source/ControlRigEditor/Private/Tools/ControlRigTweener.cpp
@@ -19,7 +19,7 @@
 #include "Channels/FloatChannelCurveModel.h"
 #include "MovieSceneSection.h"
 #include "EditorModeManager.h"
-#include "ControlRigEditMode.h"
+#include "EditMode/ControlRigEditMode.h"
 
 #define LOCTEXT_NAMESPACE "ControlRigTweenWidget"
 
diff --git a/Engine/Plugins/Animation/ControlRig/Source/ControlRigEditor/Public/ControlRigSpaceChannelEditors.h b/Engine/Plugins/Animation/ControlRig/Source/ControlRigEditor/Public/ControlRigSpaceChannelEditors.h
--- a/Engine/Plugins/Animation/ControlRig/Source/ControlRigEditor/Public/ControlRigSpaceChannelEditors.h
+++ b/Engine/Plugins/Animation/ControlRig/Source/ControlRigEditor/Public/ControlRigSpaceChannelEditors.h
@@ -15,7 +15,7 @@
 
 #include "Sequencer/MovieSceneControlRigSpaceChannel.h"
 #include "Sequencer/MovieSceneControlRigParameterSection.h"
-#include "SRigSpacePickerWidget.h"
+#include "Editor/SRigSpacePickerWidget.h"
 #include "Containers/SortedMap.h"
 #include "KeyBarCurveModel.h"
 
