Sequencer: Control Rig: Animation Mode Get rid of the advanced tab that was empty, all of it's properties are in other sections now

diff --git a/Engine/Plugins/Animation/ControlRig/Source/ControlRigEditor/Private/EditMode/ControlRigEditModeSettings.h b/Engine/Plugins/Animation/ControlRig/Source/ControlRigEditor/Private/EditMode/ControlRigEditModeSettings.h
--- a/Engine/Plugins/Animation/ControlRig/Source/ControlRigEditor/Private/EditMode/ControlRigEditModeSettings.h
+++ b/Engine/Plugins/Animation/ControlRig/Source/ControlRigEditor/Private/EditMode/ControlRigEditModeSettings.h
@@ -39,51 +39,51 @@ class UControlRigEditModeSettings : public UObject
 public:
 
 	/** Whether to show all bones in the hierarchy */
-	UPROPERTY(config, EditAnywhere, AdvancedDisplay, Category = "Animation Settings")
+	UPROPERTY(config, EditAnywhere, Category = "Animation Settings")
 	bool bDisplayHierarchy;
 
 	/** Whether to show all nulls in the hierarchy */
-	UPROPERTY(config, EditAnywhere, AdvancedDisplay, Category = "Animation Settings")
+	UPROPERTY(config, EditAnywhere, Category = "Animation Settings")
 	bool bDisplayNulls;
 
 	/** Should we always hide control shapes in viewport */
-	UPROPERTY(config, EditAnywhere, AdvancedDisplay, Category = "Animation Settings")
+	UPROPERTY(config, EditAnywhere, Category = "Animation Settings")
 	bool bHideControlShapes;
 
 	/** Should we always hide control shapes in viewport */
-	UPROPERTY(config, EditAnywhere, AdvancedDisplay, Category = "Animation Settings", meta = (EditCondition = "!bHideControlShapes"))
+	UPROPERTY(config, EditAnywhere, Category = "Animation Settings", meta = (EditCondition = "!bHideControlShapes"))
 	bool bShowAllProxyControls;
 
 	/** Determins if controls should be rendered on top of other controls */
-	UPROPERTY(config, EditAnywhere, AdvancedDisplay, Category = "Animation Settings", meta = (EditCondition = "!bHideControlShapes"))
+	UPROPERTY(config, EditAnywhere, Category = "Animation Settings", meta = (EditCondition = "!bHideControlShapes"))
 	bool bShowControlsAsOverlay;
 
 	/** Indicates a control being driven by a proxy control */
-	UPROPERTY(config, EditAnywhere, AdvancedDisplay, Category = "Animation Settings", meta = (EditCondition = "!bHideControlShapes"))
+	UPROPERTY(config, EditAnywhere, Category = "Animation Settings", meta = (EditCondition = "!bHideControlShapes"))
 	FLinearColor DrivenControlColor;
 
 	/** Should we show axes for the selected elements */
-	UPROPERTY(config, EditAnywhere, AdvancedDisplay, Category = "Animation Settings")
+	UPROPERTY(config, EditAnywhere, Category = "Animation Settings")
 	bool bDisplayAxesOnSelection;
 
 	/** The scale for axes to draw on the selection */
-	UPROPERTY(config, EditAnywhere, AdvancedDisplay, Category = "Animation Settings")
+	UPROPERTY(config, EditAnywhere, Category = "Animation Settings")
 	float AxisScale;
 
 	/** If true we restore the coordinate space when changing Widget Modes in the Viewport*/
-	UPROPERTY(config, EditAnywhere, AdvancedDisplay, Category = "Animation Settings")
+	UPROPERTY(config, EditAnywhere, Category = "Animation Settings")
 	bool bCoordSystemPerWidgetMode;
 
 	/** If true we can only select Rig Controls in the scene not other Actors. */
-	UPROPERTY(config, EditAnywhere, AdvancedDisplay, Category = "Animation Settings")
+	UPROPERTY(config, EditAnywhere, Category = "Animation Settings")
 	bool bOnlySelectRigControls;
 
 	/** If true when we transform multiple selected objects in the viewport they each transforms along their own local transform space */
-	UPROPERTY(config, EditAnywhere, AdvancedDisplay, Category = "Animation Settings")
+	UPROPERTY(config, EditAnywhere, Category = "Animation Settings")
 	bool bLocalTransformsInEachLocalSpace;
 	
 	/** The scale for Gizmos */
-	UPROPERTY(config, EditAnywhere, AdvancedDisplay, Category = "Animation Settings")
+	UPROPERTY(config, EditAnywhere, Category = "Animation Settings")
 	float GizmoScale;
 
 	UPROPERTY(config)
