Sequencer: Fix ClearOutConstraintDelegates when there is no sequence

Fix a crash after recording with Take Recorder

diff --git a/Engine/Plugins/Animation/ControlRig/Source/ControlRigEditor/Private/Sequencer/ControlRigParameterTrackEditor.cpp b/Engine/Plugins/Animation/ControlRig/Source/ControlRigEditor/Private/Sequencer/ControlRigParameterTrackEditor.cpp
--- a/Engine/Plugins/Animation/ControlRig/Source/ControlRigEditor/Private/Sequencer/ControlRigParameterTrackEditor.cpp
+++ b/Engine/Plugins/Animation/ControlRig/Source/ControlRigEditor/Private/Sequencer/ControlRigParameterTrackEditor.cpp
@@ -2514,40 +2514,42 @@ void FControlRigParameterTrackEditor::HandleSpaceKeyMoved(
 
 void FControlRigParameterTrackEditor::ClearOutAllSpaceAndConstraintDelegates(const UControlRig* InOptionalControlRig) const
 {
-	if (GetSequencer().IsValid())
+	const UMovieScene* MovieScene = GetSequencer().IsValid() && GetSequencer()->GetFocusedMovieSceneSequence() ? GetSequencer()->GetFocusedMovieSceneSequence()->GetMovieScene() : nullptr;
+	if (!MovieScene)
 	{
-		const UMovieScene* MovieScene = GetSequencer()->GetFocusedMovieSceneSequence()->GetMovieScene();
-		const TArray<FMovieSceneBinding>& Bindings = MovieScene->GetBindings();
-		for (const FMovieSceneBinding& Binding : Bindings)
+		return;
+	}
+
+	const TArray<FMovieSceneBinding>& Bindings = MovieScene->GetBindings();
+	for (const FMovieSceneBinding& Binding : Bindings)
+	{
+		const UMovieSceneTrack* Track = MovieScene->FindTrack(
+			UMovieSceneControlRigParameterTrack::StaticClass(), Binding.GetObjectGuid(), NAME_None);
+		if (const UMovieSceneControlRigParameterTrack* CRTrack = Cast<UMovieSceneControlRigParameterTrack>(Track))
 		{
-			const UMovieSceneTrack* Track = MovieScene->FindTrack(
-				UMovieSceneControlRigParameterTrack::StaticClass(), Binding.GetObjectGuid(), NAME_None);
-			if (const UMovieSceneControlRigParameterTrack* CRTrack = Cast<UMovieSceneControlRigParameterTrack>(Track))
+			if (InOptionalControlRig && CRTrack->GetControlRig() != InOptionalControlRig)
 			{
-				if (InOptionalControlRig && CRTrack->GetControlRig() != InOptionalControlRig)
-				{
-					continue;
-				}
+				continue;
+			}
 				
-				for (UMovieSceneSection* Section : Track->GetAllSections())
+			for (UMovieSceneSection* Section : Track->GetAllSections())
+			{
+				if (UMovieSceneControlRigParameterSection* CRSection = Cast<UMovieSceneControlRigParameterSection>(Section))
 				{
-					if (UMovieSceneControlRigParameterSection* CRSection = Cast<UMovieSceneControlRigParameterSection>(Section))
+					// clear space channels
+					TArray<FSpaceControlNameAndChannel>& Channels = CRSection->GetSpaceChannels();
+					for (FSpaceControlNameAndChannel& SpaceAndChannel : Channels)
 					{
-						// clear space channels
-						TArray<FSpaceControlNameAndChannel>& Channels = CRSection->GetSpaceChannels();
-						for (FSpaceControlNameAndChannel& SpaceAndChannel : Channels)
-						{
-							SpaceAndChannel.SpaceCurve.OnKeyMovedEvent().Clear();
-							SpaceAndChannel.SpaceCurve.OnKeyDeletedEvent().Clear();
-						}
+						SpaceAndChannel.SpaceCurve.OnKeyMovedEvent().Clear();
+						SpaceAndChannel.SpaceCurve.OnKeyDeletedEvent().Clear();
+					}
 
-						// clear constraint channels
-						TArray<FConstraintAndActiveChannel>& ConstraintChannels = CRSection->GetConstraintsChannels();
-						for (FConstraintAndActiveChannel& Channel: ConstraintChannels)
-						{
-							Channel.ActiveChannel.OnKeyMovedEvent().Clear();
-							Channel.ActiveChannel.OnKeyDeletedEvent().Clear();							
-						}
+					// clear constraint channels
+					TArray<FConstraintAndActiveChannel>& ConstraintChannels = CRSection->GetConstraintsChannels();
+					for (FConstraintAndActiveChannel& Channel: ConstraintChannels)
+					{
+						Channel.ActiveChannel.OnKeyMovedEvent().Clear();
+						Channel.ActiveChannel.OnKeyDeletedEvent().Clear();							
 					}
 				}
 			}
diff --git a/Engine/Source/Editor/MovieSceneTools/Private/TrackEditors/TransformTrackEditor.cpp b/Engine/Source/Editor/MovieSceneTools/Private/TrackEditors/TransformTrackEditor.cpp
--- a/Engine/Source/Editor/MovieSceneTools/Private/TrackEditors/TransformTrackEditor.cpp
+++ b/Engine/Source/Editor/MovieSceneTools/Private/TrackEditors/TransformTrackEditor.cpp
@@ -75,6 +75,7 @@ F3DTransformTrackEditor::F3DTransformTrackEditor( TSharedRef<ISequencer> InSeque
 		GEditor->RegisterForUndo(this);
 	}
 }
+
 F3DTransformTrackEditor::~F3DTransformTrackEditor()
 {
 	if (GEditor != nullptr)
@@ -106,7 +107,6 @@ void F3DTransformTrackEditor::OnRelease()
 	}
 }
 
-
 TSharedRef<ISequencerTrackEditor> F3DTransformTrackEditor::CreateTrackEditor( TSharedRef<ISequencer> InSequencer )
 {
 	return MakeShareable( new F3DTransformTrackEditor( InSequencer ) );
@@ -1554,6 +1554,7 @@ void F3DTransformTrackEditor::HandleConstraintKeyDeleted(IMovieSceneConstrainedS
 		}
 	}
 }
+
 void F3DTransformTrackEditor::HandleConstraintKeyMoved(IMovieSceneConstrainedSection* InSection, const FMovieSceneConstraintChannel* InConstraintChannel,
 	const TArray<FKeyMoveEventItem>& InMovedItems)
 {
@@ -1606,34 +1607,33 @@ void F3DTransformTrackEditor::HandleConstraintRemoved(IMovieSceneConstrainedSect
 					InSection->RemoveConstraintChannel(InConstraintName);
 				});
 	}
-	
 }
 
-
 void F3DTransformTrackEditor::ClearOutConstraintDelegates() const
 {
-	if (GetSequencer().IsValid())
+	const UMovieScene* MovieScene = GetSequencer().IsValid() && GetSequencer()->GetFocusedMovieSceneSequence() ? GetSequencer()->GetFocusedMovieSceneSequence()->GetMovieScene() : nullptr;
+	if (!MovieScene)
+	{
+		return;
+	}
+
+	const TArray<FMovieSceneBinding>& Bindings = MovieScene->GetBindings();
+	for (const FMovieSceneBinding& Binding : Bindings)
 	{
-		const UMovieScene* MovieScene = GetSequencer()->GetFocusedMovieSceneSequence()->GetMovieScene();
-		const TArray<FMovieSceneBinding>& Bindings = MovieScene->GetBindings();
-		for (const FMovieSceneBinding& Binding : Bindings)
+		const UMovieSceneTrack* Track = MovieScene->FindTrack(
+			UMovieScene3DTransformTrack::StaticClass(), Binding.GetObjectGuid(), NAME_None);
+		if (const UMovieScene3DTransformTrack* CRTrack = Cast<UMovieScene3DTransformTrack>(Track))
 		{
-			const UMovieSceneTrack* Track = MovieScene->FindTrack(
-				UMovieScene3DTransformTrack::StaticClass(), Binding.GetObjectGuid(), NAME_None);
-			if (const UMovieScene3DTransformTrack* CRTrack = Cast<UMovieScene3DTransformTrack>(Track))
+			for (UMovieSceneSection* Section : Track->GetAllSections())
 			{
-
-				for (UMovieSceneSection* Section : Track->GetAllSections())
+				if (UMovieScene3DTransformSection* CRSection = Cast<UMovieScene3DTransformSection>(Section))
 				{
-					if (UMovieScene3DTransformSection* CRSection = Cast<UMovieScene3DTransformSection>(Section))
+					// clear constraint channels
+					TArray<FConstraintAndActiveChannel>& ConstraintChannels = CRSection->GetConstraintsChannels();
+					for (FConstraintAndActiveChannel& Channel : ConstraintChannels)
 					{
-						// clear constraint channels
-						TArray<FConstraintAndActiveChannel>& ConstraintChannels = CRSection->GetConstraintsChannels();
-						for (FConstraintAndActiveChannel& Channel : ConstraintChannels)
-						{
-							Channel.ActiveChannel.OnKeyMovedEvent().Clear();
-							Channel.ActiveChannel.OnKeyDeletedEvent().Clear();
-						}
+						Channel.ActiveChannel.OnKeyMovedEvent().Clear();
+						Channel.ActiveChannel.OnKeyDeletedEvent().Clear();
 					}
 				}
 			}
@@ -1641,8 +1641,4 @@ void F3DTransformTrackEditor::ClearOutConstraintDelegates() const
 	}
 }
 
-
-
-	
-
 #undef LOCTEXT_NAMESPACE
