Control Rig: Avoid dirtying assets on load when variables change from float/double

diff --git a/Engine/Plugins/Animation/ControlRig/Source/ControlRigDeveloper/Private/ControlRigBlueprint.cpp b/Engine/Plugins/Animation/ControlRig/Source/ControlRigDeveloper/Private/ControlRigBlueprint.cpp
--- a/Engine/Plugins/Animation/ControlRig/Source/ControlRigDeveloper/Private/ControlRigBlueprint.cpp
+++ b/Engine/Plugins/Animation/ControlRig/Source/ControlRigDeveloper/Private/ControlRigBlueprint.cpp
@@ -4256,6 +4256,10 @@ void UControlRigBlueprint::PatchVariableNodesWithIncorrectType()
 				if (ValuePin->GetCPPType() != CPPType || ValuePin->GetCPPTypeObject() != CPPTypeObject)
 				{
 					Controller->RefreshVariableNode(VariableNode->GetFName(), VariableNode->GetVariableName(), CPPType, CPPTypeObject, false);
+					if (RigVMTypeUtils::AreCompatible(*ValuePin->GetCPPType(), ValuePin->GetCPPTypeObject(), *CPPType, CPPTypeObject))
+					{
+						return false;
+					}
 					return true;
 				}
 			}
