ControlRig: Animation Atttribute Node, support for user defined structs

diff --git a/Engine/Plugins/Animation/ControlRig/Source/ControlRig/Private/Units/Animation/RigUnit_AnimAttribute.cpp b/Engine/Plugins/Animation/ControlRig/Source/ControlRig/Private/Units/Animation/RigUnit_AnimAttribute.cpp
--- a/Engine/Plugins/Animation/ControlRig/Source/ControlRig/Private/Units/Animation/RigUnit_AnimAttribute.cpp
+++ b/Engine/Plugins/Animation/ControlRig/Source/ControlRig/Private/Units/Animation/RigUnit_AnimAttribute.cpp
@@ -59,6 +59,14 @@ bool FRigDispatch_AnimAttributeBase::IsTypeSupported(const TRigVMTypeIndex& InTy
 				return false;
 			}
 
+			if (UUserDefinedStruct* UserDefinedStruct = Cast<UUserDefinedStruct>(ScriptStruct))
+			{
+				// allow all user defined structs because even if a struct is not registered with anim attribute system,
+				// it could be added to or removed from the system easily. allowing all of them as valid permutations
+				// avoids having to create orphan pins.
+				return true;
+			}
+			
 			return UE::Anim::AttributeTypes::IsTypeRegistered(ScriptStruct);
 		}
 	}
@@ -240,8 +248,8 @@ void FRigDispatch_SetAnimAttribute::SetAnimAttributeDispatch(FRigVMExtendedExecu
 			if (!UE::Anim::AttributeTypes::IsTypeRegistered(ScriptStruct))
 			{
 				UE_CONTROLRIG_RIGUNIT_REPORT_WARNING(
-					TEXT("Type: '%s' is not registered with the Animation Attribute System. "
-						"Please register the type in Project Settings - Animation - CustomAttributes - User Defined Struct Attributes."),
+					TEXT("Type: '%s' is not registered in the Animation Attribute System. "
+						"Please register the type in Project Settings - Animation - CustomAttributes - User Defined Struct Animation Attributes."),
 					*ScriptStruct->GetAuthoredName());
 				Registered = false;
 			}
@@ -394,7 +402,7 @@ FRigVMTemplateTypeMap FRigDispatch_GetAnimAttribute::OnNewArgumentType(const FNa
 	// similar pattern to URigVMArrayNode's FRigVMTemplate_NewArgumentTypeDelegate to avoid double registration
 	// this is needed since a single type is called for both DefaultArg and ValueArg but we should only
 	// register one permutation
-	if (InArgumentName == DefaultArgName)
+	if (InArgumentName == ValueArgName)
 	{
 		if (IsTypeSupported(InTypeIndex))
 		{
diff --git a/Engine/Source/Runtime/RigVM/Private/RigVMCore/RigVMRegistry.cpp b/Engine/Source/Runtime/RigVM/Private/RigVMCore/RigVMRegistry.cpp
--- a/Engine/Source/Runtime/RigVM/Private/RigVMCore/RigVMRegistry.cpp
+++ b/Engine/Source/Runtime/RigVM/Private/RigVMCore/RigVMRegistry.cpp
@@ -135,10 +135,14 @@ void FRigVMRegistry::InitializeIfNeeded()
 				AssetRegistryModule->Get().OnAssetRemoved().RemoveAll(&s_RigVMRegistry);
 			}
 		}
+
+		UE::Anim::AttributeTypes::GetOnAttributeTypesChanged().RemoveAll(&s_RigVMRegistry);
 	});
 	
 	FAssetRegistryModule& AssetRegistryModule = FModuleManager::LoadModuleChecked<FAssetRegistryModule>(TEXT("AssetRegistry"));
 	AssetRegistryModule.Get().OnAssetRemoved().AddRaw(this, &FRigVMRegistry::OnAssetRemoved);
+
+	UE::Anim::AttributeTypes::GetOnAttributeTypesChanged().AddRaw(this, &FRigVMRegistry::OnAnimationAttributeTypesChanged);
 }
 
 void FRigVMRegistry::Refresh()
@@ -217,6 +221,21 @@ void FRigVMRegistry::OnAssetRemoved(const FAssetData& InAssetData)
 	}
 }
 
+void FRigVMRegistry::OnAnimationAttributeTypesChanged(const UScriptStruct* InStruct, bool bIsAdded)
+{
+	if (!ensure(InStruct))
+	{
+		return;
+	}
+
+	if (bIsAdded)
+	{
+		FindOrAddType(FRigVMTemplateArgumentType(const_cast<UScriptStruct*>(InStruct)));
+		OnRigVMRegistryChangedDelegate.Broadcast();		
+	}
+}
+
+
 void FRigVMRegistry::Reset()
 {
 	for(FRigVMDispatchFactory* Factory : Factories)
diff --git a/Engine/Source/Runtime/RigVM/Public/RigVMCore/RigVMRegistry.h b/Engine/Source/Runtime/RigVM/Public/RigVMCore/RigVMRegistry.h
--- a/Engine/Source/Runtime/RigVM/Public/RigVMCore/RigVMRegistry.h
+++ b/Engine/Source/Runtime/RigVM/Public/RigVMCore/RigVMRegistry.h
@@ -62,6 +62,10 @@ public:
 	// Update the registry when old types are removed
     void OnAssetRemoved(const FAssetData& InAssetData);
 
+	// Update the registry when new types are added to the attribute system so that they can be selected
+	// on Attribute Nodes
+	void OnAnimationAttributeTypesChanged(const UScriptStruct* InStruct, bool bIsAdded);
+	
 	// Clear the registry
 	void Reset();
 
