Control Rig: Fix duplicating graph with aggregate graph breaking links

diff --git a/Engine/Plugins/Animation/ControlRig/Source/ControlRigDeveloper/Private/ControlRigBlueprint.cpp b/Engine/Plugins/Animation/ControlRig/Source/ControlRigDeveloper/Private/ControlRigBlueprint.cpp
--- a/Engine/Plugins/Animation/ControlRig/Source/ControlRigDeveloper/Private/ControlRigBlueprint.cpp
+++ b/Engine/Plugins/Animation/ControlRig/Source/ControlRigDeveloper/Private/ControlRigBlueprint.cpp
@@ -1442,6 +1442,7 @@ void UControlRigBlueprint::RefreshAllModels(EControlRigBlueprintLoadType InLoadT
 		URigVMGraph* GraphToClean = SortedGraphsToClean[GraphIndex];
 		URigVMController* Controller = GetOrCreateController(GraphToClean);
 		TGuardValue<bool> RecomputeGuard(Controller->bSuspendRecomputingOuterTemplateFilters, true);
+		TGuardValue<bool> GuardEditGraph(GraphToClean->bEditable, true);
 		Controller->SuspendNotifications(true);
 		
 		for(URigVMNode* ModelNode : GraphToClean->GetNodes())
