ControlRig: fix crash when user defined struct is structurally changed, wrapper object class needs to be invalidated if it reference a user defined struct

diff --git a/Engine/Plugins/Animation/ControlRig/Source/ControlRigEditor/Private/ControlRigEditorModule.cpp b/Engine/Plugins/Animation/ControlRig/Source/ControlRigEditor/Private/ControlRigEditorModule.cpp
--- a/Engine/Plugins/Animation/ControlRig/Source/ControlRigEditor/Private/ControlRigEditorModule.cpp
+++ b/Engine/Plugins/Animation/ControlRig/Source/ControlRigEditor/Private/ControlRigEditorModule.cpp
@@ -2568,19 +2568,24 @@ void FControlRigEditorModule::PreChange(const UUserDefinedStruct* Changed,
 			{
 				// make sure variable properties on the BP is patched
 				// since active rig instance still references it
-				if (UControlRigBlueprintGeneratedClass* OwnerClass = Cast<UControlRigBlueprintGeneratedClass>(StructProperty->GetOwnerClass()))
+				if (UControlRigBlueprintGeneratedClass* BPClass = Cast<UControlRigBlueprintGeneratedClass>(StructProperty->GetOwnerClass()))
 				{
-					if (UControlRigBlueprint* FoundBlueprint = Cast<UControlRigBlueprint>(OwnerClass->ClassGeneratedBy))
+					if (UControlRigBlueprint* FoundBlueprint = Cast<UControlRigBlueprint>(BPClass->ClassGeneratedBy))
 					{
 						StructProperty->Struct = DuplicatedStruct;
-						StructsToRegenerateReferencesFor.Add(OwnerClass);
+						StructsToRegenerateReferencesFor.Add(BPClass);
 					}
 				}
 				// similar story, VM instructions reference properties on the GeneratorClass
-				if (URigVMMemoryStorageGeneratorClass* OwnerClass = Cast<URigVMMemoryStorageGeneratorClass>(StructProperty->GetOwnerStruct()))
+				else if (URigVMMemoryStorageGeneratorClass* MemoryClass = Cast<URigVMMemoryStorageGeneratorClass>(StructProperty->GetOwnerStruct()))
 				{
 					StructProperty->Struct = DuplicatedStruct;
-					StructsToRegenerateReferencesFor.Add(OwnerClass);
+					StructsToRegenerateReferencesFor.Add(MemoryClass);
+				}
+				else if (UDetailsViewWrapperObject::IsValidClass(StructProperty->GetOwnerClass()))
+				{
+					StructProperty->Struct = DuplicatedStruct;
+					StructsToRegenerateReferencesFor.Add(StructProperty->GetOwnerClass());
 				}
 			}
 		}
@@ -2596,6 +2601,10 @@ void FControlRigEditorModule::PreChange(const UUserDefinedStruct* Changed,
 				GeneratorClass->RefreshLinkedProperties();
 				GeneratorClass->RefreshPropertyPaths();	
 			}
+			else if (Struct->IsChildOf(UDetailsViewWrapperObject::StaticClass()))
+			{
+				UDetailsViewWrapperObject::MarkOutdatedClass(Cast<UClass>(Struct));
+			}
 		}
 
 		// as rigs are re-instanced, the duplicated struct will be GCed
@@ -2652,10 +2661,7 @@ void FControlRigEditorModule::PostChange(const UUserDefinedStruct* Changed,
 
 	for (UControlRigBlueprint* RigBlueprint : BlueprintsToRefresh)
 	{
-		// refresh all pins
-		RigBlueprint->RefreshAllModels();
-		// reflect changes in the editor
-		RigBlueprint->RebuildGraphFromModel();
+		RigBlueprint->OnRigVMRegistryChanged();
 		RigBlueprint->MarkPackageDirty();
 	}
 	
diff --git a/Engine/Plugins/Animation/ControlRig/Source/ControlRigEditor/Private/Editor/ControlRigEditor.cpp b/Engine/Plugins/Animation/ControlRig/Source/ControlRigEditor/Private/Editor/ControlRigEditor.cpp
--- a/Engine/Plugins/Animation/ControlRig/Source/ControlRigEditor/Private/Editor/ControlRigEditor.cpp
+++ b/Engine/Plugins/Animation/ControlRig/Source/ControlRigEditor/Private/Editor/ControlRigEditor.cpp
@@ -1707,7 +1707,28 @@ bool FControlRigEditor::DetailViewShowsAnyRigElement() const
 
 bool FControlRigEditor::DetailViewShowsAnyRigUnit() const
 {
-	return DetailViewShowsStruct(FRigUnit::StaticStruct());
+	if (DetailViewShowsStruct(FRigUnit::StaticStruct()))
+	{
+		return true;
+	}
+
+	const TArray< TWeakObjectPtr<UObject> >& SelectedObjects = Inspector->GetSelectedObjects();
+	for (TWeakObjectPtr<UObject> SelectedObject : SelectedObjects)
+	{
+		if (SelectedObject.IsValid())
+		{
+			if(UDetailsViewWrapperObject* WrapperObject = Cast<UDetailsViewWrapperObject>(SelectedObject.Get()))
+			{
+				const FString Notation = WrapperObject->GetWrappedNodeNotation();
+				if(!Notation.IsEmpty())
+				{
+					return true;
+				}
+			}
+		}
+	}
+	
+	return false;
 }
 
 bool FControlRigEditor::DetailViewShowsLocalVariable() const
@@ -1762,25 +1783,6 @@ bool FControlRigEditor::DetailViewShowsRigElement(FRigElementKey InKey) const
 	return false;
 }
 
-bool FControlRigEditor::DetailViewShowsRigUnit(URigVMNode* InNode) const
-{
-	TArray< TWeakObjectPtr<UObject> > SelectedObjects = Inspector->GetSelectedObjects();
-	for (TWeakObjectPtr<UObject> SelectedObject : SelectedObjects)
-	{
-		if (SelectedObject.IsValid())
-		{
-			if(UDetailsViewWrapperObject* WrapperObject = Cast<UDetailsViewWrapperObject>(SelectedObject.Get()))
-			{
-				if(WrapperObject->GetOuter() == InNode)
-				{
-					return true;
-				}
-			}
-		}
-	}
-	return false;
-}
-
 void FControlRigEditor::ClearDetailObject(bool bChangeUISelectionState)
 {
 	if(bSuspendDetailsPanelRefresh)
diff --git a/Engine/Plugins/Animation/ControlRig/Source/ControlRigEditor/Private/Editor/ControlRigEditor.h b/Engine/Plugins/Animation/ControlRig/Source/ControlRigEditor/Private/Editor/ControlRigEditor.h
--- a/Engine/Plugins/Animation/ControlRig/Source/ControlRigEditor/Private/Editor/ControlRigEditor.h
+++ b/Engine/Plugins/Animation/ControlRig/Source/ControlRigEditor/Private/Editor/ControlRigEditor.h
@@ -160,7 +160,6 @@ public:
 	bool DetailViewShowsLocalVariable() const;
 	bool DetailViewShowsStruct(UScriptStruct* InStruct) const;
 	bool DetailViewShowsRigElement(FRigElementKey InKey) const;
-	bool DetailViewShowsRigUnit(URigVMNode* InNode) const;
 
 	void ClearDetailObject(bool bChangeUISelectionState = true);
 
diff --git a/Engine/Plugins/Animation/ControlRig/Source/ControlRigEditor/Private/Editor/DetailsViewWrapperObject.cpp b/Engine/Plugins/Animation/ControlRig/Source/ControlRigEditor/Private/Editor/DetailsViewWrapperObject.cpp
--- a/Engine/Plugins/Animation/ControlRig/Source/ControlRigEditor/Private/Editor/DetailsViewWrapperObject.cpp
+++ b/Engine/Plugins/Animation/ControlRig/Source/ControlRigEditor/Private/Editor/DetailsViewWrapperObject.cpp
@@ -14,6 +14,9 @@
 
 TMap<UDetailsViewWrapperObject::FPerClassInfo, UClass*> UDetailsViewWrapperObject::InfoToClass;
 TMap<UClass*, UDetailsViewWrapperObject::FPerClassInfo> UDetailsViewWrapperObject::ClassToInfo;
+TSet<UClass*> UDetailsViewWrapperObject::OutdatedClassToRecreate;
+
+static const TCHAR DiscardedWrapperClassTemplateName[] = TEXT("DiscardedWrapperClassTemplate");
 
 UDetailsViewWrapperObject::UDetailsViewWrapperObject()
 	: bIsSettingValue(false)
@@ -26,7 +29,35 @@ UClass* UDetailsViewWrapperObject::GetClassForStruct(UScriptStruct* InStruct, bo
 
 	if(UClass** ExistingClass = InfoToClass.Find(InStruct))
 	{
-		return *ExistingClass;
+		// if properties in a class has changed (due to user defined struct change)
+		// we have to recreate the class
+		if (OutdatedClassToRecreate.Contains(*ExistingClass))
+		{
+			OutdatedClassToRecreate.Remove(*ExistingClass);
+
+			InfoToClass.Remove(InStruct);
+			ClassToInfo.Remove(*ExistingClass);
+
+			FString DiscardedWrapperClassName;
+			static int32 DiscardedWrapperClassIndex = 0;
+			do
+			{
+				DiscardedWrapperClassName = FString::Printf(TEXT("%s_%d"), DiscardedWrapperClassTemplateName, DiscardedWrapperClassIndex++);
+				if(StaticFindObjectFast(nullptr, GetTransientPackage(), *DiscardedWrapperClassName) == nullptr)
+				{
+					break;
+				}
+			}
+			while (DiscardedWrapperClassIndex < INT_MAX);
+
+			(*ExistingClass)->Rename(*DiscardedWrapperClassName, GetTransientPackage(), REN_ForceNoResetLoaders | REN_DoNotDirty | REN_DontCreateRedirectors | REN_NonTransactional);
+			
+			(*ExistingClass)->RemoveFromRoot();
+		}
+		else
+		{
+			return *ExistingClass;
+		}
 	}
 
 	if(!bCreateIfNeeded)
@@ -280,7 +311,35 @@ UClass* UDetailsViewWrapperObject::GetClassForNodes(TArray<URigVMNode*> InNodes,
 	const FPerClassInfo PerClassInfo(Notation);
 	if(UClass** ExistingClass = InfoToClass.Find(PerClassInfo))
 	{
-		return *ExistingClass;
+		// if properties in a class has changed (due to user defined struct change)
+		// we have to recreate the class
+		if (OutdatedClassToRecreate.Contains(*ExistingClass))
+		{
+			OutdatedClassToRecreate.Remove(*ExistingClass);
+
+			InfoToClass.Remove(PerClassInfo);
+			ClassToInfo.Remove(*ExistingClass);
+
+			FString DiscardedWrapperClassName;
+			static int32 DiscardedWrapperClassIndex = 0;
+			do
+			{
+				DiscardedWrapperClassName = FString::Printf(TEXT("%s_%d"), DiscardedWrapperClassTemplateName, DiscardedWrapperClassIndex++);
+				if(StaticFindObjectFast(nullptr, GetTransientPackage(), *DiscardedWrapperClassName) == nullptr)
+				{
+					break;
+				}
+			}
+			while (DiscardedWrapperClassIndex < INT_MAX);
+
+			(*ExistingClass)->Rename(*DiscardedWrapperClassName, GetTransientPackage(), REN_ForceNoResetLoaders | REN_DoNotDirty | REN_DontCreateRedirectors | REN_NonTransactional);
+			
+			(*ExistingClass)->RemoveFromRoot();
+		}
+		else
+		{
+			return *ExistingClass;
+		}
 	}
 
 	if(!bCreateIfNeeded)
@@ -469,6 +528,34 @@ UDetailsViewWrapperObject* UDetailsViewWrapperObject::MakeInstance(TArray<URigVM
 
 }
 
+void UDetailsViewWrapperObject::MarkOutdatedClass(UClass* InClass)
+{
+	if (InClass)
+	{
+		OutdatedClassToRecreate.Add(InClass);
+	}
+}
+
+bool UDetailsViewWrapperObject::IsValidClass(UClass* InClass)
+{
+	if (!InClass)
+	{
+		return false;
+	}
+	
+	if (!InClass->IsChildOf(StaticClass()))
+	{
+		return false;
+	}
+		
+	if (InClass->GetName().Contains(DiscardedWrapperClassTemplateName))
+	{
+		return false;
+	}
+
+	return true;
+}
+
 FString UDetailsViewWrapperObject::GetWrappedNodeNotation() const
 {
 	return ClassToInfo.FindChecked(GetClass()).Notation;
diff --git a/Engine/Plugins/Animation/ControlRig/Source/ControlRigEditor/Private/Editor/DetailsViewWrapperObject.h b/Engine/Plugins/Animation/ControlRig/Source/ControlRigEditor/Private/Editor/DetailsViewWrapperObject.h
--- a/Engine/Plugins/Animation/ControlRig/Source/ControlRigEditor/Private/Editor/DetailsViewWrapperObject.h
+++ b/Engine/Plugins/Animation/ControlRig/Source/ControlRigEditor/Private/Editor/DetailsViewWrapperObject.h
@@ -27,6 +27,10 @@ public:
 	// Creating wrappers from a RigVMNode
 	static UClass* GetClassForNodes(TArray<URigVMNode*> InNodes, bool bCreateIfNeeded = true);
 	static UDetailsViewWrapperObject* MakeInstance(TArray<URigVMNode*> InNodes, URigVMNode* InSubject, UObject* InOuter = nullptr);
+	
+	static void MarkOutdatedClass(UClass* InClass);
+	static bool IsValidClass(UClass* InClass);
+	
 	FString GetWrappedNodeNotation() const;
 	
 	bool IsChildOf(const UStruct* InStruct) const
@@ -110,6 +114,7 @@ private:
 	
 	static TMap<FPerClassInfo, UClass*> InfoToClass;
 	static TMap<UClass*, FPerClassInfo> ClassToInfo;
+	static TSet<UClass*> OutdatedClassToRecreate;
 	bool bIsSettingValue;
 	
 	FWrappedPropertyChangedChainEvent WrappedPropertyChangedChainEvent;
