OptimusEditor: patch outdated OptimusValueContainer class so that we don't have a mismatch during GC

+Refactor ReplaceStructWithTempDuplicate to be a general API that can be used to patch all kinds of generated class that references a User defined struct

diff --git a/Engine/Plugins/Animation/ControlRig/Source/ControlRigEditor/ControlRigEditor.Build.cs b/Engine/Plugins/Animation/ControlRig/Source/ControlRigEditor/ControlRigEditor.Build.cs
--- a/Engine/Plugins/Animation/ControlRig/Source/ControlRigEditor/ControlRigEditor.Build.cs
+++ b/Engine/Plugins/Animation/ControlRig/Source/ControlRigEditor/ControlRigEditor.Build.cs
@@ -39,6 +39,7 @@ namespace UnrealBuildTool.Rules
                     "ControlRig",
                     "ControlRigDeveloper",
                     "Kismet",
+					"KismetCompiler",
                     "EditorStyle",
 					"EditorWidgets",
                     "ApplicationCore",
diff --git a/Engine/Plugins/Animation/ControlRig/Source/ControlRigEditor/Private/ControlRigEditorModule.cpp b/Engine/Plugins/Animation/ControlRig/Source/ControlRigEditor/Private/ControlRigEditorModule.cpp
--- a/Engine/Plugins/Animation/ControlRig/Source/ControlRigEditor/Private/ControlRigEditorModule.cpp
+++ b/Engine/Plugins/Animation/ControlRig/Source/ControlRigEditor/Private/ControlRigEditorModule.cpp
@@ -130,6 +130,7 @@
 #include "AnimationToolMenuContext.h"
 #include "RigVMModel/Nodes/RigVMAggregateNode.h"
 #include "RigVMUserWorkflowRegistry.h"
+#include "UserDefinedStructureCompilerUtils.h"
 #include "Units/ControlRigNodeWorkflow.h"
 #include "Units/RigDispatchFactory.h"
 
@@ -2537,79 +2538,45 @@ void FControlRigEditorModule::PreChange(const UUserDefinedStruct* Changed,
 	//                              --2--> PreBPCompile --3--> PostBPCompile
 	
 	UUserDefinedStruct* StructureToReinstance = (UUserDefinedStruct*)Changed;
-	if (StructureToReinstance)
-	{
-		UUserDefinedStruct* DuplicatedStruct = NULL;
-		{
-			const FString ReinstancedName = FString::Printf(TEXT("STRUCT_REINST_%s"), *StructureToReinstance->GetName());
-			const FName UniqueName = MakeUniqueObjectName(GetTransientPackage(), UUserDefinedStruct::StaticClass(), FName(*ReinstancedName));
-
-			TGuardValue<bool> IsDuplicatingClassForReinstancing(GIsDuplicatingClassForReinstancing, true);
-			DuplicatedStruct = (UUserDefinedStruct*)StaticDuplicateObject(StructureToReinstance, GetTransientPackage(), UniqueName, ~RF_Transactional); 
-		}
-
-		DuplicatedStruct->Guid = StructureToReinstance->Guid;
-		DuplicatedStruct->Bind();
-		DuplicatedStruct->StaticLink(true);
-		DuplicatedStruct->PrimaryStruct = StructureToReinstance;
-		DuplicatedStruct->Status = EUserDefinedStructureStatus::UDSS_Duplicate;
-		DuplicatedStruct->SetFlags(RF_Transient);
-		DuplicatedStruct->AddToRoot();
-
-		CastChecked<UUserDefinedStructEditorData>(DuplicatedStruct->EditorData)->RecreateDefaultInstance();
 
-		// List of unique classes and structs to regenerate bytecode and property referenced objects list
-		TSet<UStruct*> StructsToRegenerateReferencesFor;
-
-		for (TAllFieldsIterator<FStructProperty> FieldIt(RF_NoFlags, EInternalObjectFlags::Garbage); FieldIt; ++FieldIt)
+	FUserDefinedStructureCompilerUtils::ReplaceStructWithTempDuplicateByPredicate(
+		StructureToReinstance,
+		[](FStructProperty* InStructProperty)
 		{
-			FStructProperty* StructProperty = *FieldIt;
-			if (StructProperty && (StructureToReinstance == StructProperty->Struct))
+			// make sure variable properties on the BP is patched
+			// since active rig instance still references it
+			if (UControlRigBlueprintGeneratedClass* BPClass = Cast<UControlRigBlueprintGeneratedClass>(InStructProperty->GetOwnerClass()))
 			{
-				// make sure variable properties on the BP is patched
-				// since active rig instance still references it
-				if (UControlRigBlueprintGeneratedClass* BPClass = Cast<UControlRigBlueprintGeneratedClass>(StructProperty->GetOwnerClass()))
-				{
-					if (UControlRigBlueprint* FoundBlueprint = Cast<UControlRigBlueprint>(BPClass->ClassGeneratedBy))
-					{
-						StructProperty->Struct = DuplicatedStruct;
-						StructsToRegenerateReferencesFor.Add(BPClass);
-					}
-				}
-				// similar story, VM instructions reference properties on the GeneratorClass
-				else if (URigVMMemoryStorageGeneratorClass* MemoryClass = Cast<URigVMMemoryStorageGeneratorClass>(StructProperty->GetOwnerStruct()))
-				{
-					StructProperty->Struct = DuplicatedStruct;
-					StructsToRegenerateReferencesFor.Add(MemoryClass);
-				}
-				else if (UDetailsViewWrapperObject::IsValidClass(StructProperty->GetOwnerClass()))
+				if (BPClass->ClassGeneratedBy->IsA<UControlRigBlueprint>())
 				{
-					StructProperty->Struct = DuplicatedStruct;
-					StructsToRegenerateReferencesFor.Add(StructProperty->GetOwnerClass());
+					return true;
 				}
 			}
-		}
-
-		// Make sure we update the list of objects referenced by structs after we replaced the struct in FStructProperties
-		for (UStruct* Struct : StructsToRegenerateReferencesFor)
-		{
-			Struct->CollectBytecodeAndPropertyReferencedObjects();
+			// similar story, VM instructions reference properties on the GeneratorClass
+			else if ((InStructProperty->GetOwnerStruct())->IsA<URigVMMemoryStorageGeneratorClass>())
+			{
+				return true;
+			}
+			else if (UDetailsViewWrapperObject::IsValidClass(InStructProperty->GetOwnerClass()))
+			{
+				return true;
+			}
 			
+			return false;
+		},
+		[](UStruct* InStruct)
+		{
 			// refresh these since VM caching references them
-			if (URigVMMemoryStorageGeneratorClass* GeneratorClass = Cast<URigVMMemoryStorageGeneratorClass>(Struct))
+			if (URigVMMemoryStorageGeneratorClass* GeneratorClass = Cast<URigVMMemoryStorageGeneratorClass>(InStruct))
 			{
 				GeneratorClass->RefreshLinkedProperties();
 				GeneratorClass->RefreshPropertyPaths();	
 			}
-			else if (Struct->IsChildOf(UDetailsViewWrapperObject::StaticClass()))
+			else if (InStruct->IsChildOf(UDetailsViewWrapperObject::StaticClass()))
 			{
-				UDetailsViewWrapperObject::MarkOutdatedClass(Cast<UClass>(Struct));
+				UDetailsViewWrapperObject::MarkOutdatedClass(Cast<UClass>(InStruct));
 			}
-		}
-
-		// as rigs are re-instanced, the duplicated struct will be GCed
-		DuplicatedStruct->RemoveFromRoot();
-	}
+		});
 	
 	// in the future we could only invalidate caches on affected rig instances, it shouldn't make too much of a difference though
 	for (TObjectIterator<UControlRig> It(RF_Transient | RF_ClassDefaultObject, /** bIncludeDerivedClasses */ true, /** InternalExcludeFlags */ EInternalObjectFlags::Garbage); It; ++It)
diff --git a/Engine/Plugins/Animation/DeformerGraph/Source/OptimusEditor/OptimusEditor.Build.cs b/Engine/Plugins/Animation/DeformerGraph/Source/OptimusEditor/OptimusEditor.Build.cs
--- a/Engine/Plugins/Animation/DeformerGraph/Source/OptimusEditor/OptimusEditor.Build.cs
+++ b/Engine/Plugins/Animation/DeformerGraph/Source/OptimusEditor/OptimusEditor.Build.cs
@@ -36,6 +36,7 @@ namespace UnrealBuildTool.Rules
 					"Persona",
 					"MessageLog",
 					"PropertyEditor",
+					"KismetCompiler"
 				}
 			);
 
diff --git a/Engine/Plugins/Animation/DeformerGraph/Source/OptimusEditor/Private/OptimusEditorModule.cpp b/Engine/Plugins/Animation/DeformerGraph/Source/OptimusEditor/Private/OptimusEditorModule.cpp
--- a/Engine/Plugins/Animation/DeformerGraph/Source/OptimusEditor/Private/OptimusEditorModule.cpp
+++ b/Engine/Plugins/Animation/DeformerGraph/Source/OptimusEditor/Private/OptimusEditorModule.cpp
@@ -30,11 +30,14 @@
 #include "OptimusValidatedName.h"
 #include "OptimusValueContainer.h"
 #include "PropertyEditorModule.h"
+#include "UserDefinedStructureCompilerUtils.h"
 #include "AssetRegistry/AssetRegistryModule.h"
 #include "Widgets/SOptimusEditorGraphExplorer.h"
 #include "Widgets/SOptimusShaderTextDocumentTextBox.h"
 #include "Algo/Transform.h"
 #include "OptimusCore/Private/Actions/OptimusResourceActions.h"
+#include "UObject/FieldIterator.h"
+#include "UserDefinedStructure/UserDefinedStructEditorData.h"
 
 #define LOCTEXT_NAMESPACE "OptimusEditorModule"
 
@@ -110,6 +113,18 @@ FOptimusEditorClipboard& FOptimusEditorModule::GetClipboard() const
 void FOptimusEditorModule::PreChange(const UUserDefinedStruct* Changed,
 	FStructureEditorUtils::EStructureEditorChangeInfo ChangedType)
 {
+	// the following is similar to
+	// FUserDefinedStructureCompilerInner::ReplaceStructWithTempDuplicate()
+	// it is necessary since things like optimus value container class/objects can reference the old struct
+	// and we need to keep the old struct alive in those classes so that they can get GCed properly.
+	
+	UUserDefinedStruct* StructureToReinstance = (UUserDefinedStruct*)Changed;
+
+	FUserDefinedStructureCompilerUtils::ReplaceStructWithTempDuplicateByPredicate(
+		StructureToReinstance,
+		[](FStructProperty* InStructProperty){ return Cast<UOptimusValueContainerGeneratorClass>(InStructProperty->GetOwnerClass()) != nullptr; },
+		[](UStruct*){});
+	
 	FAssetRegistryModule& AssetRegistryModule = FModuleManager::LoadModuleChecked<FAssetRegistryModule>("AssetRegistry");
 	TArray<FName> Referencers;
 	AssetRegistryModule.Get().GetReferencers(Changed->GetPackage()->GetFName(), Referencers);
diff --git a/Engine/Source/Editor/KismetCompiler/Private/UserDefinedStructureCompilerUtils.cpp b/Engine/Source/Editor/KismetCompiler/Private/UserDefinedStructureCompilerUtils.cpp
--- a/Engine/Source/Editor/KismetCompiler/Private/UserDefinedStructureCompilerUtils.cpp
+++ b/Engine/Source/Editor/KismetCompiler/Private/UserDefinedStructureCompilerUtils.cpp
@@ -501,4 +501,58 @@ void FUserDefinedStructureCompilerUtils::CompileStruct(class UUserDefinedStruct*
 	}
 }
 
+void FUserDefinedStructureCompilerUtils::ReplaceStructWithTempDuplicateByPredicate(
+	UUserDefinedStruct* StructureToReinstance,
+	TFunctionRef<bool(FStructProperty* InStructProperty)> ShouldReplaceStructInStructProperty,
+	TFunctionRef<void(UStruct* InStruct)> PostReplace)
+{
+	if (StructureToReinstance)
+	{
+		UUserDefinedStruct* DuplicatedStruct = NULL;
+		{
+			const FString ReinstancedName = FString::Printf(TEXT("STRUCT_REINST_%s"), *StructureToReinstance->GetName());
+			const FName UniqueName = MakeUniqueObjectName(GetTransientPackage(), UUserDefinedStruct::StaticClass(), FName(*ReinstancedName));
+
+			TGuardValue<bool> IsDuplicatingClassForReinstancing(GIsDuplicatingClassForReinstancing, true);
+			DuplicatedStruct = (UUserDefinedStruct*)StaticDuplicateObject(StructureToReinstance, GetTransientPackage(), UniqueName, ~RF_Transactional); 
+		}
+
+		DuplicatedStruct->Guid = StructureToReinstance->Guid;
+		DuplicatedStruct->Bind();
+		DuplicatedStruct->StaticLink(true);
+		DuplicatedStruct->PrimaryStruct = StructureToReinstance;
+		DuplicatedStruct->Status = EUserDefinedStructureStatus::UDSS_Duplicate;
+		DuplicatedStruct->SetFlags(RF_Transient);
+		DuplicatedStruct->AddToRoot();
+
+		CastChecked<UUserDefinedStructEditorData>(DuplicatedStruct->EditorData)->RecreateDefaultInstance();
+
+		// List of unique classes and structs to regenerate
+		TSet<UStruct*> StructsToRegenerateReferencesFor;
+
+		for (TAllFieldsIterator<FStructProperty> FieldIt(RF_NoFlags, EInternalObjectFlags::Garbage); FieldIt; ++FieldIt)
+		{
+			FStructProperty* StructProperty = *FieldIt;
+			if (StructProperty && (StructureToReinstance == StructProperty->Struct))
+			{
+				if(ShouldReplaceStructInStructProperty(StructProperty))
+				{
+					StructProperty->Struct = DuplicatedStruct;
+					StructsToRegenerateReferencesFor.Add(StructProperty->GetOwnerClass());
+				}
+			}
+		}
+
+		for (UStruct* Struct : StructsToRegenerateReferencesFor)
+		{
+			Struct->CollectBytecodeAndPropertyReferencedObjects();
+
+			PostReplace(Struct);
+		}
+
+		// as property owners are re-created, the duplicated struct will be GCed
+		DuplicatedStruct->RemoveFromRoot();
+	}
+}
+
 #undef LOCTEXT_NAMESPACE
diff --git a/Engine/Source/Editor/KismetCompiler/Public/UserDefinedStructureCompilerUtils.h b/Engine/Source/Editor/KismetCompiler/Public/UserDefinedStructureCompilerUtils.h
--- a/Engine/Source/Editor/KismetCompiler/Public/UserDefinedStructureCompilerUtils.h
+++ b/Engine/Source/Editor/KismetCompiler/Public/UserDefinedStructureCompilerUtils.h
@@ -11,5 +11,10 @@ class FUserDefinedStructureCompilerUtils
 public:
 	// ASSUMPTION, structure doesn't need to be renamed or removed
 	static void CompileStruct(class UUserDefinedStruct* Struct, class FCompilerResultsLog& MessageLog, bool bForceRecompile);
+
+	KISMETCOMPILER_API static void ReplaceStructWithTempDuplicateByPredicate(
+		UUserDefinedStruct* StructureToReinstance,
+		TFunctionRef<bool(FStructProperty* InStructProperty)> ShouldReplaceStructInStructProperty,
+		TFunctionRef<void(UStruct* InStruct)> PostReplace);
 };
 
