Control Rig: Fix GetChannel doesn't display procedural channels

diff --git a/Engine/Plugins/Animation/ControlRig/Source/ControlRigEditor/Private/Graph/ControlRigGraphPanelPinFactory.cpp b/Engine/Plugins/Animation/ControlRig/Source/ControlRigEditor/Private/Graph/ControlRigGraphPanelPinFactory.cpp
--- a/Engine/Plugins/Animation/ControlRig/Source/ControlRigEditor/Private/Graph/ControlRigGraphPanelPinFactory.cpp
+++ b/Engine/Plugins/Animation/ControlRig/Source/ControlRigEditor/Private/Graph/ControlRigGraphPanelPinFactory.cpp
@@ -149,7 +149,13 @@ TSharedPtr<SGraphPin> FControlRigGraphPanelPinFactory::CreatePin(UEdGraphPin* In
 									}
 								}
 
-								if(Blueprint->Hierarchy->Find<FRigControlElement>(ControlKey) == nullptr)
+								const URigHierarchy* Hierarchy = Blueprint->Hierarchy;
+								if (UControlRig* ControlRig = Cast<UControlRig>(Blueprint->GetObjectBeingDebugged()))
+								{
+									Hierarchy = ControlRig->GetHierarchy();
+								}
+								
+								if(Hierarchy->Find<FRigControlElement>(ControlKey) == nullptr)
 								{
 									ControlKey.Reset();
 								}
@@ -165,7 +171,7 @@ TSharedPtr<SGraphPin> FControlRigGraphPanelPinFactory::CreatePin(UEdGraphPin* In
 
 								if(!bRefreshList)
 								{
-									const int32 TopologyVersion = Blueprint->Hierarchy->GetTopologyVersion();
+									const int32 TopologyVersion = Hierarchy->GetTopologyVersion();
 									if(ChannelNames.TopologyVersion != TopologyVersion)
 									{
 										bRefreshList = true;
@@ -182,9 +188,9 @@ TSharedPtr<SGraphPin> FControlRigGraphPanelPinFactory::CreatePin(UEdGraphPin* In
 									ChannelNames.Names->Reset();
 									ChannelNames.Names->Add(MakeShareable(new FString(FName(NAME_None).ToString())));
 
-									if(const FRigControlElement* ControlElement = Blueprint->Hierarchy->Find<FRigControlElement>(ControlKey))
+									if(const FRigControlElement* ControlElement = Hierarchy->Find<FRigControlElement>(ControlKey))
 									{
-										FRigBaseElementChildrenArray Children = Blueprint->Hierarchy->GetChildren(ControlElement);
+										FRigBaseElementChildrenArray Children = Hierarchy->GetChildren(ControlElement);
 										for(const FRigBaseElement* Child : Children)
 										{
 											if(const FRigControlElement* ChildControl = Cast<FRigControlElement>(Child))
