Convert the external actor reference member to a weak ptr in the FComponentReference struct type so that it doesn't prevent GC.

diff --git a/Engine/Plugins/Animation/ControlRig/Source/ControlRig/Private/ControlRigComponent.cpp b/Engine/Plugins/Animation/ControlRig/Source/ControlRig/Private/ControlRigComponent.cpp
--- a/Engine/Plugins/Animation/ControlRig/Source/ControlRig/Private/ControlRigComponent.cpp
+++ b/Engine/Plugins/Animation/ControlRig/Source/ControlRig/Private/ControlRigComponent.cpp
@@ -1242,7 +1242,7 @@ void UControlRigComponent::ValidateMappingData()
 			MappedElement.ElementIndex = INDEX_NONE;
 			MappedElement.SubIndex = INDEX_NONE;
 
-			AActor* MappedOwner = MappedElement.ComponentReference.OtherActor == nullptr ? GetOwner() : ToRawPtr(MappedElement.ComponentReference.OtherActor);
+			AActor* MappedOwner = !MappedElement.ComponentReference.OtherActor.IsValid() ? GetOwner() : MappedElement.ComponentReference.OtherActor.Get();
 			MappedElement.SceneComponent = Cast<USceneComponent>(MappedElement.ComponentReference.GetComponent(MappedOwner));
 
 			// try again with the path to the component
diff --git a/Engine/Plugins/Runtime/CableComponent/Source/CableComponent/Private/CableComponent.cpp b/Engine/Plugins/Runtime/CableComponent/Source/CableComponent/Private/CableComponent.cpp
--- a/Engine/Plugins/Runtime/CableComponent/Source/CableComponent/Private/CableComponent.cpp
+++ b/Engine/Plugins/Runtime/CableComponent/Source/CableComponent/Private/CableComponent.cpp
@@ -581,7 +581,7 @@ void UCableComponent::SetAttachEndTo(AActor* Actor, FName ComponentProperty, FNa
 
 AActor* UCableComponent::GetAttachedActor() const
 {
-	return AttachEndTo.OtherActor;
+	return AttachEndTo.OtherActor.Get();
 }
 
 USceneComponent* UCableComponent::GetAttachedComponent() const
diff --git a/Engine/Source/Editor/DetailCustomizations/Private/ComponentReferenceCustomization.cpp b/Engine/Source/Editor/DetailCustomizations/Private/ComponentReferenceCustomization.cpp
--- a/Engine/Source/Editor/DetailCustomizations/Private/ComponentReferenceCustomization.cpp
+++ b/Engine/Source/Editor/DetailCustomizations/Private/ComponentReferenceCustomization.cpp
@@ -422,7 +422,7 @@ FPropertyAccess::Result FComponentReferenceCustomization::GetValue(FComponentRef
 
 bool FComponentReferenceCustomization::IsComponentReferenceValid(const FComponentReference& Value) const
 {
-	if (!bAllowAnyActor && Value.OtherActor)
+	if (!bAllowAnyActor && Value.OtherActor.IsValid())
 	{
 		return false;
 	}
diff --git a/Engine/Source/Runtime/Engine/Classes/Engine/EngineTypes.h b/Engine/Source/Runtime/Engine/Classes/Engine/EngineTypes.h
--- a/Engine/Source/Runtime/Engine/Classes/Engine/EngineTypes.h
+++ b/Engine/Source/Runtime/Engine/Classes/Engine/EngineTypes.h
@@ -3156,11 +3156,11 @@ struct ENGINE_API FComponentReference : public FBaseComponentReference
 	FComponentReference() : OtherActor(nullptr) {}
 
 	/** 
-	 * Hard Pointer to a different Actor that owns the Component.  
+	 * Weak Pointer to a different Actor that owns the Component.  
 	 * If this is not provided the reference refers to a component on this / the same actor.
 	 */
 	UPROPERTY(EditInstanceOnly, BlueprintReadWrite, Category=Component, meta = (DisplayName = "Referenced Actor"))
-	TObjectPtr<AActor> OtherActor;
+	TWeakObjectPtr<AActor> OtherActor;
 
 	/** Get the actual component pointer from this reference */
 	class UActorComponent* GetComponent(AActor* OwningActor) const;
diff --git a/Engine/Source/Runtime/Engine/Private/EngineTypes.cpp b/Engine/Source/Runtime/Engine/Private/EngineTypes.cpp
--- a/Engine/Source/Runtime/Engine/Private/EngineTypes.cpp
+++ b/Engine/Source/Runtime/Engine/Private/EngineTypes.cpp
@@ -165,7 +165,7 @@ UActorComponent* FBaseComponentReference::ExtractComponent(AActor* SearchActor)
 
 UActorComponent* FComponentReference::GetComponent(AActor* OwningActor) const
 {
-	AActor* SearchActor = (OtherActor != NULL) ? ToRawPtr(OtherActor) : OwningActor;
+	AActor* SearchActor = (OtherActor.IsValid()) ? OtherActor.Get() : OwningActor;
 	return ExtractComponent(SearchActor);
 }
 
