Control Rig: Fix memory consumption issue due to snapshots

diff --git a/Engine/Plugins/Animation/ControlRig/Source/ControlRig/Private/ControlRig.cpp b/Engine/Plugins/Animation/ControlRig/Source/ControlRig/Private/ControlRig.cpp
--- a/Engine/Plugins/Animation/ControlRig/Source/ControlRig/Private/ControlRig.cpp
+++ b/Engine/Plugins/Animation/ControlRig/Source/ControlRig/Private/ControlRig.cpp
@@ -147,6 +147,14 @@ void UControlRig::BeginDestroy()
 	// on destruction clear out the initialized snapshots
 	if (HasAnyFlags(RF_ClassDefaultObject))
 	{
+		for (TPair<uint32, TObjectPtr<URigVM>>& Pair : InitializedVMSnapshots)
+		{
+			if (IsValid(Pair.Value))
+			{
+				TObjectPtr<URigVM> VMSnapshot = Pair.Value;
+				VMSnapshot->RemoveFromRoot();				
+			}
+		}
 		InitializedVMSnapshots.Reset();
 	}
 
@@ -195,7 +203,7 @@ void UControlRig::Initialize(bool bInitRigUnits)
 	// recompute the hash used to differentiate VMs based on their memory layout 
 	if (HasAnyFlags(RF_ClassDefaultObject))
 	{
-		CachedMemoryHash = INDEX_NONE;
+		CachedMemoryHash = 0;
 
 		if(VM)
 		{
@@ -645,6 +653,14 @@ void UControlRig::InstantiateVMFromCDO()
 		}
 	}
 
+	if(VM)
+	{
+		CachedMemoryHash = HashCombine(
+			VM->GetLiteralMemory()->GetMemoryHash(),
+			VM->GetWorkMemory()->GetMemoryHash()
+		);
+	}
+
 	RequestInit();
 }
 
@@ -1482,7 +1498,7 @@ bool UControlRig::ExecuteUnits(FRigUnitContext& InOutContext, const FName& InEve
 				UControlRig* CDO = Cast<UControlRig>(GetClass()->GetDefaultObject());
 
 				bool bIsValidSnapshot = false;
-				if(SnapshotHash != INDEX_NONE)
+				if(SnapshotHash != 0)
 				{
 					TObjectPtr<URigVM>* InitializedVMSnapshotPtr = CDO->InitializedVMSnapshots.Find(SnapshotHash);
 					if(InitializedVMSnapshotPtr && !InitializedVMSnapshotPtr->IsNull())
@@ -3403,9 +3419,9 @@ void UControlRig::PostInitInstance(UControlRig* InCDO)
 
 uint32 UControlRig::GetHashForInitializeVMSnapShot()
 {
-	if(CachedMemoryHash == INDEX_NONE)
+	if(CachedMemoryHash == 0)
 	{
-		return INDEX_NONE;
+		return 0;
 	}
 	
 	uint32 Hash = GetHierarchy()->GetNameHash();
diff --git a/Engine/Plugins/Animation/ControlRig/Source/ControlRig/Public/ControlRig.h b/Engine/Plugins/Animation/ControlRig/Source/ControlRig/Public/ControlRig.h
--- a/Engine/Plugins/Animation/ControlRig/Source/ControlRig/Public/ControlRig.h
+++ b/Engine/Plugins/Animation/ControlRig/Source/ControlRig/Public/ControlRig.h
@@ -509,7 +509,7 @@ protected:
 	// and stores an initialized VM per hash.
 	UPROPERTY(transient)
 	TMap<uint32, TObjectPtr<URigVM>> InitializedVMSnapshots;
-	uint32 CachedMemoryHash = INDEX_NONE;
+	uint32 CachedMemoryHash = 0;
 
 	// computes the hash used to store / find VM snapshots
 	uint32 GetHashForInitializeVMSnapShot();
diff --git a/Engine/Source/Runtime/RigVM/Public/RigVMCore/RigVMExternalVariable.h b/Engine/Source/Runtime/RigVM/Public/RigVMCore/RigVMExternalVariable.h
--- a/Engine/Source/Runtime/RigVM/Public/RigVMCore/RigVMExternalVariable.h
+++ b/Engine/Source/Runtime/RigVM/Public/RigVMCore/RigVMExternalVariable.h
@@ -165,9 +165,10 @@ struct RIGVM_API FRigVMExternalVariable
 		else if (const FStructProperty *StructProperty = CastField<FStructProperty>(InProperty))
 		{
 			const UScriptStruct* StructType = StructProperty->Struct;
-			if (StructType->GetCppStructOps() && StructType->GetCppStructOps()->HasGetTypeHash())
+			// UserDefinedStruct overrides GetStructTypeHash to work without valid CppStructOps
+			if (StructType->IsA<UUserDefinedStruct>() || (StructType->GetCppStructOps() && StructType->GetCppStructOps()->HasGetTypeHash()))
 			{
-				return StructType->GetCppStructOps()->GetStructTypeHash(InMemory);
+				return StructType->GetStructTypeHash(InMemory);
 			}
 			else
 			{
