Aim Constraint: fix not having an option to constraint with offset, this is a simple node that does aim with a simple interface that is consistent with other constraint nodes

diff --git a/Engine/Plugins/Animation/ControlRig/Source/ControlRig/Private/Math/ControlRigMathLibrary.cpp b/Engine/Plugins/Animation/ControlRig/Source/ControlRig/Private/Math/ControlRigMathLibrary.cpp
--- a/Engine/Plugins/Animation/ControlRig/Source/ControlRig/Private/Math/ControlRigMathLibrary.cpp
+++ b/Engine/Plugins/Animation/ControlRig/Source/ControlRig/Private/Math/ControlRigMathLibrary.cpp
@@ -415,3 +415,55 @@ FQuat FControlRigMathLibrary::FindQuatBetweenNormals(const FVector& A, const FVe
 	Result.Normalize();
 	return Result;
 }
+
+FVector FControlRigMathLibrary::GetEquivalentEulerAngle(const FVector& InEulerAngle, const EEulerRotationOrder& InOrder)
+{
+	FVector TheOtherAngle;
+
+	switch (InOrder)
+	{
+	case EEulerRotationOrder::XYZ:
+		TheOtherAngle.X = InEulerAngle.X + 180;
+		TheOtherAngle.Y = 180 - InEulerAngle.Y;
+		TheOtherAngle.Z = InEulerAngle.Z + 180;
+		break;
+	case EEulerRotationOrder::XZY:
+		TheOtherAngle.X = InEulerAngle.X + 180;
+		TheOtherAngle.Z = 180 - InEulerAngle.Z;
+		TheOtherAngle.Y = InEulerAngle.Y + 180;
+		break;
+	case EEulerRotationOrder::YXZ:
+		TheOtherAngle.Y = InEulerAngle.Y + 180;
+		TheOtherAngle.X = 180 - InEulerAngle.X;
+		TheOtherAngle.Z = InEulerAngle.Z + 180;
+		break;
+	case EEulerRotationOrder::YZX:
+		TheOtherAngle.Y = InEulerAngle.Y + 180;
+		TheOtherAngle.Z = 180 - InEulerAngle.Z;
+		TheOtherAngle.X = InEulerAngle.X + 180;
+		break;
+	case EEulerRotationOrder::ZXY:
+		TheOtherAngle.Z = InEulerAngle.Z + 180;
+		TheOtherAngle.X = 180 - InEulerAngle.X;
+		TheOtherAngle.Y = InEulerAngle.Y + 180;
+		break;
+	case EEulerRotationOrder::ZYX:
+		TheOtherAngle.Z = InEulerAngle.Z + 180;
+		TheOtherAngle.Y = 180 - InEulerAngle.Y;
+		TheOtherAngle.X = InEulerAngle.X + 180;
+		break;
+	};
+
+	TheOtherAngle.UnwindEuler();
+	
+	return TheOtherAngle;
+}
+
+FVector& FControlRigMathLibrary::ChooseBetterEulerAngleForAxisFilter(const FVector& Base, FVector& A, FVector& B)
+{
+	// simply compare the sum of difference should be enough
+	float Diff1 = (A - Base).GetAbs().Dot(FVector::OneVector);
+	float Diff2 = (B - Base).GetAbs().Dot(FVector::OneVector);
+
+	return Diff1 < Diff2 ? A : B;
+}
diff --git a/Engine/Plugins/Animation/ControlRig/Source/ControlRig/Private/Units/Highlevel/Hierarchy/RigUnit_AimBone.cpp b/Engine/Plugins/Animation/ControlRig/Source/ControlRig/Private/Units/Highlevel/Hierarchy/RigUnit_AimBone.cpp
--- a/Engine/Plugins/Animation/ControlRig/Source/ControlRig/Private/Units/Highlevel/Hierarchy/RigUnit_AimBone.cpp
+++ b/Engine/Plugins/Animation/ControlRig/Source/ControlRig/Private/Units/Highlevel/Hierarchy/RigUnit_AimBone.cpp
@@ -3,6 +3,7 @@
 #include "Units/Highlevel/Hierarchy/RigUnit_AimBone.h"
 #include "Units/RigUnitContext.h"
 #include "Math/ControlRigMathLibrary.h"
+#include "AnimationCoreLibrary.h"
 
 FRigUnit_AimBoneMath_Execute()
 {
@@ -254,3 +255,317 @@ FRigUnit_AimItem_Execute()
 	Hierarchy->SetGlobalTransform(CachedItem, Transform);
 }
 
+
+FRigUnit_AimConstraintLocalSpaceOffset_Execute()
+{
+	DECLARE_SCOPE_HIERARCHICAL_COUNTER_RIGUNIT()
+
+	if (Weight < KINDA_SMALL_NUMBER)
+	{
+		return;
+	}
+	
+	if(Context.State == EControlRigState::Init)
+	{
+		WorldUpSpaceCache.Reset();
+		ChildCache.Reset();
+		ParentCaches.Reset();
+	}
+	
+	URigHierarchy* Hierarchy = ExecuteContext.Hierarchy;
+	if (Hierarchy)
+	{
+		WorldUpSpaceCache.UpdateCache(WorldUp.Space, Hierarchy);
+		
+		if (!ChildCache.UpdateCache(Child, Hierarchy))
+		{
+			return;
+		}
+		
+		if(ParentCaches.Num() != Parents.Num())
+		{
+			ParentCaches.SetNumZeroed(Parents.Num());
+		}
+		
+		float OverallWeight = 0;
+		for (int32 ParentIndex = 0; ParentIndex < Parents.Num(); ParentIndex++)
+		{
+			const FConstraintParent& Parent = Parents[ParentIndex];
+			if (!ParentCaches[ParentIndex].UpdateCache(Parent.Item, Hierarchy))
+			{
+				continue;
+			}
+			
+			const float ClampedWeight = FMath::Max(Parent.Weight, 0.f);
+			if (ClampedWeight < KINDA_SMALL_NUMBER)
+			{
+				continue;
+			}
+
+			OverallWeight += ClampedWeight;
+		}
+
+		const float WeightNormalizer = 1.0f / OverallWeight;
+
+		if (OverallWeight > KINDA_SMALL_NUMBER)
+		{
+			FTransform AdditionalOffsetTransform;
+			
+			const bool bChildIsControl = Child.Type == ERigElementType::Control;
+			if(bChildIsControl)
+			{
+				if (FRigControlElement* ChildAsControlElement = Hierarchy->Get<FRigControlElement>(ChildCache))
+				{
+					AdditionalOffsetTransform = Hierarchy->GetControlOffsetTransform(ChildAsControlElement, ERigTransformType::InitialLocal);
+				}
+			}
+			
+			FQuat OffsetRotation = FQuat::Identity;
+			if (bMaintainOffset)
+			{
+				FVector MixedInitialGlobalPosition = FVector::ZeroVector;
+
+				for (int32 ParentIndex = 0; ParentIndex < Parents.Num(); ParentIndex++)
+				{
+					if (!ParentCaches[ParentIndex].IsValid())
+					{
+						continue;
+					}
+					
+					const FConstraintParent& Parent = Parents[ParentIndex];
+					
+					const float ClampedWeight = FMath::Max(Parent.Weight, 0.f);
+					if (ClampedWeight < KINDA_SMALL_NUMBER)
+					{
+						continue;
+					}
+
+					const float NormalizedWeight = ClampedWeight * WeightNormalizer;
+
+					FTransform ParentInitialGlobalTransform = Hierarchy->GetGlobalTransformByIndex(ParentCaches[ParentIndex], true);
+
+					MixedInitialGlobalPosition += ParentInitialGlobalTransform.GetLocation() * NormalizedWeight;
+				}
+
+				// points the aim axis towards the parents
+				FRigUnit_AimItem_Target Primary;
+				Primary.Axis = AimAxis;
+				Primary.Target = MixedInitialGlobalPosition;
+				Primary.Kind = EControlRigVectorKind::Location;
+				Primary.Space = FRigElementKey();
+				FCachedRigElement PrimaryCachedSpace;
+
+				// points the up axis towards the world up target
+				FRigUnit_AimItem_Target Secondary;
+				Secondary.Axis = UpAxis;
+
+				FVector InitialWorldSpaceUp = WorldUp.Target;
+
+				FTransform SpaceInitialTransform = Hierarchy->GetInitialGlobalTransform(WorldUpSpaceCache);
+				if (Secondary.Kind == EControlRigVectorKind::Direction)
+				{
+					InitialWorldSpaceUp = SpaceInitialTransform.TransformVectorNoScale(InitialWorldSpaceUp);
+				}
+				else
+				{
+					InitialWorldSpaceUp = SpaceInitialTransform.TransformPositionNoScale(InitialWorldSpaceUp);
+				}	
+			
+				Secondary.Target = InitialWorldSpaceUp;
+				Secondary.Kind = WorldUp.Kind;
+				// we don't want to reference the space any more since its transform
+				// is already included in InitialWorldSpaceUp 
+				Secondary.Space = FRigElementKey();
+				FCachedRigElement SecondaryCachedSpace;
+
+			
+				FTransform ChildInitialGlobalTransform = Hierarchy->GetInitialGlobalTransform(ChildCache);
+				FTransform InitialAimResult = ChildInitialGlobalTransform;
+				FRigUnit_AimBone_DebugSettings DummyDebugSettings;
+				FRigUnit_AimBoneMath::StaticExecute(
+					RigVMExecuteContext,
+					ChildInitialGlobalTransform, // optional
+					Primary,
+					Secondary,
+					Weight,
+					InitialAimResult,
+					DummyDebugSettings,
+					PrimaryCachedSpace,
+					SecondaryCachedSpace,
+					Context);
+
+				FTransform ChildParentInitialGlobalTransform = Hierarchy->GetParentTransformByIndex(ChildCache, true);
+				FQuat MixedInitialLocalRotation = ChildParentInitialGlobalTransform.GetRotation().Inverse() * InitialAimResult.GetRotation();
+
+				FTransform ChildInitialLocalTransform = Hierarchy->GetLocalTransformByIndex(ChildCache, true);
+
+				// Controls need to be handled a bit differently
+				if (bChildIsControl)
+				{
+					ChildInitialLocalTransform *= AdditionalOffsetTransform;
+				}
+				
+				FQuat ChildInitialLocalRotation = ChildInitialLocalTransform.GetRotation();
+			
+				OffsetRotation = MixedInitialLocalRotation.Inverse() * ChildInitialLocalRotation;
+			}
+
+			FVector MixedGlobalPosition = FVector::ZeroVector;
+
+			for (int32 ParentIndex = 0; ParentIndex < Parents.Num(); ParentIndex++)
+			{
+				if (!ParentCaches[ParentIndex].IsValid())
+				{
+					continue;
+				}
+					
+				const FConstraintParent& Parent = Parents[ParentIndex];
+					
+				const float ClampedWeight = FMath::Max(Parent.Weight, 0.f);
+				if (ClampedWeight < KINDA_SMALL_NUMBER)
+				{
+					continue;
+				}
+
+				const float NormalizedWeight = ClampedWeight * WeightNormalizer;
+
+				FTransform ParentCurrentGlobalTransform = Hierarchy->GetGlobalTransformByIndex(ParentCaches[ParentIndex], false);
+				MixedGlobalPosition += ParentCurrentGlobalTransform.GetLocation() * NormalizedWeight;
+			}
+
+			// points the aim axis towards the parents
+			FRigUnit_AimItem_Target Primary;
+			Primary.Axis = AimAxis;
+			Primary.Target = MixedGlobalPosition;
+			Primary.Kind = EControlRigVectorKind::Location;
+			Primary.Space = FRigElementKey();
+			FCachedRigElement PrimaryCachedSpace;
+
+		
+			// points the up axis towards the world up target
+			FRigUnit_AimItem_Target Secondary;
+			Secondary.Axis = UpAxis;
+			Secondary.Target = WorldUp.Target;
+			Secondary.Kind = WorldUp.Kind;
+			Secondary.Space = WorldUp.Space;
+			
+			FTransform ChildGlobalTransform = Hierarchy->GetGlobalTransform(ChildCache);
+			FTransform AimResult = ChildGlobalTransform;
+			FRigUnit_AimBoneMath::StaticExecute(
+				RigVMExecuteContext,
+				ChildGlobalTransform, // optional
+				Primary,
+				Secondary,
+				Weight,
+				AimResult,
+				AdvancedSettings.DebugSettings,
+				PrimaryCachedSpace,
+				WorldUpSpaceCache,
+				Context);	
+
+			// handle filtering, performed in local space
+			FTransform ChildParentGlobalTransform = Hierarchy->GetParentTransformByIndex(ChildCache, false);
+			FQuat MixedLocalRotation = ChildParentGlobalTransform.GetRotation().Inverse() * AimResult.GetRotation();
+
+			if (bMaintainOffset)
+			{
+				MixedLocalRotation = MixedLocalRotation * OffsetRotation;
+			}
+
+			FTransform ChildCurrentLocalTransform = Hierarchy->GetLocalTransformByIndex(ChildCache, false);
+				
+			// controls need to be handled a bit differently
+			if (bChildIsControl)
+			{
+				ChildCurrentLocalTransform *= AdditionalOffsetTransform;
+			}
+
+			FQuat FilteredMixedLocalRotation = MixedLocalRotation;
+			
+			if(!Filter.HasNoEffect())
+			{
+				FVector MixedEulerRotation = AnimationCore::EulerFromQuat(MixedLocalRotation, AdvancedSettings.RotationOrderForFilter);
+
+				FVector MixedEulerRotation2 = FControlRigMathLibrary::GetEquivalentEulerAngle(MixedEulerRotation, AdvancedSettings.RotationOrderForFilter);
+
+				FQuat ChildRotation = ChildCurrentLocalTransform.GetRotation();
+				FVector ChildEulerRotation = AnimationCore::EulerFromQuat(ChildRotation, AdvancedSettings.RotationOrderForFilter);
+
+				FVector ClosestMixedEulerRotation = FControlRigMathLibrary::ChooseBetterEulerAngleForAxisFilter(ChildEulerRotation, MixedEulerRotation, MixedEulerRotation2	);
+				
+				FVector FilteredMixedEulerRotation;
+				FilteredMixedEulerRotation.X = Filter.bX ? ClosestMixedEulerRotation.X : ChildEulerRotation.X;
+            	FilteredMixedEulerRotation.Y = Filter.bY ? ClosestMixedEulerRotation.Y : ChildEulerRotation.Y;
+            	FilteredMixedEulerRotation.Z = Filter.bZ ? ClosestMixedEulerRotation.Z : ChildEulerRotation.Z;
+
+				FilteredMixedLocalRotation = AnimationCore::QuatFromEuler(FilteredMixedEulerRotation, AdvancedSettings.RotationOrderForFilter);
+			}
+
+			FTransform FilteredMixedLocalTransform = ChildCurrentLocalTransform;
+
+			FilteredMixedLocalTransform.SetRotation(FilteredMixedLocalRotation);
+
+			FTransform FinalLocalTransform = FilteredMixedLocalTransform;
+
+			if (Weight < 1.0f - KINDA_SMALL_NUMBER)
+			{
+				FinalLocalTransform = FControlRigMathLibrary::LerpTransform(ChildCurrentLocalTransform, FinalLocalTransform, Weight);
+			}
+
+			if (Child.Type == ERigElementType::Control)
+			{
+				// need to convert back to offset space for the actual control value
+				FinalLocalTransform = FinalLocalTransform.GetRelativeTransform(AdditionalOffsetTransform);
+				FinalLocalTransform.NormalizeRotation();
+			}
+		
+			Hierarchy->SetLocalTransform(ChildCache, FinalLocalTransform);
+		}
+	}
+}
+
+#if WITH_DEV_AUTOMATION_TESTS
+#include "Units/RigUnitTest.h"
+
+IMPLEMENT_RIGUNIT_AUTOMATION_TEST(FRigUnit_AimConstraintLocalSpaceOffset)
+{
+	// multi-parent test
+	{
+		// use euler rotation here to match other software's rotation representation more easily
+    	EEulerRotationOrder Order = EEulerRotationOrder::XZY;
+    	const FRigElementKey Child = Controller->AddBone(TEXT("Child"), FRigElementKey(), FTransform(FVector(0.f, 0.f, 0.f)), true, ERigBoneType::User);
+    	const FRigElementKey ChildTarget = Controller->AddBone(TEXT("ChildTarget"), FRigElementKey(TEXT("Child"), ERigElementType::Bone), FTransform(FVector(0.f, 10.f, 0.f)), false, ERigBoneType::User);
+		
+    	const FRigElementKey Parent1 = Controller->AddBone(TEXT("Parent1"), FRigElementKey(), FTransform(FVector(10.f, 10.f, 10.f)), true, ERigBoneType::User);
+    	const FRigElementKey Parent2 = Controller->AddBone(TEXT("Parent2"), FRigElementKey(), FTransform(FVector(-10.f,10.f, 10.f)), true, ERigBoneType::User);
+    	
+    	Unit.ExecuteContext.Hierarchy = Hierarchy;
+		Unit.AimAxis = FVector(0,1, 0);
+		Unit.UpAxis = FVector(0, 0,1);
+    	
+    	Unit.Child = Child;
+    
+    	Unit.Parents.Add(FConstraintParent(Parent1, 1.0));
+    	Unit.Parents.Add(FConstraintParent(Parent2, 1.0));
+    
+    	Hierarchy->ResetPoseToInitial(ERigElementType::Bone);
+    	Unit.bMaintainOffset = false;
+    	
+    	Execute();
+    	AddErrorIfFalse(Hierarchy->GetGlobalTransform(1).GetTranslation().Equals(FVector(0 , 10.f * FMath::Cos(PI / 4), 10.f * FMath::Sin(PI / 4)), 0.01f),
+    		TEXT("unexpected translation for maintain offset off"));
+    	
+    	Hierarchy->ResetPoseToInitial(ERigElementType::Bone);
+    	Hierarchy->SetGlobalTransform(2, FTransform(FVector(10.f, 10.f, 20.f)));
+    	Hierarchy->SetGlobalTransform(3, FTransform(FVector(-10.f, 10.f, 20.f)));
+    	Unit.bMaintainOffset = true;
+    	
+    	Execute();
+    	AddErrorIfFalse(Hierarchy->GetGlobalTransform(1).GetTranslation().Equals(FVector(0 , 9.487f, 3.162f), 0.01f),
+    					TEXT("unexpected translation for maintain offset on"));
+
+	}
+	
+	return true;
+}
+#endif
\ No newline at end of file
diff --git a/Engine/Plugins/Animation/ControlRig/Source/ControlRig/Private/Units/Highlevel/Hierarchy/RigUnit_AimBone.h b/Engine/Plugins/Animation/ControlRig/Source/ControlRig/Private/Units/Highlevel/Hierarchy/RigUnit_AimBone.h
--- a/Engine/Plugins/Animation/ControlRig/Source/ControlRig/Private/Units/Highlevel/Hierarchy/RigUnit_AimBone.h
+++ b/Engine/Plugins/Animation/ControlRig/Source/ControlRig/Private/Units/Highlevel/Hierarchy/RigUnit_AimBone.h
@@ -3,6 +3,9 @@
 #pragma once
 
 #include "Units/Highlevel/RigUnit_HighlevelBase.h"
+
+#include "Units/Highlevel/Hierarchy/RigUnit_TransformConstraint.h"
+
 #include "RigUnit_AimBone.generated.h"
 
 USTRUCT()
@@ -357,3 +360,140 @@ struct CONTROLRIG_API FRigUnit_AimItem: public FRigUnit_HighlevelBaseMutable
 	UPROPERTY()
 	FCachedRigElement SecondaryCachedSpace;
 };
+
+
+USTRUCT()
+struct CONTROLRIG_API FRigUnit_AimConstraint_WorldUp
+{
+	GENERATED_BODY()
+
+	FRigUnit_AimConstraint_WorldUp()
+	{
+		Target = FVector(0.f, 0.f, 1.f);
+		Kind = EControlRigVectorKind::Direction;
+		Space = FRigElementKey(NAME_None, ERigElementType::None);
+	}
+
+	/**
+	 * The target to aim at - can be a direction or location based on the Kind setting
+	 */
+	UPROPERTY(meta = (Input))
+	FVector Target;
+
+	/**
+	 * The kind of target this is representing - can be a direction or a location
+	 */
+	UPROPERTY(meta = (DisplayName="Target is ", Input))
+	EControlRigVectorKind Kind;
+
+	/**
+	 * The space in which the target is expressed, use None to indicate world space
+	 */
+	UPROPERTY(meta = (DisplayName = "Target Space", Input))
+	FRigElementKey Space;
+};
+
+
+USTRUCT()
+struct FRigUnit_AimConstraint_AdvancedSettings
+{
+	GENERATED_BODY()
+
+	FRigUnit_AimConstraint_AdvancedSettings()
+		:RotationOrderForFilter(EEulerRotationOrder::XZY)
+	{}
+	
+	/**
+	*	Settings related to debug drawings
+	*/
+	UPROPERTY(meta = (Input,DetailsOnly))
+	FRigUnit_AimBone_DebugSettings DebugSettings;
+	
+	/**
+	*	Rotation is converted to euler angles using the specified order such that individual axes can be filtered.
+	*/
+	UPROPERTY(meta = (Input))
+	EEulerRotationOrder RotationOrderForFilter;	
+};
+
+/**
+ * Orients an item such that its aim axis points towards a global target.
+ * Note: This node operates in global space!
+ */
+USTRUCT(meta=(DisplayName="Aim Constraint", Category="Hierarchy", Keywords="Lookat"))
+struct CONTROLRIG_API FRigUnit_AimConstraintLocalSpaceOffset: public FRigUnit_HighlevelBaseMutable
+{
+	GENERATED_BODY()
+
+	FRigUnit_AimConstraintLocalSpaceOffset()
+		: Child(FRigElementKey())
+		, bMaintainOffset(true)
+		, Filter(FFilterOptionPerAxis())
+		, AimAxis(FVector(1.0f,0.0f,0.0f))
+		, UpAxis(FVector(0.0f,0.0f,1.0f))
+		, WorldUp(FRigUnit_AimConstraint_WorldUp())
+		, AdvancedSettings(FRigUnit_AimConstraint_AdvancedSettings())
+		, Weight(1.0f)
+	{
+		Parents.Add(FConstraintParent());
+	};
+	
+	RIGVM_METHOD()
+	virtual void Execute(const FRigUnitContext& Context) override;
+
+	/** 
+	 * The name of the item to apply aim
+	 */
+	UPROPERTY(meta = (Input, ExpandByDefault))
+	FRigElementKey Child;
+
+	/**
+	 *	Maintains the offset between child and weighted average of parents based on initial transforms
+	 */
+	UPROPERTY(meta = (Input))
+	bool bMaintainOffset;
+
+	/**
+	 * Filters the final rotation by axes based on the euler rotation order defined in the node's advanced settings
+	 * If flipping is observed, try adjusting the rotation order
+	 */
+	UPROPERTY(meta = (Input))
+	FFilterOptionPerAxis Filter;
+
+	/**
+	 * Child is rotated so that its AimAxis points to the parents
+	 */
+	UPROPERTY(meta = (Input, ExpandByDefault))
+	FVector AimAxis;
+
+	/**
+	 * Child is rotated around the AimAxis so that its UpAxis points to/Aligns with the WorldUp target 
+	 */
+	UPROPERTY(meta = (Input, ExpandByDefault))
+	FVector UpAxis;
+
+	/**
+	 * Defines how Child should rotate around the AimAxis. This is the aim target for the UpAxis
+	 */
+	UPROPERTY(meta = (Input, ExpandByDefault))
+	FRigUnit_AimConstraint_WorldUp WorldUp;
+	
+	UPROPERTY(meta = (Input, ExpandByDefault, DefaultArraySize = 1))
+	TArray<FConstraintParent> Parents;
+
+	UPROPERTY(meta = (Input))
+	FRigUnit_AimConstraint_AdvancedSettings AdvancedSettings;
+	
+	UPROPERTY(meta = (Input))
+	float Weight;
+	
+	UPROPERTY()
+	FCachedRigElement WorldUpSpaceCache; 
+
+	UPROPERTY()
+	FCachedRigElement ChildCache;
+
+	UPROPERTY()
+	TArray<FCachedRigElement> ParentCaches;
+	
+};
diff --git a/Engine/Plugins/Animation/ControlRig/Source/ControlRig/Public/Math/ControlRigMathLibrary.h b/Engine/Plugins/Animation/ControlRig/Source/ControlRig/Public/Math/ControlRigMathLibrary.h
--- a/Engine/Plugins/Animation/ControlRig/Source/ControlRig/Public/Math/ControlRigMathLibrary.h
+++ b/Engine/Plugins/Animation/ControlRig/Source/ControlRig/Public/Math/ControlRigMathLibrary.h
@@ -77,4 +77,10 @@ public:
 	static FVector ClampSpatially(const FVector& Value, EAxis::Type Axis, EControlRigClampSpatialMode::Type Type, float Minimum, float Maximum, FTransform Space);
 	static FQuat FindQuatBetweenVectors(const FVector& A, const FVector& B);
 	static FQuat FindQuatBetweenNormals(const FVector& A, const FVector& B);
+
+	// See - "Computing Euler angles from a rotation matrix" by Gregory G. Slabaugh
+	// Each spatial orientation can be mapped to two equivalent euler angles within range (-180, 180)
+	static FVector GetEquivalentEulerAngle(const FVector& InEulerAngle, const EEulerRotationOrder& InOrder);
+	
+	static FVector& ChooseBetterEulerAngleForAxisFilter(const FVector& Base, FVector& A, FVector& B);
 };
\ No newline at end of file
