FNames containing asset paths are deprecated. FSoftObjectPath or FTopLevelAssetPath should be used instead.
Fixups for ControlRig.

diff --git a/Engine/Plugins/Animation/ControlRig/Source/ControlRig/Public/Tools/ControlRigToolAsset.h b/Engine/Plugins/Animation/ControlRig/Source/ControlRig/Public/Tools/ControlRigToolAsset.h
--- a/Engine/Plugins/Animation/ControlRig/Source/ControlRig/Public/Tools/ControlRigToolAsset.h
+++ b/Engine/Plugins/Animation/ControlRig/Source/ControlRig/Public/Tools/ControlRigToolAsset.h
@@ -38,7 +38,9 @@ public:
 
 			// Generate a unique control asset name for this take if there are already assets of the same name
 			IAssetRegistry& AssetRegistry = FModuleManager::LoadModuleChecked<FAssetRegistryModule>("AssetRegistry").Get();
-			while (AssetRegistry.GetAssetByObjectPath(*InPackageName).IsValid())
+			TArray<FAssetData> OutAssetData;
+			AssetRegistry.GetAssetsByPackageName(*InPackageName, OutAssetData);
+			while (OutAssetData.Num() > 0)
 			{
 				int32 TrimCount = InPackageName.Len() - BasePackageLength;
 				if (TrimCount > 0)
@@ -47,6 +49,8 @@ public:
 				}
 
 				InPackageName += FString::Printf(TEXT("_%04d"), UniqueIndex++);
+				OutAssetData.Empty();
+				AssetRegistry.GetAssetsByPackageName(*InPackageName, OutAssetData);
 			}
 
 			// Create the asset to record into
diff --git a/Engine/Plugins/Animation/ControlRig/Source/ControlRigDeveloper/Private/ControlRigBlueprint.cpp b/Engine/Plugins/Animation/ControlRig/Source/ControlRigDeveloper/Private/ControlRigBlueprint.cpp
--- a/Engine/Plugins/Animation/ControlRig/Source/ControlRigDeveloper/Private/ControlRigBlueprint.cpp
+++ b/Engine/Plugins/Animation/ControlRig/Source/ControlRigDeveloper/Private/ControlRigBlueprint.cpp
@@ -1107,7 +1107,7 @@ void UControlRigBlueprint::HandlePackageDone()
 
 						// find all control rigs matching the reference node
 						FAssetData AssetData = AssetRegistry.GetAssetByObjectPath(
-							Pair.Value[ReferenceIndex].ToSoftObjectPath().GetAssetPathName());
+							Pair.Value[ReferenceIndex].ToSoftObjectPath().GetWithoutSubPath());
 
 						// if the asset has never been loaded - make sure to load it once and mark as dirty
 						if(AssetData.IsValid() && !AssetData.IsAssetLoaded())
@@ -2475,7 +2475,7 @@ TArray<UControlRigBlueprint*> UControlRigBlueprint::GetDependencies(bool bRecurs
 TArray<FAssetData> UControlRigBlueprint::GetDependentAssets() const
 {
 	TArray<FAssetData> Dependents;
-	TArray<FName> AssetPaths;
+	TArray<FSoftObjectPath> AssetPaths;
 
 	if(URigVMFunctionLibrary* FunctionLibrary = RigVMClient.GetFunctionLibrary())
 	{
@@ -2489,8 +2489,8 @@ TArray<FAssetData> UControlRigBlueprint::GetDependentAssets() const
 				TArray<TSoftObjectPtr<URigVMFunctionReferenceNode>> References = FunctionLibrary->GetReferencesForFunction(FunctionName);
 				for(const TSoftObjectPtr<URigVMFunctionReferenceNode>& Reference : References)
 				{
-					const FName AssetPath = Reference.ToSoftObjectPath().GetAssetPathName();
-					if(AssetPath.ToString().StartsWith(TEXT("/Engine/Transient")))
+					const FSoftObjectPath AssetPath = Reference.ToSoftObjectPath();
+					if(AssetPath.GetLongPackageName().StartsWith(TEXT("/Engine/Transient")))
 					{
 						continue;
 					}
@@ -2499,7 +2499,7 @@ TArray<FAssetData> UControlRigBlueprint::GetDependentAssets() const
 					{
 						AssetPaths.Add(AssetPath);
 
-						const FAssetData AssetData = AssetRegistry.GetAssetByObjectPath(*AssetPath.ToString());
+						const FAssetData AssetData = AssetRegistry.GetAssetByObjectPath(AssetPath);
 						if(AssetData.IsValid())
 						{
 							Dependents.Add(AssetData);
diff --git a/Engine/Plugins/Animation/ControlRig/Source/ControlRigEditor/Private/ControlRigEditorModule.cpp b/Engine/Plugins/Animation/ControlRig/Source/ControlRigEditor/Private/ControlRigEditorModule.cpp
--- a/Engine/Plugins/Animation/ControlRig/Source/ControlRigEditor/Private/ControlRigEditorModule.cpp
+++ b/Engine/Plugins/Animation/ControlRig/Source/ControlRigEditor/Private/ControlRigEditorModule.cpp
@@ -502,7 +502,7 @@ TSharedRef< SWidget > FControlRigEditorModule::GenerateAnimationMenu(TWeakPtr<IA
 											FString PreviewSkeletalMesh = AssetData.GetTagValueRef<FString>(TEXT("PreviewSkeletalMesh"));
 											if (PreviewSkeletalMesh.Len() > 0)
 											{
-												FAssetData SkelMeshData = AssetRegistry.GetAssetByObjectPath(FName(*PreviewSkeletalMesh));
+												FAssetData SkelMeshData = AssetRegistry.GetAssetByObjectPath(FSoftObjectPath(PreviewSkeletalMesh));
 												FString PreviewSkeleton = SkelMeshData.GetTagValueRef<FString>(TEXT("Skeleton"));
 												if (PreviewSkeleton == SkeletonName)
 												{
@@ -553,8 +553,8 @@ TSharedRef< SWidget > FControlRigEditorModule::GenerateAnimationMenu(TWeakPtr<IA
 										if (bChildOfObjectClass && bMatchesFlags)
 										{
 											FString GeneratedClassPathString = InUnloadedClassData->GetClassPathName().ToString();
-											FName BlueprintPath = FName(*GeneratedClassPathString.LeftChop(2)); // Chop off _C
-											FAssetData AssetData = AssetRegistry.GetAssetByObjectPath(BlueprintPath);
+											FString BlueprintPath = GeneratedClassPathString.LeftChop(2); // Chop off _C
+											FAssetData AssetData = AssetRegistry.GetAssetByObjectPath(FSoftObjectPath(BlueprintPath));
 											return MatchesFilter(AssetData);
 
 										}
diff --git a/Engine/Plugins/Animation/ControlRig/Source/ControlRigEditor/Private/EditMode/SControlRigControlViews.cpp b/Engine/Plugins/Animation/ControlRig/Source/ControlRigEditor/Private/EditMode/SControlRigControlViews.cpp
--- a/Engine/Plugins/Animation/ControlRig/Source/ControlRigEditor/Private/EditMode/SControlRigControlViews.cpp
+++ b/Engine/Plugins/Animation/ControlRig/Source/ControlRigEditor/Private/EditMode/SControlRigControlViews.cpp
@@ -40,8 +40,7 @@ void FControlRigView::CaptureThumbnail(UObject* Asset)
 
 		TArray<FAssetData> SelectedAssets;
 		FAssetRegistryModule& AssetRegistryModule = FModuleManager::Get().LoadModuleChecked<FAssetRegistryModule>(TEXT("AssetRegistry"));
-		FName PathName = *Asset->GetPathName();
-		FAssetData AssetData = AssetRegistryModule.Get().GetAssetByObjectPath(PathName);
+		FAssetData AssetData = AssetRegistryModule.Get().GetAssetByObjectPath(FSoftObjectPath(Asset));
 		SelectedAssets.Emplace(AssetData);
 		AssetViewUtils::CaptureThumbnailFromViewport(Viewport, SelectedAssets);
 
@@ -137,7 +136,7 @@ void SControlRigAssetEditableTextBox::OnNameTextCommitted(const FText& NewText,
 		{
 
 			IAssetRegistry& AssetRegistry = FModuleManager::LoadModuleChecked<FAssetRegistryModule>("AssetRegistry").Get();
-			FAssetData AssetData = AssetRegistry.GetAssetByObjectPath(FName(*Asset->GetPathName()));
+			FAssetData AssetData = AssetRegistry.GetAssetByObjectPath(FSoftObjectPath(Asset.Get()));
 			const FString PackagePath = FPackageName::GetLongPackagePath(Asset->GetOutermost()->GetName());
 
 			//Need to save asset before renaming else may lose snapshot
diff --git a/Engine/Plugins/Animation/ControlRig/Source/ControlRigEditor/Private/Graph/NodeSpawners/ControlRigFunctionRefNodeSpawner.cpp b/Engine/Plugins/Animation/ControlRig/Source/ControlRigEditor/Private/Graph/NodeSpawners/ControlRigFunctionRefNodeSpawner.cpp
--- a/Engine/Plugins/Animation/ControlRig/Source/ControlRigEditor/Private/Graph/NodeSpawners/ControlRigFunctionRefNodeSpawner.cpp
+++ b/Engine/Plugins/Animation/ControlRig/Source/ControlRigEditor/Private/Graph/NodeSpawners/ControlRigFunctionRefNodeSpawner.cpp
@@ -84,7 +84,9 @@ UControlRigFunctionRefNodeSpawner* UControlRigFunctionRefNodeSpawner::CreateFrom
 UControlRigFunctionRefNodeSpawner* UControlRigFunctionRefNodeSpawner::CreateFromAssetData(const FAssetData& InAssetData, const FControlRigPublicFunctionData& InPublicFunction)
 {
 	UControlRigFunctionRefNodeSpawner* NodeSpawner = NewObject<UControlRigFunctionRefNodeSpawner>(GetTransientPackage());
-	NodeSpawner->ReferencedAssetObjectPath = InAssetData.ObjectPath;
+PRAGMA_DISABLE_DEPRECATION_WARNINGS
+	NodeSpawner->ReferencedAssetObjectPath = InAssetData.ToSoftObjectPath().ToFName();
+PRAGMA_ENABLE_DEPRECATION_WARNINGS
 	NodeSpawner->ReferencedPublicFunctionData = InPublicFunction;
 	NodeSpawner->NodeClass = UControlRigGraphNode::StaticClass();
 	NodeSpawner->bIsLocalFunction = false;
@@ -108,7 +110,7 @@ UControlRigFunctionRefNodeSpawner* UControlRigFunctionRefNodeSpawner::CreateFrom
 		}
 	}
 
-	const FString ObjectPathString = InAssetData.ObjectPath.ToString();
+	const FString ObjectPathString = InAssetData.GetObjectPathString();
 	if(MenuSignature.Tooltip.IsEmpty())
 	{
 		MenuSignature.Tooltip = FText::FromString(ObjectPathString);
diff --git a/Engine/Plugins/Animation/ControlRig/Source/ControlRigEditor/Private/Sequencer/ControlRigParameterTrackEditor.cpp b/Engine/Plugins/Animation/ControlRig/Source/ControlRigEditor/Private/Sequencer/ControlRigParameterTrackEditor.cpp
--- a/Engine/Plugins/Animation/ControlRig/Source/ControlRigEditor/Private/Sequencer/ControlRigParameterTrackEditor.cpp
+++ b/Engine/Plugins/Animation/ControlRig/Source/ControlRigEditor/Private/Sequencer/ControlRigParameterTrackEditor.cpp
@@ -699,7 +699,7 @@ public:
 			FString PreviewSkeletalMesh = AssetData.GetTagValueRef<FString>(TEXT("PreviewSkeletalMesh"));
 			if (PreviewSkeletalMesh.Len() > 0)
 			{
-				FAssetData SkelMeshData = AssetRegistry.GetAssetByObjectPath(FName(*PreviewSkeletalMesh));
+				FAssetData SkelMeshData = AssetRegistry.GetAssetByObjectPath(FSoftObjectPath(PreviewSkeletalMesh));
 				FString PreviewSkeleton = SkelMeshData.GetTagValueRef<FString>(TEXT("Skeleton"));
 				if (PreviewSkeleton == SkeletonName)
 				{
@@ -779,8 +779,8 @@ public:
 		if (bChildOfObjectClass && bMatchesFlags)
 		{
 			FString GeneratedClassPathString = InUnloadedClassData->GetClassPathName().ToString();
-			FName BlueprintPath = FName(*GeneratedClassPathString.LeftChop(2)); // Chop off _C
-			FAssetData AssetData = AssetRegistry.GetAssetByObjectPath(BlueprintPath);
+			FString BlueprintPath = GeneratedClassPathString.LeftChop(2); // Chop off _C
+			FAssetData AssetData = AssetRegistry.GetAssetByObjectPath(FSoftObjectPath(BlueprintPath));
 			return MatchesFilter(AssetData);
 
 		}
