As SendEvend nodes can request to add keyframes, we need to disable compensation when switching parent to avoid entering an endless loop when evaluating the graph

diff --git a/Engine/Plugins/Animation/ControlRig/Source/ControlRigEditor/Private/Sequencer/ControlRigSpaceChannelEditors.cpp b/Engine/Plugins/Animation/ControlRig/Source/ControlRigEditor/Private/Sequencer/ControlRigSpaceChannelEditors.cpp
--- a/Engine/Plugins/Animation/ControlRig/Source/ControlRigEditor/Private/Sequencer/ControlRigSpaceChannelEditors.cpp
+++ b/Engine/Plugins/Animation/ControlRig/Source/ControlRigEditor/Private/Sequencer/ControlRigSpaceChannelEditors.cpp
@@ -1180,6 +1180,8 @@ void FControlRigSpaceChannelHelpers::CompensateIfNeeded(
 
 						if (ExistingValue != PreviousValue) //if they are the same no need to do anything
 						{
+							TGuardValue<bool> CompensateGuard(bDoNotCompensate, true);
+							
 							//find global value at current time
 							TArray<FTransform> ControlRigParentWorldTransforms({FTransform::Identity});
 							TArray<FTransform> ControlWorldTransforms;
