Control Rig: Fix find hierarchy elements with tag returns duplicated elements

diff --git a/Engine/Plugins/Animation/ControlRig/Source/ControlRig/Private/Units/Hierarchy/RigUnit_Metadata.cpp b/Engine/Plugins/Animation/ControlRig/Source/ControlRig/Private/Units/Hierarchy/RigUnit_Metadata.cpp
--- a/Engine/Plugins/Animation/ControlRig/Source/ControlRig/Private/Units/Hierarchy/RigUnit_Metadata.cpp
+++ b/Engine/Plugins/Animation/ControlRig/Source/ControlRig/Private/Units/Hierarchy/RigUnit_Metadata.cpp
@@ -442,7 +442,7 @@ FRigUnit_FindItemsWithMetadata_Execute()
 	{
 		if(Element->GetMetadata(Name, Type) != nullptr)
 		{
-			Items.Add(Element->GetKey());
+			Items.AddUnique(Element->GetKey());
 		}
 		bContinue = true;
 	});
@@ -577,7 +577,7 @@ FRigUnit_FindItemsWithMetadataTag_Execute()
 		{
 			if(Md->GetValue().Contains(Tag))
 			{
-				Items.Add(Element->GetKey());
+				Items.AddUnique(Element->GetKey());
 			}
 		}
 		bContinue = true;
@@ -606,7 +606,7 @@ FRigUnit_FindItemsWithMetadataTagArray_Execute()
 
 			if(bFoundAll)
 			{
-				Items.Add(Element->GetKey());
+				Items.AddUnique(Element->GetKey());
 			}
 		}
 		bContinue = true;
