Sequencer: Spaces: Don't allow spaces to be set on parented channels, those without control shapes, not animated, etc..

diff --git a/Engine/Plugins/Animation/ControlRig/Source/ControlRigEditor/Private/Editor/SRigSpacePickerWidget.cpp b/Engine/Plugins/Animation/ControlRig/Source/ControlRigEditor/Private/Editor/SRigSpacePickerWidget.cpp
--- a/Engine/Plugins/Animation/ControlRig/Source/ControlRigEditor/Private/Editor/SRigSpacePickerWidget.cpp
+++ b/Engine/Plugins/Animation/ControlRig/Source/ControlRigEditor/Private/Editor/SRigSpacePickerWidget.cpp
@@ -222,7 +222,29 @@ void SRigSpacePickerWidget::SetControls(
 	}
 	
 	Hierarchy = InHierarchy;
-	ControlKeys = InControls;
+	ControlKeys.SetNum(0);
+	for (const FRigElementKey& Key : InControls)
+	{
+		if (const FRigControlElement* ControlElement = Hierarchy->FindChecked<FRigControlElement>(Key))
+		{
+			//if it has no shape or not animatable then bail
+			if (ControlElement->Settings.SupportsShape() == false || Hierarchy->IsAnimatable(ControlElement) == false)
+			{
+				continue;
+			}
+			if (ControlElement->Settings.ControlType == ERigControlType::Bool ||
+				ControlElement->Settings.ControlType == ERigControlType::Float ||
+				ControlElement->Settings.ControlType == ERigControlType::Integer)
+			{
+				//if it has a channel and has a parent bail
+				if (const FRigControlElement* ParentControlElement = Cast<FRigControlElement>(Hierarchy->GetFirstParent(ControlElement)))
+				{
+					continue;
+				}
+			}
+		}
+		ControlKeys.Add(Key);
+	}
 	if (Hierarchy.IsValid() && HierarchyModifiedHandle.IsValid() == false)
 	{
 		HierarchyModifiedHandle = InHierarchy->OnModified().AddSP(this, &SRigSpacePickerWidget::OnHierarchyModified);
