Control Rig: When recreating shapes check to make sure we really need to recreate them. This is a bandaid to fix the issue with ResetHierarchy constantly getting called now during construction events but it does help a bunch with performance when this event is triggered incorrectly

diff --git a/Engine/Plugins/Animation/ControlRig/Source/ControlRigEditor/Private/EditMode/ControlRigEditMode.cpp b/Engine/Plugins/Animation/ControlRig/Source/ControlRigEditor/Private/EditMode/ControlRigEditMode.cpp
--- a/Engine/Plugins/Animation/ControlRig/Source/ControlRigEditor/Private/EditMode/ControlRigEditMode.cpp
+++ b/Engine/Plugins/Animation/ControlRig/Source/ControlRigEditor/Private/EditMode/ControlRigEditMode.cpp
@@ -3121,33 +3121,8 @@ void FControlRigEditMode::RequestToRecreateControlShapeActors(UControlRig* Contr
 	}
 }
 
-void FControlRigEditMode::RecreateControlShapeActors(const TArray<FRigElementKey>& InSelectedElements)
-{
-	if (RecreateControlShapesRequired == ERecreateControlRigShape::RecreateAll)
-	{
-		for (TWeakObjectPtr<UControlRig>& RuntimeRigPtr : RuntimeControlRigs)
-		{
-			if (UControlRig* RuntimeControlRig = RuntimeRigPtr.Get())
-			{
-				DestroyShapesActors(RuntimeControlRig);
-				CreateShapeActors(RuntimeControlRig);
-			}
-		}
-	}
-	else if(ControlRigsToRecreate.Num() > 0)
-	{
-		TArray < UControlRig*> ControlRigsCopy = ControlRigsToRecreate;
-		for (UControlRig* ControlRig : ControlRigsCopy)
-		{
-			DestroyShapesActors(ControlRig);
-			CreateShapeActors(ControlRig);
-		}
-		ControlRigsToRecreate.SetNum(0);
-	}
-}
-
 // temporarily we just support following types of gizmo
-bool IsSupportedControlType(const ERigControlType ControlType)
+static bool IsSupportedControlType(const ERigControlType ControlType)
 {
 	switch (ControlType)
 	{
@@ -3172,6 +3147,72 @@ bool IsSupportedControlType(const ERigControlType ControlType)
 	return false;
 }
 
+void FControlRigEditMode::RecreateControlShapeActors(const TArray<FRigElementKey>& InSelectedElements)
+{
+	if (RecreateControlShapesRequired == ERecreateControlRigShape::RecreateAll)
+	{
+		for (TWeakObjectPtr<UControlRig>& RuntimeRigPtr : RuntimeControlRigs)
+		{
+			if (UControlRig* RuntimeControlRig = RuntimeRigPtr.Get())
+			{
+				DestroyShapesActors(RuntimeControlRig);
+				CreateShapeActors(RuntimeControlRig);
+			}
+		}
+	}
+	else if (ControlRigsToRecreate.Num() > 0)
+	{
+		TArray < UControlRig*> ControlRigsCopy = ControlRigsToRecreate;
+		for (UControlRig* ControlRig : ControlRigsCopy)
+		{
+			//check to see if actors have really changed, if not don't do it
+			bool bRecreateThem = true;
+			if (TArray<AControlRigShapeActor*>* ShapeActors = ControlRigShapeActors.Find(ControlRig))
+			{
+				TArray<FRigControlElement*> Controls = ControlRig->AvailableControls();
+				for (int32 Index = Controls.Num() - 1; Index >= 0; --Index)
+				{
+					FRigControlElement* ControlElement = Controls[Index];
+					if (!ControlElement->Settings.SupportsShape() || !IsSupportedControlType(ControlElement->Settings.ControlType))
+					{
+						Controls.RemoveAtSwap(Index);
+					}
+				}
+				//unfortunately n*n-ish but this should be very rare and much faster than recreating them
+				for (AControlRigShapeActor* Actor : *ShapeActors)
+				{
+					if (Actor)
+					{
+						for (int32 Index = 0; Index < Controls.Num(); ++Index)
+						{
+							FRigControlElement* Element = Controls[Index];
+							if (Element && Element->GetName() == Actor->ControlName)
+							{
+								Controls.RemoveAtSwap(Index);
+								break;
+							}
+						}
+					}
+					else //no actor just recreate
+					{
+						break;
+					}
+				}
+				if (Controls.Num() == 0)
+				{
+					bRecreateThem = false;
+				}
+			}
+			if (bRecreateThem)
+			{
+				DestroyShapesActors(ControlRig);
+				CreateShapeActors(ControlRig);
+			}
+		}
+		ControlRigsToRecreate.SetNum(0);
+	}
+}
+
 void FControlRigEditMode::CreateShapeActors(UControlRig* ControlRig)
 {
 	// create gizmo actors
