Separate out curve/bone control name-mapping into two TMaps to prevent possible collisions caused by HashCombine
#misc bump anim DDC key to ensure any potential invalid compressed anim data generated previously gets cleared

diff --git a/Engine/Plugins/Animation/ControlRig/Source/ControlRig/Private/Rigs/FKControlRig.cpp b/Engine/Plugins/Animation/ControlRig/Source/ControlRig/Private/Rigs/FKControlRig.cpp
--- a/Engine/Plugins/Animation/ControlRig/Source/ControlRig/Private/Rigs/FKControlRig.cpp
+++ b/Engine/Plugins/Animation/ControlRig/Source/ControlRig/Private/Rigs/FKControlRig.cpp
@@ -25,9 +25,10 @@ FName UFKControlRig::GetControlName(const FName& InName, const ERigElementType&
 {
 	if (InName != NAME_None)
 	{
-		static thread_local TMap<uint32, FName> NameToControlMapping;
-		const int32 Hash = HashCombine(GetTypeHash(InName), GetTypeHash(InType));
-		if (const FName* CachedName = NameToControlMapping.Find(Hash))
+		check(InType == ERigElementType::Bone || InType == ERigElementType::Curve);
+		static thread_local TMap<FName, FName> NameToControlMappings[2];
+		TMap<FName, FName>& NameToControlMapping = NameToControlMappings[InType == ERigElementType::Bone ? 0 : 1];
+		if (const FName* CachedName = NameToControlMapping.Find(InName))
 		{
 			return *CachedName;
 		}
@@ -47,7 +48,7 @@ FName UFKControlRig::GetControlName(const FName& InName, const ERigElementType&
 			static FString ControlPostFix = TEXT("_CONTROL");
 		
 			ScratchString.Append(ControlPostFix);
-			return NameToControlMapping.Add(Hash, FName(*ScratchString));
+			return NameToControlMapping.Add(InName, FName(*ScratchString));
 		}
 	}
 
@@ -59,9 +60,10 @@ FName UFKControlRig::GetControlTargetName(const FName& InName, const ERigElement
 {
 	if (InName != NAME_None)
 	{
-		static thread_local TMap<uint32, FName> NameToTargetMapping;
-		const int32 Hash = HashCombine(GetTypeHash(InName), GetTypeHash(InType));
-		if (const FName* CachedName = NameToTargetMapping.Find(Hash))
+		check(InType == ERigElementType::Bone || InType == ERigElementType::Curve);
+		static thread_local TMap<FName, FName> NameToTargetMappings[2];
+		TMap<FName, FName>& NameToTargetMapping = NameToTargetMappings[InType == ERigElementType::Bone ? 0 : 1];
+		if (const FName* CachedName = NameToTargetMapping.Find(InName))
 		{
 			return *CachedName;
 		}
@@ -85,7 +87,7 @@ FName UFKControlRig::GetControlTargetName(const FName& InName, const ERigElement
 			}
 
 			const FStringView ControlTargetString(*ScratchString, StartPostFix != INDEX_NONE ? StartPostFix : ScratchString.Len());
-			return NameToTargetMapping.Add(Hash, FName(ControlTargetString));
+			return NameToTargetMapping.Add(InName, FName(ControlTargetString));
 		}
 	}
 
diff --git a/Engine/Source/Runtime/Engine/Private/Animation/AnimCompressionDerivedData.cpp b/Engine/Source/Runtime/Engine/Private/Animation/AnimCompressionDerivedData.cpp
--- a/Engine/Source/Runtime/Engine/Private/Animation/AnimCompressionDerivedData.cpp
+++ b/Engine/Source/Runtime/Engine/Private/Animation/AnimCompressionDerivedData.cpp
@@ -37,7 +37,7 @@ const TCHAR* FDerivedDataAnimationCompression::GetVersionString() const
 	// This is a version string that mimics the old versioning scheme. If you
 	// want to bump this version, generate a new guid using VS->Tools->Create GUID and
 	// return it here. Ex.
-	return TEXT("E29090F49C2348068642D5BEE76253B3");
+	return TEXT("FB8D0FF1ED0C4848951601EB6BB028F6");
 }
 
 FString FDerivedDataAnimationCompression::GetDebugContextString() const
