fixed dynamic behavior for scale constraints
+ manage look at properties visibility

diff --git a/Engine/Plugins/Animation/ControlRig/Source/ControlRigEditor/Private/Editor/Constraints/SConstraintsWidget.cpp b/Engine/Plugins/Animation/ControlRig/Source/ControlRigEditor/Private/Editor/Constraints/SConstraintsWidget.cpp
--- a/Engine/Plugins/Animation/ControlRig/Source/ControlRigEditor/Private/Editor/Constraints/SConstraintsWidget.cpp
+++ b/Engine/Plugins/Animation/ControlRig/Source/ControlRigEditor/Private/Editor/Constraints/SConstraintsWidget.cpp
@@ -892,6 +892,30 @@ TSharedPtr<SWidget> SConstraintsEditionWidget::CreateContextMenu()
 	
 	TSharedRef<IDetailsView> DetailsView = FModuleManager::GetModuleChecked<FPropertyEditorModule>("PropertyEditor").CreateDetailView(DetailsViewArgs);
 
+	// manage properties visibility
+	const bool bIsLookAtConstraint = Constraint->GetClass() == UTickableLookAtConstraint::StaticClass();
+	FIsPropertyVisible PropertyVisibility = FIsPropertyVisible::CreateLambda(
+		[bIsLookAtConstraint](const FPropertyAndParent& InPropertyAndParent)
+		{
+			if (!bIsLookAtConstraint)
+			{
+				return true;
+			}
+
+			// hide offset properties for look at constraints
+			static const FName MaintainOffsetPropName = GET_MEMBER_NAME_CHECKED(UTickableTransformConstraint, bMaintainOffset);
+			static const FName DynamicOffsetPropName = GET_MEMBER_NAME_CHECKED(UTickableTransformConstraint, bDynamicOffset);
+						
+			const FName PropertyName = InPropertyAndParent.Property.GetFName();
+			if (PropertyName == MaintainOffsetPropName || PropertyName == DynamicOffsetPropName)
+			{
+				return false;
+			}
+
+			return true;
+		});
+	DetailsView->SetIsPropertyVisibleDelegate(PropertyVisibility);
+	
 	TArray<TWeakObjectPtr<UObject>> ConstrainsToEdit;
 	ConstrainsToEdit.Add(Constraint);
 	
@@ -957,32 +981,35 @@ TSharedPtr<SWidget> SConstraintsEditionWidget::CreateContextMenu()
 		{
 			return TransformConstraint->bDynamicOffset;
 		});
-	
-		MenuBuilder.BeginSection("KeyConstraint", LOCTEXT("KeyConstraintHeader", "Keys"));
-		{
-			MenuBuilder.AddMenuEntry(
-			LOCTEXT("CompensateKeyLabel", "Compensate Key"),
-			FText::Format(LOCTEXT("CompensateKeyTooltip", "Compensate transform key for {0}."), ConstraintLabel),
-			FSlateIcon(),
-			FUIAction(FExecuteAction::CreateLambda([TransformConstraint]()
-			{
-				FConstraintChannelHelper::Compensate(TransformConstraint);
-			}), IsCompensationEnabled),
-			NAME_None,
-			EUserInterfaceActionType::Button);
 
-			MenuBuilder.AddMenuEntry(
-			LOCTEXT("CompensateAllKeysLabel", "Compensate All Keys"),
-			FText::Format(LOCTEXT("CompensateAllKeysTooltip", "Compensate all transform keys for {0}."), ConstraintLabel),
-			FSlateIcon(),
-			FUIAction(FExecuteAction::CreateLambda([TransformConstraint]()
+		if (!bIsLookAtConstraint)
+		{
+			MenuBuilder.BeginSection("KeyConstraint", LOCTEXT("KeyConstraintHeader", "Keys"));
 			{
-				FConstraintChannelHelper::Compensate(TransformConstraint, true);
-			}), IsCompensationEnabled),
-			NAME_None,
-			EUserInterfaceActionType::Button);
+				MenuBuilder.AddMenuEntry(
+				LOCTEXT("CompensateKeyLabel", "Compensate Key"),
+				FText::Format(LOCTEXT("CompensateKeyTooltip", "Compensate transform key for {0}."), ConstraintLabel),
+				FSlateIcon(),
+				FUIAction(FExecuteAction::CreateLambda([TransformConstraint]()
+				{
+					FConstraintChannelHelper::Compensate(TransformConstraint);
+				}), IsCompensationEnabled),
+				NAME_None,
+				EUserInterfaceActionType::Button);
+
+				MenuBuilder.AddMenuEntry(
+				LOCTEXT("CompensateAllKeysLabel", "Compensate All Keys"),
+				FText::Format(LOCTEXT("CompensateAllKeysTooltip", "Compensate all transform keys for {0}."), ConstraintLabel),
+				FSlateIcon(),
+				FUIAction(FExecuteAction::CreateLambda([TransformConstraint]()
+				{
+					FConstraintChannelHelper::Compensate(TransformConstraint, true);
+				}), IsCompensationEnabled),
+				NAME_None,
+				EUserInterfaceActionType::Button);
+			}
+			MenuBuilder.EndSection();
 		}
-		MenuBuilder.EndSection();
 	}
 
 	return MenuBuilder.MakeWidget();
diff --git a/Engine/Source/Runtime/Experimental/Animation/Constraints/Private/TransformConstraint.cpp b/Engine/Source/Runtime/Experimental/Animation/Constraints/Private/TransformConstraint.cpp
--- a/Engine/Source/Runtime/Experimental/Animation/Constraints/Private/TransformConstraint.cpp
+++ b/Engine/Source/Runtime/Experimental/Animation/Constraints/Private/TransformConstraint.cpp
@@ -528,6 +528,33 @@ UTickableScaleConstraint::UTickableScaleConstraint()
 	Type = ETransformConstraintType::Scale;
 }
 
+#if WITH_EDITOR
+
+void UTickableScaleConstraint::PostEditChangeProperty(FPropertyChangedEvent& PropertyChangedEvent)
+{
+	Super::PostEditChangeProperty(PropertyChangedEvent);
+
+	const FName PropertyName = PropertyChangedEvent.GetPropertyName();
+	if (PropertyName == GET_MEMBER_NAME_CHECKED(UTickableTransformConstraint, bDynamicOffset))
+	{
+		if (bDynamicOffset)
+		{
+			Cache.CachedInputHash = CalculateInputHash();
+			
+			const FVector ParentScale = GetParentGlobalTransform().GetScale3D();
+			OffsetScale = GetChildGlobalTransform().GetScale3D();
+			OffsetScale[0] = FMath::Abs(ParentScale[0]) > KINDA_SMALL_NUMBER ? OffsetScale[0] / ParentScale[0] : 0.f;
+			OffsetScale[1] = FMath::Abs(ParentScale[1]) > KINDA_SMALL_NUMBER ? OffsetScale[1] / ParentScale[1] : 0.f;
+			OffsetScale[2] = FMath::Abs(ParentScale[2]) > KINDA_SMALL_NUMBER ? OffsetScale[2] / ParentScale[2] : 0.f;
+			
+			Evaluate();
+		}
+		return;
+	}
+}
+
+#endif
+
 void UTickableScaleConstraint::ComputeOffset()
 {
 	const FTransform InitParentTransform = GetParentGlobalTransform();
@@ -572,6 +599,65 @@ FConstraintTickFunction::ConstraintFunction UTickableScaleConstraint::GetFunctio
 	};
 }
 
+void UTickableScaleConstraint::OnHandleModified(UTransformableHandle* InHandle, EHandleEvent InEvent)
+{
+	Super::OnHandleModified(InHandle, InEvent);
+	
+	if (!Active || !bDynamicOffset)
+	{
+		return;
+	}
+	
+	if (!InHandle || InHandle != ChildTRSHandle)
+	{
+		return;
+	}
+
+	const bool bUpdateFromGlobal = (InEvent == EHandleEvent::GlobalTransformUpdated);
+	const bool bUpdateTransform = InEvent == EHandleEvent::LocalTransformUpdated || bUpdateFromGlobal;
+	if (!bUpdateTransform)
+	{
+		return;
+	}
+
+	const uint32 InputHash = CalculateInputHash();
+	
+	// update dynamic offset
+	if (InputHash != Cache.CachedInputHash)
+	{
+		Cache.CachedInputHash = InputHash;
+
+		if (bUpdateFromGlobal)
+		{
+			const FVector ParentScale = GetParentGlobalTransform().GetScale3D();
+			OffsetScale = GetChildGlobalTransform().GetScale3D();
+			OffsetScale[0] = FMath::Abs(ParentScale[0]) > KINDA_SMALL_NUMBER ? OffsetScale[0] / ParentScale[0] : 0.f;
+			OffsetScale[1] = FMath::Abs(ParentScale[1]) > KINDA_SMALL_NUMBER ? OffsetScale[1] / ParentScale[1] : 0.f;
+			OffsetScale[2] = FMath::Abs(ParentScale[2]) > KINDA_SMALL_NUMBER ? OffsetScale[2] / ParentScale[2] : 0.f;
+		}
+		else
+		{
+			const FTransform ChildLocalTransform = GetChildLocalTransform();
+			OffsetScale = ChildLocalTransform.GetScale3D();
+		}
+	}
+}
+
+uint32 UTickableScaleConstraint::CalculateInputHash() const
+{
+	uint32 Hash = 0;
+
+	// local scale hash
+	const FTransform ChildLocalTransform = GetChildLocalTransform();
+	Hash = HashCombine(Hash, GetTypeHash(ChildLocalTransform.GetScale3D() ));
+
+	// global scale hash
+	const FTransform ChildGlobalTransform = GetChildGlobalTransform();
+	Hash = HashCombine(Hash, GetTypeHash(ChildGlobalTransform.GetScale3D() ));
+	
+	return Hash;
+}
+
 /** 
  * UTickableParentConstraint
  **/
@@ -727,11 +813,16 @@ void UTickableParentConstraint::PostEditChangeProperty(FPropertyChangedEvent& Pr
 
 UTickableLookAtConstraint::UTickableLookAtConstraint()
 {
+	bMaintainOffset = false;
+	bDynamicOffset = false;
 	Type = ETransformConstraintType::LookAt;
 }
 
 void UTickableLookAtConstraint::ComputeOffset()
-{}
+{
+	bMaintainOffset = false;
+	bDynamicOffset = false;
+}
 
 FConstraintTickFunction::ConstraintFunction UTickableLookAtConstraint::GetFunction() const
 {
diff --git a/Engine/Source/Runtime/Experimental/Animation/Constraints/Public/TransformConstraint.h b/Engine/Source/Runtime/Experimental/Animation/Constraints/Public/TransformConstraint.h
--- a/Engine/Source/Runtime/Experimental/Animation/Constraints/Public/TransformConstraint.h
+++ b/Engine/Source/Runtime/Experimental/Animation/Constraints/Public/TransformConstraint.h
@@ -50,7 +50,7 @@ public:
 	TObjectPtr<UTransformableHandle> ChildTRSHandle;
 
 	/** Should that constraint maintain the default offset.  */
-	UPROPERTY(EditAnywhere, BlueprintReadWrite, Category="Offset")
+	UPROPERTY(EditAnywhere, BlueprintReadWrite, Category="Offset", meta=(EditCondition="Type!=ETransformConstraintType::LookAt"))
 	bool bMaintainOffset = true;
 
 	/** Defines how much the constraint will be applied. */
@@ -60,7 +60,7 @@ public:
 	float Weight = 1.f;
 
 	/** Should the child be able to change it's offset dynamically. */
-	UPROPERTY(EditAnywhere, BlueprintReadWrite, Category="Offset", meta=(EditCondition="bMaintainOffset"))
+	UPROPERTY(EditAnywhere, BlueprintReadWrite, Category="Offset", meta=(EditCondition="(Type!=ETransformConstraintType::LookAt) && bMaintainOffset"))
 	bool bDynamicOffset = false;
 
 	/** Returns the constraint type (Position, Parent, Aim...). */
@@ -232,13 +232,31 @@ public:
 	/** Returns the scale constraint function that the tick function will evaluate. */
 	virtual FConstraintTickFunction::ConstraintFunction GetFunction() const override;
 
+	/** Updates the dynamic offset based on external child's transform changes. */
+	virtual void OnHandleModified(UTransformableHandle* InHandle, EHandleEvent InEvent) override;
+	
 protected:
 	/** Computes the child's local scale offset in the parent space. */
 	virtual void ComputeOffset() override;
 
+	/** Cache data structure to store last child local/global transform. */
+	struct FDynamicCache
+	{
+		uint32 CachedInputHash = 0;
+	};
+	mutable FDynamicCache Cache;
+	uint32 CalculateInputHash() const;
+	
 	/** Defines the local child's scale offset in the parent space. */
 	UPROPERTY(EditAnywhere, BlueprintReadWrite, Category="Offset", meta=(EditCondition="bMaintainOffset"))
 	FVector OffsetScale = FVector::OneVector;
+
+#if WITH_EDITOR
+public:
+	// UObject interface
+	virtual void PostEditChangeProperty(FPropertyChangedEvent& PropertyChangedEvent) override;
+	// End of UObject interface
+#endif
 };
 
 /** 
