Control Rig: Fix variables allowing non control rig types

diff --git a/Engine/Plugins/Animation/ControlRig/Source/ControlRigEditor/Private/Editor/ControlRigEditor.cpp b/Engine/Plugins/Animation/ControlRig/Source/ControlRigEditor/Private/Editor/ControlRigEditor.cpp
--- a/Engine/Plugins/Animation/ControlRig/Source/ControlRigEditor/Private/Editor/ControlRigEditor.cpp
+++ b/Engine/Plugins/Animation/ControlRig/Source/ControlRigEditor/Private/Editor/ControlRigEditor.cpp
@@ -1421,6 +1421,11 @@ UBlueprint* FControlRigEditor::GetBlueprintObj() const
 	return nullptr;
 }
 
+TSubclassOf<UEdGraphSchema> FControlRigEditor::GetDefaultSchemaClass() const
+{
+	return UControlRigGraphSchema::StaticClass();
+}
+
 class FMemoryTypeMetaData : public ISlateMetaData
 {
 public:
diff --git a/Engine/Plugins/Animation/ControlRig/Source/ControlRigEditor/Private/Editor/ControlRigEditor.h b/Engine/Plugins/Animation/ControlRig/Source/ControlRigEditor/Private/Editor/ControlRigEditor.h
--- a/Engine/Plugins/Animation/ControlRig/Source/ControlRigEditor/Private/Editor/ControlRigEditor.h
+++ b/Engine/Plugins/Animation/ControlRig/Source/ControlRigEditor/Private/Editor/ControlRigEditor.h
@@ -73,6 +73,7 @@ public:
 
 	// FBlueprintEditor interface
 	virtual UBlueprint* GetBlueprintObj() const override;
+	virtual TSubclassOf<UEdGraphSchema> GetDefaultSchemaClass() const override;
 
 	int32 GetRigHierarchyTabCount() const { return RigHierarchyTabCount; }
 
diff --git a/Engine/Source/Editor/Kismet/Private/SBlueprintPalette.cpp b/Engine/Source/Editor/Kismet/Private/SBlueprintPalette.cpp
--- a/Engine/Source/Editor/Kismet/Private/SBlueprintPalette.cpp
+++ b/Engine/Source/Editor/Kismet/Private/SBlueprintPalette.cpp
@@ -823,6 +823,10 @@ private:
 			{
 				Schema = BlueprintEditorPtr.Pin()->GetFocusedGraph()->GetSchema();
 			}
+			else
+			{
+				Schema = BlueprintEditorPtr.Pin()->GetDefaultSchema().GetDefaultObject();
+			}
 		}
 		
 		const bool bIsDelegate = ActionPtr.IsValid() && ActionPtr.Pin()->GetTypeId() == FEdGraphSchemaAction_K2Delegate::StaticGetTypeId();
