Control Rig: Fix GetDependentAssets returning AssetData referencing function references instead of the asset itself

diff --git a/Engine/Plugins/Animation/ControlRig/Source/ControlRigDeveloper/Private/ControlRigBlueprint.cpp b/Engine/Plugins/Animation/ControlRig/Source/ControlRigDeveloper/Private/ControlRigBlueprint.cpp
--- a/Engine/Plugins/Animation/ControlRig/Source/ControlRigDeveloper/Private/ControlRigBlueprint.cpp
+++ b/Engine/Plugins/Animation/ControlRig/Source/ControlRigDeveloper/Private/ControlRigBlueprint.cpp
@@ -2489,20 +2489,24 @@ TArray<FAssetData> UControlRigBlueprint::GetDependentAssets() const
 				TArray<TSoftObjectPtr<URigVMFunctionReferenceNode>> References = FunctionLibrary->GetReferencesForFunction(FunctionName);
 				for(const TSoftObjectPtr<URigVMFunctionReferenceNode>& Reference : References)
 				{
-					const FSoftObjectPath AssetPath = Reference.ToSoftObjectPath();
-					if(AssetPath.GetLongPackageName().StartsWith(TEXT("/Engine/Transient")))
+					const TSoftObjectPtr<UPackage> Blueprint = Reference->GetTypedOuter<UControlRigBlueprint>();
+					if (!Blueprint.IsNull())
 					{
-						continue;
-					}
+						const FSoftObjectPath AssetPath = Blueprint.ToSoftObjectPath();
+						if(AssetPath.GetLongPackageName().StartsWith(TEXT("/Engine/Transient")))
+						{
+							continue;
+						}
 					
-					if(!AssetPaths.Contains(AssetPath))
-					{
-						AssetPaths.Add(AssetPath);
-
-						const FAssetData AssetData = AssetRegistry.GetAssetByObjectPath(AssetPath);
-						if(AssetData.IsValid())
+						if(!AssetPaths.Contains(AssetPath))
 						{
-							Dependents.Add(AssetData);
+							AssetPaths.Add(AssetPath);
+
+							const FAssetData AssetData = AssetRegistry.GetAssetByObjectPath(AssetPath);
+							if(AssetData.IsValid())
+							{
+								Dependents.Add(AssetData);
+							}
 						}
 					}
 				}
