Control Rig: Disable snapshots until properly working

diff --git a/Engine/Plugins/Animation/ControlRig/Source/ControlRig/Private/ControlRig.cpp b/Engine/Plugins/Animation/ControlRig/Source/ControlRig/Private/ControlRig.cpp
--- a/Engine/Plugins/Animation/ControlRig/Source/ControlRig/Private/ControlRig.cpp
+++ b/Engine/Plugins/Animation/ControlRig/Source/ControlRig/Private/ControlRig.cpp
@@ -54,6 +54,12 @@ static TAutoConsoleVariable<int32> CVarControlRigDisableExecutionAll(TEXT("Contr
 // CVar to disable swapping to nativized vms 
 static TAutoConsoleVariable<int32> CVarControlRigDisableNativizedVMs(TEXT("ControlRig.DisableNativizedVMs"), 1, TEXT("if nonzero we disable swapping to nativized VMs."));
 
+static bool bControlRigUseVMSnapshots = false;
+static FAutoConsoleVariableRef CVarControlRigUseVMSnapshots(
+	TEXT("ControlRig.UseVMSnapshots"),
+	bControlRigUseVMSnapshots,
+	TEXT("If True the VM will try to reuse previous initializations of the same rig."));
+
 UControlRig::UControlRig(const FObjectInitializer& ObjectInitializer)
 	: Super(ObjectInitializer)
 	, DeltaTime(0.0f)
@@ -1483,7 +1489,8 @@ bool UControlRig::ExecuteUnits(FRigUnitContext& InOutContext, const FName& InEve
 		}
 #endif
 		
-		const bool bUseSnapshots = !VM->IsNativized();
+		const bool bUseSnapshots = bControlRigUseVMSnapshots && !VM->IsNativized();
+		
 		TArray<URigVMMemoryStorage*> LocalMemory = VM->GetLocalMemoryArray();
 		TArray<void*> AdditionalArguments;
 		AdditionalArguments.Add(&InOutContext);
