ControlRig anim node: Fixed control rig anim node not able to accept array pin connections

diff --git a/Engine/Plugins/Animation/ControlRig/Source/ControlRig/Private/AnimNode_ControlRig.cpp b/Engine/Plugins/Animation/ControlRig/Source/ControlRig/Private/AnimNode_ControlRig.cpp
--- a/Engine/Plugins/Animation/ControlRig/Source/ControlRig/Private/AnimNode_ControlRig.cpp
+++ b/Engine/Plugins/Animation/ControlRig/Source/ControlRig/Private/AnimNode_ControlRig.cpp
@@ -585,6 +585,13 @@ void FAnimNode_ControlRig::PropagateInputProperties(const UObject* InSourceInsta
 						StructProperty->Struct->CopyScriptStruct(Variable.Memory, SrcPtr, 1);
 					}
 				}
+				else if(FArrayProperty* ArrayProperty = CastField<FArrayProperty>(CallerProperty))
+				{
+					if(ensure(ArrayProperty->SameType(Variable.Property)))
+					{
+						ArrayProperty->CopyCompleteValue(Variable.Memory, SrcPtr);
+					}
+				}
 				else
 				{
 					ensureMsgf(false, TEXT("Property %s type %s not recognized"), *CallerProperty->GetName(), *CallerProperty->GetCPPType());
