Control Rig: Fix crash in unit test ControlRig.Debug.PublicFunction

diff --git a/Engine/Plugins/Animation/ControlRig/Source/ControlRigDeveloper/Private/ControlRigBlueprint.cpp b/Engine/Plugins/Animation/ControlRig/Source/ControlRigDeveloper/Private/ControlRigBlueprint.cpp
--- a/Engine/Plugins/Animation/ControlRig/Source/ControlRigDeveloper/Private/ControlRigBlueprint.cpp
+++ b/Engine/Plugins/Animation/ControlRig/Source/ControlRigDeveloper/Private/ControlRigBlueprint.cpp
@@ -2489,24 +2489,27 @@ TArray<FAssetData> UControlRigBlueprint::GetDependentAssets() const
 				TArray<TSoftObjectPtr<URigVMFunctionReferenceNode>> References = FunctionLibrary->GetReferencesForFunction(FunctionName);
 				for(const TSoftObjectPtr<URigVMFunctionReferenceNode>& Reference : References)
 				{
-					const TSoftObjectPtr<UPackage> Blueprint = Reference->GetTypedOuter<UControlRigBlueprint>();
-					if (!Blueprint.IsNull())
+					if (const URigVMFunctionReferenceNode* ReferencePtr = Reference.Get())
 					{
-						const FSoftObjectPath AssetPath = Blueprint.ToSoftObjectPath();
-						if(AssetPath.GetLongPackageName().StartsWith(TEXT("/Engine/Transient")))
+						if (const UControlRigBlueprint* ControlRigBlueprint = ReferencePtr->GetTypedOuter<UControlRigBlueprint>())
 						{
-							continue;
-						}
-					
-						if(!AssetPaths.Contains(AssetPath))
-						{
-							AssetPaths.Add(AssetPath);
-
-							const FAssetData AssetData = AssetRegistry.GetAssetByObjectPath(AssetPath);
-							if(AssetData.IsValid())
+							const TSoftObjectPtr<UPackage> Blueprint = ControlRigBlueprint;
+							const FSoftObjectPath AssetPath = Blueprint.ToSoftObjectPath();
+							if(AssetPath.GetLongPackageName().StartsWith(TEXT("/Engine/Transient")))
 							{
-								Dependents.Add(AssetData);
+								continue;
 							}
+				
+							if(!AssetPaths.Contains(AssetPath))
+							{
+								AssetPaths.Add(AssetPath);
+
+								const FAssetData AssetData = AssetRegistry.GetAssetByObjectPath(AssetPath);
+								if(AssetData.IsValid())
+								{
+									Dependents.Add(AssetData);
+								}
+							}						
 						}
 					}
 				}
