Sequencer: Tween tool, make sure it works on frame bounds.

diff --git a/Engine/Plugins/Animation/ControlRig/Source/ControlRigEditor/Private/Tools/ControlRigTweener.cpp b/Engine/Plugins/Animation/ControlRig/Source/ControlRigEditor/Private/Tools/ControlRigTweener.cpp
--- a/Engine/Plugins/Animation/ControlRig/Source/ControlRigEditor/Private/Tools/ControlRigTweener.cpp
+++ b/Engine/Plugins/Animation/ControlRig/Source/ControlRigEditor/Private/Tools/ControlRigTweener.cpp
@@ -242,7 +242,7 @@ void FAnimSliderObjectSelection::SetupChannel(FFrameNumber CurrentFrame, TArray<
 		for (int32 Index = 0; Index < KeyTimes.Num(); Index++)
 		{
 			FFrameNumber FrameNumber = KeyTimes[Index];
-			if (FrameNumber < CurrentFrame)
+			if (FrameNumber < CurrentFrame || (FrameNumber == CurrentFrame && KeyBounds.PreviousIndex == INDEX_NONE))
 			{
 				KeyBounds.PreviousIndex = Index;
 				KeyBounds.PreviousFrame = FrameNumber;
@@ -255,7 +255,7 @@ void FAnimSliderObjectSelection::SetupChannel(FFrameNumber CurrentFrame, TArray<
 					KeyBounds.PreviousValue = DoubleValues[Index].Value;
 				}
 			}
-			else if (FrameNumber > CurrentFrame)
+			else if (FrameNumber > CurrentFrame || (FrameNumber == CurrentFrame && Index == KeyTimes.Num() -1))
 			{
 				KeyBounds.NextIndex = Index;
 				KeyBounds.NextFrame = FrameNumber;
