If RigHierarchyController is created off the game thread, remove the async flag once it is assigned to a member variable.

#localization none
#tests Tested on console

diff --git a/Engine/Plugins/Animation/ControlRig/Source/ControlRig/Private/Rigs/RigHierarchy.cpp b/Engine/Plugins/Animation/ControlRig/Source/ControlRig/Private/Rigs/RigHierarchy.cpp
--- a/Engine/Plugins/Animation/ControlRig/Source/ControlRig/Private/Rigs/RigHierarchy.cpp
+++ b/Engine/Plugins/Animation/ControlRig/Source/ControlRig/Private/Rigs/RigHierarchy.cpp
@@ -2271,9 +2271,13 @@ URigHierarchyController* URigHierarchy::GetController(bool bCreateIfNeeded)
 	}
 	else if(bCreateIfNeeded)
 	{
-		 if(!IsGarbageCollecting())
+		 if(ensure(!IsGarbageCollecting()))
 		 {
 			 HierarchyController = NewObject<URigHierarchyController>(this, TEXT("HierarchyController"), RF_Transient);
+			 if(!ensure(IsInGameThread()))
+			 {
+				 HierarchyController->ClearInternalFlags(EInternalObjectFlags::Async);
+			 }
 			 HierarchyController->SetHierarchy(this);
 			 return HierarchyController;
 		 }
