Control Rig: Fix unit test ControlRig.Debug.Snapshot not passing

diff --git a/Engine/Plugins/Animation/ControlRig/Source/ControlRig/Private/ControlRig.cpp b/Engine/Plugins/Animation/ControlRig/Source/ControlRig/Private/ControlRig.cpp
--- a/Engine/Plugins/Animation/ControlRig/Source/ControlRig/Private/ControlRig.cpp
+++ b/Engine/Plugins/Animation/ControlRig/Source/ControlRig/Private/ControlRig.cpp
@@ -1489,7 +1489,8 @@ bool UControlRig::ExecuteUnits(FRigUnitContext& InOutContext, const FName& InEve
 		}
 #endif
 		
-		const bool bUseSnapshots = bControlRigUseVMSnapshots && !VM->IsNativized();
+		const bool bUseInitializationSnapshots = bControlRigUseVMSnapshots && !VM->IsNativized();
+		const bool bUseDebuggingSnapshots = !VM->IsNativized();
 		
 		TArray<URigVMMemoryStorage*> LocalMemory = VM->GetLocalMemoryArray();
 		TArray<void*> AdditionalArguments;
@@ -1499,7 +1500,7 @@ bool UControlRig::ExecuteUnits(FRigUnitContext& InOutContext, const FName& InEve
 
 		if (InOutContext.State == EControlRigState::Init)
 		{
-			if(IsInGameThread() && bUseSnapshots)
+			if(IsInGameThread() && bUseInitializationSnapshots)
 			{
 				const uint32 SnapshotHash = GetHashForInitializeVMSnapShot();
 				UControlRig* CDO = Cast<UControlRig>(GetClass()->GetDefaultObject());
@@ -1556,7 +1557,7 @@ bool UControlRig::ExecuteUnits(FRigUnitContext& InOutContext, const FName& InEve
 		else
 		{
 #if WITH_EDITOR
-			if(bUseSnapshots)
+			if(bUseDebuggingSnapshots)
 			{
 				if(URigVM* SnapShotVM = GetSnapshotVM(false)) // don't create it for normal runs
 				{
