do not compensate in constraint space in the local value has not been set yet

diff --git a/Engine/Plugins/Animation/ControlRig/Source/ControlRigEditor/Private/EditMode/ControlRigEditMode.cpp b/Engine/Plugins/Animation/ControlRig/Source/ControlRigEditor/Private/EditMode/ControlRigEditMode.cpp
--- a/Engine/Plugins/Animation/ControlRig/Source/ControlRigEditor/Private/EditMode/ControlRigEditMode.cpp
+++ b/Engine/Plugins/Animation/ControlRig/Source/ControlRigEditor/Private/EditMode/ControlRigEditMode.cpp
@@ -60,6 +60,7 @@
 #include "LevelEditor.h"
 #include "InteractiveToolManager.h"
 #include "EdModeInteractiveToolsContext.h"
+#include "MovieSceneConstraintChannelHelper.h"
 #include "Editor/EditorPerProjectUserSettings.h"
 #include "TransformConstraint.h"
 
@@ -3817,12 +3818,17 @@ void FControlRigEditMode::SetControlShapeTransform(
 		return false;
 	});
 
-	// set the global space
-	// assumes it's attached to actor
-	ControlRig->SetControlGlobalTransform(
-		InShapeActor->ControlName, InGlobalTransform, bNotify, InContext, bUndo, bPrintPython, bFixEuler);
+	const bool bNeedsConstraintPostProcess = Constraints.IsValidIndex(LastActiveIndex);
+	
+	// set the global space, assumes it's attached to actor
+	// no need to compensate for constraints here, this will be done after when setting the control in the constraint space
+	{
+		TGuardValue<bool> CompensateGuard(FMovieSceneConstraintChannelHelper::bDoNotCompensate, true);
+		ControlRig->SetControlGlobalTransform(
+			InShapeActor->ControlName, InGlobalTransform, bNotify, InContext, bUndo, bPrintPython, bFixEuler);
+	}
 
-	if (!Constraints.IsValidIndex(LastActiveIndex))
+	if (!bNeedsConstraintPostProcess)
 	{
 		return;
 	}
