ensure a minimum widget size for readability

diff --git a/Engine/Plugins/Animation/ControlRig/Source/ControlRigEditor/Private/Editor/Constraints/SConstraintsWidget.cpp b/Engine/Plugins/Animation/ControlRig/Source/ControlRigEditor/Private/Editor/Constraints/SConstraintsWidget.cpp
--- a/Engine/Plugins/Animation/ControlRig/Source/ControlRigEditor/Private/Editor/Constraints/SConstraintsWidget.cpp
+++ b/Engine/Plugins/Animation/ControlRig/Source/ControlRigEditor/Private/Editor/Constraints/SConstraintsWidget.cpp
@@ -934,6 +934,11 @@ TSharedPtr<SWidget> SConstraintsEditionWidget::CreateContextMenu()
 	// constraint details
 	MenuBuilder.BeginSection("EditConstraint", LOCTEXT("EditConstraintHeader", "Edit Constraint"));
 	{
+		// this spacer is used to set a minimum size to the details builder so that the menu won't shrink when
+		// collapsing the transform offset for instance
+		static const FVector2D MinimumDetailsSize(300.0f, 0.0f);
+		MenuBuilder.AddWidget(SNew(SSpacer).Size(MinimumDetailsSize), FText::GetEmpty(), false, false);
+		
 		MenuBuilder.AddWidget(DetailsView, FText::GetEmpty(), true);
 	}
 	MenuBuilder.EndSection();
