Control Rig: Reconstruct channel proxies when control rig initializes and the controls are different in order to handle procedural rig's.

diff --git a/Engine/Plugins/Animation/ControlRig/Source/ControlRig/Private/Sequencer/MovieSceneControlRigParameterSection.cpp b/Engine/Plugins/Animation/ControlRig/Source/ControlRig/Private/Sequencer/MovieSceneControlRigParameterSection.cpp
--- a/Engine/Plugins/Animation/ControlRig/Source/ControlRig/Private/Sequencer/MovieSceneControlRigParameterSection.cpp
+++ b/Engine/Plugins/Animation/ControlRig/Source/ControlRig/Private/Sequencer/MovieSceneControlRigParameterSection.cpp
@@ -1386,6 +1386,37 @@ const TArray< FSpaceControlNameAndChannel>& UMovieSceneControlRigParameterSectio
 	return SpaceChannels;
 }
 
+bool UMovieSceneControlRigParameterSection::IsDifferentThanLastControlsUsedToReconstruct(const TArray<FRigControlElement*>& NewControls) const
+{
+	if (NewControls.Num() != LastControlsUsedToReconstruct.Num())
+	{
+		return true;
+	}
+	for (int32 Index = 0; Index < LastControlsUsedToReconstruct.Num(); ++Index)
+	{
+		//for the channel proxy we really just care about name and type, and if any are nullptr's
+		if (LastControlsUsedToReconstruct[Index].Key != NewControls[Index]->GetName() ||
+			LastControlsUsedToReconstruct[Index].Value != NewControls[Index]->Settings.ControlType)
+		{
+			return true;
+		}
+	}
+	return false;
+}
+
+void UMovieSceneControlRigParameterSection::StoreLastControlsUsedToReconstruct(const TArray<FRigControlElement*>& NewControls)
+{
+	LastControlsUsedToReconstruct.SetNum(NewControls.Num());
+	for (int32 Index = 0; Index < LastControlsUsedToReconstruct.Num(); ++Index)
+	{
+		if (NewControls[Index])
+		{
+			LastControlsUsedToReconstruct[Index].Key = NewControls[Index]->GetName();
+			LastControlsUsedToReconstruct[Index].Value = NewControls[Index]->Settings.ControlType;
+		}
+	}
+}
+
 void UMovieSceneControlRigParameterSection::ReconstructChannelProxy()
 {
 	FMovieSceneChannelProxyData Channels;
@@ -1395,6 +1426,7 @@ void UMovieSceneControlRigParameterSection::ReconstructChannelProxy()
 	{
 		TArray<FRigControlElement*> SortedControls;
 		ControlRig->GetControlsInOrder(SortedControls);
+		StoreLastControlsUsedToReconstruct(SortedControls);
 		bool bIsInUndo = false;
 #if WITH_EDITORONLY_DATA
 		bIsInUndo = GIsTransacting;
diff --git a/Engine/Plugins/Animation/ControlRig/Source/ControlRig/Public/Sequencer/MovieSceneControlRigParameterSection.h b/Engine/Plugins/Animation/ControlRig/Source/ControlRig/Public/Sequencer/MovieSceneControlRigParameterSection.h
--- a/Engine/Plugins/Animation/ControlRig/Source/ControlRig/Public/Sequencer/MovieSceneControlRigParameterSection.h
+++ b/Engine/Plugins/Animation/ControlRig/Source/ControlRig/Public/Sequencer/MovieSceneControlRigParameterSection.h
@@ -497,7 +497,18 @@ protected:
 	mutable bool bDoNotKey;
 
 public:
-	/** Special list of Names that we should only Modify. Needed to handle Interaction (FK/IK) since Control Rig expecting only changed value to be set
-	not all Controls*/
+	// Special list of Names that we should only Modify. Needed to handle Interaction (FK/IK) since Control Rig expecting only changed value to be set
+	//not all Controls
 	mutable TSet<FName> ControlsToSet;
+
+
+public:
+	//Test Controls really are new
+	bool IsDifferentThanLastControlsUsedToReconstruct(const TArray<FRigControlElement*>& NewControls) const;
+
+private:
+	void StoreLastControlsUsedToReconstruct(const TArray<FRigControlElement*>& NewControls);
+	//Last set of Controls used to reconstruct the channel proxies, used to make sure controls really changed if we want to reconstruct
+	//only care to check name and type
+	TArray<TPair<FName, ERigControlType>> LastControlsUsedToReconstruct;
 };
diff --git a/Engine/Plugins/Animation/ControlRig/Source/ControlRigEditor/Private/Sequencer/ControlRigParameterTrackEditor.cpp b/Engine/Plugins/Animation/ControlRig/Source/ControlRigEditor/Private/Sequencer/ControlRigParameterTrackEditor.cpp
--- a/Engine/Plugins/Animation/ControlRig/Source/ControlRigEditor/Private/Sequencer/ControlRigParameterTrackEditor.cpp
+++ b/Engine/Plugins/Animation/ControlRig/Source/ControlRigEditor/Private/Sequencer/ControlRigParameterTrackEditor.cpp
@@ -406,6 +406,7 @@ FControlRigParameterTrackEditor::~FControlRigParameterTrackEditor()
 	FMovieSceneToolsModule::Get().UnregisterAnimationBakeHelper(this);
 }
 
+
 void FControlRigParameterTrackEditor::BindControlRig(UControlRig* ControlRig)
 {
 	if (ControlRig && BoundControlRigs.Contains(ControlRig) == false)
@@ -2909,6 +2910,29 @@ void FControlRigParameterTrackEditor::HandleOnInitialized(UControlRig* ControlRi
 			GetSequencer()->NotifyMovieSceneDataChanged(EMovieSceneDataChangeType::RefreshTree);
 		}
 	}
+	//also do to new procedural rigs, we may be creating controls dynamically, so here we need to check to see if we need to reconstruct channels
+	UMovieScene* MovieScene = GetSequencer()->GetFocusedMovieSceneSequence()->GetMovieScene();
+	const TArray<FMovieSceneBinding>& Bindings = MovieScene->GetBindings();
+	for (const FMovieSceneBinding& Binding : Bindings)
+	{
+		UMovieSceneControlRigParameterTrack* Track = Cast<UMovieSceneControlRigParameterTrack>(MovieScene->FindTrack(UMovieSceneControlRigParameterTrack::StaticClass(), Binding.GetObjectGuid(), NAME_None));
+		if (Track && Track->GetControlRig() == ControlRig)
+		{
+			TArray<FRigControlElement*> SortedControls;
+			ControlRig->GetControlsInOrder(SortedControls);
+			for (UMovieSceneSection* BaseSection : Track->GetAllSections())
+			{
+				if (UMovieSceneControlRigParameterSection* Section = Cast<UMovieSceneControlRigParameterSection>(BaseSection))
+				{
+					if (Section->IsDifferentThanLastControlsUsedToReconstruct(SortedControls))
+					{
+						Section->ReconstructChannelProxy();
+						Section->MarkAsChanged();
+					}
+				}
+			}
+		}
+	}
 }
 
 void FControlRigParameterTrackEditor::HandleControlModified(UControlRig* ControlRig, FRigControlElement* ControlElement, const FRigControlModifiedContext& Context)
