ControlRig: Fix reset transform not setting IsInteracting = true during forward solve

diff --git a/Engine/Plugins/Animation/ControlRig/Source/ControlRigEditor/Private/EditMode/ControlRigEditMode.cpp b/Engine/Plugins/Animation/ControlRig/Source/ControlRigEditor/Private/EditMode/ControlRigEditMode.cpp
--- a/Engine/Plugins/Animation/ControlRig/Source/ControlRigEditor/Private/EditMode/ControlRigEditMode.cpp
+++ b/Engine/Plugins/Animation/ControlRig/Source/ControlRigEditor/Private/EditMode/ControlRigEditMode.cpp
@@ -2926,8 +2926,10 @@ void FControlRigEditMode::ResetTransforms(bool bSelectionOnly)
 		}
 		else
 		{
-			// we at least have to run the interaction event
-			TGuardValue<TArray<FName>> EventGuard(ControlRig->EventQueue, {FRigUnit_InteractionExecution::EventName});
+			// we have to insert the interaction event before we run current events
+			TArray<FName> NewEventQueue = {FRigUnit_InteractionExecution::EventName};
+			NewEventQueue.Append(ControlRig->EventQueue);
+			TGuardValue<TArray<FName>> EventGuard(ControlRig->EventQueue, NewEventQueue);
 			ControlRig->Evaluate_AnyThread();
 			for (const FRigElementKey& ControlToReset : ControlsToReset)
 			{
