Constraints: Baking, missing files to tick the skel meshes when backing

diff --git a/Engine/Plugins/Animation/ControlRig/Source/ControlRig/Private/Constraints/ControlRigTransformableHandle.cpp b/Engine/Plugins/Animation/ControlRig/Source/ControlRig/Private/Constraints/ControlRigTransformableHandle.cpp
--- a/Engine/Plugins/Animation/ControlRig/Source/ControlRig/Private/Constraints/ControlRigTransformableHandle.cpp
+++ b/Engine/Plugins/Animation/ControlRig/Source/ControlRig/Private/Constraints/ControlRigTransformableHandle.cpp
@@ -48,6 +48,33 @@ bool UTransformableControlHandle::IsValid() const
 	return true;
 }
 
+void UTransformableControlHandle::TickForBaking()
+{
+	USkeletalMeshComponent* SkeletalMeshComponent = GetSkeletalMesh();
+	if (SkeletalMeshComponent)
+	{
+		const AActor* Parent = SkeletalMeshComponent->GetOwner();
+		while (Parent)
+		{
+			TArray<USkeletalMeshComponent*> MeshComps;
+			Parent->GetComponents(MeshComps, true);
+
+			for (USkeletalMeshComponent* MeshComp : MeshComps)
+			{
+				MeshComp->TickAnimation(0.03f, false);
+				MeshComp->RefreshBoneTransforms();
+				MeshComp->RefreshFollowerComponents();
+				MeshComp->UpdateComponentToWorld();
+				MeshComp->FinalizeBoneTransform();
+				MeshComp->MarkRenderTransformDirty();
+				MeshComp->MarkRenderDynamicDataDirty();
+			}
+
+			Parent = Parent->GetAttachParentActor();
+		}
+	}
+}
+
 // NOTE should we cache the skeletal mesh and the CtrlIndex to avoid looking for if every time
 // probably not for handling runtime changes
 void UTransformableControlHandle::SetGlobalTransform(const FTransform& InGlobal) const
diff --git a/Engine/Plugins/Animation/ControlRig/Source/ControlRig/Public/Constraints/ControlRigTransformableHandle.h b/Engine/Plugins/Animation/ControlRig/Source/ControlRig/Public/Constraints/ControlRigTransformableHandle.h
--- a/Engine/Plugins/Animation/ControlRig/Source/ControlRig/Public/Constraints/ControlRigTransformableHandle.h
+++ b/Engine/Plugins/Animation/ControlRig/Source/ControlRig/Public/Constraints/ControlRigTransformableHandle.h
@@ -68,6 +68,9 @@ public:
 	/** Resolve the bound objects so that any object it references are resolved and correctly set up*/
 	virtual void ResolveBoundObjects(FMovieSceneSequenceID LocalSequenceID, IMovieScenePlayer& Player) override;
 
+	/** tick things when baking*/ 
+	virtual void TickForBaking() override;
+
 	/** Returns the skeletal mesh bound to ControlRig. */
 	USkeletalMeshComponent* GetSkeletalMesh() const;
 
diff --git a/Engine/Source/Runtime/Experimental/Animation/Constraints/Private/TransformConstraint.cpp b/Engine/Source/Runtime/Experimental/Animation/Constraints/Private/TransformConstraint.cpp
--- a/Engine/Source/Runtime/Experimental/Animation/Constraints/Private/TransformConstraint.cpp
+++ b/Engine/Source/Runtime/Experimental/Animation/Constraints/Private/TransformConstraint.cpp
@@ -251,6 +251,29 @@ void UTickableTransformConstraint::ResolveBoundObjects(FMovieSceneSequenceID Loc
 	}
 }
 
+void UTickableTransformConstraint::Evaluate(bool bTickHandlesAlso) const
+{
+	if (bTickHandlesAlso)
+	{
+		if (ParentTRSHandle)
+		{
+			ParentTRSHandle->TickForBaking();
+		}
+
+		Super::Evaluate();
+
+		if (ChildTRSHandle)
+		{
+			ChildTRSHandle->TickForBaking();
+		}
+	}
+	else
+	{
+		Super::Evaluate();
+	}
+}
+
+
 void UTickableTransformConstraint::SetChildGlobalTransform(const FTransform& InGlobal) const
 {
 	if(ChildTRSHandle && ChildTRSHandle->IsValid())
diff --git a/Engine/Source/Runtime/Experimental/Animation/Constraints/Public/TransformConstraint.h b/Engine/Source/Runtime/Experimental/Animation/Constraints/Public/TransformConstraint.h
--- a/Engine/Source/Runtime/Experimental/Animation/Constraints/Public/TransformConstraint.h
+++ b/Engine/Source/Runtime/Experimental/Animation/Constraints/Public/TransformConstraint.h
@@ -40,7 +40,6 @@ public:
 	
 	/** Test whether an InObject is referenced by that constraint. (i.e. is it's parent or child). */
 	virtual bool ReferencesObject(TWeakObjectPtr<UObject> InObject) const override;
-	
 	/** If true it contains objects bound to an external system, like sequencer so we don't do certain things, like remove constraints when they don't resolve*/
 	virtual bool HasBoundObjects() const override;
 	
@@ -50,6 +49,10 @@ public:
 	/** If Active and the handles and targets are valid*/
 	virtual bool IsFullyActive() const override;
 
+
+	/** Override the evaluate so we can tick our handles*/
+	virtual void Evaluate(bool bTickHandlesAlso = false) const override;
+
 	/** The transformable handle representing the parent of that constraint. */
 	UPROPERTY(BlueprintReadWrite, Category = "Handle")
 	TObjectPtr<UTransformableHandle> ParentTRSHandle;
