Control Rig: Allow dispatch nodes to bind pins to variables

diff --git a/Engine/Plugins/Animation/ControlRig/Source/ControlRigEditor/Private/ControlRigEditorModule.cpp b/Engine/Plugins/Animation/ControlRig/Source/ControlRigEditor/Private/ControlRigEditorModule.cpp
--- a/Engine/Plugins/Animation/ControlRig/Source/ControlRigEditor/Private/ControlRigEditorModule.cpp
+++ b/Engine/Plugins/Animation/ControlRig/Source/ControlRigEditor/Private/ControlRigEditorModule.cpp
@@ -131,6 +131,7 @@
 #include "RigVMModel/Nodes/RigVMAggregateNode.h"
 #include "RigVMUserWorkflowRegistry.h"
 #include "UserDefinedStructureCompilerUtils.h"
+#include "RigVMModel/Nodes/RigVMDispatchNode.h"
 #include "Units/ControlRigNodeWorkflow.h"
 #include "Units/RigDispatchFactory.h"
 
@@ -1513,6 +1514,7 @@ void FControlRigEditorModule::GetContextMenuActions(const UControlRigGraphSchema
 					}
 
 					if (Cast<URigVMUnitNode>(ModelPin->GetNode()) != nullptr || 
+						Cast<URigVMDispatchNode>(ModelPin->GetNode()) != nullptr || 
 						Cast<URigVMLibraryNode>(ModelPin->GetNode()) != nullptr ||
 						Cast<URigVMArrayNode>(ModelPin->GetNode()) != nullptr)
 					{
