Control Rig: Prevent UActorFactorySkeletalMesh::CanCreateActorFrom loading unrelated assets if avoidable

diff --git a/Engine/Plugins/Animation/ControlRig/Source/ControlRigEditor/Private/ControlRigBlueprintActions.cpp b/Engine/Plugins/Animation/ControlRig/Source/ControlRigEditor/Private/ControlRigBlueprintActions.cpp
--- a/Engine/Plugins/Animation/ControlRig/Source/ControlRigEditor/Private/ControlRigBlueprintActions.cpp
+++ b/Engine/Plugins/Animation/ControlRig/Source/ControlRigEditor/Private/ControlRigBlueprintActions.cpp
@@ -250,9 +250,14 @@ UControlRigBlueprint* FControlRigBlueprintActions::CreateControlRigFromSkeletalM
 	return NewControlRigBlueprint;
 }
 
-USkeletalMesh* FControlRigBlueprintActions::GetSkeletalMeshFromControlRigBlueprint(UObject* InAsset)
+USkeletalMesh* FControlRigBlueprintActions::GetSkeletalMeshFromControlRigBlueprint(const FAssetData& InAsset)
 {
-	if (UControlRigBlueprint* Blueprint = Cast<UControlRigBlueprint>(InAsset))
+	if (InAsset.GetClass() != UControlRigBlueprint::StaticClass())
+	{
+		return nullptr;
+	}
+	
+	if (const UControlRigBlueprint* Blueprint = Cast<UControlRigBlueprint>(InAsset.GetAsset()))
 	{
 		return Blueprint->GetPreviewMesh();
 	}
diff --git a/Engine/Plugins/Animation/ControlRig/Source/ControlRigEditor/Private/ControlRigBlueprintActions.h b/Engine/Plugins/Animation/ControlRig/Source/ControlRigEditor/Private/ControlRigBlueprintActions.h
--- a/Engine/Plugins/Animation/ControlRig/Source/ControlRigEditor/Private/ControlRigBlueprintActions.h
+++ b/Engine/Plugins/Animation/ControlRig/Source/ControlRigEditor/Private/ControlRigBlueprintActions.h
@@ -33,7 +33,7 @@ public:
 	static UControlRigBlueprint* CreateNewControlRigAsset(const FString& InDesiredPackagePath);
 	static UControlRigBlueprint* CreateControlRigFromSkeletalMeshOrSkeleton(UObject* InSelectedObject);
 
-	static USkeletalMesh* GetSkeletalMeshFromControlRigBlueprint(UObject* InAsset);
+	static USkeletalMesh* GetSkeletalMeshFromControlRigBlueprint(const FAssetData& InAsset);
 	static void PostSpawningSkeletalMeshActor(AActor* InSpawnedActor, UObject* InAsset);
 	static void OnSpawnedSkeletalMeshActorChanged(UObject* InObject, FPropertyChangedEvent& InEvent, UObject* InAsset);
 
diff --git a/Engine/Source/Editor/UnrealEd/Classes/ActorFactories/ActorFactorySkeletalMesh.h b/Engine/Source/Editor/UnrealEd/Classes/ActorFactories/ActorFactorySkeletalMesh.h
--- a/Engine/Source/Editor/UnrealEd/Classes/ActorFactories/ActorFactorySkeletalMesh.h
+++ b/Engine/Source/Editor/UnrealEd/Classes/ActorFactories/ActorFactorySkeletalMesh.h
@@ -12,7 +12,7 @@ class AActor;
 struct FAssetData;
 class USkeletalMesh;
 
-DECLARE_DELEGATE_RetVal_OneParam(USkeletalMesh*, FGetSkeletalMeshFromAssetDelegate, UObject* /* in asset */);
+DECLARE_DELEGATE_RetVal_OneParam(USkeletalMesh*, FGetSkeletalMeshFromAssetDelegate, const FAssetData& /* in asset */);
 DECLARE_DELEGATE_TwoParams(FPostSkeletalMeshActorSpawnedDelegate, AActor* /* spawned actor */, UObject* /* asset source*/);
 
 UCLASS(MinimalAPI, config=Editor)
diff --git a/Engine/Source/Editor/UnrealEd/Private/Factories/ActorFactory.cpp b/Engine/Source/Editor/UnrealEd/Private/Factories/ActorFactory.cpp
--- a/Engine/Source/Editor/UnrealEd/Private/Factories/ActorFactory.cpp
+++ b/Engine/Source/Editor/UnrealEd/Private/Factories/ActorFactory.cpp
@@ -1261,7 +1261,7 @@ bool UActorFactorySkeletalMesh::CanCreateActorFrom( const FAssetData& AssetData,
 	{
 		for (const auto& Pair : UActorFactorySkeletalMesh::GetSkeletalMeshDelegates)
 		{
-			if (Pair.Value.Execute(AssetData.GetAsset()) != nullptr)
+			if (Pair.Value.Execute(AssetData) != nullptr)
 			{
 				return true;
 			}
