Control Rig: Fix variable default value updates not being considered after compile

diff --git a/Engine/Plugins/Animation/ControlRig/Source/ControlRigEditor/Private/Editor/ControlRigEditor.cpp b/Engine/Plugins/Animation/ControlRig/Source/ControlRigEditor/Private/Editor/ControlRigEditor.cpp
--- a/Engine/Plugins/Animation/ControlRig/Source/ControlRigEditor/Private/Editor/ControlRigEditor.cpp
+++ b/Engine/Plugins/Animation/ControlRig/Source/ControlRigEditor/Private/Editor/ControlRigEditor.cpp
@@ -2021,9 +2021,21 @@ void FControlRigEditor::Compile()
 		}
 
 		// store the defaults from the CDO back on the new variables list
+		bool bAnyVariableValueChanged = false;
 		for(FBPVariableDescription& NewVariable : RigBlueprint->NewVariables)
 		{
-			UpdateDefaultValueForVariable(NewVariable, true);
+			bAnyVariableValueChanged |= UpdateDefaultValueForVariable(NewVariable, true);
+		}
+		if (bAnyVariableValueChanged)
+		{
+			// Go over all the instances to update the default values from CDO
+			for (const FCustomDebugObject& DebugObject : DebugList)
+			{
+				if (UControlRig* DebuggedRig = Cast<UControlRig>(DebugObject.Object))
+				{
+					DebuggedRig->CopyExternalVariableDefaultValuesFromCDO();
+				}
+			}
 		}
 
 		if (SelectedObjects.Num() > 0)
@@ -4833,8 +4845,9 @@ void FControlRigEditor::HandleJumpToHyperlink(const UObject* InSubject)
 	}
 }
 
-void FControlRigEditor::UpdateDefaultValueForVariable(FBPVariableDescription& InVariable, bool bUseCDO)
+bool FControlRigEditor::UpdateDefaultValueForVariable(FBPVariableDescription& InVariable, bool bUseCDO)
 {
+	bool bAnyValueChanged = false;
 	if (UControlRigBlueprint* ControlRigBP = GetControlRigBlueprint())
 	{
 		UClass* GeneratedClass = ControlRigBP->GeneratedClass;
@@ -4848,10 +4861,15 @@ void FControlRigEditor::UpdateDefaultValueForVariable(FBPVariableDescription& In
 				FString NewDefaultValue;
 				const uint8* Container = (const uint8*)ObjectContainer;
 				FBlueprintEditorUtils::PropertyValueToString(TargetProperty, Container, NewDefaultValue, nullptr);
-				InVariable.DefaultValue = NewDefaultValue;
+				if (InVariable.DefaultValue != NewDefaultValue)
+				{
+					InVariable.DefaultValue = NewDefaultValue;
+					bAnyValueChanged = true;
+				}
 			}
 		}
 	}
+	return bAnyValueChanged;
 }
 
 bool FControlRigEditor::SelectLocalVariable(const UEdGraph* Graph, const FName& VariableName)
diff --git a/Engine/Plugins/Animation/ControlRig/Source/ControlRigEditor/Private/Editor/ControlRigEditor.h b/Engine/Plugins/Animation/ControlRig/Source/ControlRigEditor/Private/Editor/ControlRigEditor.h
--- a/Engine/Plugins/Animation/ControlRig/Source/ControlRigEditor/Private/Editor/ControlRigEditor.h
+++ b/Engine/Plugins/Animation/ControlRig/Source/ControlRigEditor/Private/Editor/ControlRigEditor.h
@@ -427,7 +427,7 @@ protected:
 	FRigVMController_BulkEditResult OnRequestBulkEditDialog(UControlRigBlueprint* InBlueprint, URigVMController* InController, URigVMLibraryNode* InFunction, ERigVMControllerBulkEditType InEditType);
 	bool OnRequestBreakLinksDialog(TArray<URigVMLink*> InLinks);
 	void HandleJumpToHyperlink(const UObject* InSubject);
-	void UpdateDefaultValueForVariable(FBPVariableDescription& InVariable, bool bUseCDO);
+	bool UpdateDefaultValueForVariable(FBPVariableDescription& InVariable, bool bUseCDO);
 
 	URigVMGraph* GetFocusedModel() const;
 	URigVMController* GetFocusedController() const;
