Sequencer: After pasting a binding with a control rig in it the control rig wouldn't save since it woudl be set to transient. Regression from 21199364 , where we removed clearing transient flags from all subobjects since Control rig does have valid transient objects on it.
#trivial

diff --git a/Engine/Plugins/Animation/ControlRig/Source/ControlRig/Private/Sequencer/MovieSceneControlRigParameterTrack.cpp b/Engine/Plugins/Animation/ControlRig/Source/ControlRig/Private/Sequencer/MovieSceneControlRigParameterTrack.cpp
--- a/Engine/Plugins/Animation/ControlRig/Source/ControlRig/Private/Sequencer/MovieSceneControlRigParameterTrack.cpp
+++ b/Engine/Plugins/Animation/ControlRig/Source/ControlRig/Private/Sequencer/MovieSceneControlRigParameterTrack.cpp
@@ -458,7 +458,11 @@ void UMovieSceneControlRigParameterTrack::HandleControlRigPackageDone(UControlRi
 void UMovieSceneControlRigParameterTrack::PostEditImport()
 {
 	Super::PostEditImport();
-
+	if (ControlRig)
+	{
+		ControlRig->ClearFlags(RF_Transient); //when copied make sure it's no longer transient, sequencer does this for tracks/sections 
+											  //but not for all objects in them since the control rig itself has transient objects.
+	}
 	ReconstructControlRig();
 }
 
