Sequencer: Control Rig: Keep track of last evaluated time so we can correctly do additive section compensation with control rig. Since we can also recieve time from a control rig event need to make sure to use that as the evaluate time also.

diff --git a/Engine/Plugins/Animation/ControlRig/Source/ControlRigEditor/Private/Sequencer/ControlRigParameterTrackEditor.cpp b/Engine/Plugins/Animation/ControlRig/Source/ControlRigEditor/Private/Sequencer/ControlRigParameterTrackEditor.cpp
--- a/Engine/Plugins/Animation/ControlRig/Source/ControlRigEditor/Private/Sequencer/ControlRigParameterTrackEditor.cpp
+++ b/Engine/Plugins/Animation/ControlRig/Source/ControlRigEditor/Private/Sequencer/ControlRigParameterTrackEditor.cpp
@@ -3312,7 +3312,7 @@ void FControlRigParameterTrackEditor::GetControlRigKeys(
 }
 
 FKeyPropertyResult FControlRigParameterTrackEditor::AddKeysToControlRigHandle(USceneComponent* InSceneComp, UControlRig* InControlRig,
-	FGuid ObjectHandle, FFrameNumber KeyTime, FGeneratedTrackKeys& GeneratedKeys,
+	FGuid ObjectHandle, FFrameNumber KeyTime, FFrameNumber EvaluateTime, FGeneratedTrackKeys& GeneratedKeys,
 	ESequencerKeyMode KeyMode, TSubclassOf<UMovieSceneTrack> TrackClass, FName ControlRigName, FName RigControlName)
 {
 	EAutoChangeMode AutoChangeMode = GetSequencer()->GetAutoChangeMode();
@@ -3360,7 +3360,7 @@ FKeyPropertyResult FControlRigParameterTrackEditor::AddKeysToControlRigHandle(US
 		{
 			if (!bTrackCreated)
 			{
-				ModifyOurGeneratedKeysByCurrentAndWeight(InSceneComp, InControlRig, RigControlName, Track, SectionToKey, KeyTime, GeneratedKeys, Weight);
+				ModifyOurGeneratedKeysByCurrentAndWeight(InSceneComp, InControlRig, RigControlName, Track, SectionToKey, EvaluateTime, GeneratedKeys, Weight);
 			}
 			UMovieSceneControlRigParameterSection* ParamSection = Cast<UMovieSceneControlRigParameterSection>(SectionToKey);
 			if (!ParamSection->GetDoNotKey())
@@ -3394,7 +3394,7 @@ FKeyPropertyResult FControlRigParameterTrackEditor::AddKeysToControlRigHandle(US
 }
 
 FKeyPropertyResult FControlRigParameterTrackEditor::AddKeysToControlRig(
-	USceneComponent* InSceneComp, UControlRig* InControlRig, FFrameNumber KeyTime, FGeneratedTrackKeys& GeneratedKeys,
+	USceneComponent* InSceneComp, UControlRig* InControlRig, FFrameNumber KeyTime, FFrameNumber EvaluateTime, FGeneratedTrackKeys& GeneratedKeys,
 	ESequencerKeyMode KeyMode, TSubclassOf<UMovieSceneTrack> TrackClass, FName ControlRigName, FName RigControlName)
 {
 	FKeyPropertyResult KeyPropertyResult;
@@ -3411,7 +3411,7 @@ FKeyPropertyResult FControlRigParameterTrackEditor::AddKeysToControlRig(
 	KeyPropertyResult.bHandleCreated = HandleResult.bWasCreated;
 	if (ObjectHandle.IsValid())
 	{
-		KeyPropertyResult |= AddKeysToControlRigHandle(InSceneComp, InControlRig, ObjectHandle, KeyTime, GeneratedKeys, KeyMode, TrackClass, ControlRigName, RigControlName);
+		KeyPropertyResult |= AddKeysToControlRigHandle(InSceneComp, InControlRig, ObjectHandle, KeyTime, EvaluateTime, GeneratedKeys, KeyMode, TrackClass, ControlRigName, RigControlName);
 	}
 
 	return KeyPropertyResult;
@@ -3474,20 +3474,25 @@ void FControlRigParameterTrackEditor::AddControlKeys(
 	auto OnKeyProperty = [=](FFrameNumber Time) -> FKeyPropertyResult
 	{
 		FFrameNumber LocalTime = Time;
+		//for modify weights we evaluate so need to make sure we use the evaluated time
+		FFrameNumber EvaluateTime = GetSequencer()->GetLastEvaluatedLocalTime().RoundToFrame();
+		//if InLocalTime is specified that means time value was set with SetControlValue, so we don't use sequencer times at all, but this time instead
 		if (InLocalTime != FLT_MAX)
 		{
 			//convert from frame time since conversion may give us one frame less, e.g 1.53333330 * 24000.0/1.0 = 36799.999199999998
 			FFrameTime LocalFrameTime = GetSequencer()->GetFocusedTickResolution().AsFrameTime((double)InLocalTime);
 			LocalTime = LocalFrameTime.RoundToFrame();
+			EvaluateTime = LocalTime;
 		}
-		return this->AddKeysToControlRig(InSceneComp, InControlRig, LocalTime, *GeneratedKeys, KeyMode, UMovieSceneControlRigParameterTrack::StaticClass(), ControlRigName, RigControlName);
+		
+		return this->AddKeysToControlRig(InSceneComp, InControlRig, LocalTime, EvaluateTime, *GeneratedKeys, KeyMode, UMovieSceneControlRigParameterTrack::StaticClass(), ControlRigName, RigControlName);
 	};
 
 	AnimatablePropertyChanged(FOnKeyProperty::CreateLambda(OnKeyProperty));
 	EndKeying(); //fine even if we didn't BeginKeying
 }
 
-bool FControlRigParameterTrackEditor::ModifyOurGeneratedKeysByCurrentAndWeight(UObject* Object, UControlRig* InControlRig, FName RigControlName, UMovieSceneTrack* Track, UMovieSceneSection* SectionToKey, FFrameNumber KeyTime, FGeneratedTrackKeys& GeneratedTotalKeys, float Weight) const
+bool FControlRigParameterTrackEditor::ModifyOurGeneratedKeysByCurrentAndWeight(UObject* Object, UControlRig* InControlRig, FName RigControlName, UMovieSceneTrack* Track, UMovieSceneSection* SectionToKey, FFrameNumber EvaluateTime, FGeneratedTrackKeys& GeneratedTotalKeys, float Weight) const
 {
 	FFrameRate TickResolution = GetSequencer()->GetFocusedTickResolution();
 	FMovieSceneEvaluationTrack EvalTrack = CastChecked<UMovieSceneControlRigParameterTrack>(Track)->GenerateTrackTemplate(Track);
@@ -3495,7 +3500,8 @@ bool FControlRigParameterTrackEditor::ModifyOurGeneratedKeysByCurrentAndWeight(U
 	FMovieSceneInterrogationData InterrogationData;
 	GetSequencer()->GetEvaluationTemplate().CopyActuators(InterrogationData.GetAccumulator());
 
-	FMovieSceneContext Context(FMovieSceneEvaluationRange(KeyTime, GetSequencer()->GetFocusedTickResolution()));
+	//use the EvaluateTime to do the evaluation, may be different than the actually time we key
+	FMovieSceneContext Context(FMovieSceneEvaluationRange(EvaluateTime, GetSequencer()->GetFocusedTickResolution()));
 	EvalTrack.Interrogate(Context, InterrogationData, Object);
 	TArray<FRigControlElement*> Controls = InControlRig->AvailableControls();
 	UMovieSceneControlRigParameterSection* Section = Cast<UMovieSceneControlRigParameterSection>(SectionToKey);
@@ -3521,7 +3527,7 @@ bool FControlRigParameterTrackEditor::ModifyOurGeneratedKeysByCurrentAndWeight(U
 					{
 						ChannelIndex = pChannelIndex->GeneratedKeyIndex;
 						float FVal = (float)Val.Val;
-						GeneratedTotalKeys[ChannelIndex]->ModifyByCurrentAndWeight(Proxy, KeyTime, (void *)&FVal, Weight);
+						GeneratedTotalKeys[ChannelIndex]->ModifyByCurrentAndWeight(Proxy, EvaluateTime, (void *)&FVal, Weight);
 					}
 					break;
 				}
@@ -3546,9 +3552,9 @@ bool FControlRigParameterTrackEditor::ModifyOurGeneratedKeysByCurrentAndWeight(U
 					{
 						ChannelIndex = pChannelIndex->GeneratedKeyIndex;
 						float FVal = (float)Val.Val.X;
-						GeneratedTotalKeys[ChannelIndex]->ModifyByCurrentAndWeight(Proxy, KeyTime, (void *)&FVal, Weight);
+						GeneratedTotalKeys[ChannelIndex]->ModifyByCurrentAndWeight(Proxy, EvaluateTime, (void *)&FVal, Weight);
 						FVal = (float)Val.Val.Y;
-						GeneratedTotalKeys[ChannelIndex + 1]->ModifyByCurrentAndWeight(Proxy, KeyTime, (void *)&FVal, Weight);
+						GeneratedTotalKeys[ChannelIndex + 1]->ModifyByCurrentAndWeight(Proxy, EvaluateTime, (void *)&FVal, Weight);
 					}
 					break;
 				}
@@ -3571,11 +3577,11 @@ bool FControlRigParameterTrackEditor::ModifyOurGeneratedKeysByCurrentAndWeight(U
 						if (ControlElement->Settings.ControlType != ERigControlType::Rotator)
 						{
 							float FVal = (float)Val.Val.X;
-							GeneratedTotalKeys[ChannelIndex]->ModifyByCurrentAndWeight(Proxy, KeyTime, (void*)&FVal, Weight);
+							GeneratedTotalKeys[ChannelIndex]->ModifyByCurrentAndWeight(Proxy, EvaluateTime, (void*)&FVal, Weight);
 							FVal = (float)Val.Val.Y;
-							GeneratedTotalKeys[ChannelIndex + 1]->ModifyByCurrentAndWeight(Proxy, KeyTime, (void*)&FVal, Weight);
+							GeneratedTotalKeys[ChannelIndex + 1]->ModifyByCurrentAndWeight(Proxy, EvaluateTime, (void*)&FVal, Weight);
 							FVal = (float)Val.Val.Z;
-							GeneratedTotalKeys[ChannelIndex + 2]->ModifyByCurrentAndWeight(Proxy, KeyTime, (void*)&FVal, Weight);
+							GeneratedTotalKeys[ChannelIndex + 2]->ModifyByCurrentAndWeight(Proxy, EvaluateTime, (void*)&FVal, Weight);
 						}
 					}
 					break;
@@ -3605,20 +3611,20 @@ bool FControlRigParameterTrackEditor::ModifyOurGeneratedKeysByCurrentAndWeight(U
 						
 						FVector3f CurrentPos = (FVector3f)Val.Val.GetLocation();
 						FRotator3f CurrentRot = FRotator3f(Val.Val.Rotator());
-						GeneratedTotalKeys[ChannelIndex]->ModifyByCurrentAndWeight(Proxy, KeyTime, (void*)&CurrentPos.X, Weight);
-						GeneratedTotalKeys[ChannelIndex + 1]->ModifyByCurrentAndWeight(Proxy, KeyTime, (void*)&CurrentPos.Y, Weight);
-						GeneratedTotalKeys[ChannelIndex + 2]->ModifyByCurrentAndWeight(Proxy, KeyTime, (void*)&CurrentPos.Z, Weight);
+						GeneratedTotalKeys[ChannelIndex]->ModifyByCurrentAndWeight(Proxy, EvaluateTime, (void*)&CurrentPos.X, Weight);
+						GeneratedTotalKeys[ChannelIndex + 1]->ModifyByCurrentAndWeight(Proxy, EvaluateTime, (void*)&CurrentPos.Y, Weight);
+						GeneratedTotalKeys[ChannelIndex + 2]->ModifyByCurrentAndWeight(Proxy, EvaluateTime, (void*)&CurrentPos.Z, Weight);
 
-						GeneratedTotalKeys[ChannelIndex + 3]->ModifyByCurrentAndWeight(Proxy, KeyTime, (void*)&CurrentRot.Roll, Weight);
-						GeneratedTotalKeys[ChannelIndex + 4]->ModifyByCurrentAndWeight(Proxy, KeyTime, (void*)&CurrentRot.Pitch, Weight);
-						GeneratedTotalKeys[ChannelIndex + 5]->ModifyByCurrentAndWeight(Proxy, KeyTime, (void*)&CurrentRot.Yaw, Weight);
+						GeneratedTotalKeys[ChannelIndex + 3]->ModifyByCurrentAndWeight(Proxy, EvaluateTime, (void*)&CurrentRot.Roll, Weight);
+						GeneratedTotalKeys[ChannelIndex + 4]->ModifyByCurrentAndWeight(Proxy, EvaluateTime, (void*)&CurrentRot.Pitch, Weight);
+						GeneratedTotalKeys[ChannelIndex + 5]->ModifyByCurrentAndWeight(Proxy, EvaluateTime, (void*)&CurrentRot.Yaw, Weight);
 
 						if (ControlElement->Settings.ControlType == ERigControlType::Transform || ControlElement->Settings.ControlType == ERigControlType::EulerTransform)
 						{
 							FVector3f CurrentScale = (FVector3f)Val.Val.GetScale3D();
-							GeneratedTotalKeys[ChannelIndex + 6]->ModifyByCurrentAndWeight(Proxy, KeyTime, (void*)&CurrentScale.X, Weight);
-							GeneratedTotalKeys[ChannelIndex + 7]->ModifyByCurrentAndWeight(Proxy, KeyTime, (void*)&CurrentScale.Y, Weight);
-							GeneratedTotalKeys[ChannelIndex + 8]->ModifyByCurrentAndWeight(Proxy, KeyTime, (void*)&CurrentScale.Z, Weight);
+							GeneratedTotalKeys[ChannelIndex + 6]->ModifyByCurrentAndWeight(Proxy, EvaluateTime, (void*)&CurrentScale.X, Weight);
+							GeneratedTotalKeys[ChannelIndex + 7]->ModifyByCurrentAndWeight(Proxy, EvaluateTime, (void*)&CurrentScale.Y, Weight);
+							GeneratedTotalKeys[ChannelIndex + 8]->ModifyByCurrentAndWeight(Proxy, EvaluateTime, (void*)&CurrentScale.Z, Weight);
 						}
 					}
 					break;
diff --git a/Engine/Plugins/Animation/ControlRig/Source/ControlRigEditor/Private/Sequencer/ControlRigParameterTrackEditor.h b/Engine/Plugins/Animation/ControlRig/Source/ControlRigEditor/Private/Sequencer/ControlRigParameterTrackEditor.h
--- a/Engine/Plugins/Animation/ControlRig/Source/ControlRigEditor/Private/Sequencer/ControlRigParameterTrackEditor.h
+++ b/Engine/Plugins/Animation/ControlRig/Source/ControlRigEditor/Private/Sequencer/ControlRigParameterTrackEditor.h
@@ -209,10 +209,10 @@ public:
 		ESequencerKeyMode KeyMode, UMovieSceneControlRigParameterSection* SectionToKey,	FGeneratedTrackKeys& OutGeneratedKeys,
 		const bool bInConstraintSpace = false);
 	FKeyPropertyResult AddKeysToControlRig(
-		USceneComponent *InSceneComp, UControlRig* InControlRig, FFrameNumber KeyTime, FGeneratedTrackKeys& GeneratedKeys,
+		USceneComponent *InSceneComp, UControlRig* InControlRig, FFrameNumber KeyTime, FFrameNumber EvaluateTime, FGeneratedTrackKeys& GeneratedKeys,
 		ESequencerKeyMode KeyMode, TSubclassOf<UMovieSceneTrack> TrackClass, FName ControlRigName, FName RigControlName);
 	FKeyPropertyResult AddKeysToControlRigHandle(USceneComponent *InSceneComp, UControlRig* InControlRig,
-		FGuid ObjectHandle, FFrameNumber KeyTime, FGeneratedTrackKeys& GeneratedKeys,
+		FGuid ObjectHandle, FFrameNumber KeyTime, FFrameNumber EvaluateTime, FGeneratedTrackKeys& GeneratedKeys,
 		ESequencerKeyMode KeyMode, TSubclassOf<UMovieSceneTrack> TrackClass, FName ControlRigName, FName RigControlName);
 	/**
 	 * Modify the passed in Generated Keys by the current tracks values and weight at the passed in time.
@@ -224,7 +224,7 @@ public:
 	 * @param InOutGeneratedTrackKeys The Keys we need to modify. We change these values.
 	 * @param Weight The weight we need to modify the values by.
 	 */
-	bool ModifyOurGeneratedKeysByCurrentAndWeight(UObject* Object, UControlRig* InControlRig, FName RigControlName, UMovieSceneTrack *Track, UMovieSceneSection* SectionToKey, FFrameNumber Time, FGeneratedTrackKeys& InOutGeneratedTotalKeys, float Weight) const;
+	bool ModifyOurGeneratedKeysByCurrentAndWeight(UObject* Object, UControlRig* InControlRig, FName RigControlName, UMovieSceneTrack *Track, UMovieSceneSection* SectionToKey, FFrameNumber EvaluateTime, FGeneratedTrackKeys& InOutGeneratedTotalKeys, float Weight) const;
 
 	//**Function to collapse all layers from this section onto the first absoluate layer.*/
 	static bool CollapseAllLayers(TSharedPtr<ISequencer>& SequencerPtr, UMovieSceneTrack* OwnerTrack, UMovieSceneControlRigParameterSection* ParameterSection, bool bKeyReduce = false, float Tolerance = 0.001f);
diff --git a/Engine/Source/Editor/Sequencer/Private/Sequencer.cpp b/Engine/Source/Editor/Sequencer/Private/Sequencer.cpp
--- a/Engine/Source/Editor/Sequencer/Private/Sequencer.cpp
+++ b/Engine/Source/Editor/Sequencer/Private/Sequencer.cpp
@@ -2965,6 +2965,11 @@ FQualifiedFrameTime FSequencer::GetGlobalTime() const
 	return FQualifiedFrameTime(RootTime, PlayPosition.GetOutputRate());
 }
 
+FFrameTime FSequencer::GetLastEvaluatedLocalTime() const 
+{
+	return LastEvaluatedLocalTime;
+}
+
 void FSequencer::SetLocalTime( FFrameTime NewTime, ESnapTimeMode SnapTimeMode, bool bEvaluate)
 {
 	FFrameRate LocalResolution = GetFocusedTickResolution();
@@ -3008,7 +3013,6 @@ void FSequencer::SetLocalTime( FFrameTime NewTime, ESnapTimeMode SnapTimeMode, b
 	SetLocalTimeDirectly(NewTime, bEvaluate);
 }
 
-
 void FSequencer::SetLocalTimeDirectly(FFrameTime NewTime, bool bEvaluate)
 {
 	TWeakPtr<SWidget> PreviousFocusedWidget = FSlateApplication::Get().GetKeyboardFocusedWidget();
@@ -3106,6 +3110,9 @@ void FSequencer::ForceEvaluate()
 
 void FSequencer::EvaluateInternal(FMovieSceneEvaluationRange InRange, bool bHasJumped)
 {
+
+	LastEvaluatedLocalTime = GetLocalTime().Time;
+
 	if (Settings->ShouldCompileDirectorOnEvaluate())
 	{
 		RecompileDirtyDirectors();
diff --git a/Engine/Source/Editor/Sequencer/Private/Sequencer.h b/Engine/Source/Editor/Sequencer/Private/Sequencer.h
--- a/Engine/Source/Editor/Sequencer/Private/Sequencer.h
+++ b/Engine/Source/Editor/Sequencer/Private/Sequencer.h
@@ -707,6 +707,7 @@ public:
 	virtual uint32 GetLocalLoopIndex() const override;
 	virtual void SetLocalTime(FFrameTime Time, ESnapTimeMode SnapTimeMode = ESnapTimeMode::STM_None, bool bEvaluate = true) override;
 	virtual void SetLocalTimeDirectly(FFrameTime NewTime, bool bEvaluate = true) override;
+	virtual FFrameTime GetLastEvaluatedLocalTime() const override;
 	virtual void SetGlobalTime(FFrameTime Time, bool bEvaluate = true) override;
 	virtual void PlayTo(FMovieSceneSequencePlaybackParams PlaybackParams) override;
 	virtual void RestorePlaybackSpeed() override;
@@ -1439,6 +1440,8 @@ private:
 		TWeakObjectPtr<AActor> PreviousCamera;
 		TWeakObjectPtr<AActor> NextCamera;
 	};
+	/** Frame time last evaluated at, needed to make sure we evaluate at the right time when interrogating values*/
+	FFrameTime LastEvaluatedLocalTime;
 	/** Information for previewing camera cut blends. This will be applied to the editor viewport during blends. */
 	FViewModifierInfo ViewModifierInfo;
 	/** Information cached before entering silent mode, so we can restore it afterwards. */
diff --git a/Engine/Source/Editor/Sequencer/Public/ISequencer.h b/Engine/Source/Editor/Sequencer/Public/ISequencer.h
--- a/Engine/Source/Editor/Sequencer/Public/ISequencer.h
+++ b/Engine/Source/Editor/Sequencer/Public/ISequencer.h
@@ -415,6 +415,9 @@ public:
 	/** Set the global time directly, without performing any auto-scroll, snapping or other adjustments to the supplied time  */
 	virtual void SetGlobalTime(FFrameTime Time, bool bEvaluate = true) = 0;
 
+	/** Get the last evaluated time, which may be different form the current local time*/
+	virtual FFrameTime GetLastEvaluatedLocalTime() const = 0;
+
 	/** Play from the current time to the requested time */
 	virtual void PlayTo(FMovieSceneSequencePlaybackParams PlaybackParams) = 0;
 
