Control Rig: Fix Adding execute context to function with bad direction and name

diff --git a/Engine/Plugins/Animation/ControlRig/Source/ControlRigDeveloper/Private/Graph/ControlRigGraphSchema.cpp b/Engine/Plugins/Animation/ControlRig/Source/ControlRigDeveloper/Private/Graph/ControlRigGraphSchema.cpp
--- a/Engine/Plugins/Animation/ControlRig/Source/ControlRigDeveloper/Private/Graph/ControlRigGraphSchema.cpp
+++ b/Engine/Plugins/Animation/ControlRig/Source/ControlRigDeveloper/Private/Graph/ControlRigGraphSchema.cpp
@@ -1633,6 +1633,7 @@ UEdGraphPin* UControlRigGraphSchema::DropPinOnNode(UEdGraphNode* InTargetNode, c
 					{
 						if (URigVMController* Controller = RigBlueprint->GetController(Model))
 						{
+							FName SourcePinName = InSourcePinName; 
 							FString TypeName = ExternalVar.TypeName.ToString();
 							if (ExternalVar.bIsArray)
 							{
@@ -1642,6 +1643,15 @@ UEdGraphPin* UControlRigGraphSchema::DropPinOnNode(UEdGraphNode* InTargetNode, c
 							if (ExternalVar.TypeObject)
 							{
 								TypeObjectPathName = *ExternalVar.TypeObject->GetPathName();
+
+								if(const UScriptStruct* CPPTypeStruct = Cast<UScriptStruct>(ExternalVar.TypeObject))
+								{
+									if(CPPTypeStruct->IsChildOf(FRigVMExecuteContext::StaticStruct()))
+									{
+										SourcePinName = FRigVMStruct::ExecuteContextName;
+										PinDirection = ERigVMPinDirection::IO;										
+									}
+								}
 							}
 
 							FString DefaultValue;
@@ -1657,7 +1667,7 @@ UEdGraphPin* UControlRigGraphSchema::DropPinOnNode(UEdGraphNode* InTargetNode, c
 							}
 
 							FName ExposedPinName = Controller->AddExposedPin(
-								InSourcePinName,
+								SourcePinName,
 								PinDirection,
 								TypeName,
 								TypeObjectPathName,
