Control Rig: Fix string split unit node

diff --git a/Engine/Plugins/Animation/ControlRig/Source/ControlRig/Private/Units/Core/RigUnit_String.cpp b/Engine/Plugins/Animation/ControlRig/Source/ControlRig/Private/Units/Core/RigUnit_String.cpp
--- a/Engine/Plugins/Animation/ControlRig/Source/ControlRig/Private/Units/Core/RigUnit_String.cpp
+++ b/Engine/Plugins/Animation/ControlRig/Source/ControlRig/Private/Units/Core/RigUnit_String.cpp
@@ -123,7 +123,7 @@ FRigUnit_StringSplit_Execute()
 
 		FString ValueRemaining = Value;
 		FString Left, Right;
-		while(ValueRemaining.Split(TEXT("."), &Left, &Right, ESearchCase::CaseSensitive, ESearchDir::FromStart))
+		while(ValueRemaining.Split(Separator, &Left, &Right, ESearchCase::CaseSensitive, ESearchDir::FromStart))
 		{
 			Result.Add(Left);
 			Left.Empty();
