Control Rig: Fix controls not displayed in the viewport when BP Actor with CR component is selected

diff --git a/Engine/Plugins/Animation/ControlRig/Source/ControlRigEditor/Private/Sequencer/ControlRigParameterTrackEditor.cpp b/Engine/Plugins/Animation/ControlRig/Source/ControlRigEditor/Private/Sequencer/ControlRigParameterTrackEditor.cpp
--- a/Engine/Plugins/Animation/ControlRig/Source/ControlRigEditor/Private/Sequencer/ControlRigParameterTrackEditor.cpp
+++ b/Engine/Plugins/Animation/ControlRig/Source/ControlRigEditor/Private/Sequencer/ControlRigParameterTrackEditor.cpp
@@ -2064,6 +2064,7 @@ void FControlRigParameterTrackEditor::OnSelectionChanged(TArray<UMovieSceneTrack
 					if (ControlRigEditMode)
 					{
 						ControlRigEditMode->AddControlRigObject(TrackControlRig, GetSequencer());
+						ControlRigEditMode->RequestToRecreateControlShapeActors(TrackControlRig);
 						break;
 					}
 					else
@@ -2074,6 +2075,7 @@ void FControlRigParameterTrackEditor::OnSelectionChanged(TArray<UMovieSceneTrack
 							if (ControlRigEditMode)
 							{
 								ControlRigEditMode->AddControlRigObject(TrackControlRig, GetSequencer());
+								ControlRigEditMode->RequestToRecreateControlShapeActors(TrackControlRig);
 							}
 						}
 					}
