improve aim constraint setup and update
    - we set the axis at initialization so that the child doesn't "pop" when activating the constraint.
    - as this constraint is fed by the incoming rotation, compensating doesn't make sense (so no need to add keys at T-1)

diff --git a/Engine/Plugins/Animation/ControlRig/Source/ControlRigEditor/Private/Sequencer/ConstraintChannelHelper.cpp b/Engine/Plugins/Animation/ControlRig/Source/ControlRigEditor/Private/Sequencer/ConstraintChannelHelper.cpp
--- a/Engine/Plugins/Animation/ControlRig/Source/ControlRigEditor/Private/Sequencer/ConstraintChannelHelper.cpp
+++ b/Engine/Plugins/Animation/ControlRig/Source/ControlRigEditor/Private/Sequencer/ConstraintChannelHelper.cpp
@@ -198,7 +198,8 @@ void FConstraintChannelHelper::SmartControlConstraintKey(
 			bool ActiveValueToBeSet = false;
 			if (CanAddKey(Channel->ActiveChannel, Time, ActiveValueToBeSet))
 			{
-
+				const bool bNeedsCompensation = InConstraint->NeedsCompensation();
+				
 				TGuardValue<bool> CompensateGuard(FMovieSceneConstraintChannelHelper::bDoNotCompensate, true);
 				
 				// set constraint as dynamic
@@ -210,7 +211,14 @@ void FConstraintChannelHelper::SmartControlConstraintKey(
 				// store the frames to compensate
 				const TArrayView<FMovieSceneFloatChannel*> Channels = ControlHandle->GetFloatChannels(Section);
 				TArray<FFrameNumber> FramesToCompensate;
-				FMovieSceneConstraintChannelHelper::GetFramesToCompensate<FMovieSceneFloatChannel>(Channel->ActiveChannel, ActiveValueToBeSet, Time, Channels, FramesToCompensate);
+				if (bNeedsCompensation)
+				{
+					FMovieSceneConstraintChannelHelper::GetFramesToCompensate<FMovieSceneFloatChannel>(Channel->ActiveChannel, ActiveValueToBeSet, Time, Channels, FramesToCompensate);
+				}
+				else
+				{
+					FramesToCompensate.Add(Time);
+				}
 
 				// store child and space transforms for these frames
 				FCompensationEvaluator Evaluator(InConstraint);
@@ -229,7 +237,7 @@ void FConstraintChannelHelper::SmartControlConstraintKey(
 				{
 					// get the number of float channels to treat
 					NumChannels = FControlRigSpaceChannelHelpers::GetNumFloatChannels(ControlElement->Settings.ControlType);
-					if (NumChannels > 0)
+					if (bNeedsCompensation && NumChannels > 0)
 					{
 						ChannelIndex = pChannelIndex->ChannelIndex;
 						EvaluateTangentAtThisTime<FMovieSceneFloatChannel>(ChannelIndex, NumChannels, Section, Time, Tangents);
@@ -239,6 +247,7 @@ void FConstraintChannelHelper::SmartControlConstraintKey(
 				const EMovieSceneTransformChannel ChannelsToKey =InConstraint->GetChannelsToKey();
 				
 				// add child's transform key at Time-1 to keep animation
+				if (bNeedsCompensation)
 				{
 					const FFrameNumber TimeMinusOne(Time - 1);
 
@@ -268,7 +277,7 @@ void FConstraintChannelHelper::SmartControlConstraintKey(
 						ChildLocals, ChannelsToKey, TickResolution, nullptr,true);
 
 					// set tangents at Time
-					if (NumChannels > 0 )
+					if (bNeedsCompensation && NumChannels > 0)
 					{
 						SetTangentsAtThisTime<FMovieSceneFloatChannel>(ChannelIndex, NumChannels, Section, Time, Tangents);
 					}
@@ -317,7 +326,8 @@ void FConstraintChannelHelper::SmartComponentConstraintKey(
 			bool ActiveValueToBeSet = false;
 			if (CanAddKey(Channel->ActiveChannel, Time, ActiveValueToBeSet))
 			{
-	
+				const bool bNeedsCompensation = InConstraint->NeedsCompensation();
+				
 				Section->Modify();
 
 				//new for compensation
@@ -330,7 +340,14 @@ void FConstraintChannelHelper::SmartComponentConstraintKey(
 				// store the frames to compensate
 				const TArrayView<FMovieSceneDoubleChannel*> Channels = ComponentHandle->GetDoubleChannels(Section);
 				TArray<FFrameNumber> FramesToCompensate;
-				FMovieSceneConstraintChannelHelper::GetFramesToCompensate<FMovieSceneDoubleChannel>(Channel->ActiveChannel, ActiveValueToBeSet, Time, Channels, FramesToCompensate);
+				if (bNeedsCompensation)
+				{
+					FMovieSceneConstraintChannelHelper::GetFramesToCompensate<FMovieSceneDoubleChannel>(Channel->ActiveChannel, ActiveValueToBeSet, Time, Channels, FramesToCompensate);
+				}
+				else
+				{
+					FramesToCompensate.Add(Time);
+				}
 
 				// store child and space transforms for these frames
 				FCompensationEvaluator Evaluator(InConstraint);
@@ -341,11 +358,16 @@ void FConstraintChannelHelper::SmartComponentConstraintKey(
 				TArray<FMovieSceneTangentData> Tangents;
 				const int32 ChannelIndex = 0;
 				const int32 NumChannels = 9;
-				EvaluateTangentAtThisTime<FMovieSceneDoubleChannel>(ChannelIndex, NumChannels, Section, Time, Tangents);
+
+				if (bNeedsCompensation)
+				{
+					EvaluateTangentAtThisTime<FMovieSceneDoubleChannel>(ChannelIndex, NumChannels, Section, Time, Tangents);
+				}
 
 				const EMovieSceneTransformChannel ChannelsToKey = InConstraint->GetChannelsToKey();
 
 				// add child's transform key at Time-1 to keep animation
+				if (bNeedsCompensation)
 				{
 					const FFrameNumber TimeMinusOne(Time - 1);
 
@@ -369,7 +391,10 @@ void FConstraintChannelHelper::SmartComponentConstraintKey(
 					MovieSceneToolHelpers::AddTransformKeys(Section, FramesToCompensate, ChildLocals, ChannelsToKey);
 
 					// set tangents at Time
-					SetTangentsAtThisTime<FMovieSceneDoubleChannel>(ChannelIndex,NumChannels, Section, Time, Tangents);
+					if (bNeedsCompensation)
+					{
+						SetTangentsAtThisTime<FMovieSceneDoubleChannel>(ChannelIndex,NumChannels, Section, Time, Tangents);
+					}
 				}
 				// evaluate the constraint, this is needed so the global transform will be set up on the component 
 				//Todo do we need to evaluate all constraints?
@@ -491,7 +516,13 @@ void FConstraintChannelHelper::CompensateIfNeeded(
 	Algo::CopyIf(ConstraintChannels, TransformConstraintsChannels,
 		[](const FConstraintAndActiveChannel& InChannel)
 		{
-			return InChannel.Constraint.IsValid() && InChannel.Constraint->IsA<UTickableTransformConstraint>();
+			if (!InChannel.Constraint.IsValid())
+			{
+				return false;
+			}
+
+			const UTickableTransformConstraint* Constraint = Cast<UTickableTransformConstraint>(InChannel.Constraint.Get());
+			return Constraint && Constraint->NeedsCompensation();
 		}
 	);
 
diff --git a/Engine/Source/Editor/MovieSceneTools/Private/MovieSceneConstraintChannelHelper.cpp b/Engine/Source/Editor/MovieSceneTools/Private/MovieSceneConstraintChannelHelper.cpp
--- a/Engine/Source/Editor/MovieSceneTools/Private/MovieSceneConstraintChannelHelper.cpp
+++ b/Engine/Source/Editor/MovieSceneTools/Private/MovieSceneConstraintChannelHelper.cpp
@@ -515,10 +515,11 @@ void FMovieSceneConstraintChannelHelper::HandleConstraintRemoved(
 	UMovieSceneSection* InSection)
 {
 	UTickableTransformConstraint* Constraint = Cast<UTickableTransformConstraint>(InConstraint);
-	if (!Constraint || !InConstraintChannel || !InSection)
+	if (!Constraint || !Constraint->NeedsCompensation() || !InConstraintChannel || !InSection)
 	{
 		return;
 	}
+	
 	InSection->Modify();
 	TGuardValue<bool> CompensateGuard(bDoNotCompensate, true);
 	if (const UTransformableHandle* ControlHandle =Constraint->ChildTRSHandle)
@@ -565,6 +566,11 @@ void FMovieSceneConstraintChannelHelper::HandleConstraintKeyDeleted(
 	UMovieSceneSection* InSection,
 	const FFrameNumber& InTime)
 {
+	if (!InConstraint || !InConstraint->NeedsCompensation())
+	{
+		return;
+	}
+	
 	const FFrameNumber TimeMinusOne(InTime - 1);
 
 	bool CurrentValue = false, PreviousValue = false;
diff --git a/Engine/Source/Runtime/Experimental/Animation/Constraints/Private/TransformConstraint.cpp b/Engine/Source/Runtime/Experimental/Animation/Constraints/Private/TransformConstraint.cpp
--- a/Engine/Source/Runtime/Experimental/Animation/Constraints/Private/TransformConstraint.cpp
+++ b/Engine/Source/Runtime/Experimental/Animation/Constraints/Private/TransformConstraint.cpp
@@ -335,6 +335,12 @@ bool UTickableTransformConstraint::IsFullyActive() const
 		&& IsValid(ParentTRSHandle) && ParentTRSHandle->IsValid());
 }
 
+bool UTickableTransformConstraint::NeedsCompensation() const
+{
+	// NOTE: this can be extended to something more complex if needed  
+	return true;
+}
+
 /** 
  * UTickableTranslationConstraint
  **/
@@ -913,11 +919,19 @@ void UTickableLookAtConstraint::ComputeOffset()
 {
 	bMaintainOffset = false;
 	bDynamicOffset = false;
+
+	const FVector InitParentLocation = GetParentGlobalTransform().GetLocation();
+	const FTransform InitChildTransform = GetChildGlobalTransform();
+	const FVector InitLookAtDir = (InitParentLocation - InitChildTransform.GetLocation()).GetSafeNormal();
+
+	if (!InitLookAtDir.IsNearlyZero())
+	{
+		Axis = InitChildTransform.InverseTransformVectorNoScale(InitLookAtDir).GetSafeNormal();
+	}
 }
 
 FConstraintTickFunction::ConstraintFunction UTickableLookAtConstraint::GetFunction() const
 {
-	// @todo handle weight here
 	return [this]()
 	{
 		if (!IsFullyActive())
@@ -935,15 +949,24 @@ FConstraintTickFunction::ConstraintFunction UTickableLookAtConstraint::GetFuncti
 			const FVector AxisToOrient = ChildTransform.TransformVectorNoScale(Axis).GetSafeNormal();
 		
 			FQuat Rotation = FindQuatBetweenNormals(AxisToOrient, LookAtDir);
-			Rotation = Rotation * ChildTransform.GetRotation();
+			const bool bNeedsToBeRotated = !Rotation.IsIdentity();
+			if (bNeedsToBeRotated)
+			{
+				Rotation = Rotation * ChildTransform.GetRotation();
 
-			FTransform Transform = ChildTransform;
-			Transform.SetRotation(Rotation.GetNormalized());
-			SetChildGlobalTransform(Transform);
+				FTransform Transform = ChildTransform;
+				Transform.SetRotation(Rotation.GetNormalized());
+				SetChildGlobalTransform(Transform);
+			}
 		}
 	};
 }
 
+bool UTickableLookAtConstraint::NeedsCompensation() const
+{
+	return false;
+}
+
 FQuat UTickableLookAtConstraint::FindQuatBetweenNormals(const FVector& A, const FVector& B)
 {
 	const FQuat::FReal Dot = FVector::DotProduct(A, B);
@@ -969,7 +992,23 @@ FQuat UTickableLookAtConstraint::FindQuatBetweenNormals(const FVector& A, const
 
 	Result.Normalize();
 	return Result;
-};
+}
+
+#if WITH_EDITOR
+
+void UTickableLookAtConstraint::PostEditChangeProperty(FPropertyChangedEvent& PropertyChangedEvent)
+{
+	Super::PostEditChangeProperty(PropertyChangedEvent);
+
+	const FName PropertyName = PropertyChangedEvent.GetMemberPropertyName();
+	if (PropertyName == GET_MEMBER_NAME_CHECKED(UTickableLookAtConstraint, Axis))
+	{
+		Evaluate();
+		return;
+	}
+}
+
+#endif
 
 /** 
  * FTransformConstraintUtils
diff --git a/Engine/Source/Runtime/Experimental/Animation/Constraints/Public/TransformConstraint.h b/Engine/Source/Runtime/Experimental/Animation/Constraints/Public/TransformConstraint.h
--- a/Engine/Source/Runtime/Experimental/Animation/Constraints/Public/TransformConstraint.h
+++ b/Engine/Source/Runtime/Experimental/Animation/Constraints/Public/TransformConstraint.h
@@ -49,6 +49,8 @@ public:
 	/** If Active and the handles and targets are valid*/
 	virtual bool IsFullyActive() const override;
 
+	/** If that constraint needs to be handled by the compensation system. */
+	virtual bool NeedsCompensation() const;
 
 	/** Override the evaluate so we can tick our handles*/
 	virtual void Evaluate(bool bTickHandlesAlso = false) const override;
@@ -338,8 +340,13 @@ public:
 	/** Returns the look at constraint function that the tick function will evaluate. */
 	virtual FConstraintTickFunction::ConstraintFunction GetFunction() const override;
 
+	/** If that constraint needs to be handled by the compensation system. */
+	virtual bool NeedsCompensation() const override;
+
 protected:
-	/** @todo document */
+	/**
+     * Computes the initial axis that is needed to keep the child's orientation unchanged when creating the constraint.
+    */
 	virtual void ComputeOffset() override;
 
 	/** Defines the aiming axis. */
@@ -350,6 +357,13 @@ private:
 
 	/** Computes the shortest quaternion between A and B. */
 	static FQuat FindQuatBetweenNormals(const FVector& A, const FVector& B);
+
+#if WITH_EDITOR
+public:
+	// UObject interface
+	virtual void PostEditChangeProperty(FPropertyChangedEvent& PropertyChangedEvent) override;
+	// End of UObject interface
+#endif
 };
 
 /** 
