ConstraintManager: Add new icons for lookat and parent

diff --git a/Engine/Plugins/Animation/ControlRig/Source/ControlRigEditor/Private/Editor/Constraints/SConstraintsWidget.cpp b/Engine/Plugins/Animation/ControlRig/Source/ControlRigEditor/Private/Editor/Constraints/SConstraintsWidget.cpp
--- a/Engine/Plugins/Animation/ControlRig/Source/ControlRigEditor/Private/Editor/Constraints/SConstraintsWidget.cpp
+++ b/Engine/Plugins/Animation/ControlRig/Source/ControlRigEditor/Private/Editor/Constraints/SConstraintsWidget.cpp
@@ -44,8 +44,8 @@ const TArray< const FSlateBrush* >& FConstraintInfo::GetBrushes()
 		FAppStyle::Get().GetBrush("EditorViewport.TranslateMode"),
 		FAppStyle::Get().GetBrush("EditorViewport.RotateMode"),
 		FAppStyle::Get().GetBrush("EditorViewport.ScaleMode"),
-		FAppStyle::Get().GetBrush("Icons.Transform"),
-		FAppStyle::Get().GetBrush("Level.VisibleHighlightIcon16x")
+		FAppStyle::Get().GetBrush("Icons.ConstraintManager.ParentHierarchy"),
+		FAppStyle::Get().GetBrush("Icons.ConstraintManager.LookAt")
 		});
 	return Brushes;
 }
diff --git a/Engine/Source/Runtime/SlateCore/Private/Styling/StarshipCoreStyle.cpp b/Engine/Source/Runtime/SlateCore/Private/Styling/StarshipCoreStyle.cpp
--- a/Engine/Source/Runtime/SlateCore/Private/Styling/StarshipCoreStyle.cpp
+++ b/Engine/Source/Runtime/SlateCore/Private/Styling/StarshipCoreStyle.cpp
@@ -411,6 +411,10 @@ TSharedRef<ISlateStyle> FStarshipCoreStyle::Create()
 		
 		// Toolbar Size Icons
 		Style->Set("Icons.Toolbar.Settings", new IMAGE_BRUSH_SVG("Starship/Common/settings", Icon20x20));
+
+		// Constraint Manager Icons
+		Style->Set("Icons.ConstraintManager.LookAt", new IMAGE_BRUSH_SVG("Starship/Common/LookAt", Icon16x16));
+		Style->Set("Icons.ConstraintManager.ParentHierarchy", new IMAGE_BRUSH_SVG("Starship/Common/ParentHierarchy", Icon16x16));
 	}
 
 	// Tool panels
