ControlRigParameterTrack: handle failture for MakeNewSpawnable

diff --git a/Engine/Plugins/Animation/ControlRig/Source/ControlRigEditor/Private/Sequencer/ControlRigParameterTrackEditor.cpp b/Engine/Plugins/Animation/ControlRig/Source/ControlRigEditor/Private/Sequencer/ControlRigParameterTrackEditor.cpp
--- a/Engine/Plugins/Animation/ControlRig/Source/ControlRigEditor/Private/Sequencer/ControlRigParameterTrackEditor.cpp
+++ b/Engine/Plugins/Animation/ControlRig/Source/ControlRigEditor/Private/Sequencer/ControlRigParameterTrackEditor.cpp
@@ -3735,25 +3735,24 @@ bool FControlRigParameterTrackEditor::HandleAssetAdded(UObject* Asset, const FGu
 
 	const FScopedTransaction Transaction(LOCTEXT("AddControlRigAsset", "Add Control Rig"));
 
-	// Spawn a skeletal mesh actor with the preview mesh
-	ASkeletalMeshActor* SkeletalMeshActor = GCurrentLevelEditingViewportClient->GetWorld()->SpawnActor<ASkeletalMeshActor>();
-	if (!SkeletalMeshActor)
+	FGuid NewGuid = GetSequencer()->MakeNewSpawnable(*ASkeletalMeshActor::StaticClass());
+
+	// MakeNewSpawnable can fail if spawnables are not allowed
+	if (!NewGuid.IsValid())
 	{
 		return false;
 	}
-
-	SkeletalMeshActor->GetSkeletalMeshComponent()->SetSkeletalMesh(SkeletalMesh);
-
-	FGuid NewGuid = GetSequencer()->MakeNewSpawnable(*SkeletalMeshActor);
-	UObject* SpawnedSkeletalMeshActor = GetSequencer()->FindSpawnedObjectOrTemplate(NewGuid);
-	if (SkeletalMeshActor)
+	
+	ASkeletalMeshActor* SpawnedSkeletalMeshActor = Cast<ASkeletalMeshActor>(GetSequencer()->FindSpawnedObjectOrTemplate(NewGuid));
+	if (!ensure(SpawnedSkeletalMeshActor))
 	{
-		GCurrentLevelEditingViewportClient->GetWorld()->EditorDestroyActor(SkeletalMeshActor, true);
-		SkeletalMeshActor = Cast<ASkeletalMeshActor>(SpawnedSkeletalMeshActor);
+		return false;
 	}
 
+	SpawnedSkeletalMeshActor->GetSkeletalMeshComponent()->SetSkeletalMesh(SkeletalMesh);
+
 	FString NewName = MovieSceneHelpers::MakeUniqueSpawnableName(MovieScene, FName::NameToDisplayString(SkeletalMesh->GetName(), false));
-	SkeletalMeshActor->SetActorLabel(NewName, false);
+	SpawnedSkeletalMeshActor->SetActorLabel(NewName, false);
 
 	UMovieSceneControlRigParameterTrack* Track = Cast<UMovieSceneControlRigParameterTrack>(MovieScene->FindTrack(UMovieSceneControlRigParameterTrack::StaticClass(), NewGuid, NAME_None));
 	if (Track == nullptr)
@@ -3761,7 +3760,7 @@ bool FControlRigParameterTrackEditor::HandleAssetAdded(UObject* Asset, const FGu
 		UControlRig* CDO = Cast<UControlRig>(RigClass->GetDefaultObject(true /* create if needed */));
 		check(CDO);
 
-		AddControlRig(CDO->GetClass(), SkeletalMeshActor->GetSkeletalMeshComponent(), NewGuid);
+		AddControlRig(CDO->GetClass(), SpawnedSkeletalMeshActor->GetSkeletalMeshComponent(), NewGuid);
 	}
 
 	return true;
diff --git a/Engine/Source/Editor/Sequencer/Private/SequencerUtilities.cpp b/Engine/Source/Editor/Sequencer/Private/SequencerUtilities.cpp
--- a/Engine/Source/Editor/Sequencer/Private/SequencerUtilities.cpp
+++ b/Engine/Source/Editor/Sequencer/Private/SequencerUtilities.cpp
@@ -639,11 +639,17 @@ FGuid FSequencerUtilities::MakeNewSpawnable(TSharedRef<ISequencer> Sequencer, UO
 		return FGuid();
 	}
 
-	if (MovieScene->IsReadOnly() || !Sequence->AllowsSpawnableObjects())
+	if (MovieScene->IsReadOnly())
 	{
 		ShowReadOnlyError();
 		return FGuid();
 	}
+	
+	if (!Sequence->AllowsSpawnableObjects())
+	{
+		ShowSpawnableNotAllowedError();
+		return FGuid();
+	}
 
 	FGuid NewGuid = AddSpawnable(Sequencer, Object, ActorFactory, SpawnableName);
 	if (!NewGuid.IsValid())
@@ -2848,4 +2854,11 @@ void FSequencerUtilities::ShowReadOnlyError()
 	FSlateNotificationManager::Get().AddNotification(Info)->SetCompletionState(SNotificationItem::CS_Fail);
 }
 
+void FSequencerUtilities::ShowSpawnableNotAllowedError()
+{
+	FNotificationInfo Info(LOCTEXT("SequenceSpawnableNotAllowed", "Spawnable object is not allowed for Sequence."));
+	Info.ExpireDuration = 5.0f;
+	FSlateNotificationManager::Get().AddNotification(Info)->SetCompletionState(SNotificationItem::CS_Fail);	
+}
+
 #undef LOCTEXT_NAMESPACE
diff --git a/Engine/Source/Editor/Sequencer/Public/SequencerUtilities.h b/Engine/Source/Editor/Sequencer/Public/SequencerUtilities.h
--- a/Engine/Source/Editor/Sequencer/Public/SequencerUtilities.h
+++ b/Engine/Source/Editor/Sequencer/Public/SequencerUtilities.h
@@ -181,6 +181,9 @@ struct SEQUENCER_API FSequencerUtilities
 	static void ReplaceBindingWithActors(TSharedRef<ISequencer> Sequencer, const TArray<AActor*>& Actors, const FMovieSceneBindingProxy& ObjectBinding);
 	static void RemoveActorsFromBinding(TSharedRef<ISequencer> Sequencer, const TArray<AActor*>& Actors, const FMovieSceneBindingProxy& ObjectBinding);
 
-	/** Show a read only error of the movie scene is locked */
+	/** Show a read only error if the movie scene is locked */
 	static void ShowReadOnlyError();
+	
+	/** Show an error if spawnable is not allowed in a movie scene*/
+	static void ShowSpawnableNotAllowedError();
 };
