Constraints baking
    - fixed baking to ensure that we compensate the transform within the last active constraint
    - minimize the baked keys to the constraint's active ranges

Also fixed constraint key deletion failing to compensate

diff --git a/Engine/Plugins/Animation/ControlRig/Source/ControlRigEditor/Private/Editor/Constraints/SConstraintsWidget.cpp b/Engine/Plugins/Animation/ControlRig/Source/ControlRigEditor/Private/Editor/Constraints/SConstraintsWidget.cpp
--- a/Engine/Plugins/Animation/ControlRig/Source/ControlRigEditor/Private/Editor/Constraints/SConstraintsWidget.cpp
+++ b/Engine/Plugins/Animation/ControlRig/Source/ControlRigEditor/Private/Editor/Constraints/SConstraintsWidget.cpp
@@ -970,19 +970,9 @@ TSharedPtr<SWidget> SConstraintsEditionWidget::CreateContextMenu()
 						{
 							return;
 						}
-
-						// compute frames
-						const FFrameNumber StartFrame = MovieScene->GetPlaybackRange().GetLowerBoundValue();
-						const FFrameNumber EndFrame = MovieScene->GetPlaybackRange().GetUpperBoundValue();
-						TArray<FFrameNumber> Frames;
-						MovieSceneToolHelpers::CalculateFramesBetween(MovieScene, StartFrame, EndFrame, Frames);
-
-						if (Frames.IsEmpty())
-						{
-							return;
-						}
+						
 						FScopedTransaction Transaction(LOCTEXT("BakeConstraint", "Bake Constraint"));
-						FConstraintBaker::Bake(World, TransformConstraint, WeakSequencer.Pin(), Frames);
+						FConstraintBaker::Bake(World, TransformConstraint, WeakSequencer.Pin(), TOptional<TArray<FFrameNumber>>());
 
 					}
 				})),
diff --git a/Engine/Plugins/Animation/ControlRig/Source/ControlRigEditor/Private/Tools/ConstraintBaker.cpp b/Engine/Plugins/Animation/ControlRig/Source/ControlRigEditor/Private/Tools/ConstraintBaker.cpp
--- a/Engine/Plugins/Animation/ControlRig/Source/ControlRigEditor/Private/Tools/ConstraintBaker.cpp
+++ b/Engine/Plugins/Animation/ControlRig/Source/ControlRigEditor/Private/Tools/ConstraintBaker.cpp
@@ -114,26 +114,153 @@ void FConstraintBaker::AddTransformKeys(
 	}
 }
 
-
-void FConstraintBaker::Bake(UWorld* InWorld, 
-	UTickableTransformConstraint* InConstraint,
-	const TSharedPtr<ISequencer>& InSequencer, 
-	const TArray<FFrameNumber>& InFrames)
+// at this stage, we suppose that any of these arguments are safe to use. Make sure to test them before using that function
+void FConstraintBaker::GetMinimalFramesToBake(
+	UWorld* InWorld,
+	const UTickableTransformConstraint* InConstraint,
+	const TSharedPtr<ISequencer>& InSequencer,
+	IMovieSceneConstrainedSection* InSection,
+	TArray<FFrameNumber>& OutFramesToBake)
 {
-	if (InFrames.Num() == 0 || InConstraint == nullptr || InConstraint->ChildTRSHandle == nullptr)
+	const UMovieSceneSequence* MovieSceneSequence = InSequencer->GetFocusedMovieSceneSequence();
+	const UMovieScene* MovieScene = MovieSceneSequence ? MovieSceneSequence->GetMovieScene() : nullptr;
+	if (!MovieScene)
 	{
 		return;
 	}
+
+	// get constraint channel data if any
+	TArrayView<const FFrameNumber> ConstraintFrames;
+	TArrayView<const bool> ConstraintValues;
+
+	// note that we might want to bake a constraint which is not animated
+	FConstraintAndActiveChannel* ThisActiveChannel = InSection->GetConstraintChannel(InConstraint->GetFName());
+	if (ThisActiveChannel)
+	{
+		const TMovieSceneChannelData<const bool> ConstraintChannelData = ThisActiveChannel->ActiveChannel.GetData();
+		ConstraintFrames = ConstraintChannelData.GetTimes();
+		ConstraintValues = ConstraintChannelData.GetValues();
+	}
+
+	// there's no channel or an empty one and the constraint is inactive so no need to do anything if the constraint
+	// is not active
+	if (ConstraintFrames.IsEmpty() && !InConstraint->IsFullyActive())
+	{
+		return;
+	}
+
+	// default init bounds to the scene bounds
+	FFrameNumber FirstBakingFrame = MovieScene->GetPlaybackRange().GetLowerBoundValue();
+	FFrameNumber LastBakingFrame = MovieScene->GetPlaybackRange().GetUpperBoundValue();
 	
-	// compute transforms
-	FCompensationEvaluator Evaluator(InConstraint);
-	Evaluator.ComputeLocalTransformsForBaking(InWorld, InSequencer, InFrames);
-	TArray<FTransform>& Transforms = Evaluator.ChildLocals;
+	// set start to first active frame if any
+	const int32 FirstActiveIndex = ConstraintValues.IndexOfByKey(true);
+	if (ConstraintValues.IsValidIndex(FirstActiveIndex))
+	{
+		FirstBakingFrame = ConstraintFrames[FirstActiveIndex];
+	}
+
+	// set end to the last key if inactive
+	const bool bIsLastKeyInactive = (ConstraintValues.Last() == false);
+	if (bIsLastKeyInactive)
+	{
+		LastBakingFrame = ConstraintFrames.Last();
+	}
+
+	// then compute range from first to last
+	TArray<FFrameNumber> Frames;
+	MovieSceneToolHelpers::CalculateFramesBetween(MovieScene, FirstBakingFrame, LastBakingFrame, Frames);
+		
+	// Fill the frames we want to bake
+	{
+		// add constraint keys
+		for (const FFrameNumber& ConstraintFrame: ConstraintFrames)
+		{
+			OutFramesToBake.Add(ConstraintFrame);
+		}
+		
+		// keep frames within active state
+		auto IsConstraintActive = [InConstraint, ThisActiveChannel](const FFrameNumber& Time)
+		{
+			if (ThisActiveChannel)
+			{
+				bool bIsActive = false; ThisActiveChannel->ActiveChannel.Evaluate(Time, bIsActive);
+				return bIsActive;
+			}
+			return InConstraint->IsFullyActive();
+		};
+		
+		for (const FFrameNumber& InFrame: Frames)
+		{
+			if (IsConstraintActive(InFrame))
+			{
+				OutFramesToBake.Add(InFrame);
+			}
+		}
+	}
+
+	// we also need to store which T-1 frames need to be kept for other constraints compensation
+	{
+		// gather the sorted child's constraint
+		static constexpr bool bSorted = true;
+		const uint32 ChildHash = InConstraint->ChildTRSHandle->GetHash();
+		const FConstraintsManagerController& Controller = FConstraintsManagerController::Get(InWorld);
+		using ConstraintPtr = TObjectPtr<UTickableConstraint>;
+		const TArray<ConstraintPtr> Constraints = Controller.GetParentConstraints(ChildHash, bSorted);
+
+		// store the other channels that may need compensation
+		TArray<FConstraintAndActiveChannel*> OtherChannels;
+		for (const ConstraintPtr& Constraint: Constraints)
+		{
+			const UTickableTransformConstraint* TransformConstraint = Cast<UTickableTransformConstraint>(Constraint.Get());
+			if (TransformConstraint && TransformConstraint->NeedsCompensation() && TransformConstraint != InConstraint) 
+			{
+				const FName ConstraintName = TransformConstraint->GetFName();
+				if (FConstraintAndActiveChannel* ConstraintChannel = InSection->GetConstraintChannel(ConstraintName))
+				{
+					OtherChannels.Add(ConstraintChannel);
+				}
+			}
+		}
 
-	if (InFrames.Num() != Transforms.Num())
+		// check if any other channel needs to compensate at T
+		for (const FFrameNumber& ConstraintFrame: ConstraintFrames)
+		{
+			const bool bNeedsCompensation = OtherChannels.ContainsByPredicate(
+		[ConstraintFrame](const FConstraintAndActiveChannel* OtherChannel)
+			{
+				const bool bHasAtLeastOneActiveKey = OtherChannel->ActiveChannel.GetValues().Contains(true);
+				if (!bHasAtLeastOneActiveKey)
+				{
+					return false;
+				}
+				return OtherChannel->ActiveChannel.GetTimes().Contains(ConstraintFrame);
+			} );
+
+			// add T-1 if needed so that we will compensate int the remaining constraints' space 
+			if (bNeedsCompensation)
+			{
+				const FFrameNumber& FrameMinusOne(ConstraintFrame-1);
+				OutFramesToBake.Add(FrameMinusOne);
+			}
+		}
+	}
+	
+	// uniqueness
+	OutFramesToBake.Sort();
+	OutFramesToBake.SetNum(Algo::Unique(OutFramesToBake));
+}
+
+void FConstraintBaker::Bake(UWorld* InWorld, 
+	UTickableTransformConstraint* InConstraint,
+	const TSharedPtr<ISequencer>& InSequencer, 
+	const TOptional<TArray<FFrameNumber>>& InFrames)
+{
+	if ((InFrames && InFrames->IsEmpty()) || (InConstraint == nullptr) || (InConstraint->ChildTRSHandle == nullptr))
 	{
 		return;
 	}
+	
 	//get the section to be used later to delete the extra transform keys at the frame -1 times, abort if not there for some reason
 	UMovieSceneSection* Section = nullptr;
 	if (const UTransformableControlHandle* ControlHandle = Cast<UTransformableControlHandle>(InConstraint->ChildTRSHandle))
@@ -172,12 +299,31 @@ void FConstraintBaker::Bake(UWorld* InWorld,
 		return;
 	}
 
+	// compute transforms
+	TArray<FFrameNumber> FramesToBake;
+	if (InFrames)
+	{
+		FramesToBake = *InFrames;
+	}
+	else
+	{
+		GetMinimalFramesToBake(InWorld, InConstraint, InSequencer, ConstrainedSection, FramesToBake);
+	}
+	
+	FCompensationEvaluator Evaluator(InConstraint);
+	Evaluator.ComputeLocalTransformsForBaking(InWorld, InSequencer, FramesToBake);
+	const TArray<FTransform>& Transforms = Evaluator.ChildLocals;
+	if (FramesToBake.Num() != Transforms.Num())
+	{
+		return;
+	}
+	
 	Section->Modify();
 
 	// disable constraint and delete extra transform keys
 	TMovieSceneChannelData<bool> ConstraintChannelData = ActiveChannel->ActiveChannel.GetData();
-	TArrayView<const FFrameNumber> ConstraintFrames = ConstraintChannelData.GetTimes();
-
+	const TArrayView<const FFrameNumber> ConstraintFrames = ConstraintChannelData.GetTimes();
+	
 	// get transform channels
 	const TArrayView<FMovieSceneFloatChannel*> FloatTransformChannels = InConstraint->ChildTRSHandle->GetFloatChannels(Section);
 	const TArrayView<FMovieSceneDoubleChannel*> DoubleTransformChannels = InConstraint->ChildTRSHandle->GetDoubleChannels(Section);
@@ -186,12 +332,13 @@ void FConstraintBaker::Bake(UWorld* InWorld,
 	{
 		const FFrameNumber Frame = ConstraintFrames[Index];
 		//todo we need to add a key at the begin/end if there is no frame there
-		if (Frame >= InFrames[0] && Frame <= InFrames[InFrames.Num() - 1])
+		if (Frame >= FramesToBake[0] && Frame <= FramesToBake.Last())
 		{
+			// set constraint key to inactive
 			ConstraintChannelData.UpdateOrAddKey(Frame, false);
-			//delete minus one frames
-			const FFrameNumber FrameMinusOne = Frame - 1;
 
+			// delete minus one transform frames
+			const FFrameNumber FrameMinusOne = Frame - 1;
 			if (FloatTransformChannels.Num() > 0)
 			{
 				FMovieSceneConstraintChannelHelper::DeleteTransformKeys(FloatTransformChannels, FrameMinusOne);
@@ -204,7 +351,7 @@ void FConstraintBaker::Bake(UWorld* InWorld,
 	}
 	// now bake to channel curves
 	const EMovieSceneTransformChannel Channels = InConstraint->GetChannelsToKey();
-	AddTransformKeys(InSequencer, InConstraint->ChildTRSHandle, InFrames, Transforms, Channels);
+	AddTransformKeys(InSequencer, InConstraint->ChildTRSHandle, FramesToBake, Transforms, Channels);
 
 	// notify
 	InSequencer->RequestEvaluate();
diff --git a/Engine/Plugins/Animation/ControlRig/Source/ControlRigEditor/Public/Tools/ConstraintBaker.h b/Engine/Plugins/Animation/ControlRig/Source/ControlRigEditor/Public/Tools/ConstraintBaker.h
--- a/Engine/Plugins/Animation/ControlRig/Source/ControlRigEditor/Public/Tools/ConstraintBaker.h
+++ b/Engine/Plugins/Animation/ControlRig/Source/ControlRigEditor/Public/Tools/ConstraintBaker.h
@@ -5,6 +5,7 @@
 #include "Misc/FrameNumber.h"
 #include "Templates/SharedPointer.h"
 
+class IMovieSceneConstrainedSection;
 class UTransformableHandle;
 class UTickableTransformConstraint;
 class ISequencer;
@@ -16,8 +17,11 @@ struct FConstraintBaker
 {
 public:
 	/** Bake constraint over specified frames, frames must be in order*/
-	static void Bake(UWorld* InWorld, UTickableTransformConstraint* InConstraint, const TSharedPtr<ISequencer>& InSequencer, const TArray<FFrameNumber>& InFrames);
-
+	static void Bake(
+		UWorld* InWorld,
+		UTickableTransformConstraint* InConstraint,
+		const TSharedPtr<ISequencer>& InSequencer,
+		const TOptional<TArray<FFrameNumber>>& InFrames);
 	
 	/** Add InTransforms keys at InFrames into the InHandle transform animation channels. */
 	static void AddTransformKeys(
@@ -29,5 +33,12 @@ public:
 	
 private:
 
+	/** Computes the minimal frames to bake for that constraint. */
+	static void GetMinimalFramesToBake(
+		UWorld* InWorld,
+		const UTickableTransformConstraint* InConstraint,
+		const TSharedPtr<ISequencer>& InSequencer,
+		IMovieSceneConstrainedSection* InSection,
+		TArray<FFrameNumber>& OutFramesToBake);
 };
 
diff --git a/Engine/Source/Editor/MovieSceneTools/Private/MovieSceneConstraintChannelHelper.cpp b/Engine/Source/Editor/MovieSceneTools/Private/MovieSceneConstraintChannelHelper.cpp
--- a/Engine/Source/Editor/MovieSceneTools/Private/MovieSceneConstraintChannelHelper.cpp
+++ b/Engine/Source/Editor/MovieSceneTools/Private/MovieSceneConstraintChannelHelper.cpp
@@ -187,6 +187,18 @@ void FCompensationEvaluator::ComputeLocalTransformsForBaking(UWorld* InWorld, co
 
 	const TArray<UTickableTransformConstraint*> Constraints = GetHandleTransformConstraints(InWorld);
 
+	auto GetLastActiveConstraint = [this, Constraints]()
+	{
+		// find last active constraint in the list that is different than the on we want to compensate for
+		const int32 LastActiveIndex = Constraints.FindLastByPredicate([this](const UTickableTransformConstraint* InConstraint)
+		{
+			return InConstraint->Active && InConstraint->NeedsCompensation() && InConstraint->bDynamicOffset && InConstraint != Constraint;
+		});
+
+		// if found, return its parent global transform
+		return LastActiveIndex > INDEX_NONE ? Constraints[LastActiveIndex] : nullptr;
+	};
+
 	UMovieScene* MovieScene = InSequencer->GetFocusedMovieSceneSequence()->GetMovieScene();
 	const FFrameRate TickResolution = MovieScene->GetTickResolution();
 	const EMovieScenePlayerStatus::Type PlaybackStatus = InSequencer->GetPlaybackStatus();
@@ -251,6 +263,14 @@ void FCompensationEvaluator::ComputeLocalTransformsForBaking(UWorld* InWorld, co
 		// store child transforms        	
 		ChildLocal = Handle->GetLocalTransform();
 		ChildGlobal = Handle->GetGlobalTransform();
+
+		// store constraint/parent space global transform
+		if (const UTickableTransformConstraint* LastConstraint = GetLastActiveConstraint())
+		{
+			SpaceGlobal = LastConstraint->GetParentGlobalTransform();
+			ChildLocal = FTransformConstraintUtils::ComputeRelativeTransform(
+				ChildLocal, ChildGlobal, SpaceGlobal, LastConstraint);
+		}
 	}
 	for (IMovieSceneToolsAnimationBakeHelper* BakeHelper : BakeHelpers)
 	{
@@ -598,11 +618,11 @@ void FMovieSceneConstraintChannelHelper::HandleConstraintKeyDeleted(
 		TArray<FFrameNumber> FramesToCompensate;
 		if (FloatTransformChannels.Num() > 0)
 		{
-			FMovieSceneConstraintChannelHelper::GetFramesWithinActiveState(*InConstraintChannel, FloatTransformChannels, FramesToCompensate);
+			GetFramesAfter(*InConstraintChannel, InTime, FloatTransformChannels, FramesToCompensate);
 		}
 		else
 		{
-			FMovieSceneConstraintChannelHelper::GetFramesWithinActiveState(*InConstraintChannel, DoubleTransformChannels, FramesToCompensate);
+			GetFramesAfter(*InConstraintChannel,  InTime, DoubleTransformChannels, FramesToCompensate);
 		}
 		// do the compensation
 		FCompensationEvaluator Evaluator(InConstraint);
