Sequencer: Control Rig: Make sure to check all control rig's when showing available spaces.
#trivial

diff --git a/Engine/Plugins/Animation/ControlRig/Source/ControlRigEditor/Private/EditMode/SControlRigEditModeTools.cpp b/Engine/Plugins/Animation/ControlRig/Source/ControlRigEditor/Private/EditMode/SControlRigEditModeTools.cpp
--- a/Engine/Plugins/Animation/ControlRig/Source/ControlRigEditor/Private/EditMode/SControlRigEditModeTools.cpp
+++ b/Engine/Plugins/Animation/ControlRig/Source/ControlRigEditor/Private/EditMode/SControlRigEditModeTools.cpp
@@ -766,11 +766,13 @@ void SControlRigEditModeTools::OnRigElementSelected(UControlRig* Subject, FRigCo
 
 const FRigControlElementCustomization* SControlRigEditModeTools::HandleGetControlElementCustomization(URigHierarchy* InHierarchy, const FRigElementKey& InControlKey)
 {
-	//mz todo handle multiple rigs
 	UControlRig* Rig = ControlRigs.Num() > 0 ? ControlRigs[0].Get() : nullptr;
-	if (Rig)
+	for (TWeakObjectPtr<UControlRig>& ControlRig : ControlRigs)
 	{
-		return Rig->GetControlCustomization(InControlKey);
+		if (ControlRig.IsValid() && ControlRig->GetHierarchy() == InHierarchy)
+		{
+			return ControlRig->GetControlCustomization(InControlKey);
+		}
 	}
 	return nullptr;
 }
