Constraints: Add BP function to bake constraints
#trivial

diff --git a/Engine/Plugins/Animation/ControlRig/Source/ControlRigEditor/Private/ControlRigSequencerEditorLibrary.cpp b/Engine/Plugins/Animation/ControlRig/Source/ControlRigEditor/Private/ControlRigSequencerEditorLibrary.cpp
--- a/Engine/Plugins/Animation/ControlRig/Source/ControlRigEditor/Private/ControlRigSequencerEditorLibrary.cpp
+++ b/Engine/Plugins/Animation/ControlRig/Source/ControlRigEditor/Private/ControlRigSequencerEditorLibrary.cpp
@@ -37,6 +37,8 @@
 #include "Sequencer/ControlRigParameterTrackEditor.h"
 #include "ControlRigSpaceChannelEditors.h"
 #include "LevelSequenceEditorBlueprintLibrary.h"
+#include "Tools/ConstraintBaker.h"
+#include "TransformConstraint.h"
 
 #include UE_INLINE_GENERATED_CPP_BY_NAME(ControlRigSequencerEditorLibrary)
 
@@ -400,6 +402,55 @@ bool UControlRigSequencerEditorLibrary::TweenControlRig(ULevelSequence* LevelSeq
 	return false;
 }
 
+bool UControlRigSequencerEditorLibrary::BakeConstraint(UWorld* World, UTickableConstraint* Constraint, const TArray<FFrameNumber>& Frames, ESequenceTimeUnit TimeUnit)
+{
+	if (!World)
+	{
+		UE_LOG(LogControlRig, Error, TEXT("BakeConstraint: Need Valid World"));
+		return false;
+	}
+	if (UTickableTransformConstraint* TransformConstraint = Cast<UTickableTransformConstraint>(Constraint))
+	{
+		TSharedPtr<ISequencer> Sequencer = GetSequencerFromAsset();
+		if (!Sequencer || !Sequencer->GetFocusedMovieSceneSequence())
+		{
+			UE_LOG(LogControlRig, Error, TEXT("BakeConstraint: Need loaded level Sequence"));
+			return false;
+		}
+		const UMovieScene* MovieScene = Sequencer->GetFocusedMovieSceneSequence()->GetMovieScene();
+		if (!MovieScene)
+		{
+			UE_LOG(LogControlRig, Error, TEXT("BakeConstraint: Need valid Movie Scene"));
+			return false;
+		}
+		
+		TOptional<TArray<FFrameNumber>> FramesToBake;
+		TArray<FFrameNumber> RealFramesToBake;
+		RealFramesToBake.SetNum(Frames.Num());
+
+		int32 Index = 0;
+		for (FFrameNumber Frame : Frames) 
+		{
+			if (TimeUnit == ESequenceTimeUnit::DisplayRate)
+			{
+				Frame = FFrameRate::TransformTime(FFrameTime(Frame, 0), MovieScene->GetDisplayRate(), MovieScene->GetTickResolution()).RoundToFrame();
+			}
+			RealFramesToBake[Index++] = Frame;
+		} 
+		if (RealFramesToBake.Num() > 0)
+		{
+			FramesToBake = RealFramesToBake;
+		}
+		FConstraintBaker::Bake(World, TransformConstraint, Sequencer, FramesToBake);
+	}
+	else
+	{
+		UE_LOG(LogControlRig, Error, TEXT("BakeConstraint: Need Valid Constraint"));
+		return false;
+	}
+	return true;
+}
+
 
 bool UControlRigSequencerEditorLibrary::SnapControlRig(ULevelSequence* LevelSequence, FFrameNumber StartFrame, FFrameNumber EndFrame, const FControlRigSnapperSelection& ChildrenToSnap,
 	const FControlRigSnapperSelection& ParentToSnap, const UControlRigSnapSettings* SnapSettings, ESequenceTimeUnit TimeUnit)
diff --git a/Engine/Plugins/Animation/ControlRig/Source/ControlRigEditor/Public/ControlRigSequencerEditorLibrary.h b/Engine/Plugins/Animation/ControlRig/Source/ControlRigEditor/Public/ControlRigSequencerEditorLibrary.h
--- a/Engine/Plugins/Animation/ControlRig/Source/ControlRigEditor/Public/ControlRigSequencerEditorLibrary.h
+++ b/Engine/Plugins/Animation/ControlRig/Source/ControlRigEditor/Public/ControlRigSequencerEditorLibrary.h
@@ -20,6 +20,7 @@
 #include "ControlRigSequencerEditorLibrary.generated.h"
 
 class ULevelSequence;
+class UTickableConstraint;
 
 USTRUCT(BlueprintType)
 struct FControlRigSequencerBindingProxy
@@ -127,6 +128,17 @@ public:
 	static bool BakeToControlRig(UWorld* World, ULevelSequence* LevelSequence, UClass* ControlRigClass, UAnimSeqExportOption* ExportOptions, bool bReduceKeys, float Tolerance,
 			const FMovieSceneBindingProxy& Binding);
 
+	/**
+	* Bake the constraint to keys based on the passed in frames. This will use the open sequencer to bake. See ConstraintsScriptingLibrary to get the list of available constraints
+	* @param World The active world
+	* @param Constraint The Constraint to bake. After baking it will be keyed to be inactive of the range of frames that are baked
+	* @param Frames The frames to bake, if the array is empty it will use the active time ranges of the constraint to deteremine where it should bake
+	* @param TimeUnit Unit for all frame and time values, either in display rate or tick resolution
+	* @return returns True if successful, False otherwise
+	*/
+	UFUNCTION(BlueprintCallable, Category = "Editor Scripting | Sequencer Tools | Control Rig")
+	static bool BakeConstraint(UWorld* World, UTickableConstraint* Constraint, const TArray<FFrameNumber>& Frames, ESequenceTimeUnit TimeUnit = ESequenceTimeUnit::DisplayRate);
+
 	/**
 	* Peform a Tween operation on the current active sequencer time(must be visible).
 	* @param LevelSequence The LevelSequence that's loaded in the editor
