Control Rig: Fix control shapes in wrong position after undo

diff --git a/Engine/Plugins/Animation/ControlRig/Source/ControlRig/Private/Rigs/RigHierarchy.cpp b/Engine/Plugins/Animation/ControlRig/Source/ControlRig/Private/Rigs/RigHierarchy.cpp
--- a/Engine/Plugins/Animation/ControlRig/Source/ControlRig/Private/Rigs/RigHierarchy.cpp
+++ b/Engine/Plugins/Animation/ControlRig/Source/ControlRig/Private/Rigs/RigHierarchy.cpp
@@ -1767,9 +1767,11 @@ bool URigHierarchy::SetParentWeightArray(FRigBaseElement* InChild,  const TArray
 			if(FRigControlElement* ControlElement = Cast<FRigControlElement>(MultiParentElement))
 			{
 				ControlElement->Offset.MarkDirty(GlobalType);
+				ControlElement->Shape.MarkDirty(GlobalType);
 			}
 
 			PropagateDirtyFlags(MultiParentElement, ERigTransformType::IsInitial(LocalType), bAffectChildren);
+			EnsureCacheValidity();
 			
 #if WITH_EDITOR
 			if (!bPropagatingChange)
@@ -1904,6 +1906,19 @@ bool URigHierarchy::SwitchToParent(FRigBaseElement* InChild, FRigBaseElement* In
 		}
 	}
 
+	// Exit early if switching to the same parent 
+	if (InChild)
+	{
+		const FRigElementKey ChildKey = InChild->GetKey();
+		const FRigElementKey ParentKey = InParent ? InParent->GetKey() : GetDefaultParent(ChildKey);
+		const FRigElementKey ActiveParentKey = GetActiveParent(ChildKey);
+		if(ActiveParentKey == ParentKey ||
+			(ActiveParentKey == URigHierarchy::GetDefaultParentKey() && GetDefaultParent(ChildKey) == ParentKey))
+		{
+			return true;
+		}
+	}
+	
 	if(const FRigMultiParentElement* MultiParentElement = Cast<FRigMultiParentElement>(InChild))
 	{
 		int32 ParentIndex = INDEX_NONE;
diff --git a/Engine/Plugins/Animation/ControlRig/Source/ControlRigEditor/Private/EditMode/ControlRigEditMode.cpp b/Engine/Plugins/Animation/ControlRig/Source/ControlRigEditor/Private/EditMode/ControlRigEditMode.cpp
--- a/Engine/Plugins/Animation/ControlRig/Source/ControlRigEditor/Private/EditMode/ControlRigEditMode.cpp
+++ b/Engine/Plugins/Animation/ControlRig/Source/ControlRigEditor/Private/EditMode/ControlRigEditMode.cpp
@@ -567,7 +567,14 @@ void FControlRigEditMode::Tick(FEditorViewportClient* ViewportClient, float Delt
 	// We need to tick here since changing a bone for example
 	// might have changed the transform of the Control
 	{
-		PostPoseUpdate();
+		if (IsInLevelEditor())
+		{
+			TickManipulatableObjects(0.f);
+		}
+		else
+		{
+			PostPoseUpdate();
+		}
 		if (IsInLevelEditor() == false) //only do this check if not in level editor
 		{
 			for (TWeakObjectPtr<UControlRig>& RuntimeRigPtr : RuntimeControlRigs)
