Control Rig: Fix pasting rig hierarchy elements from another rig does not respect parent-child relationship

diff --git a/Engine/Plugins/Animation/ControlRig/Source/ControlRig/Private/Rigs/RigHierarchyController.cpp b/Engine/Plugins/Animation/ControlRig/Source/ControlRig/Private/Rigs/RigHierarchyController.cpp
--- a/Engine/Plugins/Animation/ControlRig/Source/ControlRig/Private/Rigs/RigHierarchyController.cpp
+++ b/Engine/Plugins/Animation/ControlRig/Source/ControlRig/Private/Rigs/RigHierarchyController.cpp
@@ -2127,8 +2127,11 @@ bool URigHierarchyController::AddParent(FRigBaseElement* InChild, FRigBaseElemen
 		{
 			if(!bRemoveAllParents)
 			{
-				ReportErrorf(TEXT("Cannot add multiple parents to animation channel '%s'."), *InChild->Key.ToString());
-				return false;
+				if (ChildControlElement->ParentConstraints.Num() > 0)
+				{
+					ReportErrorf(TEXT("Cannot add multiple parents to animation channel '%s'."), *InChild->Key.ToString());
+					return false;
+				}
 			}
 
 			bMaintainGlobalTransform = false;
