ensure that the anim instance is linked as InitAnim can reset the linked instances

+ fix to prevent FLinkedInstancesAdapter::AddLinkedInstance from adding the same instance several times

diff --git a/Engine/Plugins/Animation/ControlRig/Source/ControlRigEditor/Private/Editor/ControlRigSkeletalMeshComponent.cpp b/Engine/Plugins/Animation/ControlRig/Source/ControlRigEditor/Private/Editor/ControlRigSkeletalMeshComponent.cpp
--- a/Engine/Plugins/Animation/ControlRig/Source/ControlRigEditor/Private/Editor/ControlRigSkeletalMeshComponent.cpp
+++ b/Engine/Plugins/Animation/ControlRig/Source/ControlRigEditor/Private/Editor/ControlRigSkeletalMeshComponent.cpp
@@ -33,11 +33,16 @@ void UControlRigSkeletalMeshComponent::InitAnim(bool bForceReinit)
 	Super::InitAnim(bForceReinit);
 
 	// The preview instance or anim instance might have been created in Super::InitAnim, in which case
-	// the source animation instance must be set now
-	if (AnimInstance != GetAnimInstance() || LastPreviewInstance != PreviewInstance)
+	// the source animation instance must be set now.
+	// we also must ensure that the anim instance is linked as InitAnim can reset the linked instances
+	ControlRigInstance = Cast<UControlRigLayerInstance>(GetAnimInstance());
+	if (ControlRigInstance)
 	{
-		ControlRigInstance = Cast<UControlRigLayerInstance>(GetAnimInstance());
-		if (ControlRigInstance)
+		const bool bHaveInstancesChanged = (AnimInstance != GetAnimInstance()) || (LastPreviewInstance != PreviewInstance);
+
+		const USkeletalMeshComponent* MeshComponent = ControlRigInstance->GetOwningComponent();
+		const bool bIsAnimInstanceLinked = PreviewInstance && MeshComponent && MeshComponent->GetLinkedAnimInstances().Contains(PreviewInstance);
+		if (bHaveInstancesChanged || !bIsAnimInstanceLinked)
 		{
 			ControlRigInstance->SetSourceAnimInstance(PreviewInstance);
 		}
diff --git a/Engine/Source/Runtime/Engine/Classes/Components/SkeletalMeshComponent.h b/Engine/Source/Runtime/Engine/Classes/Components/SkeletalMeshComponent.h
--- a/Engine/Source/Runtime/Engine/Classes/Components/SkeletalMeshComponent.h
+++ b/Engine/Source/Runtime/Engine/Classes/Components/SkeletalMeshComponent.h
@@ -2640,7 +2640,7 @@ struct FLinkedInstancesAdapter
 	{
 		if (InComponent && InAnimInstance)
 		{
-			InComponent->LinkedInstances.Add(InAnimInstance);
+			InComponent->LinkedInstances.AddUnique(InAnimInstance);
 		}
 	}
 
