removed usage of TObjectPtr.IsNull and IsNullNoResolve

#p4v-preflight-copy 22193787

diff --git a/Engine/Plugins/Animation/ControlRig/Source/ControlRig/Private/ControlRig.cpp b/Engine/Plugins/Animation/ControlRig/Source/ControlRig/Private/ControlRig.cpp
--- a/Engine/Plugins/Animation/ControlRig/Source/ControlRig/Private/ControlRig.cpp
+++ b/Engine/Plugins/Animation/ControlRig/Source/ControlRig/Private/ControlRig.cpp
@@ -1511,7 +1511,7 @@ bool UControlRig::ExecuteUnits(FRigUnitContext& InOutContext, const FName& InEve
 				if(SnapshotHash != 0)
 				{
 					TObjectPtr<URigVM>* InitializedVMSnapshotPtr = CDO->InitializedVMSnapshots.Find(SnapshotHash);
-					if(InitializedVMSnapshotPtr && !InitializedVMSnapshotPtr->IsNull())
+					if(InitializedVMSnapshotPtr && *InitializedVMSnapshotPtr)
 					{
 						const URigVM* InitializedVMSnapshot = InitializedVMSnapshotPtr->Get();
 
diff --git a/Engine/Plugins/Animation/ControlRig/Source/ControlRigEditor/Private/ControlRigElementDetails.cpp b/Engine/Plugins/Animation/ControlRig/Source/ControlRigEditor/Private/ControlRigElementDetails.cpp
--- a/Engine/Plugins/Animation/ControlRig/Source/ControlRigEditor/Private/ControlRigElementDetails.cpp
+++ b/Engine/Plugins/Animation/ControlRig/Source/ControlRigEditor/Private/ControlRigElementDetails.cpp
@@ -2366,7 +2366,7 @@ void FRigControlElementDetails::CustomizeValue(IDetailLayoutBuilder& DetailBuild
 	// transform types already got filtered out earlier.
 	// integers with enums don't have limits either
 	if(ControlType != ERigControlType::Bool &&
-		(ControlType != ERigControlType::Integer || FirstControlElement->Settings.ControlEnum.IsNull()))
+		(ControlType != ERigControlType::Integer || !FirstControlElement->Settings.ControlEnum))
 	{
 		Labels.Append({
 			LOCTEXT("Min", "Min"),
@@ -2438,7 +2438,7 @@ void FRigControlElementDetails::CustomizeValue(IDetailLayoutBuilder& DetailBuild
 				{
 					if(const FRigControlElement* ControlElement = Hierarchy->Find<FRigControlElement>(Key))
 					{
-						if(!ControlElement->Settings.ControlEnum.IsNull())
+						if(ControlElement->Settings.ControlEnum)
 						{
 							bIsEnum = true;
 							break;
diff --git a/Engine/Plugins/Animation/ControlRig/Source/ControlRigEditor/Private/ControlRigGraphDetails.h b/Engine/Plugins/Animation/ControlRig/Source/ControlRigEditor/Private/ControlRigGraphDetails.h
--- a/Engine/Plugins/Animation/ControlRig/Source/ControlRigEditor/Private/ControlRigGraphDetails.h
+++ b/Engine/Plugins/Animation/ControlRig/Source/ControlRigEditor/Private/ControlRigGraphDetails.h
@@ -273,7 +273,7 @@ protected:
 		bOutEnabled = false;
 		if (!ObjectsBeingCustomized.IsEmpty())
 		{
-			if (!ObjectsBeingCustomized[0].IsNull())
+			if (ObjectsBeingCustomized[0])
 			{
 				TSharedPtr<class IPropertyHandle> ChainHandle = InPropertyHandle;
 				while (ChainHandle.IsValid() && ChainHandle->GetProperty() != nullptr)
@@ -358,9 +358,9 @@ protected:
 			return Result;
 		}
 	
-		for(TObjectPtr<UObject> Object : ObjectsBeingCustomized)
+		for(const TObjectPtr<UObject>& Object : ObjectsBeingCustomized)
 		{
-			if(!Object.IsNull() && InPropertyHandle->IsValidHandle())
+			if(Object && InPropertyHandle->IsValidHandle())
 			{
 				static VectorType ZeroVector = VectorType();
 				const VectorType& Vector = ContainerUObjectToValueRef<VectorType>(Object.Get(), ZeroVector, PropertyChain, PropertyArrayIndices);
@@ -396,8 +396,8 @@ protected:
 		TArray<UObject*> ObjectsView;
 		for(int32 Index = 0; Index < ObjectsBeingCustomized.Num(); Index++)
 		{
-			TObjectPtr<UObject> Object = ObjectsBeingCustomized[Index];
-			if (!Object.IsNull())
+			const TObjectPtr<UObject>& Object = ObjectsBeingCustomized[Index];
+			if (Object)
 			{
 				ObjectsView.Add(Object.Get());
 			}
@@ -417,8 +417,8 @@ protected:
 
 		for(int32 Index = 0; Index < ObjectsBeingCustomized.Num(); Index++)
 		{
-			TObjectPtr<UObject> Object = ObjectsBeingCustomized[Index];
-			if(!Object.IsNull() && InPropertyHandle->IsValidHandle())
+			const TObjectPtr<UObject>& Object = ObjectsBeingCustomized[Index];
+			if(Object && InPropertyHandle->IsValidHandle())
 			{
 				static VectorType ZeroVector = VectorType();
 				VectorType& Vector = ContainerUObjectToValueRef<VectorType>(Object.Get(), ZeroVector, PropertyChain, PropertyArrayIndices);
@@ -518,9 +518,9 @@ protected:
 			return Result;
 		}
 		
-		for(TObjectPtr<UObject> Object : ObjectsBeingCustomized)
+		for(const TObjectPtr<UObject>& Object : ObjectsBeingCustomized)
 		{
-			if(!Object.IsNull() && InPropertyHandle->IsValidHandle())
+			if(Object && InPropertyHandle->IsValidHandle())
 			{
 				static RotationType ZeroRotation = RotationType();
 				const RotationType& Rotation = ContainerUObjectToValueRef<RotationType>(Object.Get(), ZeroRotation, PropertyChain, PropertyArrayIndices);
@@ -555,8 +555,8 @@ protected:
 		TArray<UObject*> ObjectsView;
 		for(int32 Index = 0; Index < ObjectsBeingCustomized.Num(); Index++)
 		{
-			TObjectPtr<UObject> Object = ObjectsBeingCustomized[Index];
-			if (!Object.IsNull())
+			const TObjectPtr<UObject>& Object = ObjectsBeingCustomized[Index];
+			if (Object)
 			{
 				ObjectsView.Add(Object.Get());
 			}
@@ -576,8 +576,8 @@ protected:
 
 		for(int32 Index = 0; Index < ObjectsBeingCustomized.Num(); Index++)
 		{
-			TObjectPtr<UObject> Object = ObjectsBeingCustomized[Index];
-			if(!Object.IsNull() && InPropertyHandle->IsValidHandle())
+			const TObjectPtr<UObject>& Object = ObjectsBeingCustomized[Index];
+			if(Object && InPropertyHandle->IsValidHandle())
 			{
 				static RotationType ZeroRotation = RotationType();
 				RotationType& Rotation = ContainerUObjectToValueRef<RotationType>(Object.Get(), ZeroRotation, PropertyChain, PropertyArrayIndices);
@@ -671,9 +671,9 @@ protected:
 				return false;
 			}
 			
-			for(TObjectPtr<UObject> Object : ObjectsBeingCustomized)
+			for(const TObjectPtr<UObject>& Object : ObjectsBeingCustomized)
 			{
-				if(!Object.IsNull() && InPropertyHandle->IsValidHandle())
+				if(Object && InPropertyHandle->IsValidHandle())
 				{
 					const TransformType& Transform = ContainerUObjectToValueRef<TransformType>(Object.Get(), Identity, PropertyChain, PropertyArrayIndices);
 
@@ -727,9 +727,9 @@ protected:
 				return Result;
 			}
 			
-			for(TObjectPtr<UObject> Object : ObjectsBeingCustomized)
+			for(const TObjectPtr<UObject>& Object : ObjectsBeingCustomized)
 			{
-				if(!Object.IsNull() && InPropertyHandle->IsValidHandle())
+				if(Object && InPropertyHandle->IsValidHandle())
 				{
 					const TransformType& Transform = ContainerUObjectToValueRef<TransformType>(Object.Get(), Identity, PropertyChain, PropertyArrayIndices);
 					
@@ -779,8 +779,8 @@ protected:
 			TArray<UObject*> ObjectsView;
 			for(int32 Index = 0; Index < ObjectsBeingCustomized.Num(); Index++)
 			{
-				TObjectPtr<UObject> Object = ObjectsBeingCustomized[Index];
-				if (!Object.IsNull())
+				const TObjectPtr<UObject>& Object = ObjectsBeingCustomized[Index];
+				if (Object)
 				{
 					ObjectsView.Add(Object.Get());
 				}
@@ -800,8 +800,8 @@ protected:
 
 			for(int32 Index = 0; Index < ObjectsBeingCustomized.Num(); Index++)
 			{
-				TObjectPtr<UObject> Object = ObjectsBeingCustomized[Index];
-				if(!Object.IsNull() && InPropertyHandle->IsValidHandle())
+				const TObjectPtr<UObject>& Object = ObjectsBeingCustomized[Index];
+				if(Object && InPropertyHandle->IsValidHandle())
 				{
 					TransformType& Transform = ContainerUObjectToValueRef<TransformType>(Object.Get(), Identity, PropertyChain, PropertyArrayIndices);
 					TransformType PreviousTransform = Transform;
@@ -869,8 +869,8 @@ protected:
 			TArray<UObject*> ObjectsView;
 			for(int32 Index = 0; Index < ObjectsBeingCustomized.Num(); Index++)
 			{
-				TObjectPtr<UObject> Object = ObjectsBeingCustomized[Index];
-				if (!Object.IsNull())
+				const TObjectPtr<UObject>& Object = ObjectsBeingCustomized[Index];
+				if (Object)
 				{
 					ObjectsView.Add(Object.Get());
 				}
@@ -880,8 +880,8 @@ protected:
 			
 			for(int32 Index = 0; Index < ObjectsBeingCustomized.Num(); Index++)
 			{
-				TObjectPtr<UObject> Object = ObjectsBeingCustomized[Index];
-				if(!Object.IsNull() && InPropertyHandle->IsValidHandle())
+				const TObjectPtr<UObject>& Object = ObjectsBeingCustomized[Index];
+				if(Object && InPropertyHandle->IsValidHandle())
 				{
 					TransformType& Transform = ContainerUObjectToValueRef<TransformType>(Object.Get(), Identity, PropertyChain, PropertyArrayIndices);
 					TransformType PreviousTransform = Transform;
diff --git a/Engine/Plugins/Animation/IKRig/Source/IKRig/Private/AnimNodes/AnimNode_IKRig.cpp b/Engine/Plugins/Animation/IKRig/Source/IKRig/Private/AnimNodes/AnimNode_IKRig.cpp
--- a/Engine/Plugins/Animation/IKRig/Source/IKRig/Private/AnimNodes/AnimNode_IKRig.cpp
+++ b/Engine/Plugins/Animation/IKRig/Source/IKRig/Private/AnimNodes/AnimNode_IKRig.cpp
@@ -176,7 +176,7 @@ void FAnimNode_IKRig::Initialize_AnyThread(const FAnimationInitializeContext& Co
 	GetEvaluateGraphExposedInputs().Execute(Context);
 
 	// ensure there is always a processor available
-	if (IKRigProcessor.IsNull() && IsInGameThread())
+	if (!IKRigProcessor && IsInGameThread())
 	{
 		IKRigProcessor = NewObject<UIKRigProcessor>(Context.AnimInstanceProxy->GetSkelMeshComponent());	
 	}
diff --git a/Engine/Plugins/Animation/IKRig/Source/IKRig/Private/AnimNodes/AnimNode_RetargetPoseFromMesh.cpp b/Engine/Plugins/Animation/IKRig/Source/IKRig/Private/AnimNodes/AnimNode_RetargetPoseFromMesh.cpp
--- a/Engine/Plugins/Animation/IKRig/Source/IKRig/Private/AnimNodes/AnimNode_RetargetPoseFromMesh.cpp
+++ b/Engine/Plugins/Animation/IKRig/Source/IKRig/Private/AnimNodes/AnimNode_RetargetPoseFromMesh.cpp
@@ -14,7 +14,7 @@ void FAnimNode_RetargetPoseFromMesh::Initialize_AnyThread(const FAnimationInitia
 	// Initial update of the node, so we dont have a frame-delay on setup
 	GetEvaluateGraphExposedInputs().Execute(Context);
 
-	if (Processor.IsNull() && IsInGameThread())
+	if (!Processor && IsInGameThread())
 	{
 		Processor = NewObject<UIKRetargetProcessor>(Context.AnimInstanceProxy->GetSkelMeshComponent());	
 	}
@@ -151,12 +151,12 @@ void FAnimNode_RetargetPoseFromMesh::PreUpdate(const UAnimInstance* InAnimInstan
 {
 	DECLARE_SCOPE_HIERARCHICAL_COUNTER_ANIMNODE(PreUpdate)
 	
-	if (IKRetargeterAsset.IsNull())
+	if (!IKRetargeterAsset)
 	{
 		return;
 	}
 	
-	if (Processor.IsNull())
+	if (!Processor)
 	{
 		Processor = NewObject<UIKRetargetProcessor>(InAnimInstance->GetOwningComponent());	
 	}
@@ -250,14 +250,14 @@ bool FAnimNode_RetargetPoseFromMesh::EnsureProcessorIsInitialized(const TObjectP
 	// check that both a source and target mesh exist
 	const TObjectPtr<USkeletalMesh> SourceMesh = SourceMeshComponent->GetSkeletalMeshAsset();
 	const TObjectPtr<USkeletalMesh> TargetMesh = TargetMeshComponent->GetSkeletalMeshAsset();
-	if (SourceMesh.IsNull() || TargetMesh.IsNull())
+	if (!SourceMesh || !TargetMesh)
 	{
 		return false; // cannot initialize if components are missing skeletal mesh references
 	}
 	// check that both have skeleton assets (shouldn't get this far without a skeleton)
 	const TObjectPtr<USkeleton> SourceSkeleton = SourceMesh->GetSkeleton();
 	const TObjectPtr<USkeleton> TargetSkeleton = TargetMesh->GetSkeleton();
-	if (SourceSkeleton.IsNull() || TargetSkeleton.IsNull())
+	if (!SourceSkeleton || !TargetSkeleton)
 	{
 		return false;
 	}
@@ -301,7 +301,7 @@ void FAnimNode_RetargetPoseFromMesh::CopyBoneTransformsFromSource(USkeletalMeshC
 
 	// this should not happen as we're guaranteed to be initialized at this stage
 	// but just in case component is lost after initialization, we avoid a crash
-	if (ComponentToCopyFrom.IsNull())
+	if (!ComponentToCopyFrom)
 	{
 		return; 
 	}
diff --git a/Engine/Plugins/EnhancedInput/Source/EnhancedInput/Private/EnhancedInputPlatformSettings.cpp b/Engine/Plugins/EnhancedInput/Source/EnhancedInput/Private/EnhancedInputPlatformSettings.cpp
--- a/Engine/Plugins/EnhancedInput/Source/EnhancedInput/Private/EnhancedInputPlatformSettings.cpp
+++ b/Engine/Plugins/EnhancedInput/Source/EnhancedInput/Private/EnhancedInputPlatformSettings.cpp
@@ -17,7 +17,7 @@ EDataValidationResult UEnhancedInputPlatformData::IsDataValid(TArray<FText>& Val
 
 	for (const TPair<TObjectPtr<const UInputMappingContext>, TObjectPtr<const UInputMappingContext>> Pair : MappingContextRedirects)
 	{
-		if (Pair.Key.IsNull())
+		if (!Pair.Key)
 		{
 			Result = CombineDataValidationResults(Result, EDataValidationResult::Invalid);
 			
@@ -27,7 +27,7 @@ EDataValidationResult UEnhancedInputPlatformData::IsDataValid(TArray<FText>& Val
 			ValidationErrors.Emplace(NullKeyError);
 		}
 		
-		if (Pair.Value.IsNull())
+		if (!Pair.Value)
 		{
 			Result = CombineDataValidationResults(Result, EDataValidationResult::Invalid);
 			
diff --git a/Engine/Plugins/EnhancedInput/Source/EnhancedInput/Private/EnhancedInputSubsystemInterface.cpp b/Engine/Plugins/EnhancedInput/Source/EnhancedInput/Private/EnhancedInputSubsystemInterface.cpp
--- a/Engine/Plugins/EnhancedInput/Source/EnhancedInput/Private/EnhancedInputSubsystemInterface.cpp
+++ b/Engine/Plugins/EnhancedInput/Source/EnhancedInput/Private/EnhancedInputSubsystemInterface.cpp
@@ -604,7 +604,7 @@ void IEnhancedInputSubsystemInterface::RebuildControlMappings()
 		PlatformSettings->GetAllMappingContextRedirects(ContextRedirects);
 		for (const TPair<TObjectPtr<const UInputMappingContext>, TObjectPtr<const UInputMappingContext>> Pair : ContextRedirects)
 		{
-			if (Pair.Key.IsNull() || Pair.Value.IsNull())
+			if (!Pair.Key || !Pair.Value)
 			{
 				UE_LOG(LogEnhancedInput, Error, TEXT("An invalid Mappping Context Redirect specified in '%s'"), PlatformSettings->GetConfigOverridePlatform());
 				continue;
diff --git a/Engine/Plugins/Experimental/ChaosEditor/Source/FractureEditor/Private/FractureToolConvert.cpp b/Engine/Plugins/Experimental/ChaosEditor/Source/FractureEditor/Private/FractureToolConvert.cpp
--- a/Engine/Plugins/Experimental/ChaosEditor/Source/FractureEditor/Private/FractureToolConvert.cpp
+++ b/Engine/Plugins/Experimental/ChaosEditor/Source/FractureEditor/Private/FractureToolConvert.cpp
@@ -347,7 +347,7 @@ int32 UFractureToolConvert::ExecuteFracture(const FFractureToolContext& Fracture
 		Materials.Reserve(MainMaterials.Num());
 		for (int32 MatIdx = 0; MatIdx + SkipLastMaterialOffset < MainMaterials.Num(); MatIdx++)
 		{
-			if (MatIdx < OverrideMaterials.Num() && !OverrideMaterials[MatIdx].IsNull())
+			if (MatIdx < OverrideMaterials.Num() && OverrideMaterials[MatIdx])
 			{
 				Materials.Add(OverrideMaterials[MatIdx].Get());
 			}
diff --git a/Engine/Plugins/FX/Niagara/Source/NiagaraEditor/Private/Commandlets/NiagaraDumpModuleInfoCommandlet.cpp b/Engine/Plugins/FX/Niagara/Source/NiagaraEditor/Private/Commandlets/NiagaraDumpModuleInfoCommandlet.cpp
--- a/Engine/Plugins/FX/Niagara/Source/NiagaraEditor/Private/Commandlets/NiagaraDumpModuleInfoCommandlet.cpp
+++ b/Engine/Plugins/FX/Niagara/Source/NiagaraEditor/Private/Commandlets/NiagaraDumpModuleInfoCommandlet.cpp
@@ -141,9 +141,9 @@ void UNiagaraDumpModuleInfoCommandlet::ProcessNiagaraParameterDefinitions()
 		const TArray<UNiagaraScriptVariable*>& ScriptPtrVars =  NiagaraParameterDefinitions->GetParametersConst();
 		TArray< UNiagaraScriptVariable*> ScriptVars;
 
-		for (TObjectPtr<UNiagaraScriptVariable> Var : ScriptPtrVars)
+		for (UNiagaraScriptVariable* Var : ScriptPtrVars)
 		{
-			if (Var.IsNull())
+			if (!Var)
 				continue;
 			ScriptVars.Add(Var);
 		}
@@ -158,7 +158,7 @@ void UNiagaraDumpModuleInfoCommandlet::ProcessNiagaraParameterDefinitions()
 		OutputStream->Log(TEXT("<Info>"));
 		OutputStream->Logf(TEXT("\t<Path>%s</Path>"), *GetPathNameSafe(NiagaraParameterDefinitions));
 
-		for (UNiagaraScriptVariable*& ScriptVar : ScriptVars)
+		for (UNiagaraScriptVariable* ScriptVar : ScriptVars)
 		{
 			OutputStream->Log(TEXT("\t<ScriptVar>"));
 			OutputStream->Logf(TEXT("\t\t<Name>%s</Name>"), *ScriptVar->Variable.GetName().ToString());
@@ -264,13 +264,13 @@ void UNiagaraDumpModuleInfoCommandlet::ProcessNiagaraScripts()
 					TUniquePtr<FOutputDeviceArchiveWrapper> OutputStream(new FOutputDeviceArchiveWrapper(FileArchive.Get()));
 
 					const TMap<FNiagaraVariable, TObjectPtr<UNiagaraScriptVariable>>& AllVars = FunctionGraph->GetAllMetaData();
-					TArray< TObjectPtr<UNiagaraScriptVariable>> ScriptPtrVars;
+					TArray<TObjectPtr<UNiagaraScriptVariable>> ScriptPtrVars;
 					AllVars.GenerateValueArray(ScriptPtrVars);
-					TArray< UNiagaraScriptVariable*> ScriptVars;
+					TArray<UNiagaraScriptVariable*> ScriptVars;
 
-					for (TObjectPtr<UNiagaraScriptVariable> Var : ScriptPtrVars)
+					for (const TObjectPtr<UNiagaraScriptVariable>& Var : ScriptPtrVars)
 					{
-						if (Var.IsNull())
+						if (!Var)
 							continue;
 						ScriptVars.Add(Var);
 					}
@@ -288,7 +288,7 @@ void UNiagaraDumpModuleInfoCommandlet::ProcessNiagaraScripts()
 					OutputStream->Logf(TEXT("\t<Deprecated>%s</Deprecated>"), bDeprecated ? TEXT("true") : TEXT("false"));
 					OutputStream->Logf(TEXT("\t<Path>%s</Path>"), *GetPathNameSafe(FunctionGraph));
 
-					for (UNiagaraScriptVariable*& ScriptVar : ScriptVars)
+					for (UNiagaraScriptVariable* ScriptVar : ScriptVars)
 					{
 						OutputStream->Log(TEXT("\t<ScriptVar>"));
 						OutputStream->Logf(TEXT("\t\t<Name>%s</Name>"), *ScriptVar->Variable.GetName().ToString());
diff --git a/Engine/Plugins/FX/Niagara/Source/NiagaraEditor/Private/NiagaraEditorData.cpp b/Engine/Plugins/FX/Niagara/Source/NiagaraEditor/Private/NiagaraEditorData.cpp
--- a/Engine/Plugins/FX/Niagara/Source/NiagaraEditor/Private/NiagaraEditorData.cpp
+++ b/Engine/Plugins/FX/Niagara/Source/NiagaraEditor/Private/NiagaraEditorData.cpp
@@ -29,7 +29,7 @@ void UNiagaraEditorParametersAdapter::PostLoad()
 	{
 		for (TObjectPtr<UNiagaraScriptVariable> ScriptVariable : EditorOnlyScriptVars)
 		{
-			if (!ScriptVariable.IsNull() && ScriptVariable->Variable == OldVarType)
+			if (ScriptVariable && ScriptVariable->Variable == OldVarType)
 			{
 				ScriptVariable->Variable.SetType(FNiagaraTypeDefinition::GetPositionDef());
 			}
diff --git a/Engine/Plugins/FX/Niagara/Source/NiagaraEditor/Private/NiagaraGraph.cpp b/Engine/Plugins/FX/Niagara/Source/NiagaraEditor/Private/NiagaraGraph.cpp
--- a/Engine/Plugins/FX/Niagara/Source/NiagaraEditor/Private/NiagaraGraph.cpp
+++ b/Engine/Plugins/FX/Niagara/Source/NiagaraEditor/Private/NiagaraGraph.cpp
@@ -406,7 +406,7 @@ void UNiagaraGraph::PostLoad()
 	for (FNiagaraVariable OldVarType : OldTypes)
 	{
 		TObjectPtr<UNiagaraScriptVariable>* ScriptVariablePtr = VariableToScriptVariable.Find(OldVarType);
-		if (ScriptVariablePtr != nullptr && !ScriptVariablePtr->IsNull())
+		if (ScriptVariablePtr != nullptr && *ScriptVariablePtr)
 		{
 			UNiagaraScriptVariable& ScriptVar = *ScriptVariablePtr->Get();
 			ScriptVar.Variable.SetType(FNiagaraTypeDefinition::GetPositionDef());
@@ -707,7 +707,7 @@ void UNiagaraGraph::ChangeParameterType(const TArray<FNiagaraVariable>& Paramete
 	for(const FNiagaraVariable& CurrentParameter : ParametersToChange)
 	{		
 		TObjectPtr<UNiagaraScriptVariable>* ScriptVariablePtr = VariableToScriptVariable.Find(CurrentParameter);
-		if (ScriptVariablePtr != nullptr && !ScriptVariablePtr->IsNull())
+		if (ScriptVariablePtr != nullptr && *ScriptVariablePtr)
 		{
 			UNiagaraScriptVariable& ScriptVar = *ScriptVariablePtr->Get();
 			ScriptVar.Modify();
diff --git a/Engine/Plugins/FX/Niagara/Source/NiagaraEditor/Private/NiagaraScriptMergeManager.cpp b/Engine/Plugins/FX/Niagara/Source/NiagaraEditor/Private/NiagaraScriptMergeManager.cpp
--- a/Engine/Plugins/FX/Niagara/Source/NiagaraEditor/Private/NiagaraScriptMergeManager.cpp
+++ b/Engine/Plugins/FX/Niagara/Source/NiagaraEditor/Private/NiagaraScriptMergeManager.cpp
@@ -1836,11 +1836,11 @@ bool FNiagaraScriptMergeManager::FindBaseModule(const FVersionedNiagaraEmitter&
 					FVersionedNiagaraEmitterData* EmitterData = BaseEmitter.Emitter->GetEmitterData(BaseEmitter.Version);
 					if (EmitterData)
 					{
-						if (!EmitterData->ScratchPads.IsNull())
+						if (EmitterData->ScratchPads)
 						{
 							IndexToMatch = EmitterData->ScratchPads->FindIndexForScript(Script);
 						}
-						if (INDEX_NONE == IndexToMatch && !EmitterData->ParentScratchPads.IsNull())
+						if (INDEX_NONE == IndexToMatch && EmitterData->ParentScratchPads)
 						{
 							IndexToMatch = EmitterData->ParentScratchPads->FindIndexForScript(Script);
 							bParents = true;
diff --git a/Engine/Plugins/FX/Niagara/Source/NiagaraEditor/Private/ViewModels/Stack/NiagaraStackModuleItem.cpp b/Engine/Plugins/FX/Niagara/Source/NiagaraEditor/Private/ViewModels/Stack/NiagaraStackModuleItem.cpp
--- a/Engine/Plugins/FX/Niagara/Source/NiagaraEditor/Private/ViewModels/Stack/NiagaraStackModuleItem.cpp
+++ b/Engine/Plugins/FX/Niagara/Source/NiagaraEditor/Private/ViewModels/Stack/NiagaraStackModuleItem.cpp
@@ -249,7 +249,7 @@ const UNiagaraStackModuleItem::FCollectedUsageData& UNiagaraStackModuleItem::Get
 			}
 		}
 
-		if (!LinkedInputCollection.IsNull())
+		if (LinkedInputCollection)
 		{
 			if (LinkedInputCollection->GetCollectedUsageData().bHasReferencedParameterRead)
 				CachedCollectedUsageData.GetValue().bHasReferencedParameterRead = true;
@@ -258,7 +258,7 @@ const UNiagaraStackModuleItem::FCollectedUsageData& UNiagaraStackModuleItem::Get
 		}
 
 
-		if (!InputCollection.IsNull())
+		if (InputCollection)
 		{
 			if (InputCollection->GetCollectedUsageData().bHasReferencedParameterRead)
 				CachedCollectedUsageData.GetValue().bHasReferencedParameterRead = true;
@@ -266,7 +266,7 @@ const UNiagaraStackModuleItem::FCollectedUsageData& UNiagaraStackModuleItem::Get
 				CachedCollectedUsageData.GetValue().bHasReferencedParameterWrite = true;
 		}
 
-		if (!OutputCollection.IsNull())
+		if (OutputCollection)
 		{
 			if (OutputCollection->GetCollectedUsageData().bHasReferencedParameterRead)
 				CachedCollectedUsageData.GetValue().bHasReferencedParameterRead = true;
diff --git a/Engine/Plugins/FX/Niagara/Source/NiagaraEditorWidgets/Private/SNiagaraSystemEffectTypeBar.cpp b/Engine/Plugins/FX/Niagara/Source/NiagaraEditorWidgets/Private/SNiagaraSystemEffectTypeBar.cpp
--- a/Engine/Plugins/FX/Niagara/Source/NiagaraEditorWidgets/Private/SNiagaraSystemEffectTypeBar.cpp
+++ b/Engine/Plugins/FX/Niagara/Source/NiagaraEditorWidgets/Private/SNiagaraSystemEffectTypeBar.cpp
@@ -169,7 +169,7 @@ TSharedPtr<SWidget> SNiagaraSystemEffectTypeBar::CreatePropertyValueWidget(FProp
 		{
 			FObjectPtr& ObjectPtr = (FObjectPtr&) ObjectProperty->GetPropertyValue_InContainer(EffectType.Get());
 			
-			if(!ObjectPtr.IsNullNoResolve())
+			if(ObjectPtr)
 			{				
 				if(ObjectProperty->HasMetaData("DisplayClassDisplayName"))
 				{
diff --git a/Engine/Plugins/Interchange/Runtime/Source/Pipelines/Private/InterchangeGenericAssetsPipeline.cpp b/Engine/Plugins/Interchange/Runtime/Source/Pipelines/Private/InterchangeGenericAssetsPipeline.cpp
--- a/Engine/Plugins/Interchange/Runtime/Source/Pipelines/Private/InterchangeGenericAssetsPipeline.cpp
+++ b/Engine/Plugins/Interchange/Runtime/Source/Pipelines/Private/InterchangeGenericAssetsPipeline.cpp
@@ -46,7 +46,7 @@ UInterchangeGenericAssetsPipeline::UInterchangeGenericAssetsPipeline()
 
 void UInterchangeGenericAssetsPipeline::PreDialogCleanup(const FName PipelineStackName)
 {
-	check(!CommonSkeletalMeshesAndAnimationsProperties.IsNull())
+	check(CommonSkeletalMeshesAndAnimationsProperties)
 	//We always clean the pipeline skeleton when showing the dialog
 	CommonSkeletalMeshesAndAnimationsProperties->Skeleton = nullptr;
 
@@ -120,7 +120,7 @@ void UInterchangeGenericAssetsPipeline::AdjustSettingsForContext(EInterchangePip
 
 void UInterchangeGenericAssetsPipeline::ExecutePreImportPipeline(UInterchangeBaseNodeContainer* InBaseNodeContainer, const TArray<UInterchangeSourceData*>& InSourceDatas)
 {
-	check(!CommonSkeletalMeshesAndAnimationsProperties.IsNull());
+	check(CommonSkeletalMeshesAndAnimationsProperties);
 
 	if (!InBaseNodeContainer)
 	{
@@ -130,7 +130,7 @@ void UInterchangeGenericAssetsPipeline::ExecutePreImportPipeline(UInterchangeBas
 
 	//Set the result container to allow error message
 	//The parent Results container should be set at this point
-	ensure(!Results.IsNull());
+	ensure(Results);
 	{
 		if (MaterialPipeline)
 		{
diff --git a/Engine/Plugins/Interchange/Runtime/Source/Pipelines/Private/InterchangeGenericTexturePipeline.cpp b/Engine/Plugins/Interchange/Runtime/Source/Pipelines/Private/InterchangeGenericTexturePipeline.cpp
--- a/Engine/Plugins/Interchange/Runtime/Source/Pipelines/Private/InterchangeGenericTexturePipeline.cpp
+++ b/Engine/Plugins/Interchange/Runtime/Source/Pipelines/Private/InterchangeGenericTexturePipeline.cpp
@@ -111,7 +111,7 @@ void UInterchangeGenericTexturePipeline::AdjustSettingsForContext(EInterchangePi
 	Super::AdjustSettingsForContext(ImportType, ReimportAsset);
 
 	TArray<FString> HideCategories;
-	bool bIsObjectATexture = ReimportAsset.IsNull() ? false : ReimportAsset.IsA(UTexture::StaticClass());
+	bool bIsObjectATexture = !ReimportAsset ? false : ReimportAsset.IsA(UTexture::StaticClass());
 	if( (!bIsObjectATexture && ImportType == EInterchangePipelineContext::AssetReimport)
 		|| ImportType == EInterchangePipelineContext::AssetCustomLODImport
 		|| ImportType == EInterchangePipelineContext::AssetCustomLODReimport
diff --git a/Engine/Plugins/MeshPainting/Source/MeshPaintEditorMode/Private/MeshPaintModeHelpers.cpp b/Engine/Plugins/MeshPainting/Source/MeshPaintEditorMode/Private/MeshPaintModeHelpers.cpp
--- a/Engine/Plugins/MeshPainting/Source/MeshPaintEditorMode/Private/MeshPaintModeHelpers.cpp
+++ b/Engine/Plugins/MeshPainting/Source/MeshPaintEditorMode/Private/MeshPaintModeHelpers.cpp
@@ -333,7 +333,7 @@ void UMeshPaintModeSubsystem::ImportVertexColorsToSkeletalMesh(USkeletalMesh* Sk
 
 				if (GenericAssetPipeline)
 				{
-					if (!GenericAssetPipeline->CommonMeshesProperties.IsNull() && GenericAssetPipeline->CommonMeshesProperties->VertexColorImportOption != EInterchangeVertexColorImportOption::IVCIO_Ignore)
+					if (GenericAssetPipeline->CommonMeshesProperties && GenericAssetPipeline->CommonMeshesProperties->VertexColorImportOption != EInterchangeVertexColorImportOption::IVCIO_Ignore)
 					{
 						GenericAssetPipeline->SetFlags(RF_Transactional);
 						GenericAssetPipeline->Modify();
diff --git a/Engine/Plugins/MeshPainting/Source/MeshPaintingToolset/Private/MeshPaintSkeletalMeshAdapter.cpp b/Engine/Plugins/MeshPainting/Source/MeshPaintingToolset/Private/MeshPaintSkeletalMeshAdapter.cpp
--- a/Engine/Plugins/MeshPainting/Source/MeshPaintingToolset/Private/MeshPaintSkeletalMeshAdapter.cpp
+++ b/Engine/Plugins/MeshPainting/Source/MeshPaintingToolset/Private/MeshPaintSkeletalMeshAdapter.cpp
@@ -413,7 +413,7 @@ void FMeshPaintSkeletalMeshComponentAdapter::PreEdit()
 				
 				if (GenericAssetPipeline)
 				{
-					if (!GenericAssetPipeline->CommonMeshesProperties.IsNull() && GenericAssetPipeline->CommonMeshesProperties->VertexColorImportOption != EInterchangeVertexColorImportOption::IVCIO_Ignore)
+					if (GenericAssetPipeline->CommonMeshesProperties && GenericAssetPipeline->CommonMeshesProperties->VertexColorImportOption != EInterchangeVertexColorImportOption::IVCIO_Ignore)
 					{
 						GenericAssetPipeline->SetFlags(RF_Transactional);
 						GenericAssetPipeline->Modify();
diff --git a/Engine/Plugins/Runtime/MeshModelingToolset/Source/ModelingComponentsEditorOnly/Private/EditorModelingObjectsCreationAPI.cpp b/Engine/Plugins/Runtime/MeshModelingToolset/Source/ModelingComponentsEditorOnly/Private/EditorModelingObjectsCreationAPI.cpp
--- a/Engine/Plugins/Runtime/MeshModelingToolset/Source/ModelingComponentsEditorOnly/Private/EditorModelingObjectsCreationAPI.cpp
+++ b/Engine/Plugins/Runtime/MeshModelingToolset/Source/ModelingComponentsEditorOnly/Private/EditorModelingObjectsCreationAPI.cpp
@@ -144,7 +144,7 @@ FCreateMeshObjectResult UEditorModelingObjectsCreationAPI::CreateVolume(FCreateM
 	FActorSpawnParameters SpawnInfo;
 	FTransform NewActorTransform = FTransform::Identity;
 	UClass* VolumeClass = ABlockingVolume::StaticClass();
-	if (CreateMeshParams.TypeHintClass.IsNull() == false
+	if (CreateMeshParams.TypeHintClass
 		&& Cast<AVolume>(CreateMeshParams.TypeHintClass.Get()->GetDefaultObject(false)) != nullptr )
 	{
 		VolumeClass = CreateMeshParams.TypeHintClass;
diff --git a/Engine/Source/Editor/BlueprintGraph/Private/K2Node_CustomEvent.cpp b/Engine/Source/Editor/BlueprintGraph/Private/K2Node_CustomEvent.cpp
--- a/Engine/Source/Editor/BlueprintGraph/Private/K2Node_CustomEvent.cpp
+++ b/Engine/Source/Editor/BlueprintGraph/Private/K2Node_CustomEvent.cpp
@@ -119,7 +119,7 @@ public:
 		{
 			auto Predicate_EventGraphs = [Name](const TObjectPtr<UEdGraph>& InEventGraph) -> bool
 			{
-				return !InEventGraph.IsNull() && InEventGraph->HasAnyFlags(RF_Transient) && InEventGraph->GetName() == Name;
+				return InEventGraph && InEventGraph->HasAnyFlags(RF_Transient) && InEventGraph->GetName() == Name;
 			};
 
 			// Allow a transient event subgraph (compiler artifact) that matches the existing name
diff --git a/Engine/Source/Editor/LandscapeEditor/Private/LandscapeEdMode.cpp b/Engine/Source/Editor/LandscapeEditor/Private/LandscapeEdMode.cpp
--- a/Engine/Source/Editor/LandscapeEditor/Private/LandscapeEdMode.cpp
+++ b/Engine/Source/Editor/LandscapeEditor/Private/LandscapeEdMode.cpp
@@ -379,7 +379,7 @@ int32 FEdModeLandscape::GetAccumulatedAllLandscapesResolution() const
 bool FEdModeLandscape::IsLandscapeResolutionCompliant() const
 {
 	const TObjectPtr<const ULandscapeSettings> Settings = GetDefault<ULandscapeSettings>();
-	check(!Settings.IsNull());
+	check(Settings);
 
 	if (Settings->IsLandscapeResolutionRestricted())
 	{
@@ -400,7 +400,7 @@ bool FEdModeLandscape::DoesCurrentToolAffectEditLayers() const
 FText FEdModeLandscape::GetLandscapeResolutionErrorText() const
 {
 	const TObjectPtr<const ULandscapeSettings> Settings = GetDefault<ULandscapeSettings>();
-	check(!Settings.IsNull());
+	check(Settings);
 
 	return FText::Format(LOCTEXT("LandscapeResolutionError", "Total resolution for all Landscape actors cannot exceed the equivalent of {0} x {0}."), Settings->GetSideResolutionLimit());
 }
diff --git a/Engine/Source/Editor/MovieSceneTools/Private/TrackEditors/SkeletalAnimationTrackEditor.cpp b/Engine/Source/Editor/MovieSceneTools/Private/TrackEditors/SkeletalAnimationTrackEditor.cpp
--- a/Engine/Source/Editor/MovieSceneTools/Private/TrackEditors/SkeletalAnimationTrackEditor.cpp
+++ b/Engine/Source/Editor/MovieSceneTools/Private/TrackEditors/SkeletalAnimationTrackEditor.cpp
@@ -547,7 +547,7 @@ FText FSkeletalAnimationSection::GetSectionTitle() const
 {
 	if (Section.Params.Animation != nullptr)
 	{
-		if (Section.Params.MirrorDataTable.IsNull())
+		if (!Section.Params.MirrorDataTable)
 		{
 			return FText::FromString( Section.Params.Animation->GetName() );
 		}
diff --git a/Engine/Source/Editor/UnrealEd/Private/Kismet2/BlueprintEditorUtils.cpp b/Engine/Source/Editor/UnrealEd/Private/Kismet2/BlueprintEditorUtils.cpp
--- a/Engine/Source/Editor/UnrealEd/Private/Kismet2/BlueprintEditorUtils.cpp
+++ b/Engine/Source/Editor/UnrealEd/Private/Kismet2/BlueprintEditorUtils.cpp
@@ -5989,7 +5989,7 @@ namespace UE::Blueprint::Private
 						else if (const UWorld* WorldReferencer = Cast<const UWorld>(AssetReferencer))
 						{
 							const auto& PersistentLevel = WorldReferencer->PersistentLevel;
-							if (!PersistentLevel.IsNull() && PersistentLevel->OwningWorld && Func(PersistentLevel->GetLevelScriptBlueprint()))
+							if (PersistentLevel && PersistentLevel->OwningWorld && Func(PersistentLevel->GetLevelScriptBlueprint()))
 							{
 								GWarn->EndSlowTask();
 								return true;
diff --git a/Engine/Source/Editor/UnrealEd/Private/SSkeletonWidget.cpp b/Engine/Source/Editor/UnrealEd/Private/SSkeletonWidget.cpp
--- a/Engine/Source/Editor/UnrealEd/Private/SSkeletonWidget.cpp
+++ b/Engine/Source/Editor/UnrealEd/Private/SSkeletonWidget.cpp
@@ -879,7 +879,7 @@ FReply SReplaceMissingSkeletonDialog::OnButtonClick(EAppReturnType::Type ButtonI
 	}
 
 	const TObjectPtr<UObject> Asset = SelectedAsset.GetAsset();
-	if (Asset.IsNull())
+	if (!Asset)
 	{
 		return FReply::Handled();
 	}
diff --git a/Engine/Source/Runtime/CoreUObject/Private/UObject/PropertyBaseObject.cpp b/Engine/Source/Runtime/CoreUObject/Private/UObject/PropertyBaseObject.cpp
--- a/Engine/Source/Runtime/CoreUObject/Private/UObject/PropertyBaseObject.cpp
+++ b/Engine/Source/Runtime/CoreUObject/Private/UObject/PropertyBaseObject.cpp
@@ -254,7 +254,7 @@ void FObjectPropertyBase::ExportText_Internal( FString& ValueStr, const void* Pr
 		return;
 	}
 
-	if (Temp.IsNullNoResolve())
+	if (!Temp)
 	{
 		ValueStr += TEXT("None");
 	}
diff --git a/Engine/Source/Runtime/CoreUObject/Public/Templates/Casts.h b/Engine/Source/Runtime/CoreUObject/Public/Templates/Casts.h
--- a/Engine/Source/Runtime/CoreUObject/Public/Templates/Casts.h
+++ b/Engine/Source/Runtime/CoreUObject/Public/Templates/Casts.h
@@ -140,11 +140,9 @@ struct TCastImpl
 
 	FORCEINLINE static To* DoCast( const FObjectPtr& Src )
 	{
-		UObject* SrcObj = ResolveObjectHandleNoRead(Src.GetHandleRef());
-		if (SrcObj && Src->GetClass()->HasAnyCastFlag(TCastFlags<To>::Value))
+		if (Src && Src.GetClass()->HasAnyCastFlag(TCastFlags<To>::Value))
 		{
-			ObjectHandle_Private::OnHandleRead(SrcObj);
-			return (To*)SrcObj;
+			return (To*)Src.Get();
 		}
 		return nullptr;
 	}
@@ -182,11 +180,9 @@ struct TCastImpl<From, To, ECastType::UObjectToUObject>
 
 	FORCEINLINE static To* DoCast( const FObjectPtr& Src )
 	{
-		UObject* SrcObj = ResolveObjectHandleNoRead(Src.GetHandleRef());
-		if (SrcObj && SrcObj->IsA<To>())
+		if (Src && Src.IsA<To>())
 		{
-			ObjectHandle_Private::OnHandleRead(SrcObj);
-			return (To*)SrcObj;
+			return (To*)Src.Get();
 		}
 		return nullptr;
 	}
@@ -408,7 +404,7 @@ template <class T, class U> FORCEINLINE T* ExactCast  ( const TObjectPtr<U>& Src
 template <class T, class U> FORCEINLINE T* CastChecked( const TObjectPtr<U>& Src, ECastCheckedType::Type CheckType = ECastCheckedType::NullChecked )
 {
 #if DO_CHECK
-	if (!Src.IsNull())
+	if (Src)
 	{
 		T* Result = Cast<T>(Src);
 		if (!Result)
diff --git a/Engine/Source/Runtime/CoreUObject/Tests/ObjectPtrTest.cpp b/Engine/Source/Runtime/CoreUObject/Tests/ObjectPtrTest.cpp
--- a/Engine/Source/Runtime/CoreUObject/Tests/ObjectPtrTest.cpp
+++ b/Engine/Source/Runtime/CoreUObject/Tests/ObjectPtrTest.cpp
@@ -218,7 +218,7 @@ TEST_CASE_METHOD(FObjectPtrTestBase, "CoreUObject::TObjectPtr::Forward Declared"
 {
 	UForwardDeclaredObjDerived* PtrFwd = nullptr;
 	TObjectPtr<UForwardDeclaredObjDerived> ObjPtrFwd(MakeObjectPtrUnsafe<UForwardDeclaredObjDerived>(reinterpret_cast<UObject*>(PtrFwd)));
-	TEST_TRUE(TEXT("Null forward declared pointer used to construct a TObjectPtr should result in a null TObjectPtr"), ObjPtrFwd.IsNull());
+	TEST_TRUE(TEXT("Null forward declared pointer used to construct a TObjectPtr should result in a null TObjectPtr"), !ObjPtrFwd);
 }
 
 TEST_CASE_METHOD(FObjectPtrTestBase, "CoreUObject::TObjectPtr::Hash Consistency", "[CoreUObject][ObjectPtr]")
diff --git a/Engine/Source/Runtime/Engine/Private/Actor.cpp b/Engine/Source/Runtime/Engine/Private/Actor.cpp
--- a/Engine/Source/Runtime/Engine/Private/Actor.cpp
+++ b/Engine/Source/Runtime/Engine/Private/Actor.cpp
@@ -789,7 +789,7 @@ void AActor::Serialize(FArchive& Ar)
 					if (FObjectPtrProperty* ObjPtrProp = CastField<FObjectPtrProperty>(ObjProp))
 					{
 						FObjectPtr& ObjectPtr = (FObjectPtr&)ObjPtrProp->GetPropertyValue_InContainer(this);
-						if (!ObjectPtr.IsNullNoResolve() && ObjectPtr.IsA<UActorComponent>())
+						if (ObjectPtr && ObjectPtr.IsA<UActorComponent>())
 						{
 							ActorComponent = Cast<UActorComponent>(ObjectPtr.Get());
 						}
diff --git a/Engine/Source/Runtime/Engine/Private/EdGraph/EdGraphPin.cpp b/Engine/Source/Runtime/Engine/Private/EdGraph/EdGraphPin.cpp
--- a/Engine/Source/Runtime/Engine/Private/EdGraph/EdGraphPin.cpp
+++ b/Engine/Source/Runtime/Engine/Private/EdGraph/EdGraphPin.cpp
@@ -761,7 +761,7 @@ FString UEdGraphPin::GetDefaultAsString() const
 {
 	if(DefaultObject)
 	{
-		return DefaultObject->GetPathName();
+		return DefaultObject.GetPathName();
 	}
 	else if(!DefaultTextValue.IsEmpty())
 	{
@@ -777,7 +777,7 @@ FString UEdGraphPin::GetDefaultAsString() const
 
 bool UEdGraphPin::IsDefaultAsStringEmpty() const
 {
-	return DefaultObject.IsNullNoResolve() && DefaultTextValue.IsEmpty() && DefaultValue.IsEmpty();
+	return !DefaultObject && DefaultTextValue.IsEmpty() && DefaultValue.IsEmpty();
 }
 
 FText UEdGraphPin::GetDefaultAsText() const
diff --git a/Engine/Source/Runtime/Engine/Private/World.cpp b/Engine/Source/Runtime/Engine/Private/World.cpp
--- a/Engine/Source/Runtime/Engine/Private/World.cpp
+++ b/Engine/Source/Runtime/Engine/Private/World.cpp
@@ -8185,7 +8185,7 @@ void UWorld::GetLightMapsAndShadowMaps(ULevel* Level, TArray<UTexture2D*>& OutLi
 
 			// Don't check null references or objects already visited. Also, skip UWorlds as they will pull in more levels than desired
 			// Also skip StaticMesh as it will cause stalls during async compilation and they do not contain any lightmaps anyway.
-			if (!Obj.IsNull() && !Obj.IsA<UWorld>() && !Obj->IsA<UStaticMesh>())
+			if (Obj && !Obj.IsA<UWorld>() && !Obj.IsA<UStaticMesh>())
 			{
 				if (Obj.IsA<ULightMapTexture2D>() ||
 					Obj.IsA<UShadowMapTexture2D>() ||
diff --git a/Engine/Source/Runtime/Slate/Public/Framework/Views/TableViewTypeTraits.h b/Engine/Source/Runtime/Slate/Public/Framework/Views/TableViewTypeTraits.h
--- a/Engine/Source/Runtime/Slate/Public/Framework/Views/TableViewTypeTraits.h
+++ b/Engine/Source/Runtime/Slate/Public/Framework/Views/TableViewTypeTraits.h
@@ -484,7 +484,7 @@ public:
 
 	static bool IsPtrValid( const TObjectPtr<T>& InPtr )
 	{
-		return !InPtr.IsNull();
+		return InPtr != nullptr;
 	}
 
 	static void ResetPtr(TObjectPtr<T>& InPtr )
