Control Rig: Fix Control Rig Control Actors not rendering in packaged project

diff --git a/Engine/Plugins/Animation/ControlRig/Source/ControlRig/Private/ControlRigControlActor.cpp b/Engine/Plugins/Animation/ControlRig/Source/ControlRig/Private/ControlRigControlActor.cpp
--- a/Engine/Plugins/Animation/ControlRig/Source/ControlRig/Private/ControlRigControlActor.cpp
+++ b/Engine/Plugins/Animation/ControlRig/Source/ControlRig/Private/ControlRigControlActor.cpp
@@ -23,12 +23,9 @@ AControlRigControlActor::AControlRigControlActor(const FObjectInitializer& Objec
 	PrimaryActorTick.bTickEvenWhenPaused = true;
 	PrimaryActorTick.TickGroup = TG_PostUpdateWork;
 
-	if (WITH_EDITOR)
-	{
-		PrimaryActorTick.bStartWithTickEnabled = true;
-		bAllowTickBeforeBeginPlay = true;
-	}
-
+	PrimaryActorTick.bStartWithTickEnabled = true;
+	bAllowTickBeforeBeginPlay = true;
+	
 	SetActorEnableCollision(false);
 	
 
