Control Rig: Allow anim controls to drive other controls

diff --git a/Engine/Plugins/Animation/ControlRig/Source/ControlRig/Public/Rigs/RigHierarchyElements.h b/Engine/Plugins/Animation/ControlRig/Source/ControlRig/Public/Rigs/RigHierarchyElements.h
--- a/Engine/Plugins/Animation/ControlRig/Source/ControlRig/Public/Rigs/RigHierarchyElements.h
+++ b/Engine/Plugins/Animation/ControlRig/Source/ControlRig/Public/Rigs/RigHierarchyElements.h
@@ -1260,6 +1260,8 @@ struct CONTROLRIG_API FRigControlElement : public FRigMultiParentElement
 
 	FORCEINLINE bool IsAnimationChannel() const { return Settings.AnimationType == ERigControlAnimationType::AnimationChannel; }
 
+	FORCEINLINE bool CanDriveControls() const { return Settings.AnimationType == ERigControlAnimationType::ProxyControl || Settings.AnimationType == ERigControlAnimationType::AnimationControl; }
+
 	virtual void Save(FArchive& A, URigHierarchy* Hierarchy, ESerializationPhase SerializationPhase) override;
 	virtual void Load(FArchive& Ar, URigHierarchy* Hierarchy, ESerializationPhase SerializationPhase) override;
 
diff --git a/Engine/Plugins/Animation/ControlRig/Source/ControlRigEditor/Private/ControlRigElementDetails.cpp b/Engine/Plugins/Animation/ControlRig/Source/ControlRigEditor/Private/ControlRigElementDetails.cpp
--- a/Engine/Plugins/Animation/ControlRig/Source/ControlRigEditor/Private/ControlRigElementDetails.cpp
+++ b/Engine/Plugins/Animation/ControlRig/Source/ControlRigEditor/Private/ControlRigElementDetails.cpp
@@ -2768,7 +2768,7 @@ void FRigControlElementDetails::CustomizeControl(IDetailLayoutBuilder& DetailBui
 		}));
 	}
 	
-	if(!IsAnyControlNotOfAnimationType(ERigControlAnimationType::ProxyControl))
+	if(IsAnyControlOfAnimationType(ERigControlAnimationType::ProxyControl) || IsAnyControlOfAnimationType(ERigControlAnimationType::AnimationControl))
 	{
 		ControlCategory.AddProperty(SettingsHandle->GetChildHandle(TEXT("DrivenControls")).ToSharedRef())
 		.IsEnabled(bIsEnabled);
diff --git a/Engine/Plugins/Animation/ControlRig/Source/ControlRigEditor/Private/EditMode/ControlRigEditMode.cpp b/Engine/Plugins/Animation/ControlRig/Source/ControlRigEditor/Private/EditMode/ControlRigEditMode.cpp
--- a/Engine/Plugins/Animation/ControlRig/Source/ControlRigEditor/Private/EditMode/ControlRigEditMode.cpp
+++ b/Engine/Plugins/Animation/ControlRig/Source/ControlRigEditor/Private/EditMode/ControlRigEditMode.cpp
@@ -3554,7 +3554,7 @@ void FControlRigEditMode::OnHierarchyModified(ERigHierarchyNotification InNotif,
 							{
 								ControlProxy->SelectProxy(ControlRig,Key.Name, bSelected);
 
-								if(ControlElement->Settings.AnimationType == ERigControlAnimationType::ProxyControl)
+								if(ControlElement->CanDriveControls())
 								{
 									const UControlRigEditModeSettings* Settings = GetDefault<UControlRigEditModeSettings>();
 
@@ -4494,7 +4494,7 @@ void FControlRigEditMode::NotifyDrivenControls(UControlRig* InControlRig, const
 	// if we are changing a proxy control - we also need to notify the change for the driven controls
 	if (FRigControlElement* ControlElement = InControlRig->GetHierarchy()->Find<FRigControlElement>(InKey))
 	{
-		if(ControlElement->Settings.AnimationType == ERigControlAnimationType::ProxyControl)
+		if(ControlElement->CanDriveControls())
 		{
 			FRigControlModifiedContext Context;
 			Context.EventName = FRigUnit_BeginExecution::EventName;
diff --git a/Engine/Plugins/Animation/ControlRig/Source/ControlRigEditor/Private/Sequencer/ControlRigParameterTrackEditor.cpp b/Engine/Plugins/Animation/ControlRig/Source/ControlRigEditor/Private/Sequencer/ControlRigParameterTrackEditor.cpp
--- a/Engine/Plugins/Animation/ControlRig/Source/ControlRigEditor/Private/Sequencer/ControlRigParameterTrackEditor.cpp
+++ b/Engine/Plugins/Animation/ControlRig/Source/ControlRigEditor/Private/Sequencer/ControlRigParameterTrackEditor.cpp
@@ -2804,8 +2804,15 @@ void FControlRigParameterTrackEditor::HandleControlSelected(UControlRig* Subject
 	
 	URigHierarchy* Hierarchy = Subject->GetHierarchy();
 	static bool bIsSelectingIndirectControl = false;
+	static TArray<FRigControlElement*> SelectedElements = {};
 
-	if(ControlElement->Settings.AnimationType == ERigControlAnimationType::ProxyControl)
+	// Avoid cyclic selection
+	if (SelectedElements.Contains(ControlElement))
+	{
+		return;
+	}
+
+	if(ControlElement->CanDriveControls())
 	{
 		const TArray<FRigElementKey>& DrivenControls = ControlElement->Settings.DrivenControls;
 		for(const FRigElementKey& DrivenKey : DrivenControls)
@@ -2813,10 +2820,18 @@ void FControlRigParameterTrackEditor::HandleControlSelected(UControlRig* Subject
 			if(FRigControlElement* DrivenControl = Hierarchy->Find<FRigControlElement>(DrivenKey))
 			{
 				TGuardValue<bool> SubControlGuard(bIsSelectingIndirectControl, true);
+
+				TArray<FRigControlElement*> NewSelection = SelectedElements;
+				NewSelection.Add(ControlElement);
+				TGuardValue<TArray<FRigControlElement*>> SelectedElementsGuard(SelectedElements, NewSelection);
+				
 				HandleControlSelected(Subject, DrivenControl, bSelected);
 			}
 		}
-		return;
+		if(ControlElement->Settings.AnimationType == ERigControlAnimationType::ProxyControl)
+		{
+			return;
+		}
 	}
 	
 	//if parent selected we select child here if it's a bool,integer or single float
diff --git a/Engine/Plugins/Animation/ControlRig/Source/ControlRigEditor/Private/Sequencer/ControlRigSequencerFilter.cpp b/Engine/Plugins/Animation/ControlRig/Source/ControlRigEditor/Private/Sequencer/ControlRigSequencerFilter.cpp
--- a/Engine/Plugins/Animation/ControlRig/Source/ControlRigEditor/Private/Sequencer/ControlRigSequencerFilter.cpp
+++ b/Engine/Plugins/Animation/ControlRig/Source/ControlRigEditor/Private/Sequencer/ControlRigSequencerFilter.cpp
@@ -182,7 +182,7 @@ public:
 					}
 					if (const FRigControlElement* ControlElement = Cast<FRigControlElement>(SelectedControl))
 					{
-						if (ControlElement->Settings.AnimationType == ERigControlAnimationType::ProxyControl)
+						if (ControlElement->CanDriveControls())
 						{
 							const TArray<FRigElementKey>& DrivenControls = ControlElement->Settings.DrivenControls;
 							for (const FRigElementKey& DrivenKey : DrivenControls)
