Control Rig: Fix category for AimContraint node

diff --git a/Engine/Plugins/Animation/ControlRig/Source/ControlRig/Private/Units/Highlevel/Hierarchy/RigUnit_AimBone.h b/Engine/Plugins/Animation/ControlRig/Source/ControlRig/Private/Units/Highlevel/Hierarchy/RigUnit_AimBone.h
--- a/Engine/Plugins/Animation/ControlRig/Source/ControlRig/Private/Units/Highlevel/Hierarchy/RigUnit_AimBone.h
+++ b/Engine/Plugins/Animation/ControlRig/Source/ControlRig/Private/Units/Highlevel/Hierarchy/RigUnit_AimBone.h
@@ -420,7 +420,7 @@ struct FRigUnit_AimConstraint_AdvancedSettings
  * Orients an item such that its aim axis points towards a global target.
  * Note: This node operates in global space!
  */
-USTRUCT(meta=(DisplayName="Aim Constraint", Category="Hierarchy", Keywords="Lookat"))
+USTRUCT(meta=(DisplayName="Aim Constraint", Category="Constraints", Keywords="Lookat, Aim"))
 struct CONTROLRIG_API FRigUnit_AimConstraintLocalSpaceOffset: public FRigUnit_HighlevelBaseMutable
 {
 	GENERATED_BODY()
