Control Rig: Avoid notifications for nodes being added more than once

diff --git a/Engine/Plugins/Animation/ControlRig/Source/ControlRigDeveloper/Private/Graph/ControlRigGraph.cpp b/Engine/Plugins/Animation/ControlRig/Source/ControlRigDeveloper/Private/Graph/ControlRigGraph.cpp
--- a/Engine/Plugins/Animation/ControlRig/Source/ControlRigDeveloper/Private/Graph/ControlRigGraph.cpp
+++ b/Engine/Plugins/Animation/ControlRig/Source/ControlRigDeveloper/Private/Graph/ControlRigGraph.cpp
@@ -437,9 +437,10 @@ void UControlRigGraph::HandleModifiedEvent(ERigVMGraphNotifType InNotifType, URi
 		case ERigVMGraphNotifType::NodeAdded:
 		{
 			ModelNodePathToEdNode.Reset();
-				
+
 			if (URigVMNode* ModelNode = Cast<URigVMNode>(InSubject))
 			{
+			
 				if (!ModelNode->IsVisibleInUI())
 				{
 					if (URigVMInjectionInfo* Injection = ModelNode->GetInjectionInfo())
@@ -462,6 +463,14 @@ void UControlRigGraph::HandleModifiedEvent(ERigVMGraphNotifType InNotifType, URi
 					}
 					break;
 				}
+				else
+				{
+					// check if the node is already part of the graph
+					if(FindNodeForModelNodeName(ModelNode->GetFName()) != nullptr)
+					{
+						break;
+					}
+				}
 
 				if (URigVMCommentNode* CommentModelNode = Cast<URigVMCommentNode>(ModelNode))
 				{
