ControlRigComponent: Should reinitialize the control rig after CRBP compile. Using the same approach that the animnode_controlrig is using to detect when BP reinstancing is happening.

diff --git a/Engine/Plugins/Animation/ControlRig/Source/ControlRig/Private/ControlRigComponent.cpp b/Engine/Plugins/Animation/ControlRig/Source/ControlRig/Private/ControlRigComponent.cpp
--- a/Engine/Plugins/Animation/ControlRig/Source/ControlRig/Private/ControlRigComponent.cpp
+++ b/Engine/Plugins/Animation/ControlRig/Source/ControlRig/Private/ControlRigComponent.cpp
@@ -95,6 +95,22 @@ void UControlRigComponent::BeginDestroy()
 	gPendingSkeletalMeshes.Remove(this);
 }
 
+void UControlRigComponent::Serialize(FArchive& Ar)
+{
+	Super::Serialize(Ar);
+
+	// after compile, we have to reinitialize
+	// because it needs new execution code
+	// since memory has changed, similar to FAnimNode_ControlRig::PostSerialize
+	if (Ar.IsObjectReferenceCollector())
+	{
+		if (ControlRig)
+		{
+			ControlRig->Initialize();
+		}
+	}
+}
+
 #if WITH_EDITOR
 void UControlRigComponent::InitializeComponent()
 {
diff --git a/Engine/Plugins/Animation/ControlRig/Source/ControlRig/Public/ControlRigComponent.h b/Engine/Plugins/Animation/ControlRig/Source/ControlRig/Public/ControlRigComponent.h
--- a/Engine/Plugins/Animation/ControlRig/Source/ControlRig/Public/ControlRigComponent.h
+++ b/Engine/Plugins/Animation/ControlRig/Source/ControlRig/Public/ControlRigComponent.h
@@ -212,6 +212,8 @@ public:
 #endif
 
 	virtual void BeginDestroy() override;
+
+	virtual void Serialize(FArchive& Ar) override;
 	//~ End UObject Interface
 	
 	//~ Begin UActorComponent Interface
