Control Rig: Fix blurriness in a sequence with an FK rig

diff --git a/Engine/Plugins/Animation/ControlRig/Source/ControlRigEditor/Private/EditMode/ControlRigEditMode.cpp b/Engine/Plugins/Animation/ControlRig/Source/ControlRigEditor/Private/EditMode/ControlRigEditMode.cpp
--- a/Engine/Plugins/Animation/ControlRig/Source/ControlRigEditor/Private/EditMode/ControlRigEditMode.cpp
+++ b/Engine/Plugins/Animation/ControlRig/Source/ControlRigEditor/Private/EditMode/ControlRigEditMode.cpp
@@ -564,17 +564,11 @@ void FControlRigEditMode::Tick(FEditorViewportClient* ViewportClient, float Delt
 		ControlRigsToRecreate.SetNum(0);
 	}
 
-	// We need to tick here since changing a bone for example
-	// might have changed the transform of the Control
 	{
-		if (IsInLevelEditor())
-		{
-			TickManipulatableObjects(0.f);
-		}
-		else
-		{
-			PostPoseUpdate();
-		}
+		// We need to tick here since changing a bone for example
+		// might have changed the transform of the Control
+		PostPoseUpdate();
+		
 		if (IsInLevelEditor() == false) //only do this check if not in level editor
 		{
 			for (TWeakObjectPtr<UControlRig>& RuntimeRigPtr : RuntimeControlRigs)
@@ -3613,6 +3607,11 @@ void FControlRigEditMode::OnHierarchyModified(ERigHierarchyNotification InNotif,
 				}
 			}
 		}
+		case ERigHierarchyNotification::ParentWeightsChanged:
+		{
+			TickManipulatableObjects(0.f);
+			break;
+		}
 		case ERigHierarchyNotification::InteractionBracketOpened:
 		case ERigHierarchyNotification::InteractionBracketClosed:
 		default:
