Control Rig: UI Graph Performance Improvements

diff --git a/Engine/Plugins/Animation/ControlRig/Source/ControlRigDeveloper/Private/Graph/ControlRigGraph.cpp b/Engine/Plugins/Animation/ControlRig/Source/ControlRigDeveloper/Private/Graph/ControlRigGraph.cpp
--- a/Engine/Plugins/Animation/ControlRig/Source/ControlRigDeveloper/Private/Graph/ControlRigGraph.cpp
+++ b/Engine/Plugins/Animation/ControlRig/Source/ControlRigDeveloper/Private/Graph/ControlRigGraph.cpp
@@ -436,11 +436,8 @@ void UControlRigGraph::HandleModifiedEvent(ERigVMGraphNotifType InNotifType, URi
 		}
 		case ERigVMGraphNotifType::NodeAdded:
 		{
-			ModelNodePathToEdNode.Reset();
-
 			if (URigVMNode* ModelNode = Cast<URigVMNode>(InSubject))
 			{
-			
 				if (!ModelNode->IsVisibleInUI())
 				{
 					if (URigVMInjectionInfo* Injection = ModelNode->GetInjectionInfo())
@@ -489,6 +486,8 @@ void UControlRigGraph::HandleModifiedEvent(ERigVMGraphNotifType InNotifType, URi
 					NewNode->NodeComment = CommentModelNode->GetCommentText();
 					NewNode->SetFlags(RF_Transactional);
 					NewNode->GetNodesUnderComment();
+
+					ModelNodePathToEdNode.Add(ModelNode->GetFName(), NewNode);
 				}
 				else if (URigVMRerouteNode* RerouteModelNode = Cast<URigVMRerouteNode>(ModelNode))
 				{
@@ -504,12 +503,13 @@ void UControlRigGraph::HandleModifiedEvent(ERigVMGraphNotifType InNotifType, URi
 					NewNode->NodePosY = ModelNode->GetPosition().Y;
 
 					NewNode->SetFlags(RF_Transactional);
-					NewNode->AllocateDefaultPins();
 
 					if (UEdGraphPin* ValuePin = NewNode->FindPin(ModelNode->FindPin("Value")->GetPinPath()))
 					{
 						NewNode->SetColorFromModel(GetSchema()->GetPinTypeColor(ValuePin->PinType));
 					}
+
+					ModelNodePathToEdNode.Add(ModelNode->GetFName(), NewNode);
 				}
 				else // struct, library, parameter + variable
 				{
@@ -520,31 +520,29 @@ void UControlRigGraph::HandleModifiedEvent(ERigVMGraphNotifType InNotifType, URi
 					NewNode->CreateNewGuid();
 					NewNode->PostPlacedNewNode();
 					NewNode->AllocateDefaultPins();
+					NewNode->PostReconstructNode();
 
 					NewNode->NodePosX = ModelNode->GetPosition().X;
 					NewNode->NodePosY = ModelNode->GetPosition().Y;
 					NewNode->SetColorFromModel(ModelNode->GetNodeColor());
 					NewNode->SetFlags(RF_Transactional);
-					NewNode->AllocateDefaultPins();
+
+					ModelNodePathToEdNode.Add(ModelNode->GetFName(), NewNode);
 				}
 			}
 			break;
 		}
 		case ERigVMGraphNotifType::NodeRemoved:
 		{
-			ModelNodePathToEdNode.Reset();
-
 			if (URigVMNode* ModelNode = Cast<URigVMNode>(InSubject))
 			{
 				if (URigVMInjectionInfo* Injection = ModelNode->GetInjectionInfo())
 				{
 					if (URigVMPin* ModelPin = Injection->GetPin())
 					{
-						URigVMNode* ParentModelNode = ModelPin->GetNode();
-						if (ParentModelNode)
+						if (URigVMNode* ParentModelNode = ModelPin->GetNode())
 						{
-							UEdGraphNode* EdNode = FindNodeForModelNodeName(ParentModelNode->GetFName());
-							if (EdNode)
+							if (UEdGraphNode* EdNode = FindNodeForModelNodeName(ParentModelNode->GetFName()))
 							{
 								if (UControlRigGraphNode* RigNode = Cast<UControlRigGraphNode>(EdNode))
 								{
@@ -556,8 +554,9 @@ void UControlRigGraph::HandleModifiedEvent(ERigVMGraphNotifType InNotifType, URi
 					break;
 				}
 
-				UEdGraphNode* EdNode = FindNodeForModelNodeName(ModelNode->GetFName(), false);
-				if (EdNode)
+				ModelNodePathToEdNode.Remove(ModelNode->GetFName());
+
+				if (UEdGraphNode* EdNode = FindNodeForModelNodeName(ModelNode->GetFName(), false))
 				{
 					RemoveNode(EdNode);
 				}
@@ -599,8 +598,7 @@ void UControlRigGraph::HandleModifiedEvent(ERigVMGraphNotifType InNotifType, URi
 			{
 				if (UControlRigGraphNode* RigNode = Cast<UControlRigGraphNode>(FindNodeForModelNodeName(ModelNode->GetFName())))
 				{
-					RigNode->InvalidateNodeTitle();
-					RigNode->ReconstructNode_Internal(true);
+					RigNode->SyncGraphNodeTitleWithModelNodeTitle();
 				}
 			}
 			break;
@@ -722,7 +720,7 @@ void UControlRigGraph::HandleModifiedEvent(ERigVMGraphNotifType InNotifType, URi
 					{
 						if(ModelPin->GetNode()->IsEvent())
 						{
-							RigNode->InvalidateNodeTitle();
+							RigNode->SyncGraphNodeTitleWithModelNodeTitle();
 						}
 						break;
 					}
@@ -733,13 +731,13 @@ void UControlRigGraph::HandleModifiedEvent(ERigVMGraphNotifType InNotifType, URi
 					{
 						if (ModelPin->GetName() == TEXT("Variable"))
 						{
-							RigNode->InvalidateNodeTitle();
-							RigNode->ReconstructNode_Internal(true);
+							RigNode->SyncGraphNodeTitleWithModelNodeTitle();
+							RigNode->SynchronizeGraphPinValueWithModelPin(ModelPin);
 						}
 					}
 					else if (Cast<URigVMUnitNode>(ModelPin->GetNode()))
 					{
-						RigNode->InvalidateNodeTitle();
+						RigNode->SyncGraphNodeTitleWithModelNodeTitle();
 
 						// if the node contains a rig element key - invalidate the node
 						if(RigNode->GetAllPins().ContainsByPredicate([](UEdGraphPin* Pin) -> bool
@@ -748,7 +746,7 @@ void UControlRigGraph::HandleModifiedEvent(ERigVMGraphNotifType InNotifType, URi
 						}))
 						{
 							// we do this to enforce the refresh of the element name widgets
-							RigNode->ReconstructNode_Internal(true);
+							RigNode->SynchronizeGraphPinValueWithModelPin(ModelPin);
 						}
 					}
 				}
@@ -766,7 +764,7 @@ void UControlRigGraph::HandleModifiedEvent(ERigVMGraphNotifType InNotifType, URi
 								{
 									if (UControlRigGraphNode* HostRigNode = Cast<UControlRigGraphNode>(HostEdNode))
 									{
-										HostRigNode->ReconstructNode_Internal(true);
+										HostRigNode->ModelPinsChanged();
 									}
 								}
 							}
@@ -776,36 +774,95 @@ void UControlRigGraph::HandleModifiedEvent(ERigVMGraphNotifType InNotifType, URi
 			}
 			break;
 		}
+		case ERigVMGraphNotifType::PinAdded:
+		{
+			if (URigVMPin* ModelPin = Cast<URigVMPin>(InSubject))
+			{
+				if(URigVMNode* ModelNode = ModelPin->GetNode())
+				{
+					if (UControlRigGraphNode* RigNode = Cast<UControlRigGraphNode>(FindNodeForModelNodeName(ModelNode->GetFName())))
+					{
+						RigNode->ModelPinAdded(ModelPin);
+					}
+				}
+			}
+			break;
+		}
+		case ERigVMGraphNotifType::PinRemoved:
+		{
+			if (URigVMPin* ModelPin = Cast<URigVMPin>(InSubject))
+			{
+				if(URigVMNode* ModelNode = ModelPin->GetNode())
+				{
+					if (UControlRigGraphNode* RigNode = Cast<UControlRigGraphNode>(FindNodeForModelNodeName(ModelNode->GetFName())))
+					{
+						RigNode->ModelPinRemoved(ModelPin);
+					}
+				}
+			}
+			break;
+		}
 		case ERigVMGraphNotifType::PinArraySizeChanged:
+		{
+			// don't do anything here - the UI will update based on the
+			// PinAdded and PinRemoved notifs
+			break;
+		}
 		case ERigVMGraphNotifType::PinDirectionChanged:
-		case ERigVMGraphNotifType::PinTypeChanged:
 		case ERigVMGraphNotifType::PinIndexChanged:
 		case ERigVMGraphNotifType::PinBoundVariableChanged:
-		case ERigVMGraphNotifType::PinAdded:
-		case ERigVMGraphNotifType::PinRemoved:
+		{
+			if (URigVMPin* ModelPin = Cast<URigVMPin>(InSubject))
+			{
+				if (UControlRigGraphNode* RigNode = Cast<UControlRigGraphNode>(FindNodeForModelNodeName(ModelPin->GetNode()->GetFName())))
+				{
+					RigNode->ModelPinsChanged();
+				}
+			}
+			break;
+		}
+		case ERigVMGraphNotifType::LibraryTemplateChanged:
+		{
+			if (URigVMNode* LibraryNode = Cast<URigVMNode>(InSubject))
+			{
+				if (UControlRigGraphNode* RigNode = Cast<UControlRigGraphNode>(FindNodeForModelNodeName(LibraryNode->GetFName())))
+				{
+					RigNode->ModelPinsChanged(true);
+				}
+			}
+			break;
+		}
+		case ERigVMGraphNotifType::PinTypeChanged:
+		{
+			if (URigVMPin* ModelPin = Cast<URigVMPin>(InSubject))
+			{
+				if (UControlRigGraphNode* RigNode = Cast<UControlRigGraphNode>(FindNodeForModelNodeName(ModelPin->GetNode()->GetFName())))
+				{
+					RigNode->SynchronizeGraphPinTypeWithModelPin(ModelPin);
+				}
+			}
+			break;
+		}
 		case ERigVMGraphNotifType::PinRenamed:
 		{
 			if (URigVMPin* ModelPin = Cast<URigVMPin>(InSubject))
 			{
 				if (UControlRigGraphNode* RigNode = Cast<UControlRigGraphNode>(FindNodeForModelNodeName(ModelPin->GetNode()->GetFName())))
 				{
-					RigNode->ReconstructNode_Internal(true);
+					RigNode->SynchronizeGraphPinNameWithModelPin(ModelPin);
 				}
 			}
 			break;
 		}
 		case ERigVMGraphNotifType::NodeRenamed:
 		{
-			ModelNodePathToEdNode.Reset();
-				
 			if (URigVMNode* ModelNode = Cast<URigVMNode>(InSubject))
 			{
+				ModelNodePathToEdNode.Remove(ModelNode->GetPreviousFName());
 				if (UControlRigGraphNode* RigNode = Cast<UControlRigGraphNode>(FindNodeForModelNodeName(ModelNode->GetPreviousFName())))
 				{
-					RigNode->Rename(*ModelNode->GetName());
-					RigNode->ModelNodePath = ModelNode->GetNodePath();
-					RigNode->InvalidateNodeTitle();
-					RigNode->ReconstructNode_Internal(true);
+					RigNode->SyncGraphNodeNameWithModelNodeName(ModelNode);
+					ModelNodePathToEdNode.Add(ModelNode->GetFName(), RigNode);
 				}
 			}
 			break;
@@ -817,7 +874,7 @@ void UControlRigGraph::HandleModifiedEvent(ERigVMGraphNotifType InNotifType, URi
 			{
 				if (UControlRigGraphNode* RigNode = Cast<UControlRigGraphNode>(FindNodeForModelNodeName(ModelNode->GetFName())))
 				{
-					RigNode->InvalidateNodeTitle();
+					RigNode->SyncGraphNodeTitleWithModelNodeTitle();
 				}
 			}
 			break;
@@ -849,6 +906,16 @@ void UControlRigGraph::HandleModifiedEvent(ERigVMGraphNotifType InNotifType, URi
 			break;
 		}
 		case ERigVMGraphNotifType::PinExpansionChanged:
+		{
+			if (URigVMPin* ModelPin = Cast<URigVMPin>(InSubject))
+			{
+				if (UControlRigGraphNode* RigNode = Cast<UControlRigGraphNode>(FindNodeForModelNodeName(ModelPin->GetNode()->GetFName())))
+				{
+					RigNode->OnNodePinExpansionChanged().Broadcast();
+				}
+			}
+			break;
+		}
 		default:
 		{
 			break;
diff --git a/Engine/Plugins/Animation/ControlRig/Source/ControlRigDeveloper/Private/Graph/ControlRigGraphNode.cpp b/Engine/Plugins/Animation/ControlRig/Source/ControlRigDeveloper/Private/Graph/ControlRigGraphNode.cpp
--- a/Engine/Plugins/Animation/ControlRig/Source/ControlRigDeveloper/Private/Graph/ControlRigGraphNode.cpp
+++ b/Engine/Plugins/Animation/ControlRig/Source/ControlRigDeveloper/Private/Graph/ControlRigGraphNode.cpp
@@ -236,8 +236,6 @@ void UControlRigGraphNode::ReconstructNode_Internal(bool bForce)
 	Pins.Reset();
 
 	// Recreate the new pins
-	CachedPins.Reset();
-	CachedModelPins.Reset();
 	ReallocatePinsDuringReconstruction(OldPins);
 
 	// Maintain watches up to date
@@ -264,10 +262,8 @@ void UControlRigGraphNode::ReconstructNode_Internal(bool bForce)
 	// Let subclasses do any additional work
 	PostReconstructNode();
 
-	if (RigGraph)
-	{
-		RigGraph->NotifyGraphChanged();
-	}
+	InvalidateNodeTitle();
+	OnNodePinsChanged().Broadcast();
 }
 
 bool UControlRigGraphNode::IsDeprecated() const
@@ -467,6 +463,467 @@ void UControlRigGraphNode::ClearErrorInfo()
 	ErrorMsg = FString();	
 }
 
+URigVMPin* UControlRigGraphNode::FindModelPinFromGraphPin(const UEdGraphPin* InGraphPin) const
+{
+	if(InGraphPin == nullptr)
+	{
+		return nullptr;
+	}
+	
+	for(const auto& Pair : CachedPins)
+	{
+		if(Pair.Value.InputPin == InGraphPin ||
+			Pair.Value.OutputPin == InGraphPin)
+		{
+			return Pair.Key;
+		}
+	}
+
+	return nullptr;
+}
+
+UEdGraphPin* UControlRigGraphNode::FindGraphPinFromModelPin(const URigVMPin* InModelPin, bool bAsInput) const
+{
+	if(InModelPin == nullptr)
+	{
+		return nullptr;
+	}
+	
+	if(const FPinPair* Pair = CachedPins.Find(InModelPin))
+	{
+		return bAsInput ? Pair->InputPin : Pair->OutputPin;
+	}
+
+	return nullptr;
+}
+
+void UControlRigGraphNode::SynchronizeGraphPinNameWithModelPin(const URigVMPin* InModelPin, bool bNotify)
+{
+	auto SyncGraphPinLambda = [this](const URigVMPin* InModelPin, bool bAsInput) -> bool
+	{
+		if(UEdGraphPin* GraphPin = FindGraphPinFromModelPin(InModelPin, bAsInput))
+		{
+			const FString OldPinName = GraphPin->PinName.ToString();
+			const FString NewPinName = InModelPin->GetPinPath();
+
+			if(OldPinName != NewPinName)
+			{
+				PinPathToModelPin.Remove(OldPinName);
+				
+				GraphPin->PinName = *NewPinName;
+				GraphPin->PinFriendlyName = FText::FromName(InModelPin->GetDisplayName());
+
+				PinPathToModelPin.Add(NewPinName, (URigVMPin*)InModelPin);
+				return true;
+			}
+		}
+
+		return false;
+	};
+
+	bool bResult = false;
+	
+	if(SyncGraphPinLambda(InModelPin, true))
+	{
+		bResult = true;
+	}
+	if(SyncGraphPinLambda(InModelPin, false))
+	{
+		bResult = true;
+	}
+
+	if(bResult)
+	{
+		for (const URigVMPin* ModelSubPin : InModelPin->GetSubPins())
+		{
+			SynchronizeGraphPinNameWithModelPin(ModelSubPin, false);
+		}
+
+		if(bNotify)
+		{
+			// Notify the node widget that the pins have changed.
+			OnNodePinsChanged().Broadcast();
+		}
+	}
+}
+
+void UControlRigGraphNode::SynchronizeGraphPinValueWithModelPin(const URigVMPin* InModelPin)
+{
+	auto SyncGraphPinLambda = [this](const URigVMPin* InModelPin, bool bAsInput)
+	{
+		// If the pin has sub-pins, we may need to remove or rebuild the sub-pins.
+		if (UEdGraphPin* GraphPin = FindGraphPinFromModelPin(InModelPin, bAsInput))
+		{
+			SetupPinDefaultsFromModel(GraphPin, InModelPin);
+		}
+	};
+
+	SyncGraphPinLambda(InModelPin, true);
+	SyncGraphPinLambda(InModelPin, false);
+}
+
+void UControlRigGraphNode::SynchronizeGraphPinTypeWithModelPin(const URigVMPin* InModelPin)
+{
+	bool bNotify = false;
+
+	auto SyncGraphPinLambda = [this, &bNotify](const URigVMPin* InModelPin, bool bAsInput)
+	{
+		// If the pin has sub-pins, we may need to remove or rebuild the sub-pins.
+		if (UEdGraphPin* GraphPin = FindGraphPinFromModelPin(InModelPin, bAsInput))
+		{
+			const FEdGraphPinType NewGraphPinType = GetPinTypeForModelPin(InModelPin);
+			if(NewGraphPinType != GraphPin->PinType)
+			{
+				GraphPin->PinType = NewGraphPinType;
+				ConfigurePin(GraphPin, InModelPin);
+
+				// Create new sub-pins, if required, to reflect the new type.
+				TArray<UEdGraphPin*> GraphSubPinsToKeep;
+				if (!InModelPin->GetSubPins().IsEmpty())
+				{
+					for (const URigVMPin* ModelSubPin : InModelPin->GetSubPins())
+					{
+						if(UEdGraphPin* GraphSubPin = FindGraphPinFromModelPin(ModelSubPin, bAsInput))
+						{
+							GraphSubPinsToKeep.Add(GraphSubPin);
+
+							const FEdGraphPinType NewGraphSubPinType = GetPinTypeForModelPin(ModelSubPin);
+							if(NewGraphSubPinType != GraphSubPin->PinType)
+							{
+								GraphSubPin->PinType = NewGraphSubPinType;
+								ConfigurePin(GraphSubPin, ModelSubPin);
+							}
+						}
+						else
+						{
+							CreateGraphPinFromModelPin(ModelSubPin, GraphPin->Direction, GraphPin);
+						}
+					}
+				}
+
+				// If the graph node had other sub-pins, we need to remove those.
+				RemoveGraphSubPins(GraphPin, GraphSubPinsToKeep);
+
+				bNotify = true;
+			}
+		}
+	};
+
+	SyncGraphPinLambda(InModelPin, true);
+	SyncGraphPinLambda(InModelPin, false);
+
+	if(bNotify)
+	{
+		// Notify the node widget that the pins have changed.
+		OnNodePinsChanged().Broadcast();
+	}
+}
+
+void UControlRigGraphNode::SynchronizeGraphPinExpansionWithModelPin(const URigVMPin* InModelPin)
+{
+	OnNodePinExpansionChanged().Broadcast();
+}
+
+void UControlRigGraphNode::SyncGraphNodeTitleWithModelNodeTitle()
+{
+	InvalidateNodeTitle();
+}
+
+void UControlRigGraphNode::SyncGraphNodeNameWithModelNodeName(const URigVMNode* InModelNode)
+{
+	Rename(*InModelNode->GetName());
+	ModelNodePath = InModelNode->GetNodePath();
+	SyncGraphNodeTitleWithModelNodeTitle();
+}
+
+bool UControlRigGraphNode::CreateGraphPinFromModelPin(const URigVMPin* InModelPin, EEdGraphPinDirection InDirection,
+	UEdGraphPin* InParentPin)
+{
+	// don't create output pins for array elements
+	if(InDirection == EGPD_Output && InModelPin->IsArrayElement())
+	{
+		return false;
+	}
+
+	const FPinPair PairConst = CachedPins.FindOrAdd((URigVMPin*)InModelPin);
+
+	auto CreatePinLambda = [this](const URigVMPin* InModelPin, EEdGraphPinDirection InDirection, UEdGraphPin* InParentPin) -> UEdGraphPin*
+	{
+		UEdGraphPin* GraphPin = CreatePin(InDirection, GetPinTypeForModelPin(InModelPin), FName(*InModelPin->GetPinPath()));
+		if (GraphPin)
+		{
+			ConfigurePin(GraphPin, InModelPin);
+
+			if (InParentPin)
+			{
+				InParentPin->SubPins.Add(GraphPin);
+				GraphPin->ParentPin = InParentPin;
+			}
+			
+			for(const URigVMPin* ModelSubPin : InModelPin->GetSubPins())
+			{
+				CreateGraphPinFromModelPin(ModelSubPin, InDirection, GraphPin);
+			}
+		}
+
+		return GraphPin;
+	};
+
+	UEdGraphPin* InputPin = nullptr;
+	UEdGraphPin* OutputPin = nullptr;
+	
+	if (InDirection == EGPD_Input && PairConst.InputPin == nullptr)
+	{
+		InputPin = CreatePinLambda(InModelPin, EGPD_Input, InParentPin);
+	}
+	if (InDirection == EGPD_Output && PairConst.OutputPin == nullptr)
+	{
+		OutputPin = CreatePinLambda(InModelPin, EGPD_Output, InParentPin);
+	}
+
+	FPinPair& Pair = CachedPins.FindChecked((URigVMPin*)InModelPin);
+	Pair.InputPin = InputPin != nullptr ? InputPin : Pair.InputPin;
+	Pair.OutputPin = OutputPin != nullptr ? OutputPin : Pair.OutputPin;
+	return Pair.IsValid();
+}
+
+void UControlRigGraphNode::RemoveGraphSubPins(UEdGraphPin* InParentPin, const TArray<UEdGraphPin*>& InPinsToKeep)
+{
+	TArray<UEdGraphPin*> SubPins = InParentPin->SubPins;
+
+	TArray<URigVMPin*> ModelSubPins;
+	for (const UEdGraphPin* SubPin: SubPins)
+	{
+		if(InPinsToKeep.Contains(SubPin))
+		{
+			continue;
+		}
+		for(const auto& Pair : CachedPins)
+		{
+			if(Pair.Value.InputPin == SubPin ||
+				Pair.Value.OutputPin == SubPin)
+			{
+				ModelSubPins.Add(Pair.Key);
+			}
+		}
+	}
+
+	for(const URigVMPin* ModelSubPin : ModelSubPins)
+	{
+		PinPathToModelPin.Remove(ModelSubPin->GetPinPath());
+
+		if(const FPinPair* PinPair = CachedPins.Find(ModelSubPin))
+		{
+			auto Traverse = [this](UEdGraphPin* SubPin)
+			{
+				if(SubPin == nullptr)
+				{
+					return;
+				}
+				
+				// Remove this pin from our owned pins
+				Pins.Remove(SubPin);
+
+				if (!SubPin->SubPins.IsEmpty())
+				{
+					RemoveGraphSubPins(SubPin);
+				}
+			};
+			
+			Traverse(PinPair->InputPin);
+			Traverse(PinPair->OutputPin);
+		}
+
+		CachedPins.Remove(ModelSubPin);
+	}
+	InParentPin->SubPins.Reset();
+}
+
+bool UControlRigGraphNode::ModelPinsChanged(bool bForce)
+{
+	UpdatePinLists();
+
+	// check if any of the pins need to be added / removed
+	auto AddMissingPins = [this](const TArray<URigVMPin*>& InModelPins)
+	{
+		int32 PinsAdded = 0;
+		for(const URigVMPin* ModelPin : InModelPins)
+		{
+			if(FindGraphPinFromModelPin(ModelPin, true) == nullptr && 
+				FindGraphPinFromModelPin(ModelPin, false))
+			{
+				PinsAdded += ModelPinAdded_Internal(ModelPin) ? 1 : 0;
+			}
+		}
+		return PinsAdded;
+	};
+
+	auto RemoveObsoletePins = [this]()
+	{
+		TArray<URigVMPin*> PinsToRemove; 
+		for(const auto& Pair : CachedPins)
+		{
+			URigVMPin* ModelPin = Pair.Key;
+			if(!PinListForPin(ModelPin).Contains(ModelPin))
+			{
+				PinsToRemove.Add(ModelPin);
+			}
+		}
+
+		int32 PinsRemoved = 0;
+		for(const URigVMPin* ModelPin : PinsToRemove)
+		{
+			PinsRemoved += ModelPinRemoved_Internal(ModelPin) ? 1 : 0;
+		}
+		return PinsRemoved;
+	};
+
+	auto OrderPins = [this](const TArray<URigVMPin*>& InModelPins) -> int32
+	{
+		if(InModelPins.Num() < 2)
+		{
+			return 0;
+		}
+
+		TArray<UEdGraphPin*> OrderedGraphPins;
+		OrderedGraphPins.Reserve(InModelPins.Num() * 2);
+		int32 LastInputIndex = INDEX_NONE;
+		int32 LastOutputIndex = INDEX_NONE;
+		int32 PinsToReorder = 0;
+
+		for(const URigVMPin* ModelPin : InModelPins)
+		{
+			if(UEdGraphPin* InputGraphPin = FindGraphPinFromModelPin(ModelPin, true))
+			{
+				OrderedGraphPins.Add(InputGraphPin);
+				const int32 InputPinIndex = Pins.Find(InputGraphPin);
+				if(LastInputIndex > InputPinIndex)
+				{
+					PinsToReorder++;
+				}
+				LastInputIndex = InputPinIndex;
+			}
+			
+			if(UEdGraphPin* OutputGraphPin = FindGraphPinFromModelPin(ModelPin, false))
+			{
+				OrderedGraphPins.Add(OutputGraphPin);
+				const int32 OutputPinIndex = Pins.Find(OutputGraphPin);
+				if(LastOutputIndex > OutputPinIndex)
+				{
+					PinsToReorder++;
+				}
+				LastOutputIndex = OutputPinIndex;
+			}
+		}
+
+		if(PinsToReorder > 0)
+		{
+			for(UEdGraphPin* GraphPinInOrder : OrderedGraphPins)
+			{
+				Pins.Remove(GraphPinInOrder);
+			}
+
+			OrderedGraphPins.Append(Pins);
+			Swap(OrderedGraphPins, Pins);
+		}
+
+		return PinsToReorder;
+	};
+	
+	int32 PinsAdded = 0;
+	PinsAdded += AddMissingPins(ExecutePins);
+	PinsAdded += AddMissingPins(InputPins);
+	PinsAdded += AddMissingPins(InputOutputPins);
+	PinsAdded += AddMissingPins(OutputPins);
+
+	const int32 PinsRemoved = RemoveObsoletePins();
+
+	// working through it in the opposite order
+	// due to the use of Append within the lambda
+	int32 PinsReordered = 0;
+	PinsReordered += OrderPins(OutputPins);
+	PinsReordered += OrderPins(InputOutputPins);
+	PinsReordered += OrderPins(InputPins);
+	PinsReordered += OrderPins(ExecutePins);
+
+	const bool bResult = (PinsAdded > 0) || (PinsRemoved > 0) || (PinsReordered > 0); 
+	if(bResult || bForce)
+	{
+		OnNodePinsChanged().Broadcast();
+	}
+	return bResult;
+}
+
+bool UControlRigGraphNode::ModelPinAdded(const URigVMPin* InModelPin)
+{
+	if(InModelPin == nullptr)
+	{
+		return false;
+	}
+
+	UpdatePinLists();
+
+	bool bResult = ModelPinAdded_Internal(InModelPin);
+	if(bResult)
+	{
+		OnNodePinsChanged().Broadcast();
+	}
+	return bResult;
+}
+
+bool UControlRigGraphNode::ModelPinAdded_Internal(const URigVMPin* InModelPin)
+{
+	bool bResult = false;
+
+	UEdGraphPin* InputParentPin = FindGraphPinFromModelPin(InModelPin->GetParentPin(), true);
+	UEdGraphPin* OutputParentPin = FindGraphPinFromModelPin(InModelPin->GetParentPin(), false);
+		
+	if(InModelPin->GetDirection() == ERigVMPinDirection::Input ||
+		InModelPin->GetDirection() == ERigVMPinDirection::Visible ||
+		InModelPin->GetDirection() == ERigVMPinDirection::IO)
+	{
+		if(CreateGraphPinFromModelPin(InModelPin, EGPD_Input, InputParentPin))
+		{
+			bResult = true;
+		}
+	}
+	
+	if(InModelPin->GetDirection() == ERigVMPinDirection::Output ||
+		InModelPin->GetDirection() == ERigVMPinDirection::IO)
+	{
+		if(CreateGraphPinFromModelPin(InModelPin, EGPD_Output, OutputParentPin))
+		{
+			bResult = true;
+		}
+	}
+	
+	return bResult;
+}
+
+bool UControlRigGraphNode::ModelPinRemoved(const URigVMPin* InModelPin)
+{
+	if(InModelPin == nullptr)
+	{
+		return false;
+	}
+
+	UpdatePinLists();
+
+	const bool bResult = ModelPinRemoved_Internal(InModelPin);
+	if(bResult)
+	{
+		OnNodePinsChanged().Broadcast();
+	}
+	return bResult;
+}
+
+bool UControlRigGraphNode::ModelPinRemoved_Internal(const URigVMPin* InModelPin)
+{
+	const bool bResult = PinPathToModelPin.Remove(InModelPin->GetPinPath()) > 0;
+	CachedPins.Remove(InModelPin);
+	return bResult;
+}
+
 FLinearColor UControlRigGraphNode::GetNodeProfilingColor() const
 {
 #if WITH_EDITOR
@@ -514,7 +971,7 @@ FLinearColor UControlRigGraphNode::GetNodeProfilingColor() const
 	return FLinearColor::Black;
 }
 
-void UControlRigGraphNode::AllocateDefaultPins()
+void UControlRigGraphNode::UpdatePinLists()
 {
 	ExecutePins.Reset();
 	InputPins.Reset();
@@ -546,7 +1003,7 @@ void UControlRigGraphNode::AllocateDefaultPins()
 					InputOutputPins.Add(ModelPin);
 				}
 				else if (ModelPin->GetDirection() == ERigVMPinDirection::Input || 
-                    ModelPin->GetDirection() == ERigVMPinDirection::Visible)
+					ModelPin->GetDirection() == ERigVMPinDirection::Visible)
 				{
 					InputPins.Add(ModelPin);
 				}
@@ -557,181 +1014,46 @@ void UControlRigGraphNode::AllocateDefaultPins()
 			}
 		}
 	}
-
-	CreateExecutionPins();
-	CreateInputPins();
-	CreateInputOutputPins();
-	CreateOutputPins();
-
-	// Fill the variable list
-	ExternalVariables.Reset();
-	if (URigVMFunctionReferenceNode* FunctionReferenceNode = Cast<URigVMFunctionReferenceNode>(GetModelNode()))
-	{
-		if(FunctionReferenceNode->RequiresVariableRemapping())
-		{
-			TArray<FRigVMExternalVariable> CurrentExternalVariables = FunctionReferenceNode->GetContainedGraph()->GetExternalVariables();
-			for(const FRigVMExternalVariable& CurrentExternalVariable : CurrentExternalVariables)
-			{
-				ExternalVariables.Add(MakeShared<FRigVMExternalVariable>(CurrentExternalVariable));
-			}
-		}
-	}
-}
-
-void UControlRigGraphNode::CreateExecutionPins()
-{
-	DECLARE_SCOPE_HIERARCHICAL_COUNTER_FUNC()
-
-	const TArray<URigVMPin*> ModelPins = ExecutePins;
-	for (URigVMPin* ModelPin : ModelPins)
-	{
-		PinPair& Pair = CachedPins.FindOrAdd(ModelPin);
-		if (Pair.InputPin == nullptr)
-		{
-			Pair.InputPin = CreatePin(EGPD_Input, GetPinTypeForModelPin(ModelPin), FName(*ModelPin->GetPinPath()));
-			if (Pair.InputPin != nullptr)
-		{
-				ConfigurePin(Pair.InputPin, ModelPin, false, true);
-		}
-	}
-		if (Pair.OutputPin == nullptr)
-	{
-			Pair.OutputPin = CreatePin(EGPD_Output, GetPinTypeForModelPin(ModelPin), FName(*ModelPin->GetPinPath()));
-			if (Pair.OutputPin != nullptr)
-		{
-				ConfigurePin(Pair.OutputPin, ModelPin, false, true);
-		}
-	}
-		// note: no recursion for execution pins
-	}
 }
 
-void UControlRigGraphNode::CreateInputPins(URigVMPin* InParentPin)
-{
-	DECLARE_SCOPE_HIERARCHICAL_COUNTER_FUNC()
-
-	const TArray<URigVMPin*> ModelPins = InParentPin == nullptr ? InputPins : InParentPin->GetSubPins();
-	for (URigVMPin* ModelPin : ModelPins)
-	{
-		PinPair& Pair = CachedPins.FindOrAdd(ModelPin);
-		if (Pair.InputPin == nullptr)
-		{
-			Pair.InputPin = CreatePin(EGPD_Input, GetPinTypeForModelPin(ModelPin), FName(*ModelPin->GetPinPath()));
-			if (Pair.InputPin != nullptr)
-			{
-				ConfigurePin(Pair.InputPin, ModelPin, false, ModelPin->GetDirection() == ERigVMPinDirection::Input);
-
-				SetupPinDefaultsFromModel(Pair.InputPin);
-
-				if (InParentPin != nullptr)
-	{
-					PinPair& ParentPair = CachedPins.FindChecked(InParentPin);
-					ParentPair.InputPin->SubPins.Add(Pair.InputPin);
-					Pair.InputPin->ParentPin = ParentPair.InputPin;
-				}
-		}
-		}
-		CreateInputPins(ModelPin);
-	}
-}
-
-void UControlRigGraphNode::CreateInputOutputPins(URigVMPin* InParentPin, bool bHidden)
+void UControlRigGraphNode::AllocateDefaultPins()
 {
-	DECLARE_SCOPE_HIERARCHICAL_COUNTER_FUNC()
-
-	bool bIsContainer = false;
-	if(InParentPin)
-	{
-		bIsContainer = InParentPin->IsArray();
-	}
-
-	const TArray<URigVMPin*> ModelPins = InParentPin == nullptr ? InputOutputPins : InParentPin->GetSubPins();
-	for (URigVMPin* ModelPin : ModelPins)
-	{
-		PinPair& Pair = CachedPins.FindOrAdd(ModelPin);
-		if (Pair.InputPin == nullptr)
-		{
-			Pair.InputPin = CreatePin(EGPD_Input, GetPinTypeForModelPin(ModelPin), FName(*ModelPin->GetPinPath()));
-			if (Pair.InputPin != nullptr)
-		{
-				ConfigurePin(Pair.InputPin, ModelPin, bHidden, ModelPin->GetDirection() == ERigVMPinDirection::IO);
-
-				SetupPinDefaultsFromModel(Pair.InputPin);
+	UpdatePinLists();
 
-				if (InParentPin != nullptr)
-		{
-					PinPair& ParentPair = CachedPins.FindChecked(InParentPin);
-					ParentPair.InputPin->SubPins.Add(Pair.InputPin);
-					Pair.InputPin->ParentPin = ParentPair.InputPin;
-		}
-	}
-	}
-		if (Pair.OutputPin == nullptr && !bIsContainer)
-	{
-			Pair.OutputPin = CreatePin(EGPD_Output, GetPinTypeForModelPin(ModelPin), FName(*ModelPin->GetPinPath()));
-			if (Pair.OutputPin != nullptr)
-		{
-				ConfigurePin(Pair.OutputPin, ModelPin, bHidden, ModelPin->GetDirection() == ERigVMPinDirection::IO);
-
-				if (InParentPin != nullptr)
-		{
-					PinPair& ParentPair = CachedPins.FindChecked(InParentPin);
-					ParentPair.OutputPin->SubPins.Add(Pair.OutputPin);
-					Pair.OutputPin->ParentPin = ParentPair.OutputPin;
-				}
-		}
-	}
+	CachedPins.Reset();
+	PinPathToModelPin.Reset();
 
-		// don't recurse on knot / compact reroute nodes
-		if(URigVMRerouteNode* RerouteNode = Cast<URigVMRerouteNode>(GetModelNode()))
+	auto CreateGraphPins = [this](const TArray<URigVMPin*>& InModelPins)
 	{
-			if (!RerouteNode->GetShowsAsFullNode())
+		for (const URigVMPin* ModelPin : InModelPins)
 		{
-				bHidden = true;
+			ModelPinAdded_Internal(ModelPin);
 		}
-	}
+	};
 
-		if(bIsContainer)
-		{
-			CreateInputPins(ModelPin);
-		}
-		else
-		{
-			CreateInputOutputPins(ModelPin, bHidden);
-		}
-	}
-}
+	CreateGraphPins(ExecutePins);
+	CreateGraphPins(InputPins);
+	CreateGraphPins(InputOutputPins);
+	CreateGraphPins(OutputPins);
 
-void UControlRigGraphNode::CreateOutputPins(URigVMPin* InParentPin)
-{
-	DECLARE_SCOPE_HIERARCHICAL_COUNTER_FUNC()
-
-	const TArray<URigVMPin*> ModelPins = InParentPin == nullptr ? OutputPins : InParentPin->GetSubPins();
-	for (URigVMPin* ModelPin : ModelPins)
+	// Fill the variable list
+	ExternalVariables.Reset();
+	if (URigVMFunctionReferenceNode* FunctionReferenceNode = Cast<URigVMFunctionReferenceNode>(GetModelNode()))
 	{
-		PinPair& Pair = CachedPins.FindOrAdd(ModelPin);
-		if (Pair.OutputPin == nullptr)
+		if(FunctionReferenceNode->RequiresVariableRemapping())
 		{
-			Pair.OutputPin = CreatePin(EGPD_Output, GetPinTypeForModelPin(ModelPin), FName(*ModelPin->GetPinPath()));
-			if (Pair.OutputPin != nullptr)
+			TArray<FRigVMExternalVariable> CurrentExternalVariables = FunctionReferenceNode->GetContainedGraph()->GetExternalVariables();
+			for(const FRigVMExternalVariable& CurrentExternalVariable : CurrentExternalVariables)
 			{
-				ConfigurePin(Pair.OutputPin, ModelPin, false,  ModelPin->GetDirection() == ERigVMPinDirection::Output);
-
-				if (InParentPin != nullptr)
-	{
-					PinPair& ParentPair = CachedPins.FindChecked(InParentPin);
-					ParentPair.OutputPin->SubPins.Add(Pair.OutputPin);
-					Pair.OutputPin->ParentPin = ParentPair.OutputPin;
-				}
+				ExternalVariables.Add(MakeShared<FRigVMExternalVariable>(CurrentExternalVariable));
 			}
 		}
-		CreateOutputPins(ModelPin);
 	}
 }
 
 UClass* UControlRigGraphNode::GetControlRigGeneratedClass() const
 {
-	UControlRigBlueprint* Blueprint = GetTypedOuter<UControlRigBlueprint>();
+	const UControlRigBlueprint* Blueprint = GetTypedOuter<UControlRigBlueprint>();
 	if(Blueprint)
 	{
 		if (Blueprint->GeneratedClass)
@@ -746,7 +1068,7 @@ UClass* UControlRigGraphNode::GetControlRigGeneratedClass() const
 
 UClass* UControlRigGraphNode::GetControlRigSkeletonGeneratedClass() const
 {
-	UControlRigBlueprint* Blueprint = GetTypedOuter<UControlRigBlueprint>();
+	const UControlRigBlueprint* Blueprint = GetTypedOuter<UControlRigBlueprint>();
 	if(Blueprint)
 	{
 		if (Blueprint->SkeletonGeneratedClass)
@@ -960,7 +1282,7 @@ void UControlRigGraphNode::CopyPinDefaultsToModel(UEdGraphPin* Pin, bool bUndo,
 	}
 
 	const UEdGraphSchema_K2* K2Schema = GetDefault<UEdGraphSchema_K2>();
-	if (URigVMPin* ModelPin = GetModelPinFromPinPath(Pin->GetName()))
+	if (URigVMPin* ModelPin = FindModelPinFromGraphPin(Pin))
 	{
 		if (ModelPin->GetSubPins().Num() > 0)
 		{
@@ -1077,7 +1399,7 @@ FName UControlRigGraphNode::GetModelNodeName() const
 
 URigVMPin* UControlRigGraphNode::GetModelPinFromPinPath(const FString& InPinPath) const
 {
-	if (TWeakObjectPtr<URigVMPin> const* CachedModelPinPtr = CachedModelPins.Find(InPinPath))
+	if (TWeakObjectPtr<URigVMPin> const* CachedModelPinPtr = PinPathToModelPin.Find(InPinPath))
 	{
 		if(CachedModelPinPtr->IsValid())
 		{
@@ -1089,14 +1411,14 @@ URigVMPin* UControlRigGraphNode::GetModelPinFromPinPath(const FString& InPinPath
 		}
 	}
 
-	if (URigVMNode* ModelNode = GetModelNode())
+	if (const URigVMNode* ModelNode = GetModelNode())
 	{
-		FString PinPath = InPinPath.RightChop(ModelNode->GetNodePath().Len() + 1);
+		const FString PinPath = InPinPath.RightChop(ModelNode->GetNodePath().Len() + 1);
 		URigVMPin* ModelPin = ModelNode->FindPin(PinPath);
 		if (ModelPin)
 		{
 			UControlRigGraphNode* MutableThis = (UControlRigGraphNode*)this;
-			MutableThis->CachedModelPins.FindOrAdd(InPinPath) = ModelPin;
+			MutableThis->PinPathToModelPin.FindOrAdd(InPinPath) = ModelPin;
 		}
 		return ModelPin;
 	}
@@ -1114,30 +1436,35 @@ void UControlRigGraphNode::HandleAddAggregateElement(const FString& InNodePath)
 	}	
 }
 
-void UControlRigGraphNode::SetupPinDefaultsFromModel(UEdGraphPin* Pin)
+void UControlRigGraphNode::SetupPinDefaultsFromModel(UEdGraphPin* Pin, const URigVMPin* InModelPin)
 {
 	DECLARE_SCOPE_HIERARCHICAL_COUNTER_FUNC()
 
 	if (Pin->Direction != EGPD_Input)
 	{
 		return;
-		}
+	}
+
+	if(InModelPin == nullptr)
+	{
+		InModelPin = FindModelPinFromGraphPin(Pin);
+	}
 
 	const UEdGraphSchema_K2* K2Schema = GetDefault<UEdGraphSchema_K2>();
-	if (URigVMPin* ModelPin = GetModelPinFromPinPath(Pin->GetName()))
+	if (InModelPin)
 	{
-		if (ModelPin->GetSubPins().Num() > 0)
+		if (InModelPin->GetSubPins().Num() > 0)
 		{
 			return;
-			}
+		}
 
-		FString DefaultValueString = ModelPin->GetDefaultValue();
-		if (DefaultValueString.IsEmpty() && ModelPin->GetCPPType() == TEXT("FName"))
-	{
+		FString DefaultValueString = InModelPin->GetDefaultValue();
+		if (DefaultValueString.IsEmpty() && InModelPin->GetCPPType() == TEXT("FName"))
+		{
 			DefaultValueString = FName(NAME_None).ToString();
+		}
+		K2Schema->GetPinDefaultValuesFromString(Pin->PinType, Pin->GetOwningNodeUnchecked(), DefaultValueString, Pin->DefaultValue, Pin->DefaultObject, Pin->DefaultTextValue);
 	}
-						K2Schema->GetPinDefaultValuesFromString(Pin->PinType, Pin->GetOwningNodeUnchecked(), DefaultValueString, Pin->DefaultValue, Pin->DefaultObject, Pin->DefaultTextValue);
-					}
 }
 
 FText UControlRigGraphNode::GetTooltipText() const
@@ -1153,7 +1480,7 @@ void UControlRigGraphNode::InvalidateNodeTitle() const
 {
 	NodeTitle = FText();
 	FullNodeTitle = FText();
-	NodeTitleDirtied.ExecuteIfBound();
+	NodeTitleDirtied.Broadcast();
 }
 
 bool UControlRigGraphNode::CanCreateUnderSpecifiedSchema(const UEdGraphSchema* InSchema) const
@@ -1228,15 +1555,30 @@ bool UControlRigGraphNode::ShouldDrawNodeAsControlPointOnly(int32& OutInputPinIn
 	return false;
 }
 
-FEdGraphPinType UControlRigGraphNode::GetPinTypeForModelPin(URigVMPin* InModelPin)
+FEdGraphPinType UControlRigGraphNode::GetPinTypeForModelPin(const URigVMPin* InModelPin)
 {
 	FEdGraphPinType PinType = RigVMTypeUtils::PinTypeFromCPPType(*InModelPin->GetCPPType(), InModelPin->GetCPPTypeObject());
 	PinType.bIsConst = InModelPin->IsDefinedAsConstant();
 	return PinType;
 }
 
-void UControlRigGraphNode::ConfigurePin(UEdGraphPin* EdGraphPin, URigVMPin* ModelPin, bool bHidden, bool bConnectable)
+void UControlRigGraphNode::ConfigurePin(UEdGraphPin* EdGraphPin, const URigVMPin* ModelPin)
 {
+	const bool bConnectable =
+		ModelPin->GetDirection() == ERigVMPinDirection::Input ||
+		ModelPin->GetDirection() == ERigVMPinDirection::Output || 
+		ModelPin->GetDirection() == ERigVMPinDirection::IO;
+
+	// hide sub-pins on a compacted (knot) reroute node
+	bool bHidden = ModelPin->GetDirection() == ERigVMPinDirection::Hidden;
+	if(!bHidden && !ModelPin->IsRootPin())
+	{
+		if(const URigVMRerouteNode* RerouteNode = Cast<URigVMRerouteNode>(ModelPin->GetNode()))
+		{
+			bHidden = !RerouteNode->GetShowsAsFullNode();
+		}
+	}
+
 	EdGraphPin->bHidden = bHidden;
 	EdGraphPin->PinFriendlyName = FText::FromName(ModelPin->GetDisplayName());
 	EdGraphPin->bNotConnectable = !bConnectable;
@@ -1244,5 +1586,32 @@ void UControlRigGraphNode::ConfigurePin(UEdGraphPin* EdGraphPin, URigVMPin* Mode
 	EdGraphPin->bDisplayAsMutableRef = ModelPin->IsWildCard();
 }
 
+TArray<URigVMPin*>& UControlRigGraphNode::PinListForPin(const URigVMPin* InModelPin)
+{
+	if(IsValid(InModelPin))
+	{
+		if(InModelPin->IsExecuteContext())
+		{
+			return ExecutePins;
+		}
+		if(InModelPin->GetDirection() == ERigVMPinDirection::Input || InModelPin->GetDirection() == ERigVMPinDirection::Visible)
+		{
+			return InputPins;
+		}
+		if(InModelPin->GetDirection() == ERigVMPinDirection::IO)
+		{
+			return InputOutputPins;
+		}
+		if(InModelPin->GetDirection() == ERigVMPinDirection::Output)
+		{
+			return OutputPins;
+		}
+
+		checkNoEntry();
+	}
+	static TArray<URigVMPin*> EmptyList;
+	return EmptyList;
+}
+
 #undef LOCTEXT_NAMESPACE
 
diff --git a/Engine/Plugins/Animation/ControlRig/Source/ControlRigDeveloper/Public/Graph/ControlRigGraph.h b/Engine/Plugins/Animation/ControlRig/Source/ControlRigDeveloper/Public/Graph/ControlRigGraph.h
--- a/Engine/Plugins/Animation/ControlRig/Source/ControlRigDeveloper/Public/Graph/ControlRigGraph.h
+++ b/Engine/Plugins/Animation/ControlRig/Source/ControlRigDeveloper/Public/Graph/ControlRigGraph.h
@@ -86,6 +86,7 @@ public:
 	const UControlRigGraph* GetRootGraph() const;
 
 	void HandleModifiedEvent(ERigVMGraphNotifType InNotifType, URigVMGraph* InGraph, UObject* InSubject);
+	void ConsumeQueuedNotifications();
 
 	int32 GetInstructionIndex(const UControlRigGraphNode* InNode, bool bAsInput);
 
diff --git a/Engine/Plugins/Animation/ControlRig/Source/ControlRigDeveloper/Public/Graph/ControlRigGraphNode.h b/Engine/Plugins/Animation/ControlRig/Source/ControlRigDeveloper/Public/Graph/ControlRigGraphNode.h
--- a/Engine/Plugins/Animation/ControlRig/Source/ControlRigDeveloper/Public/Graph/ControlRigGraphNode.h
+++ b/Engine/Plugins/Animation/ControlRig/Source/ControlRigDeveloper/Public/Graph/ControlRigGraphNode.h
@@ -44,7 +44,7 @@ private:
 	TWeakObjectPtr<URigVMNode> CachedModelNode;
 
 	UPROPERTY(transient)
-	TMap<FString, TWeakObjectPtr<URigVMPin>> CachedModelPins;
+	TMap<FString, TWeakObjectPtr<URigVMPin>> PinPathToModelPin;
 
 	/** The property we represent. For template nodes this represents the struct/property type name. */
 	UPROPERTY()
@@ -76,7 +76,21 @@ private:
 
 public:
 
-	DECLARE_DELEGATE(FNodeTitleDirtied);
+	DECLARE_MULTICAST_DELEGATE(FNodeTitleDirtied);
+	DECLARE_MULTICAST_DELEGATE(FNodePinsChanged);
+	DECLARE_MULTICAST_DELEGATE(FNodePinExpansionChanged);
+
+	struct FPinPair
+	{
+		FPinPair()
+			: InputPin(nullptr)
+			, OutputPin(nullptr)
+		{}
+		UEdGraphPin* InputPin;
+		UEdGraphPin* OutputPin;
+
+		bool IsValid() const { return InputPin != nullptr || OutputPin != nullptr; }
+	};
 
 	UControlRigGraphNode();
 
@@ -151,27 +165,52 @@ public:
 	/** Insert a new array element after the element referred to by the property path */
 	void HandleInsertArrayElement(FString InPinPath);
 
-	FNodeTitleDirtied& GetNodeTitleDirtied() { return NodeTitleDirtied; }
-
 	int32 GetInstructionIndex(bool bAsInput) const;
 
 	const FRigVMTemplate* GetTemplate() const;
 
 	void ClearErrorInfo();
 
+	URigVMPin* FindModelPinFromGraphPin(const UEdGraphPin* InGraphPin) const;
+	UEdGraphPin* FindGraphPinFromModelPin(const URigVMPin* InModelPin, bool bAsInput) const;
+
+	/// Synchronize the stored name/value/type on the graph pin with the value stored on the node. 
+	/// If the pin has sub-pins, the value update is done recursively.
+	void SynchronizeGraphPinNameWithModelPin(const URigVMPin* InModelPin, bool bNotify = true);
+	void SynchronizeGraphPinValueWithModelPin(const URigVMPin* InModelPin);
+	void SynchronizeGraphPinTypeWithModelPin(const URigVMPin* InModelPin);
+	void SynchronizeGraphPinExpansionWithModelPin(const URigVMPin* InModelPin);
+	
+	void SyncGraphNodeTitleWithModelNodeTitle();
+	void SyncGraphNodeNameWithModelNodeName(const URigVMNode* InModelNode);
+
+	FNodeTitleDirtied& OnNodeTitleDirtied() { return NodeTitleDirtied; }
+	FNodePinsChanged& OnNodePinsChanged() { return NodePinsChanged; }
+	FNodePinExpansionChanged& OnNodePinExpansionChanged() { return NodePinExpansionChanged; }
+
+	/** Called when there's a drastic change in the pins */
+	bool ModelPinsChanged(bool bForce = false);
+
+	/** Called when a model pin is added after the node creation */
+	bool ModelPinAdded(const URigVMPin* InModelPin);
+
+	/** Called when a model pin is being removed */
+	bool ModelPinRemoved(const URigVMPin* InModelPin);
+
 protected:
 
 	FLinearColor GetNodeProfilingColor() const;
 	FLinearColor GetNodeOpacityColor() const;
 
 	/** Helper function for AllocateDefaultPins */
-	void CreateExecutionPins();
-	void CreateInputPins(URigVMPin* InParentPin = nullptr);
-	void CreateInputOutputPins(URigVMPin* InParentPin = nullptr, bool bHidden = false);
-	void CreateOutputPins(URigVMPin* InParentPin = nullptr);
+	void UpdatePinLists();
+	bool CreateGraphPinFromModelPin(const URigVMPin* InModelPin, EEdGraphPinDirection InDirection,  UEdGraphPin* InParentPin = nullptr);
+	void RemoveGraphSubPins(UEdGraphPin *InParentPin, const TArray<UEdGraphPin*>& InPinsToKeep = TArray<UEdGraphPin*>());
+	bool ModelPinAdded_Internal(const URigVMPin* InModelPin);
+	bool ModelPinRemoved_Internal(const URigVMPin* InModelPin);
 
 	/** Copies default values from underlying properties into pin defaults, for editing */
-	void SetupPinDefaultsFromModel(UEdGraphPin* Pin);
+	void SetupPinDefaultsFromModel(UEdGraphPin* Pin, const URigVMPin* InModelPin = nullptr);
 
 	/** Recreate pins when we reconstruct this node */
 	virtual void ReallocatePinsDuringReconstruction(const TArray<UEdGraphPin*>& OldPins);
@@ -194,14 +233,15 @@ protected:
 	UClass* GetControlRigGeneratedClass() const;
 	UClass* GetControlRigSkeletonGeneratedClass() const;
 
-	static FEdGraphPinType GetPinTypeForModelPin(URigVMPin* InModelPin);
+	static FEdGraphPinType GetPinTypeForModelPin(const URigVMPin* InModelPin);
 
 private:
 
 	FORCEINLINE int32 GetNodeTopologyVersion() const { return NodeTopologyVersion; }
 	int32 NodeTopologyVersion;
 
-	void ConfigurePin(UEdGraphPin* EdGraphPin, URigVMPin* ModelPin, bool bHidden, bool bConnectable);
+	static void ConfigurePin(UEdGraphPin* EdGraphPin, const URigVMPin* ModelPin);
+	TArray<URigVMPin*>& PinListForPin(const URigVMPin* InModelPin);
 
 	FLinearColor CachedTitleColor;
 	FLinearColor CachedNodeColor;
@@ -216,18 +256,11 @@ private:
 	TArray<URigVMPin*> OutputPins;
 	TArray<TSharedPtr<FRigVMExternalVariable>> ExternalVariables;
 	
-	struct PinPair
-	{
-		PinPair()
-			: InputPin(nullptr)
-			, OutputPin(nullptr)
-		{}
-		UEdGraphPin* InputPin;
-		UEdGraphPin* OutputPin;
-	};
-	TMap<URigVMPin*, PinPair> CachedPins;
+	TMap<URigVMPin*, FPinPair> CachedPins;
 
 	FNodeTitleDirtied NodeTitleDirtied;
+	FNodePinsChanged NodePinsChanged;
+	FNodePinExpansionChanged NodePinExpansionChanged;
 
 	mutable const FRigVMTemplate* CachedTemplate;
 
diff --git a/Engine/Plugins/Animation/ControlRig/Source/ControlRigEditor/Private/Graph/SControlRigGraphNode.cpp b/Engine/Plugins/Animation/ControlRig/Source/ControlRigEditor/Private/Graph/SControlRigGraphNode.cpp
--- a/Engine/Plugins/Animation/ControlRig/Source/ControlRigEditor/Private/Graph/SControlRigGraphNode.cpp
+++ b/Engine/Plugins/Animation/ControlRig/Source/ControlRigEditor/Private/Graph/SControlRigGraphNode.cpp
@@ -47,11 +47,6 @@
 
 void SControlRigGraphNode::Construct( const FArguments& InArgs )
 {
-	static const float PinWidgetSidePadding = 6.f;
-	static const float EmptySidePadding = 60.f;
-	static const float TopPadding = 2.f; 
-	static const float MaxHeight = 30.f;
-
 	if (CachedImg_CR_Pin_Connected == nullptr)
 	{
 		static const FName NAME_CR_Pin_Connected("ControlRig.Bug.Solid");
@@ -73,10 +68,6 @@
 
 	Blueprint = Cast<UControlRigBlueprint>(FBlueprintEditorUtils::FindBlueprintForNode(this->GraphNode));
 
-	// Re-cache variable info here (unit structure could have changed since last reconstruction, e.g. array add/remove)
-	// and also create missing pins if it hasn't created yet
-	EdGraphNode->AllocateDefaultPins();
-	
 	NodeErrorType = int32(EMessageSeverity::Info) + 1;
 	this->UpdateGraphNode();
 
@@ -85,448 +76,8 @@
 	URigVMController* Controller = EdGraphNode->GetController();
 	Controller->OnModified().AddSP(this, &SControlRigGraphNode::HandleModifiedEvent);
 
-	TMap<UEdGraphPin*, int32> EdGraphPinToInputPin;
-	for(int32 InputPinIndex = 0; InputPinIndex < InputPins.Num(); InputPinIndex++)
-	{
-		EdGraphPinToInputPin.Add(InputPins[InputPinIndex]->GetPinObj(), InputPinIndex);
-	}
-	TMap<UEdGraphPin*, int32> EdGraphPinToOutputPin;
-	for(int32 OutputPinIndex = 0; OutputPinIndex < OutputPins.Num(); OutputPinIndex++)
-	{
-		EdGraphPinToOutputPin.Add(OutputPins[OutputPinIndex]->GetPinObj(), OutputPinIndex);
-	}
+	UpdatePinTreeView();
 
-	TArray<URigVMPin*> RootModelPins = ModelNode->GetPins();
-	// orphaned pins are appended to the end of pin list on each side of the node
-	RootModelPins.Append(ModelNode->GetOrphanedPins());
-	TArray<URigVMPin*> ModelPins;
-
-	// sort model pins
-	// a) execute IOs, b) IO pins, c) input / visible pins, d) output pins
-	struct Local
-	{
-		static void VisitPinRecursively(URigVMPin* InPin, TArray<URigVMPin*>& OutPins)
-		{
-			OutPins.Add(InPin);
-
-			if (InPin->GetCPPType() == TEXT("FRotator"))
-			{
-				TArray<URigVMPin*> SubPins = InPin->GetSubPins();
-				if (SubPins.Num() == 3)
-				{
-					OutPins.Add(SubPins[2]);	
-					OutPins.Add(SubPins[0]);	
-					OutPins.Add(SubPins[1]);
-				}	
-			}
-			else
-			{				
-				for (URigVMPin* SubPin : InPin->GetSubPins())
-				{
-					VisitPinRecursively(SubPin, OutPins);
-				}
-			}
-		}
-	};
-	for(int32 SortPhase = 0; SortPhase < 4; SortPhase++)
-	{
-		for(URigVMPin* RootPin : RootModelPins)
-		{
-			switch (SortPhase)
-			{
-				case 0: // execute IO pins
-				{
-					if(RootPin->IsExecuteContext() && RootPin->GetDirection() == ERigVMPinDirection::IO)
-					{
-						Local::VisitPinRecursively(RootPin, ModelPins);
-					}
-					break;
-				}
-				case 1: // IO pins
-				{
-					if(!RootPin->IsExecuteContext() && RootPin->GetDirection() == ERigVMPinDirection::IO)
-					{
-						Local::VisitPinRecursively(RootPin, ModelPins);
-					}
-					break;
-				}
-				case 2: // input / visible pins
-				{
-					if(RootPin->GetDirection() == ERigVMPinDirection::Input || RootPin->GetDirection() == ERigVMPinDirection::Visible)
-					{
-						Local::VisitPinRecursively(RootPin, ModelPins);
-					}
-					break;
-				}
-				case 3: // output pins
-				default:
-				{
-					if(RootPin->GetDirection() == ERigVMPinDirection::Output)
-					{
-						Local::VisitPinRecursively(RootPin, ModelPins);
-					}
-					break;
-				}
-			}
-		}
-	}
-
-	// add spacer widget at the start
-	LeftNodeBox->AddSlot()
-	.HAlign(HAlign_Fill)
-	.VAlign(VAlign_Center)
-	.AutoHeight()
-	[
-		SNew(SSpacer)
-		.Size(FVector2D(1.f, 2.f))
-	];
-
-	const UControlRigGraphSchema* RigSchema = Cast<UControlRigGraphSchema>(EdGraphNode->GetSchema());
-	
-	TMap<URigVMPin*, int32> ModelPinToInfoIndex;
-	for(URigVMPin* ModelPin : ModelPins)
-	{
-		FPinInfo PinInfo;
-		PinInfo.Index = PinInfos.Num();
-		PinInfo.ParentIndex = INDEX_NONE;
-		PinInfo.bHasChildren = (ModelPin->GetSubPins().Num() > 0);
-		PinInfo.bIsContainer = ModelPin->IsArray();
-		PinInfo.Depth = 0;
-		PinInfo.bExpanded = ModelPin->IsExpanded();
-		PinInfo.ModelPinPath = ModelPin->GetPinPath();
-		PinInfo.bAutoHeight = false;
-		
-		const bool bAskSchemaForEdition = RigSchema && ModelPin->IsStruct() && !ModelPin->IsBoundToVariable();
-		PinInfo.bHideInputWidget = (!ModelPin->IsBoundToVariable()) && PinInfo.bIsContainer;
-		if (!PinInfo.bHideInputWidget)
-		{
-			if (bAskSchemaForEdition && !PinInfo.bHasChildren)
-			{
-				const bool bIsStructEditable = RigSchema->IsStructEditable(ModelPin->GetScriptStruct()); 
-				PinInfo.bHideInputWidget = !bIsStructEditable;
-				PinInfo.bAutoHeight = bIsStructEditable;
-			}
-			else if(PinInfo.bHasChildren && !ModelPin->IsBoundToVariable())
-			{
-				PinInfo.bHideInputWidget = true;
-			}
-		}
-		
-		if(URigVMPin* ParentPin = ModelPin->GetParentPin())
-		{
-			const int32* ParentIndexPtr = ModelPinToInfoIndex.Find(ParentPin);
-			if(ParentIndexPtr == nullptr)
-			{
-				continue;
-			}
-			PinInfo.ParentIndex = *ParentIndexPtr;
-			PinInfo.Depth = PinInfos[PinInfo.ParentIndex].Depth + 1;
-		}
-
-		TAttribute<EVisibility> PinVisibilityAttribute = TAttribute<EVisibility>::CreateSP(this, &SControlRigGraphNode::GetPinVisibility, PinInfo.Index);
-
-		bool bPinWidgetForExpanderLeft = false;
-		TSharedPtr<SGraphPin> PinWidgetForExpander;
-
-		bool bPinInfoIsValid = false;
-		if(UEdGraphPin* OutputEdGraphPin = EdGraphNode->FindPin(ModelPin->GetPinPath(), EEdGraphPinDirection::EGPD_Output))
-		{
-			if(const int32* PinIndexPtr = EdGraphPinToOutputPin.Find(OutputEdGraphPin))
-			{
-				PinInfo.OutputPinWidget = OutputPins[*PinIndexPtr];
-				PinInfo.OutputPinWidget->SetVisibility(PinVisibilityAttribute);
-				PinWidgetForExpander = PinInfo.OutputPinWidget;
-				bPinWidgetForExpanderLeft = false;
-				bPinInfoIsValid = true;
-			}
-		}
-		if(UEdGraphPin* InputEdGraphPin = EdGraphNode->FindPin(ModelPin->GetPinPath(), EEdGraphPinDirection::EGPD_Input))
-		{
-			if(const int32* PinIndexPtr = EdGraphPinToInputPin.Find(InputEdGraphPin))
-			{
-				PinInfo.InputPinWidget = InputPins[*PinIndexPtr];
-				PinInfo.InputPinWidget->SetVisibility(PinVisibilityAttribute);
-				PinWidgetForExpander = PinInfo.InputPinWidget;
-				bPinWidgetForExpanderLeft = true;
-				bPinInfoIsValid = true;
-			}
-		}
-
-		if(!bPinInfoIsValid)
-		{
-			continue;
-		}
-
-		ModelPinToInfoIndex.Add(ModelPin, PinInfos.Add(PinInfo));
-		
-		// check if this pin has sub pins
-		TSharedPtr<SHorizontalBox> FullPinHorizontalRowWidget = PinWidgetForExpander->GetFullPinHorizontalRowWidget().Pin();
-		if(FullPinHorizontalRowWidget.IsValid())
-		{
-			// indent the pin by padding
-			const float DepthIndentation = 12.f * float(PinInfo.Depth + (PinInfo.bHasChildren ? 0 : 1));
-			const float LeftIndentation = bPinWidgetForExpanderLeft ? DepthIndentation : 0.f;
-			const float RightIndentation = bPinWidgetForExpanderLeft ? 0.f : DepthIndentation;
-			
-			if(PinInfo.bHasChildren)
-			{
-				// only inject the expander arrow for inputs on input / IO
-				// or for output pins
-				if(
-					(
-						(
-							(ModelPin->GetDirection() == ERigVMPinDirection::Input) ||
-							(ModelPin->GetDirection() == ERigVMPinDirection::IO)
-						) &&
-						bPinWidgetForExpanderLeft
-					) ||
-					(
-						(ModelPin->GetDirection() == ERigVMPinDirection::Output) &&
-						(!bPinWidgetForExpanderLeft)
-					)
-				)
-				{
-					// Add the expander arrow
-					FullPinHorizontalRowWidget->InsertSlot(bPinWidgetForExpanderLeft ? 1 : FullPinHorizontalRowWidget->GetChildren()->Num() - 1)
-					.Padding(LeftIndentation, 0.f, RightIndentation, 0.f)
-					.AutoWidth()
-					[
-						SNew(SButton)
-						.ButtonStyle(FAppStyle::Get(), TEXT("SimpleButton"))
-						.ContentPadding(0)
-						.VAlign(VAlign_Center)
-						.HAlign(HAlign_Center)
-						.ClickMethod( EButtonClickMethod::MouseDown )
-						.OnClicked(this, &SControlRigGraphNode::OnExpanderArrowClicked, PinInfo.Index)
-						.ToolTipText(LOCTEXT("ExpandSubPin", "Expand Pin"))
-						[
-							SNew(SImage)
-							.Image(this, &SControlRigGraphNode::GetExpanderImage, PinInfo.Index, bPinWidgetForExpanderLeft, false)
-							.ColorAndOpacity(FSlateColor::UseForeground())
-						]
-					];
-				}
-			}
-			else
-			{
-				const int32 SlotToAdjustIndex = bPinWidgetForExpanderLeft ? 0 : FullPinHorizontalRowWidget->NumSlots() - 1;
-				SHorizontalBox::FSlot& Slot = FullPinHorizontalRowWidget->GetSlot(SlotToAdjustIndex);
-
-				FMargin Padding = Slot.GetPadding();
-				Padding = FMargin(RightIndentation + Padding.Left, Padding.Top, LeftIndentation + Padding.Right, Padding.Bottom);
-				Slot.SetPadding(Padding);
-			}
-		}
-	}
-
-	auto AddArrayPlusButtonLambda = [this](URigVMPin* InModelPin, TSharedPtr<SHorizontalBox> InSlotLayout, const float InEmptySidePadding)
-	{
-		// add array plus button
-		InSlotLayout->AddSlot()
-		.AutoWidth()
-		.HAlign(HAlign_Left)
-		.Padding(PinWidgetSidePadding, TopPadding, InEmptySidePadding, 0.f)
-		[
-			SNew(SButton)
-			.ContentPadding(0.0f)
-			.ButtonStyle(FAppStyle::Get(), "NoBorder")
-			.OnClicked(this, &SControlRigGraphNode::HandleAddArrayElement, InModelPin->GetPinPath())
-			.IsEnabled(this, &SGraphNode::IsNodeEditable)
-			.Cursor(EMouseCursor::Default)
-			.Visibility(this, &SControlRigGraphNode::GetArrayPlusButtonVisibility, InModelPin)
-			.ToolTipText(LOCTEXT("AddArrayElement", "Add Array Element"))
-			[
-				SNew(SHorizontalBox)
-				+SHorizontalBox::Slot()
-				.AutoWidth()
-				.VAlign(VAlign_Center)
-				[
-					SNew(SImage)
-					.Image(FAppStyle::GetBrush(TEXT("Icons.PlusCircle")))
-				]
-			]
-		];
-	};
-
-	for(const FPinInfo& PinInfo : PinInfos)
-	{
-		if(PinInfo.InputPinWidget.IsValid())
-		{
-			if(PinInfo.bHideInputWidget)
-			{
-				if(PinInfo.InputPinWidget->GetValueWidget() != SNullWidget::NullWidget)
-				{
-					PinInfo.InputPinWidget->GetValueWidget()->SetVisibility(EVisibility::Collapsed);
-				}
-			}
-				
-			// input pins
-			if(!PinInfo.OutputPinWidget.IsValid())
-			{
-				TSharedPtr<SHorizontalBox> SlotLayout;
-				SHorizontalBox::FSlot* FirstSlot = nullptr;
-
-				const float MyEmptySidePadding = PinInfo.bHideInputWidget ? EmptySidePadding : 0.f; 
-				
-				LeftNodeBox->AddSlot()
-                .HAlign(HAlign_Fill)
-				.VAlign(VAlign_Center)
-				.AutoHeight()
-				.MaxHeight(PinInfo.bAutoHeight ? TAttribute<float>() : MaxHeight)
-                [
-                    SAssignNew(SlotLayout, SHorizontalBox)
-                    .Visibility(this, &SControlRigGraphNode::GetPinVisibility, PinInfo.Index)
-		            
-                    +SHorizontalBox::Slot()
-                    .Expose(FirstSlot)
-                    .FillWidth(1.f)
-                    .HAlign(HAlign_Left)
-                    .Padding(PinWidgetSidePadding, TopPadding, PinInfo.bIsContainer ? 0.f : MyEmptySidePadding, 0.f)
-                    [
-                        PinInfo.InputPinWidget.ToSharedRef()
-                    ]
-                ];
-
-				if(PinInfo.bIsContainer)
-				{
-					URigVMPin* ModelPin = ModelNode->GetGraph()->FindPin(PinInfo.ModelPinPath);
-					if(ModelPin)
-					{
-						// make sure to minimize the width of the label
-						FirstSlot->SetAutoWidth();
-						AddArrayPlusButtonLambda(ModelPin, SlotLayout, MyEmptySidePadding);
-					}
-				}
-			}
-			// io pins
-			else
-			{
-				TSharedPtr<SHorizontalBox> SlotLayout;
-				SHorizontalBox::FSlot* FirstSlot = nullptr;
-
-				PinInfo.OutputPinWidget->SetShowLabel(false);
-			
-				LeftNodeBox->AddSlot()
-                .HAlign(HAlign_Fill)
-				.VAlign(VAlign_Center)
-                .AutoHeight()
-				.MaxHeight(PinInfo.bAutoHeight ? TAttribute<float>() : MaxHeight)
-                [
-                    SAssignNew(SlotLayout, SHorizontalBox)
-                    .Visibility(this, &SControlRigGraphNode::GetPinVisibility, PinInfo.Index)
-
-                    +SHorizontalBox::Slot()
-					.Expose(FirstSlot)
-                    .FillWidth(1.f)
-                    .HAlign(HAlign_Left)
-                    .VAlign(VAlign_Center)
-                    .Padding(PinWidgetSidePadding, TopPadding, 0.f, 0.f)
-                    [
-                        PinInfo.InputPinWidget.ToSharedRef()
-                    ]
-                ];
-
-				if(PinInfo.bIsContainer)
-				{
-					URigVMPin* ModelPin = ModelNode->GetGraph()->FindPin(PinInfo.ModelPinPath);
-					if(ModelPin)
-					{
-						// make sure to minimize the width of the label
-						FirstSlot->SetAutoWidth();
-						AddArrayPlusButtonLambda(ModelPin, SlotLayout, EmptySidePadding);
-					}
-				}
-
-				SlotLayout->AddSlot()
-				.FillWidth(1.f)
-				.HAlign(HAlign_Right)
-				.VAlign(VAlign_Center)
-				.Padding(0.f, TopPadding, PinWidgetSidePadding, 0.f)
-				[
-					PinInfo.OutputPinWidget.ToSharedRef()
-				];
-			}
-		}
-		// output pins
-		else if(PinInfo.OutputPinWidget.IsValid())
-		{
-			LeftNodeBox->AddSlot()
-            .HAlign(HAlign_Fill)
-			.VAlign(VAlign_Center)
-            .AutoHeight()
-			.MaxHeight(PinInfo.bAutoHeight ? TAttribute<float>() : MaxHeight)
-            [
-            	SNew(SHorizontalBox)
-	            .Visibility(this, &SControlRigGraphNode::GetPinVisibility, PinInfo.Index)
-	            
-				+SHorizontalBox::Slot()
-				.FillWidth(1.f)
-				.HAlign(HAlign_Right)
-				.VAlign(VAlign_Center)
-	            .Padding(EmptySidePadding, TopPadding, PinWidgetSidePadding, 0.f)
-				[
-	                PinInfo.OutputPinWidget.ToSharedRef()
-				]
-            ];
-		}
-	}
-
-	if(URigVMFunctionReferenceNode* FunctionReferenceNode = Cast<URigVMFunctionReferenceNode>(ModelNode))
-	{
-		TWeakObjectPtr<URigVMFunctionReferenceNode> WeakFunctionReferenceNode = FunctionReferenceNode;
-		TWeakObjectPtr<UControlRigBlueprint> WeakControlRigBlueprint = Blueprint;
-
-		// add the entries for the variable remapping
-		for(const TSharedPtr<FRigVMExternalVariable>& ExternalVariable : EdGraphNode->ExternalVariables)
-		{
-			LeftNodeBox->AddSlot()
-			.HAlign(HAlign_Fill)
-			.VAlign(VAlign_Center)
-			.AutoHeight()
-			.MaxHeight(MaxHeight)
-			[
-				SNew(SHorizontalBox)
-	            
-				+SHorizontalBox::Slot()
-				.AutoWidth()
-				.HAlign(HAlign_Left)
-				.VAlign(VAlign_Center)
-				.Padding(PinWidgetSidePadding, TopPadding, PinWidgetSidePadding, 0.f)
-				[
-					SNew(STextBlock)
-					.Text(FText::FromName(ExternalVariable->Name))
-					.TextStyle(FAppStyle::Get(), NAME_DefaultPinLabelStyle)
-					.ColorAndOpacity(this, &SControlRigGraphNode::GetVariableLabelTextColor, WeakFunctionReferenceNode, ExternalVariable->Name)
-					.ToolTipText(this, &SControlRigGraphNode::GetVariableLabelTooltipText, WeakControlRigBlueprint, ExternalVariable->Name)
-				]
-
-				+SHorizontalBox::Slot()
-				.AutoWidth()
-				.HAlign(HAlign_Left)
-				.VAlign(VAlign_Center)
-				.Padding(PinWidgetSidePadding, TopPadding, PinWidgetSidePadding, 0.f)
-				[
-					SNew(SControlRigVariableBinding)
-					.Blueprint(Blueprint.Get())
-					.FunctionReferenceNode(FunctionReferenceNode)
-					.InnerVariableName(ExternalVariable->Name)
-				]
-			];
-		}
-	}
-
-	CreateAggregateAddPinButton();
-	
-	// add spacer widget at the end
-	LeftNodeBox->AddSlot()
-	.HAlign(HAlign_Fill)
-	.VAlign(VAlign_Center)
-	.AutoHeight()
-	[
-		SNew(SSpacer)
-		.Size(FVector2D(1.f, 4.f))
-	];
-
 	const FSlateBrush* ImageBrush = FControlRigEditorStyle::Get().GetBrush(TEXT("ControlRig.Bug.Dot"));
 
 	VisualDebugIndicatorWidget =
@@ -554,7 +105,8 @@
 	.Visibility(EVisibility::Visible)
 	.ToolTipText(LOCTEXT("NodeDurationToolTip", "This number represents the duration in microseconds for a node.\nFor functions / collapse nodes it represents the accumulated time of contained nodes.\n\nYou can enable / disable the display of the number in the Class Settings\n(VM Runtime Settings -> Enable Profiling)"));
 
-	EdGraphNode->GetNodeTitleDirtied().BindSP(this, &SControlRigGraphNode::HandleNodeTitleDirtied);
+	EdGraphNode->OnNodeTitleDirtied().AddSP(this, &SControlRigGraphNode::HandleNodeTitleDirtied);
+	EdGraphNode->OnNodePinsChanged().AddSP(this, &SControlRigGraphNode::HandleNodePinsChanged);
 
 	LastHighDetailSize = FVector2D::ZeroVector;
 }
@@ -686,7 +238,7 @@
 		FString NodeName, PinPath;
 		if (URigVMPin::SplitPinPathAtStart(EdPinObj->GetName(), NodeName, PinPath))
 		{
-			if (URigVMPin* ModelPin = ModelNode->FindPin(PinPath))
+			if (const URigVMPin* ModelPin = ModelNode->FindPin(PinPath))
 			{
 				if (ModelPin->HasInjectedUnitNodes())
 				{
@@ -695,61 +247,71 @@
 				PinToAdd->SetToolTipText(ModelPin->GetToolTipText());
 
 				// If the pin belongs to a template node that does not own an argument for that pin, make it transparent
-				if (URigVMTemplateNode* TemplateNode = Cast<URigVMTemplateNode>(ModelPin->GetNode()))
+				if (const URigVMTemplateNode* TemplateNode = Cast<URigVMTemplateNode>(ModelPin->GetNode()))
 				{
 					if (const FRigVMTemplate* Template = TemplateNode->GetTemplate())
 					{
-						URigVMPin* RootPin = ModelPin->GetRootPin();
+						const URigVMPin* RootPin = ModelPin->GetRootPin();
+						FLinearColor PinColorAndOpacity = PinToAdd->GetColorAndOpacity();
 						if (Template->FindArgument(RootPin->GetFName()) == nullptr)
 						{
-							PinToAdd->SetColorAndOpacity(PinToAdd->GetColorAndOpacity() * FLinearColor(1.0f,1.0f,1.0f,0.2f));
+							PinColorAndOpacity.A = 0.2f;
 						}
+						else
+						{
+							PinColorAndOpacity.A = 1.0f;
+						}
+						PinToAdd->SetColorAndOpacity(PinColorAndOpacity);
 					}
 				}
 			}
 		}
 
-		// reformat the pin by
-		// 1. taking out the swrapbox widget
-		// 2. re-inserting all widgets from the label and value wrap box back in the horizontal box
-		TSharedPtr<SHorizontalBox> FullPinHorizontalRowWidget = PinToAdd->GetFullPinHorizontalRowWidget().Pin();
-		TSharedPtr<SWrapBox> LabelAndValueWidget = PinToAdd->GetLabelAndValue();
-		if(FullPinHorizontalRowWidget.IsValid() && LabelAndValueWidget.IsValid())
+		if(!PinsToKeep.Contains(EdPinObj))
 		{
-			int32 LabelAndValueWidgetIndex = INDEX_NONE;
-			for(int32 ChildIndex = 0; ChildIndex < FullPinHorizontalRowWidget->GetChildren()->Num(); ChildIndex++)
+			// reformat the pin by
+			// 1. taking out the swrapbox widget
+			// 2. re-inserting all widgets from the label and value wrap box back in the horizontal box
+			TSharedPtr<SHorizontalBox> FullPinHorizontalRowWidget = PinToAdd->GetFullPinHorizontalRowWidget().Pin();
+			TSharedPtr<SWrapBox> LabelAndValueWidget = PinToAdd->GetLabelAndValue();
+			if(FullPinHorizontalRowWidget.IsValid() && LabelAndValueWidget.IsValid())
 			{
-				TSharedRef<SWidget> ChildWidget = FullPinHorizontalRowWidget->GetChildren()->GetChildAt(ChildIndex);
-				if(ChildWidget == LabelAndValueWidget)
+				int32 LabelAndValueWidgetIndex = INDEX_NONE;
+				for(int32 ChildIndex = 0; ChildIndex < FullPinHorizontalRowWidget->GetChildren()->Num(); ChildIndex++)
 				{
-					LabelAndValueWidgetIndex = ChildIndex;
-					break;
+					TSharedRef<SWidget> ChildWidget = FullPinHorizontalRowWidget->GetChildren()->GetChildAt(ChildIndex);
+					if(ChildWidget == LabelAndValueWidget)
+					{
+						LabelAndValueWidgetIndex = ChildIndex;
+						break;
+					}
+				}
+				check(LabelAndValueWidgetIndex != INDEX_NONE);
+				
+				FullPinHorizontalRowWidget->RemoveSlot(LabelAndValueWidget.ToSharedRef());
+				
+				for(int32 ChildIndex = 0; ChildIndex < LabelAndValueWidget->GetChildren()->Num(); ChildIndex++)
+				{
+					TSharedRef<SWidget> ChildWidget = LabelAndValueWidget->GetChildren()->GetChildAt(ChildIndex);
+					if(ChildWidget != SNullWidget::NullWidget)
+					{
+						ChildWidget->AssignParentWidget(FullPinHorizontalRowWidget.ToSharedRef());
+						
+						FullPinHorizontalRowWidget->InsertSlot(LabelAndValueWidgetIndex + ChildIndex)
+						.HAlign(HAlign_Fill)
+						.VAlign(VAlign_Center)
+						.Padding(EdPinObj->Direction == EGPD_Input ? 0.f : 2.f, 0.f, EdPinObj->Direction == EGPD_Input ? 2.f : 0.f, 0.f)
+						.AutoWidth()
+						[
+							ChildWidget
+						];
+					}
 				}
 			}
-			check(LabelAndValueWidgetIndex != INDEX_NONE);
-			
-			FullPinHorizontalRowWidget->RemoveSlot(LabelAndValueWidget.ToSharedRef());
-			
-			for(int32 ChildIndex = 0; ChildIndex < LabelAndValueWidget->GetChildren()->Num(); ChildIndex++)
-			{
-				TSharedRef<SWidget> ChildWidget = LabelAndValueWidget->GetChildren()->GetChildAt(ChildIndex);
-				if(ChildWidget != SNullWidget::NullWidget)
-				{
-					ChildWidget->AssignParentWidget(FullPinHorizontalRowWidget.ToSharedRef());
-					
-					FullPinHorizontalRowWidget->InsertSlot(LabelAndValueWidgetIndex + ChildIndex)
-					.HAlign(HAlign_Fill)
-					.VAlign(VAlign_Center)
-					.Padding(EdPinObj->Direction == EGPD_Input ? 0.f : 2.f, 0.f, EdPinObj->Direction == EGPD_Input ? 2.f : 0.f, 0.f)
-					.AutoWidth()
-					[
-						ChildWidget
-					];
-				}
-			}
+
+			PinToAdd->SetOwner(SharedThis(this));
 		}
-
-		PinToAdd->SetOwner(SharedThis(this));
+		
 		if(EdPinObj->Direction == EGPD_Input)
 		{
 			InputPins.Add(PinToAdd);
@@ -761,6 +323,40 @@
 	}
 }
 
+void SControlRigGraphNode::CreateStandardPinWidget(UEdGraphPin* CurPin)
+{
+	bool bShowPin = true;
+	if(const UControlRigGraphNode* RigGraphNode = Cast<UControlRigGraphNode>(GraphNode))
+	{
+		if(const URigVMPin* ModelPin = RigGraphNode->FindModelPinFromGraphPin(CurPin))
+		{
+			bShowPin =
+				ModelPin->GetDirection() == ERigVMPinDirection::Visible ||
+				ModelPin->GetDirection() == ERigVMPinDirection::Input ||
+				ModelPin->GetDirection() == ERigVMPinDirection::Output ||
+				ModelPin->GetDirection() == ERigVMPinDirection::IO;
+		}
+	}
+	
+	if (bShowPin)
+	{
+		// Do we have this pin in our list of pins to keep?
+		TSharedPtr<SGraphPin> NewPin;
+		const TSharedRef<SGraphPin>* RecycledPin = PinsToKeep.Find(CurPin);
+		if (!RecycledPin)
+		{
+			NewPin = CreatePinWidget(CurPin);
+			check(NewPin.IsValid());
+		}
+		else
+		{
+			NewPin = *RecycledPin;
+		}
+
+		AddPin(NewPin.ToSharedRef());
+	}
+}
+
 const FSlateBrush * SControlRigGraphNode::GetNodeBodyBrush() const
 {
 	if(UControlRigGraphNode* RigNode = Cast<UControlRigGraphNode>(GetNodeObj()))
@@ -1240,6 +836,13 @@
 		GraphNode->NodeWidth = (int32)AllottedGeometry.Size.X;
 		GraphNode->NodeHeight = (int32)AllottedGeometry.Size.Y;
 		RefreshErrorInfo();
+
+		// These will be deleted on the next tick.
+		for (UEdGraphPin *PinToDelete: PinsToDelete)
+		{
+			PinToDelete->MarkAsGarbage();
+		}
+		PinsToDelete.Reset();
 	}
 
 	if((!UseLowDetailNodeContent()) && (LeftNodeBox != nullptr))
@@ -1256,6 +859,96 @@
 	}
 }
 
+void SControlRigGraphNode::HandleNodePinsChanged()
+{
+	UControlRigGraphNode* ControlRigGraphNode = Cast<UControlRigGraphNode>(GraphNode);
+	if(ControlRigGraphNode == nullptr)
+	{
+		return;
+	}
+
+	// Collect graph pins to delete. We do this here because this widget is the only entity
+	// that's aware of the lifetime requirements for the graph pins (SGraphPanel uses Slate 
+	// timers to trigger a delete, which makes deleting them from a non-widget setting). 
+	TSet<UEdGraphPin *> LocalPinsToDelete;
+	LocalPinsToDelete.Reserve(InputPins.Num() + OutputPins.Num());
+	
+	for (const TSharedRef<SGraphPin>& GraphPin: InputPins)
+	{
+		LocalPinsToDelete.Add(GraphPin->GetPinObj());
+	}
+	for (const TSharedRef<SGraphPin>& GraphPin: OutputPins)
+	{
+		LocalPinsToDelete.Add(GraphPin->GetPinObj());
+	}
+
+	check(PinsToKeep.IsEmpty());
+
+	for (const UEdGraphPin* LivePin: ControlRigGraphNode->Pins)
+	{
+		const FString PinPath = LivePin->GetName();
+		
+		const FPinInfo* PinInfoPtr = PinInfos.FindByPredicate([PinPath](const FPinInfo& PinInfo) -> bool
+		{
+			return PinInfo.ModelPinPath == PinPath;
+		});
+		
+		if (PinInfoPtr)
+		{
+			if (LivePin->Direction == EGPD_Input && PinInfoPtr->InputPinWidget.IsValid())
+			{
+				PinsToKeep.Add(LivePin, PinInfoPtr->InputPinWidget.ToSharedRef());
+			}
+			if (LivePin->Direction == EGPD_Output && PinInfoPtr->OutputPinWidget.IsValid())
+			{
+				PinsToKeep.Add(LivePin, PinInfoPtr->OutputPinWidget.ToSharedRef());
+			}
+		}
+		LocalPinsToDelete.Remove(LivePin);
+	}
+
+	for (const UEdGraphPin* DeletingPin: LocalPinsToDelete)
+	{
+		const FString PinPath = DeletingPin->GetName();
+		
+		const FPinInfo* PinInfoPtr = PinInfos.FindByPredicate([PinPath](const FPinInfo& PinInfo) -> bool
+		{
+			return PinInfo.ModelPinPath == PinPath;
+		});
+		
+		if (PinInfoPtr)
+		{
+			if (DeletingPin->Direction == EGPD_Input && PinInfoPtr->InputPinWidget.IsValid())
+			{
+				// Ensure that this pin widget can no longer depend on the soon-to-be-deleted
+				// graph pin.
+				PinInfoPtr->InputPinWidget->InvalidateGraphData();
+			}
+
+			if (DeletingPin->Direction == EGPD_Output && PinInfoPtr->OutputPinWidget.IsValid())
+			{
+				// Ensure that this pin widget can no longer depend on the soon-to-be-deleted
+				// graph pin.
+				PinInfoPtr->OutputPinWidget->InvalidateGraphData();
+			}
+		}
+	}
+	
+	PinsToDelete.Append(LocalPinsToDelete);
+
+	// Reconstruct the pin widgets. This could be done more surgically but will do for now.
+	InputPins.Reset();
+	OutputPins.Reset();
+	PinInfos.Reset();
+
+	CreatePinWidgets();
+
+	// Nix any pins left in this map. They're most likely hidden sub-pins.
+	PinsToKeep.Reset();
+
+	UpdatePinTreeView();
+}
+
 FText SControlRigGraphNode::GetInstructionCountText() const
 {
 	if(Blueprint.IsValid())
@@ -1499,4 +1192,492 @@
 	}
 }
 
+void SControlRigGraphNode::UpdatePinTreeView()
+{
+	static const float PinWidgetSidePadding = 6.f;
+	static const float EmptySidePadding = 60.f;
+	static const float TopPadding = 2.f; 
+	static const float MaxHeight = 30.f;
+
+	check(GraphNode);
+
+	// remove all existing content in the Left
+	LeftNodeBox->ClearChildren();
+
+	UControlRigGraphNode* RigGraphNode = Cast<UControlRigGraphNode>(GraphNode);
+
+	TMap<UEdGraphPin*, int32> EdGraphPinToInputPin;
+	for(int32 InputPinIndex = 0; InputPinIndex < InputPins.Num(); InputPinIndex++)
+	{
+		EdGraphPinToInputPin.Add(InputPins[InputPinIndex]->GetPinObj(), InputPinIndex);
+	}
+	TMap<UEdGraphPin*, int32> EdGraphPinToOutputPin;
+	for(int32 OutputPinIndex = 0; OutputPinIndex < OutputPins.Num(); OutputPinIndex++)
+	{
+		EdGraphPinToOutputPin.Add(OutputPins[OutputPinIndex]->GetPinObj(), OutputPinIndex);
+	}
+
+	TArray<URigVMPin*> RootModelPins = ModelNode->GetPins();
+	// orphaned pins are appended to the end of pin list on each side of the node
+	RootModelPins.Append(ModelNode->GetOrphanedPins());
+	TArray<URigVMPin*> ModelPins;
+
+	// sort model pins
+	// a) execute IOs, b) IO pins, c) input / visible pins, d) output pins
+	struct Local
+	{
+		static void VisitPinRecursively(URigVMPin* InPin, TArray<URigVMPin*>& OutPins)
+		{
+			OutPins.Add(InPin);
+
+			if (InPin->GetCPPType() == TEXT("FRotator"))
+			{
+				const TArray<URigVMPin*>& SubPins = InPin->GetSubPins();
+				if (SubPins.Num() == 3)
+				{
+					OutPins.Add(SubPins[2]);	
+					OutPins.Add(SubPins[0]);	
+					OutPins.Add(SubPins[1]);
+				}	
+			}
+			else
+			{				
+				for (URigVMPin* SubPin : InPin->GetSubPins())
+				{
+					VisitPinRecursively(SubPin, OutPins);
+				}
+			}
+		}
+	};
+	
+	for(int32 SortPhase = 0; SortPhase < 4; SortPhase++)
+	{
+		for(URigVMPin* RootPin : RootModelPins)
+		{
+			switch (SortPhase)
+			{
+				case 0: // execute IO pins
+				{
+					if(RootPin->IsExecuteContext() && RootPin->GetDirection() == ERigVMPinDirection::IO)
+					{
+						Local::VisitPinRecursively(RootPin, ModelPins);
+					}
+					break;
+				}
+				case 1: // IO pins
+				{
+					if(!RootPin->IsExecuteContext() && RootPin->GetDirection() == ERigVMPinDirection::IO)
+					{
+						Local::VisitPinRecursively(RootPin, ModelPins);
+					}
+					break;
+				}
+				case 2: // input / visible pins
+				{
+					if(RootPin->GetDirection() == ERigVMPinDirection::Input || RootPin->GetDirection() == ERigVMPinDirection::Visible)
+					{
+						Local::VisitPinRecursively(RootPin, ModelPins);
+					}
+					break;
+				}
+				case 3: // output pins
+				default:
+				{
+					if(RootPin->GetDirection() == ERigVMPinDirection::Output)
+					{
+						Local::VisitPinRecursively(RootPin, ModelPins);
+					}
+					break;
+				}
+			}
+		}
+	}
+
+	const UControlRigGraphSchema* RigSchema = Cast<UControlRigGraphSchema>(RigGraphNode->GetSchema());
+	
+	TMap<URigVMPin*, int32> ModelPinToInfoIndex;
+	for(URigVMPin* ModelPin : ModelPins)
+	{
+		FPinInfo PinInfo;
+		PinInfo.Index = PinInfos.Num();
+		PinInfo.ParentIndex = INDEX_NONE;
+		PinInfo.bHasChildren = (ModelPin->GetSubPins().Num() > 0);
+		PinInfo.bIsContainer = ModelPin->IsArray();
+		PinInfo.Depth = 0;
+		PinInfo.bExpanded = ModelPin->IsExpanded();
+		PinInfo.ModelPinPath = ModelPin->GetPinPath();
+		PinInfo.bAutoHeight = false;
+		
+		const bool bAskSchemaForEdition = RigSchema && ModelPin->IsStruct() && !ModelPin->IsBoundToVariable();
+		PinInfo.bHideInputWidget = (!ModelPin->IsBoundToVariable()) && PinInfo.bIsContainer;
+		if (!PinInfo.bHideInputWidget)
+		{
+			if (bAskSchemaForEdition && !PinInfo.bHasChildren)
+			{
+				const bool bIsStructEditable = RigSchema->IsStructEditable(ModelPin->GetScriptStruct()); 
+				PinInfo.bHideInputWidget = !bIsStructEditable;
+				PinInfo.bAutoHeight = bIsStructEditable;
+			}
+			else if(PinInfo.bHasChildren && !ModelPin->IsBoundToVariable())
+			{
+				PinInfo.bHideInputWidget = true;
+			}
+		}
+		
+		if(URigVMPin* ParentPin = ModelPin->GetParentPin())
+		{
+			const int32* ParentIndexPtr = ModelPinToInfoIndex.Find(ParentPin);
+			if(ParentIndexPtr == nullptr)
+			{
+				continue;
+			}
+			PinInfo.ParentIndex = *ParentIndexPtr;
+			PinInfo.Depth = PinInfos[PinInfo.ParentIndex].Depth + 1;
+		}
+
+		TAttribute<EVisibility> PinVisibilityAttribute = TAttribute<EVisibility>::CreateSP(this, &SControlRigGraphNode::GetPinVisibility, PinInfo.Index);
+
+		bool bPinWidgetForExpanderLeft = false;
+		TSharedPtr<SGraphPin> PinWidgetForExpander;
+
+		bool bPinInfoIsValid = false;
+		if(UEdGraphPin* OutputEdGraphPin = RigGraphNode->FindPin(ModelPin->GetPinPath(), EEdGraphPinDirection::EGPD_Output))
+		{
+			if(const int32* PinIndexPtr = EdGraphPinToOutputPin.Find(OutputEdGraphPin))
+			{
+				PinInfo.OutputPinWidget = OutputPins[*PinIndexPtr];
+				PinInfo.OutputPinWidget->SetVisibility(PinVisibilityAttribute);
+				PinWidgetForExpander = PinInfo.OutputPinWidget;
+				bPinWidgetForExpanderLeft = false;
+				bPinInfoIsValid = true;
+			}
+		}
+		
+		if(UEdGraphPin* InputEdGraphPin = RigGraphNode->FindPin(ModelPin->GetPinPath(), EEdGraphPinDirection::EGPD_Input))
+		{
+			if(const int32* PinIndexPtr = EdGraphPinToInputPin.Find(InputEdGraphPin))
+			{
+				PinInfo.InputPinWidget = InputPins[*PinIndexPtr];
+				PinInfo.InputPinWidget->SetVisibility(PinVisibilityAttribute);
+				PinWidgetForExpander = PinInfo.InputPinWidget;
+				bPinWidgetForExpanderLeft = true;
+				bPinInfoIsValid = true;
+			}
+		}
+
+		if(!bPinInfoIsValid)
+		{
+			continue;
+		}
+
+		ModelPinToInfoIndex.Add(ModelPin, PinInfos.Add(PinInfo));
+		
+		// check if this pin has sub pins
+		TSharedPtr<SHorizontalBox> FullPinHorizontalRowWidget = PinWidgetForExpander->GetFullPinHorizontalRowWidget().Pin();
+		if(FullPinHorizontalRowWidget.IsValid())
+		{
+			// indent the pin by padding
+			const float DepthIndentation = 12.f * float(PinInfo.Depth + (PinInfo.bHasChildren ? 0 : 1));
+			const float LeftIndentation = bPinWidgetForExpanderLeft ? DepthIndentation : 0.f;
+			const float RightIndentation = bPinWidgetForExpanderLeft ? 0.f : DepthIndentation;
+			const FMargin LineIndentation = FMargin(RightIndentation, 0, LeftIndentation, 0); 
+
+			static const TSharedRef<FTagMetaData> ExpanderButtonMetadata = MakeShared<FTagMetaData>(TEXT("SControlRigGraphNode.ExpanderButton"));
+
+			// check if this pin widget may already have the expander button 
+			int32 ExpanderSlotIndex = INDEX_NONE;
+			int32 PaddedSlotIndex = INDEX_NONE;
+
+			for(int32 SlotIndex=0; SlotIndex<FullPinHorizontalRowWidget->NumSlots(); SlotIndex++)
+			{
+				const TSharedRef<SWidget> Widget = FullPinHorizontalRowWidget->GetSlot(SlotIndex).GetWidget();
+				const TSharedPtr<FTagMetaData> Metadata = Widget->GetMetaData<FTagMetaData>();
+				if(Metadata.IsValid())
+				{
+					if(Metadata->Tag == ExpanderButtonMetadata->Tag)
+					{
+						ExpanderSlotIndex = SlotIndex;
+						break;
+					}
+				}
+			}
+
+			if(!PinInfo.bHasChildren && ExpanderSlotIndex != INDEX_NONE)
+			{
+				SHorizontalBox::FSlot& Slot = FullPinHorizontalRowWidget->GetSlot(ExpanderSlotIndex);
+				const TSharedRef<SWidget> Widget = Slot.GetWidget();
+				Widget->RemoveMetaData(ExpanderButtonMetadata);
+				FullPinHorizontalRowWidget->RemoveSlot(Widget);
+				ExpanderSlotIndex = INDEX_NONE;
+			}
+	
+			if(PinInfo.bHasChildren && ExpanderSlotIndex == INDEX_NONE)
+			{
+				// only inject the expander arrow for inputs on input / IO
+				// or for output pins
+				if(
+					(
+						(
+							(ModelPin->GetDirection() == ERigVMPinDirection::Input) ||
+							(ModelPin->GetDirection() == ERigVMPinDirection::IO)
+						) &&
+						bPinWidgetForExpanderLeft
+					) ||
+					(
+						(ModelPin->GetDirection() == ERigVMPinDirection::Output) &&
+						(!bPinWidgetForExpanderLeft)
+					)
+				)
+				{
+					// Add the expander arrow
+					FullPinHorizontalRowWidget->InsertSlot(bPinWidgetForExpanderLeft ? 1 : FullPinHorizontalRowWidget->GetChildren()->Num() - 1)
+					.Padding(FMargin(0, 0, 0, 0))
+					.AutoWidth()
+					[
+						SNew(SButton)
+						.ButtonStyle(FAppStyle::Get(), TEXT("SimpleButton"))
+						.ContentPadding(0)
+						.VAlign(VAlign_Center)
+						.HAlign(HAlign_Center)
+						.ClickMethod( EButtonClickMethod::MouseDown )
+						.OnClicked(this, &SControlRigGraphNode::OnExpanderArrowClicked, PinInfo.Index)
+						.ToolTipText(LOCTEXT("ExpandSubPin", "Expand Pin"))
+						[
+							SNew(SImage)
+							.Image(this, &SControlRigGraphNode::GetExpanderImage, PinInfo.Index, bPinWidgetForExpanderLeft, false)
+							.ColorAndOpacity(FSlateColor::UseForeground())
+						]
+						.AddMetaData(ExpanderButtonMetadata)
+					];
+				}
+			}
+
+			// adjust the padding
+			{
+				const int32 SlotToAdjustIndex = bPinWidgetForExpanderLeft ? 0 : FullPinHorizontalRowWidget->NumSlots() - 1;
+				SHorizontalBox::FSlot& Slot = FullPinHorizontalRowWidget->GetSlot(SlotToAdjustIndex);
+				Slot.SetPadding(LineIndentation);
+			}
+		}
+	}
+
+	// add spacer widget at the start
+	LeftNodeBox->AddSlot()
+	.HAlign(HAlign_Fill)
+	.VAlign(VAlign_Center)
+	.AutoHeight()
+	[
+		SNew(SSpacer)
+		.Size(FVector2D(1.f, 2.f))
+	];
+
+	auto AddArrayPlusButtonLambda = [this](URigVMPin* InModelPin, TSharedPtr<SHorizontalBox> InSlotLayout, const float InEmptySidePadding)
+	{
+		// add array plus button
+		InSlotLayout->AddSlot()
+		.AutoWidth()
+		.HAlign(HAlign_Left)
+		.Padding(PinWidgetSidePadding, TopPadding, InEmptySidePadding, 0.f)
+		[
+			SNew(SButton)
+			.ContentPadding(0.0f)
+			.ButtonStyle(FAppStyle::Get(), "NoBorder")
+			.OnClicked(this, &SControlRigGraphNode::HandleAddArrayElement, InModelPin->GetPinPath())
+			.IsEnabled(this, &SGraphNode::IsNodeEditable)
+			.Cursor(EMouseCursor::Default)
+			.Visibility(this, &SControlRigGraphNode::GetArrayPlusButtonVisibility, InModelPin)
+			.ToolTipText(LOCTEXT("AddArrayElement", "Add Array Element"))
+			[
+				SNew(SHorizontalBox)
+				+SHorizontalBox::Slot()
+				.AutoWidth()
+				.VAlign(VAlign_Center)
+				[
+					SNew(SImage)
+					.Image(FAppStyle::GetBrush(TEXT("Icons.PlusCircle")))
+				]
+			]
+		];
+	};
+
+	for(const FPinInfo& PinInfo : PinInfos)
+	{
+		if(PinInfo.InputPinWidget.IsValid())
+		{
+			if(PinInfo.bHideInputWidget)
+			{
+				if(PinInfo.InputPinWidget->GetValueWidget() != SNullWidget::NullWidget)
+				{
+					PinInfo.InputPinWidget->GetValueWidget()->SetVisibility(EVisibility::Collapsed);
+				}
+			}
+				
+			// input pins
+			if(!PinInfo.OutputPinWidget.IsValid())
+			{
+				TSharedPtr<SHorizontalBox> SlotLayout;
+				SHorizontalBox::FSlot* FirstSlot = nullptr;
+
+				const float MyEmptySidePadding = PinInfo.bHideInputWidget ? EmptySidePadding : 0.f; 
+				
+				LeftNodeBox->AddSlot()
+                .HAlign(HAlign_Fill)
+				.VAlign(VAlign_Center)
+				.AutoHeight()
+				.MaxHeight(PinInfo.bAutoHeight ? TAttribute<float>() : MaxHeight)
+                [
+                    SAssignNew(SlotLayout, SHorizontalBox)
+                    .Visibility(this, &SControlRigGraphNode::GetPinVisibility, PinInfo.Index)
+		            
+                    +SHorizontalBox::Slot()
+                    .Expose(FirstSlot)
+                    .FillWidth(1.f)
+                    .HAlign(HAlign_Left)
+                    .Padding(PinWidgetSidePadding, TopPadding, PinInfo.bIsContainer ? 0.f : MyEmptySidePadding, 0.f)
+                    [
+                        PinInfo.InputPinWidget.ToSharedRef()
+                    ]
+                ];
+
+				if(PinInfo.bIsContainer)
+				{
+					URigVMPin* ModelPin = ModelNode->GetGraph()->FindPin(PinInfo.ModelPinPath);
+					if(ModelPin)
+					{
+						// make sure to minimize the width of the label
+						FirstSlot->SetAutoWidth();
+						AddArrayPlusButtonLambda(ModelPin, SlotLayout, MyEmptySidePadding);
+					}
+				}
+			}
+			// io pins
+			else
+			{
+				TSharedPtr<SHorizontalBox> SlotLayout;
+				SHorizontalBox::FSlot* FirstSlot = nullptr;
+
+				PinInfo.OutputPinWidget->SetShowLabel(false);
+			
+				LeftNodeBox->AddSlot()
+                .HAlign(HAlign_Fill)
+				.VAlign(VAlign_Center)
+                .AutoHeight()
+				.MaxHeight(PinInfo.bAutoHeight ? TAttribute<float>() : MaxHeight)
+                [
+                    SAssignNew(SlotLayout, SHorizontalBox)
+                    .Visibility(this, &SControlRigGraphNode::GetPinVisibility, PinInfo.Index)
+
+                    +SHorizontalBox::Slot()
+					.Expose(FirstSlot)
+                    .FillWidth(1.f)
+                    .HAlign(HAlign_Left)
+                    .VAlign(VAlign_Center)
+                    .Padding(PinWidgetSidePadding, TopPadding, 0.f, 0.f)
+                    [
+                        PinInfo.InputPinWidget.ToSharedRef()
+                    ]
+                ];
+
+				if(PinInfo.bIsContainer)
+				{
+					URigVMPin* ModelPin = ModelNode->GetGraph()->FindPin(PinInfo.ModelPinPath);
+					if(ModelPin)
+					{
+						// make sure to minimize the width of the label
+						FirstSlot->SetAutoWidth();
+						AddArrayPlusButtonLambda(ModelPin, SlotLayout, EmptySidePadding);
+					}
+				}
+
+				SlotLayout->AddSlot()
+				.FillWidth(1.f)
+				.HAlign(HAlign_Right)
+				.VAlign(VAlign_Center)
+				.Padding(0.f, TopPadding, PinWidgetSidePadding, 0.f)
+				[
+					PinInfo.OutputPinWidget.ToSharedRef()
+				];
+			}
+		}
+		// output pins
+		else if(PinInfo.OutputPinWidget.IsValid())
+		{
+			LeftNodeBox->AddSlot()
+            .HAlign(HAlign_Fill)
+			.VAlign(VAlign_Center)
+            .AutoHeight()
+			.MaxHeight(PinInfo.bAutoHeight ? TAttribute<float>() : MaxHeight)
+            [
+            	SNew(SHorizontalBox)
+	            .Visibility(this, &SControlRigGraphNode::GetPinVisibility, PinInfo.Index)
+	            
+				+SHorizontalBox::Slot()
+				.FillWidth(1.f)
+				.HAlign(HAlign_Right)
+				.VAlign(VAlign_Center)
+	            .Padding(EmptySidePadding, TopPadding, PinWidgetSidePadding, 0.f)
+				[
+	                PinInfo.OutputPinWidget.ToSharedRef()
+				]
+            ];
+		}
+	}
+
+	if(URigVMFunctionReferenceNode* FunctionReferenceNode = Cast<URigVMFunctionReferenceNode>(ModelNode))
+	{
+		TWeakObjectPtr<URigVMFunctionReferenceNode> WeakFunctionReferenceNode = FunctionReferenceNode;
+		TWeakObjectPtr<UControlRigBlueprint> WeakControlRigBlueprint = Blueprint;
+
+		// add the entries for the variable remapping
+		for(const TSharedPtr<FRigVMExternalVariable>& ExternalVariable : RigGraphNode->ExternalVariables)
+		{
+			LeftNodeBox->AddSlot()
+			.HAlign(HAlign_Fill)
+			.VAlign(VAlign_Center)
+			.AutoHeight()
+			.MaxHeight(MaxHeight)
+			[
+				SNew(SHorizontalBox)
+	            
+				+SHorizontalBox::Slot()
+				.AutoWidth()
+				.HAlign(HAlign_Left)
+				.VAlign(VAlign_Center)
+				.Padding(PinWidgetSidePadding, TopPadding, PinWidgetSidePadding, 0.f)
+				[
+					SNew(STextBlock)
+					.Text(FText::FromName(ExternalVariable->Name))
+					.TextStyle(FAppStyle::Get(), NAME_DefaultPinLabelStyle)
+					.ColorAndOpacity(this, &SControlRigGraphNode::GetVariableLabelTextColor, WeakFunctionReferenceNode, ExternalVariable->Name)
+					.ToolTipText(this, &SControlRigGraphNode::GetVariableLabelTooltipText, WeakControlRigBlueprint, ExternalVariable->Name)
+				]
+
+				+SHorizontalBox::Slot()
+				.AutoWidth()
+				.HAlign(HAlign_Left)
+				.VAlign(VAlign_Center)
+				.Padding(PinWidgetSidePadding, TopPadding, PinWidgetSidePadding, 0.f)
+				[
+					SNew(SControlRigVariableBinding)
+					.Blueprint(Blueprint.Get())
+					.FunctionReferenceNode(FunctionReferenceNode)
+					.InnerVariableName(ExternalVariable->Name)
+				]
+			];
+		}
+	}
+
+	CreateAggregateAddPinButton();
+	
+	// add spacer widget at the end
+	LeftNodeBox->AddSlot()
+	.HAlign(HAlign_Fill)
+	.VAlign(VAlign_Center)
+	.AutoHeight()
+	[
+		SNew(SSpacer)
+		.Size(FVector2D(1.f, 4.f))
+	];
+}
+
 #undef LOCTEXT_NAMESPACE
diff --git a/Engine/Plugins/Animation/ControlRig/Source/ControlRigEditor/Private/Graph/SControlRigGraphNode.h b/Engine/Plugins/Animation/ControlRig/Source/ControlRigEditor/Private/Graph/SControlRigGraphNode.h
--- a/Engine/Plugins/Animation/ControlRig/Source/ControlRigEditor/Private/Graph/SControlRigGraphNode.h
+++ b/Engine/Plugins/Animation/ControlRig/Source/ControlRigEditor/Private/Graph/SControlRigGraphNode.h
@@ -35,6 +35,7 @@ public:
 	virtual void EndUserInteraction() const override;
 	virtual void MoveTo(const FVector2D& NewPosition, FNodeSet& NodeFilter, bool bMarkDirty = true) override;
 	virtual void AddPin( const TSharedRef<SGraphPin>& PinToAdd ) override;
+	virtual void CreateStandardPinWidget(UEdGraphPin* CurPin) override;
 	virtual void SetDefaultTitleAreaWidget(TSharedRef<SOverlay> DefaultTitleAreaWidget) override
 	{
 		TitleAreaWidget = DefaultTitleAreaWidget;
@@ -76,6 +77,7 @@ private:
 	FReply HandleAddArrayElement(FString InModelPinPath);
 
 	void HandleNodeTitleDirtied();
+	void HandleNodePinsChanged();
 
 	FText GetInstructionCountText() const;
 	FText GetInstructionDurationText() const;
@@ -87,6 +89,7 @@ private:
 	const FSlateBrush * GetExpanderImage(int32 InPinInfoIndex, bool bLeft, bool bHovered) const;
 	FReply OnExpanderArrowClicked(int32 InPinInfoIndex);
 	void HandleModifiedEvent(ERigVMGraphNotifType InNotifType, URigVMGraph* InGraph, UObject* InSubject);
+	void UpdatePinTreeView();
 
 	/** Cached widget title area */
 	TSharedPtr<SOverlay> TitleAreaWidget;
@@ -124,4 +127,13 @@ private:
 	
 	TArray<FPinInfo> PinInfos;
 	TWeakObjectPtr<URigVMNode> ModelNode;
+
+	// Pins to keep after calling HandleNodePinsChanged. We recycle these pins in
+	// CreateStandardPinWidget.
+	TMap<const UEdGraphPin *, TSharedRef<SGraphPin>> PinsToKeep;
+
+	// Delayed pin deletion. To deal with the fact that pin deletion cannot occur until we
+	// have re-generated the pin list. SControlRigGraphNode has already relinquished them
+	// but we still have a pointer to them in our pin widget.
+	TSet<UEdGraphPin *> PinsToDelete;
 };
diff --git a/Engine/Source/Developer/RigVMDeveloper/Private/RigVMModel/RigVMController.cpp b/Engine/Source/Developer/RigVMDeveloper/Private/RigVMModel/RigVMController.cpp
--- a/Engine/Source/Developer/RigVMDeveloper/Private/RigVMModel/RigVMController.cpp
+++ b/Engine/Source/Developer/RigVMDeveloper/Private/RigVMModel/RigVMController.cpp
@@ -8353,6 +8353,7 @@ URigVMPin* URigVMController::InsertArrayPin(URigVMPin* ArrayPin, int32 InIndex,
 		Pin->DefaultValue = DefaultValue;
 	}
 
+	Notify(ERigVMGraphNotifType::PinAdded, Pin);
 	Notify(ERigVMGraphNotifType::PinArraySizeChanged, ArrayPin);
 	if (!bSuspendNotifications)
 	{
@@ -8409,15 +8410,16 @@ bool URigVMController::RemoveArrayPin(const FString& InArrayElementPinPath, bool
 	}
 
 	int32 IndexToRemove = ArrayElementPin->GetPinIndex();
-	if (!RemovePin(ArrayElementPin, bSetupUndoRedo, false))
+	if (!RemovePin(ArrayElementPin, bSetupUndoRedo, true))
 	{
 		return false;
 	}
 
-	for (int32 ExistingIndex = ArrayPin->GetSubPins().Num() - 1; ExistingIndex >= IndexToRemove; ExistingIndex--)
+	for (int32 ExistingIndex = IndexToRemove; ExistingIndex < ArrayPin->GetArraySize(); ExistingIndex++)
 	{
 		URigVMPin* ExistingPin = ArrayPin->GetSubPins()[ExistingIndex];
 		ExistingPin->SetNameFromIndex();
+		Notify(ERigVMGraphNotifType::PinRenamed, ExistingPin);
 	}
 
 	if (!bSuspendNotifications)
@@ -10288,7 +10290,11 @@ bool URigVMController::AddArgumentForPin(URigVMPin* InPin, URigVMPin* InToLinkPi
 		ActionStack->AddAction(FRigVMSetLibraryTemplateAction(LibraryNode, NewTemplate));
 	}
 	
-	LibraryNode->Template = NewTemplate;	
+	LibraryNode->Template = NewTemplate;
+	{
+		FRigVMControllerGraphGuard GraphGuard(this, LibraryNode->GetTypedOuter<URigVMGraph>(), false);
+		Notify(ERigVMGraphNotifType::LibraryTemplateChanged, LibraryNode);
+	}
 	
 	// Update Entry and Return filtered permutations to all permutations
 	URigVMTemplateNode* EntryNode = Graph->GetEntryNode();
@@ -10652,6 +10658,10 @@ bool URigVMController::RenameExposedPin(const FName& InOldPinName, const FName&
 			ActionStack->AddAction(FRigVMSetLibraryTemplateAction(LibraryNode, NewTemplate));
 		}
 		LibraryNode->Template = NewTemplate;
+		{
+			FRigVMControllerGraphGuard GraphGuard(this, LibraryNode->GetTypedOuter<URigVMGraph>(), false);
+			Notify(ERigVMGraphNotifType::LibraryTemplateChanged, LibraryNode);
+		}
 	}
 
 	if (URigVMFunctionLibrary* FunctionLibrary = Cast<URigVMFunctionLibrary>(LibraryNode->GetGraph()))
@@ -17656,8 +17666,41 @@ bool URigVMController::UpdateLibraryTemplate(URigVMLibraryNode* LibraryNode, boo
 			ActionStack->AddAction(FRigVMSetTemplateFilteredPermutationsAction(LibraryNode, OldPermutations));
 			ActionStack->AddAction(FRigVMSetLibraryTemplateAction(LibraryNode, NewTemplate));		
 		}
+
+		bool bTemplateChanged = false;
+		if(!bTemplateChanged)
+		{
+			bTemplateChanged = CurTemplate->GetNotation() != NewTemplate.GetNotation();
+		}
+		if(!bTemplateChanged)
+		{
+			check(CurTemplate->NumArguments() == NewTemplate.NumArguments());
+			for(int32 ArgumentIndex = 0; (ArgumentIndex < NewTemplate.NumArguments()) && (!bTemplateChanged); ArgumentIndex++)
+			{
+				const FRigVMTemplateArgument* ArgA = CurTemplate->GetArgument(ArgumentIndex);
+				const FRigVMTemplateArgument* ArgB = NewTemplate.GetArgument(ArgumentIndex);
+
+				const TArray<TRigVMTypeIndex>& TypesA = ArgA->GetTypeIndices();
+				const TArray<TRigVMTypeIndex>& TypesB = ArgB->GetTypeIndices();
+				bTemplateChanged = TypesA.Num() != TypesB.Num();
+				if(!bTemplateChanged)
+				{
+					for(int32 TypeIndex=0; (TypeIndex<TypesA.Num()) && (!bTemplateChanged);TypeIndex++)
+					{
+						bTemplateChanged = TypesA[TypeIndex] != TypesB[TypeIndex];
+					}
+				}
+			}
+		}
+		
 		LibraryNode->Template = NewTemplate;
 		Notify(ERigVMGraphNotifType::NodeDescriptionChanged, LibraryNode);
+
+		if(bTemplateChanged)
+		{
+			FRigVMControllerGraphGuard GraphGuard(this, LibraryNode->GetTypedOuter<URigVMGraph>(), false);
+			Notify(ERigVMGraphNotifType::LibraryTemplateChanged, LibraryNode);
+		}
 	}
 
 	// Update Entry and Return filtered permutations to all permutations
@@ -18048,6 +18091,10 @@ void URigVMController::InitializeFilteredPermutationsFromTemplateTypes()
 			{
 				NewTemplate.ComputeNotationFromArguments(LibraryNode->GetName());				
 				LibraryNode->Template = NewTemplate;
+				{
+					FRigVMControllerGraphGuard GraphGuard(this, LibraryNode->GetTypedOuter<URigVMGraph>(), false);
+					Notify(ERigVMGraphNotifType::LibraryTemplateChanged, LibraryNode);
+				}
 			}
 		}
 		
diff --git a/Engine/Source/Developer/RigVMDeveloper/Public/RigVMModel/RigVMNotifications.h b/Engine/Source/Developer/RigVMDeveloper/Public/RigVMModel/RigVMNotifications.h
--- a/Engine/Source/Developer/RigVMDeveloper/Public/RigVMModel/RigVMNotifications.h
+++ b/Engine/Source/Developer/RigVMDeveloper/Public/RigVMModel/RigVMNotifications.h
@@ -51,6 +51,7 @@ enum class ERigVMGraphNotifType : uint8
 	NodeKeywordsChanged, // A node's keywords have changed (Subject == URigVMNode)
 	NodeDescriptionChanged, // A node's description has changed (Subject == URigVMNode)
 	VariableRemappingChanged, // A function reference node's remapping has changed (Subject == URigVMFunctionReferenceNode)
+	LibraryTemplateChanged, // The definition of a library node's template has changed
 	Invalid // The max for this enum (used for guarding)
 };
 
