Control Rig: Hierarchy Reorder functionality

diff --git a/Engine/Plugins/Animation/ControlRig/Source/ControlRig/Private/Rigs/RigHierarchy.cpp b/Engine/Plugins/Animation/ControlRig/Source/ControlRig/Private/Rigs/RigHierarchy.cpp
--- a/Engine/Plugins/Animation/ControlRig/Source/ControlRig/Private/Rigs/RigHierarchy.cpp
+++ b/Engine/Plugins/Animation/ControlRig/Source/ControlRig/Private/Rigs/RigHierarchy.cpp
@@ -3821,6 +3821,22 @@ bool URigHierarchy::IsDependentOn(FRigBaseElement* InDependent, FRigBaseElement*
 	return false;
 }
 
+int32 URigHierarchy::GetLocalIndex(const FRigBaseElement* InElement) const
+{
+	if(InElement == nullptr)
+	{
+		return INDEX_NONE;
+	}
+	
+	if(const FRigBaseElement* ParentElement = GetFirstParent(InElement))
+	{
+		const FRigBaseElementChildrenArray& Children = GetChildren(ParentElement);
+		return Children.Find((FRigBaseElement*)InElement);
+	}
+
+	return GetRootElements().Find((FRigBaseElement*)InElement);
+}
+
 bool URigHierarchy::IsTracingChanges() const
 {
 #if WITH_EDITOR
diff --git a/Engine/Plugins/Animation/ControlRig/Source/ControlRig/Private/Rigs/RigHierarchyController.cpp b/Engine/Plugins/Animation/ControlRig/Source/ControlRig/Private/Rigs/RigHierarchyController.cpp
--- a/Engine/Plugins/Animation/ControlRig/Source/ControlRig/Private/Rigs/RigHierarchyController.cpp
+++ b/Engine/Plugins/Animation/ControlRig/Source/ControlRig/Private/Rigs/RigHierarchyController.cpp
@@ -1873,8 +1873,57 @@ FRigElementKey URigHierarchyController::RenameElement(FRigElementKey InElement,
 	return bRenamed ? Element->GetKey() : FRigElementKey();
 }
 
-FName URigHierarchyController::SetDisplayName(FRigElementKey InControl, FName InDisplayName, bool bRenameElement, bool bSetupUndo,
+bool URigHierarchyController::ReorderElement(FRigElementKey InElement, int32 InIndex, bool bSetupUndo,
 	bool bPrintPythonCommand)
+{
+	if(!IsValid())
+	{
+		return false;
+	}
+
+	FRigBaseElement* Element = Hierarchy->Find(InElement);
+	if(Element == nullptr)
+	{
+		ReportWarningf(TEXT("Cannot Reorder Element: '%s' not found."), *InElement.ToString());
+		return false;
+	}
+
+#if WITH_EDITOR
+	TSharedPtr<FScopedTransaction> TransactionPtr;
+	if(bSetupUndo)
+	{
+		TransactionPtr = MakeShared<FScopedTransaction>(NSLOCTEXT("RigHierarchyController", "Reorder Element", "Reorder Element"));
+		Hierarchy->Modify();
+	}
+#endif
+
+	const bool bReordered = ReorderElement(Element, InIndex);
+
+#if WITH_EDITOR
+	if(!bReordered && TransactionPtr.IsValid())
+	{
+		TransactionPtr->Cancel();
+	}
+	TransactionPtr.Reset();
+
+	if (bReordered && bPrintPythonCommand && !bSuspendPythonPrinting)
+	{
+		UBlueprint* Blueprint = GetTypedOuter<UBlueprint>();
+		if (Blueprint)
+		{
+			RigVMPythonUtils::Print(Blueprint->GetFName().ToString(), 
+				FString::Printf(TEXT("hierarchy_controller.reorder_element(%s, %d)"),
+				*InElement.ToPythonString(),
+				InIndex));
+		}
+	}
+#endif
+
+	return bReordered;
+}
+
+FName URigHierarchyController::SetDisplayName(FRigElementKey InControl, FName InDisplayName, bool bRenameElement, bool bSetupUndo,
+                                              bool bPrintPythonCommand)
 {
 	if(!IsValid())
 	{
@@ -2007,6 +2056,68 @@ bool URigHierarchyController::RenameElement(FRigBaseElement* InElement, const FN
 	return true;
 }
 
+bool URigHierarchyController::ReorderElement(FRigBaseElement* InElement, int32 InIndex)
+{
+	if(InElement == nullptr)
+	{
+		return false;
+	}
+
+	InIndex = FMath::Max<int32>(InIndex, 0);
+
+	TArray<FRigBaseElement*> LocalElements;
+	const FRigBaseElement* ParentElement = Hierarchy->GetFirstParent(InElement);
+	if(ParentElement)
+	{
+		const FRigBaseElementChildrenArray& Children = Hierarchy->GetChildren(ParentElement);
+		LocalElements.Append(Children);
+	}
+	else
+	{
+		const TArray<FRigBaseElement*> RootElements = Hierarchy->GetRootElements();
+		LocalElements.Append(RootElements);
+	}
+
+	const int32 CurrentIndex = LocalElements.Find(InElement);
+	if(CurrentIndex == INDEX_NONE || CurrentIndex == InIndex || !LocalElements.IsValidIndex(InIndex))
+	{
+		return false;
+	}
+
+	Hierarchy->IncrementTopologyVersion();
+
+	TArray<int32> GlobalIndices;
+	GlobalIndices.Reserve(LocalElements.Num());
+	for(const FRigBaseElement* Element : LocalElements)
+	{
+		GlobalIndices.Add(Element->GetIndex());
+	}
+
+	LocalElements.RemoveAt(CurrentIndex);
+	LocalElements.Insert(InElement, InIndex);
+
+	const int32 LowerBound = FMath::Min<int32>(InIndex, CurrentIndex);
+	const int32 UpperBound = FMath::Max<int32>(InIndex, CurrentIndex);
+	for(int32 LocalIndex = LowerBound; LocalIndex <= UpperBound; LocalIndex++)
+	{
+		const int32 GlobalIndex = GlobalIndices[LocalIndex];
+		FRigBaseElement* Element = LocalElements[LocalIndex];
+		Hierarchy->Elements[GlobalIndex] = Element;
+		Element->Index = GlobalIndex;
+		Hierarchy->IndexLookup.FindOrAdd(Element->Key) = GlobalIndex;
+		Element->TopologyVersion = Hierarchy->GetTopologyVersion();
+	}
+
+	if(ParentElement)
+	{
+		Hierarchy->UpdateCachedChildren(ParentElement, true);
+	}
+
+	Notify(ERigHierarchyNotification::ElementReordered, InElement);
+
+	return true;
+}
+
 FName URigHierarchyController::SetDisplayName(FRigControlElement* InControlElement, const FName& InDisplayName, bool bRenameElement)
 {
 	if(InControlElement == nullptr)
diff --git a/Engine/Plugins/Animation/ControlRig/Source/ControlRig/Public/Rigs/RigHierarchy.h b/Engine/Plugins/Animation/ControlRig/Source/ControlRig/Public/Rigs/RigHierarchy.h
--- a/Engine/Plugins/Animation/ControlRig/Source/ControlRig/Public/Rigs/RigHierarchy.h
+++ b/Engine/Plugins/Animation/ControlRig/Source/ControlRig/Public/Rigs/RigHierarchy.h
@@ -331,6 +331,27 @@ public:
 		return INDEX_NONE;
 	}
 
+	/**
+	 * Returns the index of an element given its key within its default parent (or root)
+	 * @param InKey The key of the element to retrieve the index for
+	 * @return The index of the element or INDEX_NONE
+	 */
+	UFUNCTION(BlueprintCallable, Category = URigHierarchy, meta = (DisplayName = "Get Index", ScriptName = "GetIndex"))
+	FORCEINLINE int32 GetLocalIndex_ForBlueprint(FRigElementKey InKey) const
+	{
+		return GetLocalIndex(InKey);
+	}
+
+	/**
+	 * Returns the index of an element given its key within its default parent (or root)
+	 * @param InKey The key of the element to retrieve the index for
+	 * @return The index of the element or INDEX_NONE
+	 */
+	FORCEINLINE int32 GetLocalIndex(const FRigElementKey& InKey) const
+	{
+		return GetLocalIndex(Find(InKey));
+	}
+
 	/**
 	 * Returns the indices of an array of keys
 	 * @param InKeys The keys of the elements to retrieve the indices for
@@ -839,6 +860,29 @@ public:
 		return GetKeysOfType<FRigReferenceElement>(bTraverse);
 	}
 
+	/**
+	 * Returns all root elements
+	 */
+	FORCEINLINE TArray<FRigBaseElement*> GetRootElements() const
+	{
+		return GetFilteredElements<FRigBaseElement>([this](FRigBaseElement* Element) -> bool
+		{
+			return GetNumberOfParents(Element) == 0;
+		});
+	}
+
+	/**
+	 * Returns all root element keys
+	 */
+	UFUNCTION(BlueprintCallable, Category = URigHierarchy, meta = (DisplayName = "Get Controls", ScriptName = "GetControls"))
+	FORCEINLINE TArray<FRigElementKey> GetRootElementKeys() const
+	{
+		return GetKeysByPredicate([this](const FRigBaseElement& Element) -> bool
+		{
+			return GetNumberOfParents(Element.Index) == 0;
+		}, false);
+	}
+
 	/**
 	 * Returns the name of metadata for a given element
 	 * @param InItem The element key to return the metadata keys for
@@ -3373,6 +3417,13 @@ public:
 	 */
 	bool IsDependentOn(FRigBaseElement* InDependent, FRigBaseElement* InDependency, TArray<bool>& InElementsVisited, const TElementDependencyMap& InDependencyMap = TElementDependencyMap()) const;
 
+	/**
+	 * Returns the index of an element given its key within its default parent (or root)
+	 * @param InElement The element to retrieve the index for
+	 * @return The index of the element or INDEX_NONE
+	 */
+	int32 GetLocalIndex(const FRigBaseElement* InElement) const;
+
 	/**
 	 * Returns a reference to the suspend notifications flag
 	 */
diff --git a/Engine/Plugins/Animation/ControlRig/Source/ControlRig/Public/Rigs/RigHierarchyController.h b/Engine/Plugins/Animation/ControlRig/Source/ControlRig/Public/Rigs/RigHierarchyController.h
--- a/Engine/Plugins/Animation/ControlRig/Source/ControlRig/Public/Rigs/RigHierarchyController.h
+++ b/Engine/Plugins/Animation/ControlRig/Source/ControlRig/Public/Rigs/RigHierarchyController.h
@@ -404,6 +404,17 @@ public:
 	UFUNCTION(BlueprintCallable, Category = URigHierarchyController)
     FRigElementKey RenameElement(FRigElementKey InElement, FName InName, bool bSetupUndo = false, bool bPrintPythonCommand = false, bool bClearSelection = true);
 
+	/**
+	 * Changes the element's index within its default parent (or the top level)
+	 * @param InElement The key of the element to rename
+	 * @param InIndex The new index of the element to take within its default parent (or the top level)
+	 * @param bSetupUndo If set to true the stack will record the change for undo / redo
+	 * @param bPrintPythonCommand If set to true a python command equivalent to this call will be printed out
+	 * @return Returns true if the element has been reordered accordingly
+	 */
+	UFUNCTION(BlueprintCallable, Category = URigHierarchyController)
+	bool ReorderElement(FRigElementKey InElement, int32 InIndex, bool bSetupUndo = false, bool bPrintPythonCommand = false);
+
 	/**
  	 * Sets the display name on a control
  	 * @param InControl The key of the control to change the display name for
@@ -611,6 +622,14 @@ private:
 	 */
     bool RenameElement(FRigBaseElement* InElement, const FName &InName, bool bClearSelection = true);
 
+	/**
+	 * Changes the element's index within its default parent (or the top level)
+	 * @param InElement The key of the element to rename
+	 * @param InIndex The new index of the element to take within its default parent (or the top level)
+	 * @return Returns true if the element has been reordered accordingly
+	 */
+	bool ReorderElement(FRigBaseElement* InElement, int32 InIndex);
+
 	/**
  	 * Sets the display name on a control
  	 * @param InControlElement The element to change the display name for
diff --git a/Engine/Plugins/Animation/ControlRig/Source/ControlRig/Public/Rigs/RigHierarchyDefines.h b/Engine/Plugins/Animation/ControlRig/Source/ControlRig/Public/Rigs/RigHierarchyDefines.h
--- a/Engine/Plugins/Animation/ControlRig/Source/ControlRig/Public/Rigs/RigHierarchyDefines.h
+++ b/Engine/Plugins/Animation/ControlRig/Source/ControlRig/Public/Rigs/RigHierarchyDefines.h
@@ -88,6 +88,7 @@ enum class ERigHierarchyNotification : uint8
 	ParentWeightsChanged,
 	InteractionBracketOpened,
 	InteractionBracketClosed,
+	ElementReordered,
 
 	/** MAX - invalid */
 	Max UMETA(Hidden),
