Control Rig: Avoid consecutive runs of Construction during LOD change

diff --git a/Engine/Plugins/Animation/ControlRig/Source/ControlRig/Private/AnimNode_ControlRigBase.cpp b/Engine/Plugins/Animation/ControlRig/Source/ControlRig/Private/AnimNode_ControlRigBase.cpp
--- a/Engine/Plugins/Animation/ControlRig/Source/ControlRig/Private/AnimNode_ControlRigBase.cpp
+++ b/Engine/Plugins/Animation/ControlRig/Source/ControlRig/Private/AnimNode_ControlRigBase.cpp
@@ -29,6 +29,8 @@ FAnimNode_ControlRigBase::FAnimNode_ControlRigBase()
 	, OutputSettings(FControlRigIOSettings())
 	, bExecute(true)
 	, InternalBlendAlpha (1.f)
+	, bControlRigRequiresInitialization(true)
+	, LastBonesSerialNumberForCacheBones(0)
 {
 
 }
@@ -70,6 +72,8 @@ void FAnimNode_ControlRigBase::Initialize_AnyThread(const FAnimationInitializeCo
 		//Don't Inititialize the Control Rig here it may have the wrong VM on the CDO
 		SetTargetInstance(ControlRig);
 		ControlRig->RequestInit();
+		bControlRigRequiresInitialization = true;
+		LastBonesSerialNumberForCacheBones = 0;
 	}
 }
 
@@ -598,14 +602,19 @@ void FAnimNode_ControlRigBase::CacheBones_AnyThread(const FAnimationCacheBonesCo
 
 	if (UControlRig* ControlRig = GetControlRig())
 	{
+		// fill up node names
+		FBoneContainer& RequiredBones = Context.AnimInstanceProxy->GetRequiredBones();
+
+		const uint16 BonesSerialNumber = RequiredBones.GetSerialNumber();
+		const bool bIsLODChange = !bControlRigRequiresInitialization && (BonesSerialNumber != LastBonesSerialNumberForCacheBones);
+
 		// the construction event may create a set of bones that we can map to. let's run construction now.
-		if(ControlRig->IsConstructionModeEnabled() || ControlRig->IsConstructionRequired())
+		if(ControlRig->IsConstructionModeEnabled() ||
+			(ControlRig->IsConstructionRequired() && (bControlRigRequiresInitialization || bIsLODChange)))
 		{
 			ControlRig->Execute(EControlRigState::Update, FRigUnit_PrepareForExecution::EventName);
+			bControlRigRequiresInitialization = false;
 		}
-			
-		// fill up node names
-		FBoneContainer& RequiredBones = Context.AnimInstanceProxy->GetRequiredBones();
 
 		ControlRigBoneInputMappingByIndex.Reset();
 		ControlRigBoneOutputMappingByIndex.Reset();
@@ -790,10 +799,16 @@ void FAnimNode_ControlRigBase::CacheBones_AnyThread(const FAnimationCacheBonesCo
 			}
 		}
 
-		// re-init when LOD changes
-		// and restore control values
-		FControlRigControlScope Scope(ControlRig);
-		ControlRig->Execute(EControlRigState::Init, FRigUnit_BeginExecution::EventName);
+		if(bControlRigRequiresInitialization)
+		{
+			// re-init only if this is the first run
+			// and restore control values
+			FControlRigControlScope Scope(ControlRig);
+			ControlRig->Execute(EControlRigState::Init, FRigUnit_BeginExecution::EventName);
+			bControlRigRequiresInitialization = false;
+		}
+		
+		LastBonesSerialNumberForCacheBones = BonesSerialNumber;
 	}
 }
 
diff --git a/Engine/Plugins/Animation/ControlRig/Source/ControlRig/Public/AnimNode_ControlRigBase.h b/Engine/Plugins/Animation/ControlRig/Source/ControlRig/Public/AnimNode_ControlRigBase.h
--- a/Engine/Plugins/Animation/ControlRig/Source/ControlRig/Public/AnimNode_ControlRigBase.h
+++ b/Engine/Plugins/Animation/ControlRig/Source/ControlRig/Public/AnimNode_ControlRigBase.h
@@ -167,6 +167,9 @@ protected:
 	void ExecuteControlRig(FPoseContext& InOutput);
 
 	void QueueControlRigDrawInstructions(UControlRig* ControlRig, FAnimInstanceProxy* Proxy) const;
+	
+	bool bControlRigRequiresInitialization;
+	uint16 LastBonesSerialNumberForCacheBones;
 
 	friend struct FControlRigSequencerAnimInstanceProxy;
 	friend struct FControlRigLayerInstanceProxy;
