Control Rig: Fix duplicate script names in RigHierarchy header

diff --git a/Engine/Plugins/Animation/ControlRig/Source/ControlRig/Public/Rigs/RigHierarchy.h b/Engine/Plugins/Animation/ControlRig/Source/ControlRig/Public/Rigs/RigHierarchy.h
--- a/Engine/Plugins/Animation/ControlRig/Source/ControlRig/Public/Rigs/RigHierarchy.h
+++ b/Engine/Plugins/Animation/ControlRig/Source/ControlRig/Public/Rigs/RigHierarchy.h
@@ -336,7 +336,7 @@ public:
 	 * @param InKey The key of the element to retrieve the index for
 	 * @return The index of the element or INDEX_NONE
 	 */
-	UFUNCTION(BlueprintCallable, Category = URigHierarchy, meta = (DisplayName = "Get Index", ScriptName = "GetIndex"))
+	UFUNCTION(BlueprintCallable, Category = URigHierarchy, meta = (DisplayName = "Get Local Index", ScriptName = "GetLocalIndex"))
 	FORCEINLINE int32 GetLocalIndex_ForBlueprint(FRigElementKey InKey) const
 	{
 		return GetLocalIndex(InKey);
@@ -874,7 +874,7 @@ public:
 	/**
 	 * Returns all root element keys
 	 */
-	UFUNCTION(BlueprintCallable, Category = URigHierarchy, meta = (DisplayName = "Get Controls", ScriptName = "GetControls"))
+	UFUNCTION(BlueprintCallable, Category = URigHierarchy, meta = (DisplayName = "Get Root Elements", ScriptName = "GetRootElements"))
 	FORCEINLINE TArray<FRigElementKey> GetRootElementKeys() const
 	{
 		return GetKeysByPredicate([this](const FRigBaseElement& Element) -> bool
