Control Rig: Allow to reorder elements in the hierarchy

This submit contains the UI portion.

diff --git a/Engine/Plugins/Animation/ControlRig/Source/ControlRig/Private/Rigs/RigHierarchyController.cpp b/Engine/Plugins/Animation/ControlRig/Source/ControlRig/Private/Rigs/RigHierarchyController.cpp
--- a/Engine/Plugins/Animation/ControlRig/Source/ControlRig/Private/Rigs/RigHierarchyController.cpp
+++ b/Engine/Plugins/Animation/ControlRig/Source/ControlRig/Private/Rigs/RigHierarchyController.cpp
@@ -2079,7 +2079,7 @@ bool URigHierarchyController::ReorderElement(FRigBaseElement* InElement, int32 I
 	}
 
 	const int32 CurrentIndex = LocalElements.Find(InElement);
-	if(CurrentIndex == INDEX_NONE || CurrentIndex == InIndex || !LocalElements.IsValidIndex(InIndex))
+	if(CurrentIndex == INDEX_NONE || CurrentIndex == InIndex)
 	{
 		return false;
 	}
@@ -2094,8 +2094,16 @@ bool URigHierarchyController::ReorderElement(FRigBaseElement* InElement, int32 I
 	}
 
 	LocalElements.RemoveAt(CurrentIndex);
-	LocalElements.Insert(InElement, InIndex);
+	if(InIndex >= LocalElements.Num())
+	{
+		LocalElements.Add(InElement);
+	}
+	else
+	{
+		LocalElements.Insert(InElement, InIndex);
+	}
 
+	InIndex = FMath::Min<int32>(InIndex, LocalElements.Num() - 1);
 	const int32 LowerBound = FMath::Min<int32>(InIndex, CurrentIndex);
 	const int32 UpperBound = FMath::Max<int32>(InIndex, CurrentIndex);
 	for(int32 LocalIndex = LowerBound; LocalIndex <= UpperBound; LocalIndex++)
diff --git a/Engine/Plugins/Animation/ControlRig/Source/ControlRigDeveloper/Private/ControlRigBlueprint.cpp b/Engine/Plugins/Animation/ControlRig/Source/ControlRigDeveloper/Private/ControlRigBlueprint.cpp
--- a/Engine/Plugins/Animation/ControlRig/Source/ControlRigDeveloper/Private/ControlRigBlueprint.cpp
+++ b/Engine/Plugins/Animation/ControlRig/Source/ControlRigDeveloper/Private/ControlRigBlueprint.cpp
@@ -4711,6 +4711,7 @@ void UControlRigBlueprint::HandleHierarchyModified(ERigHierarchyNotification InN
 		}
 		case ERigHierarchyNotification::ElementAdded:
 		case ERigHierarchyNotification::ParentChanged:
+		case ERigHierarchyNotification::ElementReordered:
 		case ERigHierarchyNotification::HierarchyReset:
 		{
 			Modify();
diff --git a/Engine/Plugins/Animation/ControlRig/Source/ControlRigEditor/Private/ControlRigElementDetails.cpp b/Engine/Plugins/Animation/ControlRig/Source/ControlRigEditor/Private/ControlRigElementDetails.cpp
--- a/Engine/Plugins/Animation/ControlRig/Source/ControlRigEditor/Private/ControlRigElementDetails.cpp
+++ b/Engine/Plugins/Animation/ControlRig/Source/ControlRigEditor/Private/ControlRigElementDetails.cpp
@@ -3029,7 +3029,49 @@ void FRigControlElementDetails::CustomizeAnimationChannels(IDetailLayoutBuilder&
 						})),
 						LOCTEXT("DeleteAnimationChannel", "Delete"),
 						LOCTEXT("DeleteAnimationChannelTooltip", "Deletes this animation channel"),
-					FSlateIcon());
+						FSlateIcon());
+
+					const FRigElementKey ControlElementKey = ControlElement->GetKey();
+					
+					// move up or down
+					WidgetRow->AddCustomContextMenuAction(FUIAction(
+						FExecuteAction::CreateLambda([this, ControlElementKey, ChildElementKeys, HierarchyToChange]()
+						{
+							if(URigHierarchyController* Controller = HierarchyToChange->GetController(true))
+							{
+								FScopedTransaction(LOCTEXT("MoveAnimationChannelUp", "Move Animation Channel Up"));
+								HierarchyToChange->Modify();
+								
+								for(const FRigElementKey& KeyToMove : ChildElementKeys)
+								{
+									const int32 LocalIndex = HierarchyToChange->GetLocalIndex(KeyToMove);
+									Controller->ReorderElement(KeyToMove, LocalIndex - 1, true);
+								}
+								Controller->SelectElement(ControlElementKey, true, true);
+							}
+						})),
+						LOCTEXT("MoveAnimationChannelUp", "Move Up"),
+						LOCTEXT("MoveAnimationChannelUpTooltip", "Reorders this animation channel to show up one higher"),
+						FSlateIcon());
+					WidgetRow->AddCustomContextMenuAction(FUIAction(
+						FExecuteAction::CreateLambda([this, ControlElementKey, ChildElementKeys, HierarchyToChange]()
+						{
+							if(URigHierarchyController* Controller = HierarchyToChange->GetController(true))
+							{
+								FScopedTransaction(LOCTEXT("MoveAnimationChannelDown", "Move Animation Channel Down"));
+								HierarchyToChange->Modify();
+								
+								for(const FRigElementKey& KeyToMove : ChildElementKeys)
+								{
+									const int32 LocalIndex = HierarchyToChange->GetLocalIndex(KeyToMove);
+									Controller->ReorderElement(KeyToMove, LocalIndex + 1, true);
+								}
+								Controller->SelectElement(ControlElementKey, true, true);
+							}
+						})),
+						LOCTEXT("MoveAnimationChannelDown", "Move Down"),
+						LOCTEXT("MoveAnimationChannelDownTooltip", "Reorders this animation channel to show up one lower"),
+						FSlateIcon());
 				}
 			}
 		}
@@ -3059,7 +3101,7 @@ void FRigControlElementDetails::CustomizeAnimationChannels(IDetailLayoutBuilder&
 
 FReply FRigControlElementDetails::OnAddAnimationChannelClicked()
 {
-	if(!IsAnyElementNotOfType(ERigElementType::Control) || IsAnyElementProcedural())
+	if(IsAnyElementNotOfType(ERigElementType::Control) || IsAnyElementProcedural())
 	{
 		return FReply::Handled();
 	}
@@ -3075,6 +3117,7 @@ FReply FRigControlElementDetails::OnAddAnimationChannelClicked()
 	Settings.MaximumValue = FRigControlValue::Make<float>(1.f);
 	Settings.DisplayName = HierarchyToChange->GetSafeNewDisplayName(Key, ChannelName);
 	HierarchyToChange->GetController(true)->AddAnimationChannel(*ChannelName, Key, Settings, true, true);
+	HierarchyToChange->GetController(true)->SelectElement(Key);
 	return FReply::Handled();
 }
 
diff --git a/Engine/Plugins/Animation/ControlRig/Source/ControlRigEditor/Private/EditMode/ControlRigEditMode.cpp b/Engine/Plugins/Animation/ControlRig/Source/ControlRigEditor/Private/EditMode/ControlRigEditMode.cpp
--- a/Engine/Plugins/Animation/ControlRig/Source/ControlRigEditor/Private/EditMode/ControlRigEditMode.cpp
+++ b/Engine/Plugins/Animation/ControlRig/Source/ControlRigEditor/Private/EditMode/ControlRigEditMode.cpp
@@ -3432,6 +3432,7 @@ void FControlRigEditMode::OnHierarchyModified(ERigHierarchyNotification InNotif,
 		case ERigHierarchyNotification::ElementAdded:
 		case ERigHierarchyNotification::ElementRemoved:
 		case ERigHierarchyNotification::ElementRenamed:
+		case ERigHierarchyNotification::ElementReordered:
 		case ERigHierarchyNotification::HierarchyReset:
 		{
 			UControlRig* ControlRig = InHierarchy->GetTypedOuter<UControlRig>();
diff --git a/Engine/Plugins/Animation/ControlRig/Source/ControlRigEditor/Private/Editor/ControlRigSkeletalMeshComponent.cpp b/Engine/Plugins/Animation/ControlRig/Source/ControlRigEditor/Private/Editor/ControlRigSkeletalMeshComponent.cpp
--- a/Engine/Plugins/Animation/ControlRig/Source/ControlRigEditor/Private/Editor/ControlRigSkeletalMeshComponent.cpp
+++ b/Engine/Plugins/Animation/ControlRig/Source/ControlRigEditor/Private/Editor/ControlRigSkeletalMeshComponent.cpp
@@ -202,6 +202,7 @@ void UControlRigSkeletalMeshComponent::OnHierarchyModified(ERigHierarchyNotifica
 		case ERigHierarchyNotification::ElementAdded:
 		case ERigHierarchyNotification::ElementRemoved:
 		case ERigHierarchyNotification::ElementRenamed:
+		case ERigHierarchyNotification::ElementReordered:
 		case ERigHierarchyNotification::ParentChanged:
 		case ERigHierarchyNotification::HierarchyReset:
 		{
diff --git a/Engine/Plugins/Animation/ControlRig/Source/ControlRigEditor/Private/Editor/SRigCurveContainer.cpp b/Engine/Plugins/Animation/ControlRig/Source/ControlRigEditor/Private/Editor/SRigCurveContainer.cpp
--- a/Engine/Plugins/Animation/ControlRig/Source/ControlRigEditor/Private/Editor/SRigCurveContainer.cpp
+++ b/Engine/Plugins/Animation/ControlRig/Source/ControlRigEditor/Private/Editor/SRigCurveContainer.cpp
@@ -588,6 +588,7 @@ void SRigCurveContainer::OnHierarchyModified(ERigHierarchyNotification InNotif,
 		case ERigHierarchyNotification::ElementAdded:
 		case ERigHierarchyNotification::ElementRemoved:
 		case ERigHierarchyNotification::ElementRenamed:
+		case ERigHierarchyNotification::ElementReordered:
 		case ERigHierarchyNotification::HierarchyReset:
 		{
 			RefreshCurveList();
diff --git a/Engine/Plugins/Animation/ControlRig/Source/ControlRigEditor/Private/Editor/SRigHierarchy.cpp b/Engine/Plugins/Animation/ControlRig/Source/ControlRigEditor/Private/Editor/SRigHierarchy.cpp
--- a/Engine/Plugins/Animation/ControlRig/Source/ControlRigEditor/Private/Editor/SRigHierarchy.cpp
+++ b/Engine/Plugins/Animation/ControlRig/Source/ControlRigEditor/Private/Editor/SRigHierarchy.cpp
@@ -46,6 +46,8 @@
 #include "Framework/Notifications/NotificationManager.h"
 #include "Widgets/Notifications/SNotificationList.h"
 #include "Styling/AppStyle.h"
+#include "Algo/MinElement.h"
+#include "Algo/MaxElement.h"
 
 #include UE_INLINE_GENERATED_CPP_BY_NAME(SRigHierarchy)
 
@@ -700,6 +702,7 @@ void SRigHierarchy::OnHierarchyModified(ERigHierarchyNotification InNotif, URigH
 			break;
 		}
 		case ERigHierarchyNotification::ElementRenamed:
+		case ERigHierarchyNotification::ElementReordered:
 		case ERigHierarchyNotification::HierarchyReset:
 		{
 			RefreshTreeView();
@@ -2168,45 +2171,65 @@ FReply SRigHierarchy::OnDragDetected(const FGeometry& MyGeometry, const FPointer
 
 TOptional<EItemDropZone> SRigHierarchy::OnCanAcceptDrop(const FDragDropEvent& DragDropEvent, EItemDropZone DropZone, TSharedPtr<FRigTreeElement> TargetItem)
 {
+	const TOptional<EItemDropZone> InvalidDropZone;
 	TOptional<EItemDropZone> ReturnDropZone;
 
-	TSharedPtr<FRigElementHierarchyDragDropOp> RigDragDropOp = DragDropEvent.GetOperationAs<FRigElementHierarchyDragDropOp>();
+	const TSharedPtr<FRigElementHierarchyDragDropOp> RigDragDropOp = DragDropEvent.GetOperationAs<FRigElementHierarchyDragDropOp>();
 	if (RigDragDropOp.IsValid())
 	{
-		URigHierarchy* Hierarchy = GetHierarchy();
-		if (Hierarchy)
+		FRigElementKey TargetKey;
+		if (const URigHierarchy* Hierarchy = GetHierarchy())
 		{
-			for (const FRigElementKey& DraggedKey : RigDragDropOp->GetElements())
+			switch(DropZone)
 			{
-				if (DraggedKey == TargetItem->Key)
+				case EItemDropZone::AboveItem:
+				case EItemDropZone::BelowItem:
 				{
-					return ReturnDropZone;
+					TargetKey = Hierarchy->GetFirstParent(TargetItem->Key);
+					break;
 				}
+				case EItemDropZone::OntoItem:
+				{
+					TargetKey = TargetItem->Key;
+					break;
+				}
+			}
 
+			for (const FRigElementKey& DraggedKey : RigDragDropOp->GetElements())
+			{
 				if(Hierarchy->IsProcedural(DraggedKey))
 				{
-					return ReturnDropZone;
+					return InvalidDropZone;
+				}
+
+				// re-parenting directly onto an item
+				if (DraggedKey == TargetKey)
+				{
+					return InvalidDropZone;
 				}
 
-				if(Hierarchy->IsParentedTo(TargetItem->Key, DraggedKey))
+				if(DropZone == EItemDropZone::OntoItem)
 				{
-					return ReturnDropZone;
+					if(Hierarchy->IsParentedTo(TargetKey, DraggedKey))
+					{
+						return InvalidDropZone;
+					}
 				}
 			}
 		}
 
 		// don't allow dragging onto procedural items
-		if(!GetDefaultHierarchy()->Contains(TargetItem->Key))
+		if(TargetKey.IsValid() && !GetDefaultHierarchy()->Contains(TargetKey))
 		{
-			return ReturnDropZone;
+			return InvalidDropZone;
 		}
 
-		switch (TargetItem->Key.Type)
+		switch (TargetKey.Type)
 		{
 			case ERigElementType::Bone:
 			{
 				// bones can parent anything
-				ReturnDropZone = EItemDropZone::OntoItem;
+				ReturnDropZone = DropZone;
 				break;
 			}
 			case ERigElementType::Control:
@@ -2227,16 +2250,16 @@ TOptional<EItemDropZone> SRigHierarchy::OnCanAcceptDrop(const FDragDropEvent& Dr
 						}
 						default:
 						{
-							return ReturnDropZone;
+							return InvalidDropZone;
 						}
 					}
 				}
-				ReturnDropZone = EItemDropZone::OntoItem;
+				ReturnDropZone = DropZone;
 				break;
 			}
 			default:
 			{
-				ReturnDropZone = EItemDropZone::OntoItem;
+				ReturnDropZone = DropZone;
 				break;
 			}
 		}
@@ -2272,12 +2295,36 @@ FReply SRigHierarchy::OnAcceptDrop(const FDragDropEvent& DragDropEvent, EItemDro
 		}
 		else
 		{
+			const URigHierarchy* Hierarchy = GetDefaultHierarchy();
+
 			FRigElementKey TargetKey;
+			int32 LocalIndex = INDEX_NONE;
+
 			if (TargetItem.IsValid())
 			{
-				TargetKey = TargetItem->Key;
+				switch(DropZone)
+				{
+					case EItemDropZone::AboveItem:
+					{
+						TargetKey = Hierarchy->GetFirstParent(TargetItem->Key);
+						LocalIndex = Hierarchy->GetLocalIndex(TargetItem->Key);
+						break;
+					}
+					case EItemDropZone::BelowItem:
+					{
+						TargetKey = Hierarchy->GetFirstParent(TargetItem->Key);
+						LocalIndex = Hierarchy->GetLocalIndex(TargetItem->Key) + 1;
+						break;
+					}
+					case EItemDropZone::OntoItem:
+					{
+						TargetKey = TargetItem->Key;
+						break;
+					}
+				}
 			}
-			return ReparentOrMatchTransform(RigDragDropOp->GetElements(), TargetKey, bReparentItems);			
+			
+			return ReparentOrMatchTransform(RigDragDropOp->GetElements(), TargetKey, bReparentItems, LocalIndex);			
 		}
 
 	}
@@ -2708,7 +2755,7 @@ void SRigHierarchy::HandleAlign(const FToolMenuContext& Context)
 	}
 }
 
-FReply SRigHierarchy::ReparentOrMatchTransform(const TArray<FRigElementKey>& DraggedKeys, FRigElementKey TargetKey, bool bReparentItems)
+FReply SRigHierarchy::ReparentOrMatchTransform(const TArray<FRigElementKey>& DraggedKeys, FRigElementKey TargetKey, bool bReparentItems, int32 LocalIndex)
 {
 	bool bMatchTransforms = !bReparentItems;
 
@@ -2741,7 +2788,10 @@ FReply SRigHierarchy::ReparentOrMatchTransform(const TArray<FRigElementKey>& Dra
 			{
 				if(Hierarchy->IsParentedTo(TargetKey, DraggedKey))
 				{
-					return FReply::Unhandled();
+					if(LocalIndex == INDEX_NONE)
+					{
+						return FReply::Unhandled();
+					}
 				}
 			}
 
@@ -2815,41 +2865,50 @@ FReply SRigHierarchy::ReparentOrMatchTransform(const TArray<FRigElementKey>& Dra
 			const FTransform CurrentLocalTransform = DebuggedHierarchy->GetLocalTransform(DraggedKey);
 			const FTransform CurrentGlobalOffsetTransform = DebuggedHierarchy->GetGlobalControlOffsetTransform(DraggedKey, false);
 
-			if(ParentKey.IsValid())
+			bool bElementWasReparented = false;
+			if(ParentKey.IsValid() && (Hierarchy->GetFirstParent(DraggedKey) != ParentKey))
 			{
-				Controller->SetParent(DraggedKey, ParentKey, true, true, true);
+				bElementWasReparented = Controller->SetParent(DraggedKey, ParentKey, true, true, true);
 			}
-			else
+			else if(!ParentKey.IsValid() && (Hierarchy->GetNumberOfParents(DraggedKey) > 0))
 			{
-				Controller->RemoveAllParents(DraggedKey, true, true, true);
+				bElementWasReparented = Controller->RemoveAllParents(DraggedKey, true, true, true);
 			}
 
-			if (DraggedKey.Type == ERigElementType::Control)
+			if(LocalIndex != INDEX_NONE)
 			{
-				int32 ControlIndex = DebuggedHierarchy->GetIndex(DraggedKey);
-				if (ControlIndex == INDEX_NONE)
-				{
-					continue;
-				}
-
-				const FTransform GlobalParentTransform = DebuggedHierarchy->GetGlobalTransform(ParentKey, false);
-				const FTransform LocalOffsetTransform = CurrentGlobalOffsetTransform.GetRelativeTransform(GlobalParentTransform);
-
-				Hierarchy->SetControlOffsetTransformByIndex(ControlIndex, LocalOffsetTransform, ERigTransformType::InitialLocal, true, true, true);
-				Hierarchy->SetControlOffsetTransformByIndex(ControlIndex, LocalOffsetTransform, ERigTransformType::CurrentLocal, true, true, true);
-				Hierarchy->SetLocalTransform(DraggedKey, CurrentLocalTransform, true, true, true, true);
-				Hierarchy->SetInitialLocalTransform(DraggedKey, InitialLocalTransform, true, true, true);
-				DebuggedHierarchy->SetControlOffsetTransformByIndex(ControlIndex, LocalOffsetTransform, ERigTransformType::InitialLocal, true, true);
-				DebuggedHierarchy->SetControlOffsetTransformByIndex(ControlIndex, LocalOffsetTransform, ERigTransformType::CurrentLocal, true, true);
-				DebuggedHierarchy->SetLocalTransform(DraggedKey, CurrentLocalTransform, true, true, true);
-				DebuggedHierarchy->SetInitialLocalTransform(DraggedKey, InitialLocalTransform, true, true);
+				Controller->ReorderElement(DraggedKey, LocalIndex, true, true);
 			}
-			else
+
+			if(bElementWasReparented)
 			{
-				DebuggedHierarchy->SetInitialGlobalTransform(DraggedKey, InitialGlobalTransform, true, true);
-				DebuggedHierarchy->SetGlobalTransform(DraggedKey, CurrentGlobalTransform, false, true, true);
-				Hierarchy->SetInitialGlobalTransform(DraggedKey, InitialGlobalTransform, true, true);
-				Hierarchy->SetGlobalTransform(DraggedKey, CurrentGlobalTransform, false, true, true);
+				if (DraggedKey.Type == ERigElementType::Control)
+				{
+					int32 ControlIndex = DebuggedHierarchy->GetIndex(DraggedKey);
+					if (ControlIndex == INDEX_NONE)
+					{
+						continue;
+					}
+
+					const FTransform GlobalParentTransform = DebuggedHierarchy->GetGlobalTransform(ParentKey, false);
+					const FTransform LocalOffsetTransform = CurrentGlobalOffsetTransform.GetRelativeTransform(GlobalParentTransform);
+
+					Hierarchy->SetControlOffsetTransformByIndex(ControlIndex, LocalOffsetTransform, ERigTransformType::InitialLocal, true, true, true);
+					Hierarchy->SetControlOffsetTransformByIndex(ControlIndex, LocalOffsetTransform, ERigTransformType::CurrentLocal, true, true, true);
+					Hierarchy->SetLocalTransform(DraggedKey, CurrentLocalTransform, true, true, true, true);
+					Hierarchy->SetInitialLocalTransform(DraggedKey, InitialLocalTransform, true, true, true);
+					DebuggedHierarchy->SetControlOffsetTransformByIndex(ControlIndex, LocalOffsetTransform, ERigTransformType::InitialLocal, true, true);
+					DebuggedHierarchy->SetControlOffsetTransformByIndex(ControlIndex, LocalOffsetTransform, ERigTransformType::CurrentLocal, true, true);
+					DebuggedHierarchy->SetLocalTransform(DraggedKey, CurrentLocalTransform, true, true, true);
+					DebuggedHierarchy->SetInitialLocalTransform(DraggedKey, InitialLocalTransform, true, true);
+				}
+				else
+				{
+					DebuggedHierarchy->SetInitialGlobalTransform(DraggedKey, InitialGlobalTransform, true, true);
+					DebuggedHierarchy->SetGlobalTransform(DraggedKey, CurrentGlobalTransform, false, true, true);
+					Hierarchy->SetInitialGlobalTransform(DraggedKey, InitialGlobalTransform, true, true);
+					Hierarchy->SetGlobalTransform(DraggedKey, CurrentGlobalTransform, false, true, true);
+				}
 			}
 		}
 	}
diff --git a/Engine/Plugins/Animation/ControlRig/Source/ControlRigEditor/Private/Editor/SRigHierarchy.h b/Engine/Plugins/Animation/ControlRig/Source/ControlRigEditor/Private/Editor/SRigHierarchy.h
--- a/Engine/Plugins/Animation/ControlRig/Source/ControlRigEditor/Private/Editor/SRigHierarchy.h
+++ b/Engine/Plugins/Animation/ControlRig/Source/ControlRigEditor/Private/Editor/SRigHierarchy.h
@@ -212,7 +212,7 @@ private:
 
 	void HandleParent(const FToolMenuContext& Context);
 	void HandleAlign(const FToolMenuContext& Context);
-	FReply ReparentOrMatchTransform(const TArray<FRigElementKey>& DraggedKeys, FRigElementKey TargetKey, bool bReparentItems);
+	FReply ReparentOrMatchTransform(const TArray<FRigElementKey>& DraggedKeys, FRigElementKey TargetKey, bool bReparentItems, int32 LocalIndex = INDEX_NONE);
 
 	FName CreateUniqueName(const FName& InBaseName, ERigElementType InElementType) const;
 
