Control Rig: RigGraph adding / removing to support undo

diff --git a/Engine/Plugins/Animation/ControlRig/Source/ControlRigDeveloper/Private/ControlRigBlueprint.cpp b/Engine/Plugins/Animation/ControlRig/Source/ControlRigDeveloper/Private/ControlRigBlueprint.cpp
--- a/Engine/Plugins/Animation/ControlRig/Source/ControlRigDeveloper/Private/ControlRigBlueprint.cpp
+++ b/Engine/Plugins/Animation/ControlRig/Source/ControlRigDeveloper/Private/ControlRigBlueprint.cpp
@@ -2652,6 +2652,10 @@ void UControlRigBlueprint::PostTransacted(const FTransactionObjectEvent& Transac
 		if (PropertiesChanged.Contains(GET_MEMBER_NAME_CHECKED(UControlRigBlueprint, RigVMClient)) ||
 			PropertiesChanged.Contains(GET_MEMBER_NAME_CHECKED(UControlRigBlueprint, UbergraphPages)))
 		{
+			UbergraphPages.RemoveAll([](const UEdGraph* UberGraph) -> bool
+			{
+ 				return UberGraph == nullptr || !IsValid(UberGraph);
+			});
 			RigVMClient.PostTransacted(TransactionEvent);
 
 			RecompileVM();
