ensure the viewport client is safe to use when getting the world from it

diff --git a/Engine/Plugins/Animation/ControlRig/Source/ControlRigEditor/Private/Editor/Constraints/SConstraintsWidget.cpp b/Engine/Plugins/Animation/ControlRig/Source/ControlRigEditor/Private/Editor/Constraints/SConstraintsWidget.cpp
--- a/Engine/Plugins/Animation/ControlRig/Source/ControlRigEditor/Private/Editor/Constraints/SConstraintsWidget.cpp
+++ b/Engine/Plugins/Animation/ControlRig/Source/ControlRigEditor/Private/Editor/Constraints/SConstraintsWidget.cpp
@@ -85,12 +85,22 @@ int8 FConstraintInfo::GetType(UClass* InClass)
 namespace
 {
 
+UWorld* GetCurrentWorld()
+{
+	return GCurrentLevelEditingViewportClient ? GCurrentLevelEditingViewportClient->GetWorld() : nullptr;
+}
+	
 // NOTE we use this function to get the current selection as control actors are Temporary Editor Actors so won't be
 // pushed added to the selection list.
 TArray<AActor*> GetCurrentSelection()
 {
-	const UWorld* World = GCurrentLevelEditingViewportClient->GetWorld();
-	const ULevel* CurrentLevel = World->GetCurrentLevel();
+	const UWorld* World = GetCurrentWorld();
+	const ULevel* CurrentLevel = IsValid(World) ? World->GetCurrentLevel() : nullptr;
+	if (!CurrentLevel)
+	{
+		static const TArray<AActor*> DummyArray;
+		return DummyArray;
+	}
 
 	return CurrentLevel->Actors.FilterByPredicate( [](const AActor* Actor)
 	{
@@ -297,7 +307,12 @@ void SDroppableConstraintItem::CreateConstraint(
 	auto CreateConstraint = [InCreationDelegate, InConstraintType](
 		const TArray<AActor*>& Selection, AActor* InParent, const FName& InSocketName)
 	{
-		UWorld* World = GCurrentLevelEditingViewportClient->GetWorld();
+		UWorld* World = GetCurrentWorld();
+		if (!IsValid(World))
+		{
+			return;
+		}
+
 		bool bCreated = false;
 		for (AActor* Child: Selection)
 		{
@@ -313,7 +328,7 @@ void SDroppableConstraintItem::CreateConstraint(
 				}
 			}
 		}
-		
+	
 		// update list
 		if (bCreated && InCreationDelegate.IsBound())
 		{
@@ -415,9 +430,14 @@ void SEditableConstraintItem::Construct(
 	});
 
 	// constraint
-	auto GetConstraint = [this]()
+	auto GetConstraint = [this]() -> UTickableConstraint*
 	{
-		UWorld* World = GCurrentLevelEditingViewportClient->GetWorld();
+		UWorld* World = GetCurrentWorld();
+		if (!IsValid(World))
+		{
+			return nullptr;
+		}
+		
 		const FConstraintsManagerController& Controller = FConstraintsManagerController::Get(World);
 		return Controller.GetConstraint(ConstraintItem->Name);
 	};
@@ -742,7 +762,12 @@ void SConstraintsEditionWidget::MoveItemUp(const TSharedPtr<FEditableConstraintI
 	const int32 Index = ListItems.IndexOfByKey(Item);
 	if (Index > 0 && ListItems.IsValidIndex(Index))
 	{
-		UWorld* World = GCurrentLevelEditingViewportClient->GetWorld();
+		UWorld* World = GetCurrentWorld();
+		if (!IsValid(World))
+		{
+			return;
+		}
+		
 		const FConstraintsManagerController& Controller = FConstraintsManagerController::Get(World);
 
 		// the current item needs to tick before the previous item 
@@ -764,7 +789,12 @@ void SConstraintsEditionWidget::MoveItemDown(const TSharedPtr<FEditableConstrain
 	const int32 Index = ListItems.IndexOfByKey(Item);
 	if (ListItems.IsValidIndex(Index))
 	{
-		UWorld* World = GCurrentLevelEditingViewportClient->GetWorld();
+		UWorld* World = GetCurrentWorld();
+		if (!IsValid(World))
+		{
+			return;
+		}
+		
 		const FConstraintsManagerController& Controller = FConstraintsManagerController::Get(World);
 
 		// the current item needs to tick after the next item 
@@ -778,8 +808,14 @@ void SConstraintsEditionWidget::MoveItemDown(const TSharedPtr<FEditableConstrain
 
 void SConstraintsEditionWidget::RemoveItem(const TSharedPtr<FEditableConstraintItem>& Item)
 {
+	UWorld* World = GetCurrentWorld();
+	if (!IsValid(World))
+	{
+		return;
+	}
+	
 	FScopedTransaction Transaction(LOCTEXT("RemoveConstraint", "Remove Constraint"));
-	UWorld* World = GCurrentLevelEditingViewportClient->GetWorld();
+	
 	const FConstraintsManagerController& Controller = FConstraintsManagerController::Get(World);
 	
 	const FName ConstraintName(Item->Name.ToString());
@@ -814,7 +850,12 @@ void SConstraintsEditionWidget::InvalidateConstraintList()
 int32 SConstraintsEditionWidget::RefreshConstraintList()
 {
 	// get constraints
-	UWorld* World = GCurrentLevelEditingViewportClient->GetWorld();
+	UWorld* World = GetCurrentWorld();
+	if (!IsValid(World))
+	{
+		return 0;
+	}
+	
 	const TArray<AActor*> Selection = GetCurrentSelection();
 	
 	const bool bIsConstraintsActor = Selection.Num() == 1 && Selection[0]->IsA<AConstraintsActor>();
@@ -875,7 +916,12 @@ TSharedPtr<SWidget> SConstraintsEditionWidget::CreateContextMenu()
 		return SNullWidget::NullWidget;
 	}
 
-	UWorld* World = GCurrentLevelEditingViewportClient->GetWorld();
+	UWorld* World = GetCurrentWorld();
+	if (!IsValid(World))
+	{
+		return SNullWidget::NullWidget;
+	}
+	
 	const FConstraintsManagerController& Controller = FConstraintsManagerController::Get(World);
 	UTickableConstraint* Constraint = Controller.GetConstraint(ListItems[Index]->Name);
 	if (!Constraint)
@@ -1028,7 +1074,12 @@ void SConstraintsEditionWidget::OnItemDoubleClicked(ItemSharedPtr InItem)
 		return;
 	}
 
-	UWorld* World = GCurrentLevelEditingViewportClient->GetWorld();
+	UWorld* World = GetCurrentWorld();
+	if (!IsValid(World))
+	{
+		return;
+	}
+	
 	const FConstraintsManagerController& Controller = FConstraintsManagerController::Get(World);
 	UTickableConstraint* Constraint = Controller.GetConstraint(InItem->Name);
 	if (!Constraint)
