ControlRigComponent: input components should tick before CR component and output component should tick after CR component

diff --git a/Engine/Plugins/Animation/ControlRig/Source/ControlRig/Private/ControlRigComponent.cpp b/Engine/Plugins/Animation/ControlRig/Source/ControlRig/Private/ControlRigComponent.cpp
--- a/Engine/Plugins/Animation/ControlRig/Source/ControlRig/Private/ControlRigComponent.cpp
+++ b/Engine/Plugins/Animation/ControlRig/Source/ControlRig/Private/ControlRigComponent.cpp
@@ -361,17 +361,37 @@ void UControlRigComponent::Initialize()
 	// we want to make sure all components driven by control rig component tick
 	// after the control rig component such that by the time they tick they can
 	// send the latest data for rendering
-	USceneComponent* LastComponent = nullptr;
+	// also need to make sure all components driving the control rig component tick before
+	// the control rig component ticks
+	TMap<TObjectPtr<USceneComponent>, EControlRigComponentMapDirection> MappedComponents;
+	TSet<TObjectPtr<USceneComponent>> MappedComponentsWithErrors;
 	for (const FControlRigComponentMappedElement& MappedElement : MappedElements)
 	{
 		if (MappedElement.SceneComponent)
 		{
-			if (LastComponent != MappedElement.SceneComponent)
-			{	
+			if (EControlRigComponentMapDirection* Direction = MappedComponents.Find(MappedElement.SceneComponent))
+			{
+				if (*Direction != MappedElement.Direction)
+				{
+					// elements from the same component should not be mapped to both input and output
+					MappedComponentsWithErrors.Add(MappedElement.SceneComponent);
+					continue;
+				}
+			}
+
+			// input elements should tick before ControlRig updates
+			// output elements should tick after ControlRig updates
+			if (MappedElement.Direction == EControlRigComponentMapDirection::Output)
+			{
 				MappedElement.SceneComponent->AddTickPrerequisiteComponent(this);
-				LastComponent = MappedElement.SceneComponent;
+			}
+			else
+			{
+				AddTickPrerequisiteComponent(MappedElement.SceneComponent);
 			}
 
+			MappedComponents.Add(MappedElement.SceneComponent) = MappedElement.Direction;
+			
 			// make sure that the animation is updated so that bone transforms are updated on the mapped component
 			// (otherwise, FControlRigAnimInstanceProxy::Evaluate is never called when moving a control in the editor)
 			if (USkeletalMeshComponent* Component = Cast<USkeletalMeshComponent>(MappedElement.SceneComponent))
@@ -380,6 +400,14 @@ void UControlRigComponent::Initialize()
 			}
 		}
 	}
+
+	for (const TObjectPtr<USceneComponent> ErrorComponent : MappedComponentsWithErrors)
+	{
+		FString Message = FString::Printf(
+				TEXT("Elements from the same component (%s) should not be mapped to both input and output,"
+					" because it creates ambiguity when inferring tick order."), *ErrorComponent.GetName());
+		UE_LOG(LogControlRig, Warning, TEXT("%s: %s"), *GetPathName(), *Message);
+	}
 }
 
 void UControlRigComponent::Update(float DeltaTime)
