Control Rig: Fix LOCTEXT duplicate use.

diff --git a/Engine/Plugins/Animation/ControlRig/Source/ControlRigEditor/Private/ControlRigElementDetails.cpp b/Engine/Plugins/Animation/ControlRig/Source/ControlRigEditor/Private/ControlRigElementDetails.cpp
--- a/Engine/Plugins/Animation/ControlRig/Source/ControlRigEditor/Private/ControlRigElementDetails.cpp
+++ b/Engine/Plugins/Animation/ControlRig/Source/ControlRigEditor/Private/ControlRigElementDetails.cpp
@@ -3039,7 +3039,7 @@ void FRigControlElementDetails::CustomizeAnimationChannels(IDetailLayoutBuilder&
 						{
 							if(URigHierarchyController* Controller = HierarchyToChange->GetController(true))
 							{
-								FScopedTransaction(LOCTEXT("MoveAnimationChannelUp", "Move Animation Channel Up"));
+								FScopedTransaction(LOCTEXT("MoveAnimationChannelUpTransaction", "Move Animation Channel Up"));
 								HierarchyToChange->Modify();
 								
 								for(const FRigElementKey& KeyToMove : ChildElementKeys)
@@ -3058,7 +3058,7 @@ void FRigControlElementDetails::CustomizeAnimationChannels(IDetailLayoutBuilder&
 						{
 							if(URigHierarchyController* Controller = HierarchyToChange->GetController(true))
 							{
-								FScopedTransaction(LOCTEXT("MoveAnimationChannelDown", "Move Animation Channel Down"));
+								FScopedTransaction(LOCTEXT("MoveAnimationChannelDownTransaction", "Move Animation Channel Down"));
 								HierarchyToChange->Modify();
 								
 								for(const FRigElementKey& KeyToMove : ChildElementKeys)
