Control Rig: Shape Settings Nodes

diff --git a/Engine/Plugins/Animation/ControlRig/Source/ControlRig/Private/Units/Execution/RigUnit_DynamicHierarchy.cpp b/Engine/Plugins/Animation/ControlRig/Source/ControlRig/Private/Units/Execution/RigUnit_DynamicHierarchy.cpp
--- a/Engine/Plugins/Animation/ControlRig/Source/ControlRig/Private/Units/Execution/RigUnit_DynamicHierarchy.cpp
+++ b/Engine/Plugins/Animation/ControlRig/Source/ControlRig/Private/Units/Execution/RigUnit_DynamicHierarchy.cpp
@@ -906,3 +906,50 @@ FRigUnit_HierarchyAddAnimationChannelRotator_Execute()
 	}
 }
 
+FRigUnit_HierarchyGetShapeSettings_Execute()
+{
+	Settings = FRigUnit_HierarchyAddControl_ShapeSettings();
+
+	if(URigHierarchy* Hierarchy = Context.Hierarchy)
+	{
+		if(const FRigControlElement* ControlElement = Hierarchy->Find<FRigControlElement>(Item))
+		{
+			Settings.bVisible = ControlElement->Settings.bShapeVisible;
+			Settings.Color = ControlElement->Settings.ShapeColor;
+			Settings.Name = ControlElement->Settings.ShapeName;
+			Settings.Transform = Hierarchy->GetControlShapeTransform((FRigControlElement*)ControlElement, ERigTransformType::InitialLocal);
+		}
+		else
+		{
+			UE_CONTROLRIG_RIGUNIT_REPORT_ERROR(TEXT("Control '%s' does not exist."), *Item.ToString());
+		}
+	}
+}
+
+FRigUnit_HierarchySetShapeSettings_Execute()
+{
+	FString ErrorMessage;
+	if(!FRigUnit_DynamicHierarchyBase::IsValidToRunInContext(Context, ExecuteContext, true, &ErrorMessage))
+	{
+		if(!ErrorMessage.IsEmpty())
+		{
+			UE_CONTROLRIG_RIGUNIT_REPORT_ERROR(TEXT("%s"), *ErrorMessage);
+		}
+		return;
+	}
+
+	if(FRigControlElement* ControlElement = ExecuteContext.Hierarchy->Find<FRigControlElement>(Item))
+	{
+		ControlElement->Settings.bShapeVisible = Settings.bVisible;
+		ControlElement->Settings.ShapeColor = Settings.Color;
+		ControlElement->Settings.ShapeName = Settings.Name;
+		ExecuteContext.Hierarchy->Notify(ERigHierarchyNotification::ControlSettingChanged, ControlElement);
+		ExecuteContext.Hierarchy->SetControlShapeTransformByIndex(ControlElement->GetIndex(), Settings.Transform, true, false);
+		ExecuteContext.Hierarchy->SetControlShapeTransformByIndex(ControlElement->GetIndex(), Settings.Transform, false, false);
+	}
+	else
+	{
+		UE_CONTROLRIG_RIGUNIT_REPORT_ERROR(TEXT("Control '%s' does not exist."), *Item.ToString());
+	}
+
+}
\ No newline at end of file
diff --git a/Engine/Plugins/Animation/ControlRig/Source/ControlRig/Private/Units/Execution/RigUnit_DynamicHierarchy.h b/Engine/Plugins/Animation/ControlRig/Source/ControlRig/Private/Units/Execution/RigUnit_DynamicHierarchy.h
--- a/Engine/Plugins/Animation/ControlRig/Source/ControlRig/Private/Units/Execution/RigUnit_DynamicHierarchy.h
+++ b/Engine/Plugins/Animation/ControlRig/Source/ControlRig/Private/Units/Execution/RigUnit_DynamicHierarchy.h
@@ -1233,4 +1233,63 @@ struct CONTROLRIG_API FRigUnit_HierarchyAddAnimationChannelRotator : public FRig
 
 	RIGVM_METHOD()
 	virtual void Execute(const FRigUnitContext& Context) override;
-};
\ No newline at end of file
+};
+
+/**
+ * Retrieves the shape settings of a control
+ */
+USTRUCT(meta=(DisplayName="Get Shape Settings", Keywords="Construction,Create,New,Control", Varying))
+struct CONTROLRIG_API FRigUnit_HierarchyGetShapeSettings : public FRigUnit_DynamicHierarchyBase
+{
+	GENERATED_BODY()
+
+	FRigUnit_HierarchyGetShapeSettings()
+	{
+		Item = FRigElementKey(NAME_None, ERigElementType::Control);
+	}
+
+	/*
+	 * The item to change
+	 */
+	UPROPERTY(meta = (Input))
+	FRigElementKey Item;
+
+	/*
+	 * The shape settings for the control
+	 */
+	UPROPERTY(meta = (Output))
+	FRigUnit_HierarchyAddControl_ShapeSettings Settings;
+
+	RIGVM_METHOD()
+	virtual void Execute(const FRigUnitContext& Context) override;
+};
+
+/**
+ * Changes the shape settings of a control
+ * Note: This node only runs as part of the construction event.
+ */
+USTRUCT(meta=(DisplayName="Set Shape Settings", Keywords="Construction,Create,New,Control", Varying))
+struct CONTROLRIG_API FRigUnit_HierarchySetShapeSettings : public FRigUnit_DynamicHierarchyBaseMutable
+{
+	GENERATED_BODY()
+
+	FRigUnit_HierarchySetShapeSettings()
+	{
+		Item = FRigElementKey(NAME_None, ERigElementType::Control);
+	}
+
+	/*
+	 * The item to change
+	 */
+	UPROPERTY(meta = (Input))
+	FRigElementKey Item;
+
+	/*
+	 * The shape settings for the control
+	 */
+	UPROPERTY(meta = (Input))
+	FRigUnit_HierarchyAddControl_ShapeSettings Settings;
+
+	RIGVM_METHOD()
+	virtual void Execute(const FRigUnitContext& Context) override;
+};
