Control Rig: Offer editing of name for AnimationChannels

diff --git a/Engine/Plugins/Animation/ControlRig/Source/ControlRigEditor/Private/ControlRigElementDetails.cpp b/Engine/Plugins/Animation/ControlRig/Source/ControlRigEditor/Private/ControlRigElementDetails.cpp
--- a/Engine/Plugins/Animation/ControlRig/Source/ControlRigEditor/Private/ControlRigElementDetails.cpp
+++ b/Engine/Plugins/Animation/ControlRig/Source/ControlRigEditor/Private/ControlRigElementDetails.cpp
@@ -696,7 +696,11 @@ void FRigBaseElementDetails::SetName(const FText& InNewText, ETextCommit::Type I
 		
 		URigHierarchyController* Controller = Hierarchy->GetController(true);
 		check(Controller);
-		Controller->RenameElement(GetElementKey(), *InNewText.ToString(), true, true);
+		const FRigElementKey NewKey = Controller->RenameElement(GetElementKey(), *InNewText.ToString(), true, true);
+		if(NewKey.IsValid())
+		{
+			Controller->SelectElement(NewKey, true, true);
+		}
 	}
 }
 
@@ -2537,6 +2541,27 @@ void FRigControlElementDetails::CustomizeControl(IDetailLayoutBuilder& DetailBui
 		.IsEnabled(bIsEnabled && (PerElementInfos.Num() == 1))
 	];
 
+	if(bAllAnimationChannels)
+	{
+		ControlCategory.AddCustomRow(FText::FromString(TEXT("Script Name")))
+		.NameContent()
+		[
+			SNew(STextBlock)
+			.Text(FText::FromString(TEXT("Script Name")))
+			.Font(IDetailLayoutBuilder::GetDetailFont())
+			.IsEnabled(!bIsProcedural)
+		]
+		.ValueContent()
+		[
+			SNew(SInlineEditableTextBlock)
+			.Font(IDetailLayoutBuilder::GetDetailFont())
+			.Text(this, &FRigBaseElementDetails::GetName)
+			.OnTextCommitted(this, &FRigBaseElementDetails::SetName)
+			.OnVerifyTextChanged(this, &FRigBaseElementDetails::OnVerifyNameChanged)
+			.IsEnabled(!bIsProcedural && PerElementInfos.Num() == 1)
+		];
+	}
+
 	const TSharedRef<IPropertyUtilities> PropertyUtilities = DetailBuilder.GetPropertyUtilities();
 
 	// when control type changes, we have to refresh detail panel
