Control Rig: Set offset from closest bone doesn't stick

diff --git a/Engine/Plugins/Animation/ControlRig/Source/ControlRigEditor/Private/Editor/SRigHierarchy.cpp b/Engine/Plugins/Animation/ControlRig/Source/ControlRigEditor/Private/Editor/SRigHierarchy.cpp
--- a/Engine/Plugins/Animation/ControlRig/Source/ControlRigEditor/Private/Editor/SRigHierarchy.cpp
+++ b/Engine/Plugins/Animation/ControlRig/Source/ControlRigEditor/Private/Editor/SRigHierarchy.cpp
@@ -2935,6 +2935,9 @@ void SRigHierarchy::HandleSetInitialTransformFromClosestBone()
 		{
 			if (URigHierarchy* DebuggedHierarchy = GetHierarchy())
 			{
+				URigHierarchy* Hierarchy = GetDefaultHierarchy();
+  				check(Hierarchy);
+
 				FScopedTransaction Transaction(LOCTEXT("HierarchySetInitialTransforms", "Set Initial Transforms"));
 
 				TArray<FRigElementKey> SelectedKeys = GetSelectedKeys();
@@ -2970,12 +2973,12 @@ void SRigHierarchy::HandleSetInitialTransformFromClosestBone()
 
 					if (SelectedKey.Type == ERigElementType::Control)
 					{
-						if(FRigControlElement* ControlElement = GetHierarchy()->Find<FRigControlElement>(SelectedKey))
+						if(FRigControlElement* ControlElement = Hierarchy->Find<FRigControlElement>(SelectedKey))
 						{
-							GetHierarchy()->SetControlOffsetTransform(ControlElement, LocalTransform, ERigTransformType::InitialLocal, true, true, false, true);
-							GetHierarchy()->SetControlOffsetTransform(ControlElement, LocalTransform, ERigTransformType::CurrentLocal, true, true, false, true);
-							GetHierarchy()->SetTransform(ControlElement, FTransform::Identity, ERigTransformType::InitialLocal, true, true, false, true);
-							GetHierarchy()->SetTransform(ControlElement, FTransform::Identity, ERigTransformType::CurrentLocal, true, true, false, true);
+							Hierarchy->SetControlOffsetTransform(ControlElement, LocalTransform, ERigTransformType::InitialLocal, true, true, false, true);
+							Hierarchy->SetControlOffsetTransform(ControlElement, LocalTransform, ERigTransformType::CurrentLocal, true, true, false, true);
+							Hierarchy->SetTransform(ControlElement, FTransform::Identity, ERigTransformType::InitialLocal, true, true, false, true);
+							Hierarchy->SetTransform(ControlElement, FTransform::Identity, ERigTransformType::CurrentLocal, true, true, false, true);
 						}
 						if(FRigControlElement* ControlElement = DebuggedHierarchy->Find<FRigControlElement>(SelectedKey))
 						{
@@ -2990,10 +2993,10 @@ void SRigHierarchy::HandleSetInitialTransformFromClosestBone()
 					{
 						FTransform InitialTransform = LocalTransform;
 
-						if(FRigTransformElement* TransformElement = GetHierarchy()->Find<FRigTransformElement>(SelectedKey))
+						if(FRigTransformElement* TransformElement = Hierarchy->Find<FRigTransformElement>(SelectedKey))
 						{
-							GetHierarchy()->SetTransform(TransformElement, LocalTransform, ERigTransformType::InitialLocal, true, true, false, true);
-							GetHierarchy()->SetTransform(TransformElement, LocalTransform, ERigTransformType::CurrentLocal, true, true, false, true);
+							Hierarchy->SetTransform(TransformElement, LocalTransform, ERigTransformType::InitialLocal, true, true, false, true);
+							Hierarchy->SetTransform(TransformElement, LocalTransform, ERigTransformType::CurrentLocal, true, true, false, true);
 						}
 						if(FRigTransformElement* TransformElement = DebuggedHierarchy->Find<FRigTransformElement>(SelectedKey))
 						{
