Control Rig: Fix typo in node display name

diff --git a/Engine/Plugins/Animation/ControlRig/Source/ControlRig/Private/Units/Hierarchy/RigUnit_SetControlColor.h b/Engine/Plugins/Animation/ControlRig/Source/ControlRig/Private/Units/Hierarchy/RigUnit_SetControlColor.h
--- a/Engine/Plugins/Animation/ControlRig/Source/ControlRig/Private/Units/Hierarchy/RigUnit_SetControlColor.h
+++ b/Engine/Plugins/Animation/ControlRig/Source/ControlRig/Private/Units/Hierarchy/RigUnit_SetControlColor.h
@@ -8,7 +8,7 @@
 /**
  * GetControlColor is used to retrieve the color of control
  */
-USTRUCT(meta=(DisplayName="Set Gontrol Color", Category="Controls", DocumentationPolicy="Strict", Keywords = "GetControlColor,GetGizmoColor", TemplateName="GetControlColor", NodeColor = "0.0 0.36470600962638855 1.0"))
+USTRUCT(meta=(DisplayName="Set Control Color", Category="Controls", DocumentationPolicy="Strict", Keywords = "GetControlColor,GetGizmoColor", TemplateName="GetControlColor", NodeColor = "0.0 0.36470600962638855 1.0"))
 struct CONTROLRIG_API FRigUnit_GetControlColor : public FRigUnit
 {
 	GENERATED_BODY()
