Control Rig: Switch Parent crashes with Spawn Controls

diff --git a/Engine/Plugins/Animation/ControlRig/Source/ControlRigEditor/Private/EditMode/ControlRigEditMode.cpp b/Engine/Plugins/Animation/ControlRig/Source/ControlRigEditor/Private/EditMode/ControlRigEditMode.cpp
--- a/Engine/Plugins/Animation/ControlRig/Source/ControlRigEditor/Private/EditMode/ControlRigEditMode.cpp
+++ b/Engine/Plugins/Animation/ControlRig/Source/ControlRigEditor/Private/EditMode/ControlRigEditMode.cpp
@@ -4368,13 +4368,16 @@ void FControlRigEditMode::TickManipulatableObjects(float DeltaTime)
 				{
 					// "Copy Pose from Mesh" requires AnimInstance::PreUpdate() to copy the parent bone transforms.
 					// have to TickAnimation() to ensure that PreUpdate() is called on all anim instances
-					SkeletalMeshToUpdate->TickAnimation(0.0f, false);
-					SkeletalMeshToUpdate->RefreshBoneTransforms();
-					SkeletalMeshToUpdate->RefreshFollowerComponents	();
-					SkeletalMeshToUpdate->UpdateComponentToWorld();
-					SkeletalMeshToUpdate->FinalizeBoneTransform();
-					SkeletalMeshToUpdate->MarkRenderTransformDirty();
-					SkeletalMeshToUpdate->MarkRenderDynamicDataDirty();
+					if (!SkeletalMeshToUpdate->GetAnimInstance()->IsUpdatingAnimation())
+					{
+						SkeletalMeshToUpdate->TickAnimation(0.0f, false);
+						SkeletalMeshToUpdate->RefreshBoneTransforms();
+						SkeletalMeshToUpdate->RefreshFollowerComponents	();
+						SkeletalMeshToUpdate->UpdateComponentToWorld();
+						SkeletalMeshToUpdate->FinalizeBoneTransform();
+						SkeletalMeshToUpdate->MarkRenderTransformDirty();
+						SkeletalMeshToUpdate->MarkRenderDynamicDataDirty();
+					}
 				}
 			}
 		}
diff --git a/Engine/Source/Runtime/Engine/Classes/Animation/AnimInstance.h b/Engine/Source/Runtime/Engine/Classes/Animation/AnimInstance.h
--- a/Engine/Source/Runtime/Engine/Classes/Animation/AnimInstance.h
+++ b/Engine/Source/Runtime/Engine/Classes/Animation/AnimInstance.h
@@ -951,6 +951,9 @@ public:
 #if DO_CHECK
 	/** Is this animation currently running post update */
 	bool IsPostUpdatingAnimation() const { return bPostUpdatingAnimation; }
+
+	/** Is this animation currently running update */
+	bool IsUpdatingAnimation() const { return bUpdatingAnimation; }
 #endif
 
 	/** Set RootMotionMode */
