Control Rig: Fixes for fallout from UI optimizations

diff --git a/Engine/Plugins/Animation/ControlRig/Source/ControlRigDeveloper/Private/Graph/ControlRigGraphNode.cpp b/Engine/Plugins/Animation/ControlRig/Source/ControlRigDeveloper/Private/Graph/ControlRigGraphNode.cpp
--- a/Engine/Plugins/Animation/ControlRig/Source/ControlRigDeveloper/Private/Graph/ControlRigGraphNode.cpp
+++ b/Engine/Plugins/Animation/ControlRig/Source/ControlRigDeveloper/Private/Graph/ControlRigGraphNode.cpp
@@ -479,7 +479,7 @@ URigVMPin* UControlRigGraphNode::FindModelPinFromGraphPin(const UEdGraphPin* InG
 		}
 	}
 
-	return nullptr;
+	return GetModelPinFromPinPath(InGraphPin->GetName());
 }
 
 UEdGraphPin* UControlRigGraphNode::FindGraphPinFromModelPin(const URigVMPin* InModelPin, bool bAsInput) const
@@ -494,6 +494,18 @@ UEdGraphPin* UControlRigGraphNode::FindGraphPinFromModelPin(const URigVMPin* InM
 		return bAsInput ? Pair->InputPin : Pair->OutputPin;
 	}
 
+	const FString PinPath = InModelPin->GetPinPath();
+	for(UEdGraphPin* GraphPin : Pins)
+	{
+		if((GraphPin->Direction == EGPD_Input) == bAsInput)
+		{
+			if(GraphPin->GetName() == PinPath)
+			{
+				return GraphPin;
+			}
+		}
+	}
+	
 	return nullptr;
 }
 
@@ -634,6 +646,12 @@ void UControlRigGraphNode::SyncGraphNodeNameWithModelNodeName(const URigVMNode*
 	Rename(*InModelNode->GetName());
 	ModelNodePath = InModelNode->GetNodePath();
 	SyncGraphNodeTitleWithModelNodeTitle();
+
+	TArray<URigVMPin*> AllModelPins = InModelNode->GetAllPinsRecursively();
+	for(const URigVMPin* ModelPin : AllModelPins)
+	{
+		SynchronizeGraphPinNameWithModelPin(ModelPin);
+	}
 }
 
 bool UControlRigGraphNode::CreateGraphPinFromModelPin(const URigVMPin* InModelPin, EEdGraphPinDirection InDirection,
@@ -671,20 +689,23 @@ bool UControlRigGraphNode::CreateGraphPinFromModelPin(const URigVMPin* InModelPi
 
 	UEdGraphPin* InputPin = nullptr;
 	UEdGraphPin* OutputPin = nullptr;
-	
+
+	bool bResult = false;
 	if (InDirection == EGPD_Input && PairConst.InputPin == nullptr)
 	{
 		InputPin = CreatePinLambda(InModelPin, EGPD_Input, InParentPin);
+		bResult = true;
 	}
 	if (InDirection == EGPD_Output && PairConst.OutputPin == nullptr)
 	{
 		OutputPin = CreatePinLambda(InModelPin, EGPD_Output, InParentPin);
+		bResult = true;
 	}
 
 	FPinPair& Pair = CachedPins.FindChecked((URigVMPin*)InModelPin);
 	Pair.InputPin = InputPin != nullptr ? InputPin : Pair.InputPin;
 	Pair.OutputPin = OutputPin != nullptr ? OutputPin : Pair.OutputPin;
-	return Pair.IsValid();
+	return Pair.IsValid() && bResult;
 }
 
 void UControlRigGraphNode::RemoveGraphSubPins(UEdGraphPin* InParentPin, const TArray<UEdGraphPin*>& InPinsToKeep)
@@ -919,6 +940,17 @@ bool UControlRigGraphNode::ModelPinRemoved(const URigVMPin* InModelPin)
 
 bool UControlRigGraphNode::ModelPinRemoved_Internal(const URigVMPin* InModelPin)
 {
+	auto RemoveGraphPin = [this](const URigVMPin* InModelPin, bool bAsInput)
+	{
+		if(UEdGraphPin* GraphPin = FindGraphPinFromModelPin(InModelPin, bAsInput))
+		{
+			RemoveGraphSubPins(GraphPin);
+			Pins.Remove(GraphPin);
+		}
+	};
+	RemoveGraphPin(InModelPin, true);
+	RemoveGraphPin(InModelPin, false);
+		
 	const bool bResult = PinPathToModelPin.Remove(InModelPin->GetPinPath()) > 0;
 	CachedPins.Remove(InModelPin);
 	return bResult;
diff --git a/Engine/Plugins/Animation/ControlRig/Source/ControlRigEditor/Private/Graph/SControlRigGraphNode.cpp b/Engine/Plugins/Animation/ControlRig/Source/ControlRigEditor/Private/Graph/SControlRigGraphNode.cpp
--- a/Engine/Plugins/Animation/ControlRig/Source/ControlRigEditor/Private/Graph/SControlRigGraphNode.cpp
+++ b/Engine/Plugins/Animation/ControlRig/Source/ControlRigEditor/Private/Graph/SControlRigGraphNode.cpp
@@ -326,6 +326,7 @@ void SControlRigGraphNode::AddPin(const TSharedRef<SGraphPin>& PinToAdd)
 void SControlRigGraphNode::CreateStandardPinWidget(UEdGraphPin* CurPin)
 {
 	bool bShowPin = true;
+	FString CPPType;
 	if(const UControlRigGraphNode* RigGraphNode = Cast<UControlRigGraphNode>(GraphNode))
 	{
 		if(const URigVMPin* ModelPin = RigGraphNode->FindModelPinFromGraphPin(CurPin))
@@ -335,6 +336,8 @@ void SControlRigGraphNode::CreateStandardPinWidget(UEdGraphPin* CurPin)
 				ModelPin->GetDirection() == ERigVMPinDirection::Input ||
 				ModelPin->GetDirection() == ERigVMPinDirection::Output ||
 				ModelPin->GetDirection() == ERigVMPinDirection::IO;
+			
+			CPPType = ModelPin->GetCPPType();
 		}
 	}
 	
@@ -342,15 +345,30 @@ void SControlRigGraphNode::CreateStandardPinWidget(UEdGraphPin* CurPin)
 	{
 		// Do we have this pin in our list of pins to keep?
 		TSharedPtr<SGraphPin> NewPin;
-		const TSharedRef<SGraphPin>* RecycledPin = PinsToKeep.Find(CurPin);
-		if (!RecycledPin)
+		const TSharedRef<SGraphPin>* RecycledPinPtr = PinsToKeep.Find(CurPin);
+		if (RecycledPinPtr)
 		{
-			NewPin = CreatePinWidget(CurPin);
-			check(NewPin.IsValid());
+			const TSharedRef<SGraphPin>& RecycledPin = *RecycledPinPtr;
+			const TSharedPtr<FPinInfoMetaData> PinInfoMetaData = RecycledPin->GetMetaData<FPinInfoMetaData>();
+			if(PinInfoMetaData.IsValid())
+			{
+				if(PinInfoMetaData->CPPType == CPPType)
+				{
+					NewPin = RecycledPin;
+				}
+			}
 		}
-		else
+
+		if(!NewPin.IsValid())
 		{
-			NewPin = *RecycledPin;
+			if(RecycledPinPtr)
+			{
+				(*RecycledPinPtr)->InvalidateGraphData();
+			}
+			NewPin = CreatePinWidget(CurPin);
+			NewPin->AddMetadata(MakeShared<FPinInfoMetaData>(CPPType));
+			check(NewPin.IsValid());
+			PinsToKeep.Remove(CurPin);
 		}
 
 		AddPin(NewPin.ToSharedRef());
@@ -898,10 +916,12 @@ void SControlRigGraphNode::HandleNodePinsChanged()
 			if (LivePin->Direction == EGPD_Input && PinInfoPtr->InputPinWidget.IsValid())
 			{
 				PinsToKeep.Add(LivePin, PinInfoPtr->InputPinWidget.ToSharedRef());
+				LocalPinsToDelete.Remove(PinInfoPtr->InputPinWidget->GetPinObj());
 			}
 			if (LivePin->Direction == EGPD_Output && PinInfoPtr->OutputPinWidget.IsValid())
 			{
 				PinsToKeep.Add(LivePin, PinInfoPtr->OutputPinWidget.ToSharedRef());
+				LocalPinsToDelete.Remove(PinInfoPtr->OutputPinWidget->GetPinObj());
 			}
 		}
 		LocalPinsToDelete.Remove(LivePin);
diff --git a/Engine/Plugins/Animation/ControlRig/Source/ControlRigEditor/Private/Graph/SControlRigGraphNode.h b/Engine/Plugins/Animation/ControlRig/Source/ControlRigEditor/Private/Graph/SControlRigGraphNode.h
--- a/Engine/Plugins/Animation/ControlRig/Source/ControlRigEditor/Private/Graph/SControlRigGraphNode.h
+++ b/Engine/Plugins/Animation/ControlRig/Source/ControlRigEditor/Private/Graph/SControlRigGraphNode.h
@@ -124,6 +124,22 @@ private:
 		bool bExpanded;
 		bool bAutoHeight;
 	};
+
+	
+	/**
+	 * Simple tagging metadata
+	 */
+	class FPinInfoMetaData : public ISlateMetaData
+	{
+	public:
+		SLATE_METADATA_TYPE(FPinInfoMetaData, ISlateMetaData)
+
+		FPinInfoMetaData(const FString& InCPPType)
+		: CPPType(InCPPType)
+		{}
+
+		FString CPPType;
+	};
 	
 	TArray<FPinInfo> PinInfos;
 	TWeakObjectPtr<URigVMNode> ModelNode;
