Control Rig: Fix FControlRigEditMode::TickManipulatableObjects being called when SkeletalMeshComponent is already updating animation

diff --git a/Engine/Plugins/Animation/ControlRig/Source/ControlRigEditor/Private/EditMode/ControlRigEditMode.cpp b/Engine/Plugins/Animation/ControlRig/Source/ControlRigEditor/Private/EditMode/ControlRigEditMode.cpp
--- a/Engine/Plugins/Animation/ControlRig/Source/ControlRigEditor/Private/EditMode/ControlRigEditMode.cpp
+++ b/Engine/Plugins/Animation/ControlRig/Source/ControlRigEditor/Private/EditMode/ControlRigEditMode.cpp
@@ -3610,7 +3610,7 @@ void FControlRigEditMode::OnHierarchyModified(ERigHierarchyNotification InNotif,
 		}
 		case ERigHierarchyNotification::ParentWeightsChanged:
 		{
-			TickManipulatableObjects(0.f);
+			// TickManipulatableObjects(0.f);
 			break;
 		}
 		case ERigHierarchyNotification::InteractionBracketOpened:
@@ -4368,16 +4368,13 @@ void FControlRigEditMode::TickManipulatableObjects(float DeltaTime)
 				{
 					// "Copy Pose from Mesh" requires AnimInstance::PreUpdate() to copy the parent bone transforms.
 					// have to TickAnimation() to ensure that PreUpdate() is called on all anim instances
-					if (!SkeletalMeshToUpdate->GetAnimInstance()->IsUpdatingAnimation())
-					{
-						SkeletalMeshToUpdate->TickAnimation(0.0f, false);
-						SkeletalMeshToUpdate->RefreshBoneTransforms();
-						SkeletalMeshToUpdate->RefreshFollowerComponents	();
-						SkeletalMeshToUpdate->UpdateComponentToWorld();
-						SkeletalMeshToUpdate->FinalizeBoneTransform();
-						SkeletalMeshToUpdate->MarkRenderTransformDirty();
-						SkeletalMeshToUpdate->MarkRenderDynamicDataDirty();
-					}
+					SkeletalMeshToUpdate->TickAnimation(0.0f, false);
+					SkeletalMeshToUpdate->RefreshBoneTransforms();
+					SkeletalMeshToUpdate->RefreshFollowerComponents	();
+					SkeletalMeshToUpdate->UpdateComponentToWorld();
+					SkeletalMeshToUpdate->FinalizeBoneTransform();
+					SkeletalMeshToUpdate->MarkRenderTransformDirty();
+					SkeletalMeshToUpdate->MarkRenderDynamicDataDirty();
 				}
 			}
 		}
