Control Rig: Allow wildcard for name search

diff --git a/Engine/Plugins/Animation/ControlRig/Source/ControlRig/Private/Rigs/RigHierarchyDefines.cpp b/Engine/Plugins/Animation/ControlRig/Source/ControlRig/Private/Rigs/RigHierarchyDefines.cpp
--- a/Engine/Plugins/Animation/ControlRig/Source/ControlRig/Private/Rigs/RigHierarchyDefines.cpp
+++ b/Engine/Plugins/Animation/ControlRig/Source/ControlRig/Private/Rigs/RigHierarchyDefines.cpp
@@ -2,6 +2,7 @@
 
 #include "Rigs/RigHierarchyDefines.h"
 #include "Rigs/RigHierarchy.h"
+#include "Misc/WildcardString.h"
 
 #include UE_INLINE_GENERATED_CPP_BY_NAME(RigHierarchyDefines)
 
@@ -166,6 +167,15 @@ FRigElementKeyCollection FRigElementKeyCollection::MakeFromName(
 	constexpr bool bTraverse = true;
 
 	const FString PartialNameString = InPartialName.ToString();
+	const FWildcardString WildcardString(PartialNameString);
+	if(WildcardString.ContainsWildcards())
+	{
+		return InHierarchy->GetKeysByPredicate([WildcardString, InElementTypes](const FRigBaseElement& InElement) -> bool
+		{
+			return InElement.IsTypeOf(static_cast<ERigElementType>(InElementTypes)) &&
+				   WildcardString.IsMatch(InElement.GetNameString());
+		}, bTraverse);
+	}
 	
 	return InHierarchy->GetKeysByPredicate([PartialNameString, InElementTypes](const FRigBaseElement& InElement) -> bool
 	{
