Control Rig: Visual representation of casts on connection lines

diff --git a/Engine/Plugins/Animation/ControlRig/Source/ControlRigEditor/Private/Graph/ControlRigConnectionDrawingPolicy.cpp b/Engine/Plugins/Animation/ControlRig/Source/ControlRigEditor/Private/Graph/ControlRigConnectionDrawingPolicy.cpp
--- a/Engine/Plugins/Animation/ControlRig/Source/ControlRigEditor/Private/Graph/ControlRigConnectionDrawingPolicy.cpp
+++ b/Engine/Plugins/Animation/ControlRig/Source/ControlRigEditor/Private/Graph/ControlRigConnectionDrawingPolicy.cpp
@@ -102,7 +102,6 @@ void FControlRigConnectionDrawingPolicy::DrawPinGeometries(TMap<TSharedRef<SWidg
 
 						const FVector2D StartPoint = FGeometryHelper::VerticalMiddleRightOf(LinkStartWidgetGeometry->Geometry) - FVector2D(StartFudgeX, 0.0f);
 						const FVector2D EndPoint = FGeometryHelper::VerticalMiddleLeftOf(LinkEndWidgetGeometry->Geometry) - FVector2D(EndFudgeX, 0);
-						const FVector2D FakeTangent = (EndPoint - StartPoint).GetSafeNormal();
 
 						static TArray<FVector2D> LinePoints;
 						LinePoints.SetNum(2);
@@ -113,7 +112,45 @@ void FControlRigConnectionDrawingPolicy::DrawPinGeometries(TMap<TSharedRef<SWidg
 					}
 					else
 					{
-						DrawSplineWithArrow(LinkStartWidgetGeometry->Geometry, LinkEndWidgetGeometry->Geometry, Params);
+						const float StartFudgeX = 4.0f;
+						const float EndFudgeX = 4.0f;
+						FVector2D StartPoint = FGeometryHelper::VerticalMiddleRightOf(LinkStartWidgetGeometry->Geometry) - FVector2D(StartFudgeX, 0.0f);
+						FVector2D EndPoint = FGeometryHelper::VerticalMiddleLeftOf(LinkEndWidgetGeometry->Geometry) - FVector2D(ArrowRadius.X - EndFudgeX, 0);
+
+						if(Params.bUserFlag1) // indicate a cast using a mid point image
+						{
+							FConnectionParams OtherParams;
+							DetermineWiringStyle(TargetPin, ThePin, OtherParams);
+
+							if(!OtherParams.WireColor.Equals(Params.WireColor, 0.001f))
+							{
+								const FVector2D MidPoint = (StartPoint + EndPoint) * 0.5;
+								const FVector2D OriginalSplineTangent = FVector2D(100.f, 0);
+								const FVector2D InnerStartPoint = StartPoint + OriginalSplineTangent;
+								const FVector2D InnerEndPoint = EndPoint - OriginalSplineTangent;
+
+								const FVector2D DeltaPos = InnerEndPoint - InnerStartPoint;
+								const double TangentLength = FMath::Min<float>(FMath::Abs<float>(DeltaPos.X), Settings->ForwardSplineHorizontalDeltaRange);
+								OtherParams.StartTangent = DeltaPos.GetSafeNormal() * TangentLength;
+
+								DrawSplineWithArrow(MidPoint, EndPoint, OtherParams);
+
+								EndPoint = MidPoint;
+								Params.EndTangent = OtherParams.StartTangent;
+
+								// Draw the icon
+								FSlateDrawElement::MakeBox(
+									DrawElementsList,
+									ArrowLayerID,
+									FPaintGeometry(MidPoint - ArrowRadius, CastImage->ImageSize * ZoomFactor, ZoomFactor),
+									CastImage,
+									ESlateDrawEffect::None,
+									OtherParams.WireColor
+								);
+							}
+						}
+						
+						DrawSplineWithArrow(StartPoint, EndPoint, Params);
 					}
 				}
 			}
@@ -336,6 +373,15 @@ void FControlRigConnectionDrawingPolicy::DetermineWiringStyle(UEdGraphPin* Outpu
 			{
 				Params.WireColor = Params.WireColor * 0.5f;
 			}
+
+			Params.bUserFlag1 = false;
+			if(!UseLowDetailConnections() && !bInjectionIsSelected)
+			{
+				if(OutputModelPin->GetTypeIndex() != InputModelPin->GetTypeIndex())
+				{
+					Params.bUserFlag1 = true;
+				}
+			}
 		}
 	}
 }
diff --git a/Engine/Plugins/Animation/ControlRig/Source/ControlRigEditor/Private/Graph/ControlRigConnectionDrawingPolicy.h b/Engine/Plugins/Animation/ControlRig/Source/ControlRigEditor/Private/Graph/ControlRigConnectionDrawingPolicy.h
--- a/Engine/Plugins/Animation/ControlRig/Source/ControlRigEditor/Private/Graph/ControlRigConnectionDrawingPolicy.h
+++ b/Engine/Plugins/Animation/ControlRig/Source/ControlRigEditor/Private/Graph/ControlRigConnectionDrawingPolicy.h
@@ -12,6 +12,11 @@ public:
 	FControlRigConnectionDrawingPolicy(int32 InBackLayerID, int32 InFrontLayerID, float InZoomFactor, const FSlateRect& InClippingRect, FSlateWindowElementList& InDrawElements, UEdGraph* InGraphObj)
 		: FKismetConnectionDrawingPolicy(InBackLayerID, InFrontLayerID, InZoomFactor, InClippingRect, InDrawElements, InGraphObj)
 	{
+		CastImage = FAppStyle::GetBrush( TEXT("GraphEditor.Cast_16x") );
+		ArrowImage = nullptr;
+		ArrowRadius = CastImage->ImageSize * InZoomFactor * 0.5f;
+		MidpointImage = nullptr;
+		MidpointRadius = CastImage->ImageSize * InZoomFactor * 0.5f;
 	}
 
 	virtual void SetIncompatiblePinDrawState(const TSharedPtr<SGraphPin>& StartPin, const TSet< TSharedRef<SWidget> >& VisiblePins) override;
@@ -37,4 +42,5 @@ private:
 	// Average of the positions of all pins connected to InPin
 	bool GetAverageConnectedPositionForPin(UEdGraphPin* InPin, FVector2D& OutPos) const;
 
+	const FSlateBrush* CastImage;
 };
diff --git a/Engine/Source/Editor/GraphEditor/Private/ConnectionDrawingPolicy.cpp b/Engine/Source/Editor/GraphEditor/Private/ConnectionDrawingPolicy.cpp
--- a/Engine/Source/Editor/GraphEditor/Private/ConnectionDrawingPolicy.cpp
+++ b/Engine/Source/Editor/GraphEditor/Private/ConnectionDrawingPolicy.cpp
@@ -243,8 +243,9 @@ void FConnectionDrawingPolicy::DrawConnection(int32 LayerId, const FVector2D& St
 	const FVector2D& P1 = End;
 
 	const FVector2D SplineTangent = ComputeSplineTangent(P0, P1);
-	const FVector2D P0Tangent = (Params.StartDirection == EGPD_Output) ? SplineTangent : -SplineTangent;
-	const FVector2D P1Tangent = (Params.EndDirection == EGPD_Input) ? SplineTangent : -SplineTangent;
+
+	const FVector2D P0Tangent = (Params.StartTangent.IsNearlyZero()) ? ((Params.StartDirection == EGPD_Output) ? SplineTangent : -SplineTangent) : Params.StartTangent;
+	const FVector2D P1Tangent = (Params.EndTangent.IsNearlyZero()) ? ((Params.EndDirection == EGPD_Input) ? SplineTangent : -SplineTangent) : Params.EndTangent;
 
 	if (Settings->bTreatSplinesLikePins)
 	{
diff --git a/Engine/Source/Editor/GraphEditor/Public/ConnectionDrawingPolicy.h b/Engine/Source/Editor/GraphEditor/Public/ConnectionDrawingPolicy.h
--- a/Engine/Source/Editor/GraphEditor/Public/ConnectionDrawingPolicy.h
+++ b/Engine/Source/Editor/GraphEditor/Public/ConnectionDrawingPolicy.h
@@ -68,6 +68,9 @@ struct GRAPHEDITOR_API FConnectionParams
 	EEdGraphPinDirection StartDirection;
 	EEdGraphPinDirection EndDirection;
 
+	FVector2D StartTangent;
+	FVector2D EndTangent;
+
 	FConnectionParams()
 		: WireColor(FLinearColor::White)
 		, AssociatedPin1(nullptr)
@@ -78,6 +81,8 @@ struct GRAPHEDITOR_API FConnectionParams
 		, bUserFlag2(false)
 		, StartDirection(EGPD_Output)
 		, EndDirection(EGPD_Input)
+		, StartTangent(FVector2D::ZeroVector)
+		, EndTangent(FVector2D::ZeroVector)
 	{
 	}
 };
