Fix FScopedTransaction instances missing a variable name, and add NODISCARD to FScopedTransaction to prevent this from happening in future.

diff --git a/Engine/Plugins/Animation/ControlRig/Source/ControlRigEditor/Private/ControlRigElementDetails.cpp b/Engine/Plugins/Animation/ControlRig/Source/ControlRigEditor/Private/ControlRigElementDetails.cpp
--- a/Engine/Plugins/Animation/ControlRig/Source/ControlRigEditor/Private/ControlRigElementDetails.cpp
+++ b/Engine/Plugins/Animation/ControlRig/Source/ControlRigEditor/Private/ControlRigElementDetails.cpp
@@ -3043,7 +3043,7 @@ void FRigControlElementDetails::CustomizeAnimationChannels(IDetailLayoutBuilder&
 						{
 							if(URigHierarchyController* Controller = HierarchyToChange->GetController(true))
 							{
-								FScopedTransaction(LOCTEXT("DeleteAnimationChannels", "Delete Animation Channels"));
+								FScopedTransaction Transaction(LOCTEXT("DeleteAnimationChannels", "Delete Animation Channels"));
 								HierarchyToChange->Modify();
 								
 								for(const FRigElementKey& KeyToRemove : ChildElementKeys)
@@ -3064,7 +3064,7 @@ void FRigControlElementDetails::CustomizeAnimationChannels(IDetailLayoutBuilder&
 						{
 							if(URigHierarchyController* Controller = HierarchyToChange->GetController(true))
 							{
-								FScopedTransaction(LOCTEXT("MoveAnimationChannelUpTransaction", "Move Animation Channel Up"));
+								FScopedTransaction Transaction(LOCTEXT("MoveAnimationChannelUpTransaction", "Move Animation Channel Up"));
 								HierarchyToChange->Modify();
 								
 								for(const FRigElementKey& KeyToMove : ChildElementKeys)
@@ -3083,7 +3083,7 @@ void FRigControlElementDetails::CustomizeAnimationChannels(IDetailLayoutBuilder&
 						{
 							if(URigHierarchyController* Controller = HierarchyToChange->GetController(true))
 							{
-								FScopedTransaction(LOCTEXT("MoveAnimationChannelDownTransaction", "Move Animation Channel Down"));
+								FScopedTransaction Transaction(LOCTEXT("MoveAnimationChannelDownTransaction", "Move Animation Channel Down"));
 								HierarchyToChange->Modify();
 								
 								for(const FRigElementKey& KeyToMove : ChildElementKeys)
diff --git a/Engine/Source/Editor/SkeletonEditor/Private/SkeletonTreeBoneItem.cpp b/Engine/Source/Editor/SkeletonEditor/Private/SkeletonTreeBoneItem.cpp
--- a/Engine/Source/Editor/SkeletonEditor/Private/SkeletonTreeBoneItem.cpp
+++ b/Engine/Source/Editor/SkeletonEditor/Private/SkeletonTreeBoneItem.cpp
@@ -409,7 +409,7 @@ void FSkeletonTreeBoneItem::OnBlendSliderCommitted(float NewValue, ETextCommit::
 
 	if (BlendProfile)
 	{
-		FScopedTransaction(LOCTEXT("SetBlendProfileValue", "Set Blend Profile Value"));
+		FScopedTransaction Transaction(LOCTEXT("SetBlendProfileValue", "Set Blend Profile Value"));
 		BlendProfile->SetFlags(RF_Transactional);
 		BlendProfile->Modify();
 
diff --git a/Engine/Source/Editor/SkeletonEditor/Private/SkeletonTreeVirtualBoneItem.cpp b/Engine/Source/Editor/SkeletonEditor/Private/SkeletonTreeVirtualBoneItem.cpp
--- a/Engine/Source/Editor/SkeletonEditor/Private/SkeletonTreeVirtualBoneItem.cpp
+++ b/Engine/Source/Editor/SkeletonEditor/Private/SkeletonTreeVirtualBoneItem.cpp
@@ -188,7 +188,7 @@ void FSkeletonTreeVirtualBoneItem::OnBlendSliderCommitted(float NewValue, ETextC
 
 	if (BlendProfile)
 	{
-		FScopedTransaction(LOCTEXT("SetBlendProfileValue", "Set Blend Profile Value"));
+		FScopedTransaction Transaction(LOCTEXT("SetBlendProfileValue", "Set Blend Profile Value"));
 		BlendProfile->SetFlags(RF_Transactional);
 		BlendProfile->Modify();
 
diff --git a/Engine/Source/Editor/UnrealEd/Public/ScopedTransaction.h b/Engine/Source/Editor/UnrealEd/Public/ScopedTransaction.h
--- a/Engine/Source/Editor/UnrealEd/Public/ScopedTransaction.h
+++ b/Engine/Source/Editor/UnrealEd/Public/ScopedTransaction.h
@@ -18,14 +18,14 @@ public:
 	 * @param PrimaryObject				The main object being edited (if known)
 	 * @param bShouldActuallyTransact	If false the transaction is ignored. Good for creating optional transactions based on condition
 	 */
-	UNREALED_API FScopedTransaction(const TCHAR* TransactionContext, const FText& SessionName, UObject* PrimaryObject, const bool bShouldActuallyTransact = true);
+	UE_NODISCARD_CTOR UNREALED_API FScopedTransaction(const TCHAR* TransactionContext, const FText& SessionName, UObject* PrimaryObject, const bool bShouldActuallyTransact = true);
 
 	/**
 	 * Construct an FScopedTransaction with minimal context 
 	 * @param SessionName				The description of this transaction, ie what is occurring
 	 * @param bShouldActuallyTransact	If false the transaction is ignored. Good for creating optional transactions based on condition 
 	 */
-	UNREALED_API FScopedTransaction(const FText& SessionName, const bool bShouldActuallyTransact = true);
+	UE_NODISCARD_CTOR UNREALED_API FScopedTransaction(const FText& SessionName, const bool bShouldActuallyTransact = true);
 	UNREALED_API ~FScopedTransaction();
 
 	/**
