Control Rig: Item Array node

diff --git a/Engine/Plugins/Animation/ControlRig/Source/ControlRig/Private/Units/Hierarchy/RigUnit_BoneName.cpp b/Engine/Plugins/Animation/ControlRig/Source/ControlRig/Private/Units/Hierarchy/RigUnit_BoneName.cpp
--- a/Engine/Plugins/Animation/ControlRig/Source/ControlRig/Private/Units/Hierarchy/RigUnit_BoneName.cpp
+++ b/Engine/Plugins/Animation/ControlRig/Source/ControlRig/Private/Units/Hierarchy/RigUnit_BoneName.cpp
@@ -9,6 +9,10 @@ FRigUnit_Item_Execute()
 {
 }
 
+FRigUnit_ItemArray_Execute()
+{
+}
+
 FRigUnit_BoneName_Execute()
 {
 }
diff --git a/Engine/Plugins/Animation/ControlRig/Source/ControlRig/Private/Units/Hierarchy/RigUnit_BoneName.h b/Engine/Plugins/Animation/ControlRig/Source/ControlRig/Private/Units/Hierarchy/RigUnit_BoneName.h
--- a/Engine/Plugins/Animation/ControlRig/Source/ControlRig/Private/Units/Hierarchy/RigUnit_BoneName.h
+++ b/Engine/Plugins/Animation/ControlRig/Source/ControlRig/Private/Units/Hierarchy/RigUnit_BoneName.h
@@ -26,6 +26,27 @@ struct CONTROLRIG_API FRigUnit_Item : public FRigUnit
 	FRigElementKey Item;
 };
 
+/**
+ * The Item Array node is used to share an array of items across the graph
+ */
+USTRUCT(meta=(DisplayName="Item Array", Category="Hierarchy", NodeColor = "0.4627450108528137 1.0 0.3294120132923126", DocumentationPolicy = "Strict", Constant))
+struct CONTROLRIG_API FRigUnit_ItemArray : public FRigUnit
+{
+	GENERATED_BODY()
+
+	FRigUnit_ItemArray()
+	{}
+
+	RIGVM_METHOD()
+	virtual void Execute(const FRigUnitContext& Context) override;
+
+	/**
+	 * The items
+	 */
+	UPROPERTY(meta = (Input, Output, ExpandByDefault))
+	TArray<FRigElementKey> Items;
+};
+
 /**
  * BoneName is used to represent a bone name in the graph
  */
diff --git a/Engine/Plugins/Animation/ControlRig/Source/ControlRigEditor/Private/Editor/ControlRigEditor.cpp b/Engine/Plugins/Animation/ControlRig/Source/ControlRigEditor/Private/Editor/ControlRigEditor.cpp
--- a/Engine/Plugins/Animation/ControlRig/Source/ControlRigEditor/Private/Editor/ControlRigEditor.cpp
+++ b/Engine/Plugins/Animation/ControlRig/Source/ControlRigEditor/Private/Editor/ControlRigEditor.cpp
@@ -5656,12 +5656,9 @@ void FControlRigEditor::CreateRigHierarchyToGraphDragAndDropMenu() const
 										{
 											Controller->OpenUndoBracket(TEXT("Create Item Array From Selection"));
 
-											static const FString ItemArrayCPPType = FString::Printf(TEXT("TArray<%s>"), *RigVMTypeUtils::GetUniqueStructTypeName(FRigElementKey::StaticStruct()));
-											static const FString ItemArrayObjectPath = FRigElementKey::StaticStruct()->GetPathName();
-							
-											if (URigVMNode* ItemsNode = Controller->AddFreeRerouteNode(true, ItemArrayCPPType, *ItemArrayObjectPath, false, NAME_None, FString(), NodePosition))
+											if (URigVMNode* ItemsNode = Controller->AddUnitNode(FRigUnit_ItemArray::StaticStruct(), TEXT("Execute"), NodePosition))
 											{
-												if (URigVMPin* ItemsPin = ItemsNode->FindPin(TEXT("Value")))
+												if (URigVMPin* ItemsPin = ItemsNode->FindPin(TEXT("Items")))
 												{
 													Controller->SetArrayPinSize(ItemsPin->GetPinPath(), DraggedKeys.Num());
 
