Control Rig: Ability to restrict space switching

diff --git a/Engine/Plugins/Animation/ControlRig/Source/ControlRig/Private/Rigs/RigHierarchyController.cpp b/Engine/Plugins/Animation/ControlRig/Source/ControlRig/Private/Rigs/RigHierarchyController.cpp
--- a/Engine/Plugins/Animation/ControlRig/Source/ControlRig/Private/Rigs/RigHierarchyController.cpp
+++ b/Engine/Plugins/Animation/ControlRig/Source/ControlRig/Private/Rigs/RigHierarchyController.cpp
@@ -2255,6 +2255,14 @@ bool URigHierarchyController::AddParent(FRigBaseElement* InChild, FRigBaseElemen
 
 			bMaintainGlobalTransform = false;
 		}
+
+		if(ChildControlElement->Settings.bRestrictSpaceSwitching)
+		{
+			if(ChildControlElement->Settings.Customization.AvailableSpaces.Contains(InParent->GetKey()))
+			{
+				return false;
+			}
+		}
 	}
 
 	if(Hierarchy->IsParentedTo(InParent, InChild))
diff --git a/Engine/Plugins/Animation/ControlRig/Source/ControlRig/Private/Rigs/RigHierarchyElements.cpp b/Engine/Plugins/Animation/ControlRig/Source/ControlRig/Private/Rigs/RigHierarchyElements.cpp
--- a/Engine/Plugins/Animation/ControlRig/Source/ControlRig/Private/Rigs/RigHierarchyElements.cpp
+++ b/Engine/Plugins/Animation/ControlRig/Source/ControlRig/Private/Rigs/RigHierarchyElements.cpp
@@ -740,6 +740,7 @@ FRigControlSettings::FRigControlSettings()
 , ControlEnum(nullptr)
 , Customization()
 , bGroupWithParentControl(false)
+, bRestrictSpaceSwitching(false)
 {
 	// rely on the default provided by the shape definition
 	ShapeName = FControlRigShapeDefinition().ShapeName; 
@@ -783,6 +784,7 @@ void FRigControlSettings::Save(FArchive& Ar)
 	Ar << Customization.AvailableSpaces;
 	Ar << DrivenControls;
 	Ar << bGroupWithParentControl;
+	Ar << bRestrictSpaceSwitching;
 }
 
 void FRigControlSettings::Load(FArchive& Ar)
@@ -933,6 +935,15 @@ void FRigControlSettings::Load(FArchive& Ar)
 			ControlType == ERigControlType::Vector2D
 		);
 	}
+
+	if (Ar.CustomVer(FControlRigObjectVersion::GUID) >= FControlRigObjectVersion::RestrictSpaceSwitchingForControls)
+	{
+		Ar << bRestrictSpaceSwitching;
+	}
+	else
+	{
+		bRestrictSpaceSwitching = false;
+	}
 }
 
 uint32 GetTypeHash(const FRigControlSettings& Settings)
@@ -950,6 +961,7 @@ uint32 GetTypeHash(const FRigControlSettings& Settings)
 	Hash = HashCombine(Hash, GetTypeHash(Settings.ControlEnum));
 	Hash = HashCombine(Hash, GetTypeHash(Settings.DrivenControls));
 	Hash = HashCombine(Hash, GetTypeHash(Settings.bGroupWithParentControl));
+	Hash = HashCombine(Hash, GetTypeHash(Settings.bRestrictSpaceSwitching));
 	return Hash;
 }
 
@@ -1023,7 +1035,11 @@ bool FRigControlSettings::operator==(const FRigControlSettings& InOther) const
 	{
 		return false;
 	}
-	
+	if(bRestrictSpaceSwitching != InOther.bRestrictSpaceSwitching)
+	{
+		return false;
+	}
+
 	const FTransform MinimumTransform = MinimumValue.GetAsTransform(ControlType, PrimaryAxis);
 	const FTransform OtherMinimumTransform = InOther.MinimumValue.GetAsTransform(ControlType, PrimaryAxis);
 	if(!MinimumTransform.Equals(OtherMinimumTransform, 0.001))
diff --git a/Engine/Plugins/Animation/ControlRig/Source/ControlRig/Public/ControlRigObjectVersion.h b/Engine/Plugins/Animation/ControlRig/Source/ControlRig/Public/ControlRigObjectVersion.h
--- a/Engine/Plugins/Animation/ControlRig/Source/ControlRig/Public/ControlRigObjectVersion.h
+++ b/Engine/Plugins/Animation/ControlRig/Source/ControlRig/Public/ControlRigObjectVersion.h
@@ -75,6 +75,9 @@ struct CONTROLRIG_API FControlRigObjectVersion
 		// Converted library nodes to templates
 		LibraryNodeTemplates,
 
+		// Controls to be able specify space switch targets
+		RestrictSpaceSwitchingForControls,
+
 		// -----<new versions can be added above this line>-------------------------------------------------
 		VersionPlusOne,
 		LatestVersion = VersionPlusOne - 1
diff --git a/Engine/Plugins/Animation/ControlRig/Source/ControlRig/Public/Rigs/RigHierarchyElements.h b/Engine/Plugins/Animation/ControlRig/Source/ControlRig/Public/Rigs/RigHierarchyElements.h
--- a/Engine/Plugins/Animation/ControlRig/Source/ControlRig/Public/Rigs/RigHierarchyElements.h
+++ b/Engine/Plugins/Animation/ControlRig/Source/ControlRig/Public/Rigs/RigHierarchyElements.h
@@ -1120,6 +1120,12 @@ struct CONTROLRIG_API FRigControlSettings
 	UPROPERTY(BlueprintReadOnly, EditAnywhere, Category = Animation)
 	bool bGroupWithParentControl;
 
+	/**
+	 * Allow to space switch only to the available spaces
+	 */
+	UPROPERTY(BlueprintReadOnly, EditAnywhere, Category = Animation)
+	bool bRestrictSpaceSwitching;
+
 	/**
 	 * Deprecated properties.
 	 */
diff --git a/Engine/Plugins/Animation/ControlRig/Source/ControlRigEditor/Private/ControlRigElementDetails.cpp b/Engine/Plugins/Animation/ControlRig/Source/ControlRigEditor/Private/ControlRigElementDetails.cpp
--- a/Engine/Plugins/Animation/ControlRig/Source/ControlRigEditor/Private/ControlRigElementDetails.cpp
+++ b/Engine/Plugins/Animation/ControlRig/Source/ControlRigEditor/Private/ControlRigElementDetails.cpp
@@ -2759,6 +2759,12 @@ void FRigControlElementDetails::CustomizeControl(IDetailLayoutBuilder& DetailBui
 
 	if(bSupportsShape)
 	{
+		const TSharedPtr<IPropertyHandle> RestrictSpaceSwitchingHandle = SettingsHandle->GetChildHandle(TEXT("bRestrictSpaceSwitching"));
+		ControlCategory
+		.AddProperty(RestrictSpaceSwitchingHandle.ToSharedRef())
+		.DisplayName(FText::FromString(TEXT("Restrict Switching")))
+		.IsEnabled(bIsEnabled);
+		
 		const TSharedPtr<IPropertyHandle> AvailableSpacesHandle = CustomizationHandle->GetChildHandle(TEXT("AvailableSpaces"));
 		ControlCategory.AddProperty(AvailableSpacesHandle.ToSharedRef())
 		.IsEnabled(bIsEnabled);
diff --git a/Engine/Plugins/Animation/ControlRig/Source/ControlRigEditor/Private/EditMode/SControlRigEditModeTools.cpp b/Engine/Plugins/Animation/ControlRig/Source/ControlRigEditor/Private/EditMode/SControlRigEditModeTools.cpp
--- a/Engine/Plugins/Animation/ControlRig/Source/ControlRigEditor/Private/EditMode/SControlRigEditModeTools.cpp
+++ b/Engine/Plugins/Animation/ControlRig/Source/ControlRigEditor/Private/EditMode/SControlRigEditModeTools.cpp
@@ -186,6 +186,7 @@ void SControlRigEditModeTools::Construct(const FArguments& InArgs, TSharedPtr<FC
 	RigTreeDelegates.OnGetDisplaySettings = FOnGetRigTreeDisplaySettings::CreateSP(this, &SControlRigEditModeTools::GetDisplaySettings);
 	RigTreeDelegates.OnSelectionChanged = FOnRigTreeSelectionChanged::CreateSP(this, &SControlRigEditModeTools::HandleSelectionChanged);
 #endif
+
 	ChildSlot
 	[
 		SNew(SScrollBox)
@@ -307,6 +308,10 @@ void SControlRigEditModeTools::Construct(const FArguments& InArgs, TSharedPtr<FC
 							SNew(SImage)
 							.Image(FAppStyle::GetBrush(TEXT("Icons.PlusCircle")))
 						]
+						.Visibility_Lambda([this]()->EVisibility
+						{
+							return GetAddSpaceButtonVisibility();
+						})
 					]
 				]
 				.BodyContent()
@@ -871,6 +876,16 @@ FReply SControlRigEditModeTools::HandleAddSpaceClicked()
 	return SpacePickerWidget->HandleAddElementClicked();
 }
 
+EVisibility SControlRigEditModeTools::GetAddSpaceButtonVisibility() const
+{
+	return IsSpaceSwitchingRestricted() ? EVisibility::Hidden : EVisibility::Visible;
+}
+
+bool SControlRigEditModeTools::IsSpaceSwitchingRestricted() const
+{
+	return SpacePickerWidget->IsRestricted();
+}
+
 FReply SControlRigEditModeTools::OnBakeControlsToNewSpaceButtonClicked()
 {
 	if (SpacePickerWidget->GetHierarchy() == nullptr)
diff --git a/Engine/Plugins/Animation/ControlRig/Source/ControlRigEditor/Private/EditMode/SControlRigEditModeTools.h b/Engine/Plugins/Animation/ControlRig/Source/ControlRigEditor/Private/EditMode/SControlRigEditModeTools.h
--- a/Engine/Plugins/Animation/ControlRig/Source/ControlRigEditor/Private/EditMode/SControlRigEditModeTools.h
+++ b/Engine/Plugins/Animation/ControlRig/Source/ControlRigEditor/Private/EditMode/SControlRigEditModeTools.h
@@ -111,6 +111,8 @@ private:
 	void HandleActiveSpaceChanged(URigHierarchy* InHierarchy, const FRigElementKey& InControlKey, const FRigElementKey& InSpaceKey);
 	void HandleSpaceListChanged(URigHierarchy* InHierarchy, const FRigElementKey& InControlKey, const TArray<FRigElementKey>& InSpaceList);
 	FReply HandleAddSpaceClicked();
+	EVisibility GetAddSpaceButtonVisibility() const;
+	bool IsSpaceSwitchingRestricted() const;
 	FReply OnBakeControlsToNewSpaceButtonClicked();
 
 	FReply HandleAddConstraintClicked();
diff --git a/Engine/Plugins/Animation/ControlRig/Source/ControlRigEditor/Private/EditMode/SControlRigSpacePicker.cpp b/Engine/Plugins/Animation/ControlRig/Source/ControlRigEditor/Private/EditMode/SControlRigSpacePicker.cpp
--- a/Engine/Plugins/Animation/ControlRig/Source/ControlRigEditor/Private/EditMode/SControlRigSpacePicker.cpp
+++ b/Engine/Plugins/Animation/ControlRig/Source/ControlRigEditor/Private/EditMode/SControlRigSpacePicker.cpp
@@ -76,6 +76,10 @@ void SControlRigSpacePicker::Construct(const FArguments& InArgs, FControlRigEdit
 								SNew(SImage)
 								.Image(FAppStyle::GetBrush(TEXT("Icons.PlusCircle")))
 							]
+							.Visibility_Lambda([this]()
+							{
+								return GetAddSpaceButtonVisibility();
+							})
 						]
 					]
 					.BodyContent()
@@ -324,4 +328,23 @@ FReply SControlRigSpacePicker::OnBakeControlsToNewSpaceButtonClicked()
 	return FReply::Unhandled();
 }
 
+EVisibility SControlRigSpacePicker::GetAddSpaceButtonVisibility() const
+{
+	if(const URigHierarchy* Hierarchy = SpacePickerWidget->GetHierarchy())
+	{
+		for(const FRigElementKey& Control : SpacePickerWidget->GetControls())
+		{
+			if(const FRigControlElement* ControlElement = Hierarchy->Find<FRigControlElement>(Control))
+			{
+				if(ControlElement->Settings.bRestrictSpaceSwitching)
+				{
+					return EVisibility::Collapsed;
+				}
+			}
+		}
+	}
+
+	return EVisibility::Visible;
+}
+
 #undef LOCTEXT_NAMESPACE
diff --git a/Engine/Plugins/Animation/ControlRig/Source/ControlRigEditor/Private/EditMode/SControlRigSpacePicker.h b/Engine/Plugins/Animation/ControlRig/Source/ControlRigEditor/Private/EditMode/SControlRigSpacePicker.h
--- a/Engine/Plugins/Animation/ControlRig/Source/ControlRigEditor/Private/EditMode/SControlRigSpacePicker.h
+++ b/Engine/Plugins/Animation/ControlRig/Source/ControlRigEditor/Private/EditMode/SControlRigSpacePicker.h
@@ -40,6 +40,7 @@ private:
 	void HandleSpaceListChanged(URigHierarchy* InHierarchy, const FRigElementKey& InControlKey, const TArray<FRigElementKey>& InSpaceList);
 	FReply HandleAddSpaceClicked();
 	FReply OnBakeControlsToNewSpaceButtonClicked();
+	EVisibility GetAddSpaceButtonVisibility() const;
 
 	//for now picker works off of one ControlRig, this function gets the first control rig with a selection
 	UControlRig* GetControlRig();
diff --git a/Engine/Plugins/Animation/ControlRig/Source/ControlRigEditor/Private/Editor/SRigSpacePickerWidget.cpp b/Engine/Plugins/Animation/ControlRig/Source/ControlRigEditor/Private/Editor/SRigSpacePickerWidget.cpp
--- a/Engine/Plugins/Animation/ControlRig/Source/ControlRigEditor/Private/Editor/SRigSpacePickerWidget.cpp
+++ b/Engine/Plugins/Animation/ControlRig/Source/ControlRigEditor/Private/Editor/SRigSpacePickerWidget.cpp
@@ -158,6 +158,10 @@ void SRigSpacePickerWidget::Construct(const FArguments& InArgs)
 					SNew(SImage)
 					.Image(FAppStyle::GetBrush(TEXT("Icons.PlusCircle")))
 				]
+				.Visibility_Lambda([this]()
+				{
+					return IsRestricted() ? EVisibility::Collapsed : EVisibility::Visible;
+				})
 			];
 		}
 
@@ -187,7 +191,6 @@ void SRigSpacePickerWidget::Construct(const FArguments& InArgs)
 
 	SetControls(InArgs._Hierarchy, InArgs._Controls);
 	SetCanTick(true);
-	
 }
 
 SRigSpacePickerWidget::~SRigSpacePickerWidget()
@@ -470,12 +473,18 @@ void SRigSpacePickerWidget::AddSpacePickerRow(
 
 	if(!IsDefaultSpace(InKey))
 	{
+		const TAttribute<EVisibility> RestrictedVisibility = TAttribute<EVisibility>::CreateLambda([this]
+		{
+			return IsRestricted() ? EVisibility::Collapsed : EVisibility::Visible;
+		});
+		
 		if(bAllowDelete || bAllowReorder)
 		{
 			RowBox->AddSlot()
 			.FillWidth(1.f)
 			[
 				SNew(SSpacer)
+				.Visibility(RestrictedVisibility)
 			];
 		}
 		
@@ -498,6 +507,7 @@ void SRigSpacePickerWidget::AddSpacePickerRow(
 					.Image(FAppStyle::GetBrush("Icons.ChevronUp"))
 					.ColorAndOpacity(FSlateColor::UseForeground())
 				]
+				.Visibility(RestrictedVisibility)
 			];
 
 			RowBox->AddSlot()
@@ -517,18 +527,22 @@ void SRigSpacePickerWidget::AddSpacePickerRow(
 					.Image(FAppStyle::GetBrush("Icons.ChevronDown"))
 					.ColorAndOpacity(FSlateColor::UseForeground())
 				]
+				.Visibility(RestrictedVisibility)
 			];
 		}
 
 		if(bAllowDelete)
 		{
+			const TSharedRef<SWidget> ClearButton = PropertyCustomizationHelpers::MakeClearButton(FSimpleDelegate::CreateSP(this, &SRigSpacePickerWidget::HandleSpaceDelete, InKey), LOCTEXT("DeleteSpace", "Remove this space."), true);
+			ClearButton->SetVisibility(RestrictedVisibility);
+
 			RowBox->AddSlot()
 			.AutoWidth()
 			.VAlign(VAlign_Center)
 			.HAlign(HAlign_Left)
 			.Padding(0)
 			[
-				PropertyCustomizationHelpers::MakeClearButton(FSimpleDelegate::CreateSP(this, &SRigSpacePickerWidget::HandleSpaceDelete, InKey), LOCTEXT("DeleteSpace", "Remove this space."), true)
+				ClearButton
 			];
 		}
 	}
@@ -754,6 +768,24 @@ FReply SRigSpacePickerWidget::HandleAddElementClicked()
 	return FReply::Handled().SetUserFocus(SearchableTreeView->GetSearchBox(), EFocusCause::SetDirectly);
 }
 
+bool SRigSpacePickerWidget::IsRestricted() const
+{
+	if(URigHierarchy* CurrentHierarchy = GetHierarchy())
+	{
+		for(const FRigElementKey& Control : GetControls())
+		{
+			if(const FRigControlElement* ControlElement = CurrentHierarchy->Find<FRigControlElement>(Control))
+			{
+				if(ControlElement->Settings.bRestrictSpaceSwitching)
+				{
+					return true;
+				}
+			}
+		}
+	}
+	return false;
+}
+
 bool SRigSpacePickerWidget::IsSpaceMoveUpEnabled(FRigElementKey InKey) const
 {
 	if(CurrentSpaceKeys.IsEmpty())
diff --git a/Engine/Plugins/Animation/ControlRig/Source/ControlRigEditor/Private/Editor/SRigSpacePickerWidget.h b/Engine/Plugins/Animation/ControlRig/Source/ControlRigEditor/Private/Editor/SRigSpacePickerWidget.h
--- a/Engine/Plugins/Animation/ControlRig/Source/ControlRigEditor/Private/Editor/SRigSpacePickerWidget.h
+++ b/Engine/Plugins/Animation/ControlRig/Source/ControlRigEditor/Private/Editor/SRigSpacePickerWidget.h
@@ -128,9 +128,10 @@ private:
 public:
 	
 	FReply HandleAddElementClicked();
+	bool IsRestricted() const;
 
 private:
-	
+
 	bool IsSpaceMoveUpEnabled(FRigElementKey InKey) const;
 	bool IsSpaceMoveDownEnabled(FRigElementKey InKey) const;
 
