Control Rig: Add note on if / select node concerning arrays

diff --git a/Engine/Plugins/Animation/ControlRig/Source/ControlRigEditor/Private/Graph/SControlRigGraphNode.cpp b/Engine/Plugins/Animation/ControlRig/Source/ControlRigEditor/Private/Graph/SControlRigGraphNode.cpp
--- a/Engine/Plugins/Animation/ControlRig/Source/ControlRigEditor/Private/Graph/SControlRigGraphNode.cpp
+++ b/Engine/Plugins/Animation/ControlRig/Source/ControlRigEditor/Private/Graph/SControlRigGraphNode.cpp
@@ -832,6 +832,32 @@ void SControlRigGraphNode::RefreshErrorInfo()
 {
 	if (GraphNode)
 	{
+		// if the node has no further errors, check for array reference issues
+		if(const UControlRigGraphNode* RigGraphNode = Cast<UControlRigGraphNode>(GraphNode))
+		{
+			if(!GraphNode->bHasCompilerMessage &&
+				(GraphNode->ErrorType == int32(EMessageSeverity::Info) + 1))
+			{
+				if(const URigVMNode* RigModelNode = RigGraphNode->GetModelNode())
+				{
+					if(RigModelNode->IsA<URigVMIfNode>() || RigModelNode->IsA<URigVMSelectNode>())
+					{
+						for(const URigVMPin* Pin : RigModelNode->GetPins())
+						{
+							if(Pin->IsArray())
+							{
+								GraphNode->bHasCompilerMessage = true;
+								GraphNode->ErrorType = int32(EMessageSeverity::Info);
+								static const FString ArrayWarning = TEXT("This node creates a copy of the array.\nThis may cause side effects.");
+								GraphNode->ErrorMsg = ArrayWarning;
+								break;
+							}
+						}
+					}
+				}
+			}
+		}
+
 		if (NodeErrorType != GraphNode->ErrorType)
 		{
 			SGraphNode::RefreshErrorInfo();
