ensure the constraint is valid (not pending kill) when updating the constraints' list

diff --git a/Engine/Plugins/Animation/ControlRig/Source/ControlRigEditor/Private/Editor/Constraints/SConstraintsWidget.cpp b/Engine/Plugins/Animation/ControlRig/Source/ControlRigEditor/Private/Editor/Constraints/SConstraintsWidget.cpp
--- a/Engine/Plugins/Animation/ControlRig/Source/ControlRigEditor/Private/Editor/Constraints/SConstraintsWidget.cpp
+++ b/Engine/Plugins/Animation/ControlRig/Source/ControlRigEditor/Private/Editor/Constraints/SConstraintsWidget.cpp
@@ -452,7 +452,7 @@ void SEditableConstraintItem::Construct(
 	}
 	
 	FString ParentFullLabel = ParentLabel, ChildFullLabel = ChildLabel;
-	if (Constraint)
+	if (IsValid(Constraint))
 	{
 		Constraint->GetFullLabel().Split(TEXT("."), &ParentFullLabel, &ChildFullLabel);
 	}
@@ -474,7 +474,7 @@ void SEditableConstraintItem::Construct(
 			.BorderImage(RoundedBoxBrush)
 			.BorderBackgroundColor_Lambda([Constraint]()
 			{
-				if (!Constraint)
+				if (!IsValid(Constraint))
 				{
 					return FStyleColors::Transparent;
 				}
@@ -510,7 +510,7 @@ void SEditableConstraintItem::Construct(
 					})
 					.Font_Lambda([Constraint]()
 					{
-						if (!Constraint)
+						if (!IsValid(Constraint))
 						{
 							return IDetailLayoutBuilder::GetDetailFont();
 						}
@@ -518,7 +518,7 @@ void SEditableConstraintItem::Construct(
 					})
 					.ToolTipText_Lambda( [Constraint, ParentFullLabel, ChildFullLabel]()
 					{
-						if (!Constraint)
+						if (!IsValid(Constraint))
 						{
 							return FText();
 						}
@@ -550,6 +550,10 @@ void SEditableConstraintItem::Construct(
 			.ContentPadding(0)
 			.OnClicked_Lambda(	[Constraint]()
 			{
+				if (!IsValid(Constraint))
+				{
+					return FReply::Handled();
+				}
 				if (UTickableTransformConstraint* TransformConstraint = Cast<UTickableTransformConstraint>(Constraint))
 				{
 					FScopedTransaction Transaction(LOCTEXT("CreateConstraintKey", "Create Constraint Key"));
