Sequencer: Constraints: Need to keep track of the handle to the constraint deletion for control rig track editor like we did for transform track editor

diff --git a/Engine/Plugins/Animation/ControlRig/Source/ControlRigEditor/Private/Sequencer/ControlRigParameterTrackEditor.cpp b/Engine/Plugins/Animation/ControlRig/Source/ControlRigEditor/Private/Sequencer/ControlRigParameterTrackEditor.cpp
--- a/Engine/Plugins/Animation/ControlRig/Source/ControlRigEditor/Private/Sequencer/ControlRigParameterTrackEditor.cpp
+++ b/Engine/Plugins/Animation/ControlRig/Source/ControlRigEditor/Private/Sequencer/ControlRigParameterTrackEditor.cpp
@@ -517,6 +517,17 @@ void FControlRigParameterTrackEditor::ObjectImplicitlyRemoved(UObject* InObject)
 
 void FControlRigParameterTrackEditor::OnRelease()
 {
+	UWorld* World = GCurrentLevelEditingViewportClient ? GCurrentLevelEditingViewportClient->GetWorld() : nullptr;
+	FConstraintsManagerController& Controller = FConstraintsManagerController::Get(World);
+	for (FDelegateHandle& Handle : ConstraintHandlesToClear)
+	{
+		if (Handle.IsValid())
+		{
+			Controller.GetNotifyDelegate().Remove(Handle);
+		}
+	}
+	ConstraintHandlesToClear.Reset();
+
 	UnbindAllControlRigs();
 	if (GetSequencer().IsValid())
 	{
@@ -2720,7 +2731,7 @@ void FControlRigParameterTrackEditor::HandleConstraintKeyMoved(
 	}
 }
 
-void FControlRigParameterTrackEditor::HandleConstraintRemoved(IMovieSceneConstrainedSection* InSection) const
+void FControlRigParameterTrackEditor::HandleConstraintRemoved(IMovieSceneConstrainedSection* InSection) 
 {
 	UMovieSceneControlRigParameterSection* Section = Cast<UMovieSceneControlRigParameterSection>(InSection);
 
@@ -2771,6 +2782,7 @@ void FControlRigParameterTrackEditor::HandleConstraintRemoved(IMovieSceneConstra
 						break;		
 				}
 			});
+			ConstraintHandlesToClear.Add(InSection->OnConstraintRemovedHandle);
 		}
 	}
 }
diff --git a/Engine/Plugins/Animation/ControlRig/Source/ControlRigEditor/Private/Sequencer/ControlRigParameterTrackEditor.h b/Engine/Plugins/Animation/ControlRig/Source/ControlRigEditor/Private/Sequencer/ControlRigParameterTrackEditor.h
--- a/Engine/Plugins/Animation/ControlRig/Source/ControlRigEditor/Private/Sequencer/ControlRigParameterTrackEditor.h
+++ b/Engine/Plugins/Animation/ControlRig/Source/ControlRigEditor/Private/Sequencer/ControlRigParameterTrackEditor.h
@@ -159,7 +159,7 @@ private:
 		IMovieSceneConstrainedSection* InSection,
 		const FMovieSceneConstraintChannel* InConstraintChannel,
 		const TArray<FKeyMoveEventItem>& InMovedItems);
-	void HandleConstraintRemoved(IMovieSceneConstrainedSection* InSection) const;
+	void HandleConstraintRemoved(IMovieSceneConstrainedSection* InSection);
 
 	/** Select control rig if not selected, select controls from key areas */
 	void SelectRigsAndControls(UControlRig* Subject, const TArray<const IKeyArea*>& KeyAreas);
@@ -278,6 +278,9 @@ private:
 	/** Controls if the control rig track for the default animating rig should be created */
 	static bool bAutoGenerateControlRigTrack;
 
+	/** Set of delegate handles we have added delegate's too, need to clear them*/
+	TSet<FDelegateHandle> ConstraintHandlesToClear;
+
 	friend class FControlRigBlueprintActions;
 };
 
