Control Rig: Make nodes for math types

diff --git a/Engine/Plugins/Animation/ControlRig/Source/ControlRig/Private/Units/Math/RigUnit_MathBool.cpp b/Engine/Plugins/Animation/ControlRig/Source/ControlRig/Private/Units/Math/RigUnit_MathBool.cpp
--- a/Engine/Plugins/Animation/ControlRig/Source/ControlRig/Private/Units/Math/RigUnit_MathBool.cpp
+++ b/Engine/Plugins/Animation/ControlRig/Source/ControlRig/Private/Units/Math/RigUnit_MathBool.cpp
@@ -5,6 +5,10 @@
 
 #include UE_INLINE_GENERATED_CPP_BY_NAME(RigUnit_MathBool)
 
+FRigUnit_MathBoolMake_Execute()
+{
+}
+
 FRigUnit_MathBoolConstTrue_Execute()
 {
 	Value = true;
diff --git a/Engine/Plugins/Animation/ControlRig/Source/ControlRig/Private/Units/Math/RigUnit_MathBool.h b/Engine/Plugins/Animation/ControlRig/Source/ControlRig/Private/Units/Math/RigUnit_MathBool.h
--- a/Engine/Plugins/Animation/ControlRig/Source/ControlRig/Private/Units/Math/RigUnit_MathBool.h
+++ b/Engine/Plugins/Animation/ControlRig/Source/ControlRig/Private/Units/Math/RigUnit_MathBool.h
@@ -83,6 +83,26 @@ struct CONTROLRIG_API FRigUnit_MathBoolBinaryAggregateOp : public FRigUnit_MathB
 	bool Result;
 };
 
+/**
+ * A boolean constant
+ */
+USTRUCT(meta=(DisplayName="Bool", Keywords="Make,Construct,Constant"))
+struct CONTROLRIG_API FRigUnit_MathBoolMake : public FRigUnit_MathBoolBase
+{
+	GENERATED_BODY()
+
+	FRigUnit_MathBoolMake()
+	{
+		Value = false;
+	}
+	
+	RIGVM_METHOD()
+	virtual void Execute(const FRigUnitContext& Context) override;
+
+	UPROPERTY(meta=(Input, Output))
+	bool Value;
+};
+
 /**
  * Returns true
  */
diff --git a/Engine/Plugins/Animation/ControlRig/Source/ControlRig/Private/Units/Math/RigUnit_MathColor.cpp b/Engine/Plugins/Animation/ControlRig/Source/ControlRig/Private/Units/Math/RigUnit_MathColor.cpp
--- a/Engine/Plugins/Animation/ControlRig/Source/ControlRig/Private/Units/Math/RigUnit_MathColor.cpp
+++ b/Engine/Plugins/Animation/ControlRig/Source/ControlRig/Private/Units/Math/RigUnit_MathColor.cpp
@@ -6,6 +6,11 @@
 
 #include UE_INLINE_GENERATED_CPP_BY_NAME(RigUnit_MathColor)
 
+FRigUnit_MathColorMake_Execute()
+{
+	Result = FLinearColor(R, G, B, A);
+}
+
 FRigUnit_MathColorFromFloat_Execute()
 {
     DECLARE_SCOPE_HIERARCHICAL_COUNTER_RIGUNIT()
diff --git a/Engine/Plugins/Animation/ControlRig/Source/ControlRig/Private/Units/Math/RigUnit_MathColor.h b/Engine/Plugins/Animation/ControlRig/Source/ControlRig/Private/Units/Math/RigUnit_MathColor.h
--- a/Engine/Plugins/Animation/ControlRig/Source/ControlRig/Private/Units/Math/RigUnit_MathColor.h
+++ b/Engine/Plugins/Animation/ControlRig/Source/ControlRig/Private/Units/Math/RigUnit_MathColor.h
@@ -52,6 +52,39 @@ struct CONTROLRIG_API FRigUnit_MathColorBinaryAggregateOp : public FRigUnit_Math
 	FLinearColor Result;
 };
 
+/**
+ * Makes a color from its components
+ */
+USTRUCT(meta=(DisplayName="Make Color", Keywords="Make,Construct,Constant"))
+struct CONTROLRIG_API FRigUnit_MathColorMake : public FRigUnit_MathColorBase
+{
+	GENERATED_BODY()
+	RIGVM_METHOD()
+	virtual void Execute(const FRigUnitContext& Context) override;
+
+	FRigUnit_MathColorMake()
+	{
+		R = G = B = 0.f;
+		A = 1.f;
+		Result = FLinearColor::Black;
+	}
+
+	UPROPERTY(meta=(Input))
+	float R;
+
+	UPROPERTY(meta=(Input))
+	float G;
+
+	UPROPERTY(meta=(Input))
+	float B;
+
+	UPROPERTY(meta=(Input))
+	float A;
+
+	UPROPERTY(meta=(Output))
+	FLinearColor Result;
+};
+
 /**
  * Makes a vector from a single float
  */
diff --git a/Engine/Plugins/Animation/ControlRig/Source/ControlRig/Private/Units/Math/RigUnit_MathDouble.cpp b/Engine/Plugins/Animation/ControlRig/Source/ControlRig/Private/Units/Math/RigUnit_MathDouble.cpp
--- a/Engine/Plugins/Animation/ControlRig/Source/ControlRig/Private/Units/Math/RigUnit_MathDouble.cpp
+++ b/Engine/Plugins/Animation/ControlRig/Source/ControlRig/Private/Units/Math/RigUnit_MathDouble.cpp
@@ -6,6 +6,10 @@
 
 #include UE_INLINE_GENERATED_CPP_BY_NAME(RigUnit_MathDouble)
 
+FRigUnit_MathDoubleMake_Execute()
+{
+}
+
 FRigUnit_MathDoubleConstPi_Execute()
 {
 	Value = PI;
diff --git a/Engine/Plugins/Animation/ControlRig/Source/ControlRig/Private/Units/Math/RigUnit_MathDouble.h b/Engine/Plugins/Animation/ControlRig/Source/ControlRig/Private/Units/Math/RigUnit_MathDouble.h
--- a/Engine/Plugins/Animation/ControlRig/Source/ControlRig/Private/Units/Math/RigUnit_MathDouble.h
+++ b/Engine/Plugins/Animation/ControlRig/Source/ControlRig/Private/Units/Math/RigUnit_MathDouble.h
@@ -83,6 +83,26 @@ struct CONTROLRIG_API FRigUnit_MathDoubleBinaryAggregateOp : public FRigUnit_Mat
 	double Result;
 };
 
+/**
+ * A double constant
+ */
+USTRUCT(meta=(DisplayName="Double", Keywords="Make,Construct,Constant"))
+struct CONTROLRIG_API FRigUnit_MathDoubleMake : public FRigUnit_MathDoubleBase
+{
+	GENERATED_BODY()
+
+	FRigUnit_MathDoubleMake()
+	{
+		Value = 0.0;
+	}
+	
+	RIGVM_METHOD()
+	virtual void Execute(const FRigUnitContext& Context) override;
+
+	UPROPERTY(meta=(Input, Output))
+	double Value;
+};
+
 /**
  * Returns PI
  */
diff --git a/Engine/Plugins/Animation/ControlRig/Source/ControlRig/Private/Units/Math/RigUnit_MathFloat.cpp b/Engine/Plugins/Animation/ControlRig/Source/ControlRig/Private/Units/Math/RigUnit_MathFloat.cpp
--- a/Engine/Plugins/Animation/ControlRig/Source/ControlRig/Private/Units/Math/RigUnit_MathFloat.cpp
+++ b/Engine/Plugins/Animation/ControlRig/Source/ControlRig/Private/Units/Math/RigUnit_MathFloat.cpp
@@ -6,6 +6,10 @@
 
 #include UE_INLINE_GENERATED_CPP_BY_NAME(RigUnit_MathFloat)
 
+FRigUnit_MathFloatMake_Execute()
+{
+}
+
 FRigUnit_MathFloatConstPi_Execute()
 {
 	Value = PI;
diff --git a/Engine/Plugins/Animation/ControlRig/Source/ControlRig/Private/Units/Math/RigUnit_MathFloat.h b/Engine/Plugins/Animation/ControlRig/Source/ControlRig/Private/Units/Math/RigUnit_MathFloat.h
--- a/Engine/Plugins/Animation/ControlRig/Source/ControlRig/Private/Units/Math/RigUnit_MathFloat.h
+++ b/Engine/Plugins/Animation/ControlRig/Source/ControlRig/Private/Units/Math/RigUnit_MathFloat.h
@@ -82,6 +82,26 @@ struct CONTROLRIG_API FRigUnit_MathFloatBinaryAggregateOp : public FRigUnit_Math
 	float Result;
 };
 
+/**
+ * A float constant
+ */
+USTRUCT(meta=(DisplayName="Float", Keywords="Make,Construct,Constant"))
+struct CONTROLRIG_API FRigUnit_MathFloatMake : public FRigUnit_MathFloatBase
+{
+	GENERATED_BODY()
+
+	FRigUnit_MathFloatMake()
+	{
+		Value = 0.f;
+	}
+	
+	RIGVM_METHOD()
+	virtual void Execute(const FRigUnitContext& Context) override;
+
+	UPROPERTY(meta=(Input, Output))
+	float Value;
+};
+
 /**
  * Returns PI
  */
diff --git a/Engine/Plugins/Animation/ControlRig/Source/ControlRig/Private/Units/Math/RigUnit_MathInt.cpp b/Engine/Plugins/Animation/ControlRig/Source/ControlRig/Private/Units/Math/RigUnit_MathInt.cpp
--- a/Engine/Plugins/Animation/ControlRig/Source/ControlRig/Private/Units/Math/RigUnit_MathInt.cpp
+++ b/Engine/Plugins/Animation/ControlRig/Source/ControlRig/Private/Units/Math/RigUnit_MathInt.cpp
@@ -6,6 +6,10 @@
 
 #include UE_INLINE_GENERATED_CPP_BY_NAME(RigUnit_MathInt)
 
+FRigUnit_MathIntMake_Execute()
+{
+}
+
 FRigUnit_MathIntAdd_Execute()
 {
     DECLARE_SCOPE_HIERARCHICAL_COUNTER_RIGUNIT()
diff --git a/Engine/Plugins/Animation/ControlRig/Source/ControlRig/Private/Units/Math/RigUnit_MathInt.h b/Engine/Plugins/Animation/ControlRig/Source/ControlRig/Private/Units/Math/RigUnit_MathInt.h
--- a/Engine/Plugins/Animation/ControlRig/Source/ControlRig/Private/Units/Math/RigUnit_MathInt.h
+++ b/Engine/Plugins/Animation/ControlRig/Source/ControlRig/Private/Units/Math/RigUnit_MathInt.h
@@ -68,6 +68,26 @@ struct CONTROLRIG_API FRigUnit_MathIntBinaryAggregateOp : public FRigUnit_MathIn
 	int32 Result;
 };
 
+/**
+ * A integer constant
+ */
+USTRUCT(meta=(DisplayName="Integer", Keywords="Make,Construct,Constant"))
+struct CONTROLRIG_API FRigUnit_MathIntMake : public FRigUnit_MathIntBase
+{
+	GENERATED_BODY()
+
+	FRigUnit_MathIntMake()
+	{
+		Value = 0;
+	}
+	
+	RIGVM_METHOD()
+	virtual void Execute(const FRigUnitContext& Context) override;
+
+	UPROPERTY(meta=(Input, Output))
+	int32 Value;
+};
+
 /**
  * Returns the sum of the two values
  */
diff --git a/Engine/Plugins/Animation/ControlRig/Source/ControlRig/Private/Units/Math/RigUnit_MathQuaternion.cpp b/Engine/Plugins/Animation/ControlRig/Source/ControlRig/Private/Units/Math/RigUnit_MathQuaternion.cpp
--- a/Engine/Plugins/Animation/ControlRig/Source/ControlRig/Private/Units/Math/RigUnit_MathQuaternion.cpp
+++ b/Engine/Plugins/Animation/ControlRig/Source/ControlRig/Private/Units/Math/RigUnit_MathQuaternion.cpp
@@ -9,6 +9,11 @@
 
 #include UE_INLINE_GENERATED_CPP_BY_NAME(RigUnit_MathQuaternion)
 
+FRigUnit_MathQuaternionMake_Execute()
+{
+	Result = FQuat(X, Y, Z, W);
+}
+
 FRigUnit_MathQuaternionFromAxisAndAngle_Execute()
 {
     DECLARE_SCOPE_HIERARCHICAL_COUNTER_RIGUNIT()
diff --git a/Engine/Plugins/Animation/ControlRig/Source/ControlRig/Private/Units/Math/RigUnit_MathQuaternion.h b/Engine/Plugins/Animation/ControlRig/Source/ControlRig/Private/Units/Math/RigUnit_MathQuaternion.h
--- a/Engine/Plugins/Animation/ControlRig/Source/ControlRig/Private/Units/Math/RigUnit_MathQuaternion.h
+++ b/Engine/Plugins/Animation/ControlRig/Source/ControlRig/Private/Units/Math/RigUnit_MathQuaternion.h
@@ -69,6 +69,38 @@ struct CONTROLRIG_API FRigUnit_MathQuaternionBinaryAggregateOp : public FRigUnit
 	FQuat Result;
 };
 
+/**
+ * Makes a quaternion from its components
+ */
+USTRUCT(meta=(DisplayName="Make Quat", Keywords="Make,Construct,Constant"))
+struct CONTROLRIG_API FRigUnit_MathQuaternionMake : public FRigUnit_MathQuaternionBase
+{
+	GENERATED_BODY()
+	RIGVM_METHOD()
+	virtual void Execute(const FRigUnitContext& Context) override;
+
+	FRigUnit_MathQuaternionMake()
+	{
+		X = Y = Z = W = 0.f;
+		Result = FQuat::Identity;
+	}
+
+	UPROPERTY(meta=(Input))
+	float X;
+
+	UPROPERTY(meta=(Input))
+	float Y;
+
+	UPROPERTY(meta=(Input))
+	float Z;
+
+	UPROPERTY(meta=(Input))
+	float W;
+
+	UPROPERTY(meta=(Output))
+	FQuat Result;
+};
+
 /**
  * Makes a quaternion from an axis and an angle in radians
  */
diff --git a/Engine/Plugins/Animation/ControlRig/Source/ControlRig/Private/Units/Math/RigUnit_MathTransform.cpp b/Engine/Plugins/Animation/ControlRig/Source/ControlRig/Private/Units/Math/RigUnit_MathTransform.cpp
--- a/Engine/Plugins/Animation/ControlRig/Source/ControlRig/Private/Units/Math/RigUnit_MathTransform.cpp
+++ b/Engine/Plugins/Animation/ControlRig/Source/ControlRig/Private/Units/Math/RigUnit_MathTransform.cpp
@@ -9,6 +9,11 @@
 
 #include UE_INLINE_GENERATED_CPP_BY_NAME(RigUnit_MathTransform)
 
+FRigUnit_MathTransformMake_Execute()
+{
+	Result = FTransform(Rotation, Translation, Scale);
+}
+
 FRigUnit_MathTransformFromEulerTransform_Execute()
 {
     DECLARE_SCOPE_HIERARCHICAL_COUNTER_RIGUNIT()
diff --git a/Engine/Plugins/Animation/ControlRig/Source/ControlRig/Private/Units/Math/RigUnit_MathTransform.h b/Engine/Plugins/Animation/ControlRig/Source/ControlRig/Private/Units/Math/RigUnit_MathTransform.h
--- a/Engine/Plugins/Animation/ControlRig/Source/ControlRig/Private/Units/Math/RigUnit_MathTransform.h
+++ b/Engine/Plugins/Animation/ControlRig/Source/ControlRig/Private/Units/Math/RigUnit_MathTransform.h
@@ -76,6 +76,37 @@ struct CONTROLRIG_API FRigUnit_MathTransformBinaryAggregateOp : public FRigUnit_
 	FTransform Result;
 };
 
+/**
+ * Makes a transform from its components
+ */
+USTRUCT(meta=(DisplayName="Make Transform", Keywords="Make,Construct,Constant"))
+struct CONTROLRIG_API FRigUnit_MathTransformMake : public FRigUnit_MathTransformBase
+{
+	GENERATED_BODY()
+	RIGVM_METHOD()
+	virtual void Execute(const FRigUnitContext& Context) override;
+
+	FRigUnit_MathTransformMake()
+	{
+		Translation = FVector::ZeroVector;
+		Rotation = FQuat::Identity;
+		Scale = FVector::OneVector;
+		Result = FTransform::Identity;
+	}
+
+	UPROPERTY(meta=(Input))
+	FVector Translation;
+
+	UPROPERTY(meta=(Input))
+	FQuat Rotation;
+
+	UPROPERTY(meta=(Input))
+	FVector Scale;
+
+	UPROPERTY(meta=(Output))
+	FTransform Result;
+};
+
 /**
  * Makes a quaternion based transform from a euler based transform
  */
diff --git a/Engine/Plugins/Animation/ControlRig/Source/ControlRig/Private/Units/Math/RigUnit_MathVector.cpp b/Engine/Plugins/Animation/ControlRig/Source/ControlRig/Private/Units/Math/RigUnit_MathVector.cpp
--- a/Engine/Plugins/Animation/ControlRig/Source/ControlRig/Private/Units/Math/RigUnit_MathVector.cpp
+++ b/Engine/Plugins/Animation/ControlRig/Source/ControlRig/Private/Units/Math/RigUnit_MathVector.cpp
@@ -8,6 +8,11 @@
 
 #include UE_INLINE_GENERATED_CPP_BY_NAME(RigUnit_MathVector)
 
+FRigUnit_MathVectorMake_Execute()
+{
+	Result = FVector(X, Y, Z);
+}
+
 FRigUnit_MathVectorFromFloat_Execute()
 {
     DECLARE_SCOPE_HIERARCHICAL_COUNTER_RIGUNIT()
diff --git a/Engine/Plugins/Animation/ControlRig/Source/ControlRig/Private/Units/Math/RigUnit_MathVector.h b/Engine/Plugins/Animation/ControlRig/Source/ControlRig/Private/Units/Math/RigUnit_MathVector.h
--- a/Engine/Plugins/Animation/ControlRig/Source/ControlRig/Private/Units/Math/RigUnit_MathVector.h
+++ b/Engine/Plugins/Animation/ControlRig/Source/ControlRig/Private/Units/Math/RigUnit_MathVector.h
@@ -69,6 +69,35 @@ struct CONTROLRIG_API FRigUnit_MathVectorBinaryAggregateOp : public FRigUnit_Mat
 	FVector Result;
 };
 
+/**
+ * Makes a vector from its components
+ */
+USTRUCT(meta=(DisplayName="Make Vector", Keywords="Make,Construct,Constant"))
+struct CONTROLRIG_API FRigUnit_MathVectorMake : public FRigUnit_MathVectorBase
+{
+	GENERATED_BODY()
+	RIGVM_METHOD()
+	virtual void Execute(const FRigUnitContext& Context) override;
+
+	FRigUnit_MathVectorMake()
+	{
+		X = Y = Z = 0.f;
+		Result = FVector::ZeroVector;
+	}
+
+	UPROPERTY(meta=(Input))
+	float X;
+
+	UPROPERTY(meta=(Input))
+	float Y;
+
+	UPROPERTY(meta=(Input))
+	float Z;
+
+	UPROPERTY(meta=(Output))
+	FVector Result;
+};
+
 /**
  * Makes a vector from a single float
  */
