Cotnrol Rig: Go To Instruction in Stackview

diff --git a/Engine/Plugins/Animation/ControlRig/Source/ControlRigEditor/Private/ControlRigStackCommands.cpp b/Engine/Plugins/Animation/ControlRig/Source/ControlRigEditor/Private/ControlRigStackCommands.cpp
--- a/Engine/Plugins/Animation/ControlRig/Source/ControlRigEditor/Private/ControlRigStackCommands.cpp
+++ b/Engine/Plugins/Animation/ControlRig/Source/ControlRigEditor/Private/ControlRigStackCommands.cpp
@@ -7,6 +7,7 @@
 void FControlRigStackCommands::RegisterCommands()
 {
 	UI_COMMAND(FocusOnSelection, "Focus On Selection", "Finds the selected operator's node in the graph.", EUserInterfaceActionType::Button, FInputChord());
+	UI_COMMAND(GoToInstruction, "Go To Instruction", "Looks for a specific instruction by index and brings it into focus.", EUserInterfaceActionType::Button, FInputChord(EKeys::G, EModifierKey::Control));
 }
 
 #undef LOCTEXT_NAMESPACE
diff --git a/Engine/Plugins/Animation/ControlRig/Source/ControlRigEditor/Private/ControlRigStackCommands.h b/Engine/Plugins/Animation/ControlRig/Source/ControlRigEditor/Private/ControlRigStackCommands.h
--- a/Engine/Plugins/Animation/ControlRig/Source/ControlRigEditor/Private/ControlRigStackCommands.h
+++ b/Engine/Plugins/Animation/ControlRig/Source/ControlRigEditor/Private/ControlRigStackCommands.h
@@ -21,6 +21,9 @@ public:
 	/** Focuses on a selected operator's node */
 	TSharedPtr< FUICommandInfo > FocusOnSelection;
 
+	/** Looks for a specific instruction by index and brings it into focus */
+	TSharedPtr< FUICommandInfo > GoToInstruction;
+
 	/**
 	 * Initialize commands
 	 */
diff --git a/Engine/Plugins/Animation/ControlRig/Source/ControlRigEditor/Private/Editor/SControlRigStackView.cpp b/Engine/Plugins/Animation/ControlRig/Source/ControlRigEditor/Private/Editor/SControlRigStackView.cpp
--- a/Engine/Plugins/Animation/ControlRig/Source/ControlRigEditor/Private/Editor/SControlRigStackView.cpp
+++ b/Engine/Plugins/Animation/ControlRig/Source/ControlRigEditor/Private/Editor/SControlRigStackView.cpp
@@ -6,6 +6,7 @@
 #include "Widgets/Layout/SScrollBar.h"
 #include "Widgets/Layout/SSpacer.h"
 #include "Widgets/Input/SSearchBox.h"
+#include "Widgets/Input/SNumericEntryBox.h"
 #include "Editor/ControlRigEditor.h"
 #include "ControlRigEditorStyle.h"
 #include "ControlRigStackCommands.h"
@@ -16,6 +17,7 @@
 #include "DetailLayoutBuilder.h"
 #include "RigVMModel/Nodes/RigVMAggregateNode.h"
 #include "Widgets/Input/SSearchBox.h"
+#include "Dialog/SCustomDialog.h"
 
 #define LOCTEXT_NAMESPACE "SControlRigStackView"
 
@@ -426,6 +428,7 @@ void SControlRigStackView::BindCommands()
 	// create new command
 	const FControlRigStackCommands& Commands = FControlRigStackCommands::Get();
 	CommandList->MapAction(Commands.FocusOnSelection, FExecuteAction::CreateSP(this, &SControlRigStackView::HandleFocusOnSelectedGraphNode));
+	CommandList->MapAction(Commands.GoToInstruction, FExecuteAction::CreateSP(this, &SControlRigStackView::HandleGoToInstruction));
 }
 
 TSharedRef<ITableRow> SControlRigStackView::MakeTableRowWidget(TSharedPtr<FRigStackEntry> InItem, const TSharedRef<STableViewBase>& OwnerTable, TWeakObjectPtr<UControlRigBlueprint> InBlueprint)
@@ -834,6 +837,7 @@ TSharedPtr< SWidget > SControlRigStackView::CreateContextMenu()
 	{
 		MenuBuilder.BeginSection("RigStackToolsAction", LOCTEXT("ToolsAction", "Tools"));
 		MenuBuilder.AddMenuEntry(Actions.FocusOnSelection);
+		MenuBuilder.AddMenuEntry(Actions.GoToInstruction);
 		MenuBuilder.EndSection();
 	}
 
@@ -876,6 +880,60 @@ void SControlRigStackView::HandleFocusOnSelectedGraphNode()
 	}
 }
 
+void SControlRigStackView::HandleGoToInstruction()
+{
+	// figure out the current instruction's index
+	int32 Index = 0;
+	TArray<TSharedPtr<FRigStackEntry>> SelectedItems = TreeView->GetSelectedItems();
+	if (SelectedItems.Num() > 0)
+	{
+		Index = SelectedItems[0]->InstructionIndex;
+	}
+	
+	const FVector2D MouseCursorLocation = FSlateApplication::Get().GetCursorPos();
+	
+	SWindow::FArguments WindowArguments;
+	WindowArguments.ScreenPosition(MouseCursorLocation);
+	WindowArguments.AutoCenter(EAutoCenter::None);
+	WindowArguments.FocusWhenFirstShown(true);
+
+	TSharedPtr<SNumericEntryBox<int32>> NumericBox;
+	const TSharedRef<SCustomDialog> OptionsDialog = SNew(SCustomDialog)
+	.Title(FText(LOCTEXT("GoToInstructionDialog", "Go to...")))
+	.WindowArguments(WindowArguments)
+	.Content()
+	[
+		SNew(SHorizontalBox)
+		+ SHorizontalBox::Slot()
+		.FillWidth(1)
+		.VAlign(VAlign_Center)
+		.HAlign(HAlign_Left)
+		[
+			SAssignNew(NumericBox, SNumericEntryBox<int32>)
+			.Font(IDetailLayoutBuilder::GetDetailFont())
+			.Value(Index)
+			.MinDesiredValueWidth(250)
+			.MinValue(0)
+			.MaxValue(Operators.Num() - 1)
+			.OnValueChanged_Lambda([&Index](const int32& InIndex)
+			{
+				Index = InIndex;
+			})
+			.IsEnabled(true)
+		]
+	]
+	.Buttons({
+		SCustomDialog::FButton(LOCTEXT("OK", "OK")),
+		SCustomDialog::FButton(LOCTEXT("Cancel", "Cancel"))
+	});
+
+	if(OptionsDialog->ShowModal() == 0)
+	{
+		TreeView->SetSelection(Operators[Index]);
+		TreeView->SetScrollOffset(FMath::Max(Index-5, 0));
+	}
+}
+
 void SControlRigStackView::OnVMCompiled(UObject* InCompiledObject, URigVM* InCompiledVM)
 {
 	RefreshTreeView(InCompiledVM);
diff --git a/Engine/Plugins/Animation/ControlRig/Source/ControlRigEditor/Private/Editor/SControlRigStackView.h b/Engine/Plugins/Animation/ControlRig/Source/ControlRigEditor/Private/Editor/SControlRigStackView.h
--- a/Engine/Plugins/Animation/ControlRig/Source/ControlRigEditor/Private/Editor/SControlRigStackView.h
+++ b/Engine/Plugins/Animation/ControlRig/Source/ControlRigEditor/Private/Editor/SControlRigStackView.h
@@ -96,6 +96,9 @@ private:
 	/** Focus on the selected operator in the graph*/
 	void HandleFocusOnSelectedGraphNode();
 
+	/** Offers a dialog to move to a specific instruction */
+	void HandleGoToInstruction();
+
 	void OnVMCompiled(UObject* InCompiledObject, URigVM* InCompiledVM);
 
 	//* Focus on the instruction when the execution is halted */
