Control Rig: Avoid crash when RigElement cease to exist

diff --git a/Engine/Plugins/Animation/ControlRig/Source/ControlRigEditor/Private/EditMode/ControlRigEditMode.cpp b/Engine/Plugins/Animation/ControlRig/Source/ControlRigEditor/Private/EditMode/ControlRigEditMode.cpp
--- a/Engine/Plugins/Animation/ControlRig/Source/ControlRigEditor/Private/EditMode/ControlRigEditMode.cpp
+++ b/Engine/Plugins/Animation/ControlRig/Source/ControlRigEditor/Private/EditMode/ControlRigEditMode.cpp
@@ -3665,8 +3665,10 @@ void FControlRigEditMode::OnHierarchyModified_AnyThread(ERigHierarchyNotificatio
 		{
 			return;
 		}
-		const FRigBaseElement* Element = WeakHierarchy.Get()->Find(Key);
-		OnHierarchyModified(InNotif, WeakHierarchy.Get(), Element);
+		if (const FRigBaseElement* Element = WeakHierarchy.Get()->Find(Key))
+		{
+			OnHierarchyModified(InNotif, WeakHierarchy.Get(), Element);
+		}
 		
 	}, TStatId(), NULL, ENamedThreads::GameThread);
 }
diff --git a/Engine/Plugins/Animation/ControlRig/Source/ControlRigEditor/Private/Editor/ControlRigSkeletalMeshComponent.cpp b/Engine/Plugins/Animation/ControlRig/Source/ControlRigEditor/Private/Editor/ControlRigSkeletalMeshComponent.cpp
--- a/Engine/Plugins/Animation/ControlRig/Source/ControlRigEditor/Private/Editor/ControlRigSkeletalMeshComponent.cpp
+++ b/Engine/Plugins/Animation/ControlRig/Source/ControlRigEditor/Private/Editor/ControlRigSkeletalMeshComponent.cpp
@@ -254,8 +254,10 @@ void UControlRigSkeletalMeshComponent::OnHierarchyModified_AnyThread(ERigHierarc
     	{
     		return;
     	}
-        const FRigBaseElement* Element = WeakHierarchy.Get()->Find(Key);
-        OnHierarchyModified(InNotif, WeakHierarchy.Get(), Element);
+        if (const FRigBaseElement* Element = WeakHierarchy.Get()->Find(Key))
+        {
+	        OnHierarchyModified(InNotif, WeakHierarchy.Get(), Element);
+        }
 		
     }, TStatId(), NULL, ENamedThreads::GameThread);
 }
diff --git a/Engine/Plugins/Animation/ControlRig/Source/ControlRigEditor/Private/Editor/SRigHierarchy.cpp b/Engine/Plugins/Animation/ControlRig/Source/ControlRigEditor/Private/Editor/SRigHierarchy.cpp
--- a/Engine/Plugins/Animation/ControlRig/Source/ControlRigEditor/Private/Editor/SRigHierarchy.cpp
+++ b/Engine/Plugins/Animation/ControlRig/Source/ControlRigEditor/Private/Editor/SRigHierarchy.cpp
@@ -806,8 +806,10 @@ void SRigHierarchy::OnHierarchyModified_AnyThread(ERigHierarchyNotification InNo
             {
                 return;
             }
-            const FRigBaseElement* Element = WeakHierarchy.Get()->Find(Key);
-            OnHierarchyModified(InNotif, WeakHierarchy.Get(), Element);
+            if (const FRigBaseElement* Element = WeakHierarchy.Get()->Find(Key))
+            {
+	            OnHierarchyModified(InNotif, WeakHierarchy.Get(), Element);
+            }
 			
         }, TStatId(), NULL, ENamedThreads::GameThread);
 	}
