Control Rig: Fix Reroute, If, Select, Array nodes not available when pulling from a pin

diff --git a/Engine/Plugins/Animation/ControlRig/Source/ControlRigDeveloper/Private/Graph/ControlRigGraphNode.cpp b/Engine/Plugins/Animation/ControlRig/Source/ControlRigDeveloper/Private/Graph/ControlRigGraphNode.cpp
--- a/Engine/Plugins/Animation/ControlRig/Source/ControlRigDeveloper/Private/Graph/ControlRigGraphNode.cpp
+++ b/Engine/Plugins/Animation/ControlRig/Source/ControlRigDeveloper/Private/Graph/ControlRigGraphNode.cpp
@@ -442,13 +442,14 @@ const FRigVMTemplate* UControlRigGraphNode::GetTemplate() const
 {
 	if(CachedTemplate == nullptr)
 	{
+		const FRigVMRegistry& Registry = FRigVMRegistry::Get();
 		if(URigVMTemplateNode* TemplateNode = Cast<URigVMTemplateNode>(GetModelNode()))
 		{
 			CachedTemplate = TemplateNode->GetTemplate();
 		}
-		else if(ModelNodePath.Contains(TEXT("::Execute(")))
+		else if(const FRigVMTemplate* Template = Registry.FindTemplate(*ModelNodePath))
 		{
-			CachedTemplate = FRigVMRegistry::Get().FindTemplate(*ModelNodePath); 
+			CachedTemplate = Template; 
 		}
 	}
 	return CachedTemplate;
diff --git a/Engine/Plugins/Animation/ControlRig/Source/ControlRigDeveloper/Public/Graph/ControlRigGraphNode.h b/Engine/Plugins/Animation/ControlRig/Source/ControlRigDeveloper/Public/Graph/ControlRigGraphNode.h
--- a/Engine/Plugins/Animation/ControlRig/Source/ControlRigDeveloper/Public/Graph/ControlRigGraphNode.h
+++ b/Engine/Plugins/Animation/ControlRig/Source/ControlRigDeveloper/Public/Graph/ControlRigGraphNode.h
@@ -268,4 +268,8 @@ private:
 	friend class FControlRigArgumentLayout;
 	friend class FControlRigGraphDetails;
 	friend class UControlRigTemplateNodeSpawner;
+	friend class UControlRigRerouteNodeSpawner;
+	friend class UControlRigArrayNodeSpawner;
+	friend class UControlRigIfNodeSpawner;
+	friend class UControlRigSelectNodeSpawner;
 };
diff --git a/Engine/Plugins/Animation/ControlRig/Source/ControlRigEditor/Private/Graph/NodeSpawners/ControlRigArrayNodeSpawner.cpp b/Engine/Plugins/Animation/ControlRig/Source/ControlRigEditor/Private/Graph/NodeSpawners/ControlRigArrayNodeSpawner.cpp
--- a/Engine/Plugins/Animation/ControlRig/Source/ControlRigEditor/Private/Graph/NodeSpawners/ControlRigArrayNodeSpawner.cpp
+++ b/Engine/Plugins/Animation/ControlRig/Source/ControlRigEditor/Private/Graph/NodeSpawners/ControlRigArrayNodeSpawner.cpp
@@ -189,6 +189,10 @@ UEdGraphNode* UControlRigArrayNodeSpawner::Invoke(UEdGraph* ParentGraph, FBindin
 		}
 		else
 		{
+			if(NewNode)
+			{
+				NewNode->ModelNodePath = ModelNode->GetNotation().ToString();
+			}
 			Controller->RemoveNode(ModelNode, false);
 		}
 	}
diff --git a/Engine/Plugins/Animation/ControlRig/Source/ControlRigEditor/Private/Graph/NodeSpawners/ControlRigIfNodeSpawner.cpp b/Engine/Plugins/Animation/ControlRig/Source/ControlRigEditor/Private/Graph/NodeSpawners/ControlRigIfNodeSpawner.cpp
--- a/Engine/Plugins/Animation/ControlRig/Source/ControlRigEditor/Private/Graph/NodeSpawners/ControlRigIfNodeSpawner.cpp
+++ b/Engine/Plugins/Animation/ControlRig/Source/ControlRigEditor/Private/Graph/NodeSpawners/ControlRigIfNodeSpawner.cpp
@@ -115,6 +115,10 @@ UEdGraphNode* UControlRigIfNodeSpawner::Invoke(UEdGraph* ParentGraph, FBindingSe
 		}
 		else
 		{
+			if(NewNode)
+			{
+				NewNode->ModelNodePath = ModelNode->GetNotation().ToString();
+			}
 			Controller->RemoveNode(ModelNode, false);
 		}
 	}
diff --git a/Engine/Plugins/Animation/ControlRig/Source/ControlRigEditor/Private/Graph/NodeSpawners/ControlRigRerouteNodeSpawner.cpp b/Engine/Plugins/Animation/ControlRig/Source/ControlRigEditor/Private/Graph/NodeSpawners/ControlRigRerouteNodeSpawner.cpp
--- a/Engine/Plugins/Animation/ControlRig/Source/ControlRigEditor/Private/Graph/NodeSpawners/ControlRigRerouteNodeSpawner.cpp
+++ b/Engine/Plugins/Animation/ControlRig/Source/ControlRigEditor/Private/Graph/NodeSpawners/ControlRigRerouteNodeSpawner.cpp
@@ -109,7 +109,7 @@ UEdGraphNode* UControlRigRerouteNodeSpawner::Invoke(UEdGraph* ParentGraph, FBind
 		}
 	}
 
-	URigVMNode* ModelNode = nullptr;
+	URigVMRerouteNode* ModelNode = nullptr;
 	if(!PinPath.IsEmpty())
 	{
 		ModelNode = Controller->AddRerouteNodeOnPin(PinPath, bIsInput, true, Location, FString(), bIsUserFacingNode, bIsUserFacingNode);
@@ -137,6 +137,10 @@ UEdGraphNode* UControlRigRerouteNodeSpawner::Invoke(UEdGraph* ParentGraph, FBind
 		}
 		else
 		{
+			if(NewNode)
+			{
+				NewNode->ModelNodePath = ModelNode->GetNotation().ToString();
+			}
 			Controller->RemoveNode(ModelNode, false);
 		}
 	}
diff --git a/Engine/Plugins/Animation/ControlRig/Source/ControlRigEditor/Private/Graph/NodeSpawners/ControlRigSelectNodeSpawner.cpp b/Engine/Plugins/Animation/ControlRig/Source/ControlRigEditor/Private/Graph/NodeSpawners/ControlRigSelectNodeSpawner.cpp
--- a/Engine/Plugins/Animation/ControlRig/Source/ControlRigEditor/Private/Graph/NodeSpawners/ControlRigSelectNodeSpawner.cpp
+++ b/Engine/Plugins/Animation/ControlRig/Source/ControlRigEditor/Private/Graph/NodeSpawners/ControlRigSelectNodeSpawner.cpp
@@ -129,6 +129,10 @@ UEdGraphNode* UControlRigSelectNodeSpawner::Invoke(UEdGraph* ParentGraph, FBindi
 		}
 		else
 		{
+			if(NewNode)
+			{
+				NewNode->ModelNodePath = ModelNode->GetNotation().ToString();
+			}
 			Controller->RemoveNode(ModelNode, false);
 		}
 	}
