Control Rig: Fix sequence node aggregate UI

diff --git a/Engine/Plugins/Animation/ControlRig/Source/ControlRigDeveloper/Private/Graph/ControlRigGraphNode.cpp b/Engine/Plugins/Animation/ControlRig/Source/ControlRigDeveloper/Private/Graph/ControlRigGraphNode.cpp
--- a/Engine/Plugins/Animation/ControlRig/Source/ControlRigDeveloper/Private/Graph/ControlRigGraphNode.cpp
+++ b/Engine/Plugins/Animation/ControlRig/Source/ControlRigDeveloper/Private/Graph/ControlRigGraphNode.cpp
@@ -1623,7 +1623,7 @@ TArray<URigVMPin*>& UControlRigGraphNode::PinListForPin(const URigVMPin* InModel
 {
 	if(IsValid(InModelPin))
 	{
-		if(InModelPin->IsExecuteContext())
+		if(InModelPin->IsExecuteContext() && InModelPin->GetDirection() == ERigVMPinDirection::IO)
 		{
 			return ExecutePins;
 		}
