Control Rig: Fix default values not showing up correctly

diff --git a/Engine/Plugins/Animation/ControlRig/Source/ControlRigDeveloper/Private/Graph/ControlRigGraphNode.cpp b/Engine/Plugins/Animation/ControlRig/Source/ControlRigDeveloper/Private/Graph/ControlRigGraphNode.cpp
--- a/Engine/Plugins/Animation/ControlRig/Source/ControlRigDeveloper/Private/Graph/ControlRigGraphNode.cpp
+++ b/Engine/Plugins/Animation/ControlRig/Source/ControlRigDeveloper/Private/Graph/ControlRigGraphNode.cpp
@@ -1484,7 +1484,7 @@ void UControlRigGraphNode::SetupPinDefaultsFromModel(UEdGraphPin* Pin, const URi
 	}
 
 	const UEdGraphSchema_K2* K2Schema = GetDefault<UEdGraphSchema_K2>();
-	if (InModelPin && !IsValid(InModelPin))
+	if (InModelPin && IsValid(InModelPin))
 	{
 		if (InModelPin->GetSubPins().Num() > 0)
 		{
