Control Rig: Guard against nullptrs in subpins

diff --git a/Engine/Plugins/Animation/ControlRig/Source/ControlRigDeveloper/Private/Graph/ControlRigGraphNode.cpp b/Engine/Plugins/Animation/ControlRig/Source/ControlRigDeveloper/Private/Graph/ControlRigGraphNode.cpp
--- a/Engine/Plugins/Animation/ControlRig/Source/ControlRigDeveloper/Private/Graph/ControlRigGraphNode.cpp
+++ b/Engine/Plugins/Animation/ControlRig/Source/ControlRigDeveloper/Private/Graph/ControlRigGraphNode.cpp
@@ -346,7 +346,7 @@ void UControlRigGraphNode::DestroyPinList(TArray<UEdGraphPin*>& InPins)
 	for (UEdGraphPin* Pin : InPins)
 	{
 		Pin->BreakAllPinLinks(bNotify);
-
+		Pin->SubPins.Remove(nullptr);
 		UEdGraphNode::DestroyPin(Pin);
 	}
 }
@@ -1483,6 +1483,9 @@ void UControlRigGraphNode::SetupPinDefaultsFromModel(UEdGraphPin* Pin, const URi
 		InModelPin = FindModelPinFromGraphPin(Pin);
 	}
 
+	// remove stale subpins
+	Pin->SubPins.Remove(nullptr);
+
 	const UEdGraphSchema_K2* K2Schema = GetDefault<UEdGraphSchema_K2>();
 	if (InModelPin && IsValid(InModelPin))
 	{
@@ -1588,6 +1591,15 @@ bool UControlRigGraphNode::ShouldDrawNodeAsControlPointOnly(int32& OutInputPinIn
 	return false;
 }
 
+void UControlRigGraphNode::BeginDestroy()
+{
+	for(UEdGraphPin* Pin : Pins)
+	{
+		Pin->SubPins.Remove(nullptr);
+	}
+	Super::BeginDestroy();
+}
+
 FEdGraphPinType UControlRigGraphNode::GetPinTypeForModelPin(const URigVMPin* InModelPin)
 {
 	FEdGraphPinType PinType = RigVMTypeUtils::PinTypeFromCPPType(*InModelPin->GetCPPType(), InModelPin->GetCPPTypeObject());
diff --git a/Engine/Plugins/Animation/ControlRig/Source/ControlRigDeveloper/Public/Graph/ControlRigGraphNode.h b/Engine/Plugins/Animation/ControlRig/Source/ControlRigDeveloper/Public/Graph/ControlRigGraphNode.h
--- a/Engine/Plugins/Animation/ControlRig/Source/ControlRigDeveloper/Public/Graph/ControlRigGraphNode.h
+++ b/Engine/Plugins/Animation/ControlRig/Source/ControlRigDeveloper/Public/Graph/ControlRigGraphNode.h
@@ -117,6 +117,7 @@ public:
 	virtual bool SupportsCommentBubble() const override { return false; }
 	virtual bool IsSelectedInEditor() const;
 	virtual bool ShouldDrawNodeAsControlPointOnly(int32& OutInputPinIndex, int32& OutOutputPinIndex) const override;
+	virtual void BeginDestroy() override;
 
 	virtual bool IsDeprecated() const override;
 	virtual FEdGraphNodeDeprecationResponse GetDeprecationResponse(EEdGraphNodeDeprecationType DeprecationType) const override;
