Control Rig: Fix python script not identifying execute context pins correctly

diff --git a/Engine/Plugins/Animation/ControlRig/Source/ControlRigDeveloper/Private/ControlRigBlueprint.cpp b/Engine/Plugins/Animation/ControlRig/Source/ControlRigDeveloper/Private/ControlRigBlueprint.cpp
--- a/Engine/Plugins/Animation/ControlRig/Source/ControlRigDeveloper/Private/ControlRigBlueprint.cpp
+++ b/Engine/Plugins/Animation/ControlRig/Source/ControlRigDeveloper/Private/ControlRigBlueprint.cpp
@@ -2246,6 +2246,7 @@ TArray<FString> UControlRigBlueprint::GeneratePythonCommands(const FString InNew
 					{
 
 						bool bHitFirstExecute = false;
+						bool bRenamedExecute = false;
 						for (auto Pin : EntryNode->GetPins())
 						{
 							if (Pin->GetDirection() != ERigVMPinDirection::Output)
@@ -2258,6 +2259,14 @@ TArray<FString> UControlRigBlueprint::GeneratePythonCommands(const FString InNew
 								if(!bHitFirstExecute)
 								{
 									bHitFirstExecute = true;
+									if (Pin->GetName() != FRigVMStruct::ExecuteContextName.ToString())
+									{
+										bRenamedExecute = true;
+										InternalCommands.Add(FString::Printf(TEXT("blueprint.get_controller_by_name('%s').rename_exposed_pin('%s', '%s')"),
+											*Graph->GetGraphName(),
+											*FRigVMStruct::ExecuteContextName.ToString(),
+											*Pin->GetName()));
+									}
 									continue;
 								}
 							}
@@ -2285,6 +2294,14 @@ TArray<FString> UControlRigBlueprint::GeneratePythonCommands(const FString InNew
 								if(!bHitFirstExecute)
 								{
 									bHitFirstExecute = true;
+									if (!bRenamedExecute && Pin->GetName() != FRigVMStruct::ExecuteContextName.ToString())
+									{
+										bRenamedExecute = true;
+										InternalCommands.Add(FString::Printf(TEXT("blueprint.get_controller_by_name('%s').rename_exposed_pin('%s', '%s')"),
+											*Graph->GetGraphName(),
+											*FRigVMStruct::ExecuteContextName.ToString(),
+											*Pin->GetName()));
+									}
 									continue;
 								}
 							}
diff --git a/Engine/Source/Developer/RigVMDeveloper/Private/RigVMModel/RigVMNode.cpp b/Engine/Source/Developer/RigVMDeveloper/Private/RigVMModel/RigVMNode.cpp
--- a/Engine/Source/Developer/RigVMDeveloper/Private/RigVMModel/RigVMNode.cpp
+++ b/Engine/Source/Developer/RigVMDeveloper/Private/RigVMModel/RigVMNode.cpp
@@ -285,9 +285,9 @@ bool URigVMNode::IsPure() const
 
 bool URigVMNode::IsMutable() const
 {
-	if (const URigVMPin* ExecutePin = FindPin(FRigVMStruct::ExecuteContextName.ToString()))
+	for (const URigVMPin* Pin : GetPins())
 	{
-		if(const UScriptStruct* ScriptStruct = ExecutePin->GetScriptStruct())
+		if(const UScriptStruct* ScriptStruct = Pin->GetScriptStruct())
 		{
 			return ScriptStruct->IsChildOf(FRigVMExecuteContext::StaticStruct());
 		}
