Control Rig: Support auto scroll + auto expand in Rig Tree view

diff --git a/Engine/Plugins/Animation/ControlRig/Source/ControlRigEditor/Private/Editor/SRigHierarchy.cpp b/Engine/Plugins/Animation/ControlRig/Source/ControlRigEditor/Private/Editor/SRigHierarchy.cpp
--- a/Engine/Plugins/Animation/ControlRig/Source/ControlRigEditor/Private/Editor/SRigHierarchy.cpp
+++ b/Engine/Plugins/Animation/ControlRig/Source/ControlRigEditor/Private/Editor/SRigHierarchy.cpp
@@ -229,6 +229,7 @@ void SRigHierarchy::Construct(const FArguments& InArgs, TSharedRef<FControlRigEd
 				[
 					SAssignNew(TreeView, SRigHierarchyTreeView)
 					.RigTreeDelegates(Delegates)
+					.AutoScrollEnabled(true)
 				]
 			]
 		]
diff --git a/Engine/Plugins/Animation/ControlRig/Source/ControlRigEditor/Private/Editor/SRigHierarchyTreeView.cpp b/Engine/Plugins/Animation/ControlRig/Source/ControlRigEditor/Private/Editor/SRigHierarchyTreeView.cpp
--- a/Engine/Plugins/Animation/ControlRig/Source/ControlRigEditor/Private/Editor/SRigHierarchyTreeView.cpp
+++ b/Engine/Plugins/Animation/ControlRig/Source/ControlRigEditor/Private/Editor/SRigHierarchyTreeView.cpp
@@ -206,6 +206,7 @@ FText SRigHierarchyItem::GetName() const
 void SRigHierarchyTreeView::Construct(const FArguments& InArgs)
 {
 	Delegates = InArgs._RigTreeDelegates;
+	bAutoScrollEnabled = InArgs._AutoScrollEnabled;
 
 	STreeView<TSharedPtr<FRigTreeElement>>::FArguments SuperArgs;
 	SuperArgs.TreeItemsSource(&RootElements);
@@ -231,6 +232,63 @@ void SRigHierarchyTreeView::Construct(const FArguments& InArgs)
 	});
 
 	STreeView<TSharedPtr<FRigTreeElement>>::Construct(SuperArgs);
+
+	LastMousePosition = FVector2D::ZeroVector;
+	TimeAtMousePosition = 0.0;
+}
+
+void SRigHierarchyTreeView::Tick(const FGeometry& AllottedGeometry, const double InCurrentTime, const float InDeltaTime)
+{
+	STreeView<TSharedPtr<FRigTreeElement, ESPMode::ThreadSafe>>::Tick(AllottedGeometry, InCurrentTime, InDeltaTime);
+
+	const FGeometry PaintGeometry = GetPaintSpaceGeometry();
+	const FVector2D MousePosition = FSlateApplication::Get().GetCursorPos();
+
+	if(PaintGeometry.IsUnderLocation(MousePosition))
+	{
+		const FVector2D WidgetPosition = PaintGeometry.AbsoluteToLocal(MousePosition);
+
+		static constexpr float SteadyMousePositionTolerance = 5.f;
+
+		if(LastMousePosition.Equals(MousePosition, SteadyMousePositionTolerance))
+		{
+			TimeAtMousePosition += InDeltaTime;
+		}
+		else
+		{
+			LastMousePosition = MousePosition;
+			TimeAtMousePosition = 0.0;
+		}
+
+		static constexpr float AutoScrollStartDuration = 0.5f; // in seconds
+		static constexpr float AutoScrollDistance = 24.f; // in pixels
+		static constexpr float AutoScrollSpeed = 150.f;
+
+		if(TimeAtMousePosition > AutoScrollStartDuration && FSlateApplication::Get().IsDragDropping())
+		{
+			if((WidgetPosition.Y < AutoScrollDistance) || (WidgetPosition.Y > PaintGeometry.Size.Y - AutoScrollDistance))
+			{
+				if(bAutoScrollEnabled)
+				{
+					const bool bScrollUp = (WidgetPosition.Y < AutoScrollDistance);
+
+					const float DeltaInSlateUnits = (bScrollUp ? -InDeltaTime : InDeltaTime) * AutoScrollSpeed; 
+					ScrollBy(GetCachedGeometry(), DeltaInSlateUnits, EAllowOverscroll::No);
+				}
+			}
+			else
+			{
+				const TSharedPtr<FRigTreeElement> Item = FindItemAtPosition(MousePosition);
+				if(Item.IsValid())
+				{
+					if(!IsItemExpanded(Item))
+					{
+						SetItemExpansion(Item, true);
+					}
+				}
+			}
+		}
+	}
 }
 
 TSharedPtr<FRigTreeElement> SRigHierarchyTreeView::FindElement(const FRigElementKey& InElementKey, TSharedPtr<FRigTreeElement> CurrentItem)
@@ -699,9 +757,22 @@ TSharedPtr<FRigTreeElement> SRigHierarchyTreeView::FindItemAtPosition(FVector2D
 		ItemsPanel->ArrangeChildren(MyGeometry, ArrangedChildren, true);
 
 		const int32 Index = ItemsPanel->FindChildUnderPosition(ArrangedChildren, InScreenSpacePosition); 
-		if(ItemsSource->IsValidIndex(Index))
+		if(ArrangedChildren.IsValidIndex(Index))
 		{
-			return ItemsSource->operator[](Index);
+			TSharedRef<SRigHierarchyItem> ItemWidget = StaticCastSharedRef<SRigHierarchyItem>(ArrangedChildren[Index].Widget);
+			if(ItemWidget->WeakRigTreeElement.IsValid())
+			{
+				const FRigElementKey Key = ItemWidget->WeakRigTreeElement.Pin()->Key;
+				const TSharedPtr<FRigTreeElement>* ResultPtr = ItemsSource->FindByPredicate([Key](const TSharedPtr<FRigTreeElement>& Item) -> bool
+				{
+					return Item->Key == Key;
+				});
+				
+				if(ResultPtr)
+				{
+					return *ResultPtr;
+				}
+			}
 		}
 	}
 	return TSharedPtr<FRigTreeElement>();
diff --git a/Engine/Plugins/Animation/ControlRig/Source/ControlRigEditor/Private/Editor/SRigHierarchyTreeView.h b/Engine/Plugins/Animation/ControlRig/Source/ControlRigEditor/Private/Editor/SRigHierarchyTreeView.h
--- a/Engine/Plugins/Animation/ControlRig/Source/ControlRigEditor/Private/Editor/SRigHierarchyTreeView.h
+++ b/Engine/Plugins/Animation/ControlRig/Source/ControlRigEditor/Private/Editor/SRigHierarchyTreeView.h
@@ -212,19 +212,27 @@ private:
  	FRigTreeDelegates Delegates;
 
 	FText GetName() const;
+
+	friend class SRigHierarchyTreeView; 
 };
 
 class SRigHierarchyTreeView : public STreeView<TSharedPtr<FRigTreeElement>>
 {
 public:
 
-	SLATE_BEGIN_ARGS(SRigHierarchyTreeView) {}
+	SLATE_BEGIN_ARGS(SRigHierarchyTreeView)
+		: _AutoScrollEnabled(false)
+	{}
 		SLATE_ARGUMENT(FRigTreeDelegates, RigTreeDelegates)
+		SLATE_ARGUMENT(bool, AutoScrollEnabled)
 	SLATE_END_ARGS()
 
 	void Construct(const FArguments& InArgs);
 	virtual ~SRigHierarchyTreeView() {}
 
+	/** Performs auto scroll */
+	virtual void Tick(const FGeometry& AllottedGeometry, const double InCurrentTime, const float InDeltaTime) override;
+
 	virtual FReply OnFocusReceived(const FGeometry& MyGeometry, const FFocusEvent& InFocusEvent) override
 	{
 		FReply Reply = STreeView<TSharedPtr<FRigTreeElement>>::OnFocusReceived(MyGeometry, InFocusEvent);
@@ -294,6 +302,10 @@ private:
 
 	FRigTreeDelegates Delegates;
 
+	bool bAutoScrollEnabled;
+	FVector2D LastMousePosition;
+	double TimeAtMousePosition;
+
 	friend class SRigHierarchy;
 };
 
