Control Rig: Fix control manipulation in the preview window is not undoable

diff --git a/Engine/Plugins/Animation/ControlRig/Source/ControlRigEditor/Private/EditMode/ControlRigEditMode.cpp b/Engine/Plugins/Animation/ControlRig/Source/ControlRigEditor/Private/EditMode/ControlRigEditMode.cpp
--- a/Engine/Plugins/Animation/ControlRig/Source/ControlRigEditor/Private/EditMode/ControlRigEditMode.cpp
+++ b/Engine/Plugins/Animation/ControlRig/Source/ControlRigEditor/Private/EditMode/ControlRigEditMode.cpp
@@ -3817,7 +3817,7 @@ void FControlRigEditMode::SetControlShapeTransform(
 		return;
 	}
 
-	static constexpr bool bNotify = true, bFixEuler = true, bUndo = false;
+	static constexpr bool bNotify = true, bFixEuler = true, bUndo = true;
 	if (!IsInLevelEditor())
 	{
 		// assumes it's attached to actor
