Control Rig: Avoid drawing cast links for float vs double

diff --git a/Engine/Plugins/Animation/ControlRig/Source/ControlRigEditor/Private/Graph/ControlRigConnectionDrawingPolicy.cpp b/Engine/Plugins/Animation/ControlRig/Source/ControlRigEditor/Private/Graph/ControlRigConnectionDrawingPolicy.cpp
--- a/Engine/Plugins/Animation/ControlRig/Source/ControlRigEditor/Private/Graph/ControlRigConnectionDrawingPolicy.cpp
+++ b/Engine/Plugins/Animation/ControlRig/Source/ControlRigEditor/Private/Graph/ControlRigConnectionDrawingPolicy.cpp
@@ -283,6 +283,11 @@ void FControlRigConnectionDrawingPolicy::DetermineWiringStyle(UEdGraphPin* Outpu
 		return;
 	}
 
+	if(OutputPin->Direction == EGPD_Input)
+	{
+		Swap(OutputPin, InputPin);
+	}
+
 	UControlRigGraphNode* OutputNode = Cast<UControlRigGraphNode>(OutputPin->GetOwningNode());
 	UControlRigGraphNode* InputNode = Cast<UControlRigGraphNode>(InputPin->GetOwningNode());
 	if (OutputNode && InputNode)
@@ -364,6 +369,8 @@ void FControlRigConnectionDrawingPolicy::DetermineWiringStyle(UEdGraphPin* Outpu
 				}
 			}
 
+			Params.bUserFlag1 = false;
+
 			if (bVisited)
 			{
 				//Params.bDrawBubbles = true;
@@ -374,10 +381,12 @@ void FControlRigConnectionDrawingPolicy::DetermineWiringStyle(UEdGraphPin* Outpu
 				Params.WireColor = Params.WireColor * 0.5f;
 			}
 
-			Params.bUserFlag1 = false;
 			if(!UseLowDetailConnections() && !bInjectionIsSelected)
 			{
-				if(OutputModelPin->GetTypeIndex() != InputModelPin->GetTypeIndex())
+				const TRigVMTypeIndex OutputTypeIndex = OutputModelPin->GetTypeIndex();
+				const TRigVMTypeIndex InputTypeIndex = InputModelPin->GetTypeIndex();
+				if(OutputTypeIndex != InputTypeIndex
+					&& !FRigVMRegistry::Get().CanMatchTypes(OutputTypeIndex, InputTypeIndex, true))
 				{
 					Params.bUserFlag1 = true;
 				}
