Control Rig: Pose based testing framework

diff --git a/Engine/Plugins/Animation/ControlRig/Source/ControlRig/Private/ControlRigTestData.cpp b/Engine/Plugins/Animation/ControlRig/Source/ControlRig/Private/ControlRigTestData.cpp
new file mode 100644
--- /dev/null
+++ b/Engine/Plugins/Animation/ControlRig/Source/ControlRig/Private/ControlRigTestData.cpp
@@ -0,0 +1,415 @@
+// Copyright Epic Games, Inc. All Rights Reserved.
+
+#include "ControlRigTestData.h"
+#include "ControlRig.h"
+#include "HAL/PlatformTime.h"
+#if WITH_EDITOR
+#include "AssetToolsModule.h"
+#endif
+
+#include UE_INLINE_GENERATED_CPP_BY_NAME(ControlRigTestData)
+
+bool FControlRigTestDataFrame::Store(UControlRig* InControlRig, bool bInitial)
+{
+	if(InControlRig == nullptr)
+	{
+		return false;
+	}
+
+	AbsoluteTime = InControlRig->GetAbsoluteTime();
+	DeltaTime = InControlRig->GetDeltaTime();
+	Pose = InControlRig->GetHierarchy()->GetPose(bInitial);
+	Variables.Reset();
+
+	const TArray<FRigVMExternalVariable> ExternalVariables = InControlRig->GetExternalVariables(); 
+	for(const FRigVMExternalVariable& ExternalVariable : ExternalVariables)
+	{
+		FControlRigTestDataVariable VariableData;
+		VariableData.Name = ExternalVariable.Name;
+		VariableData.CPPType = ExternalVariable.TypeName;
+
+		if(ExternalVariable.Property && ExternalVariable.Memory)
+		{
+			ExternalVariable.Property->ExportText_Direct(
+				VariableData.Value,
+				ExternalVariable.Memory,
+				ExternalVariable.Memory,
+				nullptr,
+				PPF_None,
+				nullptr
+			);
+		}
+
+		Variables.Add(VariableData);
+	}
+
+	return true;
+}
+
+bool FControlRigTestDataFrame::Restore(UControlRig* InControlRig, bool bInitial) const
+{
+	if(InControlRig == nullptr)
+	{
+		return false;
+	}
+
+	URigHierarchy* Hierarchy = InControlRig->GetHierarchy();
+
+	// check if the pose can be applied
+	for(const FRigPoseElement& PoseElement : Pose.Elements)
+	{
+		const FRigElementKey& Key = PoseElement.Index.GetKey(); 
+		if(!Hierarchy->Contains(Key))
+		{
+			UE_LOG(LogControlRig, Error, TEXT("Control Rig does not contain hierarchy element '%s'. Please re-create the test data asset."), *Key.ToString());
+			return false;
+		}
+	}
+
+	Hierarchy->SetPose(Pose, bInitial ? ERigTransformType::InitialLocal : ERigTransformType::CurrentLocal);
+
+	return RestoreVariables(InControlRig);
+}
+
+bool FControlRigTestDataFrame::RestoreVariables(UControlRig* InControlRig) const
+{
+	class FControlRigTestDataFrame_ErrorPipe : public FOutputDevice
+	{
+	public:
+
+		TArray<FString> Errors;
+
+		FControlRigTestDataFrame_ErrorPipe()
+			: FOutputDevice()
+		{
+		}
+
+		virtual void Serialize(const TCHAR* V, ELogVerbosity::Type Verbosity, const class FName& Category) override
+		{
+			Errors.Add(FString::Printf(TEXT("Error convert to string: %s"), V));
+		}
+	};
+
+	const TArray<FRigVMExternalVariable> ExternalVariables = InControlRig->GetExternalVariables();
+
+	if(ExternalVariables.Num() != Variables.Num())
+	{
+		UE_LOG(LogControlRig, Error, TEXT("Variable data does not match the Rig. Please re-create the test data asset."));
+		return false;
+	}
+	
+	for(const FRigVMExternalVariable& ExternalVariable : ExternalVariables)
+	{
+		if(ExternalVariable.Memory == nullptr || ExternalVariable.Property == nullptr)
+		{
+			UE_LOG(LogControlRig, Error, TEXT("Variable '%s' is not valid."), *ExternalVariable.Name.ToString());
+			return false;
+		}
+		
+		const FControlRigTestDataVariable* VariableData = Variables.FindByPredicate(
+			[ExternalVariable](const FControlRigTestDataVariable& InVariable) -> bool
+			{
+				return InVariable.Name == ExternalVariable.Name &&
+					InVariable.CPPType == ExternalVariable.TypeName;
+			}
+		);
+
+		if(VariableData)
+		{
+			FControlRigTestDataFrame_ErrorPipe ErrorPipe;
+			ExternalVariable.Property->ImportText_Direct(
+				*VariableData->Value,
+				ExternalVariable.Memory,
+				nullptr,
+				PPF_None,
+				&ErrorPipe
+			);
+
+			if(!ErrorPipe.Errors.IsEmpty())
+			{
+				for(const FString& ImportError : ErrorPipe.Errors)
+				{
+					UE_LOG(LogControlRig, Error, TEXT("Import Error for Variable '%s': %s"), *ExternalVariable.Name.ToString(), *ImportError);
+				}
+				return false;
+			}
+		}
+		else
+		{
+			UE_LOG(LogControlRig, Error, TEXT("Variable data for '%s' is not part of the test file. Please re-create the test data asset."), *ExternalVariable.Name.ToString());
+			return false;
+		}
+	}
+
+	return true;
+}
+
+void UControlRigTestData::Serialize(FArchive& Ar)
+{
+	UObject::Serialize(Ar);
+	LastFrameIndex = INDEX_NONE;
+}
+
+UControlRigTestData* UControlRigTestData::CreateNewAsset(FString InDesiredPackagePath, FString InBlueprintPathName)
+{
+#if WITH_EDITOR
+	const FAssetToolsModule& AssetToolsModule = FModuleManager::LoadModuleChecked<FAssetToolsModule>("AssetTools");
+	FString UniquePackageName;
+	FString UniqueAssetName;
+	AssetToolsModule.Get().CreateUniqueAssetName(InDesiredPackagePath, TEXT(""), UniquePackageName, UniqueAssetName);
+
+	if (UniquePackageName.EndsWith(UniqueAssetName))
+	{
+		UniquePackageName = UniquePackageName.LeftChop(UniqueAssetName.Len() + 1);
+	}
+
+	UObject* NewAsset = AssetToolsModule.Get().CreateAsset(*UniqueAssetName, *UniquePackageName, UControlRigTestData::StaticClass(), nullptr);
+	if(NewAsset)
+	{
+		// make sure the package is never cooked.
+		UPackage* Package = NewAsset->GetOutermost();
+		Package->SetPackageFlags(Package->GetPackageFlags() | PKG_EditorOnly);
+			
+		if(UControlRigTestData* TestData = Cast<UControlRigTestData>(NewAsset))
+		{
+			TestData->ControlRigObjectPath = InBlueprintPathName;
+			return TestData;
+		}
+	}
+#endif
+	return nullptr;
+}
+
+FVector2D UControlRigTestData::GetTimeRange(bool bInput) const
+{
+	const TArray<FControlRigTestDataFrame>& Frames = bInput ? InputFrames : OutputFrames;
+	if(Frames.IsEmpty())
+	{
+		return FVector2D::ZeroVector;
+	}
+	return FVector2D(Frames[0].AbsoluteTime, Frames.Last().AbsoluteTime);
+}
+
+int32 UControlRigTestData::GetFrameIndexForTime(double InSeconds, bool bInput) const
+{
+	const TArray<FControlRigTestDataFrame>& Frames = bInput ? InputFrames : OutputFrames;
+
+	if(LastFrameIndex == INDEX_NONE)
+	{
+		LastFrameIndex = 0;
+	}
+
+	while(Frames.IsValidIndex(LastFrameIndex) && Frames[LastFrameIndex].AbsoluteTime < InSeconds)
+	{
+		LastFrameIndex++;
+	}
+
+	while(Frames.IsValidIndex(LastFrameIndex) && Frames[LastFrameIndex].AbsoluteTime > InSeconds)
+	{
+		LastFrameIndex--;
+	}
+
+	LastFrameIndex = FMath::Clamp(LastFrameIndex, 0, Frames.Num() - 1);
+
+	return LastFrameIndex;
+}
+
+bool UControlRigTestData::Record(UControlRig* InControlRig, double InRecordingDuration)
+{
+	if(InControlRig == nullptr)
+	{
+		return false;
+	}
+
+	ReleaseReplay();
+	ClearDelegates(InControlRig);
+
+	DesiredRecordingDuration = InRecordingDuration;
+	TimeAtStartOfRecording = FPlatformTime::Seconds();
+
+	// if this is the first frame
+	if(InputFrames.IsEmpty())
+	{
+		InControlRig->RequestInit();
+		PreConstructionHandle = InControlRig->OnPreConstruction_AnyThread().AddLambda(
+			[this](UControlRig* InControlRig, const EControlRigState InState, const FName& InEventName)
+			{
+				Initial.Store(InControlRig, true);
+			}
+		);
+	}
+
+	PreForwardHandle = InControlRig->OnPreForwardsSolve_AnyThread().AddLambda(
+		[this](UControlRig* InControlRig, const EControlRigState InState, const FName& InEventName)
+		{
+			FControlRigTestDataFrame Frame;
+			Frame.Store(InControlRig);
+
+			// reapply the variable data. we are doing this to make sure that
+			// the results in the rig are the same during a recording and replay.
+			Frame.RestoreVariables(InControlRig);
+			
+			InputFrames.Add(Frame);
+		}
+	);
+
+	PostForwardHandle = InControlRig->OnPostForwardsSolve_AnyThread().AddLambda(
+		[this](UControlRig* InControlRig, const EControlRigState InState, const FName& InEventName)
+		{
+			FControlRigTestDataFrame Frame;
+			Frame.Store(InControlRig);
+			OutputFrames.Add(Frame);
+			LastFrameIndex = INDEX_NONE;
+			(void)MarkPackageDirty();
+
+			const double TimeNow = FPlatformTime::Seconds();
+			const double TimeDelta = TimeNow - TimeAtStartOfRecording;
+			if(DesiredRecordingDuration <= TimeDelta)
+			{
+				ClearDelegates(InControlRig);
+				DesiredRecordingDuration = 0.0;
+			}
+		}
+	);
+
+	return true;
+}
+
+bool UControlRigTestData::SetupReplay(UControlRig* InControlRig, bool bGroundTruth)
+{
+	ReleaseReplay();
+	ClearDelegates(InControlRig);
+
+	if(InControlRig == nullptr)
+	{
+		return false;
+	}
+	
+	bIsApplyingOutputs = bGroundTruth;
+
+	if(InputFrames.IsEmpty() || OutputFrames.IsEmpty())
+	{
+		return false;
+	}
+
+	// reset the control rig's absolute time
+	InControlRig->SetAbsoluteAndDeltaTime(InputFrames[0].AbsoluteTime, InputFrames[0].DeltaTime);
+	
+	InControlRig->RequestInit();
+	PreConstructionHandle = InControlRig->OnPreConstruction_AnyThread().AddLambda(
+		[this](UControlRig* InControlRig, const EControlRigState InState, const FName& InEventName)
+		{
+			Initial.Restore(InControlRig, true);
+		}
+	);
+
+	PreForwardHandle = InControlRig->OnPreForwardsSolve_AnyThread().AddLambda(
+		[this](UControlRig* InControlRig, const EControlRigState InState, const FName& InEventName)
+		{
+			// loop the animation data
+			if(InControlRig->GetAbsoluteTime() < GetTimeRange().X - SMALL_NUMBER ||
+				InControlRig->GetAbsoluteTime() > GetTimeRange().Y + SMALL_NUMBER)
+			{
+				InControlRig->SetAbsoluteAndDeltaTime(GetTimeRange().X, InControlRig->GetDeltaTime());
+			}
+			
+			const int32 FrameIndex = GetFrameIndexForTime(InControlRig->GetAbsoluteTime(), true);
+			const FControlRigTestDataFrame& Frame = InputFrames[FrameIndex];
+			Frame.Restore(InControlRig, false);
+
+			if(Frame.DeltaTime > SMALL_NUMBER)
+			{
+				InControlRig->SetDeltaTime(Frame.DeltaTime);
+			}
+		}
+	);
+
+	PostForwardHandle = InControlRig->OnPostForwardsSolve_AnyThread().AddLambda(
+		[this](UControlRig* InControlRig, const EControlRigState InState, const FName& InEventName)
+		{
+			const FRigPose CurrentPose = InControlRig->GetHierarchy()->GetPose();
+
+			const FControlRigTestDataFrame& Frame = OutputFrames[LastFrameIndex];
+			if(bIsApplyingOutputs)
+			{
+				Frame.Restore(InControlRig, false);
+			}
+
+			const FRigPose& ExpectedPose = Frame.Pose;
+
+			// draw differences of the pose result of the rig onto the screen
+			FControlRigDrawInterface& DrawInterface = InControlRig->GetDrawInterface();
+			for(const FRigPoseElement& ExpectedPoseElement : ExpectedPose)
+			{
+				const int32 CurrentPoseIndex = CurrentPose.GetIndex(ExpectedPoseElement.Index.GetKey());
+				if(CurrentPoseIndex != INDEX_NONE)
+				{
+					const FRigPoseElement& CurrentPoseElement = CurrentPose[CurrentPoseIndex];
+
+					if(!CurrentPoseElement.LocalTransform.Equals(ExpectedPoseElement.LocalTransform, 0.001f))
+					{
+						DrawInterface.DrawAxes(
+							FTransform::Identity,
+							bIsApplyingOutputs ? CurrentPoseElement.GlobalTransform : ExpectedPoseElement.GlobalTransform,
+							bIsApplyingOutputs ? FLinearColor::Red : FLinearColor::Green,
+							15.0f,
+							1.0f
+						);
+					}
+				}
+			}
+		}
+	);
+
+	ReplayControlRig = InControlRig;
+	return true;
+}
+
+void UControlRigTestData::ReleaseReplay()
+{
+	if(UControlRig* ControlRig = ReplayControlRig.Get())
+	{
+		ClearDelegates(ControlRig);
+		ReplayControlRig.Reset();
+	}
+}
+
+EControlRigTestDataPlaybackMode UControlRigTestData::GetPlaybackMode() const
+{
+	if(IsReplaying())
+	{
+		if(bIsApplyingOutputs)
+		{
+			return EControlRigTestDataPlaybackMode::GroundTruth;
+		}
+		return EControlRigTestDataPlaybackMode::ReplayInputs;
+	}
+	return EControlRigTestDataPlaybackMode::Live;
+}
+
+bool UControlRigTestData::IsReplaying() const
+{
+	return ReplayControlRig.IsValid();
+}
+
+void UControlRigTestData::ClearDelegates(UControlRig* InControlRig)
+{
+	if(InControlRig)
+	{
+		if(PreConstructionHandle.IsValid())
+		{
+			InControlRig->OnPreConstruction_AnyThread().Remove(PreConstructionHandle);
+			PreConstructionHandle.Reset();
+		}
+		if(PreForwardHandle.IsValid())
+		{
+			InControlRig->OnPreForwardsSolve_AnyThread().Remove(PreForwardHandle);
+			PreForwardHandle.Reset();
+		}
+		if(PostForwardHandle.IsValid())
+		{
+			InControlRig->OnPostForwardsSolve_AnyThread().Remove(PostForwardHandle);
+			PostForwardHandle.Reset();
+		}
+	}
+}
diff --git a/Engine/Plugins/Animation/ControlRig/Source/ControlRig/Public/ControlRig.h b/Engine/Plugins/Animation/ControlRig/Source/ControlRig/Public/ControlRig.h
--- a/Engine/Plugins/Animation/ControlRig/Source/ControlRig/Public/ControlRig.h
+++ b/Engine/Plugins/Animation/ControlRig/Source/ControlRig/Public/ControlRig.h
@@ -138,6 +138,10 @@ public:
 	UFUNCTION(BlueprintPure, Category = "Control Rig")
 	float GetAbsoluteTime() const { return AbsoluteTime; }
 
+	/** Gets the current delta time */
+	UFUNCTION(BlueprintPure, Category = "Control Rig")
+	float GetDeltaTime() const { return DeltaTime; }
+
 	/** Set the current delta time */
 	UFUNCTION(BlueprintCallable, Category="Control Rig")
 	void SetDeltaTime(float InDeltaTime);
diff --git a/Engine/Plugins/Animation/ControlRig/Source/ControlRig/Public/ControlRigTestData.h b/Engine/Plugins/Animation/ControlRig/Source/ControlRig/Public/ControlRigTestData.h
new file mode 100644
--- /dev/null
+++ b/Engine/Plugins/Animation/ControlRig/Source/ControlRig/Public/ControlRigTestData.h
@@ -0,0 +1,141 @@
+// Copyright Epic Games, Inc. All Rights Reserved.
+
+#pragma once
+
+#include "CoreMinimal.h"
+#include "Rigs/RigHierarchyPose.h"
+
+#include "ControlRigTestData.generated.h"
+
+class UControlRigTestData;
+class UControlRig;
+
+USTRUCT(BlueprintType)
+struct CONTROLRIG_API FControlRigTestDataVariable
+{
+	GENERATED_USTRUCT_BODY()
+
+	FControlRigTestDataVariable()
+	{
+		Name = NAME_None;
+		CPPType = NAME_None;
+		Value = FString();
+	}
+
+	UPROPERTY(VisibleAnywhere, BlueprintReadOnly, Category = "ControlRigTestDataVariable")
+	FName Name;
+
+	UPROPERTY(VisibleAnywhere, BlueprintReadOnly, Category = "ControlRigTestDataVariable")
+	FName CPPType;
+
+	UPROPERTY(VisibleAnywhere, BlueprintReadOnly, Category = "ControlRigTestDataVariable")
+	FString Value;
+};
+
+USTRUCT(BlueprintType)
+struct CONTROLRIG_API FControlRigTestDataFrame
+{
+	GENERATED_USTRUCT_BODY()
+
+	FControlRigTestDataFrame()
+	{
+		AbsoluteTime = 0.0;
+		DeltaTime = 0.0;
+		Variables.Reset();
+		Pose.Reset();
+	}
+
+	bool Store(UControlRig* InControlRig, bool bInitial = false);
+	bool Restore(UControlRig* InControlRig, bool bInitial = false) const;
+	bool RestoreVariables(UControlRig* InControlRig) const;
+
+	UPROPERTY(VisibleAnywhere, BlueprintReadOnly, Category = "ControlRigTestDataFrame")
+	double AbsoluteTime;
+
+	UPROPERTY(VisibleAnywhere, BlueprintReadOnly, Category = "ControlRigTestDataFrame")
+	double DeltaTime;
+
+	UPROPERTY(VisibleAnywhere, BlueprintReadOnly, Category = "ControlRigTestDataFrame")
+	TArray<FControlRigTestDataVariable> Variables;
+
+	UPROPERTY(VisibleAnywhere, BlueprintReadOnly, Category = "ControlRigTestDataFrame")
+	FRigPose Pose;
+};
+
+UENUM()
+enum class EControlRigTestDataPlaybackMode : uint8
+{
+	Live,
+	ReplayInputs,
+	GroundTruth,
+	Max UMETA(Hidden),
+};
+
+UCLASS(BlueprintType)
+class CONTROLRIG_API UControlRigTestData : public UObject
+{
+	GENERATED_BODY()
+
+public:
+
+	UControlRigTestData()
+		: LastFrameIndex(INDEX_NONE)
+		, DesiredRecordingDuration(0.0)
+		, TimeAtStartOfRecording(0.0)
+		, bIsApplyingOutputs(false)
+	{}
+
+	virtual void Serialize(FArchive& Ar) override;
+
+	UFUNCTION(BlueprintCallable, Category = "ControlRigTestData")
+	static UControlRigTestData* CreateNewAsset(FString InDesiredPackagePath, FString InBlueprintPathName);
+
+	UFUNCTION(BlueprintPure, Category = "ControlRigTestData")
+	FVector2D GetTimeRange(bool bInput = false) const;
+
+	UFUNCTION(BlueprintPure, Category = "ControlRigTestData")
+	int32 GetFrameIndexForTime(double InSeconds, bool bInput = false) const;
+
+	UFUNCTION(BlueprintCallable, Category = "ControlRigTestData")
+	bool Record(UControlRig* InControlRig, double InRecordingDuration = 0.0);
+
+	UFUNCTION(BlueprintCallable, Category = "ControlRigTestData")
+	bool SetupReplay(UControlRig* InControlRig, bool bGroundTruth = true);
+
+	UFUNCTION(BlueprintCallable, Category = "ControlRigTestData")
+	void ReleaseReplay();
+
+	UFUNCTION(BlueprintPure, Category = "ControlRigTestData")
+	EControlRigTestDataPlaybackMode GetPlaybackMode() const;
+
+	UFUNCTION(BlueprintPure, Category = "ControlRigTestData")
+	bool IsReplaying() const;
+
+	UFUNCTION(BlueprintPure, Category = "ControlRigTestData")
+	bool IsRecording() const { return DesiredRecordingDuration >= SMALL_NUMBER; }
+
+	UPROPERTY(AssetRegistrySearchable, VisibleAnywhere, BlueprintReadOnly, Category = "ControlRigTestData")
+	FSoftObjectPath ControlRigObjectPath;
+
+	UPROPERTY(VisibleAnywhere, BlueprintReadOnly, Category = "ControlRigTestData")
+	FControlRigTestDataFrame Initial;
+
+	UPROPERTY(VisibleAnywhere, BlueprintReadOnly, Category = "ControlRigTestData")
+	TArray<FControlRigTestDataFrame> InputFrames;
+
+	UPROPERTY(VisibleAnywhere, BlueprintReadOnly, Category = "ControlRigTestData")
+	TArray<FControlRigTestDataFrame> OutputFrames;
+
+private:
+
+	void ClearDelegates(UControlRig* InControlRig);
+
+	mutable int32 LastFrameIndex;
+	double DesiredRecordingDuration;
+	double TimeAtStartOfRecording;
+	TWeakObjectPtr<UControlRig> ReplayControlRig;
+	bool bIsApplyingOutputs;
+	FDelegateHandle PreConstructionHandle;
+	FDelegateHandle PreForwardHandle;
+	FDelegateHandle PostForwardHandle;
+};
\ No newline at end of file
diff --git a/Engine/Plugins/Animation/ControlRig/Source/ControlRigEditor/Private/ControlRigEditorStyle.h b/Engine/Plugins/Animation/ControlRig/Source/ControlRigEditor/Private/ControlRigEditorStyle.h
--- a/Engine/Plugins/Animation/ControlRig/Source/ControlRigEditor/Private/ControlRigEditorStyle.h
+++ b/Engine/Plugins/Animation/ControlRig/Source/ControlRigEditor/Private/ControlRigEditorStyle.h
@@ -166,6 +166,11 @@ public:
 			Set("ControlRig.SpacePicker.RoundedRect", new FSlateRoundedBoxBrush(FStyleColors::White, 4.0f, FStyleColors::Transparent, 0.0f));
 		}
 
+		// Test Data
+		{
+			Set("ControlRig.TestData.Record", new IMAGE_BRUSH("Slate/RecordingIndicator", Icon32x32));
+		}
+
 		FSlateStyleRegistry::RegisterSlateStyle(*this);
 	}
 
diff --git a/Engine/Plugins/Animation/ControlRig/Source/ControlRigEditor/Private/Editor/ControlRigEditor.cpp b/Engine/Plugins/Animation/ControlRig/Source/ControlRigEditor/Private/Editor/ControlRigEditor.cpp
--- a/Engine/Plugins/Animation/ControlRig/Source/ControlRigEditor/Private/Editor/ControlRigEditor.cpp
+++ b/Engine/Plugins/Animation/ControlRig/Source/ControlRigEditor/Private/Editor/ControlRigEditor.cpp
@@ -103,6 +103,8 @@
 
 #define LOCTEXT_NAMESPACE "ControlRigEditor"
 
+TAutoConsoleVariable<bool> CVarControlRigShowTestingToolbar(TEXT("ControlRig.Test.EnableTestingToolbar"), false, TEXT("When true we'll show the testing toolbar in Control Rig Editor."));
+
 const FName ControlRigEditorAppName(TEXT("ControlRigEditorApp"));
 
 const FName FControlRigEditorModes::ControlRigEditorMode("Rigging");
@@ -855,6 +857,213 @@ TSharedRef<SWidget> FControlRigEditor::GenerateExecutionModeMenuContent()
 	return MenuBuilder.MakeWidget();
 }
 
+FText FControlRigEditor::GetTestAssetName() const
+{
+	if(TestDataStrongPtr.IsValid())
+	{
+		return FText::FromString(TestDataStrongPtr->GetName());
+	}
+
+	static const FText NoTestAsset = LOCTEXT("NoTestAsset", "No Test Asset");
+	return NoTestAsset;
+}
+
+FText FControlRigEditor::GetTestAssetTooltip() const
+{
+	if(TestDataStrongPtr.IsValid())
+	{
+		return FText::FromString(TestDataStrongPtr->GetPathName());
+	}
+	static const FText NoTestAssetTooltip = LOCTEXT("NoTestAssetTooltip", "Click the record button to the left to record a new frame");
+	return NoTestAssetTooltip;
+}
+
+bool FControlRigEditor::SetTestAssetPath(const FString& InAssetPath)
+{
+	if(TestDataStrongPtr.IsValid())
+	{
+		if(TestDataStrongPtr->GetPathName().Equals(InAssetPath, ESearchCase::CaseSensitive))
+		{
+			return false;
+		}
+	}
+
+	if(TestDataStrongPtr.IsValid())
+	{
+		TestDataStrongPtr->ReleaseReplay();
+		TestDataStrongPtr.Reset();
+	}
+
+	if(!InAssetPath.IsEmpty())
+	{
+		if(UControlRigTestData* TestData = LoadObject<UControlRigTestData>(GetControlRigBlueprint(), *InAssetPath))
+		{
+			TestDataStrongPtr = TStrongObjectPtr<UControlRigTestData>(TestData);
+		}
+	}
+	return true;
+}
+
+TSharedRef<SWidget> FControlRigEditor::GenerateTestAssetModeMenuContent()
+{
+	FMenuBuilder MenuBuilder(true, GetToolkitCommands());
+
+	MenuBuilder.BeginSection(TEXT("Default"));
+	MenuBuilder.AddMenuEntry(
+		LOCTEXT("ClearTestAssset", "Clear"),
+		LOCTEXT("ClearTestAsset_ToolTip", "Clears the test asset"),
+		FSlateIcon(),
+		FUIAction(
+			FExecuteAction::CreateLambda([this]()
+			{
+				SetTestAssetPath(FString());
+			}
+		)	
+	));
+	MenuBuilder.EndSection();
+
+	const FAssetRegistryModule& AssetRegistryModule = FModuleManager::LoadModuleChecked<FAssetRegistryModule>("AssetRegistry");
+	TArray<FAssetData> TestDataAssets;
+	FARFilter AssetFilter;
+	AssetFilter.ClassPaths.Add(UControlRigTestData::StaticClass()->GetClassPathName());
+	AssetRegistryModule.Get().GetAssets(AssetFilter, TestDataAssets);
+
+	const FString CurrentObjectPath = GetControlRigBlueprint()->GetPathName();
+	TestDataAssets.RemoveAll([CurrentObjectPath](const FAssetData& InAssetData)
+	{
+		const FString ControlRigObjectPath = InAssetData.GetTagValueRef<FString>(GET_MEMBER_NAME_CHECKED(UControlRigTestData, ControlRigObjectPath));
+		return ControlRigObjectPath != CurrentObjectPath;
+	});
+
+	if(!TestDataAssets.IsEmpty())
+	{
+		MenuBuilder.BeginSection(TEXT("Assets"));
+		for(const FAssetData& TestDataAsset : TestDataAssets)
+		{
+			const FString TestDataObjectPath = TestDataAsset.GetObjectPathString();
+			FString Right;
+			if(TestDataObjectPath.Split(TEXT("."), nullptr, &Right))
+			{
+				MenuBuilder.AddMenuEntry(FText::FromString(Right), FText::FromString(TestDataObjectPath), FSlateIcon(),
+					FUIAction(
+						FExecuteAction::CreateLambda([this, TestDataObjectPath]()
+						{
+							SetTestAssetPath(TestDataObjectPath);
+						}
+					)	
+				));
+			}
+		}
+		MenuBuilder.EndSection();
+	}
+	return MenuBuilder.MakeWidget();
+}
+
+TSharedRef<SWidget> FControlRigEditor::GenerateTestAssetRecordMenuContent()
+{
+	FMenuBuilder MenuBuilder(true, GetToolkitCommands());
+
+	MenuBuilder.AddMenuEntry(
+		LOCTEXT("TestAssetRecordSingleFrame", "Single Frame"),
+		LOCTEXT("TestAssetRecordSingleFrame_ToolTip", "Records a single frame into the test data asset"),
+		FSlateIcon(),
+		FUIAction(
+			FExecuteAction::CreateLambda([this]()
+			{
+				RecordTestData(0);
+			}
+		)	
+	));
+	MenuBuilder.AddMenuEntry(
+		LOCTEXT("TestAssetRecordOneSecond", "1 Second"),
+		LOCTEXT("TestAssetRecordOneSecond_ToolTip", "Records 1 second of animation into the test data asset"),
+		FSlateIcon(),
+		FUIAction(
+			FExecuteAction::CreateLambda([this]()
+			{
+				RecordTestData(1);
+			}
+		)	
+	));
+
+	MenuBuilder.AddMenuEntry(
+	LOCTEXT("TestAssetRecordFiveSeconds", "5 Seconds"),
+	LOCTEXT("TestAssetRecordFiveSeconds_ToolTip", "Records 5 seconds of animation into the test data asset"),
+	FSlateIcon(),
+	FUIAction(
+			FExecuteAction::CreateLambda([this]()
+			{
+				RecordTestData(5);
+			}
+		)	
+	));
+
+	MenuBuilder.AddMenuEntry(
+	LOCTEXT("TestAssetRecordTenSeconds", "10 Seconds"),
+	LOCTEXT("TestAssetRecordTenSeconds_ToolTip", "Records 10 seconds of animation into the test data asset"),
+	FSlateIcon(),
+	FUIAction(
+			FExecuteAction::CreateLambda([this]()
+			{
+				RecordTestData(10);
+			}
+		)	
+	));
+
+	return MenuBuilder.MakeWidget();
+}
+
+bool FControlRigEditor::RecordTestData(double InRecordingDuration)
+{
+	if(ControlRig == nullptr)
+	{
+		return false;
+	}
+	
+	if(!TestDataStrongPtr.IsValid())
+	{
+		// create a new test asset
+		static const FString Folder = TEXT("/Game/Animation/ControlRig/NoCook/");
+		const FString DesiredPackagePath =  FString::Printf(TEXT("%s/%s_TestData"), *Folder, *GetControlRigBlueprint()->GetName());
+
+		if(UControlRigTestData* TestData = UControlRigTestData::CreateNewAsset(DesiredPackagePath, GetControlRigBlueprint()->GetPathName()))
+		{
+			SetTestAssetPath(TestData->GetPathName());
+		}
+	}
+	
+	if(UControlRigTestData* TestData = TestDataStrongPtr.Get())
+	{
+		TestData->Record(ControlRig, InRecordingDuration);
+	}
+	return true;
+}
+
+void FControlRigEditor::ToggleTestData()
+{
+	if(TestDataStrongPtr.IsValid())
+	{
+		switch(TestDataStrongPtr->GetPlaybackMode())
+		{
+			case EControlRigTestDataPlaybackMode::ReplayInputs:
+			{
+				TestDataStrongPtr->ReleaseReplay();
+				break;
+			}
+			case EControlRigTestDataPlaybackMode::GroundTruth:
+			{
+				TestDataStrongPtr->SetupReplay(ControlRig, false);
+				break;
+			}
+			default:
+			{
+				TestDataStrongPtr->SetupReplay(ControlRig, true);
+				break;
+			}
+		}
+	}
+}
+
 void FControlRigEditor::ExtendToolbar()
 {
 	DECLARE_SCOPE_HIERARCHICAL_COUNTER_FUNC()
@@ -960,6 +1169,131 @@ void FControlRigEditor::FillToolbar(FToolBarBuilder& ToolbarBuilder)
 
 		ToolbarBuilder.EndStyleOverride();
 
+		if(CVarControlRigShowTestingToolbar.GetValueOnAnyThread())
+		{
+			ToolbarBuilder.AddSeparator();
+
+			FUIAction TestAssetAction;
+			ToolbarBuilder.AddComboButton(
+				FUIAction(FExecuteAction(), FCanExecuteAction::CreateLambda([this]()
+				{
+					if(TestDataStrongPtr.IsValid())
+					{
+						return !TestDataStrongPtr->IsReplaying() && !TestDataStrongPtr->IsRecording();
+					}
+					return true;
+				})),
+				FOnGetContent::CreateSP(this, &FControlRigEditor::GenerateTestAssetModeMenuContent),
+				TAttribute<FText>::CreateSP(this, &FControlRigEditor::GetTestAssetName),
+				TAttribute<FText>::CreateSP(this, &FControlRigEditor::GetTestAssetTooltip),
+				FSlateIcon(FAppStyle::GetAppStyleSetName(), "AutomationTools.TestAutomation"),
+				false);
+
+			ToolbarBuilder.AddToolBarButton(FUIAction(
+					FExecuteAction::CreateLambda([this]()
+					{
+						RecordTestData(0.0);
+					}),
+					FCanExecuteAction::CreateLambda([this]()
+					{
+						if(TestDataStrongPtr.IsValid())
+						{
+							return !TestDataStrongPtr->IsReplaying() && !TestDataStrongPtr->IsRecording();
+						}
+						return true;
+					})
+				),
+				NAME_None,
+				LOCTEXT("TestDataRecordButton", "Record"),
+				LOCTEXT("TestDataRecordButton_Tooltip", "Records a new frame into the test data asset.\nA test data asset will be created if necessary."),
+				FSlateIcon(FControlRigEditorStyle::Get().GetStyleSetName(),"ControlRig.TestData.Record")
+				);
+			ToolbarBuilder.AddComboButton(
+				FUIAction(),
+				FOnGetContent::CreateSP(this, &FControlRigEditor::GenerateTestAssetRecordMenuContent),
+				LOCTEXT("TestDataRecordMenu_Label", "Recording Modes"),
+				LOCTEXT("TestDataRecordMenu_ToolTip", "Pick between different modes for recording"),
+				FSlateIcon(FAppStyle::GetAppStyleSetName(), "LevelEditor.Recompile"),
+				true);
+
+			ToolbarBuilder.AddToolBarButton(FUIAction(
+					FExecuteAction::CreateLambda([this]()
+					{
+						ToggleTestData();
+					}),
+					FCanExecuteAction::CreateLambda([this]
+					{
+						if(TestDataStrongPtr.IsValid())
+						{
+							return !TestDataStrongPtr->IsRecording();
+						}
+						return false;
+					})
+				),
+				NAME_None,
+				TAttribute<FText>::CreateLambda([this]()
+				{
+					static const FText LiveStatus = LOCTEXT("LiveStatus", "Live"); 
+					static const FText ReplayInputsStatus = LOCTEXT("ReplayInputsStatus", "Replay Inputs");
+					static const FText GroundTruthStatus = LOCTEXT("GroundTruthStatus", "Ground Truth");
+					if(TestDataStrongPtr.IsValid())
+					{
+						switch(TestDataStrongPtr->GetPlaybackMode())
+						{
+							case EControlRigTestDataPlaybackMode::ReplayInputs:
+							{
+								return ReplayInputsStatus;
+							}
+							case EControlRigTestDataPlaybackMode::GroundTruth:
+							{
+								return GroundTruthStatus;
+							}
+							default:
+							{
+								break;
+							}
+						}
+					}
+					return LiveStatus;
+				}),
+				TAttribute<FText>::CreateLambda([this]()
+				{
+					static const FText LiveStatusTooltip = LOCTEXT("LiveStatusTooltip", "The test data is not affecting the rig."); 
+					static const FText ReplayInputsStatusTooltip = LOCTEXT("ReplayInputsStatusTooltip", "The test data inputs are being replayed onto the rig.");
+					static const FText GroundTruthStatusTooltip = LOCTEXT("GroundTruthStatusTooltip", "The complete result of the rig is being overwritten.");
+					if(TestDataStrongPtr.IsValid())
+					{
+						switch(TestDataStrongPtr->GetPlaybackMode())
+						{
+							case EControlRigTestDataPlaybackMode::ReplayInputs:
+							{
+								return ReplayInputsStatusTooltip;
+							}
+							case EControlRigTestDataPlaybackMode::GroundTruth:
+							{
+								return GroundTruthStatusTooltip;
+							}
+							default:
+							{
+								break;
+							}
+						}
+					}
+					return LiveStatusTooltip;
+				}),
+				TAttribute<FSlateIcon>::CreateLambda([this]()
+				{
+					static const FSlateIcon LiveIcon = FSlateIcon(FControlRigEditorStyle::Get().GetStyleSetName(),"ClassIcon.ControlRigBlueprint"); 
+					static const FSlateIcon ReplayIcon = FSlateIcon(FControlRigEditorStyle::Get().GetStyleSetName(),"ClassIcon.ControlRigSequence"); 
+					if(TestDataStrongPtr.IsValid() && TestDataStrongPtr->IsReplaying())
+					{
+						return ReplayIcon;
+					}
+					return LiveIcon;
+				}),
+				EUserInterfaceActionType::Button
+			);
+		}
 	}
 	ToolbarBuilder.EndSection();
 }
@@ -4032,6 +4366,11 @@ void FControlRigEditor::UpdateControlRig()
 		{
 			if (ControlRig)
 			{
+				if(TestDataStrongPtr.IsValid())
+				{
+					TestDataStrongPtr->ReleaseReplay();
+				}
+				
 				// if this control rig is from a temporary step,
 				// for example the reinstancing class, clear it 
 				// and create a new one!
diff --git a/Engine/Plugins/Animation/ControlRig/Source/ControlRigEditor/Private/Editor/ControlRigEditor.h b/Engine/Plugins/Animation/ControlRig/Source/ControlRigEditor/Private/Editor/ControlRigEditor.h
--- a/Engine/Plugins/Animation/ControlRig/Source/ControlRigEditor/Private/Editor/ControlRigEditor.h
+++ b/Engine/Plugins/Animation/ControlRig/Source/ControlRigEditor/Private/Editor/ControlRigEditor.h
@@ -21,6 +21,7 @@
 #include "Graph/ControlRigGraphNode.h"
 #include "RigVMModel/RigVMController.h"
 #include "Editor/DetailsViewWrapperObject.h"
+#include "ControlRigTestData.h"
 
 class UControlRigBlueprint;
 class IPersonaToolkit;
@@ -377,6 +378,14 @@ private:
 
 	void FrameSelection();
 
+	FText GetTestAssetName() const;
+	FText GetTestAssetTooltip() const;
+	bool SetTestAssetPath(const FString& InAssetPath);
+	TSharedRef<SWidget> GenerateTestAssetModeMenuContent();
+	TSharedRef<SWidget> GenerateTestAssetRecordMenuContent();
+	bool RecordTestData(double InRecordingDuration);
+	void ToggleTestData();
+
 protected:
 
 	/** Toolbox hosting widget */
@@ -492,6 +501,8 @@ protected:
 	bool bIsConstructionEventRunning;
 	uint32 LastHierarchyHash;
 
+	TStrongObjectPtr<UControlRigTestData> TestDataStrongPtr;
+	
 	static const TArray<FName> ForwardsSolveEventQueue;
 	static const TArray<FName> BackwardsSolveEventQueue;
 	static const TArray<FName> ConstructionEventQueue;
