Control Rig: Fix crash with bound variables

diff --git a/Engine/Plugins/Animation/ControlRig/Source/ControlRigDeveloper/Private/Graph/ControlRigGraph.cpp b/Engine/Plugins/Animation/ControlRig/Source/ControlRigDeveloper/Private/Graph/ControlRigGraph.cpp
--- a/Engine/Plugins/Animation/ControlRig/Source/ControlRigDeveloper/Private/Graph/ControlRigGraph.cpp
+++ b/Engine/Plugins/Animation/ControlRig/Source/ControlRigDeveloper/Private/Graph/ControlRigGraph.cpp
@@ -452,7 +452,7 @@ void UControlRigGraph::HandleModifiedEvent(ERigVMGraphNotifType InNotifType, URi
 								{
 									if (UControlRigGraphNode* RigNode = Cast<UControlRigGraphNode>(EdNode))
 									{
-										RigNode->ReconstructNode_Internal(true);
+										RigNode->ModelPinsChanged(true);
 									}
 								}
 							}
@@ -546,7 +546,7 @@ void UControlRigGraph::HandleModifiedEvent(ERigVMGraphNotifType InNotifType, URi
 							{
 								if (UControlRigGraphNode* RigNode = Cast<UControlRigGraphNode>(EdNode))
 								{
-									RigNode->ReconstructNode_Internal(true);
+									RigNode->ModelPinsChanged(true);
 								}
 							}
 						}
diff --git a/Engine/Plugins/Animation/ControlRig/Source/ControlRigEditor/Private/Graph/SControlRigGraphNode.cpp b/Engine/Plugins/Animation/ControlRig/Source/ControlRigEditor/Private/Graph/SControlRigGraphNode.cpp
--- a/Engine/Plugins/Animation/ControlRig/Source/ControlRigEditor/Private/Graph/SControlRigGraphNode.cpp
+++ b/Engine/Plugins/Animation/ControlRig/Source/ControlRigEditor/Private/Graph/SControlRigGraphNode.cpp
@@ -327,6 +327,7 @@ void SControlRigGraphNode::CreateStandardPinWidget(UEdGraphPin* CurPin)
 {
 	bool bShowPin = true;
 	FString CPPType;
+	FString BoundVariableName;
 	if(const UControlRigGraphNode* RigGraphNode = Cast<UControlRigGraphNode>(GraphNode))
 	{
 		if(const URigVMPin* ModelPin = RigGraphNode->FindModelPinFromGraphPin(CurPin))
@@ -338,6 +339,7 @@ void SControlRigGraphNode::CreateStandardPinWidget(UEdGraphPin* CurPin)
 				ModelPin->GetDirection() == ERigVMPinDirection::IO;
 			
 			CPPType = ModelPin->GetCPPType();
+			BoundVariableName = ModelPin->GetBoundVariableName();
 		}
 	}
 	
@@ -352,7 +354,8 @@ void SControlRigGraphNode::CreateStandardPinWidget(UEdGraphPin* CurPin)
 			const TSharedPtr<FPinInfoMetaData> PinInfoMetaData = RecycledPin->GetMetaData<FPinInfoMetaData>();
 			if(PinInfoMetaData.IsValid())
 			{
-				if(PinInfoMetaData->CPPType == CPPType)
+				if(PinInfoMetaData->CPPType == CPPType &&
+					PinInfoMetaData->BoundVariableName == BoundVariableName)
 				{
 					NewPin = RecycledPin;
 				}
@@ -366,7 +369,7 @@ void SControlRigGraphNode::CreateStandardPinWidget(UEdGraphPin* CurPin)
 				(*RecycledPinPtr)->InvalidateGraphData();
 			}
 			NewPin = CreatePinWidget(CurPin);
-			NewPin->AddMetadata(MakeShared<FPinInfoMetaData>(CPPType));
+			NewPin->AddMetadata(MakeShared<FPinInfoMetaData>(CPPType, BoundVariableName));
 			check(NewPin.IsValid());
 			PinsToKeep.Remove(CurPin);
 		}
diff --git a/Engine/Plugins/Animation/ControlRig/Source/ControlRigEditor/Private/Graph/SControlRigGraphNode.h b/Engine/Plugins/Animation/ControlRig/Source/ControlRigEditor/Private/Graph/SControlRigGraphNode.h
--- a/Engine/Plugins/Animation/ControlRig/Source/ControlRigEditor/Private/Graph/SControlRigGraphNode.h
+++ b/Engine/Plugins/Animation/ControlRig/Source/ControlRigEditor/Private/Graph/SControlRigGraphNode.h
@@ -134,11 +134,13 @@ private:
 	public:
 		SLATE_METADATA_TYPE(FPinInfoMetaData, ISlateMetaData)
 
-		FPinInfoMetaData(const FString& InCPPType)
+		FPinInfoMetaData(const FString& InCPPType, const FString& InBoundVariableName)
 		: CPPType(InCPPType)
+		, BoundVariableName(InBoundVariableName)
 		{}
 
 		FString CPPType;
+		FString BoundVariableName;
 	};
 	
 	TArray<FPinInfo> PinInfos;
