Control Rig: Allow set shape settings node in any event

diff --git a/Engine/Plugins/Animation/ControlRig/Source/ControlRig/Private/Units/Execution/RigUnit_DynamicHierarchy.cpp b/Engine/Plugins/Animation/ControlRig/Source/ControlRig/Private/Units/Execution/RigUnit_DynamicHierarchy.cpp
--- a/Engine/Plugins/Animation/ControlRig/Source/ControlRig/Private/Units/Execution/RigUnit_DynamicHierarchy.cpp
+++ b/Engine/Plugins/Animation/ControlRig/Source/ControlRig/Private/Units/Execution/RigUnit_DynamicHierarchy.cpp
@@ -929,16 +929,6 @@ FRigUnit_HierarchyGetShapeSettings_Execute()
 
 FRigUnit_HierarchySetShapeSettings_Execute()
 {
-	FString ErrorMessage;
-	if(!FRigUnit_DynamicHierarchyBase::IsValidToRunInContext(Context, ExecuteContext, true, &ErrorMessage))
-	{
-		if(!ErrorMessage.IsEmpty())
-		{
-			UE_CONTROLRIG_RIGUNIT_REPORT_ERROR(TEXT("%s"), *ErrorMessage);
-		}
-		return;
-	}
-
 	if(FRigControlElement* ControlElement = ExecuteContext.Hierarchy->Find<FRigControlElement>(Item))
 	{
 		ControlElement->Settings.bShapeVisible = Settings.bVisible;
