Sequencer:  Clear all sections when a control is deselected, otherwise they may not be deselected in the non Sections To Key.  Also add a ControlRig.Sequencer.SelectedSectionSetsSectionToKey Cvar for animators to try out which sets the section to key based upon the current selected key area. Defaults false but will make true if they like it.
jira #UE-168942

diff --git a/Engine/Plugins/Animation/ControlRig/Source/ControlRigEditor/Private/Sequencer/ControlRigParameterTrackEditor.cpp b/Engine/Plugins/Animation/ControlRig/Source/ControlRigEditor/Private/Sequencer/ControlRigParameterTrackEditor.cpp
--- a/Engine/Plugins/Animation/ControlRig/Source/ControlRigEditor/Private/Sequencer/ControlRigParameterTrackEditor.cpp
+++ b/Engine/Plugins/Animation/ControlRig/Source/ControlRigEditor/Private/Sequencer/ControlRigParameterTrackEditor.cpp
@@ -87,6 +87,8 @@
 
 TAutoConsoleVariable<bool> CVarSelectedKeysSelectControls(TEXT("ControlRig.Sequencer.SelectedKeysSelectControls"), false, TEXT("When true when we select a key in Sequencer it will select the Control, by default false."));
 
+TAutoConsoleVariable<bool> CVarSelectedSectionSetsSectionToKey(TEXT("ControlRig.Sequencer.SelectedSectionSetsSectionToKey"), false, TEXT("When true when we select a channel in a section, if it's the only section selected we set it as the Section To Key, by default false."));
+
 static USkeletalMeshComponent* AcquireSkeletalMeshFromObject(UObject* BoundObject, TSharedPtr<ISequencer> SequencerPtr)
 {
 	if (AActor* Actor = Cast<AActor>(BoundObject))
@@ -2103,6 +2105,10 @@ void FControlRigParameterTrackEditor::OnSelectionChanged(TArray<UMovieSceneTrack
 void FControlRigParameterTrackEditor::SelectRigsAndControls(UControlRig* ControlRig, const TArray<const IKeyArea*>& KeyAreas)
 {
 	FControlRigEditMode* ControlRigEditMode = GetEditMode();
+	
+	//if selection set's section to key we need to keep track of selected sections for each track.
+	const bool bSelectedSectionSetsSectionToKey = CVarSelectedSectionSetsSectionToKey.GetValueOnGameThread();
+	TMap<UMovieSceneTrack*, TSet<UMovieSceneControlRigParameterSection*>> TracksAndSections;
 
 	TArray<FString> StringArray;
 	//we have two sets here one to see if selection has really changed that contains the attirbutes, the other to select just the parent
@@ -2178,8 +2184,16 @@ void FControlRigParameterTrackEditor::SelectRigsAndControls(UControlRig* Control
 					}
 				}
 			}
+			if (bSelectedSectionSetsSectionToKey)
+			{
+				if (UMovieSceneTrack* Track = MovieSection->GetTypedOuter<UMovieSceneTrack>())
+				{
+					TracksAndSections.FindOrAdd(Track).Add(MovieSection);
+				}
+			}
 		}
 	}
+
 	//only create transaction if selection is really different.
 	bool bEndTransaction = false;
 	
@@ -2272,7 +2286,14 @@ void FControlRigParameterTrackEditor::SelectRigsAndControls(UControlRig* Control
 		}
 		SelectedPairs.Key->ClearControlSelection();
 	}
-
+	//if we have only one  selected section per track and the track has more than one section we set that to the section to key
+	for (TPair<UMovieSceneTrack*, TSet<UMovieSceneControlRigParameterSection*>>& TrackPair : TracksAndSections)
+	{
+		if (TrackPair.Key->GetAllSections().Num() > 0 && TrackPair.Value.Num() == 1)
+		{
+			TrackPair.Key->SetSectionToKey(TrackPair.Value.Array()[0]);
+		}
+	}
 	if (bEndTransaction)
 	{
 		GEditor->EndTransaction();
@@ -2933,10 +2954,23 @@ void FControlRigParameterTrackEditor::HandleControlSelected(UControlRig* Subject
 		if (Track)
 		{
 			GetSequencer()->SuspendSelectionBroadcast();
-			//Just set in the section to key not all
-			UMovieSceneSection* Section = Track->GetSectionToKey();
-			UMovieSceneControlRigParameterSection* ParamSection = Cast<UMovieSceneControlRigParameterSection>(Section);
-			SelectSequencerNodeInSection(ParamSection, ControlElement->GetName(), bSelected);
+			//Just select in section to key, if deselecting makes sure deselected everywhere
+			if (bSelected == true)
+			{
+				UMovieSceneSection* Section = Track->GetSectionToKey();
+				UMovieSceneControlRigParameterSection* ParamSection = Cast<UMovieSceneControlRigParameterSection>(Section);
+				SelectSequencerNodeInSection(ParamSection, ControlElement->GetName(), bSelected);
+			}
+			else
+			{
+				for (UMovieSceneSection* BaseSection : Track->GetAllSections())
+				{
+					if (UMovieSceneControlRigParameterSection* ParamSection = Cast< UMovieSceneControlRigParameterSection>(BaseSection))
+					{
+						SelectSequencerNodeInSection(ParamSection, ControlElement->GetName(), bSelected);
+					}
+				}
+			}
 			
 			GetSequencer()->ResumeSelectionBroadcast();
 
