Control Rig: Offer filtering for available animation channels

EnumComboBox changes are being reviewed separately

diff --git a/Engine/Plugins/Animation/ControlRig/Source/ControlRig/Private/Rigs/RigHierarchyElements.cpp b/Engine/Plugins/Animation/ControlRig/Source/ControlRig/Private/Rigs/RigHierarchyElements.cpp
--- a/Engine/Plugins/Animation/ControlRig/Source/ControlRig/Private/Rigs/RigHierarchyElements.cpp
+++ b/Engine/Plugins/Animation/ControlRig/Source/ControlRig/Private/Rigs/RigHierarchyElements.cpp
@@ -785,6 +785,7 @@ void FRigControlSettings::Save(FArchive& Ar)
 	Ar << DrivenControls;
 	Ar << bGroupWithParentControl;
 	Ar << bRestrictSpaceSwitching;
+	Ar << FilteredChannels;
 }
 
 void FRigControlSettings::Load(FArchive& Ar)
@@ -944,6 +945,15 @@ void FRigControlSettings::Load(FArchive& Ar)
 	{
 		bRestrictSpaceSwitching = false;
 	}
+
+	if (Ar.CustomVer(FControlRigObjectVersion::GUID) >= FControlRigObjectVersion::ControlTransformChannelFiltering)
+	{
+		Ar << FilteredChannels;
+	}
+	else
+	{
+		FilteredChannels.Reset();
+	}
 }
 
 uint32 GetTypeHash(const FRigControlSettings& Settings)
@@ -962,6 +972,11 @@ uint32 GetTypeHash(const FRigControlSettings& Settings)
 	Hash = HashCombine(Hash, GetTypeHash(Settings.DrivenControls));
 	Hash = HashCombine(Hash, GetTypeHash(Settings.bGroupWithParentControl));
 	Hash = HashCombine(Hash, GetTypeHash(Settings.bRestrictSpaceSwitching));
+	Hash = HashCombine(Hash, GetTypeHash(Settings.FilteredChannels.Num()));
+	for(const ERigControlTransformChannel& Channel : Settings.FilteredChannels)
+	{
+		Hash = HashCombine(Hash, GetTypeHash(Channel));
+	}
 	return Hash;
 }
 
@@ -1039,6 +1054,10 @@ bool FRigControlSettings::operator==(const FRigControlSettings& InOther) const
 	{
 		return false;
 	}
+	if(FilteredChannels != InOther.FilteredChannels)
+	{
+		return false;
+	}
 
 	const FTransform MinimumTransform = MinimumValue.GetAsTransform(ControlType, PrimaryAxis);
 	const FTransform OtherMinimumTransform = InOther.MinimumValue.GetAsTransform(ControlType, PrimaryAxis);
diff --git a/Engine/Plugins/Animation/ControlRig/Source/ControlRig/Private/Sequencer/MovieSceneControlRigParameterSection.cpp b/Engine/Plugins/Animation/ControlRig/Source/ControlRig/Private/Sequencer/MovieSceneControlRigParameterSection.cpp
--- a/Engine/Plugins/Animation/ControlRig/Source/ControlRig/Private/Sequencer/MovieSceneControlRigParameterSection.cpp
+++ b/Engine/Plugins/Animation/ControlRig/Source/ControlRig/Private/Sequencer/MovieSceneControlRigParameterSection.cpp
@@ -1617,6 +1617,23 @@ void UMovieSceneControlRigParameterSection::ReconstructChannelProxy()
 			}
 		};
 		// end constraints
+
+#if WITH_EDITOR
+		// masking for per control channels based on control filters
+		auto MaybeApplyChannelMask = [](FMovieSceneChannelMetaData& OutMetadata, const FRigControlElement* ControlElement, ERigControlTransformChannel InChannel)
+		{
+			if(!OutMetadata.bEnabled)
+			{
+				return;
+			}
+			
+			const TArray<ERigControlTransformChannel>& FilteredChannels = ControlElement->Settings.FilteredChannels;
+			if(!FilteredChannels.IsEmpty())
+			{
+				OutMetadata.bEnabled = FilteredChannels.Contains(InChannel);
+			}
+		};
+#endif
 		
 		for (FRigControlElement* ControlElement : SortedControls)
 		{
@@ -1805,6 +1822,8 @@ void UMovieSceneControlRigParameterSection::ReconstructChannelProxy()
 								ControlChannelMap.Add(Vector2D.ParameterName, FChannelMapInfo(ControlIndex, TotalIndex, FloatChannelIndex, ParentControlIndex, NAME_None, MaskIndex, CategoryIndex));
 							}
 							FParameterVectorChannelEditorData EditorData(ControlRig, Vector2D.ParameterName, bEnabled, Group, TotalIndex, 2);
+							MaybeApplyChannelMask(EditorData.MetaData[0], ControlElement, ERigControlTransformChannel::TranslationX);
+							MaybeApplyChannelMask(EditorData.MetaData[1], ControlElement, ERigControlTransformChannel::TranslationY);
 							Channels.Add(Vector2D.XCurve, EditorData.MetaData[0], EditorData.ExternalValues[0]);
 							Channels.Add(Vector2D.YCurve, EditorData.MetaData[1], EditorData.ExternalValues[1]);
 							FloatChannelIndex += 2;
@@ -1862,6 +1881,26 @@ void UMovieSceneControlRigParameterSection::ReconstructChannelProxy()
 						
 
 							FParameterVectorChannelEditorData EditorData(ControlRig, Vector.ParameterName, bEnabled, Group, TotalIndex, 3);
+
+							if(ControlElement->Settings.ControlType == ERigControlType::Position)
+							{
+								MaybeApplyChannelMask(EditorData.MetaData[0], ControlElement, ERigControlTransformChannel::TranslationX);
+								MaybeApplyChannelMask(EditorData.MetaData[1], ControlElement, ERigControlTransformChannel::TranslationY);
+								MaybeApplyChannelMask(EditorData.MetaData[2], ControlElement, ERigControlTransformChannel::TranslationZ);
+							}
+							else if(ControlElement->Settings.ControlType == ERigControlType::Rotator)
+							{
+								MaybeApplyChannelMask(EditorData.MetaData[0], ControlElement, ERigControlTransformChannel::Pitch);
+								MaybeApplyChannelMask(EditorData.MetaData[1], ControlElement, ERigControlTransformChannel::Yaw);
+								MaybeApplyChannelMask(EditorData.MetaData[2], ControlElement, ERigControlTransformChannel::Roll);
+							}
+							else if(ControlElement->Settings.ControlType == ERigControlType::Scale)
+							{
+								MaybeApplyChannelMask(EditorData.MetaData[0], ControlElement, ERigControlTransformChannel::ScaleX);
+								MaybeApplyChannelMask(EditorData.MetaData[1], ControlElement, ERigControlTransformChannel::ScaleY);
+								MaybeApplyChannelMask(EditorData.MetaData[2], ControlElement, ERigControlTransformChannel::ScaleZ);
+							}
+							
 							Channels.Add(Vector.XCurve, EditorData.MetaData[0], EditorData.ExternalValues[0]);
 							Channels.Add(Vector.YCurve, EditorData.MetaData[1], EditorData.ExternalValues[1]);
 							Channels.Add(Vector.ZCurve, EditorData.MetaData[2], EditorData.ExternalValues[2]);
@@ -1929,10 +1968,19 @@ void UMovieSceneControlRigParameterSection::ReconstructChannelProxy()
 							FParameterTransformChannelEditorData EditorData(ControlRig, Transform.ParameterName, bEnabled, TransformMask.GetChannels(), Group, 
 								TotalIndex);
 
+							MaybeApplyChannelMask(EditorData.MetaData[0], ControlElement, ERigControlTransformChannel::TranslationX);
+							MaybeApplyChannelMask(EditorData.MetaData[1], ControlElement, ERigControlTransformChannel::TranslationY);
+							MaybeApplyChannelMask(EditorData.MetaData[2], ControlElement, ERigControlTransformChannel::TranslationZ);
+							
 							Channels.Add(Transform.Translation[0], EditorData.MetaData[0], EditorData.ExternalValues[0]);
 							Channels.Add(Transform.Translation[1], EditorData.MetaData[1], EditorData.ExternalValues[1]);
 							Channels.Add(Transform.Translation[2], EditorData.MetaData[2], EditorData.ExternalValues[2]);
 
+							// note the order here is different from the rotator
+							MaybeApplyChannelMask(EditorData.MetaData[3], ControlElement, ERigControlTransformChannel::Roll);
+							MaybeApplyChannelMask(EditorData.MetaData[4], ControlElement, ERigControlTransformChannel::Pitch);
+							MaybeApplyChannelMask(EditorData.MetaData[5], ControlElement, ERigControlTransformChannel::Yaw);
+
 							Channels.Add(Transform.Rotation[0], EditorData.MetaData[3], EditorData.ExternalValues[3]);
 							Channels.Add(Transform.Rotation[1], EditorData.MetaData[4], EditorData.ExternalValues[4]);
 							Channels.Add(Transform.Rotation[2], EditorData.MetaData[5], EditorData.ExternalValues[5]);
@@ -1940,6 +1988,9 @@ void UMovieSceneControlRigParameterSection::ReconstructChannelProxy()
 							if (ControlElement->Settings.ControlType == ERigControlType::Transform ||
 								ControlElement->Settings.ControlType == ERigControlType::EulerTransform)
 							{
+								MaybeApplyChannelMask(EditorData.MetaData[6], ControlElement, ERigControlTransformChannel::ScaleX);
+								MaybeApplyChannelMask(EditorData.MetaData[7], ControlElement, ERigControlTransformChannel::ScaleY);
+								MaybeApplyChannelMask(EditorData.MetaData[8], ControlElement, ERigControlTransformChannel::ScaleZ);
 								
 								Channels.Add(Transform.Scale[0], EditorData.MetaData[6], EditorData.ExternalValues[6]);
 								Channels.Add(Transform.Scale[1], EditorData.MetaData[7], EditorData.ExternalValues[7]);
diff --git a/Engine/Plugins/Animation/ControlRig/Source/ControlRig/Private/Units/Execution/RigUnit_DynamicHierarchy.cpp b/Engine/Plugins/Animation/ControlRig/Source/ControlRig/Private/Units/Execution/RigUnit_DynamicHierarchy.cpp
--- a/Engine/Plugins/Animation/ControlRig/Source/ControlRig/Private/Units/Execution/RigUnit_DynamicHierarchy.cpp
+++ b/Engine/Plugins/Animation/ControlRig/Source/ControlRig/Private/Units/Execution/RigUnit_DynamicHierarchy.cpp
@@ -504,6 +504,7 @@ void FRigUnit_HierarchyAddControlVector2D_Settings::Configure(FRigControlSetting
 	
 	OutSettings.ControlType = ERigControlType::Vector2D;
 	OutSettings.PrimaryAxis = PrimaryAxis;
+	OutSettings.FilteredChannels = FilteredChannels;
 
 	Proxy.Configure(OutSettings);
 	Limits.Configure(OutSettings);
@@ -553,6 +554,7 @@ void FRigUnit_HierarchyAddControlVector_Settings::Configure(FRigControlSettings&
 	Super::Configure(OutSettings);
 	
 	OutSettings.ControlType = bIsPosition ? ERigControlType::Position : ERigControlType::Scale;
+	OutSettings.FilteredChannels = FilteredChannels;
 
 	Proxy.Configure(OutSettings);
 	Limits.Configure(OutSettings);
@@ -602,6 +604,7 @@ void FRigUnit_HierarchyAddControlRotator_Settings::Configure(FRigControlSettings
 	Super::Configure(OutSettings);
 	
 	OutSettings.ControlType = ERigControlType::Rotator;
+	OutSettings.FilteredChannels = FilteredChannels;
 
 	Proxy.Configure(OutSettings);
 	Limits.Configure(OutSettings);
@@ -640,6 +643,7 @@ void FRigUnit_HierarchyAddControlTransform_Settings::Configure(FRigControlSettin
 	Super::Configure(OutSettings);
 	
 	OutSettings.ControlType = ERigControlType::EulerTransform;
+	OutSettings.FilteredChannels = FilteredChannels;
 
 	Proxy.Configure(OutSettings);
 	Shape.Configure(OutSettings);
diff --git a/Engine/Plugins/Animation/ControlRig/Source/ControlRig/Private/Units/Execution/RigUnit_DynamicHierarchy.h b/Engine/Plugins/Animation/ControlRig/Source/ControlRig/Private/Units/Execution/RigUnit_DynamicHierarchy.h
--- a/Engine/Plugins/Animation/ControlRig/Source/ControlRig/Private/Units/Execution/RigUnit_DynamicHierarchy.h
+++ b/Engine/Plugins/Animation/ControlRig/Source/ControlRig/Private/Units/Execution/RigUnit_DynamicHierarchy.h
@@ -325,6 +325,8 @@ struct CONTROLRIG_API FRigUnit_HierarchyAddElement : public FRigUnit_DynamicHier
 		Name = NAME_None;
 	}
 
+	virtual ERigElementType GetElementTypeToSpawn() const { return ERigElementType::None; }
+
 	/*
 	 * The parent of the new element to add
 	 */
@@ -360,6 +362,8 @@ struct CONTROLRIG_API FRigUnit_HierarchyAddBone : public FRigUnit_HierarchyAddEl
 		Space = EBoneGetterSetterMode::LocalSpace;
 	}
 
+	virtual ERigElementType GetElementTypeToSpawn() const override { return ERigElementType::Bone; }
+
 	/*
 	 * The initial transform of the new element
 	 */
@@ -392,6 +396,8 @@ struct CONTROLRIG_API FRigUnit_HierarchyAddNull : public FRigUnit_HierarchyAddEl
 		Space = EBoneGetterSetterMode::LocalSpace;
 	}
 
+	virtual ERigElementType GetElementTypeToSpawn() const override { return ERigElementType::Null; }
+
 	/*
 	 * The initial transform of the new element
 	 */
@@ -418,7 +424,7 @@ struct CONTROLRIG_API FRigUnit_HierarchyAddControl_Settings
 	{}
 	virtual ~FRigUnit_HierarchyAddControl_Settings(){}
 
-	void Configure(FRigControlSettings& OutSettings) const;
+	virtual void Configure(FRigControlSettings& OutSettings) const;
 
 	UPROPERTY(EditAnywhere, BlueprintReadWrite, Category = "Settings")
 	FName DisplayName;
@@ -511,7 +517,7 @@ struct CONTROLRIG_API FRigUnit_HierarchyAddControlFloat_Settings : public FRigUn
 	{}
 	virtual ~FRigUnit_HierarchyAddControlFloat_Settings() override {}
 
-	void Configure(FRigControlSettings& OutSettings) const;
+	virtual void Configure(FRigControlSettings& OutSettings) const override;
 
 	UPROPERTY(EditAnywhere, BlueprintReadWrite, Category = "Settings")
 	ERigControlAxis PrimaryAxis;
@@ -528,25 +534,42 @@ struct CONTROLRIG_API FRigUnit_HierarchyAddControlFloat_Settings : public FRigUn
 
 /**
  * Adds a new control to the hierarchy
- * Note: This node only runs as part of the construction event.
  */
-USTRUCT(meta=(DisplayName="Spawn Float Control", TemplateName="SpawnControl", Keywords="Construction,Create,New,AddControl,NewControl,CreateControl", Varying))
-struct CONTROLRIG_API FRigUnit_HierarchyAddControlFloat : public FRigUnit_HierarchyAddElement
+USTRUCT(meta=(TemplateName="SpawnControl", Keywords="Construction,Create,New,AddControl,NewControl,CreateControl", Varying))
+struct CONTROLRIG_API FRigUnit_HierarchyAddControlElement : public FRigUnit_HierarchyAddElement
 {
 	GENERATED_BODY()
 
-	FRigUnit_HierarchyAddControlFloat()
+	FRigUnit_HierarchyAddControlElement()
 	{
 		Name = TEXT("NewControl");
-		OffsetTransform = FTransform::Identity;
-		InitialValue = 0.f;
 	}
 
+	virtual ERigElementType GetElementTypeToSpawn() const override { return ERigElementType::Control; }
+	virtual ERigControlType GetControlTypeToSpawn() const { return ERigControlType::Bool; }
+
 	/*
 	 * The offset transform of the new control
 	 */
 	UPROPERTY(meta = (Input))
 	FTransform OffsetTransform;
+};
+
+/**
+ * Adds a new control to the hierarchy
+ * Note: This node only runs as part of the construction event.
+ */
+USTRUCT(meta=(DisplayName="Spawn Float Control", TemplateName="SpawnControl", Keywords="Construction,Create,New,AddControl,NewControl,CreateControl", Varying))
+struct CONTROLRIG_API FRigUnit_HierarchyAddControlFloat : public FRigUnit_HierarchyAddControlElement
+{
+	GENERATED_BODY()
+
+	FRigUnit_HierarchyAddControlFloat()
+	{
+		InitialValue = 0.f;
+	}
+
+	virtual ERigControlType GetControlTypeToSpawn() const override { return ERigControlType::Float; }
 
 	/*
 	 * The initial value of the new control
@@ -602,7 +625,7 @@ struct CONTROLRIG_API FRigUnit_HierarchyAddControlInteger_Settings : public FRig
 	{}
 	virtual ~FRigUnit_HierarchyAddControlInteger_Settings() override {}
 
-	void Configure(FRigControlSettings& OutSettings) const;
+	virtual void Configure(FRigControlSettings& OutSettings) const override;
 
 	UPROPERTY(EditAnywhere, BlueprintReadWrite, Category = "Settings")
 	ERigControlAxis PrimaryAxis;
@@ -622,22 +645,16 @@ struct CONTROLRIG_API FRigUnit_HierarchyAddControlInteger_Settings : public FRig
  * Note: This node only runs as part of the construction event.
  */
 USTRUCT(meta=(DisplayName="Spawn Integer Control", TemplateName="SpawnControl", Keywords="Construction,Create,New,AddControl,NewControl,CreateControl", Varying))
-struct CONTROLRIG_API FRigUnit_HierarchyAddControlInteger : public FRigUnit_HierarchyAddElement
+struct CONTROLRIG_API FRigUnit_HierarchyAddControlInteger : public FRigUnit_HierarchyAddControlElement
 {
 	GENERATED_BODY()
 
 	FRigUnit_HierarchyAddControlInteger()
 	{
-		Name = TEXT("NewControl");
-		OffsetTransform = FTransform::Identity;
 		InitialValue = 0;
 	}
 
-	/*
-	 * The offset transform of the new control
-	 */
-	UPROPERTY(meta = (Input))
-	FTransform OffsetTransform;
+	virtual ERigControlType GetControlTypeToSpawn() const override { return ERigControlType::Integer; }
 
 	/*
 	 * The initial value of the new control
@@ -694,10 +711,12 @@ struct CONTROLRIG_API FRigUnit_HierarchyAddControlVector2D_Settings : public FRi
 	FRigUnit_HierarchyAddControlVector2D_Settings()
 		: FRigUnit_HierarchyAddControl_Settings()
 		, PrimaryAxis(ERigControlAxis::X)
-	{}
+	{
+		FilteredChannels.Reset();
+	}
 	virtual ~FRigUnit_HierarchyAddControlVector2D_Settings() override {}
 
-	void Configure(FRigControlSettings& OutSettings) const;
+	virtual void Configure(FRigControlSettings& OutSettings) const override;
 
 	UPROPERTY(EditAnywhere, BlueprintReadWrite, Category = "Settings")
 	ERigControlAxis PrimaryAxis;
@@ -710,6 +729,9 @@ struct CONTROLRIG_API FRigUnit_HierarchyAddControlVector2D_Settings : public FRi
 
 	UPROPERTY(EditAnywhere, BlueprintReadWrite, Category = "Settings")
 	FRigUnit_HierarchyAddControl_ProxySettings Proxy;
+
+	UPROPERTY(EditAnywhere, BlueprintReadWrite, Category = "Settings")
+	TArray<ERigControlTransformChannel> FilteredChannels;
 };
 
 /**
@@ -717,22 +739,16 @@ struct CONTROLRIG_API FRigUnit_HierarchyAddControlVector2D_Settings : public FRi
  * Note: This node only runs as part of the construction event.
  */
 USTRUCT(meta=(DisplayName="Spawn Vector2D Control", TemplateName="SpawnControl", Keywords="Construction,Create,New,AddControl,NewControl,CreateControl", Varying))
-struct CONTROLRIG_API FRigUnit_HierarchyAddControlVector2D : public FRigUnit_HierarchyAddElement
+struct CONTROLRIG_API FRigUnit_HierarchyAddControlVector2D : public FRigUnit_HierarchyAddControlElement
 {
 	GENERATED_BODY()
 
 	FRigUnit_HierarchyAddControlVector2D()
 	{
-		Name = TEXT("NewControl");
-		OffsetTransform = FTransform::Identity;
 		InitialValue = FVector2D::ZeroVector;
 	}
 
-	/*
-	 * The offset transform of the new control
-	 */
-	UPROPERTY(meta = (Input))
-	FTransform OffsetTransform;
+	virtual ERigControlType GetControlTypeToSpawn() const override { return ERigControlType::Vector2D; }
 
 	/*
 	 * The initial value of the new control
@@ -793,10 +809,12 @@ struct CONTROLRIG_API FRigUnit_HierarchyAddControlVector_Settings : public FRigU
 	FRigUnit_HierarchyAddControlVector_Settings()
 		: FRigUnit_HierarchyAddControl_Settings()
 		, bIsPosition(true)
-	{}
+	{
+		FilteredChannels.Reset();
+	}
 	virtual ~FRigUnit_HierarchyAddControlVector_Settings() override {}
 
-	void Configure(FRigControlSettings& OutSettings) const;
+	virtual void Configure(FRigControlSettings& OutSettings) const override;
 
 	UPROPERTY(EditAnywhere, BlueprintReadWrite, Category = "Settings")
 	bool bIsPosition;
@@ -809,6 +827,9 @@ struct CONTROLRIG_API FRigUnit_HierarchyAddControlVector_Settings : public FRigU
 
 	UPROPERTY(EditAnywhere, BlueprintReadWrite, Category = "Settings")
 	FRigUnit_HierarchyAddControl_ProxySettings Proxy;
+
+	UPROPERTY(EditAnywhere, BlueprintReadWrite, Category = "Settings")
+	TArray<ERigControlTransformChannel> FilteredChannels;
 };
 
 /**
@@ -816,22 +837,19 @@ struct CONTROLRIG_API FRigUnit_HierarchyAddControlVector_Settings : public FRigU
  * Note: This node only runs as part of the construction event.
  */
 USTRUCT(meta=(DisplayName="Spawn Vector Control", TemplateName="SpawnControl", Keywords="Construction,Create,New,AddControl,NewControl,CreateControl", Varying))
-struct CONTROLRIG_API FRigUnit_HierarchyAddControlVector : public FRigUnit_HierarchyAddElement
+struct CONTROLRIG_API FRigUnit_HierarchyAddControlVector : public FRigUnit_HierarchyAddControlElement
 {
 	GENERATED_BODY()
 
 	FRigUnit_HierarchyAddControlVector()
 	{
-		Name = TEXT("NewControl");
-		OffsetTransform = FTransform::Identity;
 		InitialValue = FVector::ZeroVector;
 	}
 
-	/*
-	 * The offset transform of the new control
-	 */
-	UPROPERTY(meta = (Input))
-	FTransform OffsetTransform;
+	virtual ERigControlType GetControlTypeToSpawn() const override
+	{
+		return Settings.bIsPosition ? ERigControlType::Position : ERigControlType::Scale;
+	}
 
 	/*
 	 * The initial value of the new control
@@ -891,10 +909,12 @@ struct CONTROLRIG_API FRigUnit_HierarchyAddControlRotator_Settings : public FRig
 	
 	FRigUnit_HierarchyAddControlRotator_Settings()
 		: FRigUnit_HierarchyAddControl_Settings()
-	{}
+	{
+		FilteredChannels.Reset();
+	}
 	virtual ~FRigUnit_HierarchyAddControlRotator_Settings() override {}
 
-	void Configure(FRigControlSettings& OutSettings) const;
+	virtual void Configure(FRigControlSettings& OutSettings) const override;
 
 	UPROPERTY(EditAnywhere, BlueprintReadWrite, Category = "Settings")
 	FRigUnit_HierarchyAddControlRotator_LimitSettings Limits;;
@@ -904,6 +924,9 @@ struct CONTROLRIG_API FRigUnit_HierarchyAddControlRotator_Settings : public FRig
 
 	UPROPERTY(EditAnywhere, BlueprintReadWrite, Category = "Settings")
 	FRigUnit_HierarchyAddControl_ProxySettings Proxy;
+
+	UPROPERTY(EditAnywhere, BlueprintReadWrite, Category = "Settings")
+	TArray<ERigControlTransformChannel> FilteredChannels;
 };
 
 /**
@@ -911,22 +934,16 @@ struct CONTROLRIG_API FRigUnit_HierarchyAddControlRotator_Settings : public FRig
  * Note: This node only runs as part of the construction event.
  */
 USTRUCT(meta=(DisplayName="Spawn Rotator Control", TemplateName="SpawnControl", Keywords="Construction,Create,New,AddControl,NewControl,CreateControl,Rotation", Varying))
-struct CONTROLRIG_API FRigUnit_HierarchyAddControlRotator : public FRigUnit_HierarchyAddElement
+struct CONTROLRIG_API FRigUnit_HierarchyAddControlRotator : public FRigUnit_HierarchyAddControlElement
 {
 	GENERATED_BODY()
 
 	FRigUnit_HierarchyAddControlRotator()
 	{
-		Name = TEXT("NewControl");
-		OffsetTransform = FTransform::Identity;
 		InitialValue = FRotator::ZeroRotator;
 	}
 
-	/*
-	 * The offset transform of the new control
-	 */
-	UPROPERTY(meta = (Input))
-	FTransform OffsetTransform;
+	virtual ERigControlType GetControlTypeToSpawn() const override { return ERigControlType::Rotator; }
 
 	/*
 	 * The initial value of the new control
@@ -951,16 +968,21 @@ struct CONTROLRIG_API FRigUnit_HierarchyAddControlTransform_Settings : public FR
 	
 	FRigUnit_HierarchyAddControlTransform_Settings()
 		: FRigUnit_HierarchyAddControl_Settings()
-	{}
+	{
+		FilteredChannels.Reset();
+	}
 	virtual ~FRigUnit_HierarchyAddControlTransform_Settings() override {}
 
-	void Configure(FRigControlSettings& OutSettings) const;
+	virtual void Configure(FRigControlSettings& OutSettings) const override;
 
 	UPROPERTY(EditAnywhere, BlueprintReadWrite, Category = "Settings")
 	FRigUnit_HierarchyAddControl_ShapeSettings Shape;
 
 	UPROPERTY(EditAnywhere, BlueprintReadWrite, Category = "Settings")
 	FRigUnit_HierarchyAddControl_ProxySettings Proxy;
+
+	UPROPERTY(EditAnywhere, BlueprintReadWrite, Category = "Settings")
+	TArray<ERigControlTransformChannel> FilteredChannels;
 };
 
 /**
@@ -968,22 +990,16 @@ struct CONTROLRIG_API FRigUnit_HierarchyAddControlTransform_Settings : public FR
  * Note: This node only runs as part of the construction event.
  */
 USTRUCT(meta=(DisplayName="Spawn Transform Control", TemplateName="SpawnControl", Keywords="Construction,Create,New,AddControl,NewControl,CreateControl", Varying))
-struct CONTROLRIG_API FRigUnit_HierarchyAddControlTransform : public FRigUnit_HierarchyAddElement
+struct CONTROLRIG_API FRigUnit_HierarchyAddControlTransform : public FRigUnit_HierarchyAddControlElement
 {
 	GENERATED_BODY()
 
 	FRigUnit_HierarchyAddControlTransform()
 	{
-		Name = TEXT("NewControl");
-		OffsetTransform = FTransform::Identity;
 		InitialValue = FTransform::Identity;
 	}
 
-	/*
-	 * The offset transform of the new control
-	 */
-	UPROPERTY(meta = (Input))
-	FTransform OffsetTransform;
+	virtual ERigControlType GetControlTypeToSpawn() const override { return ERigControlType::EulerTransform; }
 
 	/*
 	 * The initial value of the new control
diff --git a/Engine/Plugins/Animation/ControlRig/Source/ControlRig/Public/ControlRigObjectVersion.h b/Engine/Plugins/Animation/ControlRig/Source/ControlRig/Public/ControlRigObjectVersion.h
--- a/Engine/Plugins/Animation/ControlRig/Source/ControlRig/Public/ControlRigObjectVersion.h
+++ b/Engine/Plugins/Animation/ControlRig/Source/ControlRig/Public/ControlRigObjectVersion.h
@@ -78,6 +78,9 @@ struct CONTROLRIG_API FControlRigObjectVersion
 		// Controls to be able specify space switch targets
 		RestrictSpaceSwitchingForControls,
 
+		// Controls to be able specify which channels should be visible in sequencer
+		ControlTransformChannelFiltering,
+
 		// -----<new versions can be added above this line>-------------------------------------------------
 		VersionPlusOne,
 		LatestVersion = VersionPlusOne - 1
diff --git a/Engine/Plugins/Animation/ControlRig/Source/ControlRig/Public/Rigs/RigHierarchyElements.h b/Engine/Plugins/Animation/ControlRig/Source/ControlRig/Public/Rigs/RigHierarchyElements.h
--- a/Engine/Plugins/Animation/ControlRig/Source/ControlRig/Public/Rigs/RigHierarchyElements.h
+++ b/Engine/Plugins/Animation/ControlRig/Source/ControlRig/Public/Rigs/RigHierarchyElements.h
@@ -1025,6 +1025,20 @@ struct FRigControlElementCustomization
 	TArray<FRigElementKey> RemovedSpaces;
 };
 
+UENUM(BlueprintType)
+enum class ERigControlTransformChannel : uint8
+{
+	TranslationX,
+	TranslationY,
+	TranslationZ,
+	Pitch,
+	Yaw,
+	Roll,
+	ScaleX,
+	ScaleY,
+	ScaleZ
+};
+
 USTRUCT(BlueprintType)
 struct CONTROLRIG_API FRigControlSettings
 {
@@ -1126,6 +1140,12 @@ struct CONTROLRIG_API FRigControlSettings
 	UPROPERTY(BlueprintReadOnly, EditAnywhere, Category = Animation)
 	bool bRestrictSpaceSwitching;
 
+	/**
+	 * Filtered Visible Transform channels. If this is empty everything is visible
+	 */
+	UPROPERTY(BlueprintReadOnly, EditAnywhere, Category = Animation)
+	TArray<ERigControlTransformChannel> FilteredChannels;
+
 #if WITH_EDITORONLY_DATA
 	/**
 	 * Deprecated properties.
diff --git a/Engine/Plugins/Animation/ControlRig/Source/ControlRigEditor/Private/ControlRigEditorModule.cpp b/Engine/Plugins/Animation/ControlRig/Source/ControlRigEditor/Private/ControlRigEditorModule.cpp
--- a/Engine/Plugins/Animation/ControlRig/Source/ControlRigEditor/Private/ControlRigEditorModule.cpp
+++ b/Engine/Plugins/Animation/ControlRig/Source/ControlRigEditor/Private/ControlRigEditorModule.cpp
@@ -191,6 +191,9 @@ void FControlRigEditorModule::StartupModule()
 	PropertiesToUnregisterOnShutdown.Add(FControlRigAnimNodeEventName::StaticStruct()->GetFName());
 	PropertyEditorModule.RegisterCustomPropertyTypeLayout(PropertiesToUnregisterOnShutdown.Last(), FOnGetPropertyTypeCustomizationInstance::CreateStatic(&FControlRigAnimNodeEventNameDetails::MakeInstance));
 
+	PropertiesToUnregisterOnShutdown.Add(StaticEnum<ERigControlTransformChannel>()->GetFName());
+	PropertyEditorModule.RegisterCustomPropertyTypeLayout(PropertiesToUnregisterOnShutdown.Last(), FOnGetPropertyTypeCustomizationInstance::CreateStatic(&FRigControlTransformChannelDetails::MakeInstance));
+
 	FRigBaseElementDetails::RegisterSectionMappings(PropertyEditorModule);
 
 	// Register asset tools
diff --git a/Engine/Plugins/Animation/ControlRig/Source/ControlRigEditor/Private/ControlRigElementDetails.cpp b/Engine/Plugins/Animation/ControlRig/Source/ControlRigEditor/Private/ControlRigElementDetails.cpp
--- a/Engine/Plugins/Animation/ControlRig/Source/ControlRigEditor/Private/ControlRigElementDetails.cpp
+++ b/Engine/Plugins/Animation/ControlRig/Source/ControlRigEditor/Private/ControlRigElementDetails.cpp
@@ -14,6 +14,7 @@
 #include "PropertyCustomizationHelpers.h"
 #include "SEnumCombo.h"
 #include "Units/Execution/RigUnit_BeginExecution.h"
+#include "Units/Execution/RigUnit_DynamicHierarchy.h"
 #include "Graph/SControlRigGraphPinVariableBinding.h"
 #include "HAL/PlatformApplicationMisc.h"
 #include "Styling/AppStyle.h"
@@ -523,6 +524,146 @@ void FRigComputedTransformDetails::OnTransformChanged(FEditPropertyChain* InProp
 	}
 }
 
+void FRigControlTransformChannelDetails::CustomizeHeader(TSharedRef<IPropertyHandle> InStructPropertyHandle,
+	FDetailWidgetRow& HeaderRow, IPropertyTypeCustomizationUtils& StructCustomizationUtils)
+{
+	Handle = InStructPropertyHandle;
+
+	TArray<int32> VisibleEnumValues;
+	const TArray<ERigControlTransformChannel>* VisibleChannels = nullptr;
+
+	// loop for controls to figure out the control type
+	TArray<UObject*> Objects;
+	InStructPropertyHandle->GetOuterObjects(Objects);
+	for (UObject* Object : Objects)
+	{
+		if (const UDetailsViewWrapperObject* WrapperObject = Cast<UDetailsViewWrapperObject>(Object))
+		{
+			if(WrapperObject->GetWrappedStruct() == FRigControlElement::StaticStruct())
+			{
+				const FRigControlElement ControlElement = WrapperObject->GetContent<FRigControlElement>();
+				VisibleChannels = GetVisibleChannelsForControlType(ControlElement.Settings.ControlType);
+				break;
+			}
+			if (const URigVMUnitNode* UnitNode = Cast<URigVMUnitNode>(WrapperObject->GetOuter()))
+			{
+				if(UnitNode->GetScriptStruct() && UnitNode->GetScriptStruct()->IsChildOf(FRigUnit_HierarchyAddControlElement::StaticStruct()))
+				{
+					FStructOnScope StructOnScope(UnitNode->GetScriptStruct());
+					WrapperObject->GetContent(StructOnScope.GetStructMemory(), StructOnScope.GetStruct());
+
+					const FRigUnit_HierarchyAddControlElement* RigUnit = (const FRigUnit_HierarchyAddControlElement*)StructOnScope.GetStructMemory();
+					VisibleChannels = GetVisibleChannelsForControlType(RigUnit->GetControlTypeToSpawn());
+					break;
+				}
+			}
+		}
+	}
+
+	if(VisibleChannels)
+	{
+		VisibleEnumValues.Reserve(VisibleChannels->Num());
+		for(int32 Index=0; Index < VisibleChannels->Num(); Index++)
+		{
+			VisibleEnumValues.Add((int32)(*VisibleChannels)[Index]);
+		}
+	}
+	
+	HeaderRow
+	.NameContent()
+	[
+		InStructPropertyHandle->CreatePropertyNameWidget()
+	]
+	.ValueContent()
+	[
+		SNew(SEnumComboBox, StaticEnum<ERigControlTransformChannel>())
+		.CurrentValue_Raw(this, &FRigControlTransformChannelDetails::GetChannelAsInt32)
+		.OnEnumSelectionChanged_Raw(this, &FRigControlTransformChannelDetails::OnChannelChanged)
+		.Font(FAppStyle::GetFontStyle(TEXT("MenuItem.Font")))
+		.EnumValueSubset(VisibleEnumValues)
+	];
+}
+
+void FRigControlTransformChannelDetails::CustomizeChildren(TSharedRef<IPropertyHandle> InStructPropertyHandle,
+	IDetailChildrenBuilder& StructBuilder, IPropertyTypeCustomizationUtils& StructCustomizationUtils)
+{
+	// nothing to do here
+}
+
+ERigControlTransformChannel FRigControlTransformChannelDetails::GetChannel() const
+{
+	uint8 Value = 0;
+	Handle->GetValue(Value);
+	return (ERigControlTransformChannel)Value;
+}
+
+void FRigControlTransformChannelDetails::OnChannelChanged(int32 NewSelection, ESelectInfo::Type InSelectionInfo)
+{
+	Handle->SetValue((uint8)NewSelection);
+}
+
+const TArray<ERigControlTransformChannel>* FRigControlTransformChannelDetails::GetVisibleChannelsForControlType(ERigControlType InControlType)
+{
+	switch(InControlType)
+	{
+		case ERigControlType::Position:
+		{
+			static const TArray<ERigControlTransformChannel> PositionChannels = {
+				ERigControlTransformChannel::TranslationX,
+				ERigControlTransformChannel::TranslationY,
+				ERigControlTransformChannel::TranslationZ
+			};
+			return &PositionChannels;
+		}
+		case ERigControlType::Rotator:
+		{
+			static const TArray<ERigControlTransformChannel> RotatorChannels = {
+				ERigControlTransformChannel::Pitch,
+				ERigControlTransformChannel::Yaw,
+				ERigControlTransformChannel::Roll
+			};
+			return &RotatorChannels;
+		}
+		case ERigControlType::Scale:
+		{
+			static const TArray<ERigControlTransformChannel> ScaleChannels = {
+				ERigControlTransformChannel::ScaleX,
+				ERigControlTransformChannel::ScaleY,
+				ERigControlTransformChannel::ScaleZ
+			};
+			return &ScaleChannels;
+		}
+		case ERigControlType::Vector2D:
+		{
+			static const TArray<ERigControlTransformChannel> Vector2DChannels = {
+				ERigControlTransformChannel::TranslationX,
+				ERigControlTransformChannel::TranslationY
+			};
+			return &Vector2DChannels;
+		}
+		case ERigControlType::EulerTransform:
+		{
+			static const TArray<ERigControlTransformChannel> EulerTransformChannels = {
+				ERigControlTransformChannel::TranslationX,
+				ERigControlTransformChannel::TranslationY,
+				ERigControlTransformChannel::TranslationZ,
+				ERigControlTransformChannel::Pitch,
+				ERigControlTransformChannel::Yaw,
+				ERigControlTransformChannel::Roll,
+				ERigControlTransformChannel::ScaleX,
+				ERigControlTransformChannel::ScaleY,
+				ERigControlTransformChannel::ScaleZ
+			};
+			return &EulerTransformChannels;
+		}
+		default:
+		{
+			break;
+		}
+	}
+	return nullptr;
+}
+
 void FRigBaseElementDetails::CustomizeDetails(IDetailLayoutBuilder& DetailBuilder)
 {
 	PerElementInfos.Reset();
@@ -858,6 +999,96 @@ bool FRigBaseElementDetails::IsAnyElementProcedural() const
 	});
 }
 
+bool FRigBaseElementDetails::GetCommonElementType(ERigElementType& OutElementType) const
+{
+	OutElementType = ERigElementType::None;
+
+	for(const FPerElementInfo& Info : PerElementInfos)
+	{
+		const FRigElementKey& Key = Info.Element.GetKey();
+		if(Key.IsValid())
+		{
+			if(OutElementType == ERigElementType::None)
+			{
+				OutElementType = Key.Type;
+			}
+			else if(OutElementType != Key.Type)
+			{
+				OutElementType = ERigElementType::None;
+				break;
+			}
+		}
+	}
+
+	return OutElementType != ERigElementType::None;
+}
+
+bool FRigBaseElementDetails::GetCommonControlType(ERigControlType& OutControlType) const
+{
+	OutControlType = ERigControlType::Bool;
+	
+	ERigElementType ElementType = ERigElementType::None;
+	if(GetCommonElementType(ElementType))
+	{
+		if(ElementType == ERigElementType::Control)
+		{
+			bool bSuccess = false;
+			for(const FPerElementInfo& Info : PerElementInfos)
+			{
+				if(const FRigControlElement* ControlElement = Info.Element.Get<FRigControlElement>())
+				{
+					if(!bSuccess)
+					{
+						OutControlType = ControlElement->Settings.ControlType;
+						bSuccess = true;
+					}
+					else if(OutControlType != ControlElement->Settings.ControlType)
+					{
+						OutControlType = ERigControlType::Bool;
+						bSuccess = false;
+						break;
+					}
+				}
+			}
+			return bSuccess;
+		}
+	}
+	return false;
+}
+
+bool FRigBaseElementDetails::GetCommonAnimationType(ERigControlAnimationType& OutAnimationType) const
+{
+	OutAnimationType = ERigControlAnimationType::AnimationControl;
+	
+	ERigElementType ElementType = ERigElementType::None;
+	if(GetCommonElementType(ElementType))
+	{
+		if(ElementType == ERigElementType::Control)
+		{
+			bool bSuccess = false;
+			for(const FPerElementInfo& Info : PerElementInfos)
+			{
+				if(const FRigControlElement* ControlElement = Info.Element.Get<FRigControlElement>())
+				{
+					if(!bSuccess)
+					{
+						OutAnimationType = ControlElement->Settings.AnimationType;
+						bSuccess = true;
+					}
+					else if(OutAnimationType != ControlElement->Settings.AnimationType)
+					{
+						OutAnimationType = ERigControlAnimationType::AnimationControl;
+						bSuccess = false;
+						break;
+					}
+				}
+			}
+			return bSuccess;
+		}
+	}
+	return false;
+}
+
 const FRigBaseElementDetails::FPerElementInfo* FRigBaseElementDetails::FindElementByPredicate(const TFunction<bool(const FPerElementInfo&)>& InPredicate) const
 {
 	return PerElementInfos.FindByPredicate(InPredicate);
@@ -2798,6 +3029,17 @@ void FRigControlElementDetails::CustomizeControl(IDetailLayoutBuilder& DetailBui
 			return false;
 		}));
 	}
+
+	ERigControlType CommonControlType = ERigControlType::Bool;
+	if(GetCommonControlType(CommonControlType))
+	{
+		if(FRigControlTransformChannelDetails::GetVisibleChannelsForControlType(CommonControlType) != nullptr)
+		{
+			const TSharedPtr<IPropertyHandle> FilteredChannelsHandle = SettingsHandle->GetChildHandle(GET_MEMBER_NAME_CHECKED(FRigControlSettings, FilteredChannels));
+			ControlCategory.AddProperty(FilteredChannelsHandle.ToSharedRef())
+			.IsEnabled(bIsEnabled);
+		}
+	}
 	
 	if(IsAnyControlOfAnimationType(ERigControlAnimationType::ProxyControl) || IsAnyControlOfAnimationType(ERigControlAnimationType::AnimationControl))
 	{
@@ -3209,6 +3451,7 @@ void FRigControlElementDetails::HandleControlTypeChanged(ERigControlType Control
 		ControlElement->Settings.ControlType = ControlType;
 		ControlElement->Settings.LimitEnabled.Reset();
 		ControlElement->Settings.bGroupWithParentControl = false;
+		ControlElement->Settings.FilteredChannels.Reset();
 
 		switch (ControlElement->Settings.ControlType)
 		{
diff --git a/Engine/Plugins/Animation/ControlRig/Source/ControlRigEditor/Private/ControlRigElementDetails.h b/Engine/Plugins/Animation/ControlRig/Source/ControlRigEditor/Private/ControlRigElementDetails.h
--- a/Engine/Plugins/Animation/ControlRig/Source/ControlRigEditor/Private/ControlRigElementDetails.h
+++ b/Engine/Plugins/Animation/ControlRig/Source/ControlRigEditor/Private/ControlRigElementDetails.h
@@ -108,6 +108,31 @@ protected:
 	void OnTransformChanged(FEditPropertyChain* InPropertyChain);
 };
 
+class FRigControlTransformChannelDetails : public IPropertyTypeCustomization
+{
+public:
+
+	static TSharedRef<IPropertyTypeCustomization> MakeInstance()
+	{
+		return MakeShareable(new FRigControlTransformChannelDetails);
+	}
+
+	/** IPropertyTypeCustomization interface */
+	virtual void CustomizeHeader(TSharedRef<class IPropertyHandle> InStructPropertyHandle, class FDetailWidgetRow& HeaderRow, IPropertyTypeCustomizationUtils& StructCustomizationUtils) override;
+	virtual void CustomizeChildren(TSharedRef<class IPropertyHandle> InStructPropertyHandle, class IDetailChildrenBuilder& StructBuilder, IPropertyTypeCustomizationUtils& StructCustomizationUtils) override;
+
+protected:
+	ERigControlTransformChannel GetChannel() const;
+	int32 GetChannelAsInt32() const { return (int32)GetChannel(); }
+	void OnChannelChanged(int32 NewSelection, ESelectInfo::Type);
+
+public:
+	static const TArray<ERigControlTransformChannel>* GetVisibleChannelsForControlType(ERigControlType InControlType);
+
+private:
+	TSharedPtr<IPropertyHandle> Handle;
+};
+
 class FRigBaseElementDetails : public IDetailCustomization
 {
 public:
@@ -204,6 +229,9 @@ public:
 	bool IsAnyControlOfValueType(ERigControlType InType) const;
 	bool IsAnyControlNotOfValueType(ERigControlType InType) const;
 	bool IsAnyElementProcedural() const;
+	bool GetCommonElementType(ERigElementType& OutElementType) const;
+	bool GetCommonControlType(ERigControlType& OutControlType) const;
+	bool GetCommonAnimationType(ERigControlAnimationType& OutAnimationType) const;
 	const FPerElementInfo* FindElementByPredicate(const TFunction<bool(const FPerElementInfo&)>& InPredicate) const;
 	bool ContainsElementByPredicate(const TFunction<bool(const FPerElementInfo&)>& InPredicate) const;
 
diff --git a/Engine/Plugins/Animation/ControlRig/Source/ControlRigEditor/Private/Graph/ControlRigGraphPanelPinFactory.cpp b/Engine/Plugins/Animation/ControlRig/Source/ControlRigEditor/Private/Graph/ControlRigGraphPanelPinFactory.cpp
--- a/Engine/Plugins/Animation/ControlRig/Source/ControlRigEditor/Private/Graph/ControlRigGraphPanelPinFactory.cpp
+++ b/Engine/Plugins/Animation/ControlRig/Source/ControlRigEditor/Private/Graph/ControlRigGraphPanelPinFactory.cpp
@@ -16,6 +16,8 @@
 #include "Curves/CurveFloat.h"
 #include "RigVMModel/Nodes/RigVMUnitNode.h"
 #include "RigVMCore/RigVMExecuteContext.h"
+#include "Units/Execution/RigUnit_DynamicHierarchy.h"
+#include "ControlRigElementDetails.h"
 #include "IPropertyAccessEditor.h"
 
 TSharedPtr<SGraphPin> FControlRigGraphPanelPinFactory::CreatePin(UEdGraphPin* InPin) const
@@ -357,7 +359,7 @@ TSharedPtr<SGraphPin> FControlRigGraphPanelPinFactory::CreatePin(UEdGraphPin* In
 				{
 					return SNew(SGraphPinExec, InPin);
 				}
-				else if (InPin->PinType.PinSubCategoryObject == FRuntimeFloatCurve::StaticStruct())
+				if (InPin->PinType.PinSubCategoryObject == FRuntimeFloatCurve::StaticStruct())
 				{
 					return SNew(SControlRigGraphPinCurveFloat, InPin);
 				}
@@ -391,6 +393,51 @@ TSharedPtr<SGraphPin> FControlRigGraphPanelPinFactory::CreatePin(UEdGraphPin* In
 						}
 					}
 				}
+
+				const UEnum* RigControlTransformChannelEnum = StaticEnum<ERigControlTransformChannel>();
+				if (InPin->PinType.PinSubCategoryObject == RigControlTransformChannelEnum)
+				{
+					TSharedPtr<SWidget> ValueWidget = K2PinWidget->GetValueWidget();
+					if(ValueWidget.IsValid())
+					{
+						if(TSharedPtr<SPinComboBox> EnumCombo = StaticCastSharedPtr<SPinComboBox>(ValueWidget))
+						{
+							if(EnumCombo.IsValid())
+							{
+								if (const UControlRigGraphNode* RigNode = Cast<UControlRigGraphNode>(InPin->GetOwningNode()))
+								{
+									if (const URigVMPin* ModelPin = RigNode->GetModelPinFromPinPath(InPin->GetName()))
+									{
+										if(const URigVMUnitNode* UnitNode = Cast<URigVMUnitNode>(ModelPin->GetNode()))
+										{
+											if(UnitNode->GetScriptStruct() &&
+												UnitNode->GetScriptStruct()->IsChildOf(FRigUnit_HierarchyAddControlElement::StaticStruct()))
+											{
+												const TSharedPtr<FStructOnScope> StructInstanceScope = UnitNode->ConstructStructInstance();
+												const FRigUnit_HierarchyAddControlElement* StructInstance = 
+													(const FRigUnit_HierarchyAddControlElement*)StructInstanceScope->GetStructMemory();
+
+												if(const TArray<ERigControlTransformChannel>* VisibleChannels =
+													FRigControlTransformChannelDetails::GetVisibleChannelsForControlType(StructInstance->GetControlTypeToSpawn()))
+												{
+													for(int32 Index = 0; Index < RigControlTransformChannelEnum->NumEnums(); Index++)
+													{
+														const ERigControlTransformChannel Value =
+															(ERigControlTransformChannel)RigControlTransformChannelEnum->GetValueByIndex(Index);
+														if(!VisibleChannels->Contains(Value))
+														{
+															EnumCombo->RemoveItemByIndex(Index);
+														}
+													}
+												}
+											}
+										}
+									}
+								}
+							}
+						}
+					}
+				}
 			}
 
 			return K2PinWidget;
