Control Rig: Guard against invalid pin index access in UI

The crash is caused by widgets waiting to be garbage collected

diff --git a/Engine/Plugins/Animation/ControlRig/Source/ControlRigEditor/Private/Graph/SControlRigGraphNode.cpp b/Engine/Plugins/Animation/ControlRig/Source/ControlRigEditor/Private/Graph/SControlRigGraphNode.cpp
--- a/Engine/Plugins/Animation/ControlRig/Source/ControlRigEditor/Private/Graph/SControlRigGraphNode.cpp
+++ b/Engine/Plugins/Animation/ControlRig/Source/ControlRigEditor/Private/Graph/SControlRigGraphNode.cpp
@@ -1146,13 +1146,16 @@ const FSlateBrush* SControlRigGraphNode::GetExpanderImage(int32 InPinInfoIndex,
 		CollapsedRightBrush = FControlRigEditorStyle::Get().GetBrush(CollapsedRightName);
 	}
 
-	if(PinInfos[InPinInfoIndex].bExpanded)
+	if(PinInfos.IsValidIndex(InPinInfoIndex))
 	{
-		if(bHovered)
+		if(PinInfos[InPinInfoIndex].bExpanded)
 		{
-			return bLeft ? ExpandedHoveredLeftBrush : ExpandedHoveredRightBrush;
+			if(bHovered)
+			{
+				return bLeft ? ExpandedHoveredLeftBrush : ExpandedHoveredRightBrush;
+			}
+			return bLeft ? ExpandedLeftBrush : ExpandedRightBrush;
 		}
-		return bLeft ? ExpandedLeftBrush : ExpandedRightBrush;
 	}
 
 	if(bHovered)
@@ -1168,6 +1171,11 @@ FReply SControlRigGraphNode::OnExpanderArrowClicked(int32 InPinInfoIndex)
 	{
 		if(URigVMController* Controller = EdGraphNode->GetController())
 		{
+			if(!PinInfos.IsValidIndex(InPinInfoIndex))
+			{
+				return FReply::Unhandled();
+			}
+			
 			const FPinInfo& PinInfo = PinInfos[InPinInfoIndex];
 			TArray<FString> PinPathsToModify;
 			PinPathsToModify.Add(PinInfo.ModelPinPath);
