Sequencer: Control Rig: Check for valid toolkit to avoid possible crashes.
#trivial

diff --git a/Engine/Plugins/Animation/ControlRig/Source/ControlRigEditor/Private/EditMode/ControlRigEditMode.cpp b/Engine/Plugins/Animation/ControlRig/Source/ControlRigEditor/Private/EditMode/ControlRigEditMode.cpp
--- a/Engine/Plugins/Animation/ControlRig/Source/ControlRigEditor/Private/EditMode/ControlRigEditMode.cpp
+++ b/Engine/Plugins/Animation/ControlRig/Source/ControlRigEditor/Private/EditMode/ControlRigEditMode.cpp
@@ -2554,42 +2554,58 @@ void FControlRigEditMode::GetAllSelectedControls(TMap<UControlRig*, TArray<FRigE
 /** If Anim Slider is open, got to the next tool*/
 void FControlRigEditMode::ChangeAnimSliderTool()
 {
-	StaticCastSharedPtr<FControlRigEditModeToolkit>(Toolkit)->GetToNextActiveSlider();
+	if (Toolkit.IsValid())
+	{
+		StaticCastSharedPtr<FControlRigEditModeToolkit>(Toolkit)->GetToNextActiveSlider();
+	}
 }
 
 /** If Anim Slider is open, then can drag*/
 bool FControlRigEditMode::CanChangeAnimSliderTool() const
 {
-	return (StaticCastSharedPtr<FControlRigEditModeToolkit>(Toolkit)->CanChangeAnimSliderTool());
+	if (Toolkit.IsValid())
+	{
+		return (StaticCastSharedPtr<FControlRigEditModeToolkit>(Toolkit)->CanChangeAnimSliderTool());
+	}
+	return false;
 }
 
 /** If Anim Slider is open, drag the tool*/
 void FControlRigEditMode::DragAnimSliderTool(double IncrementVal)
 {
-	StaticCastSharedPtr<FControlRigEditModeToolkit>(Toolkit)->DragAnimSliderTool(IncrementVal);
+	if (Toolkit.IsValid())
+	{
+		StaticCastSharedPtr<FControlRigEditModeToolkit>(Toolkit)->DragAnimSliderTool(IncrementVal);
+	}
 }
 
 void FControlRigEditMode::StartAnimSliderTool(int32 InX)
 {
-	bisTrackingAnimToolDrag = true;
-	GEditor->BeginTransaction(LOCTEXT("AnimSliderBlend", "AnimSlider Blend"));
-	StartXValue = InX;
-	StaticCastSharedPtr<FControlRigEditModeToolkit>(Toolkit)->StartAnimSliderTool();
+	if (Toolkit.IsValid())
+	{
+		bisTrackingAnimToolDrag = true;
+		GEditor->BeginTransaction(LOCTEXT("AnimSliderBlend", "AnimSlider Blend"));
+		StartXValue = InX;
+		StaticCastSharedPtr<FControlRigEditModeToolkit>(Toolkit)->StartAnimSliderTool();
+	}
 }
 
 /** Reset and stop user the anim slider tool*/
 void FControlRigEditMode::ResetAnimSlider()
 {
-	GEditor->EndTransaction();
-	bisTrackingAnimToolDrag = false;
-	StartXValue.Reset();
-	StaticCastSharedPtr<FControlRigEditModeToolkit>(Toolkit)->ResetAnimSlider();
+	if (Toolkit.IsValid())
+	{
+		GEditor->EndTransaction();
+		bisTrackingAnimToolDrag = false;
+		StartXValue.Reset();
+		StaticCastSharedPtr<FControlRigEditModeToolkit>(Toolkit)->ResetAnimSlider();
+	}
 }
 
 /** If the Drag Anim Slider Tool is pressed*/
 bool FControlRigEditMode::IsDragAnimSliderToolPressed(FViewport* InViewport)
 {
-	if (IsInLevelEditor() && CanChangeAnimSliderTool())
+	if (IsInLevelEditor() && Toolkit.IsValid() &&  CanChangeAnimSliderTool())
 	{
 		bool bIsMovingSlider = false;
 		const FControlRigEditModeCommands& Commands = FControlRigEditModeCommands::Get();
