Allow Control Rig with invalid parent class to load with crashing.

diff --git a/Engine/Plugins/Animation/ControlRig/Source/ControlRig/Private/ControlRigBlueprintGeneratedClass.cpp b/Engine/Plugins/Animation/ControlRig/Source/ControlRig/Private/ControlRigBlueprintGeneratedClass.cpp
--- a/Engine/Plugins/Animation/ControlRig/Source/ControlRig/Private/ControlRigBlueprintGeneratedClass.cpp
+++ b/Engine/Plugins/Animation/ControlRig/Source/ControlRig/Private/ControlRigBlueprintGeneratedClass.cpp
@@ -25,16 +25,16 @@ uint8* UControlRigBlueprintGeneratedClass::GetPersistentUberGraphFrame(UObject*
 
 void UControlRigBlueprintGeneratedClass::PostInitInstance(UObject* InObj, FObjectInstancingGraph* InstanceGraph)
 {
-	UControlRig* ControlRig = Cast<UControlRig>(InObj);
-	check(ControlRig);
-
-	UControlRig* CDO = nullptr;
-	if(!ControlRig->HasAnyFlags(RF_ClassDefaultObject))
+	if (UControlRig* ControlRig = Cast<UControlRig>(InObj))
 	{
-		CDO = Cast<UControlRig>(GetDefaultObject());;
-	}
+		UControlRig* CDO = nullptr;
+		if (!ControlRig->HasAnyFlags(RF_ClassDefaultObject))
+		{
+			CDO = Cast<UControlRig>(GetDefaultObject());;
+		}
 
-	ControlRig->PostInitInstance(CDO);
+		ControlRig->PostInitInstance(CDO);
+	}
 }
 
 void UControlRigBlueprintGeneratedClass::Serialize(FArchive& Ar)
diff --git a/Engine/Plugins/Animation/ControlRig/Source/ControlRigDeveloper/Private/ControlRigBlueprint.cpp b/Engine/Plugins/Animation/ControlRig/Source/ControlRigDeveloper/Private/ControlRigBlueprint.cpp
--- a/Engine/Plugins/Animation/ControlRig/Source/ControlRigDeveloper/Private/ControlRigBlueprint.cpp
+++ b/Engine/Plugins/Animation/ControlRig/Source/ControlRigDeveloper/Private/ControlRigBlueprint.cpp
@@ -1158,17 +1158,19 @@ void UControlRigBlueprint::HandlePackageDone()
 
 void UControlRigBlueprint::BroadcastControlRigPackageDone()
 {
-	UControlRigBlueprintGeneratedClass* RigClass = GetControlRigBlueprintGeneratedClass();
-	UControlRig* CDO = Cast<UControlRig>(RigClass->GetDefaultObject(true /* create if needed */));
-	CDO->BroadCastEndLoadPackage();
-
-	TArray<UObject*> ArchetypeInstances;
-	CDO->GetArchetypeInstances(ArchetypeInstances);
-	for (UObject* Instance : ArchetypeInstances)
+	if (UControlRigBlueprintGeneratedClass* RigClass = GetControlRigBlueprintGeneratedClass())
 	{
-		if (UControlRig* InstanceRig = Cast<UControlRig>(Instance))
+		UControlRig* CDO = Cast<UControlRig>(RigClass->GetDefaultObject(true /* create if needed */));
+		CDO->BroadCastEndLoadPackage();
+
+		TArray<UObject*> ArchetypeInstances;
+		CDO->GetArchetypeInstances(ArchetypeInstances);
+		for (UObject* Instance : ArchetypeInstances)
 		{
-			InstanceRig->BroadCastEndLoadPackage();
+			if (UControlRig* InstanceRig = Cast<UControlRig>(Instance))
+			{
+				InstanceRig->BroadCastEndLoadPackage();
+			}
 		}
 	}
 }
@@ -1209,7 +1211,7 @@ void UControlRigBlueprint::RecompileVM()
 	RigGraphDisplaySettings.MaxMicroSeconds = RigGraphDisplaySettings.LastMaxMicroSeconds = (double)INDEX_NONE;
 
 	UControlRig* CDO = Cast<UControlRig>(RigClass->GetDefaultObject(true /* create if needed */));
-	if (CDO->VM != nullptr)
+	if (CDO && CDO->VM != nullptr)
 	{
 		TGuardValue<bool> ReentrantGuardSelf(bSuspendModelNotificationsForSelf, true);
 		TGuardValue<bool> ReentrantGuardOthers(bSuspendModelNotificationsForOthers, true);
@@ -1944,12 +1946,14 @@ bool UControlRigBlueprint::RemoveBreakpoint(URigVMNode* InBreakpointNode)
 
 void UControlRigBlueprint::RefreshControlRigBreakpoints()
 {
-	UControlRigBlueprintGeneratedClass* RigClass = GetControlRigBlueprintGeneratedClass();
-	UControlRig* CDO = Cast<UControlRig>(RigClass->GetDefaultObject(false));
-	CDO->GetDebugInfo().Reset();
-	for (URigVMNode* Node : RigVMBreakpointNodes)
+	if (UControlRigBlueprintGeneratedClass* RigClass = GetControlRigBlueprintGeneratedClass())
 	{
-		AddBreakpointToControlRig(Node);
+		UControlRig* CDO = Cast<UControlRig>(RigClass->GetDefaultObject(false));
+		CDO->GetDebugInfo().Reset();
+		for (URigVMNode* Node : RigVMBreakpointNodes)
+		{
+			AddBreakpointToControlRig(Node);
+		}
 	}
 }
 
@@ -1975,17 +1979,19 @@ TArray<FRigVMReferenceNodeData> UControlRigBlueprint::GetReferenceNodeData() con
 
 void UControlRigBlueprint::RequestControlRigInit()
 {
-	UControlRigBlueprintGeneratedClass* RigClass = GetControlRigBlueprintGeneratedClass();
-	UControlRig* CDO = Cast<UControlRig>(RigClass->GetDefaultObject(true /* create if needed */));
-	CDO->RequestInit();
-
-	TArray<UObject*> ArchetypeInstances;
-	CDO->GetArchetypeInstances(ArchetypeInstances);
-	for (UObject* Instance : ArchetypeInstances)
+	if (UControlRigBlueprintGeneratedClass* RigClass = GetControlRigBlueprintGeneratedClass())
 	{
-		if (UControlRig* InstanceRig = Cast<UControlRig>(Instance))
+		UControlRig* CDO = Cast<UControlRig>(RigClass->GetDefaultObject(true /* create if needed */));
+		CDO->RequestInit();
+
+		TArray<UObject*> ArchetypeInstances;
+		CDO->GetArchetypeInstances(ArchetypeInstances);
+		for (UObject* Instance : ArchetypeInstances)
 		{
-			InstanceRig->RequestInit();
+			if (UControlRig* InstanceRig = Cast<UControlRig>(Instance))
+			{
+				InstanceRig->RequestInit();
+			}
 		}
 	}
 }
@@ -3104,17 +3110,19 @@ void UControlRigBlueprint::ClearTransientControls()
 		ValueScope = MakeUnique<FControlValueScope>(this);
 	}
 
-	UControlRigBlueprintGeneratedClass* RigClass = GetControlRigBlueprintGeneratedClass();
-	UControlRig* CDO = Cast<UControlRig>(RigClass->GetDefaultObject(true /* create if needed */));
-
-	TArray<UObject*> ArchetypeInstances;
-	CDO->GetArchetypeInstances(ArchetypeInstances);
-	for (UObject* ArchetypeInstance : ArchetypeInstances)
+	if (UControlRigBlueprintGeneratedClass* RigClass = GetControlRigBlueprintGeneratedClass())
 	{
-		UControlRig* InstancedControlRig = Cast<UControlRig>(ArchetypeInstance);
-		if (InstancedControlRig)
+		UControlRig* CDO = Cast<UControlRig>(RigClass->GetDefaultObject(true /* create if needed */));
+
+		TArray<UObject*> ArchetypeInstances;
+		CDO->GetArchetypeInstances(ArchetypeInstances);
+		for (UObject* ArchetypeInstance : ArchetypeInstances)
 		{
-			InstancedControlRig->ClearTransientControls();
+			UControlRig* InstancedControlRig = Cast<UControlRig>(ArchetypeInstance);
+			if (InstancedControlRig)
+			{
+				InstancedControlRig->ClearTransientControls();
+			}
 		}
 	}
 }
diff --git a/Engine/Source/Editor/BlueprintGraph/Private/BlueprintNodeTemplateCache.cpp b/Engine/Source/Editor/BlueprintGraph/Private/BlueprintNodeTemplateCache.cpp
--- a/Engine/Source/Editor/BlueprintGraph/Private/BlueprintNodeTemplateCache.cpp
+++ b/Engine/Source/Editor/BlueprintGraph/Private/BlueprintNodeTemplateCache.cpp
@@ -394,7 +394,8 @@ UEdGraphNode* FBlueprintNodeTemplateCache::GetNodeTemplate(UBlueprintNodeSpawner
 						GeneratedClassType = TargetBlueprint->GeneratedClass->GetClass();
 					}
 
-					CompatibleBlueprint = MakeCompatibleBlueprint(BlueprintClass, TargetBlueprint->ParentClass, GeneratedClassType);
+					UClass* ParentClass = TargetBlueprint->ParentClass != nullptr ? TargetBlueprint->ParentClass.Get() : UObject::StaticClass();
+					CompatibleBlueprint = MakeCompatibleBlueprint(BlueprintClass, ParentClass, GeneratedClassType);
 					if (!CacheBlueprintOuter(CompatibleBlueprint))
 					{
 						LogCacheFullMsg();
