Sequencer: Control Rig: Tweak IsSelectable on on ControlRigShapes so that it means IsSeletactable in the Viewport, we still want it to be selected in the outliner. Also we hide controls if the hosting component's actor is not visible.

diff --git a/Engine/Plugins/Animation/ControlRig/Source/ControlRig/Private/ControlRigGizmoActor.cpp b/Engine/Plugins/Animation/ControlRig/Source/ControlRig/Private/ControlRigGizmoActor.cpp
--- a/Engine/Plugins/Animation/ControlRig/Source/ControlRig/Private/ControlRigGizmoActor.cpp
+++ b/Engine/Plugins/Animation/ControlRig/Source/ControlRig/Private/ControlRigGizmoActor.cpp
@@ -19,6 +19,7 @@ AControlRigShapeActor::AControlRigShapeActor(const FObjectInitializer& ObjectIni
 	, OverrideColor(0, 0, 0, 0)
 	, bSelected(false)
 	, bHovered(false)
+	, bSelectable(true)
 {
 
 	ActorRootComponent = CreateDefaultSubobject<USceneComponent>(TEXT("SceneComponent0"));
@@ -38,12 +39,9 @@ AControlRigShapeActor::AControlRigShapeActor(const FObjectInitializer& ObjectIni
 	StaticMeshComponent->bSelectable = true;
 }
 
+//client should check to see if it's seletable based upon how the selection occurs (viewport or outliner,etc..)
 void AControlRigShapeActor::SetSelected(bool bInSelected)
 {
-	if(!IsSelectable() && bInSelected)
-	{
-		return;
-	}
 	if(bSelected != bInSelected)
 	{
 		bSelected = bInSelected;
@@ -59,15 +57,16 @@ bool AControlRigShapeActor::IsSelectedInEditor() const
 
 bool AControlRigShapeActor::IsSelectable() const
 {
-	return StaticMeshComponent->bSelectable;
+	return bSelectable;
 }
 
+//we no longer set the StaticMeshComponent bSelectable flag since that drives if it can move with a gizmo or not
 void AControlRigShapeActor::SetSelectable(bool bInSelectable)
 {
-	if (StaticMeshComponent->bSelectable != bInSelectable)
+	if (bSelectable != bInSelectable)
 	{
-		StaticMeshComponent->bSelectable = bInSelectable;
-		if (!StaticMeshComponent->bSelectable)
+		bSelectable = bInSelectable;
+		if (!bSelectable)
 		{
 			SetSelected(false);
 		}
diff --git a/Engine/Plugins/Animation/ControlRig/Source/ControlRig/Public/ControlRigGizmoActor.h b/Engine/Plugins/Animation/ControlRig/Source/ControlRig/Public/ControlRigGizmoActor.h
--- a/Engine/Plugins/Animation/ControlRig/Source/ControlRig/Public/ControlRigGizmoActor.h
+++ b/Engine/Plugins/Animation/ControlRig/Source/ControlRig/Public/ControlRigGizmoActor.h
@@ -174,6 +174,9 @@ private:
 	UPROPERTY(BlueprintGetter = IsHovered, BlueprintSetter = SetHovered, Category = "ControlRig|Shape")
 	uint8 bHovered : 1;
 
+	/** Whether or not this control is selectable*/
+	uint8 bSelectable : 1;
+
 };
 
 /**
diff --git a/Engine/Plugins/Animation/ControlRig/Source/ControlRigEditor/Private/EditMode/ControlRigEditMode.cpp b/Engine/Plugins/Animation/ControlRig/Source/ControlRigEditor/Private/EditMode/ControlRigEditMode.cpp
--- a/Engine/Plugins/Animation/ControlRig/Source/ControlRigEditor/Private/EditMode/ControlRigEditMode.cpp
+++ b/Engine/Plugins/Animation/ControlRig/Source/ControlRigEditor/Private/EditMode/ControlRigEditMode.cpp
@@ -4192,6 +4192,23 @@ bool FControlRigEditMode::ModeSupportedByShapeActor(const AControlRigShapeActor*
 	return false;
 }
 
+bool FControlRigEditMode::IsControlRigSkelMeshVisible(UControlRig* ControlRig) const
+{
+	if (IsInLevelEditor())
+	{
+		if (ControlRig)
+		{
+			if (USceneComponent* SceneComponent = GetHostingSceneComponent(ControlRig))
+			{
+				const AActor* Actor = SceneComponent->GetTypedOuter<AActor>();
+				return Actor ? (Actor->IsHiddenEd() == false && SceneComponent->IsVisibleInEditor()) : SceneComponent->IsVisibleInEditor();
+			}
+		}
+		return false;
+	}
+	return true;
+}
+
 void FControlRigEditMode::TickControlShape(AControlRigShapeActor* ShapeActor, const FTransform& ComponentTransform)
 {
 	const UControlRigEditModeSettings* Settings = GetDefault<UControlRigEditModeSettings>();
@@ -4204,17 +4221,17 @@ void FControlRigEditMode::TickControlShape(AControlRigShapeActor* ShapeActor, co
 
 			if (FRigControlElement* ControlElement = ControlRig->FindControl(ShapeActor->ControlName))
 			{
-				const bool bControlsHiddenInViewport = Settings->bHideControlShapes || !ControlRig->GetControlsVisible();
-				
+				const bool bControlsHiddenInViewport = Settings->bHideControlShapes || !ControlRig->GetControlsVisible()
+					|| (IsControlRigSkelMeshVisible(ControlRig) == false);
+
 				bool bIsVisible = ControlElement->Settings.IsVisible();
-				bool bRespectVisibilityForSelection = true;
+				bool bRespectVisibilityForSelection = true; 
 
 				if(!bControlsHiddenInViewport)
 				{
 					if(ControlElement->Settings.AnimationType == ERigControlAnimationType::ProxyControl)
-					{
+					{					
 						bRespectVisibilityForSelection = false;
-						
 						if(Settings->bShowAllProxyControls)
 						{
 							bIsVisible = true;
diff --git a/Engine/Plugins/Animation/ControlRig/Source/ControlRigEditor/Private/EditMode/ControlRigEditMode.h b/Engine/Plugins/Animation/ControlRig/Source/ControlRigEditor/Private/EditMode/ControlRigEditMode.h
--- a/Engine/Plugins/Animation/ControlRig/Source/ControlRigEditor/Private/EditMode/ControlRigEditMode.h
+++ b/Engine/Plugins/Animation/ControlRig/Source/ControlRigEditor/Private/EditMode/ControlRigEditMode.h
@@ -436,6 +436,9 @@ protected:
 	USceneComponent* GetHostingSceneComponent(const UControlRig* ControlRig = nullptr) const;
 	FTransform	GetHostingSceneComponentTransform(const UControlRig* ControlRig =  nullptr) const;
 
+	//Get if the hosted component is visible
+	bool IsControlRigSkelMeshVisible(UControlRig* ControlRig) const;
+
 private:
 
 	// Post pose update handler
