Sequencer: Control Rig: Rotators additive section modification was turned off for some reason. Also some pointer checks to avoid Toolkit being null causing crashes

diff --git a/Engine/Plugins/Animation/ControlRig/Source/ControlRigEditor/Private/Sequencer/ControlRigParameterTrackEditor.cpp b/Engine/Plugins/Animation/ControlRig/Source/ControlRigEditor/Private/Sequencer/ControlRigParameterTrackEditor.cpp
--- a/Engine/Plugins/Animation/ControlRig/Source/ControlRigEditor/Private/Sequencer/ControlRigParameterTrackEditor.cpp
+++ b/Engine/Plugins/Animation/ControlRig/Source/ControlRigEditor/Private/Sequencer/ControlRigParameterTrackEditor.cpp
@@ -3660,16 +3660,12 @@ bool FControlRigParameterTrackEditor::ModifyOurGeneratedKeysByCurrentAndWeight(U
 					if (pChannelIndex && pChannelIndex->GeneratedKeyIndex != INDEX_NONE)
 					{
 						ChannelIndex = pChannelIndex->GeneratedKeyIndex;
-
-						if (ControlElement->Settings.ControlType != ERigControlType::Rotator)
-						{
-							float FVal = (float)Val.Val.X;
-							GeneratedTotalKeys[ChannelIndex]->ModifyByCurrentAndWeight(Proxy, EvaluateTime, (void*)&FVal, Weight);
-							FVal = (float)Val.Val.Y;
-							GeneratedTotalKeys[ChannelIndex + 1]->ModifyByCurrentAndWeight(Proxy, EvaluateTime, (void*)&FVal, Weight);
-							FVal = (float)Val.Val.Z;
-							GeneratedTotalKeys[ChannelIndex + 2]->ModifyByCurrentAndWeight(Proxy, EvaluateTime, (void*)&FVal, Weight);
-						}
+						float FVal = (float)Val.Val.X;
+						GeneratedTotalKeys[ChannelIndex]->ModifyByCurrentAndWeight(Proxy, EvaluateTime, (void*)&FVal, Weight);
+						FVal = (float)Val.Val.Y;
+						GeneratedTotalKeys[ChannelIndex + 1]->ModifyByCurrentAndWeight(Proxy, EvaluateTime, (void*)&FVal, Weight);
+						FVal = (float)Val.Val.Z;
+						GeneratedTotalKeys[ChannelIndex + 2]->ModifyByCurrentAndWeight(Proxy, EvaluateTime, (void*)&FVal, Weight);			
 					}
 					break;
 				}
@@ -4842,6 +4838,7 @@ FControlRigEditMode* FControlRigParameterTrackEditor::GetEditMode(bool bForceAct
 	return nullptr;
 }
 
+
 #undef LOCTEXT_NAMESPACE
 
 
