RigVM: Move dispatch helper functionality to RigVM module

diff --git a/Engine/Plugins/Animation/ControlRig/Source/ControlRig/Private/Units/RigDispatchFactory.cpp b/Engine/Plugins/Animation/ControlRig/Source/ControlRig/Private/Units/RigDispatchFactory.cpp
--- a/Engine/Plugins/Animation/ControlRig/Source/ControlRig/Private/Units/RigDispatchFactory.cpp
+++ b/Engine/Plugins/Animation/ControlRig/Source/ControlRig/Private/Units/RigDispatchFactory.cpp
@@ -8,65 +8,11 @@
 
 FString FRigDispatchFactory::GetArgumentDefaultValue(const FName& InArgumentName, TRigVMTypeIndex InTypeIndex) const
 {
-	if(FRigVMRegistry::Get().IsArrayType(InTypeIndex))
-	{
-		static const FString EmptyArrayString = TEXT("()");
-		return EmptyArrayString;
-	}
-	if(InTypeIndex == RigVMTypeUtils::TypeIndex::Bool)
-	{
-		static const FString FalseString = TEXT("False");
-		return FalseString;
-	}
-	if(InTypeIndex == RigVMTypeUtils::TypeIndex::Float || InTypeIndex == RigVMTypeUtils::TypeIndex::Double)
-	{
-		static const FString ZeroString = TEXT("0.000000");
-		return ZeroString;
-	}
-	if(InTypeIndex == RigVMTypeUtils::TypeIndex::Int32)
-	{
-		static const FString ZeroString = TEXT("0");
-		return ZeroString;
-	}
-	if(InTypeIndex == RigVMTypeUtils::TypeIndex::FName)
-	{
-		return FName(NAME_None).ToString();
-	}
-	if(InTypeIndex == FRigVMRegistry::Get().GetTypeIndex<FVector2D>())
-	{
-		static FString DefaultValueString = GetDefaultValueForStruct(FVector2D::ZeroVector);
-		return DefaultValueString;
-	}
-	if(InTypeIndex == FRigVMRegistry::Get().GetTypeIndex<FVector>())
-	{
-		static FString DefaultValueString = GetDefaultValueForStruct(FVector::ZeroVector); 
-		return DefaultValueString;
-	}
-	if(InTypeIndex == FRigVMRegistry::Get().GetTypeIndex<FRotator>())
-	{
-		static FString DefaultValueString = GetDefaultValueForStruct(FRotator::ZeroRotator); 
-		return DefaultValueString;
-	}
-	if(InTypeIndex == FRigVMRegistry::Get().GetTypeIndex<FQuat>())
-	{
-		static FString DefaultValueString = GetDefaultValueForStruct(FQuat::Identity); 
-		return DefaultValueString;
-	}
-	if(InTypeIndex == FRigVMRegistry::Get().GetTypeIndex<FTransform>())
-	{
-		static FString DefaultValueString = GetDefaultValueForStruct(FTransform::Identity); 
-		return DefaultValueString;
-	}
 	if(InTypeIndex == FRigVMRegistry::Get().GetTypeIndex<FEulerTransform>())
 	{
 		static FString DefaultValueString = GetDefaultValueForStruct(FEulerTransform::Identity); 
 		return DefaultValueString;
 	}
-	if(InTypeIndex == FRigVMRegistry::Get().GetTypeIndex<FLinearColor>())
-	{
-		static FString DefaultValueString = GetDefaultValueForStruct(FLinearColor::White); 
-		return DefaultValueString;
-	}
 	if(InTypeIndex == FRigVMRegistry::Get().GetTypeIndex<FRigElementKey>())
 	{
 		static FString DefaultValueString = GetDefaultValueForStruct(FRigElementKey(NAME_None, ERigElementType::Bone)); 
diff --git a/Engine/Plugins/Animation/ControlRig/Source/ControlRig/Public/Units/RigDispatchFactory.h b/Engine/Plugins/Animation/ControlRig/Source/ControlRig/Public/Units/RigDispatchFactory.h
--- a/Engine/Plugins/Animation/ControlRig/Source/ControlRig/Public/Units/RigDispatchFactory.h
+++ b/Engine/Plugins/Animation/ControlRig/Source/ControlRig/Public/Units/RigDispatchFactory.h
@@ -40,48 +40,9 @@ struct CONTROLRIG_API FRigDispatchFactory : public FRigVMDispatchFactory
 
 #endif
 
-	template <
-		typename T,
-		typename TEnableIf<TRigVMIsBaseStructure<T>::Value, T>::Type* = nullptr
-	>
-	FORCEINLINE static FString GetDefaultValueForStruct(const T& InValue)
-	{
-		static FString ValueString;
-		if(ValueString.IsEmpty())
-		{
-			TBaseStructure<T>::Get()->ExportText(ValueString, &InValue, &InValue, nullptr, PPF_None, nullptr);
-		}
-		return ValueString;
-	}
-
-	template <
-		typename T,
-		typename TEnableIf<TModels<CRigVMUStruct, T>::Value>::Type * = nullptr
-	>
-	FORCEINLINE static FString GetDefaultValueForStruct(const T& InValue)
-	{
-		static FString ValueString;
-		if(ValueString.IsEmpty())
-		{
-			T::StaticStruct()->ExportText(ValueString, &InValue, &InValue, nullptr, PPF_None, nullptr);
-		}
-		return ValueString;
-	}
-
 	FORCEINLINE static const FRigUnitContext& GetRigUnitContext(const FRigVMExtendedExecuteContext& InContext)
 	{
 		return *(const FRigUnitContext*)InContext.OpaqueArguments[0];
 	}
-
-#if WITH_EDITOR
-	FORCEINLINE bool CheckArgumentType(bool bCondition, const FName& InArgumentName) const
-	{
-		if(!bCondition)
-		{
-			UE_LOG(LogControlRig, Error, TEXT("Fatal: '%s' Argument '%s' has incorrect type."), *GetFactoryName().ToString(), *InArgumentName.ToString())
-		}
-		return bCondition;
-	}
-#endif
 };
 
diff --git a/Engine/Source/Runtime/RigVM/Private/RigVMCore/RigVMDispatchFactory.cpp b/Engine/Source/Runtime/RigVM/Private/RigVMCore/RigVMDispatchFactory.cpp
--- a/Engine/Source/Runtime/RigVM/Private/RigVMCore/RigVMDispatchFactory.cpp
+++ b/Engine/Source/Runtime/RigVM/Private/RigVMCore/RigVMDispatchFactory.cpp
@@ -39,6 +39,65 @@ FText FRigVMDispatchFactory::GetNodeTooltip(const FRigVMTemplateTypeMap& InTypes
 	return GetScriptStruct()->GetToolTipText();
 }
 
+FString FRigVMDispatchFactory::GetArgumentDefaultValue(const FName& InArgumentName, TRigVMTypeIndex InTypeIndex) const
+{
+	if(FRigVMRegistry::Get().IsArrayType(InTypeIndex))
+	{
+		static const FString EmptyArrayString = TEXT("()");
+		return EmptyArrayString;
+	}
+	if(InTypeIndex == RigVMTypeUtils::TypeIndex::Bool)
+	{
+		static const FString FalseString = TEXT("False");
+		return FalseString;
+	}
+	if(InTypeIndex == RigVMTypeUtils::TypeIndex::Float || InTypeIndex == RigVMTypeUtils::TypeIndex::Double)
+	{
+		static const FString ZeroString = TEXT("0.000000");
+		return ZeroString;
+	}
+	if(InTypeIndex == RigVMTypeUtils::TypeIndex::Int32)
+	{
+		static const FString ZeroString = TEXT("0");
+		return ZeroString;
+	}
+	if(InTypeIndex == RigVMTypeUtils::TypeIndex::FName)
+	{
+		return FName(NAME_None).ToString();
+	}
+	if(InTypeIndex == FRigVMRegistry::Get().GetTypeIndex<FVector2D>())
+	{
+		static FString DefaultValueString = GetDefaultValueForStruct(FVector2D::ZeroVector);
+		return DefaultValueString;
+	}
+	if(InTypeIndex == FRigVMRegistry::Get().GetTypeIndex<FVector>())
+	{
+		static FString DefaultValueString = GetDefaultValueForStruct(FVector::ZeroVector); 
+		return DefaultValueString;
+	}
+	if(InTypeIndex == FRigVMRegistry::Get().GetTypeIndex<FRotator>())
+	{
+		static FString DefaultValueString = GetDefaultValueForStruct(FRotator::ZeroRotator); 
+		return DefaultValueString;
+	}
+	if(InTypeIndex == FRigVMRegistry::Get().GetTypeIndex<FQuat>())
+	{
+		static FString DefaultValueString = GetDefaultValueForStruct(FQuat::Identity); 
+		return DefaultValueString;
+	}
+	if(InTypeIndex == FRigVMRegistry::Get().GetTypeIndex<FTransform>())
+	{
+		static FString DefaultValueString = GetDefaultValueForStruct(FTransform::Identity); 
+		return DefaultValueString;
+	}
+	if(InTypeIndex == FRigVMRegistry::Get().GetTypeIndex<FLinearColor>())
+	{
+		static FString DefaultValueString = GetDefaultValueForStruct(FLinearColor::White); 
+		return DefaultValueString;
+	}
+	return FString();
+}
+
 FString FRigVMDispatchFactory::GetCategory() const
 {
 	if(const UScriptStruct* ScriptStruct = GetScriptStruct())
diff --git a/Engine/Source/Runtime/RigVM/Public/RigVMCore/RigVMDispatchFactory.h b/Engine/Source/Runtime/RigVM/Public/RigVMCore/RigVMDispatchFactory.h
--- a/Engine/Source/Runtime/RigVM/Public/RigVMCore/RigVMDispatchFactory.h
+++ b/Engine/Source/Runtime/RigVM/Public/RigVMCore/RigVMDispatchFactory.h
@@ -54,7 +54,7 @@ public:
 	virtual FText GetNodeTooltip(const FRigVMTemplateTypeMap& InTypes) const;
 
 	// returns the default value for an argument
-	FORCEINLINE virtual FString GetArgumentDefaultValue(const FName& InArgumentName, TRigVMTypeIndex InTypeIndex) const { return FString(); }
+	virtual FString GetArgumentDefaultValue(const FName& InArgumentName, TRigVMTypeIndex InTypeIndex) const;
 
 	// returns the tooltip for an argument
 	FORCEINLINE virtual FText GetArgumentTooltip(const FName& InArgumentName, TRigVMTypeIndex InTypeIndex) const { return FText(); }
@@ -122,6 +122,45 @@ protected:
 
 	virtual TArray<FRigVMExecuteArgument> GetExecuteArguments_Impl() const { return TArray<FRigVMExecuteArgument>(); }
 
+	template <
+	typename T,
+	typename TEnableIf<TRigVMIsBaseStructure<T>::Value, T>::Type* = nullptr
+	>
+	FORCEINLINE static FString GetDefaultValueForStruct(const T& InValue)
+	{
+		static FString ValueString;
+		if(ValueString.IsEmpty())
+		{
+			TBaseStructure<T>::Get()->ExportText(ValueString, &InValue, &InValue, nullptr, PPF_None, nullptr);
+		}
+		return ValueString;
+	}
+
+	template <
+		typename T,
+		typename TEnableIf<TModels<CRigVMUStruct, T>::Value>::Type * = nullptr
+	>
+	FORCEINLINE static FString GetDefaultValueForStruct(const T& InValue)
+	{
+		static FString ValueString;
+		if(ValueString.IsEmpty())
+		{
+			T::StaticStruct()->ExportText(ValueString, &InValue, &InValue, nullptr, PPF_None, nullptr);
+		}
+		return ValueString;
+	}
+
+#if WITH_EDITOR
+	FORCEINLINE bool CheckArgumentType(bool bCondition, const FName& InArgumentName) const
+	{
+		if(!bCondition)
+		{
+			UE_LOG(LogRigVM, Error, TEXT("Fatal: '%s' Argument '%s' has incorrect type."), *GetFactoryName().ToString(), *InArgumentName.ToString())
+		}
+		return bCondition;
+	}
+#endif
+
 	static const FString DispatchPrefix;
 
 	UScriptStruct* FactoryScriptStruct;
