Control Rig: Fix cast links backwards comp code

diff --git a/Engine/Plugins/Animation/ControlRig/Source/ControlRigDeveloper/Private/ControlRigBlueprint.cpp b/Engine/Plugins/Animation/ControlRig/Source/ControlRigDeveloper/Private/ControlRigBlueprint.cpp
--- a/Engine/Plugins/Animation/ControlRig/Source/ControlRigDeveloper/Private/ControlRigBlueprint.cpp
+++ b/Engine/Plugins/Animation/ControlRig/Source/ControlRigDeveloper/Private/ControlRigBlueprint.cpp
@@ -4594,7 +4594,8 @@ void UControlRigBlueprint::PatchLinksWithCast()
 					Controller->ReattachLinksToPinObjects(true, &LinksToReattach, true, false, false);
 				}
 			}
-			
+
+			TGuardValue<bool> DisableRecomputeTemplates(Controller->bSuspendRecomputingTemplateFilters, true);
 			Controller->BreakLink(Tuple.Get<2>(), Tuple.Get<3>(), false);
 
 			// notify the user that the link has been broken.
diff --git a/Engine/Source/Developer/RigVMDeveloper/Public/RigVMModel/RigVMController.h b/Engine/Source/Developer/RigVMDeveloper/Public/RigVMModel/RigVMController.h
--- a/Engine/Source/Developer/RigVMDeveloper/Public/RigVMModel/RigVMController.h
+++ b/Engine/Source/Developer/RigVMDeveloper/Public/RigVMModel/RigVMController.h
@@ -1229,7 +1229,13 @@ private:
 	bool bIsTransacting; // Performing undo/redo transaction
 	bool bIsRunningUnitTest;
 	bool bIsFullyResolvingTemplateNode;
+
+public:
+	
 	bool bSuspendRecomputingTemplateFilters;
+
+private:
+
 #if WITH_EDITOR
 	bool bRegisterTemplateNodeUsage;
 #endif
diff --git a/Engine/Source/Runtime/RigVM/Private/RigVMCore/RigVMNameCache.cpp b/Engine/Source/Runtime/RigVM/Private/RigVMCore/RigVMNameCache.cpp
--- a/Engine/Source/Runtime/RigVM/Private/RigVMCore/RigVMNameCache.cpp
+++ b/Engine/Source/Runtime/RigVM/Private/RigVMCore/RigVMNameCache.cpp
@@ -5,7 +5,7 @@
 
 #if WITH_EDITOR
 #include "RigVMCore/RigVM.h"
-TAutoConsoleVariable<int32> CVarRigVMNameCacheMaxSize(TEXT("RigVM.NameCacheMaxSize"), 512, TEXT("Change to control how many names are cached per VM instance."));
+TAutoConsoleVariable<int32> CVarRigVMNameCacheMaxSize(TEXT("RigVM.NameCacheMaxSize"), 4096, TEXT("Change to control how many names are cached per VM instance."));
 #endif
 
 ////////////////////////////////////////////////////////////////////////////////
