Sequencer: ControlRig: Python Add ability to set FK Control Rig as Addtiive to Python/BP. Add Python EngineTest to make sure additve mode correctly passes through. Also added a little function to do the blend tool based on selection as requested by Greg/Fred since I was already touching these files.

diff --git a/Engine/Plugins/Animation/ControlRig/Source/ControlRig/Private/Rigs/FKControlRig.cpp b/Engine/Plugins/Animation/ControlRig/Source/ControlRig/Private/Rigs/FKControlRig.cpp
--- a/Engine/Plugins/Animation/ControlRig/Source/ControlRig/Private/Rigs/FKControlRig.cpp
+++ b/Engine/Plugins/Animation/ControlRig/Source/ControlRig/Private/Rigs/FKControlRig.cpp
@@ -27,7 +27,7 @@ FName UFKControlRig::GetControlName(const FName& InName, const ERigElementType&
 {
 	if (InName != NAME_None)
 	{
-		check(InType == ERigElementType::Bone || InType == ERigElementType::Curve);
+		//@jurre to look at after vacation, causes ensure check(InType == ERigElementType::Bone || InType == ERigElementType::Curve);
 		static thread_local TMap<FName, FName> NameToControlMappings[2];
 		TMap<FName, FName>& NameToControlMapping = NameToControlMappings[InType == ERigElementType::Bone ? 0 : 1];
 		if (const FName* CachedName = NameToControlMapping.Find(InName))
@@ -62,7 +62,7 @@ FName UFKControlRig::GetControlTargetName(const FName& InName, const ERigElement
 {
 	if (InName != NAME_None)
 	{
-		check(InType == ERigElementType::Bone || InType == ERigElementType::Curve);
+		//@jurre to look at after vacation, causes ensure check(InType == ERigElementType::Bone || InType == ERigElementType::Curve);
 		static thread_local TMap<FName, FName> NameToTargetMappings[2];
 		TMap<FName, FName>& NameToTargetMapping = NameToTargetMappings[InType == ERigElementType::Bone ? 0 : 1];
 		if (const FName* CachedName = NameToTargetMapping.Find(InName))
diff --git a/Engine/Plugins/Animation/ControlRig/Source/ControlRig/Public/Rigs/FKControlRig.h b/Engine/Plugins/Animation/ControlRig/Source/ControlRig/Public/Rigs/FKControlRig.h
--- a/Engine/Plugins/Animation/ControlRig/Source/ControlRig/Public/Rigs/FKControlRig.h
+++ b/Engine/Plugins/Animation/ControlRig/Source/ControlRig/Public/Rigs/FKControlRig.h
@@ -18,7 +18,7 @@ struct FFKBoneCheckInfo
 	bool  bActive;
 };
 
-UENUM()
+UENUM(Blueprintable)
 enum class EControlRigFKRigExecuteMode: uint8
 {
 	/** Replaces the current pose */
@@ -62,6 +62,7 @@ public:
 	void ToggleApplyMode();
 	bool CanToggleApplyMode() const { return true; }
 	bool IsApplyModeAdditive() const { return ApplyMode == EControlRigFKRigExecuteMode::Additive; }
+	EControlRigFKRigExecuteMode GetApplyMode() const { return ApplyMode; }
 
 	// Ensures that controls mask is updated according to contained ControlRig (control) elements
 	void RefreshActiveControls();
diff --git a/Engine/Plugins/Animation/ControlRig/Source/ControlRigEditor/Private/ControlRigSequencerEditorLibrary.cpp b/Engine/Plugins/Animation/ControlRig/Source/ControlRigEditor/Private/ControlRigSequencerEditorLibrary.cpp
--- a/Engine/Plugins/Animation/ControlRig/Source/ControlRigEditor/Private/ControlRigSequencerEditorLibrary.cpp
+++ b/Engine/Plugins/Animation/ControlRig/Source/ControlRigEditor/Private/ControlRigSequencerEditorLibrary.cpp
@@ -40,6 +40,7 @@
 #include "LevelSequenceEditorBlueprintLibrary.h"
 #include "Tools/ConstraintBaker.h"
 #include "TransformConstraint.h"
+#include "Rigs/FKControlRig.h"
 
 #include UE_INLINE_GENERATED_CPP_BY_NAME(ControlRigSequencerEditorLibrary)
 
@@ -403,6 +404,43 @@ bool UControlRigSequencerEditorLibrary::TweenControlRig(ULevelSequence* LevelSeq
 	return false;
 }
 
+bool UControlRigSequencerEditorLibrary::BlendValuesOnSelected(ULevelSequence* LevelSequence, EAnimToolBlendOperation BlendOperation, float BlendValue)
+{
+	TWeakPtr<ISequencer> WeakSequencer = GetSequencerFromAsset();
+	if (WeakSequencer.IsValid() && WeakSequencer.Pin()->GetFocusedMovieSceneSequence() == LevelSequence
+		&& LevelSequence->GetMovieScene())
+	{
+		FControlsToTween ControlsToTween;
+		LevelSequence->GetMovieScene()->Modify();
+		switch(BlendOperation)
+		{ 
+			case EAnimToolBlendOperation::Tween:
+			{
+				FControlsToTween BlendTool;
+				BlendTool.Setup(WeakSequencer);
+				BlendTool.Blend(WeakSequencer, BlendValue);
+				return true;
+			}
+			case EAnimToolBlendOperation::BlendToNeighbor:
+			{
+				FBlendNeighborSlider BlendTool;
+				BlendTool.Setup(WeakSequencer);
+				BlendTool.Blend(WeakSequencer, BlendValue);
+				return true;
+			}
+			case EAnimToolBlendOperation::PushPull:
+			{
+				FPushPullSlider BlendTool;
+				BlendTool.Setup(WeakSequencer);
+				BlendTool.Blend(WeakSequencer, BlendValue);
+				return true;
+			}
+		}
+	}
+	return false;
+}
+
+
 bool UControlRigSequencerEditorLibrary::BakeConstraint(UWorld* World, UTickableConstraint* Constraint, const TArray<FFrameNumber>& Frames, ESequenceTimeUnit TimeUnit)
 {
 	if (!World)
@@ -2547,5 +2585,30 @@ void UControlRigSequencerEditorLibrary::HideAllControls(UMovieSceneSection* InSe
 	ParameterSection->FillControlsMask(false);
 }
 
+bool UControlRigSequencerEditorLibrary::IsFKControlRig(UControlRig* InControlRig)
+{
+	return (InControlRig && InControlRig->IsA<UFKControlRig>());
+}
+
+EControlRigFKRigExecuteMode UControlRigSequencerEditorLibrary::GetFKControlRigApplyMode(UControlRig* InControlRig)
+{
+	EControlRigFKRigExecuteMode ApplyMode = EControlRigFKRigExecuteMode::Direct;
+	if (UFKControlRig* FKRig = Cast<UFKControlRig>(InControlRig))
+	{
+		ApplyMode = FKRig->GetApplyMode();
+	}
+	return ApplyMode;
+}
+
+bool UControlRigSequencerEditorLibrary::SetControlRigApplyMode(UControlRig* InControlRig, EControlRigFKRigExecuteMode InApplyMode)
+{
+	if (UFKControlRig* FKRig = Cast<UFKControlRig>(InControlRig))
+	{
+		FKRig->SetApplyMode(InApplyMode);
+		return true;
+	}
+	return false;
+}
+
 #undef LOCTEXT_NAMESPACE
 
diff --git a/Engine/Plugins/Animation/ControlRig/Source/ControlRigEditor/Public/ControlRigSequencerEditorLibrary.h b/Engine/Plugins/Animation/ControlRig/Source/ControlRigEditor/Public/ControlRigSequencerEditorLibrary.h
--- a/Engine/Plugins/Animation/ControlRig/Source/ControlRigEditor/Public/ControlRigSequencerEditorLibrary.h
+++ b/Engine/Plugins/Animation/ControlRig/Source/ControlRigEditor/Public/ControlRigSequencerEditorLibrary.h
@@ -48,6 +48,14 @@ struct FControlRigSequencerBindingProxy
 	TObjectPtr<UMovieSceneControlRigParameterTrack> Track;
 };
 
+UENUM(BlueprintType)
+enum class EAnimToolBlendOperation : uint8
+{
+	Tween,
+	BlendToNeighbor,
+	PushPull
+};
+
 /**
 * This is a set of helper functions to access various parts of the Sequencer and Control Rig API via Python and Blueprints.
 */
@@ -147,6 +155,16 @@ public:
 	*/
 	UFUNCTION(BlueprintCallable, Category = "Editor Scripting | Sequencer Tools | Control Rig")
 	static bool TweenControlRig(ULevelSequence* LevelSequence, UControlRig* ControlRig, float TweenValue);
+	
+	/**
+	* Peform specified blend operation based upon selected keys in the curve editor or selected control rig controls
+	* @param LevelSequence The LevelSequence that's loaded in the editor
+	* @param EAnimToolBlendOperation The operation to perform
+	* @param BlendValue The blend value to use, range from -1(blend to previous) to 1(blend to next)
+	*/
+	UFUNCTION(BlueprintCallable, Category = "Editor Scripting | Sequencer Tools | Control Rig")
+	static bool BlendValuesOnSelected(ULevelSequence* LevelSequence, EAnimToolBlendOperation BlendOperation, 
+		float BlendValue);
 
 	/**
 	* Peform a Snap operation to snap the children to the parent. 
@@ -914,4 +932,25 @@ public:
 	/** Hides all of the controls for the given section */
 	UFUNCTION(BlueprintCallable, Category = "Editor Scripting | Sequencer Tools | Control Rig")
 	static void HideAllControls(UMovieSceneSection* InSection);
+
+	/**	Whether or not the control rig is an FK Control Rig.
+	@param InControlRig Rig to test to see if FK Control Rig
+	**/
+	UFUNCTION(BlueprintCallable, Category = "Editor Scripting | Sequencer Tools | Control Rig")
+	static bool IsFKControlRig(UControlRig* InControlRig);
+
+	/**	Get FKControlRig Apply Mode.
+	@param InControlRig Rig to test
+	@return The EControlRigFKRigExecuteMode mode it is in, either Replace,Additive or Direct
+	**/
+	UFUNCTION(BlueprintCallable, Category = "Editor Scripting | Sequencer Tools | Control Rig")
+	static EControlRigFKRigExecuteMode  GetFKControlRigApplyMode(UControlRig* InControlRig);
+
+	/**	Set the FK Control Rig to apply mode
+	@param InControlRig Rig to set 
+	@param InApplyMode Set the EControlRigFKRigExecuteMode mode (Replace,Addtiive or Direct)
+	@return returns True if the mode was set, may not be set if the Control Rig doesn't support these modes currently only FKControlRig's do.
+	**/
+	UFUNCTION(BlueprintCallable, Category = "Editor Scripting | Sequencer Tools | Control Rig")
+	static bool SetControlRigApplyMode(UControlRig* InControlRig, EControlRigFKRigExecuteMode InApplyMode);
 };
