Control Rig: Fix crash when pasting node as function

diff --git a/Engine/Plugins/Animation/ControlRig/Source/ControlRigDeveloper/Private/ControlRigBlueprint.cpp b/Engine/Plugins/Animation/ControlRig/Source/ControlRigDeveloper/Private/ControlRigBlueprint.cpp
--- a/Engine/Plugins/Animation/ControlRig/Source/ControlRigDeveloper/Private/ControlRigBlueprint.cpp
+++ b/Engine/Plugins/Animation/ControlRig/Source/ControlRigDeveloper/Private/ControlRigBlueprint.cpp
@@ -261,7 +261,7 @@ bool UControlRigBlueprint::ExportGraphToText(UEdGraph* InEdGraph, FString& OutTe
 
 bool UControlRigBlueprint::CanImportGraphFromText(const FString& InClipboardText)
 {
-	return GetTemplateController()->CanImportNodesFromText(InClipboardText);
+	return GetTemplateController(true)->CanImportNodesFromText(InClipboardText);
 }
 
 void UControlRigBlueprint::PostEditChangeChainProperty(FPropertyChangedChainEvent& PropertyChangedEvent)
@@ -2313,12 +2313,19 @@ TArray<FString> UControlRigBlueprint::GeneratePythonCommands(const FString InNew
 }
 
 
-URigVMGraph* UControlRigBlueprint::GetTemplateModel()
+URigVMGraph* UControlRigBlueprint::GetTemplateModel(bool bIsFunctionLibrary)
 {
 #if WITH_EDITORONLY_DATA
 	if (TemplateModel == nullptr)
 	{
-		TemplateModel = NewObject<URigVMGraph>(this, TEXT("TemplateModel"));
+		if (bIsFunctionLibrary)
+		{
+			TemplateModel = NewObject<URigVMFunctionLibrary>(this, TEXT("TemplateFunctionLibrary"));
+		}
+		else
+		{
+			TemplateModel = NewObject<URigVMGraph>(this, TEXT("TemplateModel"));
+		}
 		TemplateModel->SetFlags(RF_Transient);
 		TemplateModel->SetExecuteContextStruct(FControlRigExecuteContext::StaticStruct());
 	}
@@ -2328,13 +2335,13 @@ URigVMGraph* UControlRigBlueprint::GetTemplateModel()
 #endif
 }
 
-URigVMController* UControlRigBlueprint::GetTemplateController()
+URigVMController* UControlRigBlueprint::GetTemplateController(bool bIsFunctionLibrary)
 {
 #if WITH_EDITORONLY_DATA
 	if (TemplateController == nullptr)
 	{
 		TemplateController = NewObject<URigVMController>(this, TEXT("TemplateController"));
-		TemplateController->SetGraph(GetTemplateModel());
+		TemplateController->SetGraph(GetTemplateModel(bIsFunctionLibrary));
 		TemplateController->EnableReporting(false);
 		TemplateController->SetFlags(RF_Transient);
 	}
diff --git a/Engine/Plugins/Animation/ControlRig/Source/ControlRigDeveloper/Public/ControlRigBlueprint.h b/Engine/Plugins/Animation/ControlRig/Source/ControlRigDeveloper/Public/ControlRigBlueprint.h
--- a/Engine/Plugins/Animation/ControlRig/Source/ControlRigDeveloper/Public/ControlRigBlueprint.h
+++ b/Engine/Plugins/Animation/ControlRig/Source/ControlRigDeveloper/Public/ControlRigBlueprint.h
@@ -349,8 +349,8 @@ public:
 	UFUNCTION(BlueprintCallable, Category = "Control Rig Blueprint")
 	TArray<FString> GeneratePythonCommands(const FString InNewBlueprintName);
 
-	URigVMGraph* GetTemplateModel();
-	URigVMController* GetTemplateController();
+	URigVMGraph* GetTemplateModel(bool bIsFunctionLibrary = false);
+	URigVMController* GetTemplateController(bool bIsFunctionLibrary = false);
 
 #if WITH_EDITOR
 	UEdGraph* GetEdGraph(URigVMGraph* InModel) const;
diff --git a/Engine/Source/Developer/RigVMDeveloper/Private/RigVMModel/RigVMController.cpp b/Engine/Source/Developer/RigVMDeveloper/Private/RigVMModel/RigVMController.cpp
--- a/Engine/Source/Developer/RigVMDeveloper/Private/RigVMModel/RigVMController.cpp
+++ b/Engine/Source/Developer/RigVMDeveloper/Private/RigVMModel/RigVMController.cpp
@@ -13776,6 +13776,14 @@ bool URigVMController::CanAddNode(URigVMNode* InNode, bool bReportErrors, bool b
 	URigVMGraph* Graph = GetGraph();
 	check(Graph);
 
+	if (Graph->IsA<URigVMFunctionLibrary>())
+	{
+		if (!InNode->IsA<URigVMCollapseNode>())
+		{
+			return false;
+		}
+	}
+
 	if (URigVMFunctionReferenceNode* FunctionRefNode = Cast<URigVMFunctionReferenceNode>(InNode))
 	{
 		if (URigVMFunctionLibrary* FunctionLibrary = FunctionRefNode->GetLibrary())
