Constraints: Finish Constraint Python/BP Api Add functions for adding, keying, compensating, etc..

diff --git a/Engine/Plugins/Animation/ControlRig/Source/ControlRigEditor/Private/ControlRigSequencerEditorLibrary.cpp b/Engine/Plugins/Animation/ControlRig/Source/ControlRigEditor/Private/ControlRigSequencerEditorLibrary.cpp
--- a/Engine/Plugins/Animation/ControlRig/Source/ControlRigEditor/Private/ControlRigSequencerEditorLibrary.cpp
+++ b/Engine/Plugins/Animation/ControlRig/Source/ControlRigEditor/Private/ControlRigSequencerEditorLibrary.cpp
@@ -41,6 +41,12 @@
 #include "Tools/ConstraintBaker.h"
 #include "TransformConstraint.h"
 #include "Rigs/FKControlRig.h"
+#include "Subsystems/UnrealEditorSubsystem.h"
+#include "EditorScriptingHelpers.h"
+#include "ConstraintsScripting.h"
+#include "Constraints/ControlRigTransformableHandle.h"
+#include "ConstraintChannelHelper.h"
+#include "Sections/MovieSceneConstrainedSection.h"
 
 #include UE_INLINE_GENERATED_CPP_BY_NAME(ControlRigSequencerEditorLibrary)
 
@@ -440,6 +446,193 @@ bool UControlRigSequencerEditorLibrary::BlendValuesOnSelected(ULevelSequence* Le
 	return false;
 }
 
+UTickableConstraint* UControlRigSequencerEditorLibrary::AddConstraint(UWorld* World, ETransformConstraintType InType, UTransformableHandle* InChildHandle, UTransformableHandle* InParentHandle, const bool bMaintainOffset)
+{
+	if (!World)
+	{
+		UE_LOG(LogControlRig, Error, TEXT("AddConstraint: Need Valid World"));
+		return nullptr;
+	}
+
+	UTickableTransformConstraint* Constraint = FTransformConstraintUtils::CreateFromType(World, InType);
+	if (!Constraint)
+	{
+		UE_LOG(LogControlRig, Error, TEXT("AddConstraint: Can not create constraint from type"));
+		return nullptr;
+	}
+
+	if(FTransformConstraintUtils::AddConstraint(World, InParentHandle, InChildHandle, Constraint, bMaintainOffset) == false)
+	{
+		UE_LOG(LogControlRig, Error, TEXT("AddConstraint: Constraint not added"));
+		Constraint->MarkAsGarbage();
+		return nullptr;
+	}
+	
+	//add key
+	FConstraintChannelHelper::SmartConstraintKey(Constraint, TOptional<bool>(), TOptional<FFrameNumber>());
+	return Constraint;
+}
+
+TArray <TObjectPtr<UTickableConstraint>> UControlRigSequencerEditorLibrary::GetConstraintsForHandle(UWorld* InWorld, const UTransformableHandle* InChild)
+{
+	TArray <TObjectPtr<UTickableConstraint>> Constraints;
+	const FConstraintsManagerController& Controller = FConstraintsManagerController::Get(InWorld);
+	const TArray<TObjectPtr<UTickableConstraint>>& AllConstraints = Controller.GetAllConstraints(false);
+	for (const TObjectPtr<UTickableConstraint>& TickConstraint : AllConstraints)
+	{
+		if (TObjectPtr<UTickableTransformConstraint> Constraint = Cast<UTickableTransformConstraint>(TickConstraint))
+		{
+			if (Constraint->ChildTRSHandle == InChild)
+			{
+				Constraints.Add(Constraint);
+			}
+		}
+	}
+	return Constraints;
+}
+
+bool UControlRigSequencerEditorLibrary::Compensate(UTickableConstraint* InConstraint, const FFrameNumber& FrameTime)
+{
+	if (UTickableTransformConstraint* Constraint = Cast<UTickableTransformConstraint>(InConstraint))
+	{
+		TOptional<FFrameNumber> OptTime(FrameTime);
+		FConstraintChannelHelper::Compensate(Constraint, OptTime);
+		return true;
+	}
+	return false;
+}
+
+bool UControlRigSequencerEditorLibrary::CompensateAll(UTickableConstraint* InConstraint)
+{
+	if (UTickableTransformConstraint* Constraint = Cast<UTickableTransformConstraint>(InConstraint))
+	{
+		FConstraintChannelHelper::Compensate(Constraint, TOptional<FFrameNumber>());
+		return true;
+	}
+	return false;
+}
+
+bool UControlRigSequencerEditorLibrary::SetConstraintActiveKey(UTickableConstraint* InConstraint, bool bActive, const FFrameNumber& FrameTime)
+{
+	if (UTickableTransformConstraint* Constraint = Cast<UTickableTransformConstraint>(InConstraint))
+	{
+		return FConstraintChannelHelper::SmartConstraintKey(Constraint, TOptional<bool>(bActive), TOptional<FFrameNumber>(FrameTime));
+	}
+	return false;
+}
+
+bool UControlRigSequencerEditorLibrary::MoveConstraintKey(UTickableConstraint* Constraint, UMovieSceneSection* ConstraintSection, FFrameNumber InTime, FFrameNumber InNewTime, ESequenceTimeUnit TimeUnit)
+{
+	TWeakPtr<ISequencer> WeakSequencer = GetSequencerFromAsset();
+	if (WeakSequencer.IsValid() == false)
+	{
+		UE_LOG(LogControlRig, Error, TEXT("MoveConstraintKey: Need open Sequencer"));
+		return false;
+	}
+	IMovieSceneConstrainedSection* ConstrainedSection = Cast<IMovieSceneConstrainedSection>(ConstraintSection);
+	if (ConstrainedSection == nullptr)
+	{
+		UE_LOG(LogControlRig, Error, TEXT("MoveConstraintKey: Section doesn't support constraints"));
+		return false;
+	}
+	if (Constraint == nullptr)
+	{
+		UE_LOG(LogControlRig, Error, TEXT("MoveConstraintKey: Constraint not valid"));
+		return false;
+	}
+	FConstraintAndActiveChannel* ConstraintAndChannel = ConstrainedSection->GetConstraintChannel(Constraint->GetFName());
+	if (ConstraintAndChannel == nullptr)
+	{
+		UE_LOG(LogControlRig, Error, TEXT("MoveConstraintKey: Constraint not found in section"));
+		return false;
+	}
+	TSharedPtr<ISequencer>  Sequencer = WeakSequencer.Pin();
+	FFrameRate TickResolution = Sequencer->GetFocusedTickResolution();
+	FFrameRate DisplayRate = Sequencer->GetFocusedDisplayRate();
+
+	if (TimeUnit == ESequenceTimeUnit::DisplayRate)
+	{
+		InTime = FFrameRate::TransformTime(FFrameTime(InTime, 0), DisplayRate, TickResolution).RoundToFrame();
+		InNewTime = FFrameRate::TransformTime(FFrameTime(InNewTime, 0), DisplayRate, TickResolution).RoundToFrame();
+	}
+
+	TArray<FFrameNumber> OurKeyTimes;
+	TArray<FKeyHandle> OurKeyHandles;
+	TRange<FFrameNumber> CurrentFrameRange;
+	CurrentFrameRange.SetLowerBound(TRangeBound<FFrameNumber>(InTime));
+	CurrentFrameRange.SetUpperBound(TRangeBound<FFrameNumber>(InTime));
+	TMovieSceneChannelData<bool> ChannelInterface = ConstraintAndChannel->ActiveChannel.GetData();
+	ChannelInterface.GetKeys(CurrentFrameRange, &OurKeyTimes, &OurKeyHandles);
+	if (OurKeyHandles.Num() > 0)
+	{
+		ConstraintSection->TryModify();
+		TArray<FFrameNumber> NewKeyTimes;
+		NewKeyTimes.Add(InNewTime);
+		ConstraintAndChannel->ActiveChannel.SetKeyTimes(OurKeyHandles, NewKeyTimes);
+		WeakSequencer.Pin()->NotifyMovieSceneDataChanged(EMovieSceneDataChangeType::TrackValueChanged);
+	}
+	else
+	{
+		UE_LOG(LogControlRig, Error, TEXT("MoveConstraintKey: No Key To Move At that Time"));
+		return false;
+	}
+	return true;
+}
+
+bool UControlRigSequencerEditorLibrary::DeleteConstraintKey(UTickableConstraint* Constraint, UMovieSceneSection* ConstraintSection, FFrameNumber InTime, ESequenceTimeUnit TimeUnit)
+{
+	TWeakPtr<ISequencer> WeakSequencer = GetSequencerFromAsset();
+	if (WeakSequencer.IsValid() == false)
+	{
+		UE_LOG(LogControlRig, Error, TEXT("DeleteConstraintKey: Need open Sequencer"));
+		return false;
+	}
+	IMovieSceneConstrainedSection* ConstrainedSection = Cast<IMovieSceneConstrainedSection>(ConstraintSection);
+	if (ConstrainedSection == nullptr)
+	{
+		UE_LOG(LogControlRig, Error, TEXT("DeleteConstraintKey: Section doesn't support constraints"));
+		return false;
+	}
+	if (Constraint == nullptr)
+	{
+		UE_LOG(LogControlRig, Error, TEXT("DeleteConstraintKey: Constraint not valid"));
+		return false;
+	}
+	FConstraintAndActiveChannel* ConstraintAndChannel = ConstrainedSection->GetConstraintChannel(Constraint->GetFName());
+	if (ConstraintAndChannel == nullptr)
+	{
+		UE_LOG(LogControlRig, Error, TEXT("DeleteConstraintKey: Constraint not found in section"));
+		return false;
+	}
+	TSharedPtr<ISequencer>  Sequencer = WeakSequencer.Pin();
+	FFrameRate TickResolution = Sequencer->GetFocusedTickResolution();
+	FFrameRate DisplayRate = Sequencer->GetFocusedDisplayRate();
+
+	if (TimeUnit == ESequenceTimeUnit::DisplayRate)
+	{
+		InTime = FFrameRate::TransformTime(FFrameTime(InTime, 0), DisplayRate, TickResolution).RoundToFrame();
+	}
+
+	TArray<FFrameNumber> OurKeyTimes;
+	TArray<FKeyHandle> OurKeyHandles;
+	TRange<FFrameNumber> CurrentFrameRange;
+	CurrentFrameRange.SetLowerBound(TRangeBound<FFrameNumber>(InTime));
+	CurrentFrameRange.SetUpperBound(TRangeBound<FFrameNumber>(InTime));
+	TMovieSceneChannelData<bool> ChannelInterface = ConstraintAndChannel->ActiveChannel.GetData();
+	ChannelInterface.GetKeys(CurrentFrameRange, &OurKeyTimes, &OurKeyHandles);
+	if (OurKeyHandles.Num() > 0)
+	{
+		ConstraintSection->TryModify();
+		ConstraintAndChannel->ActiveChannel.DeleteKeys(OurKeyHandles);
+		WeakSequencer.Pin()->NotifyMovieSceneDataChanged(EMovieSceneDataChangeType::TrackValueChanged);
+	}
+	else
+	{
+		UE_LOG(LogControlRig, Error, TEXT("DeleteConstraintKey: No Key To Delete At that Time"));
+		return false;
+	}
+	return true;
+}
 
 bool UControlRigSequencerEditorLibrary::BakeConstraint(UWorld* World, UTickableConstraint* Constraint, const TArray<FFrameNumber>& Frames, ESequenceTimeUnit TimeUnit)
 {
diff --git a/Engine/Plugins/Animation/ControlRig/Source/ControlRigEditor/Private/Editor/Constraints/SConstraintsWidget.cpp b/Engine/Plugins/Animation/ControlRig/Source/ControlRigEditor/Private/Editor/Constraints/SConstraintsWidget.cpp
--- a/Engine/Plugins/Animation/ControlRig/Source/ControlRigEditor/Private/Editor/Constraints/SConstraintsWidget.cpp
+++ b/Engine/Plugins/Animation/ControlRig/Source/ControlRigEditor/Private/Editor/Constraints/SConstraintsWidget.cpp
@@ -323,7 +323,7 @@ void SDroppableConstraintItem::CreateConstraint(
 					FTransformConstraintUtils::CreateAndAddFromActors(World, InParent, InSocketName, Child, InConstraintType);
 				if (Constraint)
 				{
-					FConstraintChannelHelper::SmartConstraintKey(Constraint);
+					FConstraintChannelHelper::SmartConstraintKey(Constraint, TOptional<bool>(), TOptional<FFrameNumber>());
 					bCreated = true;
 				}
 			}
@@ -557,7 +557,7 @@ void SEditableConstraintItem::Construct(
 				if (UTickableTransformConstraint* TransformConstraint = Cast<UTickableTransformConstraint>(Constraint))
 				{
 					FScopedTransaction Transaction(LOCTEXT("CreateConstraintKey", "Create Constraint Key"));
-					FConstraintChannelHelper::SmartConstraintKey(TransformConstraint);
+					FConstraintChannelHelper::SmartConstraintKey(TransformConstraint,TOptional<bool>(), TOptional<FFrameNumber>());
 				}
 				return FReply::Handled();
 			})
@@ -1047,7 +1047,15 @@ TSharedPtr<SWidget> SConstraintsEditionWidget::CreateContextMenu()
 				FSlateIcon(),
 				FUIAction(FExecuteAction::CreateLambda([TransformConstraint]()
 				{
-					FConstraintChannelHelper::Compensate(TransformConstraint);
+					const TWeakPtr<ISequencer> WeakSequencer = FBakingHelper::GetSequencer();
+					if (WeakSequencer.IsValid())
+					{
+						const TSharedPtr<ISequencer> Sequencer = WeakSequencer.Pin();
+						const FFrameRate TickResolution = Sequencer->GetFocusedTickResolution();
+						const FFrameTime FrameTime = Sequencer->GetLocalTime().ConvertTo(TickResolution);
+						const FFrameNumber Time = FrameTime.GetFrame();
+						FConstraintChannelHelper::Compensate(TransformConstraint, TOptional<FFrameNumber>(Time));
+					}
 				}), IsCompensationEnabled),
 				NAME_None,
 				EUserInterfaceActionType::Button);
@@ -1058,7 +1066,7 @@ TSharedPtr<SWidget> SConstraintsEditionWidget::CreateContextMenu()
 				FSlateIcon(),
 				FUIAction(FExecuteAction::CreateLambda([TransformConstraint]()
 				{
-					FConstraintChannelHelper::Compensate(TransformConstraint, true);
+					FConstraintChannelHelper::Compensate(TransformConstraint, TOptional<FFrameNumber>());
 				}), IsCompensationEnabled),
 				NAME_None,
 				EUserInterfaceActionType::Button);
diff --git a/Engine/Plugins/Animation/ControlRig/Source/ControlRigEditor/Private/Sequencer/ConstraintChannelHelper.cpp b/Engine/Plugins/Animation/ControlRig/Source/ControlRigEditor/Private/Sequencer/ConstraintChannelHelper.cpp
--- a/Engine/Plugins/Animation/ControlRig/Source/ControlRigEditor/Private/Sequencer/ConstraintChannelHelper.cpp
+++ b/Engine/Plugins/Animation/ControlRig/Source/ControlRigEditor/Private/Sequencer/ConstraintChannelHelper.cpp
@@ -73,26 +73,48 @@ bool FConstraintChannelHelper::IsKeyframingAvailable()
 	return true;
 }
 
-void FConstraintChannelHelper::SmartConstraintKey(UTickableTransformConstraint* InConstraint)
+bool FConstraintChannelHelper::SmartConstraintKey(UTickableTransformConstraint* InConstraint, 
+	const TOptional<bool>& bActive, 
+	const TOptional<FFrameNumber>& OptFrameTime)
 {
 	const TWeakPtr<ISequencer> WeakSequencer = FBakingHelper::GetSequencer();
-	if (!WeakSequencer.IsValid() || !WeakSequencer.Pin()->GetFocusedMovieSceneSequence())
+	if (!WeakSequencer.IsValid() )
 	{
-		return;
+		return false;
+	}
+	TSharedPtr<ISequencer> SequencerPtr = WeakSequencer.Pin();
+	if (SequencerPtr->GetFocusedMovieSceneSequence() == nullptr)
+	{
+		return false;
+	}
+	FFrameNumber Time;
+	if (OptFrameTime.IsSet())
+	{
+		Time = OptFrameTime.GetValue();
+	}
+	else
+	{
+		const FFrameRate TickResolution = SequencerPtr->GetFocusedTickResolution();
+		const FFrameTime FrameTime = SequencerPtr->GetLocalTime().ConvertTo(TickResolution);
+		Time = FrameTime.GetFrame();
 	}
 
+	bool bSuceeded = false;
 	if (const UTransformableComponentHandle* ComponentHandle = Cast<UTransformableComponentHandle>(InConstraint->ChildTRSHandle))
 	{
-		SmartComponentConstraintKey(InConstraint, WeakSequencer.Pin());
+		bSuceeded = SmartComponentConstraintKey(InConstraint,bActive,Time, SequencerPtr);
 	}
 	
 	else if (const UTransformableControlHandle* ControlHandle = Cast<UTransformableControlHandle>(InConstraint->ChildTRSHandle))
 	{
-		SmartControlConstraintKey(InConstraint, WeakSequencer.Pin());
+		bSuceeded = SmartControlConstraintKey(InConstraint,bActive,Time, SequencerPtr);
 	}
 
+	//todo need to revisit this to see if we need to create this even if we don't set a key, it's harmless I think either way
 	FConstraintChannelHelper::CreateBindingIDForHandle(InConstraint->ChildTRSHandle);
 	FConstraintChannelHelper::CreateBindingIDForHandle(InConstraint->ParentTRSHandle);
+
+	return bSuceeded;
 }
 
 
@@ -171,14 +193,16 @@ UMovieScene3DTransformSection* FConstraintChannelHelper::GetTransformSection(
 	return MovieSceneToolHelpers::GetTransformSection(InSequencer.Get(), Guid);
 }
 
-void FConstraintChannelHelper::SmartControlConstraintKey(
+bool FConstraintChannelHelper::SmartControlConstraintKey(
 	UTickableTransformConstraint* InConstraint,
+	const TOptional<bool>& bActive,
+	const FFrameNumber& Time,
 	const TSharedPtr<ISequencer>& InSequencer)
 {
 	const UTransformableControlHandle* ControlHandle = Cast<UTransformableControlHandle>(InConstraint->ChildTRSHandle);
 	if (!ControlHandle)
 	{
-		return;
+		return false;
 	}
 	
 	if (UMovieSceneControlRigParameterSection* Section = GetControlSection(ControlHandle, InSequencer))
@@ -195,13 +219,13 @@ void FConstraintChannelHelper::SmartControlConstraintKey(
 		// add key if needed
 		if (FConstraintAndActiveChannel* Channel = Section->GetConstraintChannel(InConstraint->GetFName()))
 		{
-			const FFrameRate TickResolution = InSequencer->GetFocusedTickResolution();
-			const FFrameTime FrameTime = InSequencer->GetLocalTime().ConvertTo(TickResolution);
-			const FFrameNumber Time = FrameTime.GetFrame();
 
 			bool ActiveValueToBeSet = false;
-			if (CanAddKey(Channel->ActiveChannel, Time, ActiveValueToBeSet))
+			//add key if we can and make sure the key we are setting is what we want
+			if (CanAddKey(Channel->ActiveChannel, Time, ActiveValueToBeSet) && (bActive.IsSet() == false || bActive.GetValue() == ActiveValueToBeSet))
 			{
+				const FFrameRate TickResolution = InSequencer->GetFocusedTickResolution();
+
 				const bool bNeedsCompensation = InConstraint->NeedsCompensation();
 				
 				TGuardValue<bool> CompensateGuard(FMovieSceneConstraintChannelHelper::bDoNotCompensate, true);
@@ -283,31 +307,35 @@ void FConstraintChannelHelper::SmartControlConstraintKey(
 						SetTangentsAtThisTime<FMovieSceneFloatChannel>(ChannelIndex, NumChannels, Section, Time, Tangents);
 					}
 				}
+				return true;
 			}
 		}
 	}
+	return false;
 }
 
-void FConstraintChannelHelper::SmartComponentConstraintKey(
+bool FConstraintChannelHelper::SmartComponentConstraintKey(
 	UTickableTransformConstraint* InConstraint,
+	const TOptional<bool>& bActive,
+	const FFrameNumber& Time,
 	const TSharedPtr<ISequencer>& InSequencer)
 {
 	const UTransformableComponentHandle* ComponentHandle = Cast<UTransformableComponentHandle>(InConstraint->ChildTRSHandle);
 	if (!ComponentHandle)
 	{
-		return;
+		return false;
 	}
 	AActor* Actor = ComponentHandle->Component->GetOwner();
 	if (!Actor)
 	{
-		return;
+		return false;
 	}
 
 	const FTransform LocalTransform = ComponentHandle->GetLocalTransform();
 	const FGuid Guid = InSequencer->GetHandleToObject(Actor, true);
 	if (!Guid.IsValid())
 	{
-		return;
+		return false;
 	}
 
 	if (UMovieScene3DTransformSection* Section = MovieSceneToolHelpers::GetTransformSection(InSequencer.Get(), Guid, LocalTransform))
@@ -324,13 +352,11 @@ void FConstraintChannelHelper::SmartComponentConstraintKey(
 		// add key if needed
 		if (FConstraintAndActiveChannel* Channel = Section->GetConstraintChannel(InConstraint->GetFName()))
 		{
-			const FFrameRate TickResolution = InSequencer->GetFocusedTickResolution();
-			const FFrameTime FrameTime = InSequencer->GetLocalTime().ConvertTo(TickResolution);
-			const FFrameNumber Time = FrameTime.GetFrame();
-
 			bool ActiveValueToBeSet = false;
-			if (CanAddKey(Channel->ActiveChannel, Time, ActiveValueToBeSet))
+			//add key if we can and make sure the key we are setting is what we want
+			if (CanAddKey(Channel->ActiveChannel, Time, ActiveValueToBeSet) && (bActive.IsSet() == false || bActive.GetValue() == ActiveValueToBeSet))
 			{
+				const FFrameRate TickResolution = InSequencer->GetFocusedTickResolution();
 				const bool bNeedsCompensation = InConstraint->NeedsCompensation();
 				
 				Section->Modify();
@@ -421,12 +447,14 @@ void FConstraintChannelHelper::SmartComponentConstraintKey(
 				{
 					InSequencer->SetAutoChangeMode(AutoChangeMode);
 				}
+				return true;
 			}
 		}
 	}
+	return false;
 }
 
-void FConstraintChannelHelper::Compensate(UTickableTransformConstraint* InConstraint, const bool bAllTimes)
+void FConstraintChannelHelper::Compensate(UTickableTransformConstraint* InConstraint, const TOptional<FFrameNumber>& OptTime)
 {
 	const TWeakPtr<ISequencer> WeakSequencer = FBakingHelper::GetSequencer();
 	if (!WeakSequencer.IsValid() || !WeakSequencer.Pin()->GetFocusedMovieSceneSequence())
@@ -473,11 +501,6 @@ void FConstraintChannelHelper::Compensate(UTickableTransformConstraint* InConstr
 		}
 	}
 
-	const FFrameRate TickResolution = Sequencer->GetFocusedTickResolution();
-	const FFrameTime FrameTime = Sequencer->GetLocalTime().ConvertTo(TickResolution);
-	const FFrameNumber Time = FrameTime.GetFrame();
-
-	const TOptional<FFrameNumber> OptTime = bAllTimes ? TOptional<FFrameNumber>() : TOptional<FFrameNumber>(Time);
 	CompensateIfNeeded(World, Sequencer, Section, OptTime);
 }
 
diff --git a/Engine/Plugins/Animation/ControlRig/Source/ControlRigEditor/Public/ConstraintChannelHelper.h b/Engine/Plugins/Animation/ControlRig/Source/ControlRigEditor/Public/ConstraintChannelHelper.h
--- a/Engine/Plugins/Animation/ControlRig/Source/ControlRigEditor/Public/ConstraintChannelHelper.h
+++ b/Engine/Plugins/Animation/ControlRig/Source/ControlRigEditor/Public/ConstraintChannelHelper.h
@@ -31,14 +31,17 @@ public:
 	/** @todo documentation. */
 	static bool IsKeyframingAvailable();
 	
-	/** Adds an active key if needed and does the compensation when switching. */
-	static void SmartConstraintKey(UTickableTransformConstraint* InConstraint);
+	/** Adds an active key if needed and does the compensation when switching. Will use the optional active and time if set. 
+	Will return true if key is actually set, may not be if the value is the same.*/
+	static bool SmartConstraintKey(UTickableTransformConstraint* InConstraint, 
+		const TOptional<bool>&  bActive ,
+		const TOptional<FFrameNumber>& FrameTime);
 
 
 	/** Compensate transform on handles when a constraint switches state. */
 	static void Compensate(
 		UTickableTransformConstraint* InConstraint,
-		const bool bAllTimes = false);
+		const TOptional<FFrameNumber>& OptionalTime);
 	static void CompensateIfNeeded(
 		UWorld* InWorld,
 		const TSharedPtr<ISequencer>& InSequencer,
@@ -55,8 +58,10 @@ private:
 	/** BEGIN CONTROL RIG SECTION */
 	
 	/** @todo documentation. */
-	static void SmartControlConstraintKey(
+	static bool SmartControlConstraintKey(
 		UTickableTransformConstraint* InConstraint,
+		const TOptional<bool>& bActive,
+		const FFrameNumber& FrameTime,
 		const TSharedPtr<ISequencer>& InSequencer);
 	
 	/** END CONTROL RIG SECTION */
@@ -64,8 +69,10 @@ private:
 	/** BEGIN COMPONENT SECTION */
 
 	/** @todo documentation. */
-	static void SmartComponentConstraintKey(
+	static bool SmartComponentConstraintKey(
 		UTickableTransformConstraint* InConstraint,
+		const TOptional<bool>&  bActive, 
+		const FFrameNumber& FrameTime,
 		const TSharedPtr<ISequencer>& InSequencer);
 
 	/** @todo documentation. */
diff --git a/Engine/Plugins/Animation/ControlRig/Source/ControlRigEditor/Public/ControlRigSequencerEditorLibrary.h b/Engine/Plugins/Animation/ControlRig/Source/ControlRigEditor/Public/ControlRigSequencerEditorLibrary.h
--- a/Engine/Plugins/Animation/ControlRig/Source/ControlRigEditor/Public/ControlRigSequencerEditorLibrary.h
+++ b/Engine/Plugins/Animation/ControlRig/Source/ControlRigEditor/Public/ControlRigSequencerEditorLibrary.h
@@ -21,6 +21,8 @@
 
 class ULevelSequence;
 class UTickableConstraint;
+class UTickableTransformConstraint;
+class UTransformableHandle;
 
 USTRUCT(BlueprintType)
 struct FControlRigSequencerBindingProxy
@@ -142,11 +144,78 @@ public:
 	* @param Constraint The Constraint to bake. After baking it will be keyed to be inactive of the range of frames that are baked
 	* @param Frames The frames to bake, if the array is empty it will use the active time ranges of the constraint to deteremine where it should bake
 	* @param TimeUnit Unit for all frame and time values, either in display rate or tick resolution
-	* @return returns True if successful, False otherwise
+	* @return Returns True if successful, False otherwise
 	*/
 	UFUNCTION(BlueprintCallable, Category = "Editor Scripting | Sequencer Tools | Control Rig")
 	static bool BakeConstraint(UWorld* World, UTickableConstraint* Constraint, const TArray<FFrameNumber>& Frames, ESequenceTimeUnit TimeUnit = ESequenceTimeUnit::DisplayRate);
 
+	/**
+	* Add a constraint possibly adding to sequencer also if one is open.
+	* @param World The active world
+	* @param InType Type of constraint to create
+	* @param InChild The handle to the transormable to be constrainted
+	* @param InParent The handle to the parent of the constraint
+	* @param bMaintainOffset Whether to maintain offset between child and parent when setting the constraint
+	* @return Returns the constraint if created all nullptr if not
+	*/
+	UFUNCTION(BlueprintCallable, Category = "Editor Scripting | Sequencer Tools | Control Rig")
+	static UTickableConstraint* AddConstraint(UWorld* World, ETransformConstraintType InType, UTransformableHandle* InChild, UTransformableHandle* InParent, const bool bMaintainOffset);
+	
+	/**
+	* Set the constraint active key in the current open Sequencer
+	* @param InConstraint The constraint to set the key
+	* @param bActive Whether or not it's active
+	* @param FrameTime Time to set the value 
+	* @return Returns true if we set the constraint to be the passed in value, false if not. We may not do so if the value is the same.
+	*/
+	UFUNCTION(BlueprintCallable, Category = "Editor Scripting | Sequencer Tools | Control Rig")
+	static bool SetConstraintActiveKey(UTickableConstraint* InConstraint, bool bActive,const FFrameNumber& FrameTime);
+	
+	/**
+	* Get all constraints for this object, which is described by a transformable handle
+	* @param InChild The handle to look for constraints controlling it
+	* @return Returns array of Constraints this handle is constrained to.
+	*/
+	UFUNCTION(BlueprintCallable, Category = "Editor Scripting | Sequencer Tools | Control Rig")
+	static TArray <TObjectPtr<UTickableConstraint>> GetConstraintsForHandle(UWorld* InWorld, const UTransformableHandle* InChild);
+
+	/** Move the constraint active key in the current open Sequencer
+	* @param InConstraint The constraint whose key to move
+	* @param ConstraintSection Section containing Cosntraint Key
+	* @param InTime Original time of the constraint key
+	* @param InNewTime New time for the constraint key
+	* @param TimeUnit Unit for the time params, either in display rate or tick resolution
+	* @return Will return false if function fails, for example if there is no key at this time it will fail, or if the new time is invalid it could fail also
+	*/
+	UFUNCTION(BlueprintCallable, Category = "Editor Scripting | Sequencer Tools | Control Rig")
+	bool MoveConstraintKey(UTickableConstraint* Constraint, UMovieSceneSection* ConstraintSection, FFrameNumber InTime, FFrameNumber InNewTime, ESequenceTimeUnit TimeUnit);
+
+	/** Delete the Key for the Constraint at the specified time. 
+	* @param InConstraint The constraint whose key to move
+	* @param ConstraintSection Section containing Cosntraint Key
+	* @param InTime Time to delete the constraint.
+	* @param TimeUnit Unit for the InTime, either in display rate or tick resolution
+	* @return Will return false if function fails,  for example if there is no key at this time it will fail.
+	*/
+	UFUNCTION(BlueprintCallable, Category = "Editor Scripting | Sequencer Tools | Control Rig")
+	static bool DeleteConstraintKey(UTickableConstraint* Constraint, UMovieSceneSection* ConstraintSection, FFrameNumber InTime, ESequenceTimeUnit TimeUnit = ESequenceTimeUnit::DisplayRate);
+	/**
+	* Compensate constraint at the specfied time 
+	* @param InConstraint The constraint to compensate
+	* @param FrameTime  Time to compensate
+	* @return Returns true if it can compensate
+	*/
+	UFUNCTION(BlueprintCallable, Category = "Editor Scripting | Sequencer Tools | Control Rig")
+	static bool Compensate(UTickableConstraint* InConstraint, const FFrameNumber& FrameTime);
+
+	/**
+	* Compensate constraint at all keys
+	* @param InConstraint The constraint to compensate
+	* @return Returns true if it can compensate
+	*/
+	UFUNCTION(BlueprintCallable, Category = "Editor Scripting | Sequencer Tools | Control Rig")
+	static bool CompensateAll(UTickableConstraint* InConstraint);
+
 	/**
 	* Peform a Tween operation on the current active sequencer time(must be visible).
 	* @param LevelSequence The LevelSequence that's loaded in the editor
@@ -905,8 +974,6 @@ public:
 	UFUNCTION(BlueprintCallable, Category = "Editor Scripting | Sequencer Tools | Control Rig")
 	static bool MoveControlRigSpace(ULevelSequence* InSequence, UControlRig* InControlRig, FName InControlName, FFrameNumber InTime, FFrameNumber InNewTime, ESequenceTimeUnit TimeUnit = ESequenceTimeUnit::DisplayRate);
 
-
-
 	/** Rename the Control Rig Channels in Sequencer to the specified new control names, which should be present on the Control Rig
 	* @param InSequence Sequence to rename controls
 	* @param InControlRig ControlRig to rename controls
diff --git a/Engine/Source/Runtime/Experimental/Animation/Constraints/Private/ConstraintsScripting.cpp b/Engine/Source/Runtime/Experimental/Animation/Constraints/Private/ConstraintsScripting.cpp
--- a/Engine/Source/Runtime/Experimental/Animation/Constraints/Private/ConstraintsScripting.cpp
+++ b/Engine/Source/Runtime/Experimental/Animation/Constraints/Private/ConstraintsScripting.cpp
@@ -43,9 +43,9 @@ bool UConstraintsScriptingLibrary::AddConstraint(UWorld* InWorld, UTransformable
 	return Val;
 }
 
-TArray<TWeakObjectPtr<UTickableConstraint>> UConstraintsScriptingLibrary::GetConstraintsArray(UWorld* InWorld)
+TArray<TObjectPtr<UTickableConstraint>> UConstraintsScriptingLibrary::GetConstraintsArray(UWorld* InWorld)
 {
-	TArray<TWeakObjectPtr<UTickableConstraint>> Constraints;
+	TArray<TObjectPtr<UTickableConstraint>> Constraints;
 	const FConstraintsManagerController& Controller = FConstraintsManagerController::Get(InWorld);
 	const TArray< TObjectPtr<UTickableConstraint> >& ConstraintsArray = Controller.GetConstraintsArray();
 	for (const TObjectPtr<UTickableConstraint>& Constraint : ConstraintsArray)
diff --git a/Engine/Source/Runtime/Experimental/Animation/Constraints/Public/ConstraintsScripting.h b/Engine/Source/Runtime/Experimental/Animation/Constraints/Public/ConstraintsScripting.h
--- a/Engine/Source/Runtime/Experimental/Animation/Constraints/Public/ConstraintsScripting.h
+++ b/Engine/Source/Runtime/Experimental/Animation/Constraints/Public/ConstraintsScripting.h
@@ -71,7 +71,8 @@ public:
 	@param InWorld World we are in
 	@return Copy of the constraints in the level
 	*/
-	static TArray<TWeakObjectPtr<UTickableConstraint>> GetConstraintsArray(UWorld* InWorld);
+	UFUNCTION(BlueprintCallable, Category = "Editor Scripting | Sequencer Tools | Control Rig | Constraints")
+	static TArray<TObjectPtr<UTickableConstraint>> GetConstraintsArray(UWorld* InWorld);
 
 	/**
 	* Remove constraint at specified index
