CIS TObjectPtr<> in UFUNCTION fix

diff --git a/Engine/Plugins/Animation/ControlRig/Source/ControlRigEditor/Private/ControlRigSequencerEditorLibrary.cpp b/Engine/Plugins/Animation/ControlRig/Source/ControlRigEditor/Private/ControlRigSequencerEditorLibrary.cpp
--- a/Engine/Plugins/Animation/ControlRig/Source/ControlRigEditor/Private/ControlRigSequencerEditorLibrary.cpp
+++ b/Engine/Plugins/Animation/ControlRig/Source/ControlRigEditor/Private/ControlRigSequencerEditorLibrary.cpp
@@ -473,9 +473,9 @@ UTickableConstraint* UControlRigSequencerEditorLibrary::AddConstraint(UWorld* Wo
 	return Constraint;
 }
 
-TArray <TObjectPtr<UTickableConstraint>> UControlRigSequencerEditorLibrary::GetConstraintsForHandle(UWorld* InWorld, const UTransformableHandle* InChild)
+TArray <UTickableConstraint*> UControlRigSequencerEditorLibrary::GetConstraintsForHandle(UWorld* InWorld, const UTransformableHandle* InChild)
 {
-	TArray <TObjectPtr<UTickableConstraint>> Constraints;
+	TArray <UTickableConstraint*> Constraints;
 	const FConstraintsManagerController& Controller = FConstraintsManagerController::Get(InWorld);
 	const TArray<TObjectPtr<UTickableConstraint>>& AllConstraints = Controller.GetAllConstraints(false);
 	for (const TObjectPtr<UTickableConstraint>& TickConstraint : AllConstraints)
@@ -484,7 +484,7 @@ TArray <TObjectPtr<UTickableConstraint>> UControlRigSequencerEditorLibrary::GetC
 		{
 			if (Constraint->ChildTRSHandle == InChild)
 			{
-				Constraints.Add(Constraint);
+				Constraints.Add(Constraint.Get());
 			}
 		}
 	}
diff --git a/Engine/Plugins/Animation/ControlRig/Source/ControlRigEditor/Public/ControlRigSequencerEditorLibrary.h b/Engine/Plugins/Animation/ControlRig/Source/ControlRigEditor/Public/ControlRigSequencerEditorLibrary.h
--- a/Engine/Plugins/Animation/ControlRig/Source/ControlRigEditor/Public/ControlRigSequencerEditorLibrary.h
+++ b/Engine/Plugins/Animation/ControlRig/Source/ControlRigEditor/Public/ControlRigSequencerEditorLibrary.h
@@ -177,7 +177,7 @@ public:
 	* @return Returns array of Constraints this handle is constrained to.
 	*/
 	UFUNCTION(BlueprintCallable, Category = "Editor Scripting | Sequencer Tools | Control Rig")
-	static TArray <TObjectPtr<UTickableConstraint>> GetConstraintsForHandle(UWorld* InWorld, const UTransformableHandle* InChild);
+	static TArray <UTickableConstraint*> GetConstraintsForHandle(UWorld* InWorld, const UTransformableHandle* InChild);
 
 	/** Move the constraint active key in the current open Sequencer
 	* @param InConstraint The constraint whose key to move
diff --git a/Engine/Source/Runtime/Experimental/Animation/Constraints/Private/ConstraintsScripting.cpp b/Engine/Source/Runtime/Experimental/Animation/Constraints/Private/ConstraintsScripting.cpp
--- a/Engine/Source/Runtime/Experimental/Animation/Constraints/Private/ConstraintsScripting.cpp
+++ b/Engine/Source/Runtime/Experimental/Animation/Constraints/Private/ConstraintsScripting.cpp
@@ -43,9 +43,9 @@ bool UConstraintsScriptingLibrary::AddConstraint(UWorld* InWorld, UTransformable
 	return Val;
 }
 
-TArray<TObjectPtr<UTickableConstraint>> UConstraintsScriptingLibrary::GetConstraintsArray(UWorld* InWorld)
+TArray<UTickableConstraint*> UConstraintsScriptingLibrary::GetConstraintsArray(UWorld* InWorld)
 {
-	TArray<TObjectPtr<UTickableConstraint>> Constraints;
+	TArray<UTickableConstraint*> Constraints;
 	const FConstraintsManagerController& Controller = FConstraintsManagerController::Get(InWorld);
 	const TArray< TObjectPtr<UTickableConstraint> >& ConstraintsArray = Controller.GetConstraintsArray();
 	for (const TObjectPtr<UTickableConstraint>& Constraint : ConstraintsArray)
diff --git a/Engine/Source/Runtime/Experimental/Animation/Constraints/Public/ConstraintsScripting.h b/Engine/Source/Runtime/Experimental/Animation/Constraints/Public/ConstraintsScripting.h
--- a/Engine/Source/Runtime/Experimental/Animation/Constraints/Public/ConstraintsScripting.h
+++ b/Engine/Source/Runtime/Experimental/Animation/Constraints/Public/ConstraintsScripting.h
@@ -72,7 +72,7 @@ public:
 	@return Copy of the constraints in the level
 	*/
 	UFUNCTION(BlueprintCallable, Category = "Editor Scripting | Sequencer Tools | Control Rig | Constraints")
-	static TArray<TObjectPtr<UTickableConstraint>> GetConstraintsArray(UWorld* InWorld);
+	static TArray<UTickableConstraint*> GetConstraintsArray(UWorld* InWorld);
 
 	/**
 	* Remove constraint at specified index
