Control Rig: Fix cannot add public functions from old libraries

diff --git a/Engine/Plugins/Animation/ControlRig/Source/ControlRigEditor/Private/ControlRigEditorModule.cpp b/Engine/Plugins/Animation/ControlRig/Source/ControlRigEditor/Private/ControlRigEditorModule.cpp
--- a/Engine/Plugins/Animation/ControlRig/Source/ControlRigEditor/Private/ControlRigEditorModule.cpp
+++ b/Engine/Plugins/Animation/ControlRig/Source/ControlRigEditor/Private/ControlRigEditorModule.cpp
@@ -1175,8 +1175,9 @@ void FControlRigEditorModule::GetTypeActions(UControlRigBlueprint* CRB, FBluepri
 		ActionRegistrar.AddBlueprintAction(ActionKey, NodeSpawner);
 	}
 
-	FArrayProperty* PublicFunctionsProperty = CastField<FArrayProperty>(UControlRigBlueprint::StaticClass()->FindPropertyByName(TEXT("PublicGraphFunctions")));
-	if(PublicFunctionsProperty)
+	FArrayProperty* PublicGraphFunctionsProperty = CastField<FArrayProperty>(UControlRigBlueprint::StaticClass()->FindPropertyByName(TEXT("PublicGraphFunctions")));
+	FArrayProperty* PublicFunctionsProperty = CastField<FArrayProperty>(UControlRigBlueprint::StaticClass()->FindPropertyByName(TEXT("PublicFunctions")));
+	if(PublicGraphFunctionsProperty || PublicFunctionsProperty)
 	{
 		// find all control rigs in the project
 		TArray<FAssetData> ControlRigAssetDatas;
@@ -1197,19 +1198,34 @@ void FControlRigEditorModule::GetTypeActions(UControlRigBlueprint* CRB, FBluepri
 			PackagesProcessed.Add(ControlRigAssetData.PackageName);
 			
 			const FString PublicFunctionsString = ControlRigAssetData.GetTagValueRef<FString>(PublicFunctionsProperty->GetFName());
-			if(PublicFunctionsString.IsEmpty())
+			const FString PublicGraphFunctionsString = ControlRigAssetData.GetTagValueRef<FString>(PublicGraphFunctionsProperty->GetFName());
+			if(PublicFunctionsString.IsEmpty() && PublicGraphFunctionsString.IsEmpty())
 			{
 				continue;
 			}
 
-			TArray<FRigVMGraphFunctionHeader> PublicFunctions;
-			PublicFunctionsProperty->ImportText_Direct(*PublicFunctionsString, &PublicFunctions, nullptr, EPropertyPortFlags::PPF_None);
+			if (!PublicFunctionsString.IsEmpty())
+			{
+				TArray<FControlRigPublicFunctionData> PublicFunctions;
+				PublicFunctionsProperty->ImportText_Direct(*PublicFunctionsString, &PublicFunctions, nullptr, EPropertyPortFlags::PPF_None);
+				for(const FControlRigPublicFunctionData& PublicFunction : PublicFunctions)
+				{
+					UBlueprintNodeSpawner* NodeSpawner = UControlRigFunctionRefNodeSpawner::CreateFromAssetData(ControlRigAssetData, PublicFunction);
+					check(NodeSpawner != nullptr);
+					ActionRegistrar.AddBlueprintAction(ActionKey, NodeSpawner);						
+				}
+			}
 
-			for(const FRigVMGraphFunctionHeader& PublicFunction : PublicFunctions)
+			if (!PublicGraphFunctionsString.IsEmpty())
 			{
-				UBlueprintNodeSpawner* NodeSpawner = UControlRigFunctionRefNodeSpawner::CreateFromAssetData(ControlRigAssetData, PublicFunction);
-				check(NodeSpawner != nullptr);
-				ActionRegistrar.AddBlueprintAction(ActionKey, NodeSpawner);
+				TArray<FRigVMGraphFunctionHeader> PublicFunctions;
+				PublicGraphFunctionsProperty->ImportText_Direct(*PublicGraphFunctionsString, &PublicFunctions, nullptr, EPropertyPortFlags::PPF_None);
+				for(const FRigVMGraphFunctionHeader& PublicFunction : PublicFunctions)
+				{
+					UBlueprintNodeSpawner* NodeSpawner = UControlRigFunctionRefNodeSpawner::CreateFromAssetData(ControlRigAssetData, PublicFunction);
+					check(NodeSpawner != nullptr);
+					ActionRegistrar.AddBlueprintAction(ActionKey, NodeSpawner);
+				}
 			}
 		}
 	}
diff --git a/Engine/Plugins/Animation/ControlRig/Source/ControlRigEditor/Private/Graph/NodeSpawners/ControlRigFunctionRefNodeSpawner.cpp b/Engine/Plugins/Animation/ControlRig/Source/ControlRigEditor/Private/Graph/NodeSpawners/ControlRigFunctionRefNodeSpawner.cpp
--- a/Engine/Plugins/Animation/ControlRig/Source/ControlRigEditor/Private/Graph/NodeSpawners/ControlRigFunctionRefNodeSpawner.cpp
+++ b/Engine/Plugins/Animation/ControlRig/Source/ControlRigEditor/Private/Graph/NodeSpawners/ControlRigFunctionRefNodeSpawner.cpp
@@ -132,6 +132,69 @@ UControlRigFunctionRefNodeSpawner* UControlRigFunctionRefNodeSpawner::CreateFrom
 	return NodeSpawner;
 }
 
+UControlRigFunctionRefNodeSpawner* UControlRigFunctionRefNodeSpawner::CreateFromAssetData(const FAssetData& InAssetData, const FControlRigPublicFunctionData& InPublicFunction)
+{
+	UControlRigFunctionRefNodeSpawner* NodeSpawner = NewObject<UControlRigFunctionRefNodeSpawner>(GetTransientPackage());
+	NodeSpawner->NodeClass = UControlRigGraphNode::StaticClass();
+	NodeSpawner->bIsLocalFunction = false;
+
+	NodeSpawner->ReferencedPublicFunctionHeader.Name = InPublicFunction.Name;
+	NodeSpawner->AssetPath = InAssetData.ToSoftObjectPath();
+
+	FBlueprintActionUiSpec& MenuSignature = NodeSpawner->DefaultMenuSignature;
+
+	const FString Category = InPublicFunction.Category;
+
+	MenuSignature.MenuName = FText::FromName(InPublicFunction.Name);
+	MenuSignature.Category = FText::FromString(Category);
+	MenuSignature.Keywords = FText::FromString(InPublicFunction.Keywords);
+
+	if(const UControlRigBlueprint* ReferencedBlueprint = Cast<UControlRigBlueprint>(InAssetData.FastGetAsset(false)))
+	{
+		if(const URigVMFunctionLibrary* FunctionLibrary = ReferencedBlueprint->GetLocalFunctionLibrary())
+		{
+			if(const URigVMLibraryNode* FunctionNode = FunctionLibrary->FindFunction(InPublicFunction.Name))
+			{
+				MenuSignature.Tooltip = FunctionNode->GetToolTipText();
+			}
+		}
+	}
+
+	const FString ObjectPathString = InAssetData.GetObjectPathString();
+	if(MenuSignature.Tooltip.IsEmpty())
+	{
+		MenuSignature.Tooltip = FText::FromString(ObjectPathString);
+	}
+	else
+	{
+		static constexpr TCHAR Format[] = TEXT("%s\n\n%s");
+		MenuSignature.Tooltip = FText::FromString(FString::Printf(Format, *MenuSignature.Tooltip.ToString(), *ObjectPathString));
+	}
+
+	// add at least one character, so that PrimeDefaultUiSpec() doesn't 
+	// attempt to query the template node
+	if (MenuSignature.Keywords.IsEmpty())
+	{
+		// want to set it to something so we won't end up back in this condition
+		MenuSignature.Keywords = FText::FromString(TEXT(" "));
+	}
+
+	MenuSignature.Icon = FSlateIcon(FAppStyle::GetAppStyleSetName(), "Kismet.AllClasses.FunctionIcon");
+
+#if WITH_EDITOR
+	if (InPublicFunction.IsMutable())
+	{
+		MenuSignature.IconTint = GetDefault<UGraphEditorSettings>()->FunctionCallNodeTitleColor;
+	}
+	else
+	{
+		MenuSignature.IconTint = GetDefault<UGraphEditorSettings>()->PureFunctionCallNodeTitleColor;
+	}
+#endif
+	
+	return NodeSpawner;
+}
+
 void UControlRigFunctionRefNodeSpawner::Prime()
 {
 	// we expect that you don't need a node template to construct menu entries
@@ -141,12 +204,25 @@ void UControlRigFunctionRefNodeSpawner::Prime()
 FBlueprintNodeSignature UControlRigFunctionRefNodeSpawner::GetSpawnerSignature() const
 {
 	FString SignatureString = TEXT("Invalid RigFunction");
-	if (!ReferencedPublicFunctionHeader.IsValid())
+	if (ReferencedPublicFunctionHeader.IsValid())
 	{
 		SignatureString = 
 			FString::Printf(TEXT("RigFunction=%s"),
 			*ReferencedPublicFunctionHeader.GetHash());
 	}
+	else if (ReferencedFunctionPtr.IsValid())
+	{
+		SignatureString = 
+			FString::Printf(TEXT("RigFunction=%s"),
+			*ReferencedFunctionPtr->GetPathName());
+	}
+	else
+	{
+		SignatureString = 
+			FString::Printf(TEXT("RigFunction=%s:%s"),
+			*AssetPath.ToString(),
+			*ReferencedPublicFunctionHeader.Name.ToString());
+	}
 
 	if(bIsLocalFunction)
 	{
@@ -169,6 +245,14 @@ UEdGraphNode* UControlRigFunctionRefNodeSpawner::Invoke(UEdGraph* ParentGraph, F
 {
 	UControlRigGraphNode* NewNode = nullptr;
 
+	if (!ReferencedPublicFunctionHeader.IsValid() && !ReferencedFunctionPtr.IsValid() && AssetPath.IsValid())
+	{
+		if (UControlRigBlueprint* Blueprint = Cast<UControlRigBlueprint>(AssetPath.TryLoad()))
+		{
+			ReferencedFunctionPtr = Blueprint->GetLocalFunctionLibrary()->FindFunction(ReferencedPublicFunctionHeader.Name);			
+		}
+	}
+	
 	if (!ReferencedFunctionPtr.ToString().IsEmpty())
 	{
 		if (URigVMLibraryNode* LibraryNode = ReferencedFunctionPtr.LoadSynchronous())
diff --git a/Engine/Plugins/Animation/ControlRig/Source/ControlRigEditor/Private/Graph/NodeSpawners/ControlRigFunctionRefNodeSpawner.h b/Engine/Plugins/Animation/ControlRig/Source/ControlRigEditor/Private/Graph/NodeSpawners/ControlRigFunctionRefNodeSpawner.h
--- a/Engine/Plugins/Animation/ControlRig/Source/ControlRigEditor/Private/Graph/NodeSpawners/ControlRigFunctionRefNodeSpawner.h
+++ b/Engine/Plugins/Animation/ControlRig/Source/ControlRigEditor/Private/Graph/NodeSpawners/ControlRigFunctionRefNodeSpawner.h
@@ -28,20 +28,19 @@ class CONTROLRIGEDITOR_API UControlRigFunctionRefNodeSpawner : public UBlueprint
 public:
 
 	/**
-	 * Creates a new UControlRigVariableNodeSpawner, charged with spawning 
-	 * a new member-variable node
+	 * Creates a new UControlRigFunctionRefNodeSpawner, charged with spawning a function reference
 	 * 
 	 * @return A newly allocated instance of this class.
 	 */
 	static UControlRigFunctionRefNodeSpawner* CreateFromFunction(URigVMLibraryNode* InFunction);
 
 	/**
-	 * Creates a new UControlRigVariableNodeSpawner, charged with spawning 
-	 * a new member-variable node
+	 * Creates a new UControlRigFunctionRefNodeSpawner, charged with spawning a function reference
 	 * 
 	 * @return A newly allocated instance of this class.
 	 */
 	static UControlRigFunctionRefNodeSpawner* CreateFromAssetData(const FAssetData& InAssetData, const FRigVMGraphFunctionHeader& InPublicFunction);
+	static UControlRigFunctionRefNodeSpawner* CreateFromAssetData(const FAssetData& InAssetData, const FControlRigPublicFunctionData& InPublicFunction);
 
 	// UBlueprintNodeSpawner interface
 	virtual void Prime() override;
@@ -65,6 +64,9 @@ private:
 	UPROPERTY(Transient)
 	bool bIsLocalFunction;
 
+	UPROPERTY()
+	FSoftObjectPath AssetPath;
+
 	static UControlRigGraphNode* SpawnNode(UEdGraph* ParentGraph, UBlueprint* Blueprint, FRigVMGraphFunctionHeader& InFunction, FVector2D const Location);
 
 	friend class UEngineTestControlRig;
