Control Rig: Do not load all control rigs with public functions when opening a control rig editor

diff --git a/Engine/Plugins/Animation/ControlRig/Source/ControlRigDeveloper/Private/ControlRigBlueprint.cpp b/Engine/Plugins/Animation/ControlRig/Source/ControlRigDeveloper/Private/ControlRigBlueprint.cpp
--- a/Engine/Plugins/Animation/ControlRig/Source/ControlRigDeveloper/Private/ControlRigBlueprint.cpp
+++ b/Engine/Plugins/Animation/ControlRig/Source/ControlRigDeveloper/Private/ControlRigBlueprint.cpp
@@ -690,11 +690,14 @@ void UControlRigBlueprint::Serialize(FArchive& Ar)
 
 	if(Ar.IsObjectReferenceCollector())
 	{
-		TArray<UControlRigBlueprint*> ReferencedBlueprints = GetReferencedControlRigBlueprints();
+		TArray<IRigVMGraphFunctionHost*> ReferencedFunctionHosts = GetReferencedFunctionHosts(false);
 
-		for(UControlRigBlueprint* ReferencedBlueprint : ReferencedBlueprints)
+		for(IRigVMGraphFunctionHost* ReferencedFunctionHost : ReferencedFunctionHosts)
 		{
-			Ar << ReferencedBlueprints;
+			if (UControlRigBlueprintGeneratedClass* BPGeneratedClass = Cast<UControlRigBlueprintGeneratedClass>(ReferencedFunctionHost))
+			{
+				Ar << BPGeneratedClass;
+			}
 		}
 
 		for(const TSoftObjectPtr<UControlRigShapeLibrary>& ShapeLibraryPtr : ShapeLibraries)
@@ -767,15 +770,18 @@ void UControlRigBlueprint::PostLoad()
 	{
 		TGuardValue<bool> IsCompilingGuard(bIsCompiling, true);
 		
-		TArray<UControlRigBlueprint*> ReferencedBlueprints = GetReferencedControlRigBlueprints();
+		TArray<IRigVMGraphFunctionHost*> ReferencedFunctionHosts = GetReferencedFunctionHosts(true);
 
-		// PostLoad all referenced BPs so that their function graphs are fully loaded 
+		// PostLoad all referenced function hosts so that their function data are fully loaded 
 		// and ready to be inlined into this BP during compilation
-		for (UControlRigBlueprint* BP : ReferencedBlueprints)
+		for (IRigVMGraphFunctionHost* FunctionHost : ReferencedFunctionHosts)
 		{
-			if (BP->HasAllFlags(RF_NeedPostLoad))
+			if (UControlRigBlueprintGeneratedClass* BPGeneratedClass = Cast<UControlRigBlueprintGeneratedClass>(FunctionHost))
 			{
-				BP->ConditionalPostLoad();
+				if (BPGeneratedClass->HasAllFlags(RF_NeedPostLoad))
+				{
+					BPGeneratedClass->ConditionalPostLoad();
+				}
 			}
 		}
 		
@@ -1639,9 +1645,9 @@ void UControlRigBlueprint::HandleReportFromCompiler(EMessageSeverity::Type InSev
 	}
 }
 
-TArray<UControlRigBlueprint*> UControlRigBlueprint::GetReferencedControlRigBlueprints()
+TArray<IRigVMGraphFunctionHost*> UControlRigBlueprint::GetReferencedFunctionHosts(bool bForceLoad)
 {
-	TArray<UControlRigBlueprint*> ReferencedBlueprints;
+	TArray<IRigVMGraphFunctionHost*> ReferencedBlueprints;
 	
 	TArray<UEdGraph*> EdGraphs;
 	GetAllGraphs(EdGraphs);
@@ -1653,16 +1659,32 @@ TArray<UControlRigBlueprint*> UControlRigBlueprint::GetReferencedControlRigBluep
 			{
 				if(URigVMFunctionReferenceNode* FunctionRefNode = Cast<URigVMFunctionReferenceNode>(RigNode->GetModelNode()))
 				{
-					if(const URigVMLibraryNode* ReferencedNode = FunctionRefNode->LoadReferencedNode())
+					IRigVMGraphFunctionHost* Host = nullptr;
+					if (bForceLoad || FunctionRefNode->IsReferencedFunctionHostLoaded())
+					{
+						// Load the function host
+						Host = FunctionRefNode->GetReferencedFunctionHeader().GetFunctionHost();
+					}
+					else if (bForceLoad || FunctionRefNode->IsReferencedNodeLoaded())
 					{
-						if(URigVMFunctionLibrary* ReferencedFunctionLibrary = ReferencedNode->GetLibrary())
+						// Load the reference library node
+						if(const URigVMLibraryNode* ReferencedNode = FunctionRefNode->LoadReferencedNode())
 						{
-							if(ReferencedFunctionLibrary == GetLocalFunctionLibrary())
+							if(URigVMFunctionLibrary* ReferencedFunctionLibrary = ReferencedNode->GetLibrary())
 							{
-								continue;
+								FSoftObjectPath FunctionHostPath = ReferencedFunctionLibrary->GetFunctionHostObjectPath();
+								if (UObject* FunctionHostObj = FunctionHostPath.TryLoad())
+								{
+									Host = Cast<IRigVMGraphFunctionHost>(FunctionHostObj);									
+								}
 							}
 						}
 					}
+
+					if (Host != nullptr && Host != GetControlRigBlueprintGeneratedClass())
+					{
+						ReferencedBlueprints.Add(Host);
+					}
 				}
 			}
 		}
diff --git a/Engine/Plugins/Animation/ControlRig/Source/ControlRigDeveloper/Public/ControlRigBlueprint.h b/Engine/Plugins/Animation/ControlRig/Source/ControlRigDeveloper/Public/ControlRigBlueprint.h
--- a/Engine/Plugins/Animation/ControlRig/Source/ControlRigDeveloper/Public/ControlRigBlueprint.h
+++ b/Engine/Plugins/Animation/ControlRig/Source/ControlRigDeveloper/Public/ControlRigBlueprint.h
@@ -735,7 +735,7 @@ private:
 	bool RemoveEdGraphForCollapseNode(URigVMCollapseNode* InNode, bool bNotify = false);
 	void HandleReportFromCompiler(EMessageSeverity::Type InSeverity, UObject* InSubject, const FString& InMessage);
 
-	TArray<UControlRigBlueprint*> GetReferencedControlRigBlueprints();
+	TArray<IRigVMGraphFunctionHost*> GetReferencedFunctionHosts(bool bForceLoad);
 	
 #if WITH_EDITOR
 private:
diff --git a/Engine/Plugins/Animation/ControlRig/Source/ControlRigEditor/Private/Editor/ControlRigEditor.cpp b/Engine/Plugins/Animation/ControlRig/Source/ControlRigEditor/Private/Editor/ControlRigEditor.cpp
--- a/Engine/Plugins/Animation/ControlRig/Source/ControlRigEditor/Private/Editor/ControlRigEditor.cpp
+++ b/Engine/Plugins/Animation/ControlRig/Source/ControlRigEditor/Private/Editor/ControlRigEditor.cpp
@@ -612,8 +612,8 @@ void FControlRigEditor::InitFunctionReferences()
 		return;
 	}
 	
-	//FFunctionGraphTask::CreateAndDispatchWhenReady([]()
-	//{
+	FFunctionGraphTask::CreateAndDispatchWhenReady([]()
+	{
 		if(URigVMBuildData* BuildData = URigVMController::GetBuildData())
 		{
 			const FArrayProperty* ReferenceNodeDataProperty =
@@ -655,7 +655,7 @@ void FControlRigEditor::InitFunctionReferences()
 				}
 			}
 		}
-	//}, TStatId(), NULL, ENamedThreads::GameThread);
+	}, TStatId(), NULL, ENamedThreads::GameThread);
 	
 	bAreFunctionReferencesInitialized = true;
 }
diff --git a/Engine/Source/Developer/RigVMDeveloper/Private/RigVMCompiler/RigVMCompiler.cpp b/Engine/Source/Developer/RigVMDeveloper/Private/RigVMCompiler/RigVMCompiler.cpp
--- a/Engine/Source/Developer/RigVMDeveloper/Private/RigVMCompiler/RigVMCompiler.cpp
+++ b/Engine/Source/Developer/RigVMDeveloper/Private/RigVMCompiler/RigVMCompiler.cpp
@@ -928,8 +928,6 @@ bool URigVMCompiler::Compile(TArray<URigVMGraph*> InGraphs, URigVMController* In
 		ReportInfof(TEXT("Total Compilation time %f\n"), CompilationTime*1000);
 	}
 
-	FPlatformMisc::LowLevelOutputDebugStringf(TEXT("%s\n\n"), *WorkData.VM->DumpByteCodeAsText());
-
 	return true;
 }
 
diff --git a/Engine/Source/Developer/RigVMDeveloper/Private/RigVMModel/Nodes/RigVMFunctionReferenceNode.cpp b/Engine/Source/Developer/RigVMDeveloper/Private/RigVMModel/Nodes/RigVMFunctionReferenceNode.cpp
--- a/Engine/Source/Developer/RigVMDeveloper/Private/RigVMModel/Nodes/RigVMFunctionReferenceNode.cpp
+++ b/Engine/Source/Developer/RigVMDeveloper/Private/RigVMModel/Nodes/RigVMFunctionReferenceNode.cpp
@@ -9,47 +9,27 @@
 
 FString URigVMFunctionReferenceNode::GetNodeTitle() const
 {
-	if (const FRigVMGraphFunctionData* Data = GetReferencedFunctionData())
-	{
-		return Data->Header.NodeTitle;
-	}
-	return Super::GetNodeTitle();
+	return ReferencedFunctionHeader.NodeTitle;
 }
 
 FLinearColor URigVMFunctionReferenceNode::GetNodeColor() const
 {
-	if (const FRigVMGraphFunctionData* Data = GetReferencedFunctionData())
-	{
-		return Data->Header.NodeColor;
-	}
-	return Super::GetNodeColor();
+	return ReferencedFunctionHeader.NodeColor;
 }
 
 FText URigVMFunctionReferenceNode::GetToolTipText() const
 {
-	if (const FRigVMGraphFunctionData* Data = GetReferencedFunctionData())
-	{
-		return Data->Header.Tooltip;
-	}
-	return Super::GetToolTipText();
+	return ReferencedFunctionHeader.Tooltip;
 }
 
 FString URigVMFunctionReferenceNode::GetNodeCategory() const
 {
-	if (const FRigVMGraphFunctionData* Data = GetReferencedFunctionData())
-	{
-		return Data->Header.Category;
-	}
-	return Super::GetNodeCategory();
+	return ReferencedFunctionHeader.Category;
 }
 
 FString URigVMFunctionReferenceNode::GetNodeKeywords() const
 {
-	if (const FRigVMGraphFunctionData* Data = GetReferencedFunctionData())
-	{
-		return Data->Header.Keywords;
-	}
-	return Super::GetNodeKeywords();
+	return ReferencedFunctionHeader.Keywords;
 }
 
 bool URigVMFunctionReferenceNode::RequiresVariableRemapping() const
@@ -172,6 +152,16 @@ FText URigVMFunctionReferenceNode::GetToolTipTextForPin(const URigVMPin* InPin)
 	return Super::GetToolTipTextForPin(InPin);
 }
 
+bool URigVMFunctionReferenceNode::IsReferencedFunctionHostLoaded() const
+{
+	return ReferencedFunctionHeader.LibraryPointer.HostObject.ResolveObject() != nullptr;
+}
+
+bool URigVMFunctionReferenceNode::IsReferencedNodeLoaded() const
+{
+	return ReferencedFunctionHeader.LibraryPointer.LibraryNode.ResolveObject() != nullptr;
+}
+
 URigVMLibraryNode* URigVMFunctionReferenceNode::LoadReferencedNode() const
 {
 	UObject* LibraryNode = ReferencedFunctionHeader.LibraryPointer.LibraryNode.ResolveObject();
diff --git a/Engine/Source/Developer/RigVMDeveloper/Private/RigVMModel/Nodes/RigVMLibraryNode.cpp b/Engine/Source/Developer/RigVMDeveloper/Private/RigVMModel/Nodes/RigVMLibraryNode.cpp
--- a/Engine/Source/Developer/RigVMDeveloper/Private/RigVMModel/Nodes/RigVMLibraryNode.cpp
+++ b/Engine/Source/Developer/RigVMDeveloper/Private/RigVMModel/Nodes/RigVMLibraryNode.cpp
@@ -221,7 +221,10 @@ TMap<FRigVMGraphFunctionIdentifier, uint32> URigVMLibraryNode::GetDependencies()
 FRigVMGraphFunctionIdentifier URigVMLibraryNode::GetFunctionIdentifier() const
 {
 	FRigVMGraphFunctionIdentifier Identifier;
-	Identifier.HostObject = GetLibrary()->GetFunctionHostObjectPath();
+	if (URigVMFunctionLibrary* Library = GetLibrary())
+	{
+		Identifier.HostObject = Library->GetFunctionHostObjectPath();
+	}
 	Identifier.LibraryNode = this;
 	return Identifier;
 }
diff --git a/Engine/Source/Developer/RigVMDeveloper/Private/RigVMModel/RigVMBuildData.cpp b/Engine/Source/Developer/RigVMDeveloper/Private/RigVMModel/RigVMBuildData.cpp
--- a/Engine/Source/Developer/RigVMDeveloper/Private/RigVMModel/RigVMBuildData.cpp
+++ b/Engine/Source/Developer/RigVMDeveloper/Private/RigVMModel/RigVMBuildData.cpp
@@ -183,10 +183,6 @@ void URigVMBuildData::RegisterFunctionReference(FRigVMReferenceNodeData InRefere
 	if (!bFound)
 	{
 		FRigVMGraphFunctionIdentifier Pointer(nullptr, LibraryNodePtr.ToSoftObjectPath());
-		if (!LibraryNodePtr.IsValid())
-		{
-			LibraryNodePtr.LoadSynchronous();
-		}
 		if (LibraryNodePtr.IsValid())
 		{
 			Pointer.HostObject = Cast<UObject>(LibraryNodePtr.Get()->GetFunctionHeader().GetFunctionHost());
diff --git a/Engine/Source/Developer/RigVMDeveloper/Public/RigVMModel/Nodes/RigVMFunctionReferenceNode.h b/Engine/Source/Developer/RigVMDeveloper/Public/RigVMModel/Nodes/RigVMFunctionReferenceNode.h
--- a/Engine/Source/Developer/RigVMDeveloper/Public/RigVMModel/Nodes/RigVMFunctionReferenceNode.h
+++ b/Engine/Source/Developer/RigVMDeveloper/Public/RigVMModel/Nodes/RigVMFunctionReferenceNode.h
@@ -33,6 +33,8 @@ public:
 	virtual const FRigVMTemplate* GetTemplate() const override { return nullptr; }
 	// end URigVMLibraryNode interface
 
+	bool IsReferencedFunctionHostLoaded() const;
+	bool IsReferencedNodeLoaded() const;
 	URigVMLibraryNode* LoadReferencedNode() const;
 
 	// Variable remapping
