Control Rig: Fix static analysis warning

diff --git a/Engine/Plugins/Animation/ControlRig/Source/ControlRigEditor/Private/ControlRigEditorModule.cpp b/Engine/Plugins/Animation/ControlRig/Source/ControlRigEditor/Private/ControlRigEditorModule.cpp
--- a/Engine/Plugins/Animation/ControlRig/Source/ControlRigEditor/Private/ControlRigEditorModule.cpp
+++ b/Engine/Plugins/Animation/ControlRig/Source/ControlRigEditor/Private/ControlRigEditorModule.cpp
@@ -1196,9 +1196,17 @@ void FControlRigEditorModule::GetTypeActions(UControlRigBlueprint* CRB, FBluepri
 				continue;
 			}
 			PackagesProcessed.Add(ControlRigAssetData.PackageName);
-			
-			const FString PublicFunctionsString = ControlRigAssetData.GetTagValueRef<FString>(PublicFunctionsProperty->GetFName());
-			const FString PublicGraphFunctionsString = ControlRigAssetData.GetTagValueRef<FString>(PublicGraphFunctionsProperty->GetFName());
+
+			FString PublicFunctionsString;
+			if (PublicFunctionsProperty)
+			{
+				PublicFunctionsString = ControlRigAssetData.GetTagValueRef<FString>(PublicFunctionsProperty->GetFName());
+			}
+			FString PublicGraphFunctionsString;
+			if (PublicGraphFunctionsProperty)
+			{
+				PublicGraphFunctionsString = ControlRigAssetData.GetTagValueRef<FString>(PublicGraphFunctionsProperty->GetFName());
+			}
 			if(PublicFunctionsString.IsEmpty() && PublicGraphFunctionsString.IsEmpty())
 			{
 				continue;
