Control Rig: Fix FunctionReferenceSpawner loading all the rigs with public functions

diff --git a/Engine/Plugins/Animation/ControlRig/Source/ControlRigEditor/Private/Graph/NodeSpawners/ControlRigFunctionRefNodeSpawner.cpp b/Engine/Plugins/Animation/ControlRig/Source/ControlRigEditor/Private/Graph/NodeSpawners/ControlRigFunctionRefNodeSpawner.cpp
--- a/Engine/Plugins/Animation/ControlRig/Source/ControlRigEditor/Private/Graph/NodeSpawners/ControlRigFunctionRefNodeSpawner.cpp
+++ b/Engine/Plugins/Animation/ControlRig/Source/ControlRigEditor/Private/Graph/NodeSpawners/ControlRigFunctionRefNodeSpawner.cpp
@@ -29,12 +29,17 @@
 
 #define LOCTEXT_NAMESPACE "ControlRigFunctionRefNodeSpawner"
 
+// 3 possible creation types:
+// - CreateFromFunction(URigVMLibraryNode) --> This is a local function. Valid Header.
+// - CreateFromAssetData(const FAssetData& InAssetData, const FRigVMGraphFunctionHeader& InPublicFunction) --> Public function. Valid Header.
+// - CreateFromAssetData(const FAssetData& InAssetData, const FControlRigPublicFunctionData& InPublicFunction) --> Public function. Valid AssetData and Header.Name
+
 UControlRigFunctionRefNodeSpawner* UControlRigFunctionRefNodeSpawner::CreateFromFunction(URigVMLibraryNode* InFunction)
 {
 	check(InFunction);
 
 	UControlRigFunctionRefNodeSpawner* NodeSpawner = NewObject<UControlRigFunctionRefNodeSpawner>(GetTransientPackage());
-	NodeSpawner->ReferencedFunctionPtr = InFunction;
+	NodeSpawner->ReferencedPublicFunctionHeader = InFunction->GetFunctionHeader();
 	NodeSpawner->NodeClass = UControlRigGraphNode::StaticClass();
 	NodeSpawner->bIsLocalFunction = true;
 
@@ -138,7 +143,21 @@ UControlRigFunctionRefNodeSpawner* UControlRigFunctionRefNodeSpawner::CreateFrom
 	NodeSpawner->NodeClass = UControlRigGraphNode::StaticClass();
 	NodeSpawner->bIsLocalFunction = false;
 
-	NodeSpawner->ReferencedPublicFunctionHeader.Name = InPublicFunction.Name;
+	FRigVMGraphFunctionHeader& Header = NodeSpawner->ReferencedPublicFunctionHeader;
+	Header.Name = InPublicFunction.Name;
+	Header.Arguments.Reserve(InPublicFunction.Arguments.Num());
+	for (const FControlRigPublicFunctionArg& Arg : InPublicFunction.Arguments)
+	{
+		FRigVMGraphFunctionArgument NewArgument;
+		NewArgument.Name = Arg.Name;
+		NewArgument.Direction = Arg.Direction;
+		NewArgument.bIsArray = Arg.bIsArray;
+		NewArgument.CPPType = Arg.CPPType;
+		PRAGMA_DISABLE_DEPRECATION_WARNINGS
+		NewArgument.CPPTypeObject = FSoftObjectPath(Arg.CPPTypeObjectPath);
+		PRAGMA_ENABLE_DEPRECATION_WARNINGS
+		Header.Arguments.Add(NewArgument);
+	}
 	NodeSpawner->AssetPath = InAssetData.ToSoftObjectPath();
 
 	FBlueprintActionUiSpec& MenuSignature = NodeSpawner->DefaultMenuSignature;
@@ -210,12 +229,6 @@ FBlueprintNodeSignature UControlRigFunctionRefNodeSpawner::GetSpawnerSignature()
 			FString::Printf(TEXT("RigFunction=%s"),
 			*ReferencedPublicFunctionHeader.GetHash());
 	}
-	else if (ReferencedFunctionPtr.IsValid())
-	{
-		SignatureString = 
-			FString::Printf(TEXT("RigFunction=%s"),
-			*ReferencedFunctionPtr->GetPathName());
-	}
 	else
 	{
 		SignatureString = 
@@ -245,25 +258,18 @@ UEdGraphNode* UControlRigFunctionRefNodeSpawner::Invoke(UEdGraph* ParentGraph, F
 {
 	UControlRigGraphNode* NewNode = nullptr;
 
-	if (!ReferencedPublicFunctionHeader.IsValid() && !ReferencedFunctionPtr.IsValid() && AssetPath.IsValid())
+	// if we are trying to build the real function ref - but we haven't loaded the asset yet...
+	if(!FBlueprintNodeTemplateCache::IsTemplateOuter(ParentGraph))
 	{
-		if (UControlRigBlueprint* Blueprint = Cast<UControlRigBlueprint>(AssetPath.TryLoad()))
+		if (!ReferencedPublicFunctionHeader.IsValid() && AssetPath.IsValid())
 		{
-			ReferencedFunctionPtr = Blueprint->GetLocalFunctionLibrary()->FindFunction(ReferencedPublicFunctionHeader.Name);			
-		}
-	}
-	
-	if (!ReferencedFunctionPtr.ToString().IsEmpty())
-	{
-		if (URigVMLibraryNode* LibraryNode = ReferencedFunctionPtr.LoadSynchronous())
-		{
-			if (URigVMFunctionLibrary* FunctionLibrary = LibraryNode->GetGraph()->GetDefaultFunctionLibrary())
+			if (UControlRigBlueprint* Blueprint = Cast<UControlRigBlueprint>(AssetPath.TryLoad()))
 			{
-				ReferencedPublicFunctionHeader = LibraryNode->GetFunctionHeader();
+				ReferencedPublicFunctionHeader = Blueprint->GetLocalFunctionLibrary()->FindFunction(ReferencedPublicFunctionHeader.Name)->GetFunctionHeader();			
 			}
 		}
 	}
-
+		
 	if(ReferencedPublicFunctionHeader.IsValid())
 	{
 #if WITH_EDITOR
@@ -276,6 +282,40 @@ UEdGraphNode* UControlRigFunctionRefNodeSpawner::Invoke(UEdGraph* ParentGraph, F
 		UBlueprint* Blueprint = FBlueprintEditorUtils::FindBlueprintForGraph(ParentGraph);
 		NewNode = SpawnNode(ParentGraph, Blueprint, ReferencedPublicFunctionHeader, Location);
 	}
+	else
+	{
+		// we are only going to get here if we are spawning a template node
+		NewNode = NewObject<UControlRigGraphNode>(ParentGraph);
+		ParentGraph->AddNode(NewNode, false);
+
+		NewNode->CreateNewGuid();
+		NewNode->PostPlacedNewNode();
+
+		for(const FRigVMGraphFunctionArgument& Arg : ReferencedPublicFunctionHeader.Arguments)
+		{
+			if(Arg.Direction ==  ERigVMPinDirection::Input ||
+				Arg.Direction ==  ERigVMPinDirection::IO)
+			{
+				UEdGraphPin* InputPin = UEdGraphPin::CreatePin(NewNode);
+				NewNode->Pins.Add(InputPin);
+
+				InputPin->Direction = EGPD_Input;
+				InputPin->PinType = RigVMTypeUtils::PinTypeFromCPPType(Arg.CPPType, Arg.CPPTypeObject.Get());
+			}
+
+			if(Arg.Direction ==  ERigVMPinDirection::Output ||
+				Arg.Direction ==  ERigVMPinDirection::IO)
+			{
+				UEdGraphPin* OutputPin = UEdGraphPin::CreatePin(NewNode);
+				NewNode->Pins.Add(OutputPin);
+
+				OutputPin->Direction = EGPD_Output;
+				OutputPin->PinType = RigVMTypeUtils::PinTypeFromCPPType(Arg.CPPType, Arg.CPPTypeObject.Get());
+			}
+		}
+
+		NewNode->SetFlags(RF_Transactional);
+}
 
 	return NewNode;
 }
@@ -374,18 +414,18 @@ bool UControlRigFunctionRefNodeSpawner::IsTemplateNodeFilteredOut(FBlueprintActi
 {
 	if(bIsLocalFunction)
 	{
-		if(ReferencedFunctionPtr.IsValid())
+		if(ReferencedPublicFunctionHeader.IsValid())
 		{
 			for (UBlueprint* Blueprint : Filter.Context.Blueprints)
 			{
-				if(Blueprint->GetOutermost() != ReferencedFunctionPtr.Get()->GetOutermost())
+				if(Cast<IRigVMGraphFunctionHost>(Blueprint->GeneratedClass) != ReferencedPublicFunctionHeader.GetFunctionHost())
 				{
 					return true;
 				}
 			}
 		}
 	}
-	const FString ReferencedAssetObjectPathString = ReferencedFunctionPtr.GetAssetName();
+	const FString ReferencedAssetObjectPathString = ReferencedPublicFunctionHeader.LibraryPointer.LibraryNode.GetAssetName();
 	for (UBlueprint* Blueprint : Filter.Context.Blueprints)
 	{
 		if(Blueprint->GetPathName() == ReferencedAssetObjectPathString)
diff --git a/Engine/Plugins/Animation/ControlRig/Source/ControlRigEditor/Private/Graph/NodeSpawners/ControlRigFunctionRefNodeSpawner.h b/Engine/Plugins/Animation/ControlRig/Source/ControlRigEditor/Private/Graph/NodeSpawners/ControlRigFunctionRefNodeSpawner.h
--- a/Engine/Plugins/Animation/ControlRig/Source/ControlRigEditor/Private/Graph/NodeSpawners/ControlRigFunctionRefNodeSpawner.h
+++ b/Engine/Plugins/Animation/ControlRig/Source/ControlRigEditor/Private/Graph/NodeSpawners/ControlRigFunctionRefNodeSpawner.h
@@ -52,10 +52,6 @@ public:
 
 private:
 
-	/** The unit type we will spawn [optional] */
-	UPROPERTY(Transient)
-	mutable TSoftObjectPtr<URigVMLibraryNode> ReferencedFunctionPtr;
-
 	/** The public function definition we will spawn from [optional] */
 	UPROPERTY(Transient)
 	mutable FRigVMGraphFunctionHeader ReferencedPublicFunctionHeader;
