Control Rig: Fix ensure ReferencedFunctionPath_DEPRECATED is empty

diff --git a/Engine/Plugins/Animation/ControlRig/Source/ControlRigDeveloper/Private/ControlRigBlueprint.cpp b/Engine/Plugins/Animation/ControlRig/Source/ControlRigDeveloper/Private/ControlRigBlueprint.cpp
--- a/Engine/Plugins/Animation/ControlRig/Source/ControlRigDeveloper/Private/ControlRigBlueprint.cpp
+++ b/Engine/Plugins/Animation/ControlRig/Source/ControlRigDeveloper/Private/ControlRigBlueprint.cpp
@@ -1374,7 +1374,15 @@ void UControlRigBlueprint::RefreshAllModels(EControlRigBlueprintLoadType InLoadT
 				{
 					if (URigVMFunctionReferenceNode* FunctionReferenceNode = Cast<URigVMFunctionReferenceNode>(LibraryNode))
 					{
-						continue;
+						if (FunctionReferenceNode->GetReferencedFunctionHeader().LibraryPointer.LibraryNode.GetLongPackageName() != GetPackage()->GetPathName())
+						{							
+							continue;
+						}
+						if (URigVMLibraryNode* ReferencedNode = FunctionReferenceNode->LoadReferencedNode())
+						{
+							ContainedGraphs.Add(ReferencedNode->GetContainedGraph());
+							continue;
+						}
 					}
 
 					if (URigVMGraph* ContainedGraph = LibraryNode->GetContainedGraph())
@@ -4084,6 +4092,16 @@ void UControlRigBlueprint::PatchFunctionReferencesOnLoad()
 						}
 					}
 				}
+
+				if (FunctionReferenceNode->ReferencedNodePtr_DEPRECATED.IsValid())
+				{
+					FunctionReferenceNode->ReferencedFunctionHeader = FunctionReferenceNode->ReferencedNodePtr_DEPRECATED->GetFunctionHeader();					
+				}
+				else
+				{
+					// At least lets make sure we store the path in the header
+					FunctionReferenceNode->ReferencedFunctionHeader.LibraryPointer.LibraryNode = FunctionReferenceNode->ReferencedNodePtr_DEPRECATED.ToSoftObjectPath();
+				}
 			}
 		}
 	}
@@ -4576,13 +4594,7 @@ void UControlRigBlueprint::PatchFunctionsOnLoad()
 					TSoftObjectPtr<URigVMLibraryNode> ReferencedLibraryNode = FuncRefNode->ReferencedNodePtr_DEPRECATED;
 					if (ReferencedLibraryNode.IsValid())
 					{
-						if (URigVMLibraryNode* LibraryNode = ReferencedLibraryNode.LoadSynchronous())
-						{
-							if (URigVMFunctionLibrary* FunctionLibrary = LibraryNode->GetGraph()->GetDefaultFunctionLibrary())
-							{
-								FuncRefNode->ReferencedFunctionHeader = LibraryNode->GetFunctionHeader();
-							}
-						}
+						FuncRefNode->ReferencedFunctionHeader = ReferencedLibraryNode->GetFunctionHeader();
 					}
 				}
 
diff --git a/Engine/Source/Developer/RigVMDeveloper/Private/RigVMModel/RigVMBuildData.cpp b/Engine/Source/Developer/RigVMDeveloper/Private/RigVMModel/RigVMBuildData.cpp
--- a/Engine/Source/Developer/RigVMDeveloper/Private/RigVMModel/RigVMBuildData.cpp
+++ b/Engine/Source/Developer/RigVMDeveloper/Private/RigVMModel/RigVMBuildData.cpp
@@ -165,14 +165,22 @@ void URigVMBuildData::RegisterFunctionReference(FRigVMReferenceNodeData InRefere
 	{
 		return RegisterFunctionReference(InReferenceNodeData.ReferencedHeader.LibraryPointer, InReferenceNodeData.GetReferenceNodeObjectPath());
 	}
+
+	if (!InReferenceNodeData.ReferencedHeader.LibraryPointer.LibraryNode.IsValid())
+	{
+		InReferenceNodeData.ReferencedHeader.LibraryPointer.LibraryNode = InReferenceNodeData.ReferencedFunctionPath_DEPRECATED;
+	}
 	
-	check(!InReferenceNodeData.ReferencedFunctionPath_DEPRECATED.IsEmpty());
+	check(InReferenceNodeData.ReferencedHeader.LibraryPointer.LibraryNode.IsValid());
+
+	FSoftObjectPath LibraryNodePath = InReferenceNodeData.ReferencedHeader.LibraryPointer.LibraryNode;
+	TSoftObjectPtr<URigVMLibraryNode> LibraryNodePtr = TSoftObjectPtr<URigVMLibraryNode>(LibraryNodePath);
 
-	TSoftObjectPtr<URigVMLibraryNode> LibraryNodePtr = TSoftObjectPtr<URigVMLibraryNode>(InReferenceNodeData.ReferencedFunctionPath_DEPRECATED);
+	// Try to find a FunctionIdentifier with the same LibraryNodePath
 	bool bFound = false;
 	for (TPair< FRigVMGraphFunctionIdentifier, FRigVMFunctionReferenceArray >& Pair : GraphFunctionReferences)
 	{
-		if (Pair.Key.LibraryNode == LibraryNodePtr.ToSoftObjectPath())
+		if (Pair.Key.LibraryNode == LibraryNodePath)
 		{
 			Pair.Value.FunctionReferences.Add(InReferenceNodeData.GetReferenceNodeObjectPath());
 			bFound = true;
@@ -180,9 +188,10 @@ void URigVMBuildData::RegisterFunctionReference(FRigVMReferenceNodeData InRefere
 		}
 	}
 
+	// Otherwise, lets add a new identifier, even if it has no function host
 	if (!bFound)
 	{
-		FRigVMGraphFunctionIdentifier Pointer(nullptr, LibraryNodePtr.ToSoftObjectPath());
+		FRigVMGraphFunctionIdentifier Pointer(nullptr, LibraryNodePath);
 		if (LibraryNodePtr.IsValid())
 		{
 			Pointer.HostObject = Cast<UObject>(LibraryNodePtr.Get()->GetFunctionHeader().GetFunctionHost());
