Control Rig: Fix static analysis warning

diff --git a/Engine/Plugins/Animation/ControlRig/Source/ControlRigEditor/Private/ControlRigEditorModule.cpp b/Engine/Plugins/Animation/ControlRig/Source/ControlRigEditor/Private/ControlRigEditorModule.cpp
--- a/Engine/Plugins/Animation/ControlRig/Source/ControlRigEditor/Private/ControlRigEditorModule.cpp
+++ b/Engine/Plugins/Animation/ControlRig/Source/ControlRigEditor/Private/ControlRigEditorModule.cpp
@@ -1212,7 +1212,7 @@ void FControlRigEditorModule::GetTypeActions(UControlRigBlueprint* CRB, FBluepri
 				continue;
 			}
 
-			if (!PublicFunctionsString.IsEmpty())
+			if (PublicFunctionsProperty && !PublicFunctionsString.IsEmpty())
 			{
 				TArray<FControlRigPublicFunctionData> PublicFunctions;
 				PublicFunctionsProperty->ImportText_Direct(*PublicFunctionsString, &PublicFunctions, nullptr, EPropertyPortFlags::PPF_None);
@@ -1224,7 +1224,7 @@ void FControlRigEditorModule::GetTypeActions(UControlRigBlueprint* CRB, FBluepri
 				}
 			}
 
-			if (!PublicGraphFunctionsString.IsEmpty())
+			if (PublicGraphFunctionsProperty && !PublicGraphFunctionsString.IsEmpty())
 			{
 				TArray<FRigVMGraphFunctionHeader> PublicFunctions;
 				PublicGraphFunctionsProperty->ImportText_Direct(*PublicGraphFunctionsString, &PublicFunctions, nullptr, EPropertyPortFlags::PPF_None);
