Control Rig: Fix nullptr issue with GraphPinVariableBinding

#jria na

diff --git a/Engine/Plugins/Animation/ControlRig/Source/ControlRigEditor/Private/Graph/SControlRigGraphPinVariableBinding.cpp b/Engine/Plugins/Animation/ControlRig/Source/ControlRigEditor/Private/Graph/SControlRigGraphPinVariableBinding.cpp
--- a/Engine/Plugins/Animation/ControlRig/Source/ControlRigEditor/Private/Graph/SControlRigGraphPinVariableBinding.cpp
+++ b/Engine/Plugins/Animation/ControlRig/Source/ControlRigEditor/Private/Graph/SControlRigGraphPinVariableBinding.cpp
@@ -61,7 +61,7 @@ void SControlRigVariableBinding::Construct(const FArguments& InArgs)
 
 FText SControlRigVariableBinding::GetBindingText(URigVMPin* ModelPin) const
 {
-	if (ModelPin)
+	if (ModelPin && ModelPin->GetGraph())
 	{
 		const FString VariablePath = ModelPin->GetBoundVariablePath();
 		return FText::FromString(VariablePath);
