Fixed CIS compile issue.

diff --git a/Engine/Plugins/Animation/ControlRig/Source/ControlRig/Private/Rigs/RigHierarchyElements.cpp b/Engine/Plugins/Animation/ControlRig/Source/ControlRig/Private/Rigs/RigHierarchyElements.cpp
--- a/Engine/Plugins/Animation/ControlRig/Source/ControlRig/Private/Rigs/RigHierarchyElements.cpp
+++ b/Engine/Plugins/Animation/ControlRig/Source/ControlRig/Private/Rigs/RigHierarchyElements.cpp
@@ -9,8 +9,6 @@
 #include "AnimationCoreLibrary.h"
 #include "Algo/Transform.h"
 
-#include UE_INLINE_GENERATED_CPP_BY_NAME(RigHierarchyElements)
-
 ////////////////////////////////////////////////////////////////////////////////
 // FRigBaseElement
 ////////////////////////////////////////////////////////////////////////////////
