Control Rig: Fix crash during graph profiling

diff --git a/Engine/Plugins/Animation/ControlRig/Source/ControlRigDeveloper/Private/Graph/ControlRigGraphNode.cpp b/Engine/Plugins/Animation/ControlRig/Source/ControlRigDeveloper/Private/Graph/ControlRigGraphNode.cpp
--- a/Engine/Plugins/Animation/ControlRig/Source/ControlRigDeveloper/Private/Graph/ControlRigGraphNode.cpp
+++ b/Engine/Plugins/Animation/ControlRig/Source/ControlRigDeveloper/Private/Graph/ControlRigGraphNode.cpp
@@ -267,8 +267,18 @@ void UControlRigGraphNode::ReconstructNode_Internal(bool bForce)
 	// Let subclasses do any additional work
 	PostReconstructNode();
 
-	InvalidateNodeTitle();
-	OnNodePinsChanged().Broadcast();
+	if(bForce)
+	{
+		if (RigGraph)
+		{
+			RigGraph->NotifyGraphChanged();
+		}
+	}
+	else
+	{
+		InvalidateNodeTitle();
+		OnNodePinsChanged().Broadcast();
+	}
 }
 
 bool UControlRigGraphNode::IsDeprecated() const
