Control Rig: Save public function access in model

diff --git a/Engine/Plugins/Animation/ControlRig/Source/ControlRigDeveloper/Private/ControlRigBlueprint.cpp b/Engine/Plugins/Animation/ControlRig/Source/ControlRigDeveloper/Private/ControlRigBlueprint.cpp
--- a/Engine/Plugins/Animation/ControlRig/Source/ControlRigDeveloper/Private/ControlRigBlueprint.cpp
+++ b/Engine/Plugins/Animation/ControlRig/Source/ControlRigDeveloper/Private/ControlRigBlueprint.cpp
@@ -33,6 +33,7 @@
 #include "Units/ControlRigNodeWorkflow.h"
 #include "Rigs/RigControlHierarchy.h"
 #include "RigVMModel/Nodes/RigVMDispatchNode.h"
+#include "UObject/UE5MainStreamObjectVersion.h"
 
 #include UE_INLINE_GENERATED_CPP_BY_NAME(ControlRigBlueprint)
 
@@ -921,23 +922,14 @@ void UControlRigBlueprint::PostLoad()
 				{
 					if(URigVMFunctionReferenceNode* FunctionReferenceNode = Cast<URigVMFunctionReferenceNode>(Node))
 					{
-						if(const URigVMLibraryNode* DependencyNode = FunctionReferenceNode->LoadReferencedNode())
+						if(URigVMBuildData* BuildData = URigVMController::GetBuildData())
 						{
-							if(UControlRigBlueprint* DependencyBlueprint = DependencyNode->GetTypedOuter<UControlRigBlueprint>())
-							{
-								if(DependencyBlueprint != this)
-								{
-									if(URigVMBuildData* BuildData = URigVMController::GetBuildData())
-									{
-										BuildData->UpdateReferencesForFunctionReferenceNode(FunctionReferenceNode);
-									}
-								}
-							}
+							BuildData->RegisterFunctionReference(FunctionReferenceNode);
 						}
 					}
 				}
 			}
-		}	
+		}
 
 		CompileLog.Messages.Reset();
 		CompileLog.NumErrors = CompileLog.NumWarnings = 0;
@@ -2397,11 +2389,7 @@ UEdGraph* UControlRigBlueprint::GetEdGraph(const FString& InNodePath) const
 
 bool UControlRigBlueprint::IsFunctionPublic(const FName& InFunctionName) const
 {
-	if (URigVMLibraryNode* LibraryNode = GetLocalFunctionLibrary()->FindFunction(InFunctionName))
-	{
-		return GetControlRigBlueprintGeneratedClass()->GetRigVMGraphFunctionStore()->IsFunctionPublic(LibraryNode->GetFunctionIdentifier());
-	}
-	return false;
+	return GetLocalFunctionLibrary()->IsFunctionPublic(InFunctionName);	
 }
 
 void UControlRigBlueprint::MarkFunctionPublic(const FName& InFunctionName, bool bIsPublic)
@@ -2411,13 +2399,8 @@ void UControlRigBlueprint::MarkFunctionPublic(const FName& InFunctionName, bool
 		return;
 	}
 	
-	Modify();
-	
-	UControlRigBlueprintGeneratedClass* Generated = GetControlRigBlueprintGeneratedClass();
-	Generated->Modify();
-
-	const FRigVMGraphFunctionIdentifier& Identifier = GetLocalFunctionLibrary()->FindFunction(InFunctionName)->GetFunctionIdentifier();
-	Generated->GetRigVMGraphFunctionStore()->MarkFunctionAsPublic(Identifier, bIsPublic);
+	URigVMController* Controller = RigVMClient.GetOrCreateController(GetLocalFunctionLibrary());
+	Controller->MarkFunctionAsPublic(InFunctionName, bIsPublic);
 }
 
 TArray<UControlRigBlueprint*> UControlRigBlueprint::GetDependencies(bool bRecursive) const
@@ -4064,8 +4047,9 @@ void UControlRigBlueprint::PatchFunctionReferencesOnLoad()
 			Nodes.Append(Library->GetContainedNodes());
 		}
 		
-		for (URigVMNode* Node : Nodes)
+		for (int32 i=0; i<Nodes.Num(); ++i)
 		{
+			URigVMNode* Node = Nodes[i];
 			if (URigVMFunctionReferenceNode* FunctionReferenceNode = Cast<URigVMFunctionReferenceNode>(Node))
 			{
 				if (!FunctionReferenceNode->ReferencedNodePtr_DEPRECATED.IsValid())
@@ -4097,12 +4081,18 @@ void UControlRigBlueprint::PatchFunctionReferencesOnLoad()
 				{
 					FunctionReferenceNode->ReferencedFunctionHeader = FunctionReferenceNode->ReferencedNodePtr_DEPRECATED->GetFunctionHeader();					
 				}
-				else
+				else if (!FunctionReferenceNode->ReferencedNodePtr_DEPRECATED.IsNull())
 				{
 					// At least lets make sure we store the path in the header
 					FunctionReferenceNode->ReferencedFunctionHeader.LibraryPointer.LibraryNode = FunctionReferenceNode->ReferencedNodePtr_DEPRECATED.ToSoftObjectPath();
 				}
 			}
+
+			if (URigVMCollapseNode* CollapseNode = Cast<URigVMCollapseNode>(Node))
+			{
+				Nodes.Append(CollapseNode->GetContainedNodes());
+			}
+			
 		}
 	}
 	FunctionReferenceNodeData = GetReferenceNodeData();
@@ -4560,56 +4550,66 @@ void UControlRigBlueprint::PatchFunctionsOnLoad()
 	UControlRigBlueprintGeneratedClass* CRGeneratedClass = GetControlRigBlueprintGeneratedClass();
 	FRigVMGraphFunctionStore& Store = CRGeneratedClass->GraphFunctionStore;
 	const URigVMFunctionLibrary* Library = GetLocalFunctionLibrary();
+
+	TMap<URigVMLibraryNode*, FRigVMGraphFunctionHeader> OldHeaders;
+
+	// Backwards compatibility. Store public access in the model
+	TArray<FName> BackwardsCompatiblePublicFunctions;
 	if (GetLinkerCustomVersion(FControlRigObjectVersion::GUID) < FControlRigObjectVersion::StoreFunctionsInGeneratedClass)
 	{
-		// Move functions to the generated class
 		for (const FControlRigPublicFunctionData& Data : PublicFunctions_DEPRECATED)
 		{
-			const URigVMLibraryNode* LibraryNode = Library->FindFunction(Data.Name);
-			Store.AddFunction(LibraryNode->GetFunctionHeader(), true);
+			BackwardsCompatiblePublicFunctions.Add(Data.Name);
 		}
 	}
-
-	// Add functions
-	for (const URigVMLibraryNode* LibraryNode : Library->GetFunctions())
+	else
 	{
-		FRigVMGraphFunctionIdentifier Identifier = LibraryNode->GetFunctionIdentifier();
-		bool bIsPublic = PublicGraphFunctions.ContainsByPredicate([Identifier](const FRigVMGraphFunctionHeader& Header)
+		if (GetLinkerCustomVersion(FUE5MainStreamObjectVersion::GUID) < FUE5MainStreamObjectVersion::RigVMSaveFunctionAccessInModel)
 		{
-			return Header.LibraryPointer == Identifier;
-		});
-		Store.AddFunction(LibraryNode->GetFunctionHeader(), bIsPublic);		
+			for (const FRigVMGraphFunctionData& FunctionData : Store.PublicFunctions)
+			{
+				BackwardsCompatiblePublicFunctions.Add(FunctionData.Header.Name);
+				URigVMLibraryNode* LibraryNode = Cast<URigVMLibraryNode>(FunctionData.Header.LibraryPointer.LibraryNode.ResolveObject());
+				OldHeaders.Add(LibraryNode, FunctionData.Header);
+			}
+		}
 	}
+	
 
-	if (GetLinkerCustomVersion(FControlRigObjectVersion::GUID) < FControlRigObjectVersion::StoreFunctionsInGeneratedClass)
+	// Lets rebuild the FunctionStore from the model
+	if (URigVMFunctionLibrary* FunctionLibrary = GetLocalFunctionLibrary())
 	{
-		TArray<URigVMGraph*> Graphs = GetAllModels();
-		for (URigVMGraph* Graph : Graphs)
+		Store.PublicFunctions.Reset();
+		Store.PrivateFunctions.Reset();
+
+		for (URigVMLibraryNode* LibraryNode : FunctionLibrary->GetFunctions())
 		{
-			TArray<URigVMNode*> Nodes = Graph->GetNodes();
-			for (int32 i=0; i<Nodes.Num(); ++i)
+			bool bIsPublic = FunctionLibrary->IsFunctionPublic(LibraryNode->GetFName());
+			if (!bIsPublic)
 			{
-				if (URigVMFunctionReferenceNode* FuncRefNode = Cast<URigVMFunctionReferenceNode>(Nodes[i]))
+				bIsPublic = BackwardsCompatiblePublicFunctions.Contains(LibraryNode->GetFName());
+				if (bIsPublic)
 				{
-					TSoftObjectPtr<URigVMLibraryNode> ReferencedLibraryNode = FuncRefNode->ReferencedNodePtr_DEPRECATED;
-					if (ReferencedLibraryNode.IsValid())
-					{
-						FuncRefNode->ReferencedFunctionHeader = ReferencedLibraryNode->GetFunctionHeader();
-					}
+					FunctionLibrary->PublicFunctionNames.Add(LibraryNode->GetFName());
 				}
+			}
 
-				if (URigVMCollapseNode* CollapseNode = Cast<URigVMCollapseNode>(Nodes[i]))
-				{
-					Nodes.Append(CollapseNode->GetContainedNodes());
-				}
+			FRigVMGraphFunctionHeader Header = LibraryNode->GetFunctionHeader(CRGeneratedClass);
+			if (FRigVMGraphFunctionHeader* OldHeader = OldHeaders.Find(LibraryNode))
+			{				
+				Header.ExternalVariables = OldHeader->ExternalVariables;
+				Header.Dependencies = OldHeader->Dependencies;
 			}
+			Store.AddFunction(Header, bIsPublic);
+			
 		}
 	}
 
-	// Update dependencies
-	for (const URigVMLibraryNode* LibraryNode : Library->GetFunctions())
+	// Update dependencies and external variables if needed
+	for (URigVMLibraryNode* LibraryNode : Library->GetFunctions())
 	{
-		GetRigVMClient()->UpdateDependenciesForFunction(LibraryNode);		
+		GetRigVMClient()->UpdateExternalVariablesForFunction(LibraryNode);
+		GetRigVMClient()->UpdateDependenciesForFunction(LibraryNode);
 	}
 }
 
diff --git a/Engine/Source/Developer/RigVMDeveloper/Private/RigVMModel/RigVMBuildData.cpp b/Engine/Source/Developer/RigVMDeveloper/Private/RigVMModel/RigVMBuildData.cpp
--- a/Engine/Source/Developer/RigVMDeveloper/Private/RigVMModel/RigVMBuildData.cpp
+++ b/Engine/Source/Developer/RigVMDeveloper/Private/RigVMModel/RigVMBuildData.cpp
@@ -84,40 +84,6 @@ void URigVMBuildData::ForEachFunctionReferenceSoftPtr(const FRigVMGraphFunctionI
 	}
 }
 
-void URigVMBuildData::UpdateReferencesForFunctionReferenceNode(URigVMFunctionReferenceNode* InReferenceNode)
-{
-	check(InReferenceNode);
-
-	if(InReferenceNode->GetOutermost() == GetTransientPackage())
-	{
-		return;
-	}
-
-	const FRigVMGraphFunctionHeader& Header = InReferenceNode->GetReferencedFunctionHeader();
-	FRigVMFunctionReferenceArray* ReferencesEntry = GraphFunctionReferences.Find(Header.LibraryPointer);
-	if (ReferencesEntry == nullptr)
-	{
-		Modify();
-		GraphFunctionReferences.Add(Header.LibraryPointer);
-		ReferencesEntry = GraphFunctionReferences.Find(Header.LibraryPointer);
-	}
-
-	const FString ReferenceNodePathName = InReferenceNode->GetPathName();
-	for (int32 ReferenceIndex = 0; ReferenceIndex < ReferencesEntry->Num(); ReferenceIndex++)
-	{
-		const TSoftObjectPtr<URigVMFunctionReferenceNode>& Reference = ReferencesEntry->operator [](ReferenceIndex);
-		if(Reference.ToString() == ReferenceNodePathName)
-		{
-			return;
-		}
-	}
-
-	Modify();
-	ReferencesEntry->FunctionReferences.Add(InReferenceNode);
-	MarkPackageDirty();
-
-}
-
 void URigVMBuildData::RegisterFunctionReference(const FRigVMGraphFunctionIdentifier& InFunction, URigVMFunctionReferenceNode* InReference)
 {
 	if(InReference == nullptr)
diff --git a/Engine/Source/Developer/RigVMDeveloper/Private/RigVMModel/RigVMClient.cpp b/Engine/Source/Developer/RigVMDeveloper/Private/RigVMModel/RigVMClient.cpp
--- a/Engine/Source/Developer/RigVMDeveloper/Private/RigVMModel/RigVMClient.cpp
+++ b/Engine/Source/Developer/RigVMDeveloper/Private/RigVMModel/RigVMClient.cpp
@@ -775,6 +775,7 @@ void FRigVMClient::HandleGraphModifiedEvent(ERigVMGraphNotifType InNotifType, UR
 					UpdateGraphFunctionData(CollapseNode);
 				}
 			}
+			break;
 		}
 		case ERigVMGraphNotifType::NodeReferenceChanged: // A node has changed it's referenced function (Subject == URigVMFunctionReferenceNode)
 		case ERigVMGraphNotifType::LibraryTemplateChanged: // The definition of a library node's template has changed (Subject == URigVMLibraryNode)
@@ -807,6 +808,7 @@ void FRigVMClient::HandleGraphModifiedEvent(ERigVMGraphNotifType InNotifType, UR
 					}
 				}
 			}
+			break;
 		}
 
 		case ERigVMGraphNotifType::LinkAdded: // A link has been added (Subject == URigVMLink)
@@ -822,6 +824,24 @@ void FRigVMClient::HandleGraphModifiedEvent(ERigVMGraphNotifType InNotifType, UR
 					}
 				}
 			}
+
+			break;
+		}
+
+		case ERigVMGraphNotifType::FunctionAccessChanged: // A function was made public/private
+		{
+			if (URigVMLibraryNode* LibraryNode = Cast<URigVMLibraryNode>(InSubject))
+			{
+				if (URigVMFunctionLibrary* Library = Cast<URigVMFunctionLibrary>(InGraph))
+				{
+					bool bIsPublic = Library->IsFunctionPublic(LibraryNode->GetFName());
+					if (FRigVMGraphFunctionStore* Store = FindFunctionStore(LibraryNode))
+					{
+						Store->MarkFunctionAsPublic(LibraryNode->GetFunctionIdentifier(), bIsPublic);
+					}
+				}
+			}
+			break;	
 		}
 		
 		
@@ -853,10 +873,18 @@ bool FRigVMClient::UpdateFunctionReferences(const FRigVMGraphFunctionHeader& Hea
 	{
 		for (int32 i=0; i<FunctionReferenceArray->Num(); ++i)
 		{
-			TSoftObjectPtr<URigVMFunctionReferenceNode> Reference = FunctionReferenceArray->FunctionReferences[i];
+			const TSoftObjectPtr<URigVMFunctionReferenceNode>& Reference = FunctionReferenceArray->FunctionReferences[i];
+
+			// Load reference package
+			if (!Reference.IsValid())
+			{
+				Reference.LoadSynchronous();
+			}
 			if (Reference.IsValid())
 			{
 				URigVMFunctionReferenceNode* Node = Reference.Get();
+
+				Node->Modify();
 				Node->ReferencedFunctionHeader = Header;
 
 				if (bUpdateDependencies || bUpdateExternalVariables)
@@ -874,6 +902,7 @@ bool FRigVMClient::UpdateFunctionReferences(const FRigVMGraphFunctionHeader& Hea
 						}
 					}
 				}
+				Node->MarkPackageDirty();
 			}
 		}
 	}
@@ -937,35 +966,12 @@ bool FRigVMClient::DirtyGraphFunctionCompilationData(URigVMLibraryNode* InLibrar
 		FRigVMGraphFunctionIdentifier Identifier = InLibraryNode->GetFunctionIdentifier();
 		if (const FRigVMGraphFunctionData* Data = Store->FindFunction(Identifier))
 		{
-			if (Store->RemoveFunctionCompilationData(Identifier))
-			{
-				URigVMBuildData* BuildData = GetFunctionLibrary()->GetBuildData();
-				if (const FRigVMFunctionReferenceArray* FunctionReferenceArray = BuildData->FindFunctionReferences(Identifier))
-				{
-					for (int32 i=0; i<FunctionReferenceArray->Num(); ++i)
-					{
-						TSoftObjectPtr<URigVMFunctionReferenceNode> Reference = FunctionReferenceArray->FunctionReferences[i];
-						{
-							if (!Reference.IsValid())
-							{
-								Reference.LoadSynchronous();
-							}
-							if (Reference.IsValid())
-							{
-								if (URigVMFunctionLibrary* ReferenceFunctionLibrary = Reference->GetGraph()->GetDefaultFunctionLibrary())
-								{
-									if (URigVMLibraryNode* ReferenceLibraryNode = Reference.Get()->FindFunctionForNode())
-									{
-										IRigVMClientHost* OtherClientHost = ReferenceLibraryNode->GetImplementingOuter<IRigVMClientHost>();
-										OtherClientHost->GetRigVMClient()->DirtyGraphFunctionCompilationData(ReferenceLibraryNode);
-									}
-								}							
-							}
-						}
-					}
-				}
-				return true;
-			}
+			Store->RemoveFunctionCompilationData(Identifier);
+
+			// References to this function will check if the compilation hash matches, and will recompile if they
+			// see a different compilation hash. No need to dirty their compilation data.
+			
+			return true;
 		}
 	}
 
diff --git a/Engine/Source/Developer/RigVMDeveloper/Private/RigVMModel/RigVMController.cpp b/Engine/Source/Developer/RigVMDeveloper/Private/RigVMModel/RigVMController.cpp
--- a/Engine/Source/Developer/RigVMDeveloper/Private/RigVMModel/RigVMController.cpp
+++ b/Engine/Source/Developer/RigVMDeveloper/Private/RigVMModel/RigVMController.cpp
@@ -6016,6 +6016,16 @@ bool URigVMController::RemoveNode(URigVMNode* InNode, bool bSetupUndoRedo, bool
 					break;
 				}
 			}
+
+			if (FunctionLibrary->PublicFunctionNames.Contains(LibraryNode->GetFName()))
+			{
+				FunctionLibrary->PublicFunctionNames.Remove(LibraryNode->GetFName());
+
+				if (bSetupUndoRedo)
+				{
+					ActionStack->AddAction(FRigVMMarkFunctionPublicAction(LibraryNode->GetFName(), true));
+				}
+			}
 		}
 	}
 
@@ -6229,6 +6239,12 @@ bool URigVMController::RenameNode(URigVMNode* InNode, const FName& InNewName, bo
 				FRigVMControllerGraphGuard GraphGuard(this, ReferenceNode->GetGraph(), false);
                 RenameNode(ReferenceNode, InNewName, false);
 			});
+
+			if (FunctionLibrary->PublicFunctionNames.Contains(InNode->PreviousName))
+			{
+				FunctionLibrary->PublicFunctionNames.Remove(InNode->PreviousName);
+				FunctionLibrary->PublicFunctionNames.Add(ValidNewName);
+			}
 		}
 	}
 
@@ -11660,6 +11676,117 @@ bool URigVMController::RenameFunction(const FName& InOldFunctionName, const FNam
 	return RenameNode(Node, InNewFunctionName, bSetupUndoRedo);
 }
 
+bool URigVMController::MarkFunctionAsPublic(const FName& InFunctionName, bool bInIsPublic, bool bSetupUndoRedo,	bool bPrintPythonCommand)
+{
+	if (!IsValidGraph())
+	{
+		return false;
+	}
+
+	if (!bIsTransacting && !IsGraphEditable())
+	{
+		return false;
+	}
+
+	URigVMGraph* Graph = GetGraph();
+	check(Graph);
+
+
+	if (!Graph->IsA<URigVMFunctionLibrary>())
+	{
+		ReportError(TEXT("Can only change function definitions from function library graphs."));
+		return false;
+	}
+
+	URigVMNode* Node = Graph->FindNode(InFunctionName.ToString());
+	if (!Node)
+	{
+		ReportErrorf(TEXT("Could not find function called '%s'."), *InFunctionName.ToString());
+		return false;
+	}
+
+	if (URigVMFunctionLibrary* FunctionLibrary = Cast<URigVMFunctionLibrary>(Graph))
+	{
+		bool bOldIsPublic = FunctionLibrary->PublicFunctionNames.Contains(InFunctionName); 
+		if ((bInIsPublic && bOldIsPublic) || (!bInIsPublic && !bOldIsPublic))
+		{
+			return true;
+		}
+
+		if (bSetupUndoRedo)
+		{
+			FRigVMBaseAction BaseAction;
+			BaseAction.Title = FString::Printf(TEXT("Mark function %s as %s"), *InFunctionName.ToString(), (bInIsPublic) ? TEXT("Public") : TEXT("Private"));
+			ActionStack->BeginAction(BaseAction);
+			ActionStack->AddAction(FRigVMMarkFunctionPublicAction(InFunctionName, bInIsPublic));
+			ActionStack->EndAction(BaseAction);
+		}
+
+		if (bInIsPublic)
+		{
+			FunctionLibrary->PublicFunctionNames.Add(InFunctionName);
+		}
+		else
+		{
+			FunctionLibrary->PublicFunctionNames.Remove(InFunctionName);
+		}
+	}
+
+	Notify(ERigVMGraphNotifType::FunctionAccessChanged, Node);
+
+	if (!bSuspendNotifications)
+	{
+		Graph->MarkPackageDirty();
+	}
+
+	if (bPrintPythonCommand)
+	{
+		RigVMPythonUtils::Print(GetGraphOuterName(), 
+			FString::Printf(TEXT("library_controller.mark_function_as_public('%s', %s)"),
+				*GetSanitizedNodeName(InFunctionName.ToString()),
+				(bInIsPublic) ? TEXT("True") : TEXT("False")));
+	}
+
+	return true;
+}
+
+bool URigVMController::IsFunctionPublic(const FName& InFunctionName)
+{
+	if (!IsValidGraph())
+	{
+		return false;
+	}
+
+	if (!bIsTransacting && !IsGraphEditable())
+	{
+		return false;
+	}
+
+	URigVMGraph* Graph = GetGraph();
+	check(Graph);
+
+
+	if (!Graph->IsA<URigVMFunctionLibrary>())
+	{
+		ReportError(TEXT("Can only check function definitions from function library graphs."));
+		return false;
+	}
+
+	URigVMNode* Node = Graph->FindNode(InFunctionName.ToString());
+	if (!Node)
+	{
+		ReportErrorf(TEXT("Could not find function called '%s'."), *InFunctionName.ToString());
+		return false;
+	}
+
+	if (URigVMFunctionLibrary* FunctionLibrary = Cast<URigVMFunctionLibrary>(Graph))
+	{
+		return FunctionLibrary->PublicFunctionNames.Contains(InFunctionName);
+	}
+
+	return false;
+}
+
 FRigVMGraphVariableDescription URigVMController::AddLocalVariable(const FName& InVariableName, const FString& InCPPType, UObject* InCPPTypeObject, const FString& InDefaultValue, bool bSetupUndoRedo, bool bPrintPythonCommand)
 {
 	FRigVMGraphVariableDescription NewVariable;
diff --git a/Engine/Source/Developer/RigVMDeveloper/Private/RigVMModel/RigVMControllerActions.cpp b/Engine/Source/Developer/RigVMDeveloper/Private/RigVMModel/RigVMControllerActions.cpp
--- a/Engine/Source/Developer/RigVMDeveloper/Private/RigVMModel/RigVMControllerActions.cpp
+++ b/Engine/Source/Developer/RigVMDeveloper/Private/RigVMModel/RigVMControllerActions.cpp
@@ -2423,3 +2423,35 @@ bool FRigVMAddInvokeEntryNodeAction::Redo(URigVMController* InController)
 	return false;
 }
 
+FRigVMMarkFunctionPublicAction::FRigVMMarkFunctionPublicAction()
+	: FunctionName(NAME_None)
+	, bIsPublic(false)
+{
+}
+
+FRigVMMarkFunctionPublicAction::FRigVMMarkFunctionPublicAction(const FName& InFunctionName, bool bInIsPublic)
+	: FunctionName(InFunctionName)
+	, bIsPublic(bInIsPublic)
+{
+}
+
+bool FRigVMMarkFunctionPublicAction::Undo(URigVMController* InController)
+{
+	if (!FRigVMBaseAction::Undo(InController))
+	{
+		return false;
+	}
+	return InController->MarkFunctionAsPublic(FunctionName, !bIsPublic, false);
+}
+
+bool FRigVMMarkFunctionPublicAction::Redo(URigVMController* InController)
+{
+#if WITH_EDITOR
+	if (InController->MarkFunctionAsPublic(FunctionName, bIsPublic, false))
+	{
+		return FRigVMBaseAction::Redo(InController);
+	}
+#endif
+	return false;
+}
+
diff --git a/Engine/Source/Developer/RigVMDeveloper/Public/RigVMModel/RigVMBuildData.h b/Engine/Source/Developer/RigVMDeveloper/Public/RigVMModel/RigVMBuildData.h
--- a/Engine/Source/Developer/RigVMDeveloper/Public/RigVMModel/RigVMBuildData.h
+++ b/Engine/Source/Developer/RigVMDeveloper/Public/RigVMModel/RigVMBuildData.h
@@ -89,9 +89,6 @@ public:
 	void ForEachFunctionReferenceSoftPtr(const FRigVMGraphFunctionIdentifier& InFunction, TFunction<void(TSoftObjectPtr<URigVMFunctionReferenceNode>)>
 	                                     PerReferenceFunction) const;
 
-	// Update the references list for a given reference node
-	void UpdateReferencesForFunctionReferenceNode(URigVMFunctionReferenceNode* InReferenceNode);
-
 	// registers a new reference node for a given function
 	void RegisterFunctionReference(const FRigVMGraphFunctionIdentifier& InFunction, URigVMFunctionReferenceNode* InReference);
 
diff --git a/Engine/Source/Developer/RigVMDeveloper/Public/RigVMModel/RigVMController.h b/Engine/Source/Developer/RigVMDeveloper/Public/RigVMModel/RigVMController.h
--- a/Engine/Source/Developer/RigVMDeveloper/Public/RigVMModel/RigVMController.h
+++ b/Engine/Source/Developer/RigVMDeveloper/Public/RigVMModel/RigVMController.h
@@ -789,6 +789,14 @@ public:
 	UFUNCTION(BlueprintCallable, Category = RigVMController)
 	bool RenameFunction(const FName& InOldFunctionName, const FName& InNewFunctionName, bool bSetupUndoRedo = true);
 
+	// Mark a function as public/private in the function library
+	UFUNCTION(BlueprintCallable, Category = RigVMController)
+	bool MarkFunctionAsPublic(const FName& InFunctionName, bool bInIsPublic, bool bSetupUndoRedo = true, bool bPrintPythonCommand = false);
+
+	// Returns true if a function is marked as public in the function library
+	UFUNCTION(BlueprintCallable, Category = RigVMController)
+	bool IsFunctionPublic(const FName& InFunctionName);
+
 	// Add a local variable to the graph
 	UFUNCTION(BlueprintCallable, Category = RigVMController)
 	FRigVMGraphVariableDescription AddLocalVariable(const FName& InVariableName, const FString& InCPPType, UObject* InCPPTypeObject, const FString& InDefaultValue, bool bSetupUndoRedo = true, bool bPrintPythonCommand = false);
diff --git a/Engine/Source/Developer/RigVMDeveloper/Public/RigVMModel/RigVMControllerActions.h b/Engine/Source/Developer/RigVMDeveloper/Public/RigVMModel/RigVMControllerActions.h
--- a/Engine/Source/Developer/RigVMDeveloper/Public/RigVMModel/RigVMControllerActions.h
+++ b/Engine/Source/Developer/RigVMDeveloper/Public/RigVMModel/RigVMControllerActions.h
@@ -1854,3 +1854,27 @@ public:
 	UPROPERTY()
 	FString NodePath;
 };
+
+/**
+ * An action marking a function as public/private.
+ */
+USTRUCT()
+struct FRigVMMarkFunctionPublicAction : public FRigVMBaseAction
+{
+	GENERATED_BODY()
+
+public:
+
+	FRigVMMarkFunctionPublicAction();
+	FRigVMMarkFunctionPublicAction(const FName& InFunctionName, bool bInIsPublic);
+	virtual ~FRigVMMarkFunctionPublicAction() {};
+	virtual UScriptStruct* GetScriptStruct() const override { return FRigVMMarkFunctionPublicAction::StaticStruct(); }
+	virtual bool Undo(URigVMController* InController) override;
+	virtual bool Redo(URigVMController* InController) override;
+
+	UPROPERTY()
+	FName FunctionName;
+	
+	UPROPERTY()
+	bool bIsPublic;
+};
diff --git a/Engine/Source/Developer/RigVMDeveloper/Public/RigVMModel/RigVMFunctionLibrary.h b/Engine/Source/Developer/RigVMDeveloper/Public/RigVMModel/RigVMFunctionLibrary.h
--- a/Engine/Source/Developer/RigVMDeveloper/Public/RigVMModel/RigVMFunctionLibrary.h
+++ b/Engine/Source/Developer/RigVMDeveloper/Public/RigVMModel/RigVMFunctionLibrary.h
@@ -76,8 +76,12 @@ public:
 	const FSoftObjectPath GetFunctionHostObjectPath() const;
 	URigVMFunctionLibrary_GetFunctionHostObjectPath GetFunctionHostObjectPathDelegate;
 
+	bool IsFunctionPublic(const FName& InFunctionName) const { return PublicFunctionNames.Contains(InFunctionName); }
+
 private:
 
+	TArray<FName> PublicFunctionNames;
+
 #if WITH_EDITORONLY_DATA
 	UPROPERTY()
 	TMap< TObjectPtr<URigVMLibraryNode>, FRigVMFunctionReferenceArray > FunctionReferences_DEPRECATED;
diff --git a/Engine/Source/Developer/RigVMDeveloper/Public/RigVMModel/RigVMNotifications.h b/Engine/Source/Developer/RigVMDeveloper/Public/RigVMModel/RigVMNotifications.h
--- a/Engine/Source/Developer/RigVMDeveloper/Public/RigVMModel/RigVMNotifications.h
+++ b/Engine/Source/Developer/RigVMDeveloper/Public/RigVMModel/RigVMNotifications.h
@@ -51,7 +51,8 @@ enum class ERigVMGraphNotifType : uint8
 	NodeKeywordsChanged, // A node's keywords have changed (Subject == URigVMNode)
 	NodeDescriptionChanged, // A node's description has changed (Subject == URigVMNode)
 	VariableRemappingChanged, // A function reference node's remapping has changed (Subject == URigVMFunctionReferenceNode)
-	LibraryTemplateChanged, // The definition of a library node's template has changed
+	LibraryTemplateChanged, // The definition of a library node's template has changed (Subject == URigVMLibraryNode)
+	FunctionAccessChanged, // The function has been made public/private (Subject == URigVMLibraryNode)
 	Invalid // The max for this enum (used for guarding)
 };
 
diff --git a/Engine/Source/Runtime/Core/Public/UObject/UE5MainStreamObjectVersion.h b/Engine/Source/Runtime/Core/Public/UObject/UE5MainStreamObjectVersion.h
--- a/Engine/Source/Runtime/Core/Public/UObject/UE5MainStreamObjectVersion.h
+++ b/Engine/Source/Runtime/Core/Public/UObject/UE5MainStreamObjectVersion.h
@@ -300,6 +300,9 @@ struct CORE_API FUE5MainStreamObjectVersion
 
 		// Adding additional object version to defer out-of-date pose asset warning until next resaves
 		PoseAssetRawDataGUIDUpdate,
+
+		// Store function information (and compilation data) in blueprint generated class
+		RigVMSaveFunctionAccessInModel,
 		
 		// -----<new versions can be added above this line>-------------------------------------------------
 		VersionPlusOne,
diff --git a/Engine/Source/Runtime/RigVM/Private/RigVMCore/RigVMGraphFunctionHost.cpp b/Engine/Source/Runtime/RigVM/Private/RigVMCore/RigVMGraphFunctionHost.cpp
--- a/Engine/Source/Runtime/RigVM/Private/RigVMCore/RigVMGraphFunctionHost.cpp
+++ b/Engine/Source/Runtime/RigVM/Private/RigVMCore/RigVMGraphFunctionHost.cpp
@@ -193,6 +193,34 @@ bool FRigVMGraphFunctionStore::UpdateDependencies(const FRigVMGraphFunctionIdent
 {
 	if (FRigVMGraphFunctionData* Data = FindFunction(InLibraryPointer))
 	{
+		FRigVMGraphFunctionHeader& Header = Data->Header;
+		
+		// Check if they are the same
+		if (Header.Dependencies.Num() == Dependencies.Num())
+		{
+			bool bFoundDifference = false;
+			for (const TPair<FRigVMGraphFunctionIdentifier, uint32>& Pair : Dependencies)
+			{
+				if (uint32* Hash = Header.Dependencies.Find(Pair.Key))
+				{
+					if (*Hash != Pair.Value)
+					{
+						bFoundDifference = true;
+						break;
+					}
+				}
+				else
+				{
+					bFoundDifference = true;
+					break;
+				}
+			}
+			if (!bFoundDifference)
+			{
+				return false;
+			}
+		}
+		
 		Data->Header.Dependencies = Dependencies;		
 		return true;
 	}
@@ -204,10 +232,35 @@ bool FRigVMGraphFunctionStore::UpdateExternalVariables(const FRigVMGraphFunction
 	if (FRigVMGraphFunctionData* Data = FindFunction(InLibraryPointer))
 	{
 		FRigVMGraphFunctionHeader& Header = Data->Header;
+
+		// Check if they are the same
+		if (Header.ExternalVariables.Num() == ExternalVariables.Num())
+		{
+			bool bFoundDifference = false;
+			for (const FRigVMExternalVariable& Variable : ExternalVariables)
+			{
+				if (!Header.ExternalVariables.ContainsByPredicate([Variable](const FRigVMExternalVariable& ExternalVariable)
+				{
+					return Variable.Name == ExternalVariable.Name &&
+							Variable.TypeName == ExternalVariable.TypeName &&
+							Variable.TypeObject == ExternalVariable.TypeObject &&
+							Variable.bIsArray == ExternalVariable.bIsArray &&
+							Variable.bIsPublic == ExternalVariable.bIsPublic &&
+							Variable.bIsReadOnly == ExternalVariable.bIsReadOnly &&
+							Variable.Size == ExternalVariable.Size;
+				}))
+				{
+					bFoundDifference = true;
+					break;
+				}
+			}
+			if (!bFoundDifference)
+			{
+				return false;
+			}
+		}
 		
-		// Find Dependencies
 		Header.ExternalVariables = ExternalVariables;
-
 		return true;
 	}
 	return false;
@@ -218,6 +271,11 @@ bool FRigVMGraphFunctionStore::UpdateFunctionCompilationData(const FRigVMGraphFu
 	FRigVMGraphFunctionData* Info = FindFunction(InLibraryPointer);
 	if (Info)
 	{
+		if (Info->CompilationData.Hash == CompilationData.Hash)
+		{
+			return false;
+		}
+		
 		Info->CompilationData = CompilationData;
 		return true;
 	}
diff --git a/Engine/Source/Runtime/RigVM/Public/RigVMCore/RigVMGraphFunctionDefinition.h b/Engine/Source/Runtime/RigVM/Public/RigVMCore/RigVMGraphFunctionDefinition.h
--- a/Engine/Source/Runtime/RigVM/Public/RigVMCore/RigVMGraphFunctionDefinition.h
+++ b/Engine/Source/Runtime/RigVM/Public/RigVMCore/RigVMGraphFunctionDefinition.h
@@ -310,11 +310,6 @@ struct RIGVM_API FRigVMGraphFunctionIdentifier
 	FRigVMGraphFunctionIdentifier(FSoftObjectPath InHostObject, FSoftObjectPath InLibraryNode)
 		: LibraryNode(InLibraryNode), HostObject(InHostObject) {}
 
-	FORCEINLINE FString GetHash() const
-	{
-		return LibraryNode.ToString();
-	}
-
 	friend FORCEINLINE uint32 GetTypeHash(const FRigVMGraphFunctionIdentifier& Pointer)
 	{
 		return HashCombine(GetTypeHash(Pointer.LibraryNode), GetTypeHash(Pointer.HostObject));
