The editor crashes upon clicking on the Upperarm_l

diff --git a/Engine/Plugins/Animation/ControlRig/Source/ControlRig/Private/Rigs/FKControlRig.cpp b/Engine/Plugins/Animation/ControlRig/Source/ControlRig/Private/Rigs/FKControlRig.cpp
--- a/Engine/Plugins/Animation/ControlRig/Source/ControlRig/Private/Rigs/FKControlRig.cpp
+++ b/Engine/Plugins/Animation/ControlRig/Source/ControlRig/Private/Rigs/FKControlRig.cpp
@@ -27,31 +27,38 @@ FName UFKControlRig::GetControlName(const FName& InName, const ERigElementType&
 {
 	if (InName != NAME_None)
 	{
-		//@jurre to look at after vacation, causes ensure check(InType == ERigElementType::Bone || InType == ERigElementType::Curve);
-		static thread_local TMap<FName, FName> NameToControlMappings[2];
-		TMap<FName, FName>& NameToControlMapping = NameToControlMappings[InType == ERigElementType::Bone ? 0 : 1];
-		if (const FName* CachedName = NameToControlMapping.Find(InName))
+		if((InType == ERigElementType::Bone || InType == ERigElementType::Curve))
 		{
-			return *CachedName;
-		}
-		else
-		{
-			static thread_local FStringBuilderBase ScratchString;		
-
-			// ToString performs ScratchString.Reset() internally
-			InName.ToString(ScratchString);
-
-			if (InType == ERigElementType::Curve)
+			static thread_local TMap<FName, FName> NameToControlMappings[2];
+			TMap<FName, FName>& NameToControlMapping = NameToControlMappings[InType == ERigElementType::Bone ? 0 : 1];
+			if (const FName* CachedName = NameToControlMapping.Find(InName))
 			{
-				static const TCHAR* CurvePostFix = TEXT("_CURVE");
-				ScratchString.Append(CurvePostFix);
+				return *CachedName;
 			}
+			else
+			{
+				static thread_local FStringBuilderBase ScratchString;		
+
+				// ToString performs ScratchString.Reset() internally
+				InName.ToString(ScratchString);
+
+				if (InType == ERigElementType::Curve)
+				{
+					static const TCHAR* CurvePostFix = TEXT("_CURVE");
+					ScratchString.Append(CurvePostFix);
+				}
 
-			static FString ControlPostFix = TEXT("_CONTROL");
+				static FString ControlPostFix = TEXT("_CONTROL");
 		
-			ScratchString.Append(ControlPostFix);
-			return NameToControlMapping.Add(InName, FName(*ScratchString));
+				ScratchString.Append(ControlPostFix);
+				return NameToControlMapping.Add(InName, FName(*ScratchString));
+			}
 		}
+		else if (InType == ERigElementType::Control)
+		{
+			checkSlow(InName.ToString().Contains(TEXT("_CONTROL")));
+			return InName;
+		}	
 	}
 
 	// if control name is coming as none, we don't append the postfix
@@ -62,7 +69,7 @@ FName UFKControlRig::GetControlTargetName(const FName& InName, const ERigElement
 {
 	if (InName != NAME_None)
 	{
-		//@jurre to look at after vacation, causes ensure check(InType == ERigElementType::Bone || InType == ERigElementType::Curve);
+		check(InType == ERigElementType::Bone || InType == ERigElementType::Curve);
 		static thread_local TMap<FName, FName> NameToTargetMappings[2];
 		TMap<FName, FName>& NameToTargetMapping = NameToTargetMappings[InType == ERigElementType::Bone ? 0 : 1];
 		if (const FName* CachedName = NameToTargetMapping.Find(InName))
diff --git a/Engine/Plugins/Animation/ControlRig/Source/ControlRigEditor/Private/EditMode/SControlRigOutliner.cpp b/Engine/Plugins/Animation/ControlRig/Source/ControlRigEditor/Private/EditMode/SControlRigOutliner.cpp
--- a/Engine/Plugins/Animation/ControlRig/Source/ControlRigEditor/Private/EditMode/SControlRigOutliner.cpp
+++ b/Engine/Plugins/Animation/ControlRig/Source/ControlRigEditor/Private/EditMode/SControlRigOutliner.cpp
@@ -590,11 +590,11 @@ bool SMultiRigHierarchyTreeView::AddElement(UControlRig* InControlRig, const FRi
 						if(ControlElement->Settings.AnimationType == ERigControlAnimationType::AnimationControl)
 						{
 							const FRigElementKey& ElementKey = InElement->GetKey();
-							const FName BoneName = FKControlRig->GetControlTargetName(ElementKey.Name, ElementKey.Type);
+							const FName BoneName = FKControlRig->GetControlTargetName(ElementKey.Name, ParentKey.Type);
 							const FRigElementKey ParentBoneKey = Hierarchy->GetFirstParent(FRigElementKey(BoneName, ERigElementType::Bone));
 							if(ParentBoneKey.IsValid())
 							{
-								ParentKey = FRigElementKey(FKControlRig->GetControlName(ParentBoneKey.Name, ElementKey.Type), ElementKey.Type);
+								ParentKey = FRigElementKey(FKControlRig->GetControlName(ParentBoneKey.Name, ParentKey.Type), ElementKey.Type);
 							}
 						}
 					}
