Control Rig: Correctly filter dispatch nodes with execute context pins

diff --git a/Engine/Plugins/Animation/ControlRig/Source/ControlRigEditor/Private/Graph/NodeSpawners/ControlRigTemplateNodeSpawner.cpp b/Engine/Plugins/Animation/ControlRig/Source/ControlRigEditor/Private/Graph/NodeSpawners/ControlRigTemplateNodeSpawner.cpp
--- a/Engine/Plugins/Animation/ControlRig/Source/ControlRigEditor/Private/Graph/NodeSpawners/ControlRigTemplateNodeSpawner.cpp
+++ b/Engine/Plugins/Animation/ControlRig/Source/ControlRigEditor/Private/Graph/NodeSpawners/ControlRigTemplateNodeSpawner.cpp
@@ -98,6 +98,14 @@ bool UControlRigTemplateNodeSpawner::IsTemplateNodeFilteredOut(FBlueprintActionF
 							{
 								if (const FRigVMTemplate* Template = FRigVMRegistry::Get().FindTemplate(TemplateNotation))
 								{
+									if(ModelPin->IsExecuteContext())
+									{
+										if(Template->NumExecuteArguments() > 0)
+										{
+											return false;
+										}
+									}
+									
 									for (int32 i=0; i<Template->NumArguments(); ++i)
 									{
 										const FRigVMTemplateArgument* Argument = Template->GetArgument(i);
