RigVM: Remove UE_RIGVM_ENABLE_TEMPLATE_NODES define

diff --git a/Engine/Plugins/Animation/ControlRig/Source/ControlRigEditor/Private/ControlRigEditorModule.cpp b/Engine/Plugins/Animation/ControlRig/Source/ControlRigEditor/Private/ControlRigEditorModule.cpp
--- a/Engine/Plugins/Animation/ControlRig/Source/ControlRigEditor/Private/ControlRigEditorModule.cpp
+++ b/Engine/Plugins/Animation/ControlRig/Source/ControlRigEditor/Private/ControlRigEditorModule.cpp
@@ -1007,8 +1007,6 @@ void FControlRigEditorModule::GetTypeActions(UControlRigBlueprint* CRB, FBluepri
 	FRigVMRegistry& Registry = FRigVMRegistry::Get();
 	Registry.RefreshEngineTypes();
 
-#if UE_RIGVM_ENABLE_TEMPLATE_NODES
-	
 	for (const FRigVMTemplate& Template : Registry.GetTemplates())
 	{
 		// factories are registered below
@@ -1066,8 +1064,6 @@ void FControlRigEditorModule::GetTypeActions(UControlRigBlueprint* CRB, FBluepri
 		ActionRegistrar.AddBlueprintAction(ActionKey, NodeSpawner);
 	};
 
-#endif
-
 	// Add all rig units
 	for(const FRigVMFunction& Function : Registry.GetFunctions())
 	{
@@ -1077,13 +1073,11 @@ void FControlRigEditorModule::GetTypeActions(UControlRigBlueprint* CRB, FBluepri
 			continue;
 		}
 
-#if UE_RIGVM_ENABLE_TEMPLATE_NODES
 		// skip rig units which have a template
 		if (Function.GetTemplate())
 		{
 			continue;
 		}
-#endif
 
 		// skip deprecated units
 		if(Function.Struct->HasMetaData(FRigVMStruct::DeprecatedMetaName))
diff --git a/Engine/Source/Developer/RigVMDeveloper/Private/RigVMModel/RigVMController.cpp b/Engine/Source/Developer/RigVMDeveloper/Private/RigVMModel/RigVMController.cpp
--- a/Engine/Source/Developer/RigVMDeveloper/Private/RigVMModel/RigVMController.cpp
+++ b/Engine/Source/Developer/RigVMDeveloper/Private/RigVMModel/RigVMController.cpp
@@ -834,8 +834,6 @@ URigVMUnitNode* URigVMController::AddUnitNode(UScriptStruct* InScriptStruct, con
 		return nullptr;
 	}
 
-#if UE_RIGVM_ENABLE_TEMPLATE_NODES
-
 	if(const FRigVMTemplate* Template = Function->GetTemplate())
 	{
 		if(bSetupUndoRedo)
@@ -882,8 +880,6 @@ URigVMUnitNode* URigVMController::AddUnitNode(UScriptStruct* InScriptStruct, con
 		return TemplateNode;
 	}
 
-#endif
-
 	FStructOnScope StructOnScope(InScriptStruct);
 	FRigVMStruct* StructMemory = (FRigVMStruct*)StructOnScope.GetStructMemory();
 	const bool bIsEventNode = (!StructMemory->GetEventName().IsNone());
diff --git a/Engine/Source/Developer/RigVMDeveloper/Public/RigVMModel/RigVMController.h b/Engine/Source/Developer/RigVMDeveloper/Public/RigVMModel/RigVMController.h
--- a/Engine/Source/Developer/RigVMDeveloper/Public/RigVMModel/RigVMController.h
+++ b/Engine/Source/Developer/RigVMDeveloper/Public/RigVMModel/RigVMController.h
@@ -23,10 +23,6 @@
 #include "UObject/Interface.h"
 #include "RigVMController.generated.h"
 
-#ifndef UE_RIGVM_ENABLE_TEMPLATE_NODES
-#define UE_RIGVM_ENABLE_TEMPLATE_NODES 1
-#endif
-
 class URigVMActionStack;
 struct FRigVMGraphFunctionArgument;
 struct FRigVMGraphFunctionHeader;
