ControlRig, RigVM: Remove use of FORCEINLINE

diff --git a/Engine/Plugins/Animation/ControlRig/Source/ControlRig/Private/Units/Animation/RigUnit_AnimAttribute.h b/Engine/Plugins/Animation/ControlRig/Source/ControlRig/Private/Units/Animation/RigUnit_AnimAttribute.h
--- a/Engine/Plugins/Animation/ControlRig/Source/ControlRig/Private/Units/Animation/RigUnit_AnimAttribute.h
+++ b/Engine/Plugins/Animation/ControlRig/Source/ControlRig/Private/Units/Animation/RigUnit_AnimAttribute.h
@@ -191,7 +191,7 @@ protected:
 
 #if WITH_EDITOR
 	template<typename ValueType>
-	FORCEINLINE_DEBUGGABLE bool CheckArgumentTypes(FRigVMMemoryHandleArray Handles) const
+	bool CheckArgumentTypes(FRigVMMemoryHandleArray Handles) const
 	{
 		return CheckArgumentType(Handles[NameArgIndex].IsType<FName>(), NameArgName) &&
 			CheckArgumentType(Handles[BoneNameArgIndex].IsType<FName>(), BoneNameArgName) &&
@@ -202,7 +202,7 @@ protected:
 			CheckArgumentType(Handles[FoundArgIndex].IsType<bool>(), FoundArgName);
 	}
 
-	FORCEINLINE_DEBUGGABLE bool CheckArgumentTypes(FRigVMMemoryHandleArray Handles) const
+	bool CheckArgumentTypes(FRigVMMemoryHandleArray Handles) const
 	{
 		return CheckArgumentType(Handles[NameArgIndex].IsType<FName>(), NameArgName) &&
 			CheckArgumentType(Handles[BoneNameArgIndex].IsType<FName>(), BoneNameArgName) &&
@@ -214,7 +214,7 @@ protected:
 
 	// dispatch function for built-in types
 	template<typename ValueType>	
-	FORCEINLINE_DEBUGGABLE static void GetAnimAttributeDispatch(FRigVMExtendedExecuteContext& InContext, FRigVMMemoryHandleArray Handles)
+	static void GetAnimAttributeDispatch(FRigVMExtendedExecuteContext& InContext, FRigVMMemoryHandleArray Handles)
 	{
 		const FRigDispatch_GetAnimAttribute* Factory = static_cast<const FRigDispatch_GetAnimAttribute*>(InContext.Factory);
 		
@@ -265,7 +265,7 @@ protected:
 
 #if WITH_EDITOR
 	template<typename ValueType>
-	FORCEINLINE_DEBUGGABLE bool CheckArgumentTypes(FRigVMMemoryHandleArray Handles) const
+	bool CheckArgumentTypes(FRigVMMemoryHandleArray Handles) const
 	{
 		return CheckArgumentType(Handles[NameArgIndex].IsType<FName>(), NameArgName) &&
 			CheckArgumentType(Handles[BoneNameArgIndex].IsType<FName>(), BoneNameArgName) &&
@@ -275,7 +275,7 @@ protected:
 			CheckArgumentType(Handles[SuccessArgIndex].IsType<bool>(), FoundArgName);
 	}
 
-	FORCEINLINE_DEBUGGABLE bool CheckArgumentTypes(FRigVMMemoryHandleArray Handles) const
+	bool CheckArgumentTypes(FRigVMMemoryHandleArray Handles) const
 	{
 		return CheckArgumentType(Handles[NameArgIndex].IsType<FName>(), NameArgName) &&
 			CheckArgumentType(Handles[BoneNameArgIndex].IsType<FName>(), BoneNameArgName) &&
@@ -287,7 +287,7 @@ protected:
 
 	// dispatch function for built-in types
 	template<typename ValueType>	
-	FORCEINLINE_DEBUGGABLE static void SetAnimAttributeDispatch(FRigVMExtendedExecuteContext& InContext, FRigVMMemoryHandleArray Handles)
+	static void SetAnimAttributeDispatch(FRigVMExtendedExecuteContext& InContext, FRigVMMemoryHandleArray Handles)
 	{
 		const FRigDispatch_SetAnimAttribute* Factory = static_cast<const FRigDispatch_SetAnimAttribute*>(InContext.Factory);
 		
diff --git a/Engine/Plugins/Animation/ControlRig/Source/ControlRig/Private/Units/Deprecated/Math/MathLibrary.h b/Engine/Plugins/Animation/ControlRig/Source/ControlRig/Private/Units/Deprecated/Math/MathLibrary.h
--- a/Engine/Plugins/Animation/ControlRig/Source/ControlRig/Private/Units/Deprecated/Math/MathLibrary.h
+++ b/Engine/Plugins/Animation/ControlRig/Source/ControlRig/Private/Units/Deprecated/Math/MathLibrary.h
@@ -7,25 +7,25 @@
 namespace FRigMathLibrary
 {
 	template<typename Type>
-	FORCEINLINE Type Multiply(const Type Argument0, const Type Argument1)
+	Type Multiply(const Type Argument0, const Type Argument1)
 	{
 		return Argument0 * Argument1;
 	}
 
 	template<typename Type>
-	FORCEINLINE Type Add(const Type Argument0, const Type Argument1)
+	Type Add(const Type Argument0, const Type Argument1)
 	{
 		return Argument0 + Argument1;
 	}
 
 	template<typename Type>
-	FORCEINLINE Type Subtract(const Type Argument0, const Type Argument1)
+	Type Subtract(const Type Argument0, const Type Argument1)
 	{
 		return Argument0 - Argument1;
 	}
 
 	template<typename Type>
-	FORCEINLINE Type Divide(const Type Argument0, const Type Argument1)
+	Type Divide(const Type Argument0, const Type Argument1)
 	{
 		return Argument0 / Argument1;
 	}
diff --git a/Engine/Plugins/Animation/ControlRig/Source/ControlRig/Private/Units/Execution/RigUnit_Collection.h b/Engine/Plugins/Animation/ControlRig/Source/ControlRig/Private/Units/Execution/RigUnit_Collection.h
--- a/Engine/Plugins/Animation/ControlRig/Source/ControlRig/Private/Units/Execution/RigUnit_Collection.h
+++ b/Engine/Plugins/Animation/ControlRig/Source/ControlRig/Private/Units/Execution/RigUnit_Collection.h
@@ -616,7 +616,7 @@ struct CONTROLRIG_API FRigUnit_CollectionLoop : public FRigUnit_CollectionBaseMu
 		return Blocks;
 	}
 	virtual const bool IsControlFlowBlockSliced(const FName& InBlockName) const { return InBlockName == ExecuteContextName; }
-	FORCEINLINE virtual int32 GetNumSlices() const override { return Count; }
+	virtual int32 GetNumSlices() const override { return Count; }
 
 	RIGVM_METHOD()
 	virtual void Execute(const FRigUnitContext& Context) override;
diff --git a/Engine/Plugins/Animation/ControlRig/Source/ControlRig/Private/Units/Execution/RigUnit_ForLoop.h b/Engine/Plugins/Animation/ControlRig/Source/ControlRig/Private/Units/Execution/RigUnit_ForLoop.h
--- a/Engine/Plugins/Animation/ControlRig/Source/ControlRig/Private/Units/Execution/RigUnit_ForLoop.h
+++ b/Engine/Plugins/Animation/ControlRig/Source/ControlRig/Private/Units/Execution/RigUnit_ForLoop.h
@@ -29,7 +29,7 @@ struct CONTROLRIG_API FRigUnit_ForLoopCount : public FRigUnitMutable
 		return Blocks;
 	}
 	virtual const bool IsControlFlowBlockSliced(const FName& InBlockName) const { return InBlockName == ExecuteContextName; }
-	FORCEINLINE virtual int32 GetNumSlices() const override { return Count; }
+	virtual int32 GetNumSlices() const override { return Count; }
 
 	RIGVM_METHOD()
 	virtual void Execute(const FRigUnitContext& Context) override;
diff --git a/Engine/Plugins/Animation/ControlRig/Source/ControlRig/Private/Units/Hierarchy/RigUnit_Metadata.h b/Engine/Plugins/Animation/ControlRig/Source/ControlRig/Private/Units/Hierarchy/RigUnit_Metadata.h
--- a/Engine/Plugins/Animation/ControlRig/Source/ControlRig/Private/Units/Hierarchy/RigUnit_Metadata.h
+++ b/Engine/Plugins/Animation/ControlRig/Source/ControlRig/Private/Units/Hierarchy/RigUnit_Metadata.h
@@ -60,7 +60,7 @@ protected:
 
 #if WITH_EDITOR
 	template<typename ValueType>
-	FORCEINLINE_DEBUGGABLE bool CheckArgumentTypes(FRigVMMemoryHandleArray Handles) const
+	bool CheckArgumentTypes(FRigVMMemoryHandleArray Handles) const
 	{
 		return CheckArgumentType(Handles[ItemArgIndex].IsType<FRigElementKey>(), ItemArgName) &&
 			CheckArgumentType(Handles[NameArgIndex].IsType<FName>(), NameArgName) &&
@@ -72,7 +72,7 @@ protected:
 #endif
 
 	template<typename ValueType, typename MetadataType, ERigMetadataType EnumValue>
-	FORCEINLINE_DEBUGGABLE static void GetMetadataDispatch(FRigVMExtendedExecuteContext& InContext, FRigVMMemoryHandleArray Handles)
+	static void GetMetadataDispatch(FRigVMExtendedExecuteContext& InContext, FRigVMMemoryHandleArray Handles)
 	{
 		const FRigDispatch_GetMetadata* Factory = static_cast<const FRigDispatch_GetMetadata*>(InContext.Factory);
 
@@ -124,7 +124,7 @@ protected:
 
 #if WITH_EDITOR
 	template<typename ValueType>
-	FORCEINLINE_DEBUGGABLE bool CheckArgumentTypes(FRigVMMemoryHandleArray Handles) const
+	bool CheckArgumentTypes(FRigVMMemoryHandleArray Handles) const
 	{
 		return CheckArgumentType(Handles[ItemArgIndex].IsType<FRigElementKey>(), ItemArgName) &&
 			CheckArgumentType(Handles[NameArgIndex].IsType<FName>(), NameArgName) &&
@@ -135,7 +135,7 @@ protected:
 #endif
 	
 	template<typename ValueType, typename MetadataType, ERigMetadataType EnumValue>
-	FORCEINLINE_DEBUGGABLE static void SetMetadataDispatch(FRigVMExtendedExecuteContext& InContext, FRigVMMemoryHandleArray Handles)
+	static void SetMetadataDispatch(FRigVMExtendedExecuteContext& InContext, FRigVMMemoryHandleArray Handles)
 	{
 		const FRigDispatch_SetMetadata* Factory = static_cast<const FRigDispatch_SetMetadata*>(InContext.Factory);
 
diff --git a/Engine/Plugins/Animation/ControlRig/Source/ControlRig/Public/AnimNode_ControlRigBase.h b/Engine/Plugins/Animation/ControlRig/Source/ControlRig/Public/AnimNode_ControlRigBase.h
--- a/Engine/Plugins/Animation/ControlRig/Source/ControlRig/Public/AnimNode_ControlRigBase.h
+++ b/Engine/Plugins/Animation/ControlRig/Source/ControlRig/Public/AnimNode_ControlRigBase.h
@@ -19,12 +19,12 @@ struct CONTROLRIG_API FControlRigIOSettings
 		, bUpdateCurves(true)
 	{}
 
-	FORCEINLINE static FControlRigIOSettings MakeEnabled()
+	static FControlRigIOSettings MakeEnabled()
 	{
 		return FControlRigIOSettings();
 	}
 
-	FORCEINLINE static FControlRigIOSettings MakeDisabled()
+	static FControlRigIOSettings MakeDisabled()
 	{
 		FControlRigIOSettings Settings;
 		Settings.bUpdatePose = Settings.bUpdateCurves = false;
diff --git a/Engine/Plugins/Animation/ControlRig/Source/ControlRig/Public/ControlRig.h b/Engine/Plugins/Animation/ControlRig/Source/ControlRig/Public/ControlRig.h
--- a/Engine/Plugins/Animation/ControlRig/Source/ControlRig/Public/ControlRig.h
+++ b/Engine/Plugins/Animation/ControlRig/Source/ControlRig/Public/ControlRig.h
@@ -214,7 +214,7 @@ public:
 	bool SetVariableFromString(const FName& InVariableName, const FString& InValue);
 
 	template<class T>
-	FORCEINLINE T GetPublicVariableValue(const FName& InVariableName)
+	T GetPublicVariableValue(const FName& InVariableName)
 	{
 		FRigVMExternalVariable Variable = GetPublicVariableByName(InVariableName);
 		if (Variable.IsValid())
@@ -225,7 +225,7 @@ public:
 	}
 
 	template<class T>
-	FORCEINLINE void SetPublicVariableValue(const FName& InVariableName, const T& InValue)
+	void SetPublicVariableValue(const FName& InVariableName, const T& InValue)
 	{
 		FRigVMExternalVariable Variable = GetPublicVariableByName(InVariableName);
 		if (Variable.IsValid())
@@ -235,7 +235,7 @@ public:
 	}
 
 	template<class T>
-	FORCEINLINE bool SupportsEvent() const
+	bool SupportsEvent() const
 	{
 		return SupportsEvent(T::EventName);
 	}
@@ -247,13 +247,13 @@ public:
 	TArray<FName> GetSupportedEvents() const;
 
 	/** Setup bindings to a runtime object (or clear by passing in nullptr). */
-	FORCEINLINE void SetObjectBinding(TSharedPtr<IControlRigObjectBinding> InObjectBinding)
+	void SetObjectBinding(TSharedPtr<IControlRigObjectBinding> InObjectBinding)
 	{
 		ObjectBinding = InObjectBinding;
 		OnControlRigBound.Broadcast(this);
 	}
 
-	FORCEINLINE TSharedPtr<IControlRigObjectBinding> GetObjectBinding() const
+	TSharedPtr<IControlRigObjectBinding> GetObjectBinding() const
 	{
 		return ObjectBinding;
 	}
@@ -270,12 +270,12 @@ public:
 	}
 
 	UFUNCTION(BlueprintPure, Category = "Control Rig")
-	FORCEINLINE_DEBUGGABLE URigHierarchy* GetHierarchy()
+	URigHierarchy* GetHierarchy()
 	{
 		return DynamicHierarchy;
 	}
 	
-	FORCEINLINE_DEBUGGABLE URigHierarchy* GetHierarchy() const
+	URigHierarchy* GetHierarchy() const
 	{
 		return DynamicHierarchy;
 	}
@@ -357,21 +357,21 @@ public:
 
 	// Sets the relative value of a Control
 	template<class T>
-	FORCEINLINE_DEBUGGABLE void SetControlValue(const FName& InControlName, T InValue, bool bNotify = true,
+	void SetControlValue(const FName& InControlName, T InValue, bool bNotify = true,
 		const FRigControlModifiedContext& Context = FRigControlModifiedContext(), bool bSetupUndo = true, bool bPrintPythonCommnds = false, bool bFixEulerFlips = false)
 	{
 		SetControlValueImpl(InControlName, FRigControlValue::Make<T>(InValue), bNotify, Context, bSetupUndo, bPrintPythonCommnds, bFixEulerFlips);
 	}
 
 	// Returns the value of a Control
-	FORCEINLINE_DEBUGGABLE FRigControlValue GetControlValue(const FName& InControlName)
+	FRigControlValue GetControlValue(const FName& InControlName)
 	{
 		const FRigElementKey Key(InControlName, ERigElementType::Control);
 		return DynamicHierarchy->GetControlValue(Key);
 	}
 
 	// Sets the relative value of a Control
-	FORCEINLINE_DEBUGGABLE virtual void SetControlValueImpl(const FName& InControlName, const FRigControlValue& InValue, bool bNotify = true,
+	virtual void SetControlValueImpl(const FName& InControlName, const FRigControlValue& InValue, bool bNotify = true,
 		const FRigControlModifiedContext& Context = FRigControlModifiedContext(), bool bSetupUndo = true, bool bPrintPythonCommnds = false, bool bFixEulerFlips = false)
 	{
 		const FRigElementKey Key(InControlName, ERigElementType::Control);
@@ -557,7 +557,7 @@ private:
 	
 	TArray<FRigVMExternalVariable> GetExternalVariablesImpl(bool bFallbackToBlueprint) const;
 
-	FORCEINLINE FProperty* GetPublicVariableProperty(const FName& InVariableName) const
+	FProperty* GetPublicVariableProperty(const FName& InVariableName) const
 	{
 		if (FProperty* Property = GetClass()->FindPropertyByName(InVariableName))
 		{
@@ -763,43 +763,43 @@ protected:
 
 	TWeakObjectPtr<USceneComponent> OuterSceneComponent;
 
-	FORCEINLINE bool IsInitializing() const
+	bool IsInitializing() const
 	{
 		return InitBracket > 0;
 	}
 
-	FORCEINLINE bool IsExecuting() const
+	bool IsExecuting() const
 	{
 		return UpdateBracket > 0;
 	}
 
-	FORCEINLINE bool IsRunningPreConstruction() const
+	bool IsRunningPreConstruction() const
 	{
 		return PreConstructionBracket > 0;
 	}
 
-	FORCEINLINE bool IsRunningPostConstruction() const
+	bool IsRunningPostConstruction() const
 	{
 		return PostConstructionBracket > 0;
 	}
 
-	FORCEINLINE bool IsInteracting() const
+	bool IsInteracting() const
 	{
 		return InteractionBracket > 0;
 	}
 
-	FORCEINLINE uint8 GetInteractionType() const
+	uint8 GetInteractionType() const
 	{
 		return InteractionType;
 	}
 
-	FORCEINLINE bool IsSyncingWithOtherRig() const
+	bool IsSyncingWithOtherRig() const
 	{
 		return InterRigSyncBracket > 0;
 	}
 
 #if WITH_EDITOR
-	FORCEINLINE static void OnHierarchyTransformUndoRedoWeak(URigHierarchy* InHierarchy, const FRigElementKey& InKey, ERigTransformType::Type InTransformType, const FTransform& InTransform, bool bIsUndo, TWeakObjectPtr<UControlRig> WeakThis)
+	static void OnHierarchyTransformUndoRedoWeak(URigHierarchy* InHierarchy, const FRigElementKey& InKey, ERigTransformType::Type InTransformType, const FTransform& InTransform, bool bIsUndo, TWeakObjectPtr<UControlRig> WeakThis)
 	{
 		if(WeakThis.IsValid() && InHierarchy != nullptr)
 		{
@@ -920,7 +920,7 @@ public:
 	class FTransientControlPoseScope
 	{
 	public:
-		FORCEINLINE_DEBUGGABLE FTransientControlPoseScope(TObjectPtr<UControlRig> InControlRig)
+		FTransientControlPoseScope(TObjectPtr<UControlRig> InControlRig)
 		{
 			ControlRig = InControlRig;
 
@@ -936,7 +936,7 @@ public:
 				CachedPose = ControlRig->GetHierarchy()->GetPose(false, ERigElementType::Control, TArrayView<FRigElementKey>(Keys));
 			}
 		}
-		FORCEINLINE_DEBUGGABLE ~FTransientControlPoseScope()
+		~FTransientControlPoseScope()
 		{
 			check(ControlRig);
 
@@ -989,13 +989,13 @@ class CONTROLRIG_API FControlRigBracketScope
 {
 public:
 
-	FORCEINLINE FControlRigBracketScope(int32& InBracket)
+	FControlRigBracketScope(int32& InBracket)
 		: Bracket(InBracket)
 	{
 		Bracket++;
 	}
 
-	FORCEINLINE ~FControlRigBracketScope()
+	~FControlRigBracketScope()
 	{
 		Bracket--;
 	}
@@ -1009,7 +1009,7 @@ class CONTROLRIG_API FControlRigInteractionScope
 {
 public:
 
-	FORCEINLINE_DEBUGGABLE FControlRigInteractionScope(UControlRig* InControlRig)
+	FControlRigInteractionScope(UControlRig* InControlRig)
 		: ControlRig(InControlRig)
 		, InteractionBracketScope(InControlRig->InteractionBracket)
 		, SyncBracketScope(InControlRig->InterRigSyncBracket)
@@ -1017,7 +1017,7 @@ public:
 		InControlRig->GetHierarchy()->StartInteraction();
 	}
 
-	FORCEINLINE_DEBUGGABLE FControlRigInteractionScope(
+	FControlRigInteractionScope(
 		UControlRig* InControlRig,
 		const FRigElementKey& InKey,
 		EControlRigInteractionType InInteractionType = EControlRigInteractionType::All
@@ -1032,7 +1032,7 @@ public:
 		InControlRig->GetHierarchy()->StartInteraction();
 	}
 
-	FORCEINLINE_DEBUGGABLE FControlRigInteractionScope(
+	FControlRigInteractionScope(
 		UControlRig* InControlRig,
 		const TArray<FRigElementKey>& InKeys,
 		EControlRigInteractionType InInteractionType = EControlRigInteractionType::All
@@ -1047,7 +1047,7 @@ public:
 		InControlRig->GetHierarchy()->StartInteraction();
 	}
 
-	FORCEINLINE_DEBUGGABLE ~FControlRigInteractionScope()
+	~FControlRigInteractionScope()
 	{
 		if(ensure(ControlRig.IsValid()))
 		{
diff --git a/Engine/Plugins/Animation/ControlRig/Source/ControlRig/Public/Drawing/ControlRigDrawContainer.h b/Engine/Plugins/Animation/ControlRig/Source/ControlRig/Public/Drawing/ControlRigDrawContainer.h
--- a/Engine/Plugins/Animation/ControlRig/Source/ControlRig/Public/Drawing/ControlRigDrawContainer.h
+++ b/Engine/Plugins/Animation/ControlRig/Source/ControlRig/Public/Drawing/ControlRigDrawContainer.h
@@ -12,19 +12,19 @@ struct CONTROLRIG_API FControlRigDrawContainer
 	GENERATED_BODY()
 	virtual ~FControlRigDrawContainer() {}
 
-	FORCEINLINE int32 Num() const { return Instructions.Num(); }
+	int32 Num() const { return Instructions.Num(); }
 	int32 GetIndex(const FName& InName) const;
-	FORCEINLINE const FControlRigDrawInstruction& operator[](int32 InIndex) const { return Instructions[InIndex]; }
-	FORCEINLINE FControlRigDrawInstruction& operator[](int32 InIndex) { return Instructions[InIndex]; }
-	FORCEINLINE const FControlRigDrawInstruction& operator[](const FName& InName) const { return Instructions[GetIndex(InName)]; }
-	FORCEINLINE FControlRigDrawInstruction& operator[](const FName& InName) { return Instructions[GetIndex(InName)]; }
+	const FControlRigDrawInstruction& operator[](int32 InIndex) const { return Instructions[InIndex]; }
+	FControlRigDrawInstruction& operator[](int32 InIndex) { return Instructions[InIndex]; }
+	const FControlRigDrawInstruction& operator[](const FName& InName) const { return Instructions[GetIndex(InName)]; }
+	FControlRigDrawInstruction& operator[](const FName& InName) { return Instructions[GetIndex(InName)]; }
 
-	FORCEINLINE TArray<FControlRigDrawInstruction>::RangedForIteratorType      begin() { return Instructions.begin(); }
-	FORCEINLINE TArray<FControlRigDrawInstruction>::RangedForConstIteratorType begin() const { return Instructions.begin(); }
-	FORCEINLINE TArray<FControlRigDrawInstruction>::RangedForIteratorType      end() { return Instructions.end(); }
-	FORCEINLINE TArray<FControlRigDrawInstruction>::RangedForConstIteratorType end() const { return Instructions.end(); }
+	TArray<FControlRigDrawInstruction>::RangedForIteratorType      begin() { return Instructions.begin(); }
+	TArray<FControlRigDrawInstruction>::RangedForConstIteratorType begin() const { return Instructions.begin(); }
+	TArray<FControlRigDrawInstruction>::RangedForIteratorType      end() { return Instructions.end(); }
+	TArray<FControlRigDrawInstruction>::RangedForConstIteratorType end() const { return Instructions.end(); }
 
-	FORCEINLINE uint32 GetAllocatedSize(void) const { return Instructions.GetAllocatedSize(); }
+	uint32 GetAllocatedSize(void) const { return Instructions.GetAllocatedSize(); }
 
 	void Reset();
 
diff --git a/Engine/Plugins/Animation/ControlRig/Source/ControlRig/Public/Rigs/RigBoneHierarchy.h b/Engine/Plugins/Animation/ControlRig/Source/ControlRig/Public/Rigs/RigBoneHierarchy.h
--- a/Engine/Plugins/Animation/ControlRig/Source/ControlRig/Public/Rigs/RigBoneHierarchy.h
+++ b/Engine/Plugins/Animation/ControlRig/Source/ControlRig/Public/Rigs/RigBoneHierarchy.h
@@ -53,12 +53,12 @@ struct CONTROLRIG_API FRigBone: public FRigElement
 	UPROPERTY(BlueprintReadWrite, VisibleAnywhere, Category = FRigElement)
 	ERigBoneType Type;
 
-	FORCEINLINE virtual ERigElementType GetElementType() const override
+	virtual ERigElementType GetElementType() const override
 	{
 		return ERigElementType::Bone;
 	}
 
-	FORCEINLINE virtual FRigElementKey GetParentElementKey(bool bForce = false) const
+	virtual FRigElementKey GetParentElementKey(bool bForce = false) const
 	{
 		return FRigElementKey(ParentName, GetElementType());
 	}
@@ -71,10 +71,10 @@ struct CONTROLRIG_API FRigBoneHierarchy
 
 	FRigBoneHierarchy();
 
-	FORCEINLINE TArray<FRigBone>::RangedForIteratorType      begin()       { return Bones.begin(); }
-	FORCEINLINE TArray<FRigBone>::RangedForConstIteratorType begin() const { return Bones.begin(); }
-	FORCEINLINE TArray<FRigBone>::RangedForIteratorType      end()         { return Bones.end();   }
-	FORCEINLINE TArray<FRigBone>::RangedForConstIteratorType end() const   { return Bones.end();   }
+	TArray<FRigBone>::RangedForIteratorType      begin()       { return Bones.begin(); }
+	TArray<FRigBone>::RangedForConstIteratorType begin() const { return Bones.begin(); }
+	TArray<FRigBone>::RangedForIteratorType      end()         { return Bones.end();   }
+	TArray<FRigBone>::RangedForConstIteratorType end() const   { return Bones.end();   }
 
 	FRigBone& Add(const FName& InNewName, const FName& InParentName, ERigBoneType InType, const FTransform& InInitTransform, const FTransform& InLocalTransform, const FTransform& InGlobalTransform);
 	
diff --git a/Engine/Plugins/Animation/ControlRig/Source/ControlRig/Public/Rigs/RigControlHierarchy.h b/Engine/Plugins/Animation/ControlRig/Source/ControlRig/Public/Rigs/RigControlHierarchy.h
--- a/Engine/Plugins/Animation/ControlRig/Source/ControlRig/Public/Rigs/RigControlHierarchy.h
+++ b/Engine/Plugins/Animation/ControlRig/Source/ControlRig/Public/Rigs/RigControlHierarchy.h
@@ -154,27 +154,27 @@ struct CONTROLRIG_API FRigControl : public FRigElement
 	UPROPERTY(BlueprintReadOnly, EditAnywhere, Category = Control)
 	TObjectPtr<UEnum> ControlEnum;
 
-	FORCEINLINE_DEBUGGABLE virtual ERigElementType GetElementType() const override
+	virtual ERigElementType GetElementType() const override
 	{
 		return ERigElementType::Control;
 	}
 
-	FORCEINLINE_DEBUGGABLE const FName& GetDisplayName() const
+	const FName& GetDisplayName() const
 	{
 		return DisplayName.IsNone() ? Name : DisplayName;
 	}
 
-	FORCEINLINE_DEBUGGABLE virtual FRigElementKey GetParentElementKey() const
+	virtual FRigElementKey GetParentElementKey() const
 	{
 		return FRigElementKey(ParentName, GetElementType());
 	}
 
-	FORCEINLINE_DEBUGGABLE virtual FRigElementKey GetSpaceElementKey() const
+	virtual FRigElementKey GetSpaceElementKey() const
 	{
 		return FRigElementKey(SpaceName, ERigElementType::Null);
 	}
 
-	FORCEINLINE_DEBUGGABLE const FRigControlValue& GetValue(ERigControlValueType InValueType = ERigControlValueType::Current) const
+	const FRigControlValue& GetValue(ERigControlValueType InValueType = ERigControlValueType::Current) const
 	{
 		switch(InValueType)
 		{
@@ -198,7 +198,7 @@ struct CONTROLRIG_API FRigControl : public FRigElement
 		return Value;
 	}
 
-	FORCEINLINE_DEBUGGABLE FRigControlValue& GetValue(ERigControlValueType InValueType = ERigControlValueType::Current)
+	FRigControlValue& GetValue(ERigControlValueType InValueType = ERigControlValueType::Current)
 	{
 		switch(InValueType)
 		{
@@ -224,7 +224,7 @@ struct CONTROLRIG_API FRigControl : public FRigElement
 
 	void ApplyLimits(FRigControlValue& InOutValue) const;
 
-	FORCEINLINE_DEBUGGABLE static FProperty* FindPropertyForValueType(ERigControlValueType InValueType)
+	static FProperty* FindPropertyForValueType(ERigControlValueType InValueType)
 	{
 		switch (InValueType)
 		{
@@ -260,11 +260,11 @@ struct CONTROLRIG_API FRigControlHierarchy
 
 	FRigControlHierarchy();
 
-	FORCEINLINE_DEBUGGABLE int32 Num() const { return Controls.Num(); }
-	FORCEINLINE_DEBUGGABLE TArray<FRigControl>::RangedForIteratorType      begin()       { return Controls.begin(); }
-	FORCEINLINE_DEBUGGABLE TArray<FRigControl>::RangedForConstIteratorType begin() const { return Controls.begin(); }
-	FORCEINLINE_DEBUGGABLE TArray<FRigControl>::RangedForIteratorType      end()         { return Controls.end();   }
-	FORCEINLINE_DEBUGGABLE TArray<FRigControl>::RangedForConstIteratorType end() const   { return Controls.end();   }
+	int32 Num() const { return Controls.Num(); }
+	TArray<FRigControl>::RangedForIteratorType      begin()       { return Controls.begin(); }
+	TArray<FRigControl>::RangedForConstIteratorType begin() const { return Controls.begin(); }
+	TArray<FRigControl>::RangedForIteratorType      end()         { return Controls.end();   }
+	TArray<FRigControl>::RangedForConstIteratorType end() const   { return Controls.end();   }
 
 	FRigControl& Add(
 	    const FName& InNewName,
diff --git a/Engine/Plugins/Animation/ControlRig/Source/ControlRig/Public/Rigs/RigCurveContainer.h b/Engine/Plugins/Animation/ControlRig/Source/ControlRig/Public/Rigs/RigCurveContainer.h
--- a/Engine/Plugins/Animation/ControlRig/Source/ControlRig/Public/Rigs/RigCurveContainer.h
+++ b/Engine/Plugins/Animation/ControlRig/Source/ControlRig/Public/Rigs/RigCurveContainer.h
@@ -27,7 +27,7 @@ struct FRigCurve : public FRigElement
 	UPROPERTY(BlueprintReadWrite, VisibleAnywhere, Category = FRigElement)
 	float Value;
 
-	FORCEINLINE virtual ERigElementType GetElementType() const override
+	virtual ERigElementType GetElementType() const override
 	{
 		return ERigElementType::Curve;
 	}
@@ -42,10 +42,10 @@ public:
 
 	FRigCurveContainer();
 
-	FORCEINLINE TArray<FRigCurve>::RangedForIteratorType      begin()       { return Curves.begin(); }
-	FORCEINLINE TArray<FRigCurve>::RangedForConstIteratorType begin() const { return Curves.begin(); }
-	FORCEINLINE TArray<FRigCurve>::RangedForIteratorType      end()         { return Curves.end();   }
-	FORCEINLINE TArray<FRigCurve>::RangedForConstIteratorType end() const   { return Curves.end();   }
+	TArray<FRigCurve>::RangedForIteratorType      begin()       { return Curves.begin(); }
+	TArray<FRigCurve>::RangedForConstIteratorType begin() const { return Curves.begin(); }
+	TArray<FRigCurve>::RangedForIteratorType      end()         { return Curves.end();   }
+	TArray<FRigCurve>::RangedForConstIteratorType end() const   { return Curves.end();   }
 
 	FRigCurve& Add(const FName& InNewName, float InValue);
 	
