ControlRig, RigVM: Remove use of FORCEINLINE

diff --git a/Engine/Plugins/Animation/ControlRig/Source/ControlRig/Public/Rigs/RigHierarchy.h b/Engine/Plugins/Animation/ControlRig/Source/ControlRig/Public/Rigs/RigHierarchy.h
--- a/Engine/Plugins/Animation/ControlRig/Source/ControlRig/Public/Rigs/RigHierarchy.h
+++ b/Engine/Plugins/Animation/ControlRig/Source/ControlRig/Public/Rigs/RigHierarchy.h
@@ -38,7 +38,7 @@ struct FRigTransformStackEntry
 {
 	GENERATED_BODY()
 
-	FORCEINLINE FRigTransformStackEntry()
+	FRigTransformStackEntry()
 		: Key()
 		, EntryType(ERigTransformStackEntryType::TransformPose)
 		, TransformType(ERigTransformType::CurrentLocal)
@@ -48,7 +48,7 @@ struct FRigTransformStackEntry
 		, Callstack() 
 	{}
 
-	FORCEINLINE FRigTransformStackEntry(
+	FRigTransformStackEntry(
 		const FRigElementKey& InKey,
 		ERigTransformStackEntryType InEntryType,
 		ERigTransformType::Type InTransformType,
@@ -201,7 +201,7 @@ public:
 	 * @return The number of elements in the Hierarchy
 	 */
 	UFUNCTION(BlueprintCallable, Category = URigHierarchy)
-	FORCEINLINE int32 Num() const
+	int32 Num() const
 	{
 		return Elements.Num();
 	}
@@ -214,14 +214,14 @@ public:
     int32 Num(ERigElementType InElementType) const;
 
 	// iterators
-	FORCEINLINE TArray<FRigBaseElement*>::RangedForIteratorType      begin() { return Elements.begin(); }
-	FORCEINLINE TArray<FRigBaseElement*>::RangedForIteratorType      end() { return Elements.end(); }
+	TArray<FRigBaseElement*>::RangedForIteratorType      begin() { return Elements.begin(); }
+	TArray<FRigBaseElement*>::RangedForIteratorType      end() { return Elements.end(); }
 
 	/**
 	 * Iterator function to invoke a lambda / TFunction for each element
 	 * @param PerElementFunction The function to invoke for each element
 	 */
-	FORCEINLINE_DEBUGGABLE void ForEach(TFunction<bool(FRigBaseElement*)> PerElementFunction) const
+	void ForEach(TFunction<bool(FRigBaseElement*)> PerElementFunction) const
 	{
 		for (int32 ElementIndex = 0; ElementIndex < Elements.Num(); ElementIndex++)
 		{
@@ -237,7 +237,7 @@ public:
 	 * @param PerElementFunction The function to invoke for each element of a given type
 	 */
 	template<typename T>
-	FORCEINLINE_DEBUGGABLE void ForEach(TFunction<bool(T*)> PerElementFunction) const
+	void ForEach(TFunction<bool(T*)> PerElementFunction) const
 	{
 		for (int32 ElementIndex = 0; ElementIndex < Elements.Num(); ElementIndex++)
 		{
@@ -257,7 +257,7 @@ public:
 	 * @return Returns true if the provided element index is valid
 	 */
 	UFUNCTION(BlueprintCallable, Category = URigHierarchy)
-	FORCEINLINE bool IsValidIndex(int32 InElementIndex) const
+	bool IsValidIndex(int32 InElementIndex) const
 	{
 		return Elements.IsValidIndex(InElementIndex);
 	}
@@ -268,7 +268,7 @@ public:
 	 * @return Returns true if the provided element key is valid
 	 */
 	UFUNCTION(BlueprintCallable, Category = URigHierarchy, meta = (DisplayName = "Contains", ScriptName = "Contains"))
-	FORCEINLINE bool Contains_ForBlueprint(FRigElementKey InKey) const
+	bool Contains_ForBlueprint(FRigElementKey InKey) const
 	{
 		return Contains(InKey);
 	}
@@ -278,7 +278,7 @@ public:
 	 * @param InKey The key to validate
 	 * @return Returns true if the provided element key is valid
 	 */
-	FORCEINLINE bool Contains(const FRigElementKey& InKey) const
+	bool Contains(const FRigElementKey& InKey) const
 	{
 		return GetIndex(InKey) != INDEX_NONE;
 	}
@@ -289,7 +289,7 @@ public:
 	 * @return Returns true if the provided element key is valid
 	 */
 	template<typename T>
-	FORCEINLINE bool Contains(const FRigElementKey& InKey) const
+	bool Contains(const FRigElementKey& InKey) const
 	{
 		return Find<T>(InKey) != nullptr;
 	}
@@ -315,7 +315,7 @@ public:
 	 * @return The index of the element or INDEX_NONE
 	 */
 	UFUNCTION(BlueprintCallable, Category = URigHierarchy, meta = (DisplayName = "Get Index", ScriptName = "GetIndex"))
-	FORCEINLINE int32 GetIndex_ForBlueprint(FRigElementKey InKey) const
+	int32 GetIndex_ForBlueprint(FRigElementKey InKey) const
 	{
 		return GetIndex(InKey);
 	}
@@ -325,7 +325,7 @@ public:
 	 * @param InKey The key of the element to retrieve the index for
 	 * @return The index of the element or INDEX_NONE
 	 */
-	FORCEINLINE int32 GetIndex(const FRigElementKey& InKey) const
+	int32 GetIndex(const FRigElementKey& InKey) const
 	{
 		if(const int32* Index = IndexLookup.Find(InKey))
 		{
@@ -340,7 +340,7 @@ public:
 	 * @return The index of the element or INDEX_NONE
 	 */
 	UFUNCTION(BlueprintCallable, Category = URigHierarchy, meta = (DisplayName = "Get Local Index", ScriptName = "GetLocalIndex"))
-	FORCEINLINE int32 GetLocalIndex_ForBlueprint(FRigElementKey InKey) const
+	int32 GetLocalIndex_ForBlueprint(FRigElementKey InKey) const
 	{
 		return GetLocalIndex(InKey);
 	}
@@ -350,7 +350,7 @@ public:
 	 * @param InKey The key of the element to retrieve the index for
 	 * @return The index of the element or INDEX_NONE
 	 */
-	FORCEINLINE int32 GetLocalIndex(const FRigElementKey& InKey) const
+	int32 GetLocalIndex(const FRigElementKey& InKey) const
 	{
 		return GetLocalIndex(Find(InKey));
 	}
@@ -360,7 +360,7 @@ public:
 	 * @param InKeys The keys of the elements to retrieve the indices for
 	 * @return The indices of the elements or INDEX_NONE
 	 */
-	FORCEINLINE TArray<int32> GetIndices(const TArray<FRigElementKey>& InKeys) const
+	TArray<int32> GetIndices(const TArray<FRigElementKey>& InKeys) const
 	{
 		TArray<int32> Indices;
 		for(const FRigElementKey& Key : InKeys)
@@ -376,7 +376,7 @@ public:
 	 * @return The key of an element given its index
 	 */
 	UFUNCTION(BlueprintCallable, Category = URigHierarchy)
-    FORCEINLINE FRigElementKey GetKey(int32 InElementIndex) const
+    FRigElementKey GetKey(int32 InElementIndex) const
 	{
 		if(Elements.IsValidIndex(InElementIndex))
 		{
@@ -391,7 +391,7 @@ public:
 	 * @return The keys of the elements given the indices
 	 */
 	UFUNCTION(BlueprintCallable, Category = URigHierarchy)
-	FORCEINLINE TArray<FRigElementKey> GetKeys(const TArray<int32> InElementIndices) const
+	TArray<FRigElementKey> GetKeys(const TArray<int32> InElementIndices) const
 	{
 		TArray<FRigElementKey> Keys;
 		for(int32 Index : InElementIndices)
@@ -406,7 +406,7 @@ public:
 	 * @param InIndex The index of the element to retrieve.
 	 * @return The retrieved element or nullptr.
 	 */
-	FORCEINLINE const FRigBaseElement* Get(int32 InIndex) const
+	const FRigBaseElement* Get(int32 InIndex) const
 	{
 		if(Elements.IsValidIndex(InIndex))
 		{
@@ -437,7 +437,7 @@ public:
 	 * @return The retrieved element or nullptr.
 	 */
 	template<typename T>
-	FORCEINLINE const T* Get(int32 InIndex) const
+	const T* Get(int32 InIndex) const
 	{
 		return Cast<T>(Get(InIndex));
 	}
@@ -450,7 +450,7 @@ public:
 	 * @return The retrieved element or nullptr.
 	 */
 	template<typename T>
-	FORCEINLINE T* Get(int32 InIndex)
+	T* Get(int32 InIndex)
 	{
 		return Cast<T>(Get(InIndex));
 	}
@@ -463,7 +463,7 @@ public:
 	 * @return The retrieved element or nullptr.
 	 */
 	template<typename T>
-	FORCEINLINE const T* GetChecked(int32 InIndex) const
+	const T* GetChecked(int32 InIndex) const
 	{
 		return CastChecked<T>(Get(InIndex));
 	}
@@ -476,7 +476,7 @@ public:
 	 * @return The retrieved element or nullptr.
 	 */
 	template<typename T>
-    FORCEINLINE T* GetChecked(int32 InIndex)
+    T* GetChecked(int32 InIndex)
 	{
 		return CastChecked<T>(Get(InIndex));
 	}
@@ -486,7 +486,7 @@ public:
      * @param InKey The key of the handle to retrieve.
      * @return The retrieved handle (may be invalid)
      */
-	FORCEINLINE FRigElementHandle GetHandle(const FRigElementKey& InKey) const
+	FRigElementHandle GetHandle(const FRigElementKey& InKey) const
 	{
 		if(Contains(InKey))
 		{
@@ -500,7 +500,7 @@ public:
 	 * @param InKey The key of the element to retrieve.
 	 * @return The retrieved element or nullptr.
 	 */
-	FORCEINLINE const FRigBaseElement* Find(const FRigElementKey& InKey) const
+	const FRigBaseElement* Find(const FRigElementKey& InKey) const
 	{
 		return Get(GetIndex(InKey));
 	}
@@ -510,7 +510,7 @@ public:
 	 * @param InKey The key of the element to retrieve.
 	 * @return The retrieved element or nullptr.
 	 */
-	FORCEINLINE FRigBaseElement* Find(const FRigElementKey& InKey)
+	FRigBaseElement* Find(const FRigElementKey& InKey)
 	{
 		return Get(GetIndex(InKey));
 	}
@@ -520,7 +520,7 @@ public:
 	 * @param InKey The key of the element to retrieve.
 	 * @return The retrieved element or nullptr.
 	 */
-	FORCEINLINE const FRigBaseElement* FindChecked(const FRigElementKey& InKey) const
+	const FRigBaseElement* FindChecked(const FRigElementKey& InKey) const
 	{
 		const FRigBaseElement* Element = Get(GetIndex(InKey));
 		check(Element);
@@ -532,7 +532,7 @@ public:
 	 * @param InKey The key of the element to retrieve.
 	 * @return The retrieved element or nullptr.
 	 */
-	FORCEINLINE FRigBaseElement* FindChecked(const FRigElementKey& InKey)
+	FRigBaseElement* FindChecked(const FRigElementKey& InKey)
 	{
 		FRigBaseElement* Element = Get(GetIndex(InKey));
 		check(Element);
@@ -547,7 +547,7 @@ public:
 	 * @return The retrieved element or nullptr.
 	 */
 	template<typename T>
-	FORCEINLINE const T* Find(const FRigElementKey& InKey) const
+	const T* Find(const FRigElementKey& InKey) const
 	{
 		return Get<T>(GetIndex(InKey));
 	}
@@ -560,7 +560,7 @@ public:
 	 * @return The retrieved element or nullptr.
 	 */
 	template<typename T>
-	FORCEINLINE T* Find(const FRigElementKey& InKey)
+	T* Find(const FRigElementKey& InKey)
 	{
 		return Get<T>(GetIndex(InKey));
 	}
@@ -617,7 +617,7 @@ public:
 	 * @return The retrieved element or nullptr.
 	 */
 	template<typename T>
-	FORCEINLINE const T* FindChecked(const FRigElementKey& InKey) const
+	const T* FindChecked(const FRigElementKey& InKey) const
 	{
 		return GetChecked<T>(GetIndex(InKey));
 	}
@@ -630,7 +630,7 @@ public:
 	 * @return The retrieved element or nullptr.
 	 */
 	template<typename T>
-    FORCEINLINE T* FindChecked(const FRigElementKey& InKey)
+    T* FindChecked(const FRigElementKey& InKey)
 	{
 		return GetChecked<T>(GetIndex(InKey));
 	}
@@ -640,7 +640,7 @@ public:
 	 * @param bTraverse Returns the elements in order of a depth first traversal
 	 */
 	template<typename T>
-	FORCEINLINE TArray<T*> GetElementsOfType(bool bTraverse = false) const
+	TArray<T*> GetElementsOfType(bool bTraverse = false) const
 	{
 		TArray<T*> Results;
 
@@ -681,7 +681,7 @@ public:
 	 * @param bTraverse Returns the element keys in order of a depth first traversal
 	 */
 	template<typename T>
-    FORCEINLINE TArray<FRigElementKey> GetKeysOfType(bool bTraverse = false) const
+    TArray<FRigElementKey> GetKeysOfType(bool bTraverse = false) const
 	{
 		TArray<FRigElementKey> Keys;
 		TArray<T*> Results = GetElementsOfType<T>(bTraverse);
@@ -698,7 +698,7 @@ public:
 	 * @param bTraverse Returns the elements in order of a depth first traversal
 	 */
 	template<typename T>
-    FORCEINLINE TArray<T*> GetFilteredElements(TFunction<bool(T*)> InKeepElementFunction, bool bTraverse = false) const
+    TArray<T*> GetFilteredElements(TFunction<bool(T*)> InKeepElementFunction, bool bTraverse = false) const
 	{
 		TArray<T*> Results;
 
@@ -744,7 +744,7 @@ public:
 	 * Returns all Bone elements
 	 * @param bTraverse Returns the elements in order of a depth first traversal
 	 */
-	FORCEINLINE TArray<FRigBoneElement*> GetBones(bool bTraverse = false) const
+	TArray<FRigBoneElement*> GetBones(bool bTraverse = false) const
 	{
 		return GetElementsOfType<FRigBoneElement>(bTraverse);
 	}
@@ -754,7 +754,7 @@ public:
 	 * @param bTraverse Returns the elements in order of a depth first traversal
 	 */
 	UFUNCTION(BlueprintCallable, Category = URigHierarchy, meta = (DisplayName = "Get Bones", ScriptName = "GetBones"))
-	FORCEINLINE TArray<FRigElementKey> GetBoneKeys(bool bTraverse = true) const
+	TArray<FRigElementKey> GetBoneKeys(bool bTraverse = true) const
 	{
 		return GetKeysOfType<FRigBoneElement>(bTraverse);
 	}
@@ -763,7 +763,7 @@ public:
 	 * Returns all Null elements
 	 * @param bTraverse Returns the elements in order of a depth first traversal
 	 */
-	FORCEINLINE TArray<FRigNullElement*> GetNulls(bool bTraverse = false) const
+	TArray<FRigNullElement*> GetNulls(bool bTraverse = false) const
 	{
 		return GetElementsOfType<FRigNullElement>(bTraverse);
 	}
@@ -773,7 +773,7 @@ public:
 	* @param bTraverse Returns the elements in order of a depth first traversal
 	*/
 	UFUNCTION(BlueprintCallable, Category = URigHierarchy, meta = (DisplayName = "Get Nulls", ScriptName = "GetNulls"))
-	FORCEINLINE TArray<FRigElementKey> GetNullKeys(bool bTraverse = true) const
+	TArray<FRigElementKey> GetNullKeys(bool bTraverse = true) const
 	{
 		return GetKeysOfType<FRigNullElement>(bTraverse);
 	}
@@ -782,7 +782,7 @@ public:
 	 * Returns all Control elements
 	 * @param bTraverse Returns the elements in order of a depth first traversal
 	 */
-	FORCEINLINE TArray<FRigControlElement*> GetControls(bool bTraverse = false) const
+	TArray<FRigControlElement*> GetControls(bool bTraverse = false) const
 	{
 		return GetElementsOfType<FRigControlElement>(bTraverse);
 	}
@@ -792,7 +792,7 @@ public:
 	 * @param bTraverse Returns the elements in order of a depth first traversal
 	 */
 	UFUNCTION(BlueprintCallable, Category = URigHierarchy, meta = (DisplayName = "Get Controls", ScriptName = "GetControls"))
-	FORCEINLINE TArray<FRigElementKey> GetControlKeys(bool bTraverse = true) const
+	TArray<FRigElementKey> GetControlKeys(bool bTraverse = true) const
 	{
 		return GetKeysOfType<FRigControlElement>(bTraverse);
 	}
@@ -800,7 +800,7 @@ public:
 	/**
 	 * Returns all transient Control elements
 	 */
-	FORCEINLINE TArray<FRigControlElement*> GetTransientControls() const
+	TArray<FRigControlElement*> GetTransientControls() const
 	{
 		return GetFilteredElements<FRigControlElement>([](FRigControlElement* ControlElement) -> bool
 		{
@@ -811,7 +811,7 @@ public:
 	/**
 	 * Returns all Curve elements
 	 */
-	FORCEINLINE TArray<FRigCurveElement*> GetCurves() const
+	TArray<FRigCurveElement*> GetCurves() const
 	{
 		return GetElementsOfType<FRigCurveElement>();
 	}
@@ -820,7 +820,7 @@ public:
 	 * Returns all Curve elements
 	 */
 	UFUNCTION(BlueprintCallable, Category = URigHierarchy, meta = (DisplayName = "Get Curves", ScriptName = "GetCurves"))
-	FORCEINLINE TArray<FRigElementKey> GetCurveKeys() const
+	TArray<FRigElementKey> GetCurveKeys() const
 	{
 		return GetKeysOfType<FRigCurveElement>(false);
 	}
@@ -829,7 +829,7 @@ public:
 	 * Returns all RigidBody elements
 	 * @param bTraverse Returns the elements in order of a depth first traversal
 	 */
-	FORCEINLINE TArray<FRigRigidBodyElement*> GetRigidBodies(bool bTraverse = false) const
+	TArray<FRigRigidBodyElement*> GetRigidBodies(bool bTraverse = false) const
 	{
 		return GetElementsOfType<FRigRigidBodyElement>(bTraverse);
 	}
@@ -839,7 +839,7 @@ public:
 	 * @param bTraverse Returns the elements in order of a depth first traversal
 	 */
 	UFUNCTION(BlueprintCallable, Category = URigHierarchy, meta = (DisplayName = "Get RigidBodies", ScriptName = "GetRigidBodies"))
-    FORCEINLINE TArray<FRigElementKey> GetRigidBodyKeys(bool bTraverse = true) const
+    TArray<FRigElementKey> GetRigidBodyKeys(bool bTraverse = true) const
 	{
 		return GetKeysOfType<FRigRigidBodyElement>(bTraverse);
 	}
@@ -848,7 +848,7 @@ public:
 	 * Returns all references
 	 * @param bTraverse Returns the elements in order of a depth first traversal
 	 */
-	FORCEINLINE TArray<FRigReferenceElement*> GetReferences(bool bTraverse = false) const
+	TArray<FRigReferenceElement*> GetReferences(bool bTraverse = false) const
 	{
 		return GetElementsOfType<FRigReferenceElement>(bTraverse);
 	}
@@ -858,7 +858,7 @@ public:
 	 * @param bTraverse Returns the elements in order of a depth first traversal
 	 */
 	UFUNCTION(BlueprintCallable, Category = URigHierarchy, meta = (DisplayName = "Get References", ScriptName = "GetReferences"))
-    FORCEINLINE TArray<FRigElementKey> GetReferenceKeys(bool bTraverse = true) const
+    TArray<FRigElementKey> GetReferenceKeys(bool bTraverse = true) const
 	{
 		return GetKeysOfType<FRigReferenceElement>(bTraverse);
 	}
@@ -866,7 +866,7 @@ public:
 	/**
 	 * Returns all root elements
 	 */
-	FORCEINLINE TArray<FRigBaseElement*> GetRootElements() const
+	TArray<FRigBaseElement*> GetRootElements() const
 	{
 		return GetFilteredElements<FRigBaseElement>([this](FRigBaseElement* Element) -> bool
 		{
@@ -878,7 +878,7 @@ public:
 	 * Returns all root element keys
 	 */
 	UFUNCTION(BlueprintCallable, Category = URigHierarchy, meta = (DisplayName = "Get Root Elements", ScriptName = "GetRootElements"))
-	FORCEINLINE TArray<FRigElementKey> GetRootElementKeys() const
+	TArray<FRigElementKey> GetRootElementKeys() const
 	{
 		return GetKeysByPredicate([this](const FRigBaseElement& Element) -> bool
 		{
@@ -923,7 +923,7 @@ public:
 	 * @param DefaultValue The default value to fall back on
 	 */
 	UFUNCTION(BlueprintPure, Category = URigHierarchy)
-	FORCEINLINE bool GetBoolMetadata(FRigElementKey InItem, FName InMetadataName, bool DefaultValue) const
+	bool GetBoolMetadata(FRigElementKey InItem, FName InMetadataName, bool DefaultValue) const
 	{
 		return GetMetadata<bool>(InItem, ERigMetadataType::Bool, InMetadataName, DefaultValue);
 	}
@@ -934,7 +934,7 @@ public:
 	 * @param InMetadataName The name of the metadata to query
 	 */
 	UFUNCTION(BlueprintPure, Category = URigHierarchy)
-	FORCEINLINE TArray<bool> GetBoolArrayMetadata(FRigElementKey InItem, FName InMetadataName) const
+	TArray<bool> GetBoolArrayMetadata(FRigElementKey InItem, FName InMetadataName) const
 	{
 		return GetArrayMetadata<bool>(InItem, ERigMetadataType::BoolArray, InMetadataName);
 	}
@@ -947,7 +947,7 @@ public:
 	 * @return Returns true if setting the metadata was successful
 	 */
 	UFUNCTION(BlueprintCallable, Category = URigHierarchy)
-	FORCEINLINE bool SetBoolMetadata(FRigElementKey InItem, FName InMetadataName, bool InValue)
+	bool SetBoolMetadata(FRigElementKey InItem, FName InMetadataName, bool InValue)
 	{
 		return SetMetadata<bool>(InItem, ERigMetadataType::Bool, InMetadataName, InValue);
 	}
@@ -960,7 +960,7 @@ public:
 	 * @return Returns true if setting the metadata was successful
 	 */
 	UFUNCTION(BlueprintCallable, Category = URigHierarchy)
-	FORCEINLINE bool SetBoolArrayMetadata(FRigElementKey InItem, FName InMetadataName, TArray<bool> InValue)
+	bool SetBoolArrayMetadata(FRigElementKey InItem, FName InMetadataName, TArray<bool> InValue)
 	{
 		return SetArrayMetadata<bool>(InItem, ERigMetadataType::BoolArray, InMetadataName, InValue);
 	}
@@ -972,7 +972,7 @@ public:
 	 * @param DefaultValue The default value to fall back on
 	 */
 	UFUNCTION(BlueprintPure, Category = URigHierarchy)
-	FORCEINLINE float GetFloatMetadata(FRigElementKey InItem, FName InMetadataName, float DefaultValue) const
+	float GetFloatMetadata(FRigElementKey InItem, FName InMetadataName, float DefaultValue) const
 	{
 		return GetMetadata<float>(InItem, ERigMetadataType::Float, InMetadataName, DefaultValue);
 	}
@@ -983,7 +983,7 @@ public:
 	 * @param InMetadataName The name of the metadata to query
 	 */
 	UFUNCTION(BlueprintPure, Category = URigHierarchy)
-	FORCEINLINE TArray<float> GetFloatArrayMetadata(FRigElementKey InItem, FName InMetadataName) const
+	TArray<float> GetFloatArrayMetadata(FRigElementKey InItem, FName InMetadataName) const
 	{
 		return GetArrayMetadata<float>(InItem, ERigMetadataType::FloatArray, InMetadataName);
 	}
@@ -996,7 +996,7 @@ public:
 	 * @return Returns true if setting the metadata was successful
 	 */
 	UFUNCTION(BlueprintCallable, Category = URigHierarchy)
-	FORCEINLINE bool SetFloatMetadata(FRigElementKey InItem, FName InMetadataName, float InValue)
+	bool SetFloatMetadata(FRigElementKey InItem, FName InMetadataName, float InValue)
 	{
 		return SetMetadata<float>(InItem, ERigMetadataType::Float, InMetadataName, InValue);
 	}
@@ -1009,7 +1009,7 @@ public:
 	 * @return Returns true if setting the metadata was successful
 	 */
 	UFUNCTION(BlueprintCallable, Category = URigHierarchy)
-	FORCEINLINE bool SetFloatArrayMetadata(FRigElementKey InItem, FName InMetadataName, TArray<float> InValue)
+	bool SetFloatArrayMetadata(FRigElementKey InItem, FName InMetadataName, TArray<float> InValue)
 	{
 		return SetArrayMetadata<float>(InItem, ERigMetadataType::FloatArray, InMetadataName, InValue);
 	}
@@ -1021,7 +1021,7 @@ public:
 	 * @param DefaultValue The default value to fall back on
 	 */
 	UFUNCTION(BlueprintPure, Category = URigHierarchy)
-	FORCEINLINE int32 GetInt32Metadata(FRigElementKey InItem, FName InMetadataName, int32 DefaultValue) const
+	int32 GetInt32Metadata(FRigElementKey InItem, FName InMetadataName, int32 DefaultValue) const
 	{
 		return GetMetadata<int32>(InItem, ERigMetadataType::Int32, InMetadataName, DefaultValue);
 	}
@@ -1032,7 +1032,7 @@ public:
 	 * @param InMetadataName The name of the metadata to query
 	 */
 	UFUNCTION(BlueprintPure, Category = URigHierarchy)
-	FORCEINLINE TArray<int32> GetInt32ArrayMetadata(FRigElementKey InItem, FName InMetadataName) const
+	TArray<int32> GetInt32ArrayMetadata(FRigElementKey InItem, FName InMetadataName) const
 	{
 		return GetArrayMetadata<int32>(InItem, ERigMetadataType::Int32Array, InMetadataName);
 	}
@@ -1045,7 +1045,7 @@ public:
 	 * @return Returns true if setting the metadata was successful
 	 */
 	UFUNCTION(BlueprintCallable, Category = URigHierarchy)
-	FORCEINLINE bool SetInt32Metadata(FRigElementKey InItem, FName InMetadataName, int32 InValue)
+	bool SetInt32Metadata(FRigElementKey InItem, FName InMetadataName, int32 InValue)
 	{
 		return SetMetadata<int32>(InItem, ERigMetadataType::Int32, InMetadataName, InValue);
 	}
@@ -1058,7 +1058,7 @@ public:
 	 * @return Returns true if setting the metadata was successful
 	 */
 	UFUNCTION(BlueprintCallable, Category = URigHierarchy)
-	FORCEINLINE bool SetInt32ArrayMetadata(FRigElementKey InItem, FName InMetadataName, TArray<int32> InValue)
+	bool SetInt32ArrayMetadata(FRigElementKey InItem, FName InMetadataName, TArray<int32> InValue)
 	{
 		return SetArrayMetadata<int32>(InItem, ERigMetadataType::Int32Array, InMetadataName, InValue);
 	}
@@ -1070,7 +1070,7 @@ public:
 	 * @param DefaultValue The default value to fall back on
 	 */
 	UFUNCTION(BlueprintPure, Category = URigHierarchy)
-	FORCEINLINE FName GetNameMetadata(FRigElementKey InItem, FName InMetadataName, FName DefaultValue) const
+	FName GetNameMetadata(FRigElementKey InItem, FName InMetadataName, FName DefaultValue) const
 	{
 		return GetMetadata<FName>(InItem, ERigMetadataType::Name, InMetadataName, DefaultValue);
 	}
@@ -1081,7 +1081,7 @@ public:
 	 * @param InMetadataName The name of the metadata to query
 	 */
 	UFUNCTION(BlueprintPure, Category = URigHierarchy)
-	FORCEINLINE TArray<FName> GetNameArrayMetadata(FRigElementKey InItem, FName InMetadataName) const
+	TArray<FName> GetNameArrayMetadata(FRigElementKey InItem, FName InMetadataName) const
 	{
 		return GetArrayMetadata<FName>(InItem, ERigMetadataType::NameArray, InMetadataName);
 	}
@@ -1094,7 +1094,7 @@ public:
 	 * @return Returns true if setting the metadata was successful
 	 */
 	UFUNCTION(BlueprintCallable, Category = URigHierarchy)
-	FORCEINLINE bool SetNameMetadata(FRigElementKey InItem, FName InMetadataName, FName InValue)
+	bool SetNameMetadata(FRigElementKey InItem, FName InMetadataName, FName InValue)
 	{
 		return SetMetadata<FName>(InItem, ERigMetadataType::Name, InMetadataName, InValue);
 	}
@@ -1107,7 +1107,7 @@ public:
 	 * @return Returns true if setting the metadata was successful
 	 */
 	UFUNCTION(BlueprintCallable, Category = URigHierarchy)
-	FORCEINLINE bool SetNameArrayMetadata(FRigElementKey InItem, FName InMetadataName, TArray<FName> InValue)
+	bool SetNameArrayMetadata(FRigElementKey InItem, FName InMetadataName, TArray<FName> InValue)
 	{
 		return SetArrayMetadata<FName>(InItem, ERigMetadataType::FloatArray, InMetadataName, InValue);
 	}
@@ -1119,7 +1119,7 @@ public:
 	 * @param DefaultValue The default value to fall back on
 	 */
 	UFUNCTION(BlueprintPure, Category = URigHierarchy)
-	FORCEINLINE FVector GetVectorMetadata(FRigElementKey InItem, FName InMetadataName, FVector DefaultValue) const
+	FVector GetVectorMetadata(FRigElementKey InItem, FName InMetadataName, FVector DefaultValue) const
 	{
 		return GetMetadata<FVector>(InItem, ERigMetadataType::Vector, InMetadataName, DefaultValue);
 	}
@@ -1130,7 +1130,7 @@ public:
 	 * @param InMetadataName The name of the metadata to query
 	 */
 	UFUNCTION(BlueprintPure, Category = URigHierarchy)
-	FORCEINLINE TArray<FVector> GetVectorArrayMetadata(FRigElementKey InItem, FName InMetadataName) const
+	TArray<FVector> GetVectorArrayMetadata(FRigElementKey InItem, FName InMetadataName) const
 	{
 		return GetArrayMetadata<FVector>(InItem, ERigMetadataType::VectorArray, InMetadataName);
 	}
@@ -1143,7 +1143,7 @@ public:
 	 * @return Returns true if setting the metadata was successful
 	 */
 	UFUNCTION(BlueprintCallable, Category = URigHierarchy)
-	FORCEINLINE bool SetVectorMetadata(FRigElementKey InItem, FName InMetadataName, FVector InValue)
+	bool SetVectorMetadata(FRigElementKey InItem, FName InMetadataName, FVector InValue)
 	{
 		return SetMetadata<FVector>(InItem, ERigMetadataType::Vector, InMetadataName, InValue);
 	}
@@ -1156,7 +1156,7 @@ public:
 	 * @return Returns true if setting the metadata was successful
 	 */
 	UFUNCTION(BlueprintCallable, Category = URigHierarchy)
-	FORCEINLINE bool SetVectorArrayMetadata(FRigElementKey InItem, FName InMetadataName, TArray<FVector> InValue)
+	bool SetVectorArrayMetadata(FRigElementKey InItem, FName InMetadataName, TArray<FVector> InValue)
 	{
 		return SetArrayMetadata<FVector>(InItem, ERigMetadataType::VectorArray, InMetadataName, InValue);
 	}
@@ -1168,7 +1168,7 @@ public:
 	 * @param DefaultValue The default value to fall back on
 	 */
 	UFUNCTION(BlueprintPure, Category = URigHierarchy)
-	FORCEINLINE FRotator GetRotatorMetadata(FRigElementKey InItem, FName InMetadataName, FRotator DefaultValue) const
+	FRotator GetRotatorMetadata(FRigElementKey InItem, FName InMetadataName, FRotator DefaultValue) const
 	{
 		return GetMetadata<FRotator>(InItem, ERigMetadataType::Rotator, InMetadataName, DefaultValue);
 	}
@@ -1179,7 +1179,7 @@ public:
 	 * @param InMetadataName The name of the metadata to query
 	 */
 	UFUNCTION(BlueprintPure, Category = URigHierarchy)
-	FORCEINLINE TArray<FRotator> GetRotatorArrayMetadata(FRigElementKey InItem, FName InMetadataName) const
+	TArray<FRotator> GetRotatorArrayMetadata(FRigElementKey InItem, FName InMetadataName) const
 	{
 		return GetArrayMetadata<FRotator>(InItem, ERigMetadataType::RotatorArray, InMetadataName);
 	}
@@ -1192,7 +1192,7 @@ public:
 	 * @return Returns true if setting the metadata was successful
 	 */
 	UFUNCTION(BlueprintCallable, Category = URigHierarchy)
-	FORCEINLINE bool SetRotatorMetadata(FRigElementKey InItem, FName InMetadataName, FRotator InValue)
+	bool SetRotatorMetadata(FRigElementKey InItem, FName InMetadataName, FRotator InValue)
 	{
 		return SetMetadata<FRotator>(InItem, ERigMetadataType::Rotator, InMetadataName, InValue);
 	}
@@ -1205,7 +1205,7 @@ public:
 	 * @return Returns true if setting the metadata was successful
 	 */
 	UFUNCTION(BlueprintCallable, Category = URigHierarchy)
-	FORCEINLINE bool SetRotatorArrayMetadata(FRigElementKey InItem, FName InMetadataName, TArray<FRotator> InValue)
+	bool SetRotatorArrayMetadata(FRigElementKey InItem, FName InMetadataName, TArray<FRotator> InValue)
 	{
 		return SetArrayMetadata<FRotator>(InItem, ERigMetadataType::RotatorArray, InMetadataName, InValue);
 	}
@@ -1217,7 +1217,7 @@ public:
 	 * @param DefaultValue The default value to fall back on
 	 */
 	UFUNCTION(BlueprintPure, Category = URigHierarchy)
-	FORCEINLINE FQuat GetQuatMetadata(FRigElementKey InItem, FName InMetadataName, FQuat DefaultValue) const
+	FQuat GetQuatMetadata(FRigElementKey InItem, FName InMetadataName, FQuat DefaultValue) const
 	{
 		return GetMetadata<FQuat>(InItem, ERigMetadataType::Quat, InMetadataName, DefaultValue);
 	}
@@ -1228,7 +1228,7 @@ public:
 	 * @param InMetadataName The name of the metadata to query
 	 */
 	UFUNCTION(BlueprintPure, Category = URigHierarchy)
-	FORCEINLINE TArray<FQuat> GetQuatArrayMetadata(FRigElementKey InItem, FName InMetadataName) const
+	TArray<FQuat> GetQuatArrayMetadata(FRigElementKey InItem, FName InMetadataName) const
 	{
 		return GetArrayMetadata<FQuat>(InItem, ERigMetadataType::QuatArray, InMetadataName);
 	}
@@ -1241,7 +1241,7 @@ public:
 	 * @return Returns true if setting the metadata was successful
 	 */
 	UFUNCTION(BlueprintCallable, Category = URigHierarchy)
-	FORCEINLINE bool SetQuatMetadata(FRigElementKey InItem, FName InMetadataName, FQuat InValue)
+	bool SetQuatMetadata(FRigElementKey InItem, FName InMetadataName, FQuat InValue)
 	{
 		return SetMetadata<FQuat>(InItem, ERigMetadataType::Quat, InMetadataName, InValue);
 	}
@@ -1254,7 +1254,7 @@ public:
 	 * @return Returns true if setting the metadata was successful
 	 */
 	UFUNCTION(BlueprintCallable, Category = URigHierarchy)
-	FORCEINLINE bool SetQuatArrayMetadata(FRigElementKey InItem, FName InMetadataName, TArray<FQuat> InValue)
+	bool SetQuatArrayMetadata(FRigElementKey InItem, FName InMetadataName, TArray<FQuat> InValue)
 	{
 		return SetArrayMetadata<FQuat>(InItem, ERigMetadataType::QuatArray, InMetadataName, InValue);
 	}
@@ -1266,7 +1266,7 @@ public:
 	 * @param DefaultValue The default value to fall back on
 	 */
 	UFUNCTION(BlueprintPure, Category = URigHierarchy)
-	FORCEINLINE FTransform GetTransformMetadata(FRigElementKey InItem, FName InMetadataName, FTransform DefaultValue) const
+	FTransform GetTransformMetadata(FRigElementKey InItem, FName InMetadataName, FTransform DefaultValue) const
 	{
 		return GetMetadata<FTransform>(InItem, ERigMetadataType::Transform, InMetadataName, DefaultValue);
 	}
@@ -1277,7 +1277,7 @@ public:
 	 * @param InMetadataName The name of the metadata to query
 	 */
 	UFUNCTION(BlueprintPure, Category = URigHierarchy)
-	FORCEINLINE TArray<FTransform> GetTransformArrayMetadata(FRigElementKey InItem, FName InMetadataName) const
+	TArray<FTransform> GetTransformArrayMetadata(FRigElementKey InItem, FName InMetadataName) const
 	{
 		return GetArrayMetadata<FTransform>(InItem, ERigMetadataType::TransformArray, InMetadataName);
 	}
@@ -1290,7 +1290,7 @@ public:
 	 * @return Returns true if setting the metadata was successful
 	 */
 	UFUNCTION(BlueprintCallable, Category = URigHierarchy)
-	FORCEINLINE bool SetTransformMetadata(FRigElementKey InItem, FName InMetadataName, FTransform InValue)
+	bool SetTransformMetadata(FRigElementKey InItem, FName InMetadataName, FTransform InValue)
 	{
 		return SetMetadata<FTransform>(InItem, ERigMetadataType::Transform, InMetadataName, InValue);
 	}
@@ -1303,7 +1303,7 @@ public:
 	 * @return Returns true if setting the metadata was successful
 	 */
 	UFUNCTION(BlueprintCallable, Category = URigHierarchy)
-	FORCEINLINE bool SetTransformArrayMetadata(FRigElementKey InItem, FName InMetadataName, TArray<FTransform> InValue)
+	bool SetTransformArrayMetadata(FRigElementKey InItem, FName InMetadataName, TArray<FTransform> InValue)
 	{
 		return SetArrayMetadata<FTransform>(InItem, ERigMetadataType::TransformArray, InMetadataName, InValue);
 	}
@@ -1315,7 +1315,7 @@ public:
 	 * @param DefaultValue The default value to fall back on
 	 */
 	UFUNCTION(BlueprintPure, Category = URigHierarchy)
-	FORCEINLINE FLinearColor GetLinearColorMetadata(FRigElementKey InItem, FName InMetadataName, FLinearColor DefaultValue) const
+	FLinearColor GetLinearColorMetadata(FRigElementKey InItem, FName InMetadataName, FLinearColor DefaultValue) const
 	{
 		return GetMetadata<FLinearColor>(InItem, ERigMetadataType::LinearColor, InMetadataName, DefaultValue);
 	}
@@ -1326,7 +1326,7 @@ public:
 	 * @param InMetadataName The name of the metadata to query
 	 */
 	UFUNCTION(BlueprintPure, Category = URigHierarchy)
-	FORCEINLINE TArray<FLinearColor> GetLinearColorArrayMetadata(FRigElementKey InItem, FName InMetadataName) const
+	TArray<FLinearColor> GetLinearColorArrayMetadata(FRigElementKey InItem, FName InMetadataName) const
 	{
 		return GetArrayMetadata<FLinearColor>(InItem, ERigMetadataType::LinearColorArray, InMetadataName);
 	}
@@ -1339,7 +1339,7 @@ public:
 	 * @return Returns true if setting the metadata was successful
 	 */
 	UFUNCTION(BlueprintCallable, Category = URigHierarchy)
-	FORCEINLINE bool SetLinearColorMetadata(FRigElementKey InItem, FName InMetadataName, FLinearColor InValue)
+	bool SetLinearColorMetadata(FRigElementKey InItem, FName InMetadataName, FLinearColor InValue)
 	{
 		return SetMetadata<FLinearColor>(InItem, ERigMetadataType::LinearColor, InMetadataName, InValue);
 	}
@@ -1352,7 +1352,7 @@ public:
 	 * @return Returns true if setting the metadata was successful
 	 */
 	UFUNCTION(BlueprintCallable, Category = URigHierarchy)
-	FORCEINLINE bool SetLinearColorArrayMetadata(FRigElementKey InItem, FName InMetadataName, TArray<FLinearColor> InValue)
+	bool SetLinearColorArrayMetadata(FRigElementKey InItem, FName InMetadataName, TArray<FLinearColor> InValue)
 	{
 		return SetArrayMetadata<FLinearColor>(InItem, ERigMetadataType::LinearColorArray, InMetadataName, InValue);
 	}
@@ -1364,7 +1364,7 @@ public:
 	 * @param DefaultValue The default value to fall back on
 	 */
 	UFUNCTION(BlueprintPure, Category = URigHierarchy)
-	FORCEINLINE FRigElementKey GetRigElementKeyMetadata(FRigElementKey InItem, FName InMetadataName, FRigElementKey DefaultValue) const
+	FRigElementKey GetRigElementKeyMetadata(FRigElementKey InItem, FName InMetadataName, FRigElementKey DefaultValue) const
 	{
 		return GetMetadata<FRigElementKey>(InItem, ERigMetadataType::RigElementKey, InMetadataName, DefaultValue);
 	}
@@ -1375,7 +1375,7 @@ public:
 	 * @param InMetadataName The name of the metadata to query
 	 */
 	UFUNCTION(BlueprintPure, Category = URigHierarchy)
-	FORCEINLINE TArray<FRigElementKey> GetRigElementKeyArrayMetadata(FRigElementKey InItem, FName InMetadataName) const
+	TArray<FRigElementKey> GetRigElementKeyArrayMetadata(FRigElementKey InItem, FName InMetadataName) const
 	{
 		return GetArrayMetadata<FRigElementKey>(InItem, ERigMetadataType::RigElementKeyArray, InMetadataName);
 	}
@@ -1388,7 +1388,7 @@ public:
 	 * @return Returns true if setting the metadata was successful
 	 */
 	UFUNCTION(BlueprintCallable, Category = URigHierarchy)
-	FORCEINLINE bool SetRigElementKeyMetadata(FRigElementKey InItem, FName InMetadataName, FRigElementKey InValue)
+	bool SetRigElementKeyMetadata(FRigElementKey InItem, FName InMetadataName, FRigElementKey InValue)
 	{
 		return SetMetadata<FRigElementKey>(InItem, ERigMetadataType::RigElementKey, InMetadataName, InValue);
 	}
@@ -1401,7 +1401,7 @@ public:
 	 * @return Returns true if setting the metadata was successful
 	 */
 	UFUNCTION(BlueprintCallable, Category = URigHierarchy)
-	FORCEINLINE bool SetRigElementKeyArrayMetadata(FRigElementKey InItem, FName InMetadataName, TArray<FRigElementKey> InValue)
+	bool SetRigElementKeyArrayMetadata(FRigElementKey InItem, FName InMetadataName, TArray<FRigElementKey> InValue)
 	{
 		return SetArrayMetadata<FRigElementKey>(InItem, ERigMetadataType::RigElementKeyArray, InMetadataName, InValue);
 	}
@@ -1411,7 +1411,7 @@ public:
 	 * @param InItem The item to return the tags for
 	 */
 	UFUNCTION(BlueprintPure, Category = URigHierarchy)
-	FORCEINLINE TArray<FName> GetTags(FRigElementKey InItem) const
+	TArray<FName> GetTags(FRigElementKey InItem) const
 	{
 		return GetNameArrayMetadata(InItem, TagMetadataName);
 	}
@@ -1422,7 +1422,7 @@ public:
 	 * @param InTag The tag to check
 	 */
 	UFUNCTION(BlueprintPure, Category = URigHierarchy)
-	FORCEINLINE bool HasTag(FRigElementKey InItem, FName InTag) const
+	bool HasTag(FRigElementKey InItem, FName InTag) const
 	{
 		return GetTags(InItem).Contains(InTag);
 	}
@@ -1433,7 +1433,7 @@ public:
      * @param InTag The tag to set
      */
 	UFUNCTION(BlueprintCallable, Category = URigHierarchy)
-	FORCEINLINE bool SetTag(FRigElementKey InItem, FName InTag)
+	bool SetTag(FRigElementKey InItem, FName InTag)
 	{
 		TArray<FName> Tags = GetTags(InItem);
 		Tags.AddUnique(InTag);
@@ -1461,7 +1461,7 @@ public:
 	 * @return true if a given element is selected
 	 */
 	UFUNCTION(BlueprintCallable, Category = URigHierarchy)
-	FORCEINLINE bool IsSelected(FRigElementKey InKey) const
+	bool IsSelected(FRigElementKey InKey) const
 	{
 		return IsSelected(Find(InKey));
 	}
@@ -1472,12 +1472,12 @@ public:
 	 * @return true if a given element is selected
 	 */
 	UFUNCTION(BlueprintCallable, Category = URigHierarchy)
-    FORCEINLINE bool IsSelectedByIndex(int32 InIndex) const
+    bool IsSelectedByIndex(int32 InIndex) const
 	{
 		return IsSelected(Get(InIndex));
 	}
 
-	FORCEINLINE bool IsSelected(int32 InIndex) const
+	bool IsSelected(int32 InIndex) const
 	{
 		return IsSelectedByIndex(InIndex);
 	}
@@ -1488,7 +1488,7 @@ public:
 	 * @return The sorted keys
 	 */
 	UFUNCTION(BlueprintCallable, Category = URigHierarchy)
-	FORCEINLINE TArray<FRigElementKey> SortKeys(const TArray<FRigElementKey>& InKeys) const
+	TArray<FRigElementKey> SortKeys(const TArray<FRigElementKey>& InKeys) const
 	{
 		TArray<FRigElementKey> Result;
 		Traverse([InKeys, &Result](FRigBaseElement* Element, bool& bContinue)
@@ -1584,7 +1584,7 @@ public:
 	 * @return The local current or initial transform's value.
 	 */
 	UFUNCTION(BlueprintCallable, Category = URigHierarchy)
-    FORCEINLINE_DEBUGGABLE FTransform GetLocalTransform(FRigElementKey InKey, bool bInitial = false) const
+    FTransform GetLocalTransform(FRigElementKey InKey, bool bInitial = false) const
 	{
 		return GetLocalTransformByIndex(GetIndex(InKey), bInitial);
 	}
@@ -1597,7 +1597,7 @@ public:
 	 * @return The local current or initial transform's value.
 	 */
 	UFUNCTION(BlueprintCallable, Category = URigHierarchy)
-	FORCEINLINE_DEBUGGABLE FTransform GetLocalTransformByIndex(int32 InElementIndex, bool bInitial = false) const
+	FTransform GetLocalTransformByIndex(int32 InElementIndex, bool bInitial = false) const
 	{
 		if(Elements.IsValidIndex(InElementIndex))
 		{
@@ -1609,16 +1609,16 @@ public:
 		return FTransform::Identity;
 	}
 
-	FORCEINLINE_DEBUGGABLE FTransform GetLocalTransform(int32 InElementIndex) const
+	FTransform GetLocalTransform(int32 InElementIndex) const
 	{
 		return GetLocalTransformByIndex(InElementIndex, false);
 	}
-	FORCEINLINE_DEBUGGABLE FTransform GetInitialLocalTransform(int32 InElementIndex) const
+	FTransform GetInitialLocalTransform(int32 InElementIndex) const
 	{
 		return GetLocalTransformByIndex(InElementIndex, true);
 	}
 
-	FORCEINLINE_DEBUGGABLE FTransform GetInitialLocalTransform(const FRigElementKey &InKey) const
+	FTransform GetInitialLocalTransform(const FRigElementKey &InKey) const
 	{
 		return GetLocalTransform(InKey, true);
 	}
@@ -1632,7 +1632,7 @@ public:
 	 * @param bAffectChildren If set to false children will not move (maintain global).
 	 */
 	UFUNCTION(BlueprintCallable, Category = URigHierarchy)
-    FORCEINLINE_DEBUGGABLE void SetLocalTransform(FRigElementKey InKey, FTransform InTransform, bool bInitial = false, bool bAffectChildren = true, bool bSetupUndo = false, bool bPrintPythonCommands = false)
+    void SetLocalTransform(FRigElementKey InKey, FTransform InTransform, bool bInitial = false, bool bAffectChildren = true, bool bSetupUndo = false, bool bPrintPythonCommands = false)
 	{
 		SetLocalTransformByIndex(GetIndex(InKey), InTransform, bInitial, bAffectChildren, bSetupUndo, bPrintPythonCommands);
 	}
@@ -1646,7 +1646,7 @@ public:
 	 * @param bAffectChildren If set to false children will not move (maintain global).
 	 */
 	UFUNCTION(BlueprintCallable, Category = URigHierarchy)
-    FORCEINLINE_DEBUGGABLE void SetLocalTransformByIndex(int32 InElementIndex, FTransform InTransform, bool bInitial = false, bool bAffectChildren = true, bool bSetupUndo = false, bool bPrintPythonCommands = false)
+    void SetLocalTransformByIndex(int32 InElementIndex, FTransform InTransform, bool bInitial = false, bool bAffectChildren = true, bool bSetupUndo = false, bool bPrintPythonCommands = false)
 	{
 		if(Elements.IsValidIndex(InElementIndex))
 		{
@@ -1657,17 +1657,17 @@ public:
 		}
 	}
 
-	FORCEINLINE_DEBUGGABLE void SetLocalTransform(int32 InElementIndex, const FTransform& InTransform, bool bAffectChildren = true, bool bSetupUndo = false, bool bPrintPythonCommands = false)
+	void SetLocalTransform(int32 InElementIndex, const FTransform& InTransform, bool bAffectChildren = true, bool bSetupUndo = false, bool bPrintPythonCommands = false)
 	{
 		SetLocalTransformByIndex(InElementIndex, InTransform, false, bAffectChildren, bSetupUndo, bPrintPythonCommands);
 	}
 
-	FORCEINLINE_DEBUGGABLE void SetInitialLocalTransform(int32 InElementIndex, const FTransform& InTransform, bool bAffectChildren = true, bool bSetupUndo = false, bool bPrintPythonCommands = false)
+	void SetInitialLocalTransform(int32 InElementIndex, const FTransform& InTransform, bool bAffectChildren = true, bool bSetupUndo = false, bool bPrintPythonCommands = false)
 	{
 		SetLocalTransformByIndex(InElementIndex, InTransform, true, bAffectChildren, bSetupUndo, bPrintPythonCommands);
     }
 
-	FORCEINLINE_DEBUGGABLE void SetInitialLocalTransform(const FRigElementKey& InKey, const FTransform& InTransform, bool bAffectChildren = true, bool bSetupUndo = false, bool bPrintPythonCommands = false)
+	void SetInitialLocalTransform(const FRigElementKey& InKey, const FTransform& InTransform, bool bAffectChildren = true, bool bSetupUndo = false, bool bPrintPythonCommands = false)
 	{
 		SetLocalTransform(InKey, InTransform, true, bAffectChildren, bSetupUndo, bPrintPythonCommands);
 	}
@@ -1680,7 +1680,7 @@ public:
 	 * @return The global current or initial transform's value.
 	 */
 	UFUNCTION(BlueprintCallable, Category = URigHierarchy)
-    FORCEINLINE_DEBUGGABLE FTransform GetGlobalTransform(FRigElementKey InKey, bool bInitial = false) const
+    FTransform GetGlobalTransform(FRigElementKey InKey, bool bInitial = false) const
 	{
 		return GetGlobalTransformByIndex(GetIndex(InKey), bInitial);
 	}
@@ -1693,7 +1693,7 @@ public:
 	 * @return The global current or initial transform's value.
 	 */
 	UFUNCTION(BlueprintCallable, Category = URigHierarchy)
-    FORCEINLINE_DEBUGGABLE FTransform GetGlobalTransformByIndex(int32 InElementIndex, bool bInitial = false) const
+    FTransform GetGlobalTransformByIndex(int32 InElementIndex, bool bInitial = false) const
 	{
 		if(Elements.IsValidIndex(InElementIndex))
 		{
@@ -1705,16 +1705,16 @@ public:
 		return FTransform::Identity;
 	}
 
-	FORCEINLINE_DEBUGGABLE FTransform GetGlobalTransform(int32 InElementIndex) const
+	FTransform GetGlobalTransform(int32 InElementIndex) const
 	{
 		return GetGlobalTransformByIndex(InElementIndex, false);
 	}
-	FORCEINLINE_DEBUGGABLE FTransform GetInitialGlobalTransform(int32 InElementIndex) const
+	FTransform GetInitialGlobalTransform(int32 InElementIndex) const
 	{
 		return GetGlobalTransformByIndex(InElementIndex, true);
 	}
 
-	FORCEINLINE_DEBUGGABLE FTransform GetInitialGlobalTransform(const FRigElementKey &InKey) const
+	FTransform GetInitialGlobalTransform(const FRigElementKey &InKey) const
 	{
 		return GetGlobalTransform(InKey, true);
 	}
@@ -1728,7 +1728,7 @@ public:
 	 * @param bSetupUndo If true the transform stack will be setup for undo / redo
 	 */
 	UFUNCTION(BlueprintCallable, Category = URigHierarchy)
-    FORCEINLINE_DEBUGGABLE void SetGlobalTransform(FRigElementKey InKey, FTransform InTransform, bool bInitial = false, bool bAffectChildren = true, bool bSetupUndo = false, bool bPrintPythonCommand = false)
+    void SetGlobalTransform(FRigElementKey InKey, FTransform InTransform, bool bInitial = false, bool bAffectChildren = true, bool bSetupUndo = false, bool bPrintPythonCommand = false)
 	{
 		SetGlobalTransformByIndex(GetIndex(InKey), InTransform, bInitial, bAffectChildren, bSetupUndo, bPrintPythonCommand);
 	}
@@ -1742,7 +1742,7 @@ public:
 	 * @param bSetupUndo If true the transform stack will be setup for undo / redo
 	 */
 	UFUNCTION(BlueprintCallable, Category = URigHierarchy)
-    FORCEINLINE_DEBUGGABLE void SetGlobalTransformByIndex(int32 InElementIndex, FTransform InTransform, bool bInitial = false, bool bAffectChildren = true, bool bSetupUndo = false, bool bPrintPythonCommand = false)
+    void SetGlobalTransformByIndex(int32 InElementIndex, FTransform InTransform, bool bInitial = false, bool bAffectChildren = true, bool bSetupUndo = false, bool bPrintPythonCommand = false)
 	{
 		if(Elements.IsValidIndex(InElementIndex))
 		{
@@ -1753,17 +1753,17 @@ public:
 		}
 	}
 
-	FORCEINLINE_DEBUGGABLE void SetGlobalTransform(int32 InElementIndex, const FTransform& InTransform, bool bAffectChildren = true, bool bSetupUndo = false)
+	void SetGlobalTransform(int32 InElementIndex, const FTransform& InTransform, bool bAffectChildren = true, bool bSetupUndo = false)
 	{
 		SetGlobalTransformByIndex(InElementIndex, InTransform, false, bAffectChildren, bSetupUndo);
 	}
 
-	FORCEINLINE_DEBUGGABLE void SetInitialGlobalTransform(int32 InElementIndex, const FTransform& InTransform, bool bAffectChildren = true, bool bSetupUndo = false)
+	void SetInitialGlobalTransform(int32 InElementIndex, const FTransform& InTransform, bool bAffectChildren = true, bool bSetupUndo = false)
 	{
 		SetGlobalTransformByIndex(InElementIndex, InTransform, true, bAffectChildren, bSetupUndo);
 	}
 
-	FORCEINLINE_DEBUGGABLE void SetInitialGlobalTransform(const FRigElementKey& InKey, const FTransform& InTransform, bool bAffectChildren = true, bool bSetupUndo = false)
+	void SetInitialGlobalTransform(const FRigElementKey& InKey, const FTransform& InTransform, bool bAffectChildren = true, bool bSetupUndo = false)
 	{
 		SetGlobalTransform(InKey, InTransform, true, bAffectChildren, bSetupUndo);
 	}
@@ -1775,7 +1775,7 @@ public:
 	 * @return The global offset transform
 	 */
 	UFUNCTION(BlueprintCallable, Category = URigHierarchy)
-    FORCEINLINE_DEBUGGABLE FTransform GetGlobalControlOffsetTransform(FRigElementKey InKey, bool bInitial = false) const
+    FTransform GetGlobalControlOffsetTransform(FRigElementKey InKey, bool bInitial = false) const
 	{
 		return GetGlobalControlOffsetTransformByIndex(GetIndex(InKey), bInitial);
 	}
@@ -1787,7 +1787,7 @@ public:
 	 * @return The global offset transform
 	 */
 	UFUNCTION(BlueprintCallable, Category = URigHierarchy)
-    FORCEINLINE_DEBUGGABLE FTransform GetGlobalControlOffsetTransformByIndex(int32 InElementIndex, bool bInitial = false) const
+    FTransform GetGlobalControlOffsetTransformByIndex(int32 InElementIndex, bool bInitial = false) const
 	{
 		if(Elements.IsValidIndex(InElementIndex))
 		{
@@ -1806,7 +1806,7 @@ public:
 	 * @return The local shape transform
 	 */
 	UFUNCTION(BlueprintCallable, Category = URigHierarchy)
-	FORCEINLINE_DEBUGGABLE FTransform GetLocalControlShapeTransform(FRigElementKey InKey, bool bInitial = false) const
+	FTransform GetLocalControlShapeTransform(FRigElementKey InKey, bool bInitial = false) const
 	{
 		return GetLocalControlShapeTransformByIndex(GetIndex(InKey), bInitial);
 	}
@@ -1818,7 +1818,7 @@ public:
 	 * @return The local shape transform
 	 */
 	UFUNCTION(BlueprintCallable, Category = URigHierarchy)
-	FORCEINLINE_DEBUGGABLE FTransform GetLocalControlShapeTransformByIndex(int32 InElementIndex, bool bInitial = false) const
+	FTransform GetLocalControlShapeTransformByIndex(int32 InElementIndex, bool bInitial = false) const
 	{
 		if(Elements.IsValidIndex(InElementIndex))
 		{
@@ -1837,7 +1837,7 @@ public:
 	 * @return The global shape transform
 	 */
 	UFUNCTION(BlueprintCallable, Category = URigHierarchy)
-    FORCEINLINE_DEBUGGABLE FTransform GetGlobalControlShapeTransform(FRigElementKey InKey, bool bInitial = false) const
+    FTransform GetGlobalControlShapeTransform(FRigElementKey InKey, bool bInitial = false) const
 	{
 		return GetGlobalControlShapeTransformByIndex(GetIndex(InKey), bInitial);
 	}
@@ -1849,7 +1849,7 @@ public:
 	 * @return The global shape transform
 	 */
 	UFUNCTION(BlueprintCallable, Category = URigHierarchy)
-    FORCEINLINE_DEBUGGABLE FTransform GetGlobalControlShapeTransformByIndex(int32 InElementIndex, bool bInitial = false) const
+    FTransform GetGlobalControlShapeTransformByIndex(int32 InElementIndex, bool bInitial = false) const
 	{
 		if(Elements.IsValidIndex(InElementIndex))
 		{
@@ -1868,7 +1868,7 @@ public:
 	 * @return Returns the current value of the control
 	 */
 	UFUNCTION(BlueprintCallable, Category = URigHierarchy)
-    FORCEINLINE_DEBUGGABLE FRigControlValue GetControlValue(FRigElementKey InKey, ERigControlValueType InValueType = ERigControlValueType::Current) const
+    FRigControlValue GetControlValue(FRigElementKey InKey, ERigControlValueType InValueType = ERigControlValueType::Current) const
 	{
 		return GetControlValueByIndex(GetIndex(InKey), InValueType);
 	}
@@ -1880,7 +1880,7 @@ public:
 	 * @return Returns the current value of the control
 	 */
 	template<typename T>
-	FORCEINLINE_DEBUGGABLE T GetControlValue(FRigElementKey InKey, ERigControlValueType InValueType = ERigControlValueType::Current) const
+	T GetControlValue(FRigElementKey InKey, ERigControlValueType InValueType = ERigControlValueType::Current) const
 	{
 		return GetControlValueByIndex(GetIndex(InKey), InValueType).Get<T>();
 	}
@@ -1892,7 +1892,7 @@ public:
 	 * @return Returns the current value of the control
 	 */
 	UFUNCTION(BlueprintCallable, Category = URigHierarchy)
-    FORCEINLINE_DEBUGGABLE FRigControlValue GetControlValueByIndex(int32 InElementIndex, ERigControlValueType InValueType = ERigControlValueType::Current) const
+    FRigControlValue GetControlValueByIndex(int32 InElementIndex, ERigControlValueType InValueType = ERigControlValueType::Current) const
 	{
 		if(Elements.IsValidIndex(InElementIndex))
 		{
@@ -1910,7 +1910,7 @@ public:
 	 * @param InValueType The type of value to return
 	 * @return Returns the current value of the control
 	 */
-	FORCEINLINE_DEBUGGABLE FRigControlValue GetControlValue(int32 InElementIndex, ERigControlValueType InValueType = ERigControlValueType::Current) const
+	FRigControlValue GetControlValue(int32 InElementIndex, ERigControlValueType InValueType = ERigControlValueType::Current) const
 	{
 		return GetControlValueByIndex(InElementIndex, InValueType);
 	}
@@ -1922,7 +1922,7 @@ public:
 	 * @return Returns the current value of the control
 	 */
 	template<typename T>
-	FORCEINLINE_DEBUGGABLE T GetControlValue(int32 InElementIndex, ERigControlValueType InValueType = ERigControlValueType::Current) const
+	T GetControlValue(int32 InElementIndex, ERigControlValueType InValueType = ERigControlValueType::Current) const
 	{
 		return GetControlValueByIndex(InElementIndex, InValueType).Get<T>();
 	}
@@ -1932,7 +1932,7 @@ public:
 	 * @param InElementIndex The index of the element to retrieve the initial value for
 	 * @return Returns the current value of the control
 	 */
-	FORCEINLINE_DEBUGGABLE FRigControlValue GetInitialControlValue(int32 InElementIndex) const
+	FRigControlValue GetInitialControlValue(int32 InElementIndex) const
 	{
 		return GetControlValueByIndex(InElementIndex, ERigControlValueType::Initial);
 	}
@@ -1943,7 +1943,7 @@ public:
 	 * @return Returns the current value of the control
 	 */
 	template<typename T>
-	FORCEINLINE_DEBUGGABLE T GetInitialControlValue(int32 InElementIndex) const
+	T GetInitialControlValue(int32 InElementIndex) const
 	{
 		return GetInitialControlValue(InElementIndex).Get<T>();
 	}
@@ -1955,7 +1955,7 @@ public:
 	 * @return Returns the current preferred rotator
 	 */
 	UFUNCTION(BlueprintCallable, Category = URigHierarchy)
-	FORCEINLINE_DEBUGGABLE FRotator GetControlPreferredRotator(FRigElementKey InKey, bool bInitial = false) const
+	FRotator GetControlPreferredRotator(FRigElementKey InKey, bool bInitial = false) const
 	{
 		return GetControlPreferredRotatorByIndex(GetIndex(InKey), bInitial);
 	}
@@ -1967,7 +1967,7 @@ public:
 	 * @return Returns the current preferred rotator
 	 */
 	UFUNCTION(BlueprintCallable, Category = URigHierarchy)
-	FORCEINLINE_DEBUGGABLE FRotator GetControlPreferredRotatorByIndex(int32 InElementIndex, bool bInitial = false) const
+	FRotator GetControlPreferredRotatorByIndex(int32 InElementIndex, bool bInitial = false) const
 	{
 		if(Elements.IsValidIndex(InElementIndex))
 		{
@@ -1985,7 +1985,7 @@ public:
 	 * @param bInitial If true we'll return the preferred rotator for the initial - otherwise current transform
 	 * @return Returns the current preferred rotator
 	 */
-	FORCEINLINE_DEBUGGABLE FRotator GetControlPreferredRotator(FRigControlElement* InControlElement, bool bInitial = false) const
+	FRotator GetControlPreferredRotator(FRigControlElement* InControlElement, bool bInitial = false) const
 	{
 		if(InControlElement)
 		{
@@ -2002,7 +2002,7 @@ public:
 	 * @param bFixEulerFlips If true the new rotator value will use the shortest path
 	 */
 	UFUNCTION(BlueprintCallable, Category = URigHierarchy)
-	FORCEINLINE_DEBUGGABLE void SetControlPreferredRotator(FRigElementKey InKey, const FRotator& InValue, bool bInitial = false, bool bFixEulerFlips = false)
+	void SetControlPreferredRotator(FRigElementKey InKey, const FRotator& InValue, bool bInitial = false, bool bFixEulerFlips = false)
 	{
 		SetControlPreferredRotatorByIndex(GetIndex(InKey), InValue, bInitial, bFixEulerFlips);
 	}
@@ -2016,7 +2016,7 @@ public:
 	 * @param bFixEulerFlips If true the new rotator value will use the shortest path
 	 */
 	UFUNCTION(BlueprintCallable, Category = URigHierarchy)
-	FORCEINLINE_DEBUGGABLE void SetControlPreferredRotatorByIndex(int32 InElementIndex, const FRotator& InValue, bool bInitial = false, bool bFixEulerFlips = false)
+	void SetControlPreferredRotatorByIndex(int32 InElementIndex, const FRotator& InValue, bool bInitial = false, bool bFixEulerFlips = false)
 	{
 		if(Elements.IsValidIndex(InElementIndex))
 		{
@@ -2034,7 +2034,7 @@ public:
 	 * @param bInitial If true we'll return the preferred rotator for the initial - otherwise current transform
 	 * @param bFixEulerFlips If true the new rotator value will use the shortest path
 	 */
-	FORCEINLINE_DEBUGGABLE void SetControlPreferredRotator(FRigControlElement* InControlElement, const FRotator& InValue, bool bInitial = false, bool bFixEulerFlips = false)
+	void SetControlPreferredRotator(FRigControlElement* InControlElement, const FRotator& InValue, bool bInitial = false, bool bFixEulerFlips = false)
 	{
 		if(InControlElement)
 		{
@@ -2066,7 +2066,7 @@ public:
 	 * @param bSetupUndo If true the transform stack will be setup for undo / redo
 	 */
 	UFUNCTION(BlueprintCallable, Category = URigHierarchy)
-    FORCEINLINE_DEBUGGABLE void SetControlValue(FRigElementKey InKey, FRigControlValue InValue, ERigControlValueType InValueType = ERigControlValueType::Current, bool bSetupUndo = false, bool bPrintPythonCommands = false)
+    void SetControlValue(FRigElementKey InKey, FRigControlValue InValue, ERigControlValueType InValueType = ERigControlValueType::Current, bool bSetupUndo = false, bool bPrintPythonCommands = false)
 	{
 		SetControlValueByIndex(GetIndex(InKey), InValue, InValueType, bSetupUndo, bPrintPythonCommands);
 	}
@@ -2079,7 +2079,7 @@ public:
 	 * @param bSetupUndo If true the transform stack will be setup for undo / redo
 	 */
 	template<typename T>
-	FORCEINLINE_DEBUGGABLE void SetControlValue(FRigElementKey InKey, const T& InValue, ERigControlValueType InValueType = ERigControlValueType::Current, bool bSetupUndo = false, bool bPrintPythonCommands = false) const
+	void SetControlValue(FRigElementKey InKey, const T& InValue, ERigControlValueType InValueType = ERigControlValueType::Current, bool bSetupUndo = false, bool bPrintPythonCommands = false) const
 	{
 		return SetControlValue(InKey, FRigControlValue::Make<T>(InValue), InValueType, bSetupUndo, bPrintPythonCommands);
 	}
@@ -2092,7 +2092,7 @@ public:
   	 * @param bSetupUndo If true the transform stack will be setup for undo / redo
 	 */
 	UFUNCTION(BlueprintCallable, Category = URigHierarchy)
-	FORCEINLINE_DEBUGGABLE void SetControlValueByIndex(int32 InElementIndex, FRigControlValue InValue, ERigControlValueType InValueType = ERigControlValueType::Current, bool bSetupUndo = false, bool bPrintPythonCommands = false)
+	void SetControlValueByIndex(int32 InElementIndex, FRigControlValue InValue, ERigControlValueType InValueType = ERigControlValueType::Current, bool bSetupUndo = false, bool bPrintPythonCommands = false)
 	{
 		if(Elements.IsValidIndex(InElementIndex))
 		{
@@ -2110,7 +2110,7 @@ public:
 	 * @param InValueType The type of value to set
 	 * @param bSetupUndo If true the transform stack will be setup for undo / redo
 	 */
-	FORCEINLINE_DEBUGGABLE void SetControlValue(int32 InElementIndex, const FRigControlValue& InValue, ERigControlValueType InValueType = ERigControlValueType::Current, bool bSetupUndo = false, bool bPrintPythonCommands = false)
+	void SetControlValue(int32 InElementIndex, const FRigControlValue& InValue, ERigControlValueType InValueType = ERigControlValueType::Current, bool bSetupUndo = false, bool bPrintPythonCommands = false)
 	{
 		SetControlValueByIndex(InElementIndex, InValue, InValueType, bSetupUndo, bPrintPythonCommands);
 	}
@@ -2123,9 +2123,9 @@ public:
 	 * @param bSetupUndo If true the transform stack will be setup for undo / redo
 	 */
 	template<typename T>
-	FORCEINLINE_DEBUGGABLE void SetControlValue(int32 InElementIndex, const T& InValue, ERigControlValueType InValueType = ERigControlValueType::Current, bool bSetupUndo = false) const
+	void SetControlValue(int32 InElementIndex, const T& InValue, ERigControlValueType InValueType = ERigControlValueType::Current, bool bSetupUndo = false)
 	{
-		return SetControlValue(InElementIndex, FRigControlValue::Make<T>(InValue), InValueType, bSetupUndo);
+		return SetControlValueByIndex(InElementIndex, FRigControlValue::Make<T>(InValue), InValueType, bSetupUndo, false);
 	}
 
 	/**
@@ -2134,7 +2134,7 @@ public:
 	 * @param InValue The value to set on the control
 	 * @param bSetupUndo If true the transform stack will be setup for undo / redo
 	 */
-	FORCEINLINE_DEBUGGABLE void SetInitialControlValue(int32 InElementIndex, const FRigControlValue& InValue, bool bSetupUndo = false)
+	void SetInitialControlValue(int32 InElementIndex, const FRigControlValue& InValue, bool bSetupUndo = false)
 	{
 		SetControlValueByIndex(InElementIndex, InValue, ERigControlValueType::Initial, bSetupUndo);
 	}
@@ -2146,7 +2146,7 @@ public:
 	 * @param bSetupUndo If true the transform stack will be setup for undo / redo
 	 */
 	template<typename T>
-	FORCEINLINE_DEBUGGABLE void SetInitialControlValue(int32 InElementIndex, const T& InValue, bool bSetupUndo = false) const
+	void SetInitialControlValue(int32 InElementIndex, const T& InValue, bool bSetupUndo = false) const
 	{
 		return SetInitialControlValue(InElementIndex, FRigControlValue::Make<T>(InValue), bSetupUndo);
 	}
@@ -2157,7 +2157,7 @@ public:
 	 * @param bVisibility The visibility to set on the control
 	 */
 	UFUNCTION(BlueprintCallable, Category = URigHierarchy)
-    FORCEINLINE_DEBUGGABLE void SetControlVisibility(FRigElementKey InKey, bool bVisibility)
+    void SetControlVisibility(FRigElementKey InKey, bool bVisibility)
 	{
 		SetControlVisibilityByIndex(GetIndex(InKey), bVisibility);
 	}
@@ -2168,7 +2168,7 @@ public:
 	 * @param bVisibility The visibility to set on the control
 	 */
 	UFUNCTION(BlueprintCallable, Category = URigHierarchy)
-    FORCEINLINE_DEBUGGABLE void SetControlVisibilityByIndex(int32 InElementIndex, bool bVisibility)
+    void SetControlVisibilityByIndex(int32 InElementIndex, bool bVisibility)
 	{
 		if(Elements.IsValidIndex(InElementIndex))
 		{
@@ -2179,7 +2179,7 @@ public:
 		}
 	}
 
-	FORCEINLINE_DEBUGGABLE void SetControlVisibility(int32 InElementIndex, bool bVisibility)
+	void SetControlVisibility(int32 InElementIndex, bool bVisibility)
 	{
 		SetControlVisibilityByIndex(InElementIndex, bVisibility);
 	}
@@ -2190,7 +2190,7 @@ public:
 	 * @return Returns the value of the curve
 	 */
 	UFUNCTION(BlueprintCallable, Category = URigHierarchy)
-    FORCEINLINE_DEBUGGABLE float GetCurveValue(FRigElementKey InKey) const
+    float GetCurveValue(FRigElementKey InKey) const
 	{
 		return GetCurveValueByIndex(GetIndex(InKey));
 	}
@@ -2201,7 +2201,7 @@ public:
 	 * @return Returns the value of the curve
 	 */
 	UFUNCTION(BlueprintCallable, Category = URigHierarchy)
-    FORCEINLINE_DEBUGGABLE float GetCurveValueByIndex(int32 InElementIndex) const
+    float GetCurveValueByIndex(int32 InElementIndex) const
 	{
 		if(Elements.IsValidIndex(InElementIndex))
 		{
@@ -2214,7 +2214,7 @@ public:
 	}
 
 	// TODO: Deprecate?
-	FORCEINLINE_DEBUGGABLE float GetCurveValue(int32 InElementIndex) const
+	float GetCurveValue(int32 InElementIndex) const
 	{
 		return GetCurveValueByIndex(InElementIndex);
 	}
@@ -2225,7 +2225,7 @@ public:
 	 * @return Returns true if the value is set, false otherwise.
 	 */
 	UFUNCTION(BlueprintCallable, Category = URigHierarchy)
-	FORCEINLINE_DEBUGGABLE bool IsCurveValueSet(FRigElementKey InKey) const
+	bool IsCurveValueSet(FRigElementKey InKey) const
 	{
 		return IsCurveValueSetByIndex(GetIndex(InKey));
 	}
@@ -2236,7 +2236,7 @@ public:
 	 * @return Returns true if the value is set, false otherwise.
 	 */
 	UFUNCTION(BlueprintCallable, Category = URigHierarchy)
-	FORCEINLINE_DEBUGGABLE bool IsCurveValueSetByIndex(int32 InElementIndex) const
+	bool IsCurveValueSetByIndex(int32 InElementIndex) const
 	{
 		if(Elements.IsValidIndex(InElementIndex))
 		{
@@ -2255,7 +2255,7 @@ public:
 	 * @param bSetupUndo If true the transform stack will be setup for undo / redo
 	 */
 	UFUNCTION(BlueprintCallable, Category = URigHierarchy)
-    FORCEINLINE_DEBUGGABLE void SetCurveValue(FRigElementKey InKey, float InValue, bool bSetupUndo = false)
+    void SetCurveValue(FRigElementKey InKey, float InValue, bool bSetupUndo = false)
 	{
 		SetCurveValueByIndex(GetIndex(InKey), InValue, bSetupUndo);
 	}
@@ -2267,7 +2267,7 @@ public:
 	 * @param bSetupUndo If true the transform stack will be setup for undo / redo
 	 */
 	UFUNCTION(BlueprintCallable, Category = URigHierarchy)
-	FORCEINLINE_DEBUGGABLE void SetCurveValueByIndex(int32 InElementIndex, float InValue, bool bSetupUndo = false)
+	void SetCurveValueByIndex(int32 InElementIndex, float InValue, bool bSetupUndo = false)
 	{
 		if(Elements.IsValidIndex(InElementIndex))
 		{
@@ -2279,7 +2279,7 @@ public:
 	}
 
 	// TODO: Deprecate?
-	FORCEINLINE_DEBUGGABLE void SetCurveValue(int32 InElementIndex, float InValue, bool bSetupUndo = false)
+	void SetCurveValue(int32 InElementIndex, float InValue, bool bSetupUndo = false)
 	{
 		SetCurveValueByIndex(InElementIndex, InValue, bSetupUndo);
 	}
@@ -2291,7 +2291,7 @@ public:
 	 * @param bSetupUndo If true the transform stack will be setup for undo / redo
 	 */
 	UFUNCTION(BlueprintCallable, Category = URigHierarchy)
-    FORCEINLINE_DEBUGGABLE void UnsetCurveValue(FRigElementKey InKey, bool bSetupUndo = false)
+    void UnsetCurveValue(FRigElementKey InKey, bool bSetupUndo = false)
 	{
 		UnsetCurveValueByIndex(GetIndex(InKey), bSetupUndo);
 	}
@@ -2303,7 +2303,7 @@ public:
 	 * @param bSetupUndo If true the transform stack will be setup for undo / redo
 	 */
 	UFUNCTION(BlueprintCallable, Category = URigHierarchy)
-	FORCEINLINE_DEBUGGABLE void UnsetCurveValueByIndex(int32 InElementIndex, bool bSetupUndo = false)
+	void UnsetCurveValueByIndex(int32 InElementIndex, bool bSetupUndo = false)
 	{
 		if(Elements.IsValidIndex(InElementIndex))
 		{
@@ -2323,7 +2323,7 @@ public:
 	 * @param bSetupUndo If true the transform stack will be setup for undo / redo
 	 */
 	UFUNCTION(BlueprintCallable, Category = URigHierarchy)
-    FORCEINLINE_DEBUGGABLE void SetControlOffsetTransform(FRigElementKey InKey, FTransform InTransform, bool bInitial = false, bool bAffectChildren = true, bool bSetupUndo = false, bool bPrintPythonCommands = false)
+    void SetControlOffsetTransform(FRigElementKey InKey, FTransform InTransform, bool bInitial = false, bool bAffectChildren = true, bool bSetupUndo = false, bool bPrintPythonCommands = false)
 	{
 		return SetControlOffsetTransformByIndex(GetIndex(InKey), InTransform, bInitial, bAffectChildren, bSetupUndo, bPrintPythonCommands);
 	}
@@ -2337,7 +2337,7 @@ public:
 	 * @param bSetupUndo If true the transform stack will be setup for undo / redo
 	 */
 	UFUNCTION(BlueprintCallable, Category = URigHierarchy)
-    FORCEINLINE_DEBUGGABLE void SetControlOffsetTransformByIndex(int32 InElementIndex, FTransform InTransform, bool bInitial = false, bool bAffectChildren = true, bool bSetupUndo = false, bool bPrintPythonCommands = false)
+    void SetControlOffsetTransformByIndex(int32 InElementIndex, FTransform InTransform, bool bInitial = false, bool bAffectChildren = true, bool bSetupUndo = false, bool bPrintPythonCommands = false)
 	{
 		if(Elements.IsValidIndex(InElementIndex))
 		{
@@ -2356,7 +2356,7 @@ public:
 	 * @param bSetupUndo If true the transform stack will be setup for undo / redo
 	 */
 	UFUNCTION(BlueprintCallable, Category = URigHierarchy)
-    FORCEINLINE_DEBUGGABLE void SetControlShapeTransform(FRigElementKey InKey, FTransform InTransform, bool bInitial = false, bool bSetupUndo = false)
+    void SetControlShapeTransform(FRigElementKey InKey, FTransform InTransform, bool bInitial = false, bool bSetupUndo = false)
 	{
 		return SetControlShapeTransformByIndex(GetIndex(InKey), InTransform, bInitial, bSetupUndo);
 	}
@@ -2369,7 +2369,7 @@ public:
 	 * @param bSetupUndo If true the transform stack will be setup for undo / redo
 	 */
 	UFUNCTION(BlueprintCallable, Category = URigHierarchy)
-    FORCEINLINE_DEBUGGABLE void SetControlShapeTransformByIndex(int32 InElementIndex, FTransform InTransform, bool bInitial = false, bool bSetupUndo = false)
+    void SetControlShapeTransformByIndex(int32 InElementIndex, FTransform InTransform, bool bInitial = false, bool bSetupUndo = false)
 	{
 		if(Elements.IsValidIndex(InElementIndex))
 		{
@@ -2387,7 +2387,7 @@ public:
 	 * @param bSetupUndo If true the transform stack will be setup for undo / redo
 	 */
 	UFUNCTION(BlueprintCallable, Category = URigHierarchy)
-	FORCEINLINE_DEBUGGABLE void SetControlSettings(FRigElementKey InKey, FRigControlSettings InSettings, bool bSetupUndo = false, bool bForce = false, bool bPrintPythonCommands = false)
+	void SetControlSettings(FRigElementKey InKey, FRigControlSettings InSettings, bool bSetupUndo = false, bool bForce = false, bool bPrintPythonCommands = false)
 	{
 		return SetControlSettingsByIndex(GetIndex(InKey), InSettings, bSetupUndo, bForce, bPrintPythonCommands);
 	}
@@ -2399,7 +2399,7 @@ public:
 	 * @param bSetupUndo If true the transform stack will be setup for undo / redo
 	 */
 	UFUNCTION(BlueprintCallable, Category = URigHierarchy)
-	FORCEINLINE_DEBUGGABLE void SetControlSettingsByIndex(int32 InElementIndex, FRigControlSettings InSettings, bool bSetupUndo = false, bool bForce = false, bool bPrintPythonCommands = false)
+	void SetControlSettingsByIndex(int32 InElementIndex, FRigControlSettings InSettings, bool bSetupUndo = false, bool bForce = false, bool bPrintPythonCommands = false)
 	{
 		if(Elements.IsValidIndex(InElementIndex))
 		{
@@ -2418,7 +2418,7 @@ public:
 	 * @return The element's parent's global current or initial transform's value.
 	 */
 	UFUNCTION(BlueprintCallable, Category = URigHierarchy)
-    FORCEINLINE_DEBUGGABLE FTransform GetParentTransform(FRigElementKey InKey, bool bInitial = false) const
+    FTransform GetParentTransform(FRigElementKey InKey, bool bInitial = false) const
 	{
 		return GetParentTransformByIndex(GetIndex(InKey), bInitial);
 	}
@@ -2431,7 +2431,7 @@ public:
 	 * @return The element's parent's global current or initial transform's value.
 	 */
 	UFUNCTION(BlueprintCallable, Category = URigHierarchy)
-    FORCEINLINE_DEBUGGABLE FTransform GetParentTransformByIndex(int32 InElementIndex, bool bInitial = false) const
+    FTransform GetParentTransformByIndex(int32 InElementIndex, bool bInitial = false) const
 	{
 		if(Elements.IsValidIndex(InElementIndex))
 		{
@@ -2776,7 +2776,7 @@ public:
 	 * @return True if the given parent and given child have a parent-child relationship
 	 */
 	UFUNCTION(BlueprintCallable, Category = URigHierarchy)
-	FORCEINLINE bool IsParentedTo(FRigElementKey InChild, FRigElementKey InParent) const
+	bool IsParentedTo(FRigElementKey InChild, FRigElementKey InParent) const
 	{
 		return IsParentedTo(GetIndex(InChild), GetIndex(InParent));
 	}
@@ -2788,7 +2788,7 @@ public:
 	 * @param InDependencyMap An additional map of dependencies to respect
 	 * @return True if the given parent and given child have a parent-child relationship
 	 */
-    FORCEINLINE bool IsParentedTo(int32 InChildIndex, int32 InParentIndex, const TElementDependencyMap& InDependencyMap = TElementDependencyMap()) const
+    bool IsParentedTo(int32 InChildIndex, int32 InParentIndex, const TElementDependencyMap& InDependencyMap = TElementDependencyMap()) const
 	{
 		if(Elements.IsValidIndex(InChildIndex) && Elements.IsValidIndex(InParentIndex))
 		{
@@ -2820,7 +2820,7 @@ public:
 	 * @return The keys of all elements
 	 */
 	UFUNCTION(BlueprintCallable, Category = URigHierarchy, meta = (DisplayName = "Get All Keys", ScriptName = "GetAllKeys"))
-	FORCEINLINE TArray<FRigElementKey> GetAllKeys_ForBlueprint(bool bTraverse = true) const
+	TArray<FRigElementKey> GetAllKeys_ForBlueprint(bool bTraverse = true) const
 	{
 		return GetAllKeys(bTraverse, ERigElementType::All);
 	}
@@ -2865,7 +2865,7 @@ public:
 	 * Starts an interaction on the rig.
 	 * This will cause all transform actions happening to be merged
 	 */
-	FORCEINLINE void EndInteraction()
+	void EndInteraction()
 	{
 		bIsInteracting = false;
 		LastInteractedKey.Reset();
@@ -2896,7 +2896,7 @@ public:
 	 * Returns the delegate to listen to for events coming from this hierarchy
 	 * @return The delegate to listen to for events coming from this hierarchy
 	 */
-	FORCEINLINE FRigEventDelegate& OnEventReceived() { return EventDelegate; }
+	FRigEventDelegate& OnEventReceived() { return EventDelegate; }
 	
 	/**
 	 * Returns true if the hierarchy controller is currently available
@@ -2935,7 +2935,7 @@ public:
 	/**
 	 * Increments the metadata version
 	 */
-	FORCEINLINE void IncrementMetadataVersion(const FRigElementKey& InKey, const FName& InName)
+	void IncrementMetadataVersion(const FRigElementKey& InKey, const FName& InName)
 	{
 		MetadataVersion += 1 + (int32)HashCombine(GetTypeHash(InKey), GetTypeHash(InName));
 	}
@@ -2948,7 +2948,7 @@ public:
 	/**
 	 * Increments the metadataTag version
 	 */
-	FORCEINLINE void IncrementMetadataTagVersion(const FRigElementKey& InKey, const FName& InTag, bool bAdded)
+	void IncrementMetadataTagVersion(const FRigElementKey& InKey, const FName& InTag, bool bAdded)
 	{
 		MetadataTagVersion += 1 + (int32)HashCombine(GetTypeHash(InKey), GetTypeHash(InTag));
 	}
@@ -2959,7 +2959,7 @@ public:
 	 * @return The pose of the hierarchy
 	 */
 	UFUNCTION(BlueprintCallable, Category = URigHierarchy)
-	FORCEINLINE FRigPose GetPose(
+	FRigPose GetPose(
 		bool bInitial = false
 	) const
 	{
@@ -2997,7 +2997,7 @@ public:
 	 * @param InPose The pose to set on the hierarchy
 	 * @param InTransformType The transform type to set
 	 */
-	FORCEINLINE void SetPose(
+	void SetPose(
 		const FRigPose& InPose,
 		ERigTransformType::Type InTransformType = ERigTransformType::CurrentLocal
 	)
@@ -3042,7 +3042,7 @@ public:
 	 * @param InPose The pose to set on the hierarchy
 	 */
 	UFUNCTION(BlueprintCallable, Category = URigHierarchy, meta = (DisplayName = "Set Pose", ScriptName = "SetPose"))
-	FORCEINLINE void SetPose_ForBlueprint(FRigPose InPose)
+	void SetPose_ForBlueprint(FRigPose InPose)
 	{
 		return SetPose(InPose);
 	}
@@ -3053,7 +3053,7 @@ public:
 	 * @return The converted control rig val ue
 	 */
 	UFUNCTION(BlueprintCallable, Category = URigHierarchy)
-	static FORCEINLINE FRigControlValue MakeControlValueFromBool(bool InValue)
+	static FRigControlValue MakeControlValueFromBool(bool InValue)
 	{
 		return FRigControlValue::Make<bool>(InValue);
 	}
@@ -3064,7 +3064,7 @@ public:
 	 * @return The converted control rig val ue
 	 */
 	UFUNCTION(BlueprintCallable, Category = URigHierarchy)
-    static FORCEINLINE FRigControlValue MakeControlValueFromFloat(float InValue)
+    static FRigControlValue MakeControlValueFromFloat(float InValue)
 	{
 		return FRigControlValue::Make<float>(InValue);
 	}
@@ -3075,7 +3075,7 @@ public:
 	 * @return The converted float value
 	 */
 	UFUNCTION(BlueprintCallable, Category = URigHierarchy)
-	static FORCEINLINE float GetFloatFromControlValue(FRigControlValue InValue)
+	static float GetFloatFromControlValue(FRigControlValue InValue)
 	{
 		return InValue.Get<float>();
 	}
@@ -3086,7 +3086,7 @@ public:
 	 * @return The converted control rig val ue
 	 */
 	UFUNCTION(BlueprintCallable, Category = URigHierarchy)
-    static FORCEINLINE FRigControlValue MakeControlValueFromInt(int32 InValue)
+    static FRigControlValue MakeControlValueFromInt(int32 InValue)
 	{
 		return FRigControlValue::Make<int32>(InValue);
 	}
@@ -3097,7 +3097,7 @@ public:
 	 * @return The converted int32 value
 	 */
 	UFUNCTION(BlueprintCallable, Category = URigHierarchy)
-	static FORCEINLINE int32 GetIntFromControlValue(FRigControlValue InValue)
+	static int32 GetIntFromControlValue(FRigControlValue InValue)
 	{
 		return InValue.Get<int32>();
 	}
@@ -3108,7 +3108,7 @@ public:
 	 * @return The converted control rig val ue
 	 */
 	UFUNCTION(BlueprintCallable, Category = URigHierarchy)
-    static FORCEINLINE FRigControlValue MakeControlValueFromVector2D(FVector2D InValue)
+    static FRigControlValue MakeControlValueFromVector2D(FVector2D InValue)
 	{
 		return FRigControlValue::Make<FVector3f>(FVector3f(InValue.X, InValue.Y, 0.f));
 	}
@@ -3119,7 +3119,7 @@ public:
 	 * @return The converted FVector2D value
 	 */
 	UFUNCTION(BlueprintCallable, Category = URigHierarchy)
-	static FORCEINLINE FVector2D GetVector2DFromControlValue(FRigControlValue InValue)
+	static FVector2D GetVector2DFromControlValue(FRigControlValue InValue)
 	{
 		const FVector3f Vector = InValue.Get<FVector3f>();
 		return FVector2D(Vector.X, Vector.Y);
@@ -3131,7 +3131,7 @@ public:
 	 * @return The converted control rig val ue
 	 */
 	UFUNCTION(BlueprintCallable, Category = URigHierarchy)
-    static FORCEINLINE FRigControlValue MakeControlValueFromVector(FVector InValue)
+    static FRigControlValue MakeControlValueFromVector(FVector InValue)
 	{
 		return FRigControlValue::Make<FVector>(InValue);
 	}
@@ -3142,7 +3142,7 @@ public:
 	 * @return The converted FVector value
 	 */
 	UFUNCTION(BlueprintCallable, Category = URigHierarchy)
-	static FORCEINLINE FVector GetVectorFromControlValue(FRigControlValue InValue)
+	static FVector GetVectorFromControlValue(FRigControlValue InValue)
 	{
 		return (FVector)InValue.Get<FVector3f>();
 	}
@@ -3153,7 +3153,7 @@ public:
 	 * @return The converted control rig val ue
 	 */
 	UFUNCTION(BlueprintCallable, Category = URigHierarchy)
-    static FORCEINLINE FRigControlValue MakeControlValueFromRotator(FRotator InValue)
+    static FRigControlValue MakeControlValueFromRotator(FRotator InValue)
 	{
 		return FRigControlValue::Make<FVector>(InValue.Euler());
 	}
@@ -3164,7 +3164,7 @@ public:
 	 * @return The converted FRotator value
 	 */
 	UFUNCTION(BlueprintCallable, Category = URigHierarchy)
-	static FORCEINLINE FRotator GetRotatorFromControlValue(FRigControlValue InValue)
+	static FRotator GetRotatorFromControlValue(FRigControlValue InValue)
 	{
 		return FRotator::MakeFromEuler((FVector)InValue.Get<FVector3f>());
 	}
@@ -3175,7 +3175,7 @@ public:
 	 * @return The converted control rig val ue
 	 */
 	UFUNCTION(BlueprintCallable, Category = URigHierarchy)
-    static FORCEINLINE FRigControlValue MakeControlValueFromTransform(FTransform InValue)
+    static FRigControlValue MakeControlValueFromTransform(FTransform InValue)
 	{
 		return FRigControlValue::Make<FRigControlValue::FTransform_Float>(InValue);
 	}
@@ -3186,7 +3186,7 @@ public:
 	 * @return The converted FTransform value
 	 */
 	UFUNCTION(BlueprintCallable, Category = URigHierarchy)
-	static FORCEINLINE FTransform GetTransformFromControlValue(FRigControlValue InValue)
+	static FTransform GetTransformFromControlValue(FRigControlValue InValue)
 	{
 		return InValue.Get<FRigControlValue::FTransform_Float>().ToTransform();
 	}
@@ -3197,7 +3197,7 @@ public:
 	 * @return The converted control rig val ue
 	 */
 	UFUNCTION(BlueprintCallable, Category = URigHierarchy)
-    static FORCEINLINE FRigControlValue MakeControlValueFromEulerTransform(FEulerTransform InValue)
+    static FRigControlValue MakeControlValueFromEulerTransform(FEulerTransform InValue)
 	{
 		return FRigControlValue::Make<FRigControlValue::FEulerTransform_Float>(InValue);
 	}
@@ -3208,7 +3208,7 @@ public:
 	 * @return The converted FEulerTransform value
 	 */
 	UFUNCTION(BlueprintCallable, Category = URigHierarchy)
-	static FORCEINLINE FEulerTransform GetEulerTransformFromControlValue(FRigControlValue InValue)
+	static FEulerTransform GetEulerTransformFromControlValue(FRigControlValue InValue)
 	{
 		return InValue.Get<FRigControlValue::FEulerTransform_Float>().ToTransform();
 	}
@@ -3219,7 +3219,7 @@ public:
 	 * @return The converted control rig val ue
 	 */
 	UFUNCTION(BlueprintCallable, Category = URigHierarchy)
-    static FORCEINLINE FRigControlValue MakeControlValueFromTransformNoScale(FTransformNoScale InValue)
+    static FRigControlValue MakeControlValueFromTransformNoScale(FTransformNoScale InValue)
 	{
 		return FRigControlValue::Make<FRigControlValue::FTransformNoScale_Float>(InValue);
 	}
@@ -3230,7 +3230,7 @@ public:
 	 * @return The converted FTransformNoScale value
 	 */
 	UFUNCTION(BlueprintCallable, Category = URigHierarchy)
-	static FORCEINLINE FTransformNoScale GetTransformNoScaleFromControlValue(FRigControlValue InValue)
+	static FTransformNoScale GetTransformNoScaleFromControlValue(FRigControlValue InValue)
 	{
 		return InValue.Get<FRigControlValue::FTransformNoScale_Float>().ToTransform();
 	}
@@ -3328,7 +3328,7 @@ public:
 	FRigControlValue GetControlValue(FRigControlElement* InControlElement, ERigControlValueType InValueType) const;
 
 	template<typename T>
-	FORCEINLINE_DEBUGGABLE T GetControlValue(FRigControlElement* InControlElement, ERigControlValueType InValueType) const
+	T GetControlValue(FRigControlElement* InControlElement, ERigControlValueType InValueType) const
 	{
 		return GetControlValue(InControlElement, InValueType).Get<T>();
 	}
@@ -3344,7 +3344,7 @@ public:
 	void SetControlValue(FRigControlElement* InControlElement, const FRigControlValue& InValue, ERigControlValueType InValueType, bool bSetupUndo = false, bool bForce = false, bool bPrintPythonCommands = false, bool bFixEulerFlips = false);
 
 	template<typename T>
-	FORCEINLINE_DEBUGGABLE void SetControlValue(FRigControlElement* InControlElement, const T& InValue, ERigControlValueType InValueType, bool bSetupUndo = false, bool bForce = false) const
+	void SetControlValue(FRigControlElement* InControlElement, const T& InValue, ERigControlValueType InValueType, bool bSetupUndo = false, bool bForce = false) const
 	{
 		SetControlValue(InControlElement, FRigControlValue::Make<T>(InValue), InValueType, bSetupUndo, bForce);
 	}
@@ -3430,7 +3430,7 @@ public:
 	/**
 	 * Returns a reference to the suspend notifications flag
 	 */
-	FORCEINLINE bool& GetSuspendNotificationsFlag() { return bSuspendNotifications; }
+	bool& GetSuspendNotificationsFlag() { return bSuspendNotifications; }
 
 	/*
 	 * Returns true if a hierarchy will record any change.
@@ -3533,7 +3533,7 @@ private:
 	 * Templated helper function to create an element
 	 */
 	template<typename ElementType = FRigBaseElement>
-	FORCEINLINE ElementType* NewElement(int32 Num = 1)
+	ElementType* NewElement(int32 Num = 1)
 	{
 		ElementType* NewElements = (ElementType*)FMemory::Malloc(sizeof(ElementType) * Num);
 		for(int32 Index=0;Index<Num;Index++)
@@ -3559,7 +3559,7 @@ public:
 	/**
 	* Performs validation of the cache within the hierarchy on any mutation.
 	*/
-	FORCEINLINE void EnsureCacheValidity() const
+	void EnsureCacheValidity() const
 	{
 #if WITH_EDITOR
 		if(bEnableCacheValidityCheck)
@@ -3773,7 +3773,7 @@ protected:
 
 #endif
 
-	FORCEINLINE static int32 RigElementTypeToFlatIndex(ERigElementType InElementType)
+	static int32 RigElementTypeToFlatIndex(ERigElementType InElementType)
 	{
 		switch(InElementType)
 		{
@@ -3816,7 +3816,7 @@ protected:
 		return INDEX_NONE;
 	}
 
-	FORCEINLINE static ERigElementType FlatIndexToRigElementType(int32 InIndex)
+	static ERigElementType FlatIndexToRigElementType(int32 InIndex)
 	{
 		switch(InIndex)
 		{
@@ -3870,7 +3870,7 @@ private:
 
 	FTransform GetWorldTransformForReference(const FRigUnitContext* InContext, const FRigElementKey& InKey, bool bInitial);
 	
-	FORCEINLINE static float GetWeightForLerp(const float WeightA, const float WeightB)
+	static float GetWeightForLerp(const float WeightA, const float WeightB)
 	{
 		float Weight = 0.f;
 		const float ClampedWeightA = FMath::Max(WeightA, 0.f);
@@ -3957,13 +3957,13 @@ private:
 #endif
 
 	template<typename T>
-	FORCEINLINE const T& GetMetadata(const FRigElementKey& InItem, ERigMetadataType InType, const FName& InMetadataName, const T& DefaultValue) const
+	const T& GetMetadata(const FRigElementKey& InItem, ERigMetadataType InType, const FName& InMetadataName, const T& DefaultValue) const
 	{
 		return GetMetadata<T>(Find(InItem), InType, InMetadataName, DefaultValue);
 	}
 
 	template<typename T>
-	FORCEINLINE const T& GetMetadata(const FRigBaseElement* InElement, ERigMetadataType InType, const FName& InMetadataName, const T& DefaultValue) const
+	const T& GetMetadata(const FRigBaseElement* InElement, ERigMetadataType InType, const FName& InMetadataName, const T& DefaultValue) const
 	{
 		if(InElement)
 		{
@@ -3976,26 +3976,26 @@ private:
 	}
 
 	template<typename T>
-	FORCEINLINE const TArray<T>& GetArrayMetadata(const FRigElementKey& InItem, ERigMetadataType InType, const FName& InMetadataName) const
+	const TArray<T>& GetArrayMetadata(const FRigElementKey& InItem, ERigMetadataType InType, const FName& InMetadataName) const
 	{
 		return GetArrayMetadata<T>(Find(InItem), InType, InMetadataName);
 	}
 
 	template<typename T>
-	FORCEINLINE const TArray<T>& GetArrayMetadata(const FRigBaseElement* InElement, ERigMetadataType InType, const FName& InMetadataName) const
+	const TArray<T>& GetArrayMetadata(const FRigBaseElement* InElement, ERigMetadataType InType, const FName& InMetadataName) const
 	{
 		static const TArray<T> EmptyArray;
 		return GetMetadata<TArray<T>>(InElement, InType, InMetadataName, EmptyArray);
 	}
 
 	template<typename T>
-	FORCEINLINE bool SetMetadata(const FRigElementKey& InItem, ERigMetadataType InType, const FName& InMetadataName, const T& InValue)
+	bool SetMetadata(const FRigElementKey& InItem, ERigMetadataType InType, const FName& InMetadataName, const T& InValue)
 	{
 		return SetMetadata<T>(Find(InItem), InType, InMetadataName, InValue);
 	}
 
 	template<typename T>
-	FORCEINLINE bool SetMetadata(FRigBaseElement* InElement, ERigMetadataType InType, const FName& InMetadataName, const T& InValue)
+	bool SetMetadata(FRigBaseElement* InElement, ERigMetadataType InType, const FName& InMetadataName, const T& InValue)
 	{
 		if(InElement)
 		{
@@ -4005,13 +4005,13 @@ private:
 	}
 
 	template<typename T>
-	FORCEINLINE bool SetArrayMetadata(const FRigElementKey& InItem, ERigMetadataType InType, const FName& InMetadataName, const TArray<T>& InValue)
+	bool SetArrayMetadata(const FRigElementKey& InItem, ERigMetadataType InType, const FName& InMetadataName, const TArray<T>& InValue)
 	{
 		return SetMetadata<TArray<T>>(Find(InItem), InType, InMetadataName, InValue);
 	}
 
 	template<typename T>
-	FORCEINLINE bool SetArrayMetadata(FRigBaseElement* InElement, ERigMetadataType InType, const FName& InMetadataName, const TArray<T>& InValue)
+	bool SetArrayMetadata(FRigBaseElement* InElement, ERigMetadataType InType, const FName& InMetadataName, const TArray<T>& InValue)
 	{
 		return SetMetadata<TArray<T>>(InElement, InType, InMetadataName, InValue);
 	}
@@ -4021,13 +4021,13 @@ private:
 
 protected:
 
-	FORCEINLINE static void OnMetadataChanged_Static(const FRigElementKey& InKey, const FName& InName, URigHierarchy* InHierarchy)
+	static void OnMetadataChanged_Static(const FRigElementKey& InKey, const FName& InName, URigHierarchy* InHierarchy)
 	{
 		check(InHierarchy);
 		check(IsValid(InHierarchy));
 		InHierarchy->OnMetadataChanged(InKey, InName);
 	}
-	FORCEINLINE static void OnMetadataTagChanged_Static(const FRigElementKey& InKey, const FName& InTag, bool bAdded, URigHierarchy* InHierarchy)
+	static void OnMetadataTagChanged_Static(const FRigElementKey& InKey, const FName& InTag, bool bAdded, URigHierarchy* InHierarchy)
 	{
 		check(InHierarchy);
 		check(IsValid(InHierarchy));
@@ -4072,7 +4072,7 @@ private:
 		ERigHierarchyNotification Type;
 		FRigElementKey Key;
 		
-		FORCEINLINE bool operator == (const FQueuedNotification& InOther) const
+		bool operator == (const FQueuedNotification& InOther) const
 		{
 			return Type == InOther.Type && Key == InOther.Key;
 		}
@@ -4206,14 +4206,14 @@ private:
 };
 
 template<>
-FORCEINLINE_DEBUGGABLE FVector2D URigHierarchy::GetControlValue(FRigControlElement* InControlElement, ERigControlValueType InValueType) const
+inline FVector2D URigHierarchy::GetControlValue(FRigControlElement* InControlElement, ERigControlValueType InValueType) const
 {
 	const FVector3f Value = GetControlValue(InControlElement, InValueType).Get<FVector3f>();
 	return FVector2D(Value.X, Value.Y);
 }
 
 template<>
-FORCEINLINE_DEBUGGABLE void URigHierarchy::SetControlValue(int32 InElementIndex, const FVector2D& InValue, ERigControlValueType InValueType, bool bSetupUndo) const
+inline void URigHierarchy::SetControlValue(int32 InElementIndex, const FVector2D& InValue, ERigControlValueType InValueType, bool bSetupUndo)
 {
 	return SetControlValue(InElementIndex, FRigControlValue::Make<FVector3f>(FVector3f(InValue.X, InValue.Y, 0.f)), InValueType, bSetupUndo);
 }
@@ -4223,7 +4223,7 @@ FORCEINLINE_DEBUGGABLE void URigHierarchy::SetControlValue(int32 InElementIndex,
 class FRigHierarchyListenerGuard
 {
 public:
-	FORCEINLINE_DEBUGGABLE FRigHierarchyListenerGuard(
+	FRigHierarchyListenerGuard(
 		URigHierarchy* InHierarchy, 
 		bool bInEnableInitialChanges, 
 		bool bInEnableCurrentChanges,
@@ -4269,7 +4269,7 @@ public:
 		}
 	}
 
-	FORCEINLINE_DEBUGGABLE ~FRigHierarchyListenerGuard()
+	~FRigHierarchyListenerGuard()
 	{
 		if(ListeningHierarchy == nullptr)
 		{
diff --git a/Engine/Plugins/Animation/ControlRig/Source/ControlRig/Public/Rigs/RigHierarchyCache.h b/Engine/Plugins/Animation/ControlRig/Source/ControlRig/Public/Rigs/RigHierarchyCache.h
--- a/Engine/Plugins/Animation/ControlRig/Source/ControlRig/Public/Rigs/RigHierarchyCache.h
+++ b/Engine/Plugins/Animation/ControlRig/Source/ControlRig/Public/Rigs/RigHierarchyCache.h
@@ -17,13 +17,13 @@ struct CONTROLRIG_API FCachedRigElement
 
 public:
 
-	FORCEINLINE FCachedRigElement()
+	FCachedRigElement()
 		: Key()
 		, Index(UINT16_MAX)
 		, ContainerVersion(INDEX_NONE)
 	{}
 
-	FORCEINLINE FCachedRigElement(const FRigElementKey& InKey, const URigHierarchy* InHierarchy)
+	FCachedRigElement(const FRigElementKey& InKey, const URigHierarchy* InHierarchy)
 		: Key()
 		, Index(UINT16_MAX)
 		, ContainerVersion(INDEX_NONE)
@@ -31,12 +31,12 @@ public:
 		UpdateCache(InKey, InHierarchy);
 	}
 
-	FORCEINLINE bool IsValid() const
+	bool IsValid() const
 	{
 		return GetIndex() != INDEX_NONE && Key.IsValid();
 	}
 
-	FORCEINLINE void Reset()
+	void Reset()
 	{
 		Key = FRigElementKey();
 		Index = UINT16_MAX;
@@ -44,22 +44,22 @@ public:
 		Element = nullptr;
 	}
 
-	FORCEINLINE explicit operator bool() const
+	explicit operator bool() const
 	{
 		return IsValid();
 	}
 
-	FORCEINLINE operator int32() const
+	operator int32() const
 	{
 		return GetIndex();
 	}
 
-	FORCEINLINE explicit operator FRigElementKey() const
+	explicit operator FRigElementKey() const
 	{
 		return Key;
 	}
 
-	FORCEINLINE int32 GetIndex() const
+	int32 GetIndex() const
 	{
 		if(Index == UINT16_MAX)
 		{
@@ -68,12 +68,12 @@ public:
 		return (int32)Index;
 	}
 
-	FORCEINLINE const FRigElementKey& GetKey() const
+	const FRigElementKey& GetKey() const
 	{
 		return Key;
 	}
 
-	FORCEINLINE const FRigBaseElement* GetElement() const
+	const FRigBaseElement* GetElement() const
 	{
 		return Element;
 	}
@@ -82,44 +82,44 @@ public:
 
 	bool UpdateCache(const FRigElementKey& InKey, const URigHierarchy* InHierarchy);
 
-	friend FORCEINLINE uint32 GetTypeHash(const FCachedRigElement& Cache)
+	friend uint32 GetTypeHash(const FCachedRigElement& Cache)
 	{
 		return GetTypeHash(Cache.Key) * 13 + (uint32)Cache.Index;
 	}
 
 	bool IsIdentical(const FRigElementKey& InKey, const URigHierarchy* InHierarchy);
 
-	FORCEINLINE bool operator ==(const FCachedRigElement& Other) const
+	bool operator ==(const FCachedRigElement& Other) const
 	{
 		return Index == Other.Index && Key == Other.Key;
 	}
 
-	FORCEINLINE bool operator !=(const FCachedRigElement& Other) const
+	bool operator !=(const FCachedRigElement& Other) const
 	{
 		return Index != Other.Index || Key != Other.Key;
 	}
 
-	FORCEINLINE bool operator ==(const FRigElementKey& Other) const
+	bool operator ==(const FRigElementKey& Other) const
 	{
 		return Key == Other;
 	}
 
-	FORCEINLINE bool operator !=(const FRigElementKey& Other) const
+	bool operator !=(const FRigElementKey& Other) const
 	{
 		return Key != Other;
 	}
 
-	FORCEINLINE bool operator ==(const int32& Other) const
+	bool operator ==(const int32& Other) const
 	{
 		return GetIndex() == Other;
 	}
 
-	FORCEINLINE bool operator !=(const int32& Other) const
+	bool operator !=(const int32& Other) const
 	{
 		return GetIndex() != Other;
 	}
 
-	FORCEINLINE bool operator <(const FCachedRigElement& Other) const
+	bool operator <(const FCachedRigElement& Other) const
 	{
 		if (Key < Other.Key)
 		{
@@ -128,7 +128,7 @@ public:
 		return Index < Other.Index;
 	}
 
-	FORCEINLINE bool operator >(const FCachedRigElement& Other) const
+	bool operator >(const FCachedRigElement& Other) const
 	{
 		if (Key > Other.Key)
 		{
diff --git a/Engine/Plugins/Animation/ControlRig/Source/ControlRig/Public/Rigs/RigHierarchyController.h b/Engine/Plugins/Animation/ControlRig/Source/ControlRig/Public/Rigs/RigHierarchyController.h
--- a/Engine/Plugins/Animation/ControlRig/Source/ControlRig/Public/Rigs/RigHierarchyController.h
+++ b/Engine/Plugins/Animation/ControlRig/Source/ControlRig/Public/Rigs/RigHierarchyController.h
@@ -53,7 +53,7 @@ public:
 	 * @return Returns true if the selection was applied
 	 */
 	UFUNCTION(BlueprintCallable, Category = URigHierarchyController)
-    FORCEINLINE bool DeselectElement(FRigElementKey InKey)
+    bool DeselectElement(FRigElementKey InKey)
 	{
 		return SelectElement(InKey, false);
 	}
@@ -71,7 +71,7 @@ public:
 	 * @return Returns true if the selection was applied
 	 */
 	UFUNCTION(BlueprintCallable, Category = URigHierarchyController)
-    FORCEINLINE bool ClearSelection()
+    bool ClearSelection()
 	{
 		return SetSelection(TArray<FRigElementKey>());
 	}
@@ -136,7 +136,7 @@ public:
 	 * @return The key for the newly created control.
 	 */
 	UFUNCTION(BlueprintCallable, Category = URigHierarchyController, meta = (DisplayName = "Add Control", ScriptName = "AddControl"))
-    FORCEINLINE FRigElementKey AddControl_ForBlueprint(
+    FRigElementKey AddControl_ForBlueprint(
         FName InName,
         FRigElementKey InParent,
         FRigControlSettings InSettings,
@@ -174,7 +174,7 @@ public:
 	 * @return The key for the newly created animation channel.
 	 */
 	UFUNCTION(BlueprintCallable, Category = URigHierarchyController, meta = (DisplayName = "Add Control", ScriptName = "AddAnimationChannel"))
-    FORCEINLINE FRigElementKey AddAnimationChannel_ForBlueprint(
+    FRigElementKey AddAnimationChannel_ForBlueprint(
         FName InName,
         FRigElementKey InParentControl,
         FRigControlSettings InSettings,
@@ -557,7 +557,7 @@ public:
 	/**
 	 * Returns a reference to the suspend notifications flag
 	 */
-	FORCEINLINE bool& GetSuspendNotificationsFlag() { return bSuspendAllNotifications; }
+	bool& GetSuspendNotificationsFlag() { return bSuspendAllNotifications; }
 
 #if WITH_EDITOR
 	UFUNCTION(BlueprintCallable, Category = URigHierarchyController)
@@ -595,7 +595,7 @@ private:
 	 * Returns true if this controller is valid / linked to a valid hierarchy.
 	 * @return Returns true if this controller is valid / linked to a valid hierarchy.
 	 */
-	FORCEINLINE bool IsValid() const { return Hierarchy.IsValid(); }
+	bool IsValid() const { return Hierarchy.IsValid(); }
 
 	/**
 	 * Adds a new element to the hierarchy
