ControlRig, RigVM: Remove use of FORCEINLINE

diff --git a/Engine/Plugins/Animation/ControlRig/Source/ControlRig/Public/Rigs/RigHierarchyDefines.h b/Engine/Plugins/Animation/ControlRig/Source/ControlRig/Public/Rigs/RigHierarchyDefines.h
--- a/Engine/Plugins/Animation/ControlRig/Source/ControlRig/Public/Rigs/RigHierarchyDefines.h
+++ b/Engine/Plugins/Animation/ControlRig/Source/ControlRig/Public/Rigs/RigHierarchyDefines.h
@@ -215,18 +215,18 @@ struct CONTROLRIG_API FRigControlLimitEnabled
 	}
 
 	void Serialize(FArchive& Ar);
-	FORCEINLINE friend FArchive& operator<<(FArchive& Ar, FRigControlLimitEnabled& P)
+	friend FArchive& operator<<(FArchive& Ar, FRigControlLimitEnabled& P)
 	{
 		P.Serialize(Ar);
 		return Ar;
 	}
 
-	FORCEINLINE bool operator ==(const FRigControlLimitEnabled& Other) const
+	bool operator ==(const FRigControlLimitEnabled& Other) const
 	{
 		return bMinimum == Other.bMinimum && bMaximum == Other.bMaximum;
 	}
 
-	FORCEINLINE bool operator !=(const FRigControlLimitEnabled& Other) const
+	bool operator !=(const FRigControlLimitEnabled& Other) const
 	{
 		return bMinimum != Other.bMinimum || bMaximum != Other.bMaximum;
 	}
@@ -398,38 +398,38 @@ public:
 	{
 	}
 
-	FORCEINLINE_DEBUGGABLE bool IsValid() const
+	bool IsValid() const
 	{
 		return FloatStorage.bValid;
 	}
 
 	template<class T>
-	FORCEINLINE_DEBUGGABLE T Get() const
+	T Get() const
 	{
 		return GetRef<T>();
 	}
 
 	template<class T>
-	FORCEINLINE_DEBUGGABLE T& GetRef()
+	T& GetRef()
 	{
 		FloatStorage.bValid = true;
 		return *(T*)&FloatStorage;
 	}
 
 	template<class T>
-	FORCEINLINE_DEBUGGABLE const T& GetRef() const
+	const T& GetRef() const
 	{
 		return *(T*)&FloatStorage;
 	}
 
 	template<class T>
-	FORCEINLINE_DEBUGGABLE void Set(T InValue)
+	void Set(T InValue)
 	{
 		GetRef<T>() = InValue;
 	}
 
 	template<class T>
-	FORCEINLINE_DEBUGGABLE FString ToString() const
+	FString ToString() const
 	{
 		FString Result;
 		TBaseStructure<T>::Get()->ExportText(Result, &GetRef<T>(), nullptr, nullptr, PPF_None, nullptr);
@@ -437,7 +437,7 @@ public:
 	}
 
 	template<class T>
-	FORCEINLINE_DEBUGGABLE T SetFromString(const FString& InString)
+	T SetFromString(const FString& InString)
 	{
 		T Value;
 		TBaseStructure<T>::Get()->ImportText(*InString, &Value, nullptr, PPF_None, nullptr, TBaseStructure<T>::Get()->GetName());
@@ -447,7 +447,7 @@ public:
 
 	friend FArchive& operator <<(FArchive& Ar, FRigControlValue& Value);
 
-	FORCEINLINE_DEBUGGABLE FString ToPythonString(ERigControlType InControlType) const
+	FString ToPythonString(ERigControlType InControlType) const
 	{
 		FString ValueStr;
 		
@@ -499,14 +499,14 @@ public:
 	}
 
 	template<class T>
-	FORCEINLINE_DEBUGGABLE static FRigControlValue Make(T InValue)
+	static FRigControlValue Make(T InValue)
 	{
 		FRigControlValue Value;
 		Value.Set<T>(InValue);
 		return Value;
 	}
 
-	FORCEINLINE FTransform GetAsTransform(ERigControlType InControlType, ERigControlAxis InPrimaryAxis) const
+	FTransform GetAsTransform(ERigControlType InControlType, ERigControlAxis InPrimaryAxis) const
 	{
 		FTransform Transform = FTransform::Identity;
 		switch (InControlType)
@@ -630,7 +630,7 @@ public:
 		return Transform;
 	}
 	
-	FORCEINLINE void SetFromTransform(const FTransform& InTransform, ERigControlType InControlType, ERigControlAxis InPrimaryAxis)
+	void SetFromTransform(const FTransform& InTransform, ERigControlType InControlType, ERigControlAxis InPrimaryAxis)
 	{
 		switch (InControlType)
 		{
@@ -769,7 +769,7 @@ public:
 		}
 	}
 
-	FORCEINLINE void ApplyLimits(
+	void ApplyLimits(
 		const TArray<FRigControlLimitEnabled>& LimitEnabled,
 		ERigControlType InControlType,
 		const FRigControlValue& InMinimumValue,
@@ -1447,39 +1447,39 @@ public:
 	void Serialize(FArchive& Ar);
 	void Save(FArchive& Ar);
 	void Load(FArchive& Ar);
-	FORCEINLINE friend FArchive& operator<<(FArchive& Ar, FRigElementKey& P)
+	friend FArchive& operator<<(FArchive& Ar, FRigElementKey& P)
 	{
 		P.Serialize(Ar);
 		return Ar;
 	}
 
-	FORCEINLINE bool IsValid() const
+	bool IsValid() const
 	{
 		return Name != NAME_None && Type != ERigElementType::None;
 	}
 
-	FORCEINLINE explicit operator bool() const
+	explicit operator bool() const
 	{
 		return IsValid();
 	}
 
-	FORCEINLINE void Reset()
+	void Reset()
 	{
 		Type = ERigElementType::Curve;
 		Name = NAME_None;
 	}
 
-	FORCEINLINE bool IsTypeOf(ERigElementType InElementType) const
+	bool IsTypeOf(ERigElementType InElementType) const
 	{
 		return ((uint8)InElementType & (uint8)Type) == (uint8)Type;
 	}
 
-	friend FORCEINLINE uint32 GetTypeHash(const FRigElementKey& Key)
+	friend uint32 GetTypeHash(const FRigElementKey& Key)
 	{
 		return GetTypeHash(Key.Name) * 10 + (uint32)Key.Type;
 	}
 
-	friend FORCEINLINE uint32 GetTypeHash(const TArrayView<const FRigElementKey>& Keys)
+	friend uint32 GetTypeHash(const TArrayView<const FRigElementKey>& Keys)
 	{
 		uint32 Hash = (uint32)(Keys.Num() * 17 + 3);
 		for (const FRigElementKey& Key : Keys)
@@ -1489,22 +1489,22 @@ public:
 		return Hash;
 	}
 
-	friend FORCEINLINE uint32 GetTypeHash(const TArray<FRigElementKey>& Keys)
+	friend uint32 GetTypeHash(const TArray<FRigElementKey>& Keys)
 	{
 		return GetTypeHash(TArrayView<const FRigElementKey>(Keys.GetData(), Keys.Num()));
 	}
 
-	FORCEINLINE bool operator ==(const FRigElementKey& Other) const
+	bool operator ==(const FRigElementKey& Other) const
 	{
 		return Name == Other.Name && Type == Other.Type;
 	}
 
-	FORCEINLINE bool operator !=(const FRigElementKey& Other) const
+	bool operator !=(const FRigElementKey& Other) const
 	{
 		return Name != Other.Name || Type != Other.Type;
 	}
 
-	FORCEINLINE bool operator <(const FRigElementKey& Other) const
+	bool operator <(const FRigElementKey& Other) const
 	{
 		if (Type < Other.Type)
 		{
@@ -1513,7 +1513,7 @@ public:
 		return Name.LexicalLess(Other.Name);
 	}
 
-	FORCEINLINE bool operator >(const FRigElementKey& Other) const
+	bool operator >(const FRigElementKey& Other) const
 	{
 		if (Type > Other.Type)
 		{
@@ -1522,7 +1522,7 @@ public:
 		return Other.Name.LexicalLess(Name);
 	}
 
-	FORCEINLINE FString ToString() const
+	FString ToString() const
 	{
 		switch (Type)
 		{
@@ -1562,98 +1562,98 @@ struct CONTROLRIG_API FRigElementKeyCollection
 {
 	GENERATED_BODY()
 
-	FORCEINLINE FRigElementKeyCollection()
+	FRigElementKeyCollection()
 	{
 	}
 
-	FORCEINLINE FRigElementKeyCollection(const TArray<FRigElementKey>& InKeys)
+	FRigElementKeyCollection(const TArray<FRigElementKey>& InKeys)
 		: Keys(InKeys)
 	{
 	}
 
 	// Resets the data structure and maintains all storage.
-	FORCEINLINE void Reset()
+	void Reset()
 	{
 		Keys.Reset();
 	}
 
 	// Resets the data structure and removes all storage.
-	FORCEINLINE void Empty()
+	void Empty()
 	{
 		Keys.Empty();
 	}
 
 	// Returns true if a given instruction index is valid.
-	FORCEINLINE bool IsValidIndex(int32 InIndex) const
+	bool IsValidIndex(int32 InIndex) const
 	{
 		return Keys.IsValidIndex(InIndex);
 	}
 
 	// Returns the number of elements in this collection.
-	FORCEINLINE int32 Num() const { return Keys.Num(); }
+	int32 Num() const { return Keys.Num(); }
 
 	// Returns true if this collection contains no elements.
-	FORCEINLINE bool IsEmpty() const
+	bool IsEmpty() const
 	{
 		return Num() == 0;
 	}
 
 	// Returns the first element of this collection
-	FORCEINLINE const FRigElementKey& First() const
+	const FRigElementKey& First() const
 	{
 		return Keys[0];
 	}
 
 	// Returns the first element of this collection
-	FORCEINLINE FRigElementKey& First()
+	FRigElementKey& First()
 	{
 		return Keys[0];
 	}
 
 	// Returns the last element of this collection
-	FORCEINLINE const FRigElementKey& Last() const
+	const FRigElementKey& Last() const
 	{
 		return Keys.Last();
 	}
 
 	// Returns the last element of this collection
-	FORCEINLINE FRigElementKey& Last()
+	FRigElementKey& Last()
 	{
 		return Keys.Last();
 	}
 
-	FORCEINLINE int32 Add(const FRigElementKey& InKey)
+	int32 Add(const FRigElementKey& InKey)
 	{
 		return Keys.Add(InKey);
 	}
 
-	FORCEINLINE int32 AddUnique(const FRigElementKey& InKey)
+	int32 AddUnique(const FRigElementKey& InKey)
 	{
 		return Keys.AddUnique(InKey);
 	}
 
-	FORCEINLINE bool Contains(const FRigElementKey& InKey) const
+	bool Contains(const FRigElementKey& InKey) const
 	{
 		return Keys.Contains(InKey);
 	}
 
 	// const accessor for an element given its index
-	FORCEINLINE const FRigElementKey& operator[](int32 InIndex) const
+	const FRigElementKey& operator[](int32 InIndex) const
 	{
 		return Keys[InIndex];
 	}
 
-	FORCEINLINE const TArray<FRigElementKey>& GetKeys() const
+	const TArray<FRigElementKey>& GetKeys() const
 	{
 		return Keys;
 	}
 	   
-	FORCEINLINE TArray<FRigElementKey>::RangedForIteratorType      begin() { return Keys.begin(); }
-	FORCEINLINE TArray<FRigElementKey>::RangedForConstIteratorType begin() const { return Keys.begin(); }
-	FORCEINLINE TArray<FRigElementKey>::RangedForIteratorType      end() { return Keys.end(); }
-	FORCEINLINE TArray<FRigElementKey>::RangedForConstIteratorType end() const { return Keys.end(); }
+	TArray<FRigElementKey>::RangedForIteratorType      begin() { return Keys.begin(); }
+	TArray<FRigElementKey>::RangedForConstIteratorType begin() const { return Keys.begin(); }
+	TArray<FRigElementKey>::RangedForIteratorType      end() { return Keys.end(); }
+	TArray<FRigElementKey>::RangedForConstIteratorType end() const { return Keys.end(); }
 
-	friend FORCEINLINE uint32 GetTypeHash(const FRigElementKeyCollection& Collection)
+	friend uint32 GetTypeHash(const FRigElementKeyCollection& Collection)
 	{
 		return GetTypeHash(Collection.GetKeys());
 	}
diff --git a/Engine/Plugins/Animation/ControlRig/Source/ControlRig/Public/Rigs/RigHierarchyElements.h b/Engine/Plugins/Animation/ControlRig/Source/ControlRig/Public/Rigs/RigHierarchyElements.h
--- a/Engine/Plugins/Animation/ControlRig/Source/ControlRig/Public/Rigs/RigHierarchyElements.h
+++ b/Engine/Plugins/Animation/ControlRig/Source/ControlRig/Public/Rigs/RigHierarchyElements.h
@@ -19,22 +19,22 @@ DECLARE_DELEGATE_RetVal_ThreeParams(FTransform, FRigReferenceGetWorldTransformDe
 
 #define DECLARE_RIG_ELEMENT_METHODS(ElementType) \
 template<typename T> \
-friend FORCEINLINE const T* Cast(const ElementType* InElement) \
+friend const T* Cast(const ElementType* InElement) \
 { \
    return Cast<T>((const FRigBaseElement*) InElement); \
 } \
 template<typename T> \
-friend FORCEINLINE T* Cast(ElementType* InElement) \
+friend T* Cast(ElementType* InElement) \
 { \
    return Cast<T>((FRigBaseElement*) InElement); \
 } \
 template<typename T> \
-friend FORCEINLINE const T* CastChecked(const ElementType* InElement) \
+friend const T* CastChecked(const ElementType* InElement) \
 { \
 	return CastChecked<T>((const FRigBaseElement*) InElement); \
 } \
 template<typename T> \
-friend FORCEINLINE T* CastChecked(ElementType* InElement) \
+friend T* CastChecked(ElementType* InElement) \
 { \
 	return CastChecked<T>((FRigBaseElement*) InElement); \
 }
@@ -54,7 +54,7 @@ namespace ERigTransformType
 
 namespace ERigTransformType
 {
-	FORCEINLINE ERigTransformType::Type SwapCurrentAndInitial(const Type InTransformType)
+	inline ERigTransformType::Type SwapCurrentAndInitial(const Type InTransformType)
 	{
 		switch(InTransformType)
 		{
@@ -78,7 +78,7 @@ namespace ERigTransformType
 		return CurrentGlobal;
 	}
 
-	FORCEINLINE Type SwapLocalAndGlobal(const Type InTransformType)
+	inline Type SwapLocalAndGlobal(const Type InTransformType)
 	{
 		switch(InTransformType)
 		{
@@ -102,7 +102,7 @@ namespace ERigTransformType
 		return InitialLocal;
 	}
 
-	FORCEINLINE Type MakeLocal(const Type InTransformType)
+	inline Type MakeLocal(const Type InTransformType)
 	{
 		switch(InTransformType)
 		{
@@ -119,7 +119,7 @@ namespace ERigTransformType
 		return InitialLocal;
 	}
 
-	FORCEINLINE Type MakeGlobal(const Type InTransformType)
+	inline Type MakeGlobal(const Type InTransformType)
 	{
 		switch(InTransformType)
 		{
@@ -136,7 +136,7 @@ namespace ERigTransformType
 		return InitialGlobal;
 	}
 
-	FORCEINLINE Type MakeInitial(const Type InTransformType)
+	inline Type MakeInitial(const Type InTransformType)
 	{
 		switch(InTransformType)
 		{
@@ -153,7 +153,7 @@ namespace ERigTransformType
 		return InitialGlobal;
 	}
 
-	FORCEINLINE Type MakeCurrent(const Type InTransformType)
+	inline Type MakeCurrent(const Type InTransformType)
 	{
 		switch(InTransformType)
 		{
@@ -170,7 +170,7 @@ namespace ERigTransformType
 		return CurrentGlobal;
 	}
 
-	FORCEINLINE bool IsLocal(const Type InTransformType)
+	inline bool IsLocal(const Type InTransformType)
 	{
 		switch(InTransformType)
 		{
@@ -187,7 +187,7 @@ namespace ERigTransformType
 		return false;
 	}
 
-	FORCEINLINE bool IsGlobal(const Type InTransformType)
+	inline bool IsGlobal(const Type InTransformType)
 	{
 		switch(InTransformType)
 		{
@@ -204,7 +204,7 @@ namespace ERigTransformType
 		return false;
 	}
 
-	FORCEINLINE bool IsInitial(const Type InTransformType)
+	inline bool IsInitial(const Type InTransformType)
 	{
 		switch(InTransformType)
 		{
@@ -221,7 +221,7 @@ namespace ERigTransformType
 		return false;
 	}
 
-	FORCEINLINE bool IsCurrent(const Type InTransformType)
+	inline bool IsCurrent(const Type InTransformType)
 	{
 		switch(InTransformType)
 		{
@@ -252,7 +252,7 @@ struct CONTROLRIG_API FRigComputedTransform
 	void Save(FArchive& Ar);
 	void Load(FArchive& Ar);
 
-	FORCEINLINE_DEBUGGABLE void Set(const FTransform& InTransform)
+	void Set(const FTransform& InTransform)
 	{
 #if WITH_EDITOR
 		ensure(InTransform.GetRotation().IsNormalized());
@@ -264,14 +264,14 @@ struct CONTROLRIG_API FRigComputedTransform
 		bDirty = false;
 	}
 
-	FORCEINLINE static bool Equals(const FTransform& A, const FTransform& B, const float InTolerance = 0.0001f)
+	static bool Equals(const FTransform& A, const FTransform& B, const float InTolerance = 0.0001f)
 	{
 		return (A.GetTranslation() - B.GetTranslation()).IsNearlyZero(InTolerance) &&
 			A.GetRotation().Equals(B.GetRotation(), InTolerance) &&
 			(A.GetScale3D() - B.GetScale3D()).IsNearlyZero(InTolerance);
 	}
 
-	FORCEINLINE bool operator == (const FRigComputedTransform& Other) const
+	bool operator == (const FRigComputedTransform& Other) const
 	{
 		return bDirty == Other.bDirty && Equals(Transform, Other.Transform);
     }
@@ -295,7 +295,7 @@ struct CONTROLRIG_API FRigLocalAndGlobalTransform
 	void Save(FArchive& Ar);
 	void Load(FArchive& Ar);
 
-	FORCEINLINE bool operator == (const FRigLocalAndGlobalTransform& Other) const
+	bool operator == (const FRigLocalAndGlobalTransform& Other) const
 	{
 		return Local == Other.Local && Global == Other.Global;
 	}
@@ -317,7 +317,7 @@ struct CONTROLRIG_API FRigCurrentAndInitialTransform
     , Initial()
 	{}
 
-	FORCEINLINE const FRigComputedTransform& operator[](const ERigTransformType::Type InTransformType) const
+	const FRigComputedTransform& operator[](const ERigTransformType::Type InTransformType) const
 	{
 		switch(InTransformType)
 		{
@@ -341,7 +341,7 @@ struct CONTROLRIG_API FRigCurrentAndInitialTransform
 		return Initial.Global;
 	}
 
-	FORCEINLINE FRigComputedTransform& operator[](const ERigTransformType::Type InTransformType)
+	FRigComputedTransform& operator[](const ERigTransformType::Type InTransformType)
 	{
 		switch(InTransformType)
 		{
@@ -365,22 +365,22 @@ struct CONTROLRIG_API FRigCurrentAndInitialTransform
 		return Initial.Global;
 	}
 
-	FORCEINLINE const FTransform& Get(const ERigTransformType::Type InTransformType) const
+	const FTransform& Get(const ERigTransformType::Type InTransformType) const
 	{
 		return operator[](InTransformType).Transform;
 	}
 
-	FORCEINLINE void Set(const ERigTransformType::Type InTransformType, const FTransform& InTransform)
+	void Set(const ERigTransformType::Type InTransformType, const FTransform& InTransform)
 	{
 		operator[](InTransformType).Set(InTransform);
 	}
 
-	FORCEINLINE bool IsDirty(const ERigTransformType::Type InTransformType) const
+	bool IsDirty(const ERigTransformType::Type InTransformType) const
 	{
 		return operator[](InTransformType).bDirty;
 	}
 
-	FORCEINLINE void MarkDirty(const ERigTransformType::Type InTransformType)
+	void MarkDirty(const ERigTransformType::Type InTransformType)
 	{
 		ensure(!(operator[](ERigTransformType::SwapLocalAndGlobal(InTransformType)).bDirty));
 		operator[](InTransformType).bDirty = true;
@@ -389,7 +389,7 @@ struct CONTROLRIG_API FRigCurrentAndInitialTransform
 	void Save(FArchive& Ar);
 	void Load(FArchive& Ar);
 
-	FORCEINLINE bool operator == (const FRigCurrentAndInitialTransform& Other) const
+	bool operator == (const FRigCurrentAndInitialTransform& Other) const
 	{
 		return Current == Other.Current && Initial == Other.Initial;
 	}
@@ -417,7 +417,7 @@ struct CONTROLRIG_API FRigPreferredEulerAngles
 	void Save(FArchive& Ar);
 	void Load(FArchive& Ar);
 
-	FORCEINLINE bool operator == (const FRigPreferredEulerAngles& Other) const
+	bool operator == (const FRigPreferredEulerAngles& Other) const
 	{
 		return RotationOrder == Other.RotationOrder &&
 			Current == Other.Current &&
@@ -425,8 +425,8 @@ struct CONTROLRIG_API FRigPreferredEulerAngles
 	}
 
 	void Reset();
-	FORCEINLINE FVector& Get(bool bInitial = false) { return bInitial ? Initial : Current; }
-	FORCEINLINE const FVector& Get(bool bInitial = false) const { return bInitial ? Initial : Current; }
+	FVector& Get(bool bInitial = false) { return bInitial ? Initial : Current; }
+	const FVector& Get(bool bInitial = false) const { return bInitial ? Initial : Current; }
 	FRotator GetRotator(bool bInitial = false) const;
 	FRotator SetRotator(const FRotator& InValue, bool bInitial = false, bool bFixEulerFlips = false);
 	FVector GetAngles(bool bInitial = false, EEulerRotationOrder InRotationOrder = DefaultRotationOrder) const;
@@ -554,7 +554,7 @@ protected:
 	TArray<FRigBaseMetadata*> Metadata;
 	TMap<FName,int32> MetadataNameToIndex;
 
-	FORCEINLINE static bool IsClassOf(const FRigBaseElement* InElement)
+	static bool IsClassOf(const FRigBaseElement* InElement)
 	{
 		return true;
 	}
@@ -566,21 +566,21 @@ public:
 	virtual void Save(FArchive& Ar, URigHierarchy* Hierarchy, ESerializationPhase SerializationPhase);
 	virtual void Load(FArchive& Ar, URigHierarchy* Hierarchy, ESerializationPhase SerializationPhase);
 
-	FORCEINLINE const FName& GetName() const { return Key.Name; }
-	FORCEINLINE const FString& GetNameString() const { return NameString; }
-	FORCEINLINE virtual const FName& GetDisplayName() const { return GetName(); }
-	FORCEINLINE ERigElementType GetType() const { return Key.Type; }
-	FORCEINLINE const FRigElementKey& GetKey() const { return Key; }
-	FORCEINLINE int32 GetIndex() const { return Index; }
-	FORCEINLINE int32 GetSubIndex() const { return SubIndex; }
-	FORCEINLINE bool IsSelected() const { return bSelected; }
-	FORCEINLINE int32 GetCreatedAtInstructionIndex() const { return CreatedAtInstructionIndex; }
-	FORCEINLINE bool IsProcedural() const { return CreatedAtInstructionIndex != INDEX_NONE; }
-	FORCEINLINE int32 GetMetadataVersion() const { return MetadataVersion; }
-
-	FORCEINLINE int32 NumMetadata() const { return Metadata.Num(); }
-	FORCEINLINE FRigBaseMetadata* GetMetadata(int32 InIndex) const { return Metadata[InIndex]; }
-	FORCEINLINE FRigBaseMetadata* GetMetadata(const FName& InName) const
+	const FName& GetName() const { return Key.Name; }
+	const FString& GetNameString() const { return NameString; }
+	virtual const FName& GetDisplayName() const { return GetName(); }
+	ERigElementType GetType() const { return Key.Type; }
+	const FRigElementKey& GetKey() const { return Key; }
+	int32 GetIndex() const { return Index; }
+	int32 GetSubIndex() const { return SubIndex; }
+	bool IsSelected() const { return bSelected; }
+	int32 GetCreatedAtInstructionIndex() const { return CreatedAtInstructionIndex; }
+	bool IsProcedural() const { return CreatedAtInstructionIndex != INDEX_NONE; }
+	int32 GetMetadataVersion() const { return MetadataVersion; }
+
+	int32 NumMetadata() const { return Metadata.Num(); }
+	FRigBaseMetadata* GetMetadata(int32 InIndex) const { return Metadata[InIndex]; }
+	FRigBaseMetadata* GetMetadata(const FName& InName) const
 	{
 		if(const int32* MetadataIndex = MetadataNameToIndex.Find(InName))
 		{
@@ -588,7 +588,7 @@ public:
 		}
 		return nullptr;
 	}
-	FORCEINLINE FRigBaseMetadata* GetMetadata(const FName& InName, ERigMetadataType InType) const
+	FRigBaseMetadata* GetMetadata(const FName& InName, ERigMetadataType InType) const
 	{
 		if(const int32* MetadataIndex = MetadataNameToIndex.Find(InName))
 		{
@@ -600,7 +600,7 @@ public:
 		}
 		return nullptr;
 	}
-	FORCEINLINE bool SetMetaData(const FName& InName, ERigMetadataType InType, const void* InData, int32 InSize)
+	bool SetMetaData(const FName& InName, ERigMetadataType InType, const void* InData, int32 InSize)
 	{
 		if(FRigBaseMetadata* Md = SetupValidMetadata(InName, InType))
 		{
@@ -612,15 +612,15 @@ public:
 	bool RemoveAllMetadata();
 
 	template<typename T>
-	FORCEINLINE bool IsA() const { return T::IsClassOf(this); }
+	bool IsA() const { return T::IsClassOf(this); }
 
-	FORCEINLINE bool IsTypeOf(ERigElementType InElementType) const
+	bool IsTypeOf(ERigElementType InElementType) const
 	{
 		return Key.IsTypeOf(InElementType);
 	}
 
 	template<typename T>
-    friend FORCEINLINE const T* Cast(const FRigBaseElement* InElement)
+    friend const T* Cast(const FRigBaseElement* InElement)
 	{
 		if(InElement)
 		{
@@ -633,7 +633,7 @@ public:
 	}
 
 	template<typename T>
-    friend FORCEINLINE T* Cast(FRigBaseElement* InElement)
+    friend T* Cast(FRigBaseElement* InElement)
 	{
 		if(InElement)
 		{
@@ -646,7 +646,7 @@ public:
 	}
 
 	template<typename T>
-    friend FORCEINLINE const T* CastChecked(const FRigBaseElement* InElement)
+    friend const T* CastChecked(const FRigBaseElement* InElement)
 	{
 		const T* Element = Cast<T>(InElement);
 		check(Element);
@@ -654,7 +654,7 @@ public:
 	}
 
 	template<typename T>
-    friend FORCEINLINE T* CastChecked(FRigBaseElement* InElement)
+    friend T* CastChecked(FRigBaseElement* InElement)
 	{
 		T* Element = Cast<T>(InElement);
 		check(Element);
@@ -712,22 +712,22 @@ protected:
 
 	struct FElementToDirty
 	{
-		FORCEINLINE FElementToDirty()
+		FElementToDirty()
 			: Element(nullptr)
 			, HierarchyDistance(INDEX_NONE)
 		{}
 
-		FORCEINLINE FElementToDirty(FRigTransformElement* InElement, int32 InHierarchyDistance = INDEX_NONE)
+		FElementToDirty(FRigTransformElement* InElement, int32 InHierarchyDistance = INDEX_NONE)
 			: Element(InElement)
 			, HierarchyDistance(InHierarchyDistance)
 		{}
 
-		FORCEINLINE bool operator ==(const FElementToDirty& Other) const
+		bool operator ==(const FElementToDirty& Other) const
 		{
 			return Element == Other.Element;
 		}
 
-		FORCEINLINE bool operator !=(const FElementToDirty& Other) const
+		bool operator !=(const FElementToDirty& Other) const
 		{
 			return Element != Other.Element;
 		}
@@ -740,7 +740,7 @@ protected:
 	typedef TArray<FElementToDirty, TInlineAllocator<3>> FElementsToDirtyArray;  
 	FElementsToDirtyArray ElementsToDirty;
 
-	FORCEINLINE static bool IsClassOf(const FRigBaseElement* InElement)
+	static bool IsClassOf(const FRigBaseElement* InElement)
 	{
 		return InElement->GetType() == ERigElementType::Bone ||
 			InElement->GetType() == ERigElementType::Null ||
@@ -788,7 +788,7 @@ protected:
 
 	virtual void CopyFrom(URigHierarchy* InHierarchy, FRigBaseElement* InOther, URigHierarchy* InOtherHierarchy) override;
 
-	FORCEINLINE static bool IsClassOf(const FRigBaseElement* InElement)
+	static bool IsClassOf(const FRigBaseElement* InElement)
 	{
 		return InElement->GetType() == ERigElementType::Bone ||
 			InElement->GetType() == ERigElementType::RigidBody ||
@@ -813,25 +813,25 @@ public:
 	UPROPERTY(BlueprintReadWrite, EditAnywhere, Category = Weight)
 	float Scale;
 
-	FORCEINLINE FRigElementWeight()
+	FRigElementWeight()
 		: Location(1.f)
 		, Rotation(1.f)
 		, Scale(1.f)
 	{}
 
-	FORCEINLINE FRigElementWeight(float InWeight)
+	FRigElementWeight(float InWeight)
 		: Location(InWeight)
 		, Rotation(InWeight)
 		, Scale(InWeight)
 	{}
 
-	FORCEINLINE FRigElementWeight(float InLocation, float InRotation, float InScale)
+	FRigElementWeight(float InLocation, float InRotation, float InScale)
 		: Location(InLocation)
 		, Rotation(InRotation)
 		, Scale(InScale)
 	{}
 
-	FORCEINLINE friend FArchive& operator <<(FArchive& Ar, FRigElementWeight& Weight)
+	friend FArchive& operator <<(FArchive& Ar, FRigElementWeight& Weight)
 	{
 		Ar << Weight.Location;
 		Ar << Weight.Rotation;
@@ -839,32 +839,32 @@ public:
 		return Ar;
 	}
 
-	FORCEINLINE bool AffectsLocation() const
+	bool AffectsLocation() const
 	{
 		return Location > SMALL_NUMBER;
 	}
 
-	FORCEINLINE bool AffectsRotation() const
+	bool AffectsRotation() const
 	{
 		return Rotation > SMALL_NUMBER;
 	}
 
-	FORCEINLINE bool AffectsScale() const
+	bool AffectsScale() const
 	{
 		return Scale > SMALL_NUMBER;
 	}
 
-	FORCEINLINE bool IsAlmostZero() const
+	bool IsAlmostZero() const
 	{
 		return !AffectsLocation() && !AffectsRotation() && !AffectsScale();
 	}
 
-	FORCEINLINE friend FRigElementWeight operator *(FRigElementWeight InWeight, float InScale)
+	friend FRigElementWeight operator *(FRigElementWeight InWeight, float InScale)
 	{
 		return FRigElementWeight(InWeight.Location * InScale, InWeight.Rotation * InScale, InWeight.Scale * InScale);
 	}
 
-	FORCEINLINE friend FRigElementWeight operator *(float InScale, FRigElementWeight InWeight)
+	friend FRigElementWeight operator *(float InScale, FRigElementWeight InWeight)
 	{
 		return FRigElementWeight(InWeight.Location * InScale, InWeight.Rotation * InScale, InWeight.Scale * InScale);
 	}
@@ -881,18 +881,18 @@ public:
 	FRigElementWeight InitialWeight;
 	mutable FRigComputedTransform Cache;
 		
-	FORCEINLINE FRigElementParentConstraint()
+	FRigElementParentConstraint()
 		: ParentElement(nullptr)
 	{
 		Cache.bDirty = true;
 	}
 
-	FORCEINLINE const FRigElementWeight& GetWeight(bool bInitial = false) const
+	const FRigElementWeight& GetWeight(bool bInitial = false) const
 	{
 		return bInitial ? InitialWeight : Weight;
 	}
 
-	FORCEINLINE void CopyPose(const FRigElementParentConstraint& InOther, bool bCurrent, bool bInitial)
+	void CopyPose(const FRigElementParentConstraint& InOther, bool bCurrent, bool bInitial)
 	{
 		if(bCurrent)
 		{
@@ -936,7 +936,7 @@ protected:
 
 	virtual void CopyFrom(URigHierarchy* InHierarchy, FRigBaseElement* InOther, URigHierarchy* InOtherHierarchy) override;
 
-	FORCEINLINE static bool IsClassOf(const FRigBaseElement* InElement)
+	static bool IsClassOf(const FRigBaseElement* InElement)
 	{
 		return InElement->GetType() == ERigElementType::Null ||
 			InElement->GetType() == ERigElementType::Control;
@@ -979,7 +979,7 @@ protected:
 
 	virtual void CopyFrom(URigHierarchy* InHierarchy, FRigBaseElement* InOther, URigHierarchy* InOtherHierarchy) override;
 
-	FORCEINLINE static bool IsClassOf(const FRigBaseElement* InElement)
+	static bool IsClassOf(const FRigBaseElement* InElement)
 	{
 		return InElement->GetType() == ERigElementType::Bone;
 	}
@@ -1005,7 +1005,7 @@ public:
 
 protected:
 	
-	FORCEINLINE static bool IsClassOf(const FRigBaseElement* InElement)
+	static bool IsClassOf(const FRigBaseElement* InElement)
 	{
 		return InElement->GetType() == ERigElementType::Null;
 	}
@@ -1158,13 +1158,13 @@ struct CONTROLRIG_API FRigControlSettings
 #endif
 	
 	/** Applies the limits expressed by these settings to a value */
-	FORCEINLINE void ApplyLimits(FRigControlValue& InOutValue) const
+	void ApplyLimits(FRigControlValue& InOutValue) const
 	{
 		InOutValue.ApplyLimits(LimitEnabled, ControlType, MinimumValue, MaximumValue);
 	}
 
 	/** Applies the limits expressed by these settings to a transform */
-	FORCEINLINE void ApplyLimits(FTransform& InOutValue) const
+	void ApplyLimits(FTransform& InOutValue) const
 	{
 		FRigControlValue Value;
 		Value.SetFromTransform(InOutValue, ControlType, PrimaryAxis);
@@ -1172,7 +1172,7 @@ struct CONTROLRIG_API FRigControlSettings
 		InOutValue = Value.GetAsTransform(ControlType, PrimaryAxis);
 	}
 
-	FORCEINLINE FRigControlValue GetIdentityValue() const
+	FRigControlValue GetIdentityValue() const
 	{
 		FRigControlValue Value;
 		Value.SetFromTransform(FTransform::Identity, ControlType, PrimaryAxis);
@@ -1181,36 +1181,36 @@ struct CONTROLRIG_API FRigControlSettings
 
 	bool operator == (const FRigControlSettings& InOther) const;
 
-	FORCEINLINE bool operator != (const FRigControlSettings& InOther) const
+	bool operator != (const FRigControlSettings& InOther) const
 	{
 		return !(*this == InOther);
 	}
 
 	void SetupLimitArrayForType(bool bLimitTranslation = false, bool bLimitRotation = false, bool bLimitScale = false);
 
-	FORCEINLINE bool IsAnimatable() const
+	bool IsAnimatable() const
 	{
 		return (AnimationType == ERigControlAnimationType::AnimationControl) ||
 			(AnimationType == ERigControlAnimationType::AnimationChannel);
 	}
 
-	FORCEINLINE bool ShouldBeGrouped() const
+	bool ShouldBeGrouped() const
 	{
 		return IsAnimatable() && bGroupWithParentControl;
 	}
 
-	FORCEINLINE bool SupportsShape() const
+	bool SupportsShape() const
 	{
 		return (AnimationType != ERigControlAnimationType::AnimationChannel) &&
 			(ControlType != ERigControlType::Bool);
 	}
 
-	FORCEINLINE bool IsVisible() const
+	bool IsVisible() const
 	{
 		return SupportsShape() && bShapeVisible;
 	}
 	
-	FORCEINLINE bool SetVisible(bool bVisible, bool bForce = false)
+	bool SetVisible(bool bVisible, bool bForce = false)
 	{
 		if(!bForce)
 		{
@@ -1234,14 +1234,14 @@ struct CONTROLRIG_API FRigControlSettings
 		return SupportsShape();
 	}
 
-	FORCEINLINE bool IsSelectable(bool bRespectVisibility = true) const
+	bool IsSelectable(bool bRespectVisibility = true) const
 	{
 		return (AnimationType == ERigControlAnimationType::AnimationControl ||
 			AnimationType == ERigControlAnimationType::ProxyControl) &&
 			(IsVisible() || !bRespectVisibility);
 	}
 
-	FORCEINLINE void SetAnimationTypeFromDeprecatedData(bool bAnimatable, bool bShapeEnabled)
+	void SetAnimationTypeFromDeprecatedData(bool bAnimatable, bool bShapeEnabled)
 	{
 		if(bAnimatable)
 		{
@@ -1277,7 +1277,7 @@ struct CONTROLRIG_API FRigControlElement : public FRigMultiParentElement
 
 	virtual ~FRigControlElement(){}
 	
-	FORCEINLINE virtual const FName& GetDisplayName() const override
+	virtual const FName& GetDisplayName() const override
 	{
 		if(!Settings.DisplayName.IsNone())
 		{
@@ -1286,9 +1286,9 @@ struct CONTROLRIG_API FRigControlElement : public FRigMultiParentElement
 		return FRigMultiParentElement::GetDisplayName();
 	}
 
-	FORCEINLINE bool IsAnimationChannel() const { return Settings.AnimationType == ERigControlAnimationType::AnimationChannel; }
+	bool IsAnimationChannel() const { return Settings.AnimationType == ERigControlAnimationType::AnimationChannel; }
 
-	FORCEINLINE bool CanDriveControls() const { return Settings.AnimationType == ERigControlAnimationType::ProxyControl || Settings.AnimationType == ERigControlAnimationType::AnimationControl; }
+	bool CanDriveControls() const { return Settings.AnimationType == ERigControlAnimationType::ProxyControl || Settings.AnimationType == ERigControlAnimationType::AnimationControl; }
 
 	virtual void Save(FArchive& A, URigHierarchy* Hierarchy, ESerializationPhase SerializationPhase) override;
 	virtual void Load(FArchive& Ar, URigHierarchy* Hierarchy, ESerializationPhase SerializationPhase) override;
@@ -1313,7 +1313,7 @@ protected:
 	UPROPERTY(BlueprintReadOnly, EditAnywhere, Category = RigElement)
 	FRigPreferredEulerAngles PreferredEulerAngles;
 	
-	FORCEINLINE static bool IsClassOf(const FRigBaseElement* InElement)
+	static bool IsClassOf(const FRigBaseElement* InElement)
 	{
 		return InElement->GetType() == ERigElementType::Control;
 	}
@@ -1360,7 +1360,7 @@ public:
 	
 	float Value;
 
-	FORCEINLINE static bool IsClassOf(const FRigBaseElement* InElement)
+	static bool IsClassOf(const FRigBaseElement* InElement)
 	{
 		return InElement->GetType() == ERigElementType::Curve;
 	}
@@ -1419,7 +1419,7 @@ public:
 
 protected:
 	
-	FORCEINLINE static bool IsClassOf(const FRigBaseElement* InElement)
+	static bool IsClassOf(const FRigBaseElement* InElement)
 	{
 		return InElement->GetType() == ERigElementType::RigidBody;
 	}
@@ -1456,7 +1456,7 @@ protected:
 
 	virtual void CopyFrom(URigHierarchy* InHierarchy, FRigBaseElement* InOther, URigHierarchy* InOtherHierarchy) override;
 
-	FORCEINLINE static bool IsClassOf(const FRigBaseElement* InElement)
+	static bool IsClassOf(const FRigBaseElement* InElement)
 	{
 		return InElement->GetType() == ERigElementType::Reference;
 	}
diff --git a/Engine/Plugins/Animation/ControlRig/Source/ControlRig/Public/Rigs/RigHierarchyMetadata.h b/Engine/Plugins/Animation/ControlRig/Source/ControlRig/Public/Rigs/RigHierarchyMetadata.h
--- a/Engine/Plugins/Animation/ControlRig/Source/ControlRig/Public/Rigs/RigHierarchyMetadata.h
+++ b/Engine/Plugins/Animation/ControlRig/Source/ControlRig/Public/Rigs/RigHierarchyMetadata.h
@@ -11,22 +11,22 @@ class URigHierarchyController;
 
 #define DECLARE_RIG_METADATA_METHODS(MetadataType) \
 template<typename T> \
-friend FORCEINLINE const T* Cast(const MetadataType* InMetadata) \
+friend const T* Cast(const MetadataType* InMetadata) \
 { \
    return Cast<T>((const FRigBaseMetadata*) InMetadata); \
 } \
 template<typename T> \
-friend FORCEINLINE T* Cast(MetadataType* InMetadata) \
+friend T* Cast(MetadataType* InMetadata) \
 { \
    return Cast<T>((FRigBaseMetadata*) InMetadata); \
 } \
 template<typename T> \
-friend FORCEINLINE const T* CastChecked(const MetadataType* InMetadata) \
+friend const T* CastChecked(const MetadataType* InMetadata) \
 { \
 	return CastChecked<T>((const FRigBaseMetadata*) InMetadata); \
 } \
 template<typename T> \
-friend FORCEINLINE T* CastChecked(MetadataType* InMetadata) \
+friend T* CastChecked(MetadataType* InMetadata) \
 { \
 	return CastChecked<T>((FRigBaseMetadata*) InMetadata); \
 }
@@ -59,7 +59,7 @@ protected:
 
 	mutable const FProperty* ValueProperty;
 
-	FORCEINLINE static bool IsClassOf(const FRigBaseMetadata* InMetadata)
+	static bool IsClassOf(const FRigBaseMetadata* InMetadata)
 	{
 		return true;
 	}
@@ -69,24 +69,24 @@ public:
 	UScriptStruct* GetMetadataStruct() const;
 	virtual void Serialize(FArchive& Ar, bool bIsLoading);
 
-	FORCEINLINE bool IsValid() const { return (GetType() != ERigMetadataType::Invalid) && GetKey().IsValid(); }
-	FORCEINLINE const FRigBaseElement* GetElement() const { return Element; }
+	bool IsValid() const { return (GetType() != ERigMetadataType::Invalid) && GetKey().IsValid(); }
+	const FRigBaseElement* GetElement() const { return Element; }
 	const FRigElementKey& GetKey() const;
-	FORCEINLINE const FName& GetName() const { return Name; }
-	FORCEINLINE const ERigMetadataType& GetType() const { return Type; }
-	FORCEINLINE bool IsArray() const
+	const FName& GetName() const { return Name; }
+	const ERigMetadataType& GetType() const { return Type; }
+	bool IsArray() const
 	{
 		return GetValueProperty()->IsA<FArrayProperty>();
 	}
-	FORCEINLINE const void* GetValueData() const
+	const void* GetValueData() const
 	{
 		return GetValueProperty()->ContainerPtrToValuePtr<void>(this);
 	}
-	FORCEINLINE const int32& GetValueSize() const
+	const int32& GetValueSize() const
 	{
 		return GetValueProperty()->ElementSize;
 	}
-	FORCEINLINE bool SetValueData(const void* InData, int32 InSize)
+	bool SetValueData(const void* InData, int32 InSize)
 	{
 		if(InData)
 		{
@@ -99,10 +99,10 @@ public:
 	}
 
 	template<typename T>
-	FORCEINLINE bool IsA() const { return T::IsClassOf(this); }
+	bool IsA() const { return T::IsClassOf(this); }
 
 	template<typename T>
-    friend FORCEINLINE const T* Cast(const FRigBaseMetadata* InMetadata)
+    friend const T* Cast(const FRigBaseMetadata* InMetadata)
 	{
 		if(InMetadata)
 		{
@@ -115,7 +115,7 @@ public:
 	}
 
 	template<typename T>
-    friend FORCEINLINE T* Cast(FRigBaseMetadata* InMetadata)
+    friend T* Cast(FRigBaseMetadata* InMetadata)
 	{
 		if(InMetadata)
 		{
@@ -128,7 +128,7 @@ public:
 	}
 
 	template<typename T>
-    friend FORCEINLINE const T* CastChecked(const FRigBaseMetadata* InMetadata)
+    friend const T* CastChecked(const FRigBaseMetadata* InMetadata)
 	{
 		const T* CastElement = Cast<T>(InMetadata);
 		check(CastElement);
@@ -136,7 +136,7 @@ public:
 	}
 
 	template<typename T>
-    friend FORCEINLINE T* CastChecked(FRigBaseMetadata* InMetadata)
+    friend T* CastChecked(FRigBaseMetadata* InMetadata)
 	{
 		T* CastElement = Cast<T>(InMetadata);
 		check(CastElement);
@@ -149,7 +149,7 @@ protected:
 	static FRigBaseMetadata* MakeMetadata(const FRigBaseElement* InElement, const FName& InName, ERigMetadataType InType);
 	static void DestroyMetadata(FRigBaseMetadata** Metadata);
 
-	FORCEINLINE const FProperty* GetValueProperty() const
+	const FProperty* GetValueProperty() const
 	{
 		if(ValueProperty == nullptr)
 		{
@@ -159,12 +159,12 @@ protected:
 		return ValueProperty;
 	}
 
-	FORCEINLINE void* GetValueData()
+	void* GetValueData()
 	{
 		return GetValueProperty()->ContainerPtrToValuePtr<void>(this);
 	}
 
-	FORCEINLINE virtual bool SetValueDataImpl(const void* InData)
+	virtual bool SetValueDataImpl(const void* InData)
 	{
 		const FProperty* Property = GetValueProperty();
 		if(!Property->Identical(GetValueData(), InData))
@@ -193,16 +193,16 @@ public:
 		, Value(false)
 	{}
 	virtual ~FRigBoolMetadata() override {}
-	FORCEINLINE virtual void Serialize(FArchive& Ar, bool bIsLoading) override
+	virtual void Serialize(FArchive& Ar, bool bIsLoading) override
 	{
 		Super::Serialize(Ar, bIsLoading);
 		Ar << Value;
 	}
-	FORCEINLINE const bool& GetValue() const { return Value; }
-	FORCEINLINE bool& GetValue() { return Value; }
-	FORCEINLINE bool SetValue(const bool& InValue) { return SetValueData(&InValue, sizeof(bool)); }
+	const bool& GetValue() const { return Value; }
+	bool& GetValue() { return Value; }
+	bool SetValue(const bool& InValue) { return SetValueData(&InValue, sizeof(bool)); }
 
-	FORCEINLINE static bool IsClassOf(const FRigBaseMetadata* InMetadata)
+	static bool IsClassOf(const FRigBaseMetadata* InMetadata)
 	{
 		return InMetadata->GetType() == ERigMetadataType::Bool;
 	}
@@ -225,16 +225,16 @@ public:
 		: FRigBaseMetadata()
 	{}
 	virtual ~FRigBoolArrayMetadata() override {}
-	FORCEINLINE virtual void Serialize(FArchive& Ar, bool bIsLoading) override
+	virtual void Serialize(FArchive& Ar, bool bIsLoading) override
 	{
 		Super::Serialize(Ar, bIsLoading);
 		Ar << Value;
 	}
-	FORCEINLINE const TArray<bool>& GetValue() const { return Value; }
-	FORCEINLINE TArray<bool>& GetValue() { return Value; }
-	FORCEINLINE bool SetValue(const TArray<bool>& InValue) { return SetValueData(&InValue, sizeof(TArray<bool>)); }
+	const TArray<bool>& GetValue() const { return Value; }
+	TArray<bool>& GetValue() { return Value; }
+	bool SetValue(const TArray<bool>& InValue) { return SetValueData(&InValue, sizeof(TArray<bool>)); }
 
-	FORCEINLINE static bool IsClassOf(const FRigBaseMetadata* InMetadata)
+	static bool IsClassOf(const FRigBaseMetadata* InMetadata)
 	{
 		return InMetadata->GetType() == ERigMetadataType::BoolArray;
 	}
@@ -258,16 +258,16 @@ public:
 		, Value(0.f)
 	{}
 	virtual ~FRigFloatMetadata() override {}
-	FORCEINLINE virtual void Serialize(FArchive& Ar, bool bIsLoading) override
+	virtual void Serialize(FArchive& Ar, bool bIsLoading) override
 	{
 		Super::Serialize(Ar, bIsLoading);
 		Ar << Value;
 	}
-	FORCEINLINE const float& GetValue() const { return Value; } 
-	FORCEINLINE float& GetValue() { return Value; }
-	FORCEINLINE bool SetValue(const float& InValue) { return SetValueData(&InValue, sizeof(float)); }
+	const float& GetValue() const { return Value; } 
+	float& GetValue() { return Value; }
+	bool SetValue(const float& InValue) { return SetValueData(&InValue, sizeof(float)); }
 
-	FORCEINLINE static bool IsClassOf(const FRigBaseMetadata* InMetadata)
+	static bool IsClassOf(const FRigBaseMetadata* InMetadata)
     {
     	return InMetadata->GetType() == ERigMetadataType::Float;
     }
@@ -290,16 +290,16 @@ public:
 		: FRigBaseMetadata()
 	{}
 	virtual ~FRigFloatArrayMetadata() override {}
-	FORCEINLINE virtual void Serialize(FArchive& Ar, bool bIsLoading) override
+	virtual void Serialize(FArchive& Ar, bool bIsLoading) override
 	{
 		Super::Serialize(Ar, bIsLoading);
 		Ar << Value;
 	}
-	FORCEINLINE const TArray<float>& GetValue() const { return Value; }
-	FORCEINLINE TArray<float>& GetValue() { return Value; }
-	FORCEINLINE bool SetValue(const TArray<float>& InValue) { return SetValueData(&InValue, sizeof(TArray<float>)); }
+	const TArray<float>& GetValue() const { return Value; }
+	TArray<float>& GetValue() { return Value; }
+	bool SetValue(const TArray<float>& InValue) { return SetValueData(&InValue, sizeof(TArray<float>)); }
 
-	FORCEINLINE static bool IsClassOf(const FRigBaseMetadata* InMetadata)
+	static bool IsClassOf(const FRigBaseMetadata* InMetadata)
 	{
 		return InMetadata->GetType() == ERigMetadataType::FloatArray;
 	}
@@ -323,16 +323,16 @@ public:
 		, Value(0)
 	{}
 	virtual ~FRigInt32Metadata() override {}
-	FORCEINLINE virtual void Serialize(FArchive& Ar, bool bIsLoading) override
+	virtual void Serialize(FArchive& Ar, bool bIsLoading) override
 	{
 		Super::Serialize(Ar, bIsLoading);
 		Ar << Value;
 	}
-	FORCEINLINE const int32& GetValue() const { return Value; } 
-	FORCEINLINE int32& GetValue() { return Value; }
-	FORCEINLINE bool SetValue(const int32& InValue) { return SetValueData(&InValue, sizeof(int32)); }
+	const int32& GetValue() const { return Value; } 
+	int32& GetValue() { return Value; }
+	bool SetValue(const int32& InValue) { return SetValueData(&InValue, sizeof(int32)); }
 
-	FORCEINLINE static bool IsClassOf(const FRigBaseMetadata* InMetadata)
+	static bool IsClassOf(const FRigBaseMetadata* InMetadata)
     {
     	return InMetadata->GetType() == ERigMetadataType::Int32;
     }
@@ -355,16 +355,16 @@ public:
 		: FRigBaseMetadata()
 	{}
 	virtual ~FRigInt32ArrayMetadata() override {}
-	FORCEINLINE virtual void Serialize(FArchive& Ar, bool bIsLoading) override
+	virtual void Serialize(FArchive& Ar, bool bIsLoading) override
 	{
 		Super::Serialize(Ar, bIsLoading);
 		Ar << Value;
 	}
-	FORCEINLINE const TArray<int32>& GetValue() const { return Value; }
-	FORCEINLINE TArray<int32>& GetValue() { return Value; }
-	FORCEINLINE bool SetValue(const TArray<int32>& InValue) { return SetValueData(&InValue, sizeof(TArray<int32>)); }
+	const TArray<int32>& GetValue() const { return Value; }
+	TArray<int32>& GetValue() { return Value; }
+	bool SetValue(const TArray<int32>& InValue) { return SetValueData(&InValue, sizeof(TArray<int32>)); }
 
-	FORCEINLINE static bool IsClassOf(const FRigBaseMetadata* InMetadata)
+	static bool IsClassOf(const FRigBaseMetadata* InMetadata)
 	{
 		return InMetadata->GetType() == ERigMetadataType::Int32Array;
 	}
@@ -388,16 +388,16 @@ public:
 		, Value(NAME_None)
 	{}
 	virtual ~FRigNameMetadata() override {}
-	FORCEINLINE virtual void Serialize(FArchive& Ar, bool bIsLoading) override
+	virtual void Serialize(FArchive& Ar, bool bIsLoading) override
 	{
 		Super::Serialize(Ar, bIsLoading);
 		Ar << Value;
 	}
-	FORCEINLINE const FName& GetValue() const { return Value; } 
-	FORCEINLINE FName& GetValue() { return Value; }
-	FORCEINLINE bool SetValue(const FName& InValue) { return SetValueData(&InValue, sizeof(FName)); }
+	const FName& GetValue() const { return Value; } 
+	FName& GetValue() { return Value; }
+	bool SetValue(const FName& InValue) { return SetValueData(&InValue, sizeof(FName)); }
 
-	FORCEINLINE static bool IsClassOf(const FRigBaseMetadata* InMetadata)
+	static bool IsClassOf(const FRigBaseMetadata* InMetadata)
     {
     	return InMetadata->GetType() == ERigMetadataType::Name;
     }
@@ -420,16 +420,16 @@ public:
 		: FRigBaseMetadata()
 	{}
 	virtual ~FRigNameArrayMetadata() override {}
-	FORCEINLINE virtual void Serialize(FArchive& Ar, bool bIsLoading) override
+	virtual void Serialize(FArchive& Ar, bool bIsLoading) override
 	{
 		Super::Serialize(Ar, bIsLoading);
 		Ar << Value;
 	}
-	FORCEINLINE const TArray<FName>& GetValue() const { return Value; }
-	FORCEINLINE TArray<FName>& GetValue() { return Value; }
-	FORCEINLINE bool SetValue(const TArray<FName>& InValue) { return SetValueData(&InValue, sizeof(TArray<FName>)); }
+	const TArray<FName>& GetValue() const { return Value; }
+	TArray<FName>& GetValue() { return Value; }
+	bool SetValue(const TArray<FName>& InValue) { return SetValueData(&InValue, sizeof(TArray<FName>)); }
 
-	FORCEINLINE static bool IsClassOf(const FRigBaseMetadata* InMetadata)
+	static bool IsClassOf(const FRigBaseMetadata* InMetadata)
 	{
 		return InMetadata->GetType() == ERigMetadataType::NameArray;
 	}
@@ -453,16 +453,16 @@ public:
 		, Value(FVector::ZeroVector)
 	{}
 	virtual ~FRigVectorMetadata() override {}
-	FORCEINLINE virtual void Serialize(FArchive& Ar, bool bIsLoading) override
+	virtual void Serialize(FArchive& Ar, bool bIsLoading) override
 	{
 		Super::Serialize(Ar, bIsLoading);
 		Ar << Value;
 	}
-	FORCEINLINE const FVector& GetValue() const { return Value; } 
-	FORCEINLINE FVector& GetValue() { return Value; }
-	FORCEINLINE bool SetValue(const FVector& InValue) { return SetValueData(&InValue, sizeof(FVector)); }
+	const FVector& GetValue() const { return Value; } 
+	FVector& GetValue() { return Value; }
+	bool SetValue(const FVector& InValue) { return SetValueData(&InValue, sizeof(FVector)); }
 
-	FORCEINLINE static bool IsClassOf(const FRigBaseMetadata* InMetadata)
+	static bool IsClassOf(const FRigBaseMetadata* InMetadata)
     {
     	return InMetadata->GetType() == ERigMetadataType::Vector;
     }
@@ -485,16 +485,16 @@ public:
 		: FRigBaseMetadata()
 	{}
 	virtual ~FRigVectorArrayMetadata() override {}
-	FORCEINLINE virtual void Serialize(FArchive& Ar, bool bIsLoading) override
+	virtual void Serialize(FArchive& Ar, bool bIsLoading) override
 	{
 		Super::Serialize(Ar, bIsLoading);
 		Ar << Value;
 	}
-	FORCEINLINE const TArray<FVector>& GetValue() const { return Value; }
-	FORCEINLINE TArray<FVector>& GetValue() { return Value; }
-	FORCEINLINE bool SetValue(const TArray<FVector>& InValue) { return SetValueData(&InValue, sizeof(TArray<FVector>)); }
+	const TArray<FVector>& GetValue() const { return Value; }
+	TArray<FVector>& GetValue() { return Value; }
+	bool SetValue(const TArray<FVector>& InValue) { return SetValueData(&InValue, sizeof(TArray<FVector>)); }
 
-	FORCEINLINE static bool IsClassOf(const FRigBaseMetadata* InMetadata)
+	static bool IsClassOf(const FRigBaseMetadata* InMetadata)
 	{
 		return InMetadata->GetType() == ERigMetadataType::VectorArray;
 	}
@@ -518,16 +518,16 @@ public:
 		, Value(FRotator::ZeroRotator)
 	{}
 	virtual ~FRigRotatorMetadata() override {}
-	FORCEINLINE virtual void Serialize(FArchive& Ar, bool bIsLoading) override
+	virtual void Serialize(FArchive& Ar, bool bIsLoading) override
 	{
 		Super::Serialize(Ar, bIsLoading);
 		Ar << Value;
 	}
-	FORCEINLINE const FRotator& GetValue() const { return Value; } 
-	FORCEINLINE FRotator& GetValue() { return Value; }
-	FORCEINLINE bool SetValue(const FRotator& InValue) { return SetValueData(&InValue, sizeof(FRotator)); }
+	const FRotator& GetValue() const { return Value; } 
+	FRotator& GetValue() { return Value; }
+	bool SetValue(const FRotator& InValue) { return SetValueData(&InValue, sizeof(FRotator)); }
 
-	FORCEINLINE static bool IsClassOf(const FRigBaseMetadata* InMetadata)
+	static bool IsClassOf(const FRigBaseMetadata* InMetadata)
     {
     	return InMetadata->GetType() == ERigMetadataType::Rotator;
     }
@@ -550,16 +550,16 @@ public:
 		: FRigBaseMetadata()
 	{}
 	virtual ~FRigRotatorArrayMetadata() override {}
-	FORCEINLINE virtual void Serialize(FArchive& Ar, bool bIsLoading) override
+	virtual void Serialize(FArchive& Ar, bool bIsLoading) override
 	{
 		Super::Serialize(Ar, bIsLoading);
 		Ar << Value;
 	}
-	FORCEINLINE const TArray<FRotator>& GetValue() const { return Value; }
-	FORCEINLINE TArray<FRotator>& GetValue() { return Value; }
-	FORCEINLINE bool SetValue(const TArray<FRotator>& InValue) { return SetValueData(&InValue, sizeof(TArray<FRotator>)); }
+	const TArray<FRotator>& GetValue() const { return Value; }
+	TArray<FRotator>& GetValue() { return Value; }
+	bool SetValue(const TArray<FRotator>& InValue) { return SetValueData(&InValue, sizeof(TArray<FRotator>)); }
 
-	FORCEINLINE static bool IsClassOf(const FRigBaseMetadata* InMetadata)
+	static bool IsClassOf(const FRigBaseMetadata* InMetadata)
 	{
 		return InMetadata->GetType() == ERigMetadataType::RotatorArray;
 	}
@@ -583,16 +583,16 @@ public:
 		, Value(FQuat::Identity)
 	{}
 	virtual ~FRigQuatMetadata() override {}
-	FORCEINLINE virtual void Serialize(FArchive& Ar, bool bIsLoading) override
+	virtual void Serialize(FArchive& Ar, bool bIsLoading) override
 	{
 		Super::Serialize(Ar, bIsLoading);
 		Ar << Value;
 	}
-	FORCEINLINE const FQuat& GetValue() const { return Value; } 
-	FORCEINLINE FQuat& GetValue() { return Value; }
-	FORCEINLINE bool SetValue(const FQuat& InValue) { return SetValueData(&InValue, sizeof(FQuat)); }
+	const FQuat& GetValue() const { return Value; } 
+	FQuat& GetValue() { return Value; }
+	bool SetValue(const FQuat& InValue) { return SetValueData(&InValue, sizeof(FQuat)); }
 
-	FORCEINLINE static bool IsClassOf(const FRigBaseMetadata* InMetadata)
+	static bool IsClassOf(const FRigBaseMetadata* InMetadata)
     {
     	return InMetadata->GetType() == ERigMetadataType::Quat;
     }
@@ -615,16 +615,16 @@ public:
 		: FRigBaseMetadata()
 	{}
 	virtual ~FRigQuatArrayMetadata() override {}
-	FORCEINLINE virtual void Serialize(FArchive& Ar, bool bIsLoading) override
+	virtual void Serialize(FArchive& Ar, bool bIsLoading) override
 	{
 		Super::Serialize(Ar, bIsLoading);
 		Ar << Value;
 	}
-	FORCEINLINE const TArray<FQuat>& GetValue() const { return Value; }
-	FORCEINLINE TArray<FQuat>& GetValue() { return Value; }
-	FORCEINLINE bool SetValue(const TArray<FQuat>& InValue) { return SetValueData(&InValue, sizeof(TArray<FQuat>)); }
+	const TArray<FQuat>& GetValue() const { return Value; }
+	TArray<FQuat>& GetValue() { return Value; }
+	bool SetValue(const TArray<FQuat>& InValue) { return SetValueData(&InValue, sizeof(TArray<FQuat>)); }
 
-	FORCEINLINE static bool IsClassOf(const FRigBaseMetadata* InMetadata)
+	static bool IsClassOf(const FRigBaseMetadata* InMetadata)
 	{
 		return InMetadata->GetType() == ERigMetadataType::QuatArray;
 	}
@@ -648,16 +648,16 @@ public:
 		, Value(FTransform::Identity)
 	{}
 	virtual ~FRigTransformMetadata() override {}
-	FORCEINLINE virtual void Serialize(FArchive& Ar, bool bIsLoading) override
+	virtual void Serialize(FArchive& Ar, bool bIsLoading) override
 	{
 		Super::Serialize(Ar, bIsLoading);
 		Ar << Value;
 	}
-	FORCEINLINE const FTransform& GetValue() const { return Value; } 
-	FORCEINLINE FTransform& GetValue() { return Value; }
-	FORCEINLINE bool SetValue(const FTransform& InValue) { return SetValueData(&InValue, sizeof(FTransform)); }
+	const FTransform& GetValue() const { return Value; } 
+	FTransform& GetValue() { return Value; }
+	bool SetValue(const FTransform& InValue) { return SetValueData(&InValue, sizeof(FTransform)); }
 
-	FORCEINLINE static bool IsClassOf(const FRigBaseMetadata* InMetadata)
+	static bool IsClassOf(const FRigBaseMetadata* InMetadata)
     {
     	return InMetadata->GetType() == ERigMetadataType::Transform;
     }
@@ -680,16 +680,16 @@ public:
 		: FRigBaseMetadata()
 	{}
 	virtual ~FRigTransformArrayMetadata() override {}
-	FORCEINLINE virtual void Serialize(FArchive& Ar, bool bIsLoading) override
+	virtual void Serialize(FArchive& Ar, bool bIsLoading) override
 	{
 		Super::Serialize(Ar, bIsLoading);
 		Ar << Value;
 	}
-	FORCEINLINE const TArray<FTransform>& GetValue() const { return Value; }
-	FORCEINLINE TArray<FTransform>& GetValue() { return Value; }
-	FORCEINLINE bool SetValue(const TArray<FTransform>& InValue) { return SetValueData(&InValue, sizeof(TArray<FTransform>)); }
+	const TArray<FTransform>& GetValue() const { return Value; }
+	TArray<FTransform>& GetValue() { return Value; }
+	bool SetValue(const TArray<FTransform>& InValue) { return SetValueData(&InValue, sizeof(TArray<FTransform>)); }
 
-	FORCEINLINE static bool IsClassOf(const FRigBaseMetadata* InMetadata)
+	static bool IsClassOf(const FRigBaseMetadata* InMetadata)
 	{
 		return InMetadata->GetType() == ERigMetadataType::TransformArray;
 	}
@@ -713,16 +713,16 @@ public:
 		, Value(FLinearColor::White)
 	{}
 	virtual ~FRigLinearColorMetadata() override {}
-	FORCEINLINE virtual void Serialize(FArchive& Ar, bool bIsLoading) override
+	virtual void Serialize(FArchive& Ar, bool bIsLoading) override
 	{
 		Super::Serialize(Ar, bIsLoading);
 		Ar << Value;
 	}
-	FORCEINLINE const FLinearColor& GetValue() const { return Value; } 
-	FORCEINLINE FLinearColor& GetValue() { return Value; }
-	FORCEINLINE bool SetValue(const FLinearColor& InValue) { return SetValueData(&InValue, sizeof(FLinearColor)); }
+	const FLinearColor& GetValue() const { return Value; } 
+	FLinearColor& GetValue() { return Value; }
+	bool SetValue(const FLinearColor& InValue) { return SetValueData(&InValue, sizeof(FLinearColor)); }
 
-	FORCEINLINE static bool IsClassOf(const FRigBaseMetadata* InMetadata)
+	static bool IsClassOf(const FRigBaseMetadata* InMetadata)
 	{
 		return InMetadata->GetType() == ERigMetadataType::LinearColor;
 	}
@@ -745,16 +745,16 @@ public:
 		: FRigBaseMetadata()
 	{}
 	virtual ~FRigLinearColorArrayMetadata() override {}
-	FORCEINLINE virtual void Serialize(FArchive& Ar, bool bIsLoading) override
+	virtual void Serialize(FArchive& Ar, bool bIsLoading) override
 	{
 		Super::Serialize(Ar, bIsLoading);
 		Ar << Value;
 	}
-	FORCEINLINE const TArray<FLinearColor>& GetValue() const { return Value; }
-	FORCEINLINE TArray<FLinearColor>& GetValue() { return Value; }
-	FORCEINLINE bool SetValue(const TArray<FLinearColor>& InValue) { return SetValueData(&InValue, sizeof(TArray<FLinearColor>)); }
+	const TArray<FLinearColor>& GetValue() const { return Value; }
+	TArray<FLinearColor>& GetValue() { return Value; }
+	bool SetValue(const TArray<FLinearColor>& InValue) { return SetValueData(&InValue, sizeof(TArray<FLinearColor>)); }
 
-	FORCEINLINE static bool IsClassOf(const FRigBaseMetadata* InMetadata)
+	static bool IsClassOf(const FRigBaseMetadata* InMetadata)
 	{
 		return InMetadata->GetType() == ERigMetadataType::LinearColorArray;
 	}
@@ -777,16 +777,16 @@ public:
 		: FRigBaseMetadata()
 	{}
 	virtual ~FRigElementKeyMetadata() override {}
-	FORCEINLINE virtual void Serialize(FArchive& Ar, bool bIsLoading) override
+	virtual void Serialize(FArchive& Ar, bool bIsLoading) override
 	{
 		Super::Serialize(Ar, bIsLoading);
 		Ar << Value;
 	}
-	FORCEINLINE const FRigElementKey& GetValue() const { return Value; } 
-	FORCEINLINE FRigElementKey& GetValue() { return Value; }
-	FORCEINLINE bool SetValue(const FRigElementKey& InValue) { return SetValueData(&InValue, sizeof(FRigElementKey)); }
+	const FRigElementKey& GetValue() const { return Value; } 
+	FRigElementKey& GetValue() { return Value; }
+	bool SetValue(const FRigElementKey& InValue) { return SetValueData(&InValue, sizeof(FRigElementKey)); }
 
-	FORCEINLINE static bool IsClassOf(const FRigBaseMetadata* InMetadata)
+	static bool IsClassOf(const FRigBaseMetadata* InMetadata)
     {
     	return InMetadata->GetType() == ERigMetadataType::RigElementKey;
     }
@@ -809,16 +809,16 @@ public:
 		: FRigBaseMetadata()
 	{}
 	virtual ~FRigElementKeyArrayMetadata() override {}
-	FORCEINLINE virtual void Serialize(FArchive& Ar, bool bIsLoading) override
+	virtual void Serialize(FArchive& Ar, bool bIsLoading) override
 	{
 		Super::Serialize(Ar, bIsLoading);
 		Ar << Value;
 	}
-	FORCEINLINE const TArray<FRigElementKey>& GetValue() const { return Value; }
-	FORCEINLINE TArray<FRigElementKey>& GetValue() { return Value; }
-	FORCEINLINE bool SetValue(const TArray<FRigElementKey>& InValue) { return SetValueData(&InValue, sizeof(TArray<FRigElementKey>)); }
+	const TArray<FRigElementKey>& GetValue() const { return Value; }
+	TArray<FRigElementKey>& GetValue() { return Value; }
+	bool SetValue(const TArray<FRigElementKey>& InValue) { return SetValueData(&InValue, sizeof(TArray<FRigElementKey>)); }
 
-	FORCEINLINE static bool IsClassOf(const FRigBaseMetadata* InMetadata)
+	static bool IsClassOf(const FRigBaseMetadata* InMetadata)
 	{
 		return InMetadata->GetType() == ERigMetadataType::RigElementKeyArray;
 	}
diff --git a/Engine/Plugins/Animation/ControlRig/Source/ControlRig/Public/Rigs/RigHierarchyPose.h b/Engine/Plugins/Animation/ControlRig/Source/ControlRig/Public/Rigs/RigHierarchyPose.h
--- a/Engine/Plugins/Animation/ControlRig/Source/ControlRig/Public/Rigs/RigHierarchyPose.h
+++ b/Engine/Plugins/Animation/ControlRig/Source/ControlRig/Public/Rigs/RigHierarchyPose.h
@@ -49,12 +49,12 @@ public:
 	{
 	}
 
-	FORCEINLINE void Reset() { Elements.Reset(); }
+	void Reset() { Elements.Reset(); }
 
-	FORCEINLINE int32 Num() const { return Elements.Num(); }
-	FORCEINLINE bool IsValidIndex(int32 InIndex) const { return Elements.IsValidIndex(InIndex); }
+	int32 Num() const { return Elements.Num(); }
+	bool IsValidIndex(int32 InIndex) const { return Elements.IsValidIndex(InIndex); }
 
-	FORCEINLINE int32 GetIndex(const FRigElementKey& InKey) const
+	int32 GetIndex(const FRigElementKey& InKey) const
 	{
 		if(CachedPoseHash != PoseHash)
 		{
@@ -80,17 +80,17 @@ public:
 		return INDEX_NONE;
 	}
 
-	FORCEINLINE bool Contains(const FRigElementKey& InKey) const
+	bool Contains(const FRigElementKey& InKey) const
 	{
 		return GetIndex(InKey) != INDEX_NONE;
 	}
 
-	FORCEINLINE const FRigPoseElement& operator[](int32 InIndex) const { return Elements[InIndex]; }
-	FORCEINLINE FRigPoseElement& operator[](int32 InIndex) { return Elements[InIndex]; }
-	FORCEINLINE TArray<FRigPoseElement>::RangedForIteratorType      begin()       { return Elements.begin(); }
-	FORCEINLINE TArray<FRigPoseElement>::RangedForConstIteratorType begin() const { return Elements.begin(); }
-	FORCEINLINE TArray<FRigPoseElement>::RangedForIteratorType      end()         { return Elements.end();   }
-	FORCEINLINE TArray<FRigPoseElement>::RangedForConstIteratorType end() const   { return Elements.end();   }
+	const FRigPoseElement& operator[](int32 InIndex) const { return Elements[InIndex]; }
+	FRigPoseElement& operator[](int32 InIndex) { return Elements[InIndex]; }
+	TArray<FRigPoseElement>::RangedForIteratorType      begin()       { return Elements.begin(); }
+	TArray<FRigPoseElement>::RangedForConstIteratorType begin() const { return Elements.begin(); }
+	TArray<FRigPoseElement>::RangedForIteratorType      end()         { return Elements.end();   }
+	TArray<FRigPoseElement>::RangedForConstIteratorType end() const   { return Elements.end();   }
 
 	UPROPERTY()
 	TArray<FRigPoseElement> Elements;
diff --git a/Engine/Plugins/Animation/ControlRig/Source/ControlRig/Public/Rigs/RigInfluenceMap.h b/Engine/Plugins/Animation/ControlRig/Source/ControlRig/Public/Rigs/RigInfluenceMap.h
--- a/Engine/Plugins/Animation/ControlRig/Source/ControlRig/Public/Rigs/RigInfluenceMap.h
+++ b/Engine/Plugins/Animation/ControlRig/Source/ControlRig/Public/Rigs/RigInfluenceMap.h
@@ -17,33 +17,33 @@ struct FRigInfluenceEntry
 	{
 	}
 
-	FORCEINLINE const FRigElementKey& GetSource() const { return Source; }
+	const FRigElementKey& GetSource() const { return Source; }
 
-	FORCEINLINE int32 Num() const { return AffectedList.Num(); }
-	FORCEINLINE const FRigElementKey& operator[](int32 InIndex) const { return AffectedList[InIndex]; }
-	FORCEINLINE FRigElementKey& operator[](int32 InIndex) { return AffectedList[InIndex]; }
+	int32 Num() const { return AffectedList.Num(); }
+	const FRigElementKey& operator[](int32 InIndex) const { return AffectedList[InIndex]; }
+	FRigElementKey& operator[](int32 InIndex) { return AffectedList[InIndex]; }
 
-	FORCEINLINE TArray<FRigElementKey>::RangedForIteratorType      begin() { return AffectedList.begin(); }
-	FORCEINLINE TArray<FRigElementKey>::RangedForConstIteratorType begin() const { return AffectedList.begin(); }
-	FORCEINLINE TArray<FRigElementKey>::RangedForIteratorType      end() { return AffectedList.end(); }
-	FORCEINLINE TArray<FRigElementKey>::RangedForConstIteratorType end() const { return AffectedList.end(); }
+	TArray<FRigElementKey>::RangedForIteratorType      begin() { return AffectedList.begin(); }
+	TArray<FRigElementKey>::RangedForConstIteratorType begin() const { return AffectedList.begin(); }
+	TArray<FRigElementKey>::RangedForIteratorType      end() { return AffectedList.end(); }
+	TArray<FRigElementKey>::RangedForConstIteratorType end() const { return AffectedList.end(); }
 
-	FORCEINLINE int32 AddUnique(const FRigElementKey& InKey)
+	int32 AddUnique(const FRigElementKey& InKey)
 	{
 		return AffectedList.AddUnique(InKey);
 	}
 
-	FORCEINLINE void Remove(const FRigElementKey& InKey)
+	void Remove(const FRigElementKey& InKey)
 	{
 		AffectedList.Remove(InKey);
 	}
 
-	FORCEINLINE bool Contains(const FRigElementKey& InKey) const
+	bool Contains(const FRigElementKey& InKey) const
 	{
 		return AffectedList.Contains(InKey);
 	}
 
-	FORCEINLINE bool Merge(const FRigInfluenceEntry& Other)
+	bool Merge(const FRigInfluenceEntry& Other)
 	{
 		if(Other.Source != Source)
 		{
@@ -86,14 +86,14 @@ struct FRigInfluenceEntryModifier
 	UPROPERTY(EditAnywhere, Category = "Inversion")
 	TArray<FRigElementKey> AffectedList;
 
-	FORCEINLINE int32 Num() const { return AffectedList.Num(); }
-	FORCEINLINE const FRigElementKey& operator[](int32 InIndex) const { return AffectedList[InIndex]; }
-	FORCEINLINE FRigElementKey& operator[](int32 InIndex) { return AffectedList[InIndex]; }
+	int32 Num() const { return AffectedList.Num(); }
+	const FRigElementKey& operator[](int32 InIndex) const { return AffectedList[InIndex]; }
+	FRigElementKey& operator[](int32 InIndex) { return AffectedList[InIndex]; }
 
-	FORCEINLINE TArray<FRigElementKey>::RangedForIteratorType      begin() { return AffectedList.begin(); }
-	FORCEINLINE TArray<FRigElementKey>::RangedForConstIteratorType begin() const { return AffectedList.begin(); }
-	FORCEINLINE TArray<FRigElementKey>::RangedForIteratorType      end() { return AffectedList.end(); }
-	FORCEINLINE TArray<FRigElementKey>::RangedForConstIteratorType end() const { return AffectedList.end(); }
+	TArray<FRigElementKey>::RangedForIteratorType      begin() { return AffectedList.begin(); }
+	TArray<FRigElementKey>::RangedForConstIteratorType begin() const { return AffectedList.begin(); }
+	TArray<FRigElementKey>::RangedForIteratorType      end() { return AffectedList.end(); }
+	TArray<FRigElementKey>::RangedForConstIteratorType end() const { return AffectedList.end(); }
 };
 
 
@@ -110,16 +110,16 @@ public:
 
 	const FName& GetEventName() const { return EventName; }
 
-	FORCEINLINE int32 Num() const { return Entries.Num(); }
-	FORCEINLINE const FRigInfluenceEntry& operator[](int32 InIndex) const { return Entries[InIndex]; }
-	FORCEINLINE FRigInfluenceEntry& operator[](int32 InIndex) { return Entries[InIndex]; }
-	FORCEINLINE const FRigInfluenceEntry& operator[](const FRigElementKey& InKey) const { return Entries[GetIndex(InKey)]; }
-	FORCEINLINE FRigInfluenceEntry& operator[](const FRigElementKey& InKey) { return FindOrAdd(InKey); }
+	int32 Num() const { return Entries.Num(); }
+	const FRigInfluenceEntry& operator[](int32 InIndex) const { return Entries[InIndex]; }
+	FRigInfluenceEntry& operator[](int32 InIndex) { return Entries[InIndex]; }
+	const FRigInfluenceEntry& operator[](const FRigElementKey& InKey) const { return Entries[GetIndex(InKey)]; }
+	FRigInfluenceEntry& operator[](const FRigElementKey& InKey) { return FindOrAdd(InKey); }
 
-	FORCEINLINE TArray<FRigInfluenceEntry>::RangedForIteratorType      begin()       { return Entries.begin(); }
-	FORCEINLINE TArray<FRigInfluenceEntry>::RangedForConstIteratorType begin() const { return Entries.begin(); }
-	FORCEINLINE TArray<FRigInfluenceEntry>::RangedForIteratorType      end()         { return Entries.end();   }
-	FORCEINLINE TArray<FRigInfluenceEntry>::RangedForConstIteratorType end() const   { return Entries.end();   }
+	TArray<FRigInfluenceEntry>::RangedForIteratorType      begin()       { return Entries.begin(); }
+	TArray<FRigInfluenceEntry>::RangedForConstIteratorType begin() const { return Entries.begin(); }
+	TArray<FRigInfluenceEntry>::RangedForIteratorType      end()         { return Entries.end();   }
+	TArray<FRigInfluenceEntry>::RangedForConstIteratorType end() const   { return Entries.end();   }
 
 	FRigInfluenceEntry& FindOrAdd(const FRigElementKey& InKey);
 
@@ -127,12 +127,12 @@ public:
 
 	void Remove(const FRigElementKey& InKey);
 
-	FORCEINLINE bool Contains(const FRigElementKey& InKey) const
+	bool Contains(const FRigElementKey& InKey) const
 	{
 		return GetIndex(InKey) != INDEX_NONE;
 	}
 
-	FORCEINLINE int32 GetIndex(const FRigElementKey& InKey) const
+	int32 GetIndex(const FRigElementKey& InKey) const
 	{
 		const int32* Index = KeyToIndex.Find(InKey);
 		if (Index)
@@ -178,16 +178,16 @@ struct CONTROLRIG_API FRigInfluenceMapPerEvent
 
 public:
 
-	FORCEINLINE int32 Num() const { return Maps.Num(); }
-	FORCEINLINE const FRigInfluenceMap& operator[](int32 InIndex) const { return Maps[InIndex]; }
-	FORCEINLINE FRigInfluenceMap& operator[](int32 InIndex) { return Maps[InIndex]; }
-	FORCEINLINE const FRigInfluenceMap& operator[](const FName& InEventName) const { return Maps[GetIndex(InEventName)]; }
-	FORCEINLINE FRigInfluenceMap& operator[](const FName& InEventName) { return FindOrAdd(InEventName); }
+	int32 Num() const { return Maps.Num(); }
+	const FRigInfluenceMap& operator[](int32 InIndex) const { return Maps[InIndex]; }
+	FRigInfluenceMap& operator[](int32 InIndex) { return Maps[InIndex]; }
+	const FRigInfluenceMap& operator[](const FName& InEventName) const { return Maps[GetIndex(InEventName)]; }
+	FRigInfluenceMap& operator[](const FName& InEventName) { return FindOrAdd(InEventName); }
 
-	FORCEINLINE TArray<FRigInfluenceMap>::RangedForIteratorType      begin()       { return Maps.begin(); }
-	FORCEINLINE TArray<FRigInfluenceMap>::RangedForConstIteratorType begin() const { return Maps.begin(); }
-	FORCEINLINE TArray<FRigInfluenceMap>::RangedForIteratorType      end()         { return Maps.end();   }
-	FORCEINLINE TArray<FRigInfluenceMap>::RangedForConstIteratorType end() const   { return Maps.end();   }
+	TArray<FRigInfluenceMap>::RangedForIteratorType      begin()       { return Maps.begin(); }
+	TArray<FRigInfluenceMap>::RangedForConstIteratorType begin() const { return Maps.begin(); }
+	TArray<FRigInfluenceMap>::RangedForIteratorType      end()         { return Maps.end();   }
+	TArray<FRigInfluenceMap>::RangedForConstIteratorType end() const   { return Maps.end();   }
 
 	FRigInfluenceMap& FindOrAdd(const FName& InEventName);
 
@@ -195,12 +195,12 @@ public:
 
 	void Remove(const FName& InEventName);
 
-	FORCEINLINE bool Contains(const FName& InEventName) const
+	bool Contains(const FName& InEventName) const
 	{
 		return GetIndex(InEventName) != INDEX_NONE;
 	}
 
-	FORCEINLINE int32 GetIndex(const FName& InEventName) const
+	int32 GetIndex(const FName& InEventName) const
 	{
 		const int32* Index = EventToIndex.Find(InEventName);
 		if (Index)
diff --git a/Engine/Plugins/Animation/ControlRig/Source/ControlRig/Public/Rigs/RigSpaceHierarchy.h b/Engine/Plugins/Animation/ControlRig/Source/ControlRig/Public/Rigs/RigSpaceHierarchy.h
--- a/Engine/Plugins/Animation/ControlRig/Source/ControlRig/Public/Rigs/RigSpaceHierarchy.h
+++ b/Engine/Plugins/Animation/ControlRig/Source/ControlRig/Public/Rigs/RigSpaceHierarchy.h
@@ -57,12 +57,12 @@ struct CONTROLRIG_API FRigSpace : public FRigElement
 	UPROPERTY(BlueprintReadOnly, transient, EditAnywhere, Category = FRigElement)
 	FTransform LocalTransform;
 
-	FORCEINLINE virtual ERigElementType GetElementType() const override
+	virtual ERigElementType GetElementType() const override
 	{
 		return ERigElementType::Null;
 	}
 
-	FORCEINLINE virtual FRigElementKey GetParentElementKey() const
+	virtual FRigElementKey GetParentElementKey() const
 	{
 		switch (SpaceType)
 		{
@@ -94,10 +94,10 @@ struct CONTROLRIG_API FRigSpaceHierarchy
 
 	FRigSpaceHierarchy();
 
-	FORCEINLINE TArray<FRigSpace>::RangedForIteratorType      begin()       { return Spaces.begin(); }
-	FORCEINLINE TArray<FRigSpace>::RangedForConstIteratorType begin() const { return Spaces.begin(); }
-	FORCEINLINE TArray<FRigSpace>::RangedForIteratorType      end()         { return Spaces.end();   }
-	FORCEINLINE TArray<FRigSpace>::RangedForConstIteratorType end() const   { return Spaces.end();   }
+	TArray<FRigSpace>::RangedForIteratorType      begin()       { return Spaces.begin(); }
+	TArray<FRigSpace>::RangedForConstIteratorType begin() const { return Spaces.begin(); }
+	TArray<FRigSpace>::RangedForIteratorType      end()         { return Spaces.end();   }
+	TArray<FRigSpace>::RangedForConstIteratorType end() const   { return Spaces.end();   }
 
 	FRigSpace& Add(const FName& InNewName, ERigSpaceType InSpaceType, const FName& InParentName, const FTransform& InTransform);
 	
diff --git a/Engine/Plugins/Animation/ControlRig/Source/ControlRig/Public/Sequencer/MovieSceneControlRigSpaceChannel.h b/Engine/Plugins/Animation/ControlRig/Source/ControlRig/Public/Sequencer/MovieSceneControlRigSpaceChannel.h
--- a/Engine/Plugins/Animation/ControlRig/Source/ControlRig/Public/Sequencer/MovieSceneControlRigSpaceChannel.h
+++ b/Engine/Plugins/Animation/ControlRig/Source/ControlRig/Public/Sequencer/MovieSceneControlRigSpaceChannel.h
@@ -73,7 +73,7 @@ struct CONTROLRIG_API FMovieSceneControlRigSpaceChannel : public FMovieSceneChan
 	*
 	* @return An object that is able to manipulate this channel's data
 	*/
-	FORCEINLINE TMovieSceneChannelData<FMovieSceneControlRigSpaceBaseKey> GetData()
+	TMovieSceneChannelData<FMovieSceneControlRigSpaceBaseKey> GetData()
 	{
 		return TMovieSceneChannelData<FMovieSceneControlRigSpaceBaseKey>(&KeyTimes, &KeyValues, &KeyHandles, this);
 	}
@@ -83,7 +83,7 @@ struct CONTROLRIG_API FMovieSceneControlRigSpaceChannel : public FMovieSceneChan
 	*
 	* @return An object that is able to interrogate this channel's data
 	*/
-	FORCEINLINE TMovieSceneChannelData<const FMovieSceneControlRigSpaceBaseKey> GetData() const
+	TMovieSceneChannelData<const FMovieSceneControlRigSpaceBaseKey> GetData() const
 	{
 		return TMovieSceneChannelData<const FMovieSceneControlRigSpaceBaseKey>(&KeyTimes, &KeyValues);
 	}
diff --git a/Engine/Plugins/Animation/ControlRig/Source/ControlRig/Public/Units/RigDispatchFactory.h b/Engine/Plugins/Animation/ControlRig/Source/ControlRig/Public/Units/RigDispatchFactory.h
--- a/Engine/Plugins/Animation/ControlRig/Source/ControlRig/Public/Units/RigDispatchFactory.h
+++ b/Engine/Plugins/Animation/ControlRig/Source/ControlRig/Public/Units/RigDispatchFactory.h
@@ -14,19 +14,19 @@ struct CONTROLRIG_API FRigDispatchFactory : public FRigVMDispatchFactory
 {
 	GENERATED_BODY()
 
-	FORCEINLINE virtual UScriptStruct* GetExecuteContextStruct() const override
+	virtual UScriptStruct* GetExecuteContextStruct() const override
 	{
 		return FControlRigExecuteContext::StaticStruct();
 	}
 
-	FORCEINLINE virtual void RegisterDependencyTypes() const override
+	virtual void RegisterDependencyTypes() const override
 	{
 		FRigVMRegistry::Get().FindOrAddType(FControlRigExecuteContext::StaticStruct());
 		FRigVMRegistry::Get().FindOrAddType(FRigElementKey::StaticStruct());
     	FRigVMRegistry::Get().FindOrAddType(FCachedRigElement::StaticStruct());
 	}
 
-	FORCEINLINE virtual TArray<TPair<FName,FString>> GetOpaqueArguments() const override
+	virtual TArray<TPair<FName,FString>> GetOpaqueArguments() const override
 	{
 		static const TArray<TPair<FName,FString>> OpaqueArguments = {
 			TPair<FName,FString>(TEXT("Context"), TEXT("const FRigUnitContext&"))
@@ -40,7 +40,7 @@ struct CONTROLRIG_API FRigDispatchFactory : public FRigVMDispatchFactory
 
 #endif
 
-	FORCEINLINE static const FRigUnitContext& GetRigUnitContext(const FRigVMExtendedExecuteContext& InContext)
+	static const FRigUnitContext& GetRigUnitContext(const FRigVMExtendedExecuteContext& InContext)
 	{
 		return *(const FRigUnitContext*)InContext.OpaqueArguments[0];
 	}
diff --git a/Engine/Plugins/Animation/ControlRig/Source/ControlRig/Public/Units/RigUnit.h b/Engine/Plugins/Animation/ControlRig/Source/ControlRig/Public/Units/RigUnit.h
--- a/Engine/Plugins/Animation/ControlRig/Source/ControlRig/Public/Units/RigUnit.h
+++ b/Engine/Plugins/Animation/ControlRig/Source/ControlRig/Public/Units/RigUnit.h
@@ -32,7 +32,7 @@ struct CONTROLRIG_API FRigUnit : public FRigVMStruct
 	virtual FTransform DetermineOffsetTransformForPin(const FString& InPinPath, void* InUserContext) const { return FTransform::Identity; }
 	
 	/** The name of the method used within each rig unit */
-	FORCEINLINE_DEBUGGABLE static FName GetMethodName()
+	static FName GetMethodName()
 	{
 		static FName MethodName = FName(NAME_None);
 		if(MethodName.IsNone())
diff --git a/Engine/Plugins/Animation/ControlRig/Source/ControlRig/Public/Units/RigUnitContext.h b/Engine/Plugins/Animation/ControlRig/Source/ControlRig/Public/Units/RigUnitContext.h
--- a/Engine/Plugins/Animation/ControlRig/Source/ControlRig/Public/Units/RigUnitContext.h
+++ b/Engine/Plugins/Animation/ControlRig/Source/ControlRig/Public/Units/RigUnitContext.h
@@ -137,7 +137,7 @@ struct FRigUnitContext
 	 * @return The requested data source
 	 */
 	template<class T>
-	FORCEINLINE T* RequestDataSource(const FName& InName) const
+	T* RequestDataSource(const FName& InName) const
 	{
 		if (DataSourceRegistry == nullptr)
 		{
@@ -149,7 +149,7 @@ struct FRigUnitContext
 	/**
 	 * Converts a transform from rig (global) space to world space
 	 */
-	FORCEINLINE FTransform ToWorldSpace(const FTransform& InTransform) const
+	FTransform ToWorldSpace(const FTransform& InTransform) const
 	{
 		return InTransform * ToWorldSpaceTransform;
 	}
@@ -157,7 +157,7 @@ struct FRigUnitContext
 	/**
 	 * Converts a transform from world space to rig (global) space
 	 */
-	FORCEINLINE FTransform ToRigSpace(const FTransform& InTransform) const
+	FTransform ToRigSpace(const FTransform& InTransform) const
 	{
 		return InTransform.GetRelativeTransform(ToWorldSpaceTransform);
 	}
@@ -165,7 +165,7 @@ struct FRigUnitContext
 	/**
 	 * Converts a location from rig (global) space to world space
 	 */
-	FORCEINLINE FVector ToWorldSpace(const FVector& InLocation) const
+	FVector ToWorldSpace(const FVector& InLocation) const
 	{
 		return ToWorldSpaceTransform.TransformPosition(InLocation);
 	}
@@ -173,7 +173,7 @@ struct FRigUnitContext
 	/**
 	 * Converts a location from world space to rig (global) space
 	 */
-	FORCEINLINE FVector ToRigSpace(const FVector& InLocation) const
+	FVector ToRigSpace(const FVector& InLocation) const
 	{
 		return ToWorldSpaceTransform.InverseTransformPosition(InLocation);
 	}
@@ -181,7 +181,7 @@ struct FRigUnitContext
 	/**
 	 * Converts a rotation from rig (global) space to world space
 	 */
-	FORCEINLINE FQuat ToWorldSpace(const FQuat& InRotation) const
+	FQuat ToWorldSpace(const FQuat& InRotation) const
 	{
 		return ToWorldSpaceTransform.TransformRotation(InRotation);
 	}
@@ -189,7 +189,7 @@ struct FRigUnitContext
 	/**
 	 * Converts a rotation from world space to rig (global) space
 	 */
-	FORCEINLINE FQuat ToRigSpace(const FQuat& InRotation) const
+	FQuat ToRigSpace(const FQuat& InRotation) const
 	{
 		return ToWorldSpaceTransform.InverseTransformRotation(InRotation);
 	}
@@ -197,7 +197,7 @@ struct FRigUnitContext
 	/**
 	 * Returns true if this context is currently being interacted on
 	 */
-	FORCEINLINE bool IsInteracting() const
+	bool IsInteracting() const
 	{
 		return InteractionType != (uint8)EControlRigInteractionType::None;
 	}
diff --git a/Engine/Plugins/Animation/ControlRig/Source/ControlRigDeveloper/Public/ControlRigBlueprint.h b/Engine/Plugins/Animation/ControlRig/Source/ControlRigDeveloper/Public/ControlRigBlueprint.h
--- a/Engine/Plugins/Animation/ControlRig/Source/ControlRigDeveloper/Public/ControlRigBlueprint.h
+++ b/Engine/Plugins/Animation/ControlRig/Source/ControlRigDeveloper/Public/ControlRigBlueprint.h
@@ -748,7 +748,7 @@ private:
 public:
 
 	/** Sets the execution mode. In Release mode the rig will ignore all breakpoints. */
-	FORCEINLINE void SetDebugMode(const bool bValue) { bCompileInDebugMode = bValue; }
+	void SetDebugMode(const bool bValue) { bCompileInDebugMode = bValue; }
 
 	/** Removes all the breakpoints from the blueprint and the VM */
 	void ClearBreakpoints();
diff --git a/Engine/Plugins/Animation/ControlRig/Source/ControlRigDeveloper/Public/Graph/ControlRigGraph.h b/Engine/Plugins/Animation/ControlRig/Source/ControlRigDeveloper/Public/Graph/ControlRigGraph.h
--- a/Engine/Plugins/Animation/ControlRig/Source/ControlRigDeveloper/Public/Graph/ControlRigGraph.h
+++ b/Engine/Plugins/Animation/ControlRig/Source/ControlRigDeveloper/Public/Graph/ControlRigGraph.h
@@ -45,23 +45,23 @@ public:
 #endif
 #if WITH_EDITOR
 
-	FORCEINLINE const TArray<TSharedPtr<FString>>* GetBoneNameList(URigVMPin* InPin = nullptr) const
+	const TArray<TSharedPtr<FString>>* GetBoneNameList(URigVMPin* InPin = nullptr) const
 	{
 		return GetElementNameList(ERigElementType::Bone);
 	}
-	FORCEINLINE const TArray<TSharedPtr<FString>>* GetControlNameList(URigVMPin* InPin = nullptr) const
+	const TArray<TSharedPtr<FString>>* GetControlNameList(URigVMPin* InPin = nullptr) const
 	{
 		return GetElementNameList(ERigElementType::Control);
 	}
-	FORCEINLINE const TArray<TSharedPtr<FString>>* GetControlNameListWithoutAnimationChannels(URigVMPin* InPin = nullptr) const
+	const TArray<TSharedPtr<FString>>* GetControlNameListWithoutAnimationChannels(URigVMPin* InPin = nullptr) const
 	{
 		return &ControlNameListWithoutAnimationChannels;
 	}
-	FORCEINLINE const TArray<TSharedPtr<FString>>* GetNullNameList(URigVMPin* InPin = nullptr) const
+	const TArray<TSharedPtr<FString>>* GetNullNameList(URigVMPin* InPin = nullptr) const
 	{
 		return GetElementNameList(ERigElementType::Null);
 	}
-	FORCEINLINE const TArray<TSharedPtr<FString>>* GetCurveNameList(URigVMPin* InPin = nullptr) const
+	const TArray<TSharedPtr<FString>>* GetCurveNameList(URigVMPin* InPin = nullptr) const
 	{
 		return GetElementNameList(ERigElementType::Curve);
 	}
diff --git a/Engine/Plugins/Animation/ControlRig/Source/ControlRigDeveloper/Public/Graph/ControlRigGraphNode.h b/Engine/Plugins/Animation/ControlRig/Source/ControlRigDeveloper/Public/Graph/ControlRigGraphNode.h
--- a/Engine/Plugins/Animation/ControlRig/Source/ControlRigDeveloper/Public/Graph/ControlRigGraphNode.h
+++ b/Engine/Plugins/Animation/ControlRig/Source/ControlRigDeveloper/Public/Graph/ControlRigGraphNode.h
@@ -243,7 +243,7 @@ protected:
 
 private:
 
-	FORCEINLINE int32 GetNodeTopologyVersion() const { return NodeTopologyVersion; }
+	int32 GetNodeTopologyVersion() const { return NodeTopologyVersion; }
 	int32 NodeTopologyVersion;
 
 	static void ConfigurePin(UEdGraphPin* EdGraphPin, const URigVMPin* ModelPin);
diff --git a/Engine/Plugins/Animation/ControlRig/Source/ControlRigEditor/Private/ControlRigGraphDetails.h b/Engine/Plugins/Animation/ControlRig/Source/ControlRigEditor/Private/ControlRigGraphDetails.h
--- a/Engine/Plugins/Animation/ControlRig/Source/ControlRigEditor/Private/ControlRigGraphDetails.h
+++ b/Engine/Plugins/Animation/ControlRig/Source/ControlRigEditor/Private/ControlRigGraphDetails.h
@@ -266,7 +266,7 @@ public:
 
 protected:
 
-	FORCEINLINE_DEBUGGABLE bool GetPropertyChain(TSharedRef<class IPropertyHandle> InPropertyHandle, FEditPropertyChain& OutPropertyChain, TArray<int32> &OutPropertyArrayIndices, bool& bOutEnabled)
+	bool GetPropertyChain(TSharedRef<class IPropertyHandle> InPropertyHandle, FEditPropertyChain& OutPropertyChain, TArray<int32> &OutPropertyArrayIndices, bool& bOutEnabled)
 	{
 		OutPropertyChain.Empty();
 		OutPropertyArrayIndices.Reset();
@@ -296,7 +296,7 @@ protected:
 
 	// extracts the value for a nested property (for Example Settings.WorldTransform) from an outer owner
 	template<typename ValueType>
-	FORCEINLINE_DEBUGGABLE ValueType& ContainerUObjectToValueRef(UObject* InOwner, ValueType& InDefault, FEditPropertyChain& InPropertyChain, TArray<int32> &InPropertyArrayIndices) const
+	ValueType& ContainerUObjectToValueRef(UObject* InOwner, ValueType& InDefault, FEditPropertyChain& InPropertyChain, TArray<int32> &InPropertyArrayIndices) const
 	{
 		if (InPropertyChain.GetHead() == nullptr)
 		{
@@ -441,10 +441,10 @@ protected:
 
 	// specializations for FVector and FVector4 at the end of this file
 	template<typename VectorType>
-	FORCEINLINE void ExtendVectorArgs(TSharedRef<class IPropertyHandle> InPropertyHandle, void* ArgumentsPtr) {}
+	void ExtendVectorArgs(TSharedRef<class IPropertyHandle> InPropertyHandle, void* ArgumentsPtr) {}
 
 	template<typename VectorType, int32 NumberOfComponents>
-	FORCEINLINE void CustomizeVector(TSharedRef<class IPropertyHandle> InPropertyHandle, class IDetailChildrenBuilder& StructBuilder, IPropertyTypeCustomizationUtils& StructCustomizationUtils)
+	void CustomizeVector(TSharedRef<class IPropertyHandle> InPropertyHandle, class IDetailChildrenBuilder& StructBuilder, IPropertyTypeCustomizationUtils& StructCustomizationUtils)
 	{
 		typedef typename VectorType::FReal NumericType;
 		typedef SNumericVectorInputBox<NumericType, VectorType, NumberOfComponents> SLocalVectorInputBox;
@@ -600,11 +600,11 @@ protected:
 
 	// specializations for FRotator and FQuat at the end of this file
 	template<typename RotationType>
-	FORCEINLINE void ExtendRotationArgs(TSharedRef<class IPropertyHandle> InPropertyHandle, void* ArgumentsPtr) {}
+	void ExtendRotationArgs(TSharedRef<class IPropertyHandle> InPropertyHandle, void* ArgumentsPtr) {}
 
 	// add the widget for a rotation (rotator or quat)
 	template<typename RotationType>
-	FORCEINLINE void CustomizeRotation(TSharedRef<class IPropertyHandle> InPropertyHandle, class IDetailChildrenBuilder& StructBuilder, IPropertyTypeCustomizationUtils& StructCustomizationUtils)
+	void CustomizeRotation(TSharedRef<class IPropertyHandle> InPropertyHandle, class IDetailChildrenBuilder& StructBuilder, IPropertyTypeCustomizationUtils& StructCustomizationUtils)
 	{
 		FEditPropertyChain PropertyChain;
 		TArray<int32> PropertyArrayIndices;
@@ -642,7 +642,7 @@ protected:
 
 	// add the widget for a transform / euler transform
 	template<typename TransformType>
-	FORCEINLINE void CustomizeTransform(TSharedRef<class IPropertyHandle> InPropertyHandle, class IDetailChildrenBuilder& StructBuilder, IPropertyTypeCustomizationUtils& StructCustomizationUtils)
+	void CustomizeTransform(TSharedRef<class IPropertyHandle> InPropertyHandle, class IDetailChildrenBuilder& StructBuilder, IPropertyTypeCustomizationUtils& StructCustomizationUtils)
 	{
 		FEditPropertyChain PropertyChain;
 		TArray<int32> PropertyArrayIndices;
@@ -940,13 +940,13 @@ protected:
 };
 
 template<>
-FORCEINLINE bool FControlRigGraphMathTypeDetails::IsQuaternionBasedRotation<FEulerTransform>() { return false; }
+inline bool FControlRigGraphMathTypeDetails::IsQuaternionBasedRotation<FEulerTransform>() { return false; }
 
 template<>
-FORCEINLINE bool FControlRigGraphMathTypeDetails::IsQuaternionBasedRotation<FRotator>() { return false; }
+inline bool FControlRigGraphMathTypeDetails::IsQuaternionBasedRotation<FRotator>() { return false; }
 
 template<>
-FORCEINLINE void FControlRigGraphMathTypeDetails::ExtendVectorArgs<FVector>(TSharedRef<class IPropertyHandle> InPropertyHandle, void* ArgumentsPtr)
+inline void FControlRigGraphMathTypeDetails::ExtendVectorArgs<FVector>(TSharedRef<class IPropertyHandle> InPropertyHandle, void* ArgumentsPtr)
 {
 	using VectorType = FVector;
 	typedef typename VectorType::FReal NumericType;
@@ -969,7 +969,7 @@ FORCEINLINE void FControlRigGraphMathTypeDetails::ExtendVectorArgs<FVector>(TSha
 }
 
 template<>
-FORCEINLINE void FControlRigGraphMathTypeDetails::ExtendVectorArgs<FVector4>(TSharedRef<class IPropertyHandle> InPropertyHandle, void* ArgumentsPtr)
+inline void FControlRigGraphMathTypeDetails::ExtendVectorArgs<FVector4>(TSharedRef<class IPropertyHandle> InPropertyHandle, void* ArgumentsPtr)
 {
 	using VectorType = FVector4;
 	typedef typename VectorType::FReal NumericType;
@@ -1004,7 +1004,7 @@ FORCEINLINE void FControlRigGraphMathTypeDetails::ExtendVectorArgs<FVector4>(TSh
 }
 
 template<>
-FORCEINLINE void FControlRigGraphMathTypeDetails::ExtendRotationArgs<FQuat>(TSharedRef<class IPropertyHandle> InPropertyHandle, void* ArgumentsPtr)
+inline void FControlRigGraphMathTypeDetails::ExtendRotationArgs<FQuat>(TSharedRef<class IPropertyHandle> InPropertyHandle, void* ArgumentsPtr)
 {
 	using RotationType = FQuat;
 	typedef typename RotationType::FReal NumericType;
@@ -1028,7 +1028,7 @@ FORCEINLINE void FControlRigGraphMathTypeDetails::ExtendRotationArgs<FQuat>(TSha
 }
 
 template<>
-FORCEINLINE void FControlRigGraphMathTypeDetails::ExtendRotationArgs<FRotator>(TSharedRef<class IPropertyHandle> InPropertyHandle, void* ArgumentsPtr)
+inline void FControlRigGraphMathTypeDetails::ExtendRotationArgs<FRotator>(TSharedRef<class IPropertyHandle> InPropertyHandle, void* ArgumentsPtr)
 {
 	using RotationType = FRotator;
 	typedef typename RotationType::FReal NumericType;
diff --git a/Engine/Plugins/Animation/ControlRig/Source/ControlRigEditor/Private/EditMode/SControlRigOutliner.h b/Engine/Plugins/Animation/ControlRig/Source/ControlRigEditor/Private/EditMode/SControlRigOutliner.h
--- a/Engine/Plugins/Animation/ControlRig/Source/ControlRigEditor/Private/EditMode/SControlRigOutliner.h
+++ b/Engine/Plugins/Animation/ControlRig/Source/ControlRigEditor/Private/EditMode/SControlRigOutliner.h
@@ -49,7 +49,7 @@ struct CONTROLRIGEDITOR_API FMultiRigTreeDelegates
 	}
 
 
-	FORCEINLINE const FRigTreeDisplaySettings& GetDisplaySettings() const
+	const FRigTreeDisplaySettings& GetDisplaySettings() const
 	{
 		if (OnGetDisplaySettings.IsBound())
 		{
@@ -58,7 +58,7 @@ struct CONTROLRIGEDITOR_API FMultiRigTreeDelegates
 		return DefaultDisplaySettings;
 	}
 
-	FORCEINLINE void HandleSelectionChanged(TSharedPtr<FMultiRigTreeElement> Selection, ESelectInfo::Type SelectInfo)
+	void HandleSelectionChanged(TSharedPtr<FMultiRigTreeElement> Selection, ESelectInfo::Type SelectInfo)
 	{
 		if (bIsChangingRigHierarchy)
 		{
diff --git a/Engine/Plugins/Animation/ControlRig/Source/ControlRigEditor/Private/Editor/DetailsViewWrapperObject.h b/Engine/Plugins/Animation/ControlRig/Source/ControlRigEditor/Private/Editor/DetailsViewWrapperObject.h
--- a/Engine/Plugins/Animation/ControlRig/Source/ControlRigEditor/Private/Editor/DetailsViewWrapperObject.h
+++ b/Engine/Plugins/Animation/ControlRig/Source/ControlRigEditor/Private/Editor/DetailsViewWrapperObject.h
@@ -96,17 +96,17 @@ private:
 		, ScriptStruct(nullptr)
 		{}
 
-		friend FORCEINLINE uint32 GetTypeHash(const FPerClassInfo& Info)
+		friend uint32 GetTypeHash(const FPerClassInfo& Info)
 		{
 			return HashCombine(GetTypeHash(Info.Notation), GetTypeHash(Info.ScriptStruct));
 		}
 
-		FORCEINLINE bool operator ==(const FPerClassInfo& Other) const
+		bool operator ==(const FPerClassInfo& Other) const
 		{
 			return Notation == Other.Notation && ScriptStruct == Other.ScriptStruct;
 		}
 
-		FORCEINLINE bool operator !=(const FPerClassInfo& Other) const
+		bool operator !=(const FPerClassInfo& Other) const
 		{
 			return Notation != Other.Notation || ScriptStruct != Other.ScriptStruct;
 		}
diff --git a/Engine/Plugins/Animation/ControlRig/Source/ControlRigEditor/Private/Editor/SControlRigGizmoNameList.h b/Engine/Plugins/Animation/ControlRig/Source/ControlRigEditor/Private/Editor/SControlRigGizmoNameList.h
--- a/Engine/Plugins/Animation/ControlRig/Source/ControlRigEditor/Private/Editor/SControlRigGizmoNameList.h
+++ b/Engine/Plugins/Animation/ControlRig/Source/ControlRigEditor/Private/Editor/SControlRigGizmoNameList.h
@@ -39,7 +39,7 @@ protected:
 	void OnNameListChanged(TSharedPtr<FString> NewSelection, ESelectInfo::Type SelectInfo);
 	void OnNameListComboBox();
 
-	FORCEINLINE static const TArray< TSharedPtr<FString> >& GetEmptyList()
+	static const TArray< TSharedPtr<FString> >& GetEmptyList()
 	{
 		static const TArray< TSharedPtr<FString> > EmptyList;
 		return EmptyList;
diff --git a/Engine/Plugins/Animation/ControlRig/Source/ControlRigEditor/Private/Editor/SRigHierarchy.h b/Engine/Plugins/Animation/ControlRig/Source/ControlRigEditor/Private/Editor/SRigHierarchy.h
--- a/Engine/Plugins/Animation/ControlRig/Source/ControlRigEditor/Private/Editor/SRigHierarchy.h
+++ b/Engine/Plugins/Animation/ControlRig/Source/ControlRigEditor/Private/Editor/SRigHierarchy.h
@@ -74,7 +74,7 @@ public:
 
 	void Construct(const FArguments& InArgs, TSharedRef<FControlRigEditor> InControlRigEditor);
 
-	FORCEINLINE FControlRigEditor* GetControlRigEditor() const
+	FControlRigEditor* GetControlRigEditor() const
 	{
 		if(ControlRigEditor.IsValid())
 		{
diff --git a/Engine/Plugins/Animation/ControlRig/Source/ControlRigEditor/Private/Editor/SRigHierarchyTreeView.h b/Engine/Plugins/Animation/ControlRig/Source/ControlRigEditor/Private/Editor/SRigHierarchyTreeView.h
--- a/Engine/Plugins/Animation/ControlRig/Source/ControlRigEditor/Private/Editor/SRigHierarchyTreeView.h
+++ b/Engine/Plugins/Animation/ControlRig/Source/ControlRigEditor/Private/Editor/SRigHierarchyTreeView.h
@@ -92,7 +92,7 @@ struct CONTROLRIGEDITOR_API FRigTreeDelegates
 		bIsChangingRigHierarchy = false;
 	}
 
-	FORCEINLINE const URigHierarchy* GetHierarchy() const
+	const URigHierarchy* GetHierarchy() const
 	{
 		if(OnGetHierarchy.IsBound())
 		{
@@ -101,7 +101,7 @@ struct CONTROLRIGEDITOR_API FRigTreeDelegates
 		return nullptr;
 	}
 
-	FORCEINLINE const FRigTreeDisplaySettings& GetDisplaySettings() const
+	const FRigTreeDisplaySettings& GetDisplaySettings() const
 	{
 		if(OnGetDisplaySettings.IsBound())
 		{
@@ -110,7 +110,7 @@ struct CONTROLRIGEDITOR_API FRigTreeDelegates
 		return DefaultDisplaySettings;
 	}
 	
-	FORCEINLINE FName HandleRenameElement(const FRigElementKey& OldKey, const FString& NewName) const
+	FName HandleRenameElement(const FRigElementKey& OldKey, const FString& NewName) const
 	{
 		if(OnRenameElement.IsBound())
 		{
@@ -119,7 +119,7 @@ struct CONTROLRIGEDITOR_API FRigTreeDelegates
 		return OldKey.Name;
 	}
 	
-	FORCEINLINE bool HandleVerifyElementNameChanged(const FRigElementKey& OldKey, const FString& NewName, FText& OutErrorMessage) const
+	bool HandleVerifyElementNameChanged(const FRigElementKey& OldKey, const FString& NewName, FText& OutErrorMessage) const
 	{
 		if(OnVerifyElementNameChanged.IsBound())
 		{
@@ -128,7 +128,7 @@ struct CONTROLRIGEDITOR_API FRigTreeDelegates
 		return false;
 	}
 
-	FORCEINLINE void HandleSelectionChanged(TSharedPtr<FRigTreeElement> Selection, ESelectInfo::Type SelectInfo)
+	void HandleSelectionChanged(TSharedPtr<FRigTreeElement> Selection, ESelectInfo::Type SelectInfo)
 	{
 		if(bIsChangingRigHierarchy)
 		{
diff --git a/Engine/Plugins/Animation/ControlRig/Source/ControlRigEditor/Private/Graph/ControlRigConnectionDrawingPolicy.h b/Engine/Plugins/Animation/ControlRig/Source/ControlRigEditor/Private/Graph/ControlRigConnectionDrawingPolicy.h
--- a/Engine/Plugins/Animation/ControlRig/Source/ControlRigEditor/Private/Graph/ControlRigConnectionDrawingPolicy.h
+++ b/Engine/Plugins/Animation/ControlRig/Source/ControlRigEditor/Private/Graph/ControlRigConnectionDrawingPolicy.h
@@ -36,7 +36,7 @@ private:
 	// Each time a reroute node is encountered, input geometry is compared to output geometry to see if the pins on the reroute node need to be reversed
 	TMap<UControlRigGraphNode*, bool> RerouteNodeToReversedDirectionMap;
 
-	FORCEINLINE bool UseLowDetailConnections() const { return ZoomFactor <= 0.175f; /* zoom level -9 */ }
+	bool UseLowDetailConnections() const { return ZoomFactor <= 0.175f; /* zoom level -9 */ }
 
 	bool ShouldChangeTangentForReouteControlPoint(UControlRigGraphNode* Node);
 	// Average of the positions of all pins connected to InPin
diff --git a/Engine/Plugins/Animation/ControlRig/Source/ControlRigEditor/Public/IControlRigEditorModule.h b/Engine/Plugins/Animation/ControlRig/Source/ControlRigEditor/Public/IControlRigEditorModule.h
--- a/Engine/Plugins/Animation/ControlRig/Source/ControlRigEditor/Public/IControlRigEditorModule.h
+++ b/Engine/Plugins/Animation/ControlRig/Source/ControlRigEditor/Public/IControlRigEditorModule.h
@@ -22,7 +22,7 @@ class IControlRigEditorModule : public IModuleInterface, public IHasMenuExtensib
 {
 public:
 
-	static FORCEINLINE IControlRigEditorModule& Get()
+	static IControlRigEditorModule& Get()
 	{
 		return FModuleManager::LoadModuleChecked< IControlRigEditorModule >(TEXT("ControlRigEditor"));
 	}
diff --git a/Engine/Source/Developer/RigVMDeveloper/Private/RigVMCompiler/RigVMCodeGenerator.cpp b/Engine/Source/Developer/RigVMDeveloper/Private/RigVMCompiler/RigVMCodeGenerator.cpp
--- a/Engine/Source/Developer/RigVMDeveloper/Private/RigVMCompiler/RigVMCodeGenerator.cpp
+++ b/Engine/Source/Developer/RigVMDeveloper/Private/RigVMCompiler/RigVMCodeGenerator.cpp
@@ -12,7 +12,7 @@ static constexpr TCHAR RigVM_NewLineFormat[] = TEXT("\r\n");
 static constexpr TCHAR RigVM_IncludeBracketFormat[] = TEXT("#include <{0}>");
 static constexpr TCHAR RigVM_IncludeQuoteFormat[] = TEXT("#include \"{0}.h\"");
 static constexpr TCHAR RigVM_DispatchKeyFormat[] = TEXT("{0}_{1}");
-static constexpr TCHAR RigVM_DispatchDeclarationFormat[] = TEXT("\tFORCEINLINE bool {0}({1})\r\n\t{\r\n\t\tstatic const FRigVMFunction* Dispatch = FRigVMRegistry::Get().FindFunction(TEXT(\"{2}\"));\r\n\t\tif(Dispatch == nullptr) return false;");
+static constexpr TCHAR RigVM_DispatchDeclarationFormat[] = TEXT("\tbool {0}({1})\r\n\t{\r\n\t\tstatic const FRigVMFunction* Dispatch = FRigVMRegistry::Get().FindFunction(TEXT(\"{2}\"));\r\n\t\tif(Dispatch == nullptr) return false;");
 static constexpr TCHAR RigVM_UPropertyDeclareFormat[] = TEXT("\tUPROPERTY()\r\n\t{0} {1};");
 static constexpr TCHAR RigVM_UPropertyMemberFormat[] = TEXT("\tstatic const FProperty* {0}_Ptr;");
 static constexpr TCHAR RigVM_UPropertyMember2Format[] = TEXT("const FProperty* U{0}::{1}_Ptr = nullptr;");
