Constraint Scripting Fixes:
Constraint Manager, add function to remove Constraint.
CreateTransformableHandle to Control Rig properly exposed.
Added TimeUnit to MoveConstraintKey, Compensate and SetConstraintActiveKey so you can use display time.
Created GetConstraintKeys that will give return array of times and bools.

diff --git a/Engine/Plugins/Animation/ControlRig/Source/ControlRig/Public/ControlRig.h b/Engine/Plugins/Animation/ControlRig/Source/ControlRig/Public/ControlRig.h
--- a/Engine/Plugins/Animation/ControlRig/Source/ControlRig/Public/ControlRig.h
+++ b/Engine/Plugins/Animation/ControlRig/Source/ControlRig/Public/ControlRig.h
@@ -164,7 +164,7 @@ public:
 
 	/** Creates a transformable control handle for the specified control to be used by the constraints system. Should use the UObject from 
 	ConstraintsScriptingLibrary::GetManager(UWorld* InWorld)*/
-	UFUNCTION(BlueprintPure, Category = "Control Rig | Constraints")
+	UFUNCTION(BlueprintCallable, Category = "Control Rig | Constraints")
 	UTransformableControlHandle* CreateTransformableControlHandle(UObject* Outer, const FName& ControlName) const;
 
 
diff --git a/Engine/Plugins/Animation/ControlRig/Source/ControlRigEditor/Private/ControlRigSequencerEditorLibrary.cpp b/Engine/Plugins/Animation/ControlRig/Source/ControlRigEditor/Private/ControlRigSequencerEditorLibrary.cpp
--- a/Engine/Plugins/Animation/ControlRig/Source/ControlRigEditor/Private/ControlRigSequencerEditorLibrary.cpp
+++ b/Engine/Plugins/Animation/ControlRig/Source/ControlRigEditor/Private/ControlRigSequencerEditorLibrary.cpp
@@ -491,36 +491,156 @@ TArray <UTickableConstraint*> UControlRigSequencerEditorLibrary::GetConstraintsF
 	return Constraints;
 }
 
-bool UControlRigSequencerEditorLibrary::Compensate(UTickableConstraint* InConstraint, const FFrameNumber& FrameTime)
+bool UControlRigSequencerEditorLibrary::Compensate(UTickableConstraint* InConstraint, FFrameNumber InTime, ESequenceTimeUnit TimeUnit)
 {
+	TWeakPtr<ISequencer> WeakSequencer = GetSequencerFromAsset();
+	if (WeakSequencer.IsValid() == false)
+	{
+		UE_LOG(LogControlRig, Error, TEXT("Compensate: Need open Sequencer"));
+		return false;
+	}
 	if (UTickableTransformConstraint* Constraint = Cast<UTickableTransformConstraint>(InConstraint))
 	{
-		TOptional<FFrameNumber> OptTime(FrameTime);
+		TSharedPtr<ISequencer>  Sequencer = WeakSequencer.Pin();
+		FFrameRate TickResolution = Sequencer->GetFocusedTickResolution();
+		FFrameRate DisplayRate = Sequencer->GetFocusedDisplayRate();
+		if (TimeUnit == ESequenceTimeUnit::DisplayRate)
+		{
+			InTime = FFrameRate::TransformTime(FFrameTime(InTime, 0), DisplayRate, TickResolution).RoundToFrame();
+		}
+		TOptional<FFrameNumber> OptTime(InTime);
 		FConstraintChannelHelper::Compensate(Constraint, OptTime);
 		return true;
 	}
+	else
+	{
+		UE_LOG(LogControlRig, Error, TEXT("Compensate: Not Transform Constraint"));
+	}
 	return false;
 }
 
 bool UControlRigSequencerEditorLibrary::CompensateAll(UTickableConstraint* InConstraint)
 {
+	TWeakPtr<ISequencer> WeakSequencer = GetSequencerFromAsset();
+	if (WeakSequencer.IsValid() == false)
+	{
+		UE_LOG(LogControlRig, Error, TEXT("CompensateAll: Need open Sequencer"));
+		return false;
+	}
 	if (UTickableTransformConstraint* Constraint = Cast<UTickableTransformConstraint>(InConstraint))
 	{
 		FConstraintChannelHelper::Compensate(Constraint, TOptional<FFrameNumber>());
 		return true;
 	}
+	else
+	{
+		UE_LOG(LogControlRig, Error, TEXT("CompensateAll: Not Transform Constraint"));
+	}
 	return false;
 }
 
-bool UControlRigSequencerEditorLibrary::SetConstraintActiveKey(UTickableConstraint* InConstraint, bool bActive, const FFrameNumber& FrameTime)
+bool UControlRigSequencerEditorLibrary::SetConstraintActiveKey(UTickableConstraint* InConstraint, bool bActive, FFrameNumber InTime, ESequenceTimeUnit TimeUnit)
 {
+	TWeakPtr<ISequencer> WeakSequencer = GetSequencerFromAsset();
+	if (WeakSequencer.IsValid() == false)
+	{
+		UE_LOG(LogControlRig, Error, TEXT("SetConstraintActiveKey: Need open Sequencer"));
+		return false;
+	}
+
 	if (UTickableTransformConstraint* Constraint = Cast<UTickableTransformConstraint>(InConstraint))
 	{
-		return FConstraintChannelHelper::SmartConstraintKey(Constraint, TOptional<bool>(bActive), TOptional<FFrameNumber>(FrameTime));
+		TSharedPtr<ISequencer>  Sequencer = WeakSequencer.Pin();
+		FFrameRate TickResolution = Sequencer->GetFocusedTickResolution();
+		FFrameRate DisplayRate = Sequencer->GetFocusedDisplayRate();
+
+		if (TimeUnit == ESequenceTimeUnit::DisplayRate)
+		{
+			InTime = FFrameRate::TransformTime(FFrameTime(InTime, 0), DisplayRate, TickResolution).RoundToFrame();
+		}
+		return FConstraintChannelHelper::SmartConstraintKey(Constraint, TOptional<bool>(bActive), TOptional<FFrameNumber>(InTime));
+	}
+	else
+	{
+		UE_LOG(LogControlRig, Error, TEXT("SetConstraintActiveKey: Not Transform Constraint"));
 	}
 	return false;
 }
 
+bool UControlRigSequencerEditorLibrary::GetConstraintKeys(UTickableConstraint* InConstraint, UMovieSceneSection* ConstraintSection, TArray<bool>& OutBools, TArray<FFrameNumber>& OutFrames, ESequenceTimeUnit TimeUnit)
+{
+	if (InConstraint == nullptr)
+	{
+		UE_LOG(LogControlRig, Error, TEXT("GetConstraintKeys: Constraint not valid"));
+		return false;
+	}
+	TWeakPtr<ISequencer> WeakSequencer = GetSequencerFromAsset();
+	if (WeakSequencer.IsValid() == false)
+	{
+		UE_LOG(LogControlRig, Error, TEXT("GetConstraintKeys: Need open Sequencer"));
+		return false;
+	}
+
+	IMovieSceneConstrainedSection* ConstrainedSection = Cast<IMovieSceneConstrainedSection>(ConstraintSection);
+	if (ConstrainedSection == nullptr)
+	{
+		UE_LOG(LogControlRig, Error, TEXT("GetConstraintKeys: Section doesn't support constraints"));
+		return false;
+	}
+	FConstraintAndActiveChannel* ConstraintAndChannel = ConstrainedSection->GetConstraintChannel(InConstraint->GetFName());
+	if (ConstraintAndChannel == nullptr)
+	{
+		UE_LOG(LogControlRig, Error, TEXT("GetConstraintKeys: Constraint not found in section"));
+		return false;
+	}
+
+	if (UTickableTransformConstraint* Constraint = Cast<UTickableTransformConstraint>(InConstraint))
+	{
+		TSharedPtr<ISequencer>  Sequencer = WeakSequencer.Pin();
+		FFrameRate TickResolution = Sequencer->GetFocusedTickResolution();
+		FFrameRate DisplayRate = Sequencer->GetFocusedDisplayRate();
+
+
+		TArray<FFrameNumber> OurKeyTimes;
+		TArray<FKeyHandle> OurKeyHandles;
+		TRange<FFrameNumber> CurrentFrameRange;
+		CurrentFrameRange.SetLowerBound(TRangeBound<FFrameNumber>());
+		CurrentFrameRange.SetUpperBound(TRangeBound<FFrameNumber>());
+
+		TMovieSceneChannelData<bool> ChannelInterface = ConstraintAndChannel->ActiveChannel.GetData();
+		ChannelInterface.GetKeys(CurrentFrameRange, &OurKeyTimes, &OurKeyHandles);
+		
+		if (OurKeyTimes.Num() > 0)
+		{
+			OutBools.SetNum(OurKeyTimes.Num());
+			OutFrames.SetNum(OurKeyTimes.Num());
+			int32 Index = 0;
+			for (FFrameNumber Frame : OurKeyTimes)
+			{
+				ConstraintAndChannel->ActiveChannel.Evaluate(FFrameTime(Frame, 0.0f), OutBools[Index]);
+				if (TimeUnit == ESequenceTimeUnit::DisplayRate)
+				{
+					Frame = FFrameRate::TransformTime(FFrameTime(Frame, 0), TickResolution,DisplayRate).RoundToFrame();
+				}
+				OutFrames[Index] = Frame;
+				++Index;
+			}
+		}
+		else
+		{
+			OutBools.SetNum(0);
+			OutFrames.SetNum(0);
+		}
+	}
+	else
+	{
+		UE_LOG(LogControlRig, Error, TEXT("GetConstraintKeys: Not Transform Constraint"));
+		return false;
+	}
+	return true;
+}
+
+
 bool UControlRigSequencerEditorLibrary::MoveConstraintKey(UTickableConstraint* Constraint, UMovieSceneSection* ConstraintSection, FFrameNumber InTime, FFrameNumber InNewTime, ESequenceTimeUnit TimeUnit)
 {
 	TWeakPtr<ISequencer> WeakSequencer = GetSequencerFromAsset();
diff --git a/Engine/Plugins/Animation/ControlRig/Source/ControlRigEditor/Public/ControlRigSequencerEditorLibrary.h b/Engine/Plugins/Animation/ControlRig/Source/ControlRigEditor/Public/ControlRigSequencerEditorLibrary.h
--- a/Engine/Plugins/Animation/ControlRig/Source/ControlRigEditor/Public/ControlRigSequencerEditorLibrary.h
+++ b/Engine/Plugins/Animation/ControlRig/Source/ControlRigEditor/Public/ControlRigSequencerEditorLibrary.h
@@ -161,16 +161,29 @@ public:
 	UFUNCTION(BlueprintCallable, Category = "Editor Scripting | Sequencer Tools | Control Rig")
 	static UTickableConstraint* AddConstraint(UWorld* World, ETransformConstraintType InType, UTransformableHandle* InChild, UTransformableHandle* InParent, const bool bMaintainOffset);
 	
+	/**
+	* Get the constraint keys for the specified constraint
+	* @param InConstraint The constraint to get
+	* @param ConstraintSection Section containing Cosntraint Key
+	* @param OutBools Array of whether or not it's active at the specified times
+	* @param OutFrames The Times for the keys
+	* @param TimeUnit Unit for the time params, either in display rate or tick resolution
+	* @return Returns true if we got the keys from this constraint
+	*/
+	UFUNCTION(BlueprintCallable, Category = "Editor Scripting | Sequencer Tools | Control Rig")
+	static bool GetConstraintKeys(UTickableConstraint* InConstraint, UMovieSceneSection* ConstraintSection,TArray<bool>& OutBools, TArray<FFrameNumber>& OutFrames, ESequenceTimeUnit TimeUnit = ESequenceTimeUnit::DisplayRate);
+	
 	/**
 	* Set the constraint active key in the current open Sequencer
 	* @param InConstraint The constraint to set the key
 	* @param bActive Whether or not it's active
-	* @param FrameTime Time to set the value 
+	* @param FrameTime Time to set the value
+	* @param TimeUnit Unit for the time params, either in display rate or tick resolution
 	* @return Returns true if we set the constraint to be the passed in value, false if not. We may not do so if the value is the same.
 	*/
 	UFUNCTION(BlueprintCallable, Category = "Editor Scripting | Sequencer Tools | Control Rig")
-	static bool SetConstraintActiveKey(UTickableConstraint* InConstraint, bool bActive,const FFrameNumber& FrameTime);
-	
+	static bool SetConstraintActiveKey(UTickableConstraint* InConstraint, bool bActive, FFrameNumber InFrame, ESequenceTimeUnit TimeUnit = ESequenceTimeUnit::DisplayRate);
+
 	/**
 	* Get all constraints for this object, which is described by a transformable handle
 	* @param InChild The handle to look for constraints controlling it
@@ -188,7 +201,7 @@ public:
 	* @return Will return false if function fails, for example if there is no key at this time it will fail, or if the new time is invalid it could fail also
 	*/
 	UFUNCTION(BlueprintCallable, Category = "Editor Scripting | Sequencer Tools | Control Rig")
-	bool MoveConstraintKey(UTickableConstraint* Constraint, UMovieSceneSection* ConstraintSection, FFrameNumber InTime, FFrameNumber InNewTime, ESequenceTimeUnit TimeUnit);
+	static bool MoveConstraintKey(UTickableConstraint* Constraint, UMovieSceneSection* ConstraintSection, FFrameNumber InTime, FFrameNumber InNewTime, ESequenceTimeUnit TimeUnit = ESequenceTimeUnit::DisplayRate);
 
 	/** Delete the Key for the Constraint at the specified time. 
 	* @param InConstraint The constraint whose key to move
@@ -202,11 +215,12 @@ public:
 	/**
 	* Compensate constraint at the specfied time 
 	* @param InConstraint The constraint to compensate
-	* @param FrameTime  Time to compensate
+	* @param InTime Time to compensate
+	* @param TimeUnit Unit for the InTime, either in display rate or tick resolution
 	* @return Returns true if it can compensate
 	*/
 	UFUNCTION(BlueprintCallable, Category = "Editor Scripting | Sequencer Tools | Control Rig")
-	static bool Compensate(UTickableConstraint* InConstraint, const FFrameNumber& FrameTime);
+	static bool Compensate(UTickableConstraint* InConstraint,  FFrameNumber InTime, ESequenceTimeUnit TimeUnit = ESequenceTimeUnit::DisplayRate);
 
 	/**
 	* Compensate constraint at all keys
@@ -982,7 +996,7 @@ public:
 	* @return Return true if the function succeeds, false if it doesn't which can happen if the name arrays don't match in size or any of the new Control Names aren't valid
 	*/
 	UFUNCTION(BlueprintCallable, Category = "Editor Scripting | Sequencer Tools | Control Rig")
-	bool RenameControlRigControlChannels(ULevelSequence* InSequence, UControlRig* InControlRig, const TArray<FName>& InOldControlNames, const TArray<FName>& InNewControlNames);
+	static bool RenameControlRigControlChannels(ULevelSequence* InSequence, UControlRig* InControlRig, const TArray<FName>& InOldControlNames, const TArray<FName>& InNewControlNames);
 
 	/** Get the controls mask for the given ControlName */
 	UFUNCTION(BlueprintPure, Category = "Editor Scripting | Sequencer Tools | Control Rig")
diff --git a/Engine/Source/Runtime/Experimental/Animation/Constraints/Private/ConstraintsScripting.cpp b/Engine/Source/Runtime/Experimental/Animation/Constraints/Private/ConstraintsScripting.cpp
--- a/Engine/Source/Runtime/Experimental/Animation/Constraints/Private/ConstraintsScripting.cpp
+++ b/Engine/Source/Runtime/Experimental/Animation/Constraints/Private/ConstraintsScripting.cpp
@@ -61,3 +61,17 @@ bool UConstraintsScriptingLibrary::RemoveConstraint(UWorld* InWorld, int32 InInd
 	return Controller.RemoveConstraint(InIndex);
 }
 
+bool UConstraintsScriptingLibrary::RemoveThisConstraint(UWorld* InWorld, UTickableConstraint* InTickableConstraint)
+{
+	const FConstraintsManagerController& Controller = FConstraintsManagerController::Get(InWorld);
+	const TArray< TObjectPtr<UTickableConstraint> >& ConstraintsArray = Controller.GetConstraintsArray();
+	for (int32 Index = 0; Index < ConstraintsArray.Num(); ++Index)
+	{
+		if (InTickableConstraint == ConstraintsArray[Index])
+		{
+			return Controller.RemoveConstraint(Index);
+
+		}
+	}
+	return false;
+}
diff --git a/Engine/Source/Runtime/Experimental/Animation/Constraints/Public/ConstraintsScripting.h b/Engine/Source/Runtime/Experimental/Animation/Constraints/Public/ConstraintsScripting.h
--- a/Engine/Source/Runtime/Experimental/Animation/Constraints/Public/ConstraintsScripting.h
+++ b/Engine/Source/Runtime/Experimental/Animation/Constraints/Public/ConstraintsScripting.h
@@ -74,9 +74,19 @@ public:
 	UFUNCTION(BlueprintCallable, Category = "Editor Scripting | Sequencer Tools | Control Rig | Constraints")
 	static TArray<UTickableConstraint*> GetConstraintsArray(UWorld* InWorld);
 
+
+	/**
+	* Remove specified constraint 
+	* @param InWorld World to remove the constraint
+	* @param InTickableConstraint Constraint to remove
+	* @return return If constraint removed correctly
+	*/
+	UFUNCTION(BlueprintCallable, Category = "Editor Scripting | Sequencer Tools | Control Rig | Constraints")
+	static bool RemoveThisConstraint(UWorld* InWorld, UTickableConstraint* InTickableConstraint);
+
 	/**
 	* Remove constraint at specified index
-	* @param InWorld World to create the constraint
+	* @param InWorld World to remove the constraint
 	* @param InIndex Index to remove from
 	* @return return If constraint removed correctly
 	*/
