Control Rig: Turn graphs read only when using a nativized rig

diff --git a/Engine/Plugins/Animation/ControlRig/Source/ControlRigEditor/Private/Editor/ControlRigEditor.cpp b/Engine/Plugins/Animation/ControlRig/Source/ControlRigEditor/Private/Editor/ControlRigEditor.cpp
--- a/Engine/Plugins/Animation/ControlRig/Source/ControlRigEditor/Private/Editor/ControlRigEditor.cpp
+++ b/Engine/Plugins/Animation/ControlRig/Source/ControlRigEditor/Private/Editor/ControlRigEditor.cpp
@@ -101,6 +101,7 @@
 #include "Styling/AppStyle.h"
 #include "AssetRegistry/AssetRegistryModule.h"
 #include "RigVMFunctions/RigVMFunction_ControlFlow.h"
+#include "RigVMModel/Nodes/RigVMAggregateNode.h"
 
 #define LOCTEXT_NAMESPACE "ControlRigEditor"
 
@@ -3589,7 +3590,32 @@ void FControlRigEditor::Tick(float DeltaTime)
 
 bool FControlRigEditor::IsEditable(UEdGraph* InGraph) const
 {
-	return IsGraphInCurrentBlueprint(InGraph);
+	if(!IsGraphInCurrentBlueprint(InGraph))
+	{
+		return false;
+	}
+	
+	if(UControlRigBlueprint* ControlRigBlueprint = CastChecked<UControlRigBlueprint>(GetBlueprintObj()))
+	{
+		// aggregate graphs are always read only
+		if(const URigVMGraph* Model = ControlRigBlueprint->GetModel(InGraph))
+		{
+			if(Model->GetOuter()->IsA<URigVMAggregateNode>())
+			{
+				return false;
+			}
+		}
+
+		if(ControlRig && ControlRig->VM)
+		{
+			const bool bIsReadOnly = ControlRig->VM->IsNativized();
+			const bool bIsEditable = !bIsReadOnly;
+			InGraph->bEditable = bIsEditable ? 1 : 0;
+			return bIsEditable;
+		}
+	}
+
+	return IControlRigEditor::IsEditable(InGraph);
 }
 
 bool FControlRigEditor::IsCompilingEnabled() const
