Control Rig: Fix ExecutionStack not selecting correct instructions

diff --git a/Engine/Plugins/Animation/ControlRig/Source/ControlRigEditor/Private/Editor/ControlRigEditor.cpp b/Engine/Plugins/Animation/ControlRig/Source/ControlRigEditor/Private/Editor/ControlRigEditor.cpp
--- a/Engine/Plugins/Animation/ControlRig/Source/ControlRigEditor/Private/Editor/ControlRigEditor.cpp
+++ b/Engine/Plugins/Animation/ControlRig/Source/ControlRigEditor/Private/Editor/ControlRigEditor.cpp
@@ -500,7 +500,10 @@ void FControlRigEditor::InitControlRigEditor(const EToolkitMode::Type Mode, cons
 				InControlRigBlueprint->RebuildGraphFromModel();
 
 				// selection state does not need to be persistent, even though it is saved in the RigVM.
-				InControlRigBlueprint->GetController()->ClearNodeSelection(false);
+				for (URigVMGraph* Graph : InControlRigBlueprint->GetAllModels())
+				{
+					InControlRigBlueprint->GetController(Graph)->ClearNodeSelection(false);
+				}
 
 				if (UPackage* Package = InControlRigBlueprint->GetOutermost())
 				{
diff --git a/Engine/Plugins/Animation/ControlRig/Source/ControlRigEditor/Private/Editor/SControlRigStackView.cpp b/Engine/Plugins/Animation/ControlRig/Source/ControlRigEditor/Private/Editor/SControlRigStackView.cpp
--- a/Engine/Plugins/Animation/ControlRig/Source/ControlRigEditor/Private/Editor/SControlRigStackView.cpp
+++ b/Engine/Plugins/Animation/ControlRig/Source/ControlRigEditor/Private/Editor/SControlRigStackView.cpp
@@ -123,6 +123,7 @@ void SRigStackItem::Construct(const FArguments& InArgs, const TSharedRef<STableV
 					SAssignNew(TextWidget, STextBlock)
 					.Text(this, &SRigStackItem::GetLabelText)
 					.Font(this, &SRigStackItem::GetLabelFont)
+					.ToolTipText(this, &SRigStackItem::GetTooltip)
 					.Justification(ETextJustify::Left)
 				]
 			]
@@ -183,6 +184,11 @@ FSlateFontInfo SRigStackItem::GetLabelFont() const
 	return IDetailLayoutBuilder::GetDetailFontBold();
 }
 
+FText SRigStackItem::GetTooltip() const
+{
+	return FText::FromString(WeakStackEntry.Pin()->Callstack.GetCallPath(true));
+}
+
 FText SRigStackItem::GetVisitedCountText() const
 {
 	if(WeakStackEntry.IsValid() && WeakBlueprint.IsValid())
@@ -267,7 +273,6 @@ void SControlRigStackView::Construct( const FArguments& InArgs, TSharedRef<FCont
 {
 	ControlRigEditor = InControlRigEditor;
 	ControlRigBlueprint = ControlRigEditor.Pin()->GetControlRigBlueprint();
-	Graph = Cast<UControlRigGraph>(ControlRigBlueprint->GetLastEditedUberGraph());
 	CommandList = MakeShared<FUICommandList>();
 	bSuspendModelNotifications = false;
 	bSuspendControllerSelection = false;
@@ -997,17 +1002,33 @@ void SControlRigStackView::HandleModifiedEvent(ERigVMGraphNotifType InNotifType,
 
 	switch (InNotifType)
 	{
-		case ERigVMGraphNotifType::NodeSelected:
-		case ERigVMGraphNotifType::NodeDeselected:
+		case ERigVMGraphNotifType::NodeSelectionChanged:
 		{
+			const TGuardValue<bool> SuspendNotifs(bSuspendModelNotifications, true);
+			TreeView->ClearSelection();
+			TArray<const URigVMNode*> SelectedNodes;
+			const TArray<FName>& SelectedNames = InGraph->GetSelectNodes();
+			for (const FName& Selected : SelectedNames)
+			{
+				if (const URigVMNode* Node = InGraph->FindNodeByName(Selected))
+				{
+					SelectedNodes.Add(Node);
+				}
+			}		
+
+			TArray<TSharedPtr<FRigStackEntry>> SelectedItems;
 			for (TSharedPtr<FRigStackEntry>& Operator : Operators)
 			{
-				if(Operator->Callstack.Contains(InSubject))
+				for (const URigVMNode* Node : SelectedNodes)
 				{
-					const TGuardValue<bool> SuspendNotifs(bSuspendModelNotifications, true);
-					TreeView->SetItemSelection(Operator, InNotifType == ERigVMGraphNotifType::NodeSelected, ESelectInfo::Direct);
+					if(Operator->Callstack.Contains(Node))
+					{
+						SelectedItems.Add(Operator);
+						break;
+					}
 				}
 			}
+			TreeView->SetItemSelection(SelectedItems, true, ESelectInfo::Direct);
 			break;
 		}
 		default:
diff --git a/Engine/Plugins/Animation/ControlRig/Source/ControlRigEditor/Private/Editor/SControlRigStackView.h b/Engine/Plugins/Animation/ControlRig/Source/ControlRigEditor/Private/Editor/SControlRigStackView.h
--- a/Engine/Plugins/Animation/ControlRig/Source/ControlRigEditor/Private/Editor/SControlRigStackView.h
+++ b/Engine/Plugins/Animation/ControlRig/Source/ControlRigEditor/Private/Editor/SControlRigStackView.h
@@ -58,6 +58,7 @@ private:
 	FText GetIndexText() const;
 	FText GetLabelText() const;
 	FSlateFontInfo GetLabelFont() const;
+	FText GetTooltip() const;
 	FText GetVisitedCountText() const;
 	FText GetDurationText() const;
 };
@@ -127,7 +128,6 @@ private:
 
 	TWeakPtr<FControlRigEditor> ControlRigEditor;
 	TWeakObjectPtr<UControlRigBlueprint> ControlRigBlueprint;
-	TWeakObjectPtr<UControlRigGraph> Graph;
 
 	TArray<TSharedPtr<FRigStackEntry>> Operators;
 
diff --git a/Engine/Source/Developer/RigVMDeveloper/Private/RigVMCompiler/RigVMCompiler.cpp b/Engine/Source/Developer/RigVMDeveloper/Private/RigVMCompiler/RigVMCompiler.cpp
--- a/Engine/Source/Developer/RigVMDeveloper/Private/RigVMCompiler/RigVMCompiler.cpp
+++ b/Engine/Source/Developer/RigVMDeveloper/Private/RigVMCompiler/RigVMCompiler.cpp
@@ -2806,21 +2806,33 @@ void URigVMCompiler::AddCopyOperator(const FRigVMCopyOp& InOp, const FRigVMAssig
 	int32 InstructionIndex = WorkData.VM->GetByteCode().GetNumInstructions() - 1;
 	if (Settings.SetupNodeInstructionIndex)
 	{
+		bool bSetSubject = false;
 		if (URigVMPin* SourcePin = InAssignExpr->GetSourcePin())
 		{
 			if (URigVMVariableNode* VariableNode = Cast<URigVMVariableNode>(SourcePin->GetNode()))
 			{
 				const FRigVMCallstack Callstack = InSourceExpr->GetProxy().GetSibling(VariableNode).GetCallstack();
 				WorkData.VM->GetByteCode().SetSubject(InstructionIndex, Callstack.GetCallPath(), Callstack.GetStack());
+				bSetSubject = true;
 			}
 		}
 
-		if (URigVMPin* TargetPin = InAssignExpr->GetTargetPin())
+		if (!bSetSubject)
 		{
-			if (URigVMVariableNode* VariableNode = Cast<URigVMVariableNode>(TargetPin->GetNode()))
+			if (URigVMPin* TargetPin = InAssignExpr->GetTargetPin())
 			{
-				const FRigVMCallstack Callstack = InTargetExpr->GetProxy().GetSibling(VariableNode).GetCallstack();
-				WorkData.VM->GetByteCode().SetSubject(InstructionIndex, Callstack.GetCallPath(), Callstack.GetStack());
+				if (URigVMVariableNode* VariableNode = Cast<URigVMVariableNode>(TargetPin->GetNode()))
+				{
+					const FRigVMCallstack Callstack = InTargetExpr->GetProxy().GetSibling(VariableNode).GetCallstack();
+					WorkData.VM->GetByteCode().SetSubject(InstructionIndex, Callstack.GetCallPath(), Callstack.GetStack());
+					bSetSubject = true;
+				}
+				else
+				{
+					const FRigVMCallstack Callstack = InTargetExpr->GetProxy().GetSibling(TargetPin->GetNode()).GetCallstack();
+					WorkData.VM->GetByteCode().SetSubject(InstructionIndex, Callstack.GetCallPath(), Callstack.GetStack());
+					bSetSubject = true;
+				}
 			}
 		}
 	}
