Proposed solution to fix unbounded accumulation of UMovieSceneControlRigParameterTrack::HandlePackageDone calls in the FCoreUObjectDelegates::OnEndLoadPackage when rigs are loaded from a saved asset (as opposed to created programmatically as part of the animation upgrade).  The previous code was attempting to ensure that the rig's class package had finished loading operations by checking its "end loaded" status before proceeding wth reconstructing the rig.  When the rig's class was native, its package would never get the "end loaded" state and thus would not reconstruct the control rig or remove itself from the OnEndLoadPackage delegate.  Over the span of a cook (with editor domain enabled so the assets were being read from serialized form), it would accumulate more and more of these delegates and it would substantially slow down the cook.

The original issue was introduced in CL17884029 in response to UE-131056.

diff --git a/Engine/Plugins/Animation/ControlRig/Source/ControlRig/Private/Sequencer/MovieSceneControlRigParameterTrack.cpp b/Engine/Plugins/Animation/ControlRig/Source/ControlRig/Private/Sequencer/MovieSceneControlRigParameterTrack.cpp
--- a/Engine/Plugins/Animation/ControlRig/Source/ControlRig/Private/Sequencer/MovieSceneControlRigParameterTrack.cpp
+++ b/Engine/Plugins/Animation/ControlRig/Source/ControlRig/Private/Sequencer/MovieSceneControlRigParameterTrack.cpp
@@ -441,7 +441,7 @@ void UMovieSceneControlRigParameterTrack::HandlePackageDone(const FEndLoadPackag
 	}
 
 	// ensure both packages are fully end-loaded
-	if (ControlRig)
+	if (ControlRig && !ControlRig->GetClass()->IsNative())
 	{
 		if (const UPackage* ControlRigPackage = Cast<UPackage>(ControlRig->GetClass()->GetOutermost()))
 		{
