Control Rig: Fix UHTCompatEditor error when having a property declared inside an interface

diff --git a/Engine/Plugins/Animation/ControlRig/Source/ControlRigDeveloper/Public/ControlRigBlueprint.h b/Engine/Plugins/Animation/ControlRig/Source/ControlRigDeveloper/Public/ControlRigBlueprint.h
--- a/Engine/Plugins/Animation/ControlRig/Source/ControlRigDeveloper/Public/ControlRigBlueprint.h
+++ b/Engine/Plugins/Animation/ControlRig/Source/ControlRigDeveloper/Public/ControlRigBlueprint.h
@@ -417,6 +417,10 @@ protected:
 	UPROPERTY(AssetRegistrySearchable)
 	TArray<FRigVMGraphFunctionHeader> PublicGraphFunctions;
 
+	/** Asset searchable information function references in this rig */
+	UPROPERTY(AssetRegistrySearchable)
+	TArray<FRigVMReferenceNodeData> FunctionReferenceNodeData;
+
 #if WITH_EDITORONLY_DATA
 
 	UPROPERTY(transient, DuplicateTransient)
diff --git a/Engine/Source/Developer/RigVMDeveloper/Public/RigVMModel/RigVMClient.h b/Engine/Source/Developer/RigVMDeveloper/Public/RigVMModel/RigVMClient.h
--- a/Engine/Source/Developer/RigVMDeveloper/Public/RigVMModel/RigVMClient.h
+++ b/Engine/Source/Developer/RigVMDeveloper/Public/RigVMModel/RigVMClient.h
@@ -49,9 +49,6 @@ public:
 	// Reacts to a request to configure a controller
 	virtual void HandleConfigureRigVMController(const FRigVMClient* InClient, URigVMController* InControllerToConfigure) = 0;
 
-	/** Asset searchable information function references in this rig */
-	UPROPERTY(AssetRegistrySearchable)
-	TArray<FRigVMReferenceNodeData> FunctionReferenceNodeData;
 };
 
 UINTERFACE()
