RigVM: Remove support for opaque attributes

diff --git a/Engine/Plugins/Animation/ControlRig/Source/ControlRig/Private/ControlRig.cpp b/Engine/Plugins/Animation/ControlRig/Source/ControlRig/Private/ControlRig.cpp
--- a/Engine/Plugins/Animation/ControlRig/Source/ControlRig/Private/ControlRig.cpp
+++ b/Engine/Plugins/Animation/ControlRig/Source/ControlRig/Private/ControlRig.cpp
@@ -730,7 +730,7 @@ bool UControlRig::Execute(const EControlRigState InState, const FName& InEventNa
 	LatestExecutedState = InState;
 
 	FRigVMExtendedExecuteContext& ExtendedExecuteContext = GetVM()->GetContext();
-	FRigVMExecuteContext& PublicContext = ExtendedExecuteContext.GetPublicData<>();
+	FControlRigExecuteContext& PublicContext = ExtendedExecuteContext.GetPublicData<FControlRigExecuteContext>();
 	PublicContext.SetDeltaTime(DeltaTime);
 	PublicContext.SetAbsoluteTime(AbsoluteTime);
 
@@ -810,7 +810,6 @@ bool UControlRig::Execute(const EControlRigState InState, const FName& InEventNa
 	if (bRequiresInitExecution)
 	{
 		bRequiresInitExecution = false;
-		VM->SetContextPublicDataStruct(FControlRigExecuteContext::StaticStruct());
 
 		if (!bIsInitializingMemory)
 		{
@@ -823,7 +822,7 @@ bool UControlRig::Execute(const EControlRigState InState, const FName& InEventNa
 		}
 	}
 
-	FRigUnitContext Context;
+	FRigUnitContext& Context = PublicContext.UnitContext;
 
 	// setup the draw interface for debug drawing
 	if(!bIsEventInQueue || bIsEventFirstInQueue)
@@ -865,11 +864,11 @@ bool UControlRig::Execute(const EControlRigState InState, const FName& InEventNa
 	Context.InteractionType = InteractionType;
 	Context.ElementsBeingInteracted = ElementsBeingInteracted;
 	Context.State = InState;
+	PublicContext.Hierarchy = GetHierarchy();
 
 	// allow access to the hierarchy
-	Context.Hierarchy = GetHierarchy();
 	Context.HierarchySettings = HierarchySettings;
-	check(Context.Hierarchy);
+	check(PublicContext.Hierarchy);
 
 	// allow access to the default hierarchy to allow to reset
 	if(!HasAnyFlags(RF_ClassDefaultObject))
@@ -878,13 +877,13 @@ bool UControlRig::Execute(const EControlRigState InState, const FName& InEventNa
 		{
 			if(URigHierarchy* DefaultHierarchy = CDO->GetHierarchy())
 			{
-				Context.Hierarchy->DefaultHierarchyPtr = DefaultHierarchy;
+				PublicContext.Hierarchy->DefaultHierarchyPtr = DefaultHierarchy;
 			}
 		}
 	}
 
 	// disable any controller access outside of the construction event
-	FRigHierarchyEnableControllerBracket DisableHierarchyController(Context.Hierarchy, bIsConstructionEvent);
+	FRigHierarchyEnableControllerBracket DisableHierarchyController(PublicContext.Hierarchy, bIsConstructionEvent);
 
 	// setup the context with further fields
 	Context.ToWorldSpaceTransform = FTransform::Identity;
@@ -1515,8 +1514,6 @@ bool UControlRig::ExecuteUnits(FRigUnitContext& InOutContext, const FName& InEve
 		const bool bUseDebuggingSnapshots = !VM->IsNativized();
 		
 		TArray<URigVMMemoryStorage*> LocalMemory = VM->GetLocalMemoryArray();
-		TArray<void*> AdditionalArguments;
-		AdditionalArguments.Add(&InOutContext);
 
 		bool bSuccess = true;
 
@@ -1542,7 +1539,7 @@ bool UControlRig::ExecuteUnits(FRigUnitContext& InOutContext, const FName& InEve
 
 							VM->WorkMemoryStorageObject->CopyFrom(InitializedVMSnapshot->WorkMemoryStorageObject);
 							VM->InvalidateCachedMemory();
-							VM->Initialize(LocalMemory, AdditionalArguments, false);
+							VM->Initialize(LocalMemory, false);
 							bIsValidSnapshot = true;
 						}
 						else
@@ -1554,7 +1551,7 @@ bool UControlRig::ExecuteUnits(FRigUnitContext& InOutContext, const FName& InEve
 				
 				if(!bIsValidSnapshot)
 				{
-					VM->Initialize(LocalMemory, AdditionalArguments);
+					VM->Initialize(LocalMemory);
 
 					// objects assigned to transient properties need transient flag, otherwise it might get exported during save
 					URigVM* InitializedVMSnapshot = NewObject<URigVM>(CDO, NAME_None, RF_Public | RF_Transient);
@@ -1572,7 +1569,7 @@ bool UControlRig::ExecuteUnits(FRigUnitContext& InOutContext, const FName& InEve
 			}
 			else
 			{
-				bSuccess = VM->Initialize(LocalMemory, AdditionalArguments);
+				bSuccess = VM->Initialize(LocalMemory);
 			}
 			bRequiresConstructionEvent = true;
 		}
@@ -1623,7 +1620,7 @@ bool UControlRig::ExecuteUnits(FRigUnitContext& InOutContext, const FName& InEve
 #endif
 			FRigHierarchyExecuteContextBracket HierarchyContextGuard(Hierarchy, &VM->GetContext());
 
-			bSuccess = VM->Execute(LocalMemory, AdditionalArguments, InEventName) != ERigVMExecuteResult::Failed;
+			bSuccess = VM->Execute(LocalMemory, InEventName) != ERigVMExecuteResult::Failed;
 		}
 
 #if UE_CONTROLRIG_PROFILE_EXECUTE_UNITS_NUM
diff --git a/Engine/Plugins/Animation/ControlRig/Source/ControlRig/Private/Rigs/AdditiveControlRig.cpp b/Engine/Plugins/Animation/ControlRig/Source/ControlRig/Private/Rigs/AdditiveControlRig.cpp
--- a/Engine/Plugins/Animation/ControlRig/Source/ControlRig/Private/Rigs/AdditiveControlRig.cpp
+++ b/Engine/Plugins/Animation/ControlRig/Source/ControlRig/Private/Rigs/AdditiveControlRig.cpp
@@ -46,14 +46,17 @@ bool UAdditiveControlRig::ExecuteUnits(FRigUnitContext& InOutContext, const FNam
 {
 	if (InEventName == FRigUnit_BeginExecution::EventName)
 	{
+		FControlRigExecuteContext ExecuteContext;
+		ExecuteContext.Hierarchy = GetHierarchy();
+		ExecuteContext.EventName = InEventName;
+		ExecuteContext.UnitContext = InOutContext;
+
 		for (FRigUnit_AddBoneTransform& Unit : AddBoneRigUnits)
 		{
 			FName ControlName = GetControlName(Unit.Bone);
 			const int32 Index = GetHierarchy()->GetIndex(FRigElementKey(ControlName, ERigElementType::Control));
 			Unit.Transform = GetHierarchy()->GetLocalTransform(Index);
-			Unit.ExecuteContext.Hierarchy = GetHierarchy();
-			Unit.ExecuteContext.EventName = InEventName;
-			Unit.Execute(InOutContext);
+			Unit.Execute(ExecuteContext);
 		}
 	}
 	else if(InEventName == FRigUnit_PrepareForExecution::EventName)
diff --git a/Engine/Plugins/Animation/ControlRig/Source/ControlRig/Private/Rigs/RigHierarchy.cpp b/Engine/Plugins/Animation/ControlRig/Source/ControlRig/Private/Rigs/RigHierarchy.cpp
--- a/Engine/Plugins/Animation/ControlRig/Source/ControlRig/Private/Rigs/RigHierarchy.cpp
+++ b/Engine/Plugins/Animation/ControlRig/Source/ControlRig/Private/Rigs/RigHierarchy.cpp
@@ -2800,15 +2800,15 @@ void URigHierarchy::SetTransform(FRigTransformElement* InTransformElement, const
 
 						if(bChildFound)
 						{
-							const FRigUnitContext* UnitContext = (const FRigUnitContext*)ExecuteContext->OpaqueArguments[0];
-							if(UnitContext && UnitContext->Log)
+							const FRigUnitContext& UnitContext = ExecuteContext->GetPublicData<FControlRigExecuteContext>().UnitContext;
+							if(UnitContext.Log)
 							{
 								static constexpr TCHAR MessageFormat[] = TEXT("Setting transform of parent (%s) after setting child (%s).\nThis may lead to unexpected results.");
 								const FString& Message = FString::Printf(
 									MessageFormat,
 									*InTransformElement->GetName().ToString(),
 									*Child->GetName().ToString());
-								UnitContext->Log->Report(
+								UnitContext.Log->Report(
 									EMessageSeverity::Info,
 									ExecuteContext->GetPublicData<>().GetFunctionName(),
 									ExecuteContext->GetPublicData<>().GetInstructionIndex(),
diff --git a/Engine/Plugins/Animation/ControlRig/Source/ControlRig/Private/Units/Animation/RigUnit_AnimAttribute.cpp b/Engine/Plugins/Animation/ControlRig/Source/ControlRig/Private/Units/Animation/RigUnit_AnimAttribute.cpp
--- a/Engine/Plugins/Animation/ControlRig/Source/ControlRig/Private/Units/Animation/RigUnit_AnimAttribute.cpp
+++ b/Engine/Plugins/Animation/ControlRig/Source/ControlRig/Private/Units/Animation/RigUnit_AnimAttribute.cpp
@@ -184,7 +184,7 @@ void FRigDispatch_GetAnimAttribute::GetAnimAttributeDispatch(FRigVMExtendedExecu
 
 #if WITH_EDITOR
 		{
-			FRigVMExecuteContext& ExecuteContext = InContext.GetPublicData<>();
+			FControlRigExecuteContext& ExecuteContext = InContext.GetPublicData<FControlRigExecuteContext>();
 			if (!UE::Anim::AttributeTypes::IsTypeRegistered(ScriptStruct))
 			{
 				UE_CONTROLRIG_RIGUNIT_REPORT_WARNING(
@@ -246,7 +246,7 @@ void FRigDispatch_SetAnimAttribute::SetAnimAttributeDispatch(FRigVMExtendedExecu
 
 #if WITH_EDITOR
 		{
-			FRigVMExecuteContext& ExecuteContext = InContext.GetPublicData<>();
+			FControlRigExecuteContext& ExecuteContext = InContext.GetPublicData<FControlRigExecuteContext>();
 			if (!UE::Anim::AttributeTypes::IsTypeRegistered(ScriptStruct))
 			{
 				UE_CONTROLRIG_RIGUNIT_REPORT_WARNING(
diff --git a/Engine/Plugins/Animation/ControlRig/Source/ControlRig/Private/Units/Animation/RigUnit_AnimEasing.h b/Engine/Plugins/Animation/ControlRig/Source/ControlRig/Private/Units/Animation/RigUnit_AnimEasing.h
--- a/Engine/Plugins/Animation/ControlRig/Source/ControlRig/Private/Units/Animation/RigUnit_AnimEasing.h
+++ b/Engine/Plugins/Animation/ControlRig/Source/ControlRig/Private/Units/Animation/RigUnit_AnimEasing.h
@@ -20,7 +20,7 @@ struct CONTROLRIG_API FRigUnit_AnimEasingType : public FRigUnit_AnimBase
 	}
 
 	RIGVM_METHOD()
-	virtual void Execute(const FRigUnitContext& Context) override;
+	virtual void Execute() override;
 
 	UPROPERTY(meta = (Input, Output))
 	EControlRigAnimEasingType Type;
@@ -43,7 +43,7 @@ struct CONTROLRIG_API FRigUnit_AnimEasing : public FRigUnit_AnimBase
 	}
 
 	RIGVM_METHOD()
-	virtual void Execute(const FRigUnitContext& Context) override;
+	virtual void Execute() override;
 
 	UPROPERTY(meta=(Input))
 	float Value;
diff --git a/Engine/Plugins/Animation/ControlRig/Source/ControlRig/Private/Units/Animation/RigUnit_AnimEvalRichCurve.h b/Engine/Plugins/Animation/ControlRig/Source/ControlRig/Private/Units/Animation/RigUnit_AnimEvalRichCurve.h
--- a/Engine/Plugins/Animation/ControlRig/Source/ControlRig/Private/Units/Animation/RigUnit_AnimEvalRichCurve.h
+++ b/Engine/Plugins/Animation/ControlRig/Source/ControlRig/Private/Units/Animation/RigUnit_AnimEvalRichCurve.h
@@ -25,7 +25,7 @@ struct CONTROLRIG_API FRigUnit_AnimEvalRichCurve : public FRigUnit_AnimBase
 	}
 
 	RIGVM_METHOD()
-	virtual void Execute(const FRigUnitContext& Context) override;
+	virtual void Execute() override;
 
 	UPROPERTY(meta=(Input))
 	float Value;
diff --git a/Engine/Plugins/Animation/ControlRig/Source/ControlRig/Private/Units/Animation/RigUnit_AnimRichCurve.h b/Engine/Plugins/Animation/ControlRig/Source/ControlRig/Private/Units/Animation/RigUnit_AnimRichCurve.h
--- a/Engine/Plugins/Animation/ControlRig/Source/ControlRig/Private/Units/Animation/RigUnit_AnimRichCurve.h
+++ b/Engine/Plugins/Animation/ControlRig/Source/ControlRig/Private/Units/Animation/RigUnit_AnimRichCurve.h
@@ -22,7 +22,7 @@ struct CONTROLRIG_API FRigUnit_AnimRichCurve : public FRigUnit_AnimBase
 	}
 
 	RIGVM_METHOD()
-	virtual void Execute(const FRigUnitContext& Context) override;
+	virtual void Execute() override;
 
 	UPROPERTY(meta=(Input, Output, Constant))
 	FRuntimeFloatCurve Curve;
diff --git a/Engine/Plugins/Animation/ControlRig/Source/ControlRig/Private/Units/Animation/RigUnit_GetDeltaTime.cpp b/Engine/Plugins/Animation/ControlRig/Source/ControlRig/Private/Units/Animation/RigUnit_GetDeltaTime.cpp
--- a/Engine/Plugins/Animation/ControlRig/Source/ControlRig/Private/Units/Animation/RigUnit_GetDeltaTime.cpp
+++ b/Engine/Plugins/Animation/ControlRig/Source/ControlRig/Private/Units/Animation/RigUnit_GetDeltaTime.cpp
@@ -8,7 +8,7 @@
 FRigUnit_GetDeltaTime_Execute()
 {
     DECLARE_SCOPE_HIERARCHICAL_COUNTER_RIGUNIT()
-	Result = Context.DeltaTime;
+	Result = ExecuteContext.UnitContext.DeltaTime;
 }
 
 #if WITH_DEV_AUTOMATION_TESTS
@@ -16,7 +16,7 @@ FRigUnit_GetDeltaTime_Execute()
 
 IMPLEMENT_RIGUNIT_AUTOMATION_TEST(FRigUnit_GetDeltaTime)
 {
-	Context.DeltaTime = 0.2f;
+	ExecuteContext.UnitContext.DeltaTime = 0.2f;
 	Execute();
 	AddErrorIfFalse(FMath::IsNearlyEqual(Unit.Result, 0.2f), TEXT("unexpected delta time"));
 	return true;
diff --git a/Engine/Plugins/Animation/ControlRig/Source/ControlRig/Private/Units/Animation/RigUnit_TimeConversion.cpp b/Engine/Plugins/Animation/ControlRig/Source/ControlRig/Private/Units/Animation/RigUnit_TimeConversion.cpp
--- a/Engine/Plugins/Animation/ControlRig/Source/ControlRig/Private/Units/Animation/RigUnit_TimeConversion.cpp
+++ b/Engine/Plugins/Animation/ControlRig/Source/ControlRig/Private/Units/Animation/RigUnit_TimeConversion.cpp
@@ -9,9 +9,9 @@ FRigUnit_FramesToSeconds_Execute()
 {
     DECLARE_SCOPE_HIERARCHICAL_COUNTER_RIGUNIT()
 
-    if(Context.FramesPerSecond > SMALL_NUMBER)
+    if(ExecuteContext.UnitContext.FramesPerSecond > SMALL_NUMBER)
     {
-		Seconds = Frames / Context.FramesPerSecond;
+		Seconds = Frames / ExecuteContext.UnitContext.FramesPerSecond;
 	}
 	else
 	{
@@ -23,9 +23,9 @@ FRigUnit_SecondsToFrames_Execute()
 {
     DECLARE_SCOPE_HIERARCHICAL_COUNTER_RIGUNIT()
 
-    if(Context.FramesPerSecond > SMALL_NUMBER)
+    if(ExecuteContext.UnitContext.FramesPerSecond > SMALL_NUMBER)
     {
-		Frames = Seconds * Context.FramesPerSecond;
+		Frames = Seconds * ExecuteContext.UnitContext.FramesPerSecond;
 	}
 	else
 	{
diff --git a/Engine/Plugins/Animation/ControlRig/Source/ControlRig/Private/Units/Animation/RigUnit_TimeConversion.h b/Engine/Plugins/Animation/ControlRig/Source/ControlRig/Private/Units/Animation/RigUnit_TimeConversion.h
--- a/Engine/Plugins/Animation/ControlRig/Source/ControlRig/Private/Units/Animation/RigUnit_TimeConversion.h
+++ b/Engine/Plugins/Animation/ControlRig/Source/ControlRig/Private/Units/Animation/RigUnit_TimeConversion.h
@@ -19,7 +19,7 @@ struct CONTROLRIG_API FRigUnit_FramesToSeconds : public FRigUnit_AnimBase
 	}
 
 	RIGVM_METHOD()
-	virtual void Execute(const FRigUnitContext& Context) override;
+	virtual void Execute() override;
 
 	UPROPERTY(meta=(Input))
 	float Frames;
@@ -42,7 +42,7 @@ struct CONTROLRIG_API FRigUnit_SecondsToFrames : public FRigUnit_AnimBase
 	}
 
 	RIGVM_METHOD()
-	virtual void Execute(const FRigUnitContext& Context) override;
+	virtual void Execute() override;
 
 	UPROPERTY(meta=(Input))
 	float Seconds;
diff --git a/Engine/Plugins/Animation/ControlRig/Source/ControlRig/Private/Units/Collision/RigUnit_WorldCollision.cpp b/Engine/Plugins/Animation/ControlRig/Source/ControlRig/Private/Units/Collision/RigUnit_WorldCollision.cpp
--- a/Engine/Plugins/Animation/ControlRig/Source/ControlRig/Private/Units/Collision/RigUnit_WorldCollision.cpp
+++ b/Engine/Plugins/Animation/ControlRig/Source/ControlRig/Private/Units/Collision/RigUnit_WorldCollision.cpp
@@ -16,7 +16,7 @@ FRigUnit_SphereTraceWorld_Execute()
 	HitLocation = FVector::ZeroVector;
 	HitNormal = FVector(0.f, 0.f, 1.f);
 
-	if(Context.World == nullptr)
+	if(ExecuteContext.UnitContext.World == nullptr)
 	{
 		return;
 	}
@@ -24,11 +24,11 @@ FRigUnit_SphereTraceWorld_Execute()
 	FCollisionQueryParams QueryParams;
 	QueryParams.bTraceComplex = true;
 
-	if (Context.OwningActor)
+	if (ExecuteContext.UnitContext.OwningActor)
 	{
-		QueryParams.AddIgnoredActor(Context.OwningActor);
+		QueryParams.AddIgnoredActor(ExecuteContext.UnitContext.OwningActor);
 	}
-	else if (UPrimitiveComponent* PrimitiveComponent = Cast<UPrimitiveComponent>(Context.OwningComponent))
+	else if (UPrimitiveComponent* PrimitiveComponent = Cast<UPrimitiveComponent>(ExecuteContext.UnitContext.OwningComponent))
 	{
 		QueryParams.AddIgnoredComponent(PrimitiveComponent);
 	}
@@ -39,13 +39,13 @@ FRigUnit_SphereTraceWorld_Execute()
 	const FCollisionShape CollisionShape = FCollisionShape::MakeSphere(Radius);
 
 	FHitResult HitResult;
-	bHit = Context.World->SweepSingleByChannel(HitResult, Context.ToWorldSpace(Start), Context.ToWorldSpace(End), 
+	bHit = ExecuteContext.UnitContext.World->SweepSingleByChannel(HitResult, ExecuteContext.UnitContext.ToWorldSpace(Start), ExecuteContext.UnitContext.ToWorldSpace(End), 
 			FQuat::Identity, Channel, CollisionShape, QueryParams, ResponseParams);
 
 	if (bHit)
 	{
-		HitLocation = Context.ToRigSpace(HitResult.ImpactPoint);
-		HitNormal = Context.ToWorldSpaceTransform.InverseTransformVector(HitResult.ImpactNormal);
+		HitLocation = ExecuteContext.UnitContext.ToRigSpace(HitResult.ImpactPoint);
+		HitNormal = ExecuteContext.UnitContext.ToWorldSpaceTransform.InverseTransformVector(HitResult.ImpactNormal);
 	}
 }
 
@@ -67,7 +67,7 @@ FRigUnit_SphereTraceByTraceChannel_Execute()
 	HitLocation = FVector::ZeroVector;
 	HitNormal = FVector(0.f, 0.f, 1.f);
 
-	if (Context.World == nullptr)
+	if (ExecuteContext.UnitContext.World == nullptr)
 	{
 		return;
 	}
@@ -78,22 +78,22 @@ FRigUnit_SphereTraceByTraceChannel_Execute()
 	FCollisionQueryParams QueryParams;
 	QueryParams.bTraceComplex = true;
 
-	if (Context.OwningActor)
+	if (ExecuteContext.UnitContext.OwningActor)
 	{
-		QueryParams.AddIgnoredActor(Context.OwningActor);
+		QueryParams.AddIgnoredActor(ExecuteContext.UnitContext.OwningActor);
 	}
-	else if (UPrimitiveComponent* PrimitiveComponent = Cast<UPrimitiveComponent>(Context.OwningComponent))
+	else if (UPrimitiveComponent* PrimitiveComponent = Cast<UPrimitiveComponent>(ExecuteContext.UnitContext.OwningComponent))
 	{
 		QueryParams.AddIgnoredComponent(PrimitiveComponent);
 	}
 
 	FHitResult HitResult;
-	bHit = Context.World->SweepSingleByChannel(HitResult, Context.ToWorldSpace(Start), Context.ToWorldSpace(End), FQuat::Identity, CollisionChannel, CollisionShape, QueryParams);
+	bHit = ExecuteContext.UnitContext.World->SweepSingleByChannel(HitResult, ExecuteContext.UnitContext.ToWorldSpace(Start), ExecuteContext.UnitContext.ToWorldSpace(End), FQuat::Identity, CollisionChannel, CollisionShape, QueryParams);
 	
 	if (bHit)
 	{
-		HitLocation = Context.ToRigSpace(HitResult.ImpactPoint);
-		HitNormal = Context.ToWorldSpaceTransform.InverseTransformVector(HitResult.ImpactNormal);
+		HitLocation = ExecuteContext.UnitContext.ToRigSpace(HitResult.ImpactPoint);
+		HitNormal = ExecuteContext.UnitContext.ToWorldSpaceTransform.InverseTransformVector(HitResult.ImpactNormal);
 	}
 }
 
@@ -105,7 +105,7 @@ FRigUnit_SphereTraceByObjectTypes_Execute()
 	HitLocation = FVector::ZeroVector;
 	HitNormal = FVector(0.f, 0.f, 1.f);
 
-	if (Context.World == nullptr)
+	if (ExecuteContext.UnitContext.World == nullptr)
 	{
 		return;
 	}
@@ -133,22 +133,22 @@ FRigUnit_SphereTraceByObjectTypes_Execute()
 	FCollisionQueryParams QueryParams;
 	QueryParams.bTraceComplex = true;
 
-	if (Context.OwningActor)
+	if (ExecuteContext.UnitContext.OwningActor)
 	{
-		QueryParams.AddIgnoredActor(Context.OwningActor);
+		QueryParams.AddIgnoredActor(ExecuteContext.UnitContext.OwningActor);
 	}
-	else if (UPrimitiveComponent* PrimitiveComponent = Cast<UPrimitiveComponent>(Context.OwningComponent))
+	else if (UPrimitiveComponent* PrimitiveComponent = Cast<UPrimitiveComponent>(ExecuteContext.UnitContext.OwningComponent))
 	{
 		QueryParams.AddIgnoredComponent(PrimitiveComponent);
 	} 
 
 	FHitResult HitResult;
-	bHit = Context.World->SweepSingleByObjectType(HitResult, Context.ToWorldSpace(Start), Context.ToWorldSpace(End), FQuat::Identity, ObjectParams, CollisionShape, QueryParams);
+	bHit = ExecuteContext.UnitContext.World->SweepSingleByObjectType(HitResult, ExecuteContext.UnitContext.ToWorldSpace(Start), ExecuteContext.UnitContext.ToWorldSpace(End), FQuat::Identity, ObjectParams, CollisionShape, QueryParams);
 
 	if (bHit)
 	{
-		HitLocation = Context.ToRigSpace(HitResult.ImpactPoint);
-		HitNormal = Context.ToWorldSpaceTransform.InverseTransformVector(HitResult.ImpactNormal);
+		HitLocation = ExecuteContext.UnitContext.ToRigSpace(HitResult.ImpactPoint);
+		HitNormal = ExecuteContext.UnitContext.ToWorldSpaceTransform.InverseTransformVector(HitResult.ImpactNormal);
 	}
 }
 
diff --git a/Engine/Plugins/Animation/ControlRig/Source/ControlRig/Private/Units/Collision/RigUnit_WorldCollision.h b/Engine/Plugins/Animation/ControlRig/Source/ControlRig/Private/Units/Collision/RigUnit_WorldCollision.h
--- a/Engine/Plugins/Animation/ControlRig/Source/ControlRig/Private/Units/Collision/RigUnit_WorldCollision.h
+++ b/Engine/Plugins/Animation/ControlRig/Source/ControlRig/Private/Units/Collision/RigUnit_WorldCollision.h
@@ -27,7 +27,7 @@ struct CONTROLRIG_API FRigUnit_SphereTraceWorld : public FRigUnit
 	}
 
 	RIGVM_METHOD()
-	virtual void Execute(const FRigUnitContext& Context) override;
+	virtual void Execute() override;
 
 	/** Start of the trace in rig / global space */
 	UPROPERTY(meta = (Input))
@@ -83,7 +83,7 @@ struct CONTROLRIG_API FRigUnit_SphereTraceByTraceChannel : public FRigUnit
 	}
 
 	RIGVM_METHOD()
-	virtual void Execute(const FRigUnitContext& Context) override;
+	virtual void Execute() override;
 
 	/** Start of the trace in rig / global space */
 	UPROPERTY(meta = (Input))
@@ -136,7 +136,7 @@ struct CONTROLRIG_API FRigUnit_SphereTraceByObjectTypes : public FRigUnit
 	}
 
 	RIGVM_METHOD()
-	virtual void Execute(const FRigUnitContext& Context) override;
+	virtual void Execute() override;
 
 	/** Start of the trace in rig / global space */
 	UPROPERTY(meta = (Input))
diff --git a/Engine/Plugins/Animation/ControlRig/Source/ControlRig/Private/Units/Control/RigUnit_Control.cpp b/Engine/Plugins/Animation/ControlRig/Source/ControlRig/Private/Units/Control/RigUnit_Control.cpp
--- a/Engine/Plugins/Animation/ControlRig/Source/ControlRig/Private/Units/Control/RigUnit_Control.cpp
+++ b/Engine/Plugins/Animation/ControlRig/Source/ControlRig/Private/Units/Control/RigUnit_Control.cpp
@@ -8,7 +8,7 @@
 FRigUnit_Control_Execute()
 {
     DECLARE_SCOPE_HIERARCHICAL_COUNTER_RIGUNIT()
-	if (Context.State == EControlRigState::Init)
+	if (ExecuteContext.UnitContext.State == EControlRigState::Init)
 	{
 		Transform.FromFTransform(InitTransform);
 	}
diff --git a/Engine/Plugins/Animation/ControlRig/Source/ControlRig/Private/Units/Control/RigUnit_Control_StaticMesh.cpp b/Engine/Plugins/Animation/ControlRig/Source/ControlRig/Private/Units/Control/RigUnit_Control_StaticMesh.cpp
--- a/Engine/Plugins/Animation/ControlRig/Source/ControlRig/Private/Units/Control/RigUnit_Control_StaticMesh.cpp
+++ b/Engine/Plugins/Animation/ControlRig/Source/ControlRig/Private/Units/Control/RigUnit_Control_StaticMesh.cpp
@@ -10,7 +10,7 @@ FRigUnit_Control_StaticMesh::FRigUnit_Control_StaticMesh()
 
 FRigUnit_Control_StaticMesh_Execute()
 {
-	FRigUnit_Control::StaticExecute(ExecuteContext, Transform, Base, InitTransform, Result, Filter, Context);
+	FRigUnit_Control::StaticExecute(ExecuteContext, Transform, Base, InitTransform, Result, Filter);
 }
 
 FRigVMStructUpgradeInfo FRigUnit_Control_StaticMesh::GetUpgradeInfo() const
diff --git a/Engine/Plugins/Animation/ControlRig/Source/ControlRig/Private/Units/Debug/RigUnit_DebugBezier.cpp b/Engine/Plugins/Animation/ControlRig/Source/ControlRig/Private/Units/Debug/RigUnit_DebugBezier.cpp
--- a/Engine/Plugins/Animation/ControlRig/Source/ControlRig/Private/Units/Debug/RigUnit_DebugBezier.cpp
+++ b/Engine/Plugins/Animation/ControlRig/Source/ControlRig/Private/Units/Debug/RigUnit_DebugBezier.cpp
@@ -17,8 +17,7 @@ FRigUnit_DebugBezier_Execute()
 		Detail, 
 		FRigElementKey(Space, ERigElementType::Bone), 
 		WorldOffset, 
-		bEnabled,
-		Context);
+		bEnabled);
 }
 
 FRigVMStructUpgradeInfo FRigUnit_DebugBezier::GetUpgradeInfo() const
@@ -30,12 +29,12 @@ FRigVMStructUpgradeInfo FRigUnit_DebugBezier::GetUpgradeInfo() const
 FRigUnit_DebugBezierItemSpace_Execute()
 {
     DECLARE_SCOPE_HIERARCHICAL_COUNTER_RIGUNIT()
-	if (Context.State == EControlRigState::Init)
+	if (ExecuteContext.UnitContext.State == EControlRigState::Init)
 	{
 		return;
 	}
 
-	if (Context.DrawInterface == nullptr || !bEnabled)
+	if (ExecuteContext.UnitContext.DrawInterface == nullptr || !bEnabled)
 	{
 		return;
 	}
@@ -43,10 +42,10 @@ FRigUnit_DebugBezierItemSpace_Execute()
 	FTransform Transform = WorldOffset;
 	if (Space.IsValid())
 	{
-		Transform = Transform * Context.Hierarchy->GetGlobalTransform(Space);
+		Transform = Transform * ExecuteContext.Hierarchy->GetGlobalTransform(Space);
 	}
 
-	Context.DrawInterface->DrawBezier(Transform, Bezier, MinimumU, MaximumU, Color, Thickness, Detail);
+	ExecuteContext.UnitContext.DrawInterface->DrawBezier(Transform, Bezier, MinimumU, MaximumU, Color, Thickness, Detail);
 }
 
 FRigVMStructUpgradeInfo FRigUnit_DebugBezierItemSpace::GetUpgradeInfo() const
diff --git a/Engine/Plugins/Animation/ControlRig/Source/ControlRig/Private/Units/Debug/RigUnit_DebugHierarchy.cpp b/Engine/Plugins/Animation/ControlRig/Source/ControlRig/Private/Units/Debug/RigUnit_DebugHierarchy.cpp
--- a/Engine/Plugins/Animation/ControlRig/Source/ControlRig/Private/Units/Debug/RigUnit_DebugHierarchy.cpp
+++ b/Engine/Plugins/Animation/ControlRig/Source/ControlRig/Private/Units/Debug/RigUnit_DebugHierarchy.cpp
@@ -8,40 +8,40 @@
 FRigUnit_DebugHierarchy_Execute()
 {
     DECLARE_SCOPE_HIERARCHICAL_COUNTER_RIGUNIT()
-	if (Context.State == EControlRigState::Init)
+	if (ExecuteContext.UnitContext.State == EControlRigState::Init)
 	{
 		return;
 	}
 
-	if (Context.DrawInterface == nullptr || !bEnabled)
+	if (ExecuteContext.UnitContext.DrawInterface == nullptr || !bEnabled)
 	{
 		return;
 	}
 
-	URigHierarchy* Hierarchy = Context.Hierarchy;
+	URigHierarchy* Hierarchy = ExecuteContext.Hierarchy;
 	if (Hierarchy)
 	{
-		Context.DrawInterface->DrawHierarchy(WorldOffset, Hierarchy, EControlRigDrawHierarchyMode::Axes, Scale, Color, Thickness);
+		ExecuteContext.UnitContext.DrawInterface->DrawHierarchy(WorldOffset, Hierarchy, EControlRigDrawHierarchyMode::Axes, Scale, Color, Thickness);
 	}
 }
 
 FRigUnit_DebugPose_Execute()
 {
 	DECLARE_SCOPE_HIERARCHICAL_COUNTER_RIGUNIT()
-	if (Context.State == EControlRigState::Init)
+	if (ExecuteContext.UnitContext.State == EControlRigState::Init)
 	{
 		return;
 	}
 
-	if (Context.DrawInterface == nullptr || !bEnabled)
+	if (ExecuteContext.UnitContext.DrawInterface == nullptr || !bEnabled)
 	{
 		return;
 	}
 
-	URigHierarchy* Hierarchy = Context.Hierarchy;
+	URigHierarchy* Hierarchy = ExecuteContext.Hierarchy;
 	if (Hierarchy)
 	{
-		Context.DrawInterface->DrawHierarchy(WorldOffset, Hierarchy, EControlRigDrawHierarchyMode::Axes, Scale, Color, Thickness, &Pose);
+		ExecuteContext.UnitContext.DrawInterface->DrawHierarchy(WorldOffset, Hierarchy, EControlRigDrawHierarchyMode::Axes, Scale, Color, Thickness, &Pose);
 	}
 }
 
diff --git a/Engine/Plugins/Animation/ControlRig/Source/ControlRig/Private/Units/Debug/RigUnit_DebugLine.cpp b/Engine/Plugins/Animation/ControlRig/Source/ControlRig/Private/Units/Debug/RigUnit_DebugLine.cpp
--- a/Engine/Plugins/Animation/ControlRig/Source/ControlRig/Private/Units/Debug/RigUnit_DebugLine.cpp
+++ b/Engine/Plugins/Animation/ControlRig/Source/ControlRig/Private/Units/Debug/RigUnit_DebugLine.cpp
@@ -15,8 +15,7 @@ FRigUnit_DebugLine_Execute()
 		Thickness,
 		FRigElementKey(Space, ERigElementType::Bone), 
 		WorldOffset, 
-		bEnabled,
-		Context);
+		bEnabled);
 }
 
 FRigVMStructUpgradeInfo FRigUnit_DebugLine::GetUpgradeInfo() const
@@ -38,12 +37,12 @@ FRigVMStructUpgradeInfo FRigUnit_DebugLine::GetUpgradeInfo() const
 FRigUnit_DebugLineItemSpace_Execute()
 {
     DECLARE_SCOPE_HIERARCHICAL_COUNTER_RIGUNIT()
-	if (Context.State == EControlRigState::Init)
+	if (ExecuteContext.UnitContext.State == EControlRigState::Init)
 	{
 		return;
 	}
 
-	if (Context.DrawInterface == nullptr || !bEnabled)
+	if (ExecuteContext.UnitContext.DrawInterface == nullptr || !bEnabled)
 	{
 		return;
 	}
@@ -51,10 +50,10 @@ FRigUnit_DebugLineItemSpace_Execute()
 	FVector DrawA = A, DrawB = B;
 	if (Space.IsValid())
 	{
-		FTransform Transform = Context.Hierarchy->GetGlobalTransform(Space);
+		FTransform Transform = ExecuteContext.Hierarchy->GetGlobalTransform(Space);
 		DrawA = Transform.TransformPosition(DrawA);
 		DrawB = Transform.TransformPosition(DrawB);
 	}
 
-	Context.DrawInterface->DrawLine(WorldOffset, DrawA, DrawB, Color, Thickness);
+	ExecuteContext.UnitContext.DrawInterface->DrawLine(WorldOffset, DrawA, DrawB, Color, Thickness);
 }
diff --git a/Engine/Plugins/Animation/ControlRig/Source/ControlRig/Private/Units/Debug/RigUnit_DebugLineStrip.cpp b/Engine/Plugins/Animation/ControlRig/Source/ControlRig/Private/Units/Debug/RigUnit_DebugLineStrip.cpp
--- a/Engine/Plugins/Animation/ControlRig/Source/ControlRig/Private/Units/Debug/RigUnit_DebugLineStrip.cpp
+++ b/Engine/Plugins/Animation/ControlRig/Source/ControlRig/Private/Units/Debug/RigUnit_DebugLineStrip.cpp
@@ -14,8 +14,7 @@ FRigUnit_DebugLineStrip_Execute()
 		Thickness,
 		FRigElementKey(Space, ERigElementType::Bone), 
 		WorldOffset, 
-		bEnabled,
-		Context);
+		bEnabled);
 }
 
 FRigVMStructUpgradeInfo FRigUnit_DebugLineStrip::GetUpgradeInfo() const
@@ -36,30 +35,30 @@ FRigVMStructUpgradeInfo FRigUnit_DebugLineStrip::GetUpgradeInfo() const
 FRigUnit_DebugLineStripItemSpace_Execute()
 {
     DECLARE_SCOPE_HIERARCHICAL_COUNTER_RIGUNIT()
-	if (Context.State == EControlRigState::Init)
+	if (ExecuteContext.UnitContext.State == EControlRigState::Init)
 	{
 		return;
 	}
 
-	if (Context.DrawInterface == nullptr || !bEnabled)
+	if (ExecuteContext.UnitContext.DrawInterface == nullptr || !bEnabled)
 	{
 		return;
 	}
 
 	if (Space.IsValid())
 	{
-		FTransform Transform = Context.Hierarchy->GetGlobalTransform(Space);
+		FTransform Transform = ExecuteContext.Hierarchy->GetGlobalTransform(Space);
 		TArray<FVector> PointsTransformed;
 		PointsTransformed.Reserve(Points.Num());
 		for(const FVector& Point : Points)
 		{
 			PointsTransformed.Add(Transform.TransformPosition(Point));
 		}
-		Context.DrawInterface->DrawLineStrip(WorldOffset, PointsTransformed, Color, Thickness);
+		ExecuteContext.UnitContext.DrawInterface->DrawLineStrip(WorldOffset, PointsTransformed, Color, Thickness);
 	}
 	else
 	{
-		Context.DrawInterface->DrawLineStrip(WorldOffset, TArrayView<const FVector>(Points.GetData(), Points.Num()), Color, Thickness);
+		ExecuteContext.UnitContext.DrawInterface->DrawLineStrip(WorldOffset, TArrayView<const FVector>(Points.GetData(), Points.Num()), Color, Thickness);
 	}
 }
 
diff --git a/Engine/Plugins/Animation/ControlRig/Source/ControlRig/Private/Units/Debug/RigUnit_DebugPoint.cpp b/Engine/Plugins/Animation/ControlRig/Source/ControlRig/Private/Units/Debug/RigUnit_DebugPoint.cpp
--- a/Engine/Plugins/Animation/ControlRig/Source/ControlRig/Private/Units/Debug/RigUnit_DebugPoint.cpp
+++ b/Engine/Plugins/Animation/ControlRig/Source/ControlRig/Private/Units/Debug/RigUnit_DebugPoint.cpp
@@ -9,21 +9,21 @@
 FRigUnit_DebugPoint_Execute()
 {
     DECLARE_SCOPE_HIERARCHICAL_COUNTER_RIGUNIT()
-	if (Context.State == EControlRigState::Init)
+	if (ExecuteContext.UnitContext.State == EControlRigState::Init)
 	{
 		return;
 	}
 
-	if (Context.DrawInterface == nullptr || !bEnabled)
+	if (ExecuteContext.UnitContext.DrawInterface == nullptr || !bEnabled)
 	{
 		return;
 	}
 
 	FVector Center = FVector::ZeroVector;
 	FVector DrawVector = Vector;
-	if (Space != NAME_None && Context.Hierarchy != nullptr)
+	if (Space != NAME_None && ExecuteContext.Hierarchy != nullptr)
 	{
-		const FTransform Transform = Context.Hierarchy->GetGlobalTransform(FRigElementKey(Space, ERigElementType::Bone));
+		const FTransform Transform = ExecuteContext.Hierarchy->GetGlobalTransform(FRigElementKey(Space, ERigElementType::Bone));
 		Center = Transform.GetLocation();
 		DrawVector = Transform.TransformPosition(DrawVector);
 	}
@@ -32,12 +32,12 @@ FRigUnit_DebugPoint_Execute()
 	{
 		case ERigUnitDebugPointMode::Point:
 		{
-			Context.DrawInterface->DrawPoint(WorldOffset, DrawVector, Scale, Color);
+			ExecuteContext.UnitContext.DrawInterface->DrawPoint(WorldOffset, DrawVector, Scale, Color);
 			break;
 		}
 		case ERigUnitDebugPointMode::Vector:
 		{
-			Context.DrawInterface->DrawLine(WorldOffset, Center, DrawVector, Color, Thickness);
+			ExecuteContext.UnitContext.DrawInterface->DrawLine(WorldOffset, Center, DrawVector, Color, Thickness);
 			break;
 		}
 	}
@@ -62,21 +62,21 @@ FRigVMStructUpgradeInfo FRigUnit_DebugPoint::GetUpgradeInfo() const
 FRigUnit_DebugPointMutable_Execute()
 {
     DECLARE_SCOPE_HIERARCHICAL_COUNTER_RIGUNIT()
-	if (Context.State == EControlRigState::Init)
+	if (ExecuteContext.UnitContext.State == EControlRigState::Init)
 	{
 		return;
 	}
 
-	if (Context.DrawInterface == nullptr || !bEnabled)
+	if (ExecuteContext.UnitContext.DrawInterface == nullptr || !bEnabled)
 	{
 		return;
 	}
 
 	FVector Center = FVector::ZeroVector;
 	FVector DrawVector = Vector;
-	if (Space != NAME_None && Context.Hierarchy != nullptr)
+	if (Space != NAME_None && ExecuteContext.Hierarchy != nullptr)
 	{
-		const FTransform Transform = Context.Hierarchy->GetGlobalTransform(FRigElementKey(Space, ERigElementType::Bone));
+		const FTransform Transform = ExecuteContext.Hierarchy->GetGlobalTransform(FRigElementKey(Space, ERigElementType::Bone));
 		Center = Transform.GetLocation();
 		DrawVector = Transform.TransformPosition(DrawVector);
 	}
@@ -85,12 +85,12 @@ FRigUnit_DebugPointMutable_Execute()
 	{
 		case ERigUnitDebugPointMode::Point:
 		{
-			Context.DrawInterface->DrawPoint(WorldOffset, DrawVector, Scale, Color);
+			ExecuteContext.UnitContext.DrawInterface->DrawPoint(WorldOffset, DrawVector, Scale, Color);
 			break;
 		}
 		case ERigUnitDebugPointMode::Vector:
 		{
-			Context.DrawInterface->DrawLine(WorldOffset, Center, DrawVector, Color, Thickness);
+			ExecuteContext.UnitContext.DrawInterface->DrawLine(WorldOffset, Center, DrawVector, Color, Thickness);
 			break;
 		}
 	}
diff --git a/Engine/Plugins/Animation/ControlRig/Source/ControlRig/Private/Units/Debug/RigUnit_DebugPrimitives.cpp b/Engine/Plugins/Animation/ControlRig/Source/ControlRig/Private/Units/Debug/RigUnit_DebugPrimitives.cpp
--- a/Engine/Plugins/Animation/ControlRig/Source/ControlRig/Private/Units/Debug/RigUnit_DebugPrimitives.cpp
+++ b/Engine/Plugins/Animation/ControlRig/Source/ControlRig/Private/Units/Debug/RigUnit_DebugPrimitives.cpp
@@ -15,8 +15,7 @@ FRigUnit_DebugRectangle_Execute()
 		Thickness,
 		FRigElementKey(Space, ERigElementType::Bone), 
 		WorldOffset, 
-		bEnabled,
-		Context);
+		bEnabled);
 }
 
 FRigVMStructUpgradeInfo FRigUnit_DebugRectangle::GetUpgradeInfo() const
@@ -38,12 +37,12 @@ FRigVMStructUpgradeInfo FRigUnit_DebugRectangle::GetUpgradeInfo() const
 FRigUnit_DebugRectangleItemSpace_Execute()
 {
     DECLARE_SCOPE_HIERARCHICAL_COUNTER_RIGUNIT()
-	if (Context.State == EControlRigState::Init)
+	if (ExecuteContext.UnitContext.State == EControlRigState::Init)
 	{
 		return;
 	}
 
-	if (Context.DrawInterface == nullptr || !bEnabled)
+	if (ExecuteContext.UnitContext.DrawInterface == nullptr || !bEnabled)
 	{
 		return;
 	}
@@ -51,10 +50,10 @@ FRigUnit_DebugRectangleItemSpace_Execute()
 	FTransform DrawTransform = Transform;
 	if (Space.IsValid())
 	{
-		DrawTransform = DrawTransform * Context.Hierarchy->GetGlobalTransform(Space);
+		DrawTransform = DrawTransform * ExecuteContext.Hierarchy->GetGlobalTransform(Space);
 	}
 
-	Context.DrawInterface->DrawRectangle(WorldOffset, DrawTransform, Scale, Color, Thickness);
+	ExecuteContext.UnitContext.DrawInterface->DrawRectangle(WorldOffset, DrawTransform, Scale, Color, Thickness);
 }
 
 FRigUnit_DebugArc_Execute()
@@ -70,8 +69,7 @@ FRigUnit_DebugArc_Execute()
 		Detail,
 		FRigElementKey(Space, ERigElementType::Bone), 
 		WorldOffset, 
-		bEnabled,
-		Context);
+		bEnabled);
 }
 
 FRigVMStructUpgradeInfo FRigUnit_DebugArc::GetUpgradeInfo() const
@@ -96,12 +94,12 @@ FRigVMStructUpgradeInfo FRigUnit_DebugArc::GetUpgradeInfo() const
 FRigUnit_DebugArcItemSpace_Execute()
 {
     DECLARE_SCOPE_HIERARCHICAL_COUNTER_RIGUNIT()
-	if (Context.State == EControlRigState::Init)
+	if (ExecuteContext.UnitContext.State == EControlRigState::Init)
 	{
 		return;
 	}
 
-	if (Context.DrawInterface == nullptr || !bEnabled)
+	if (ExecuteContext.UnitContext.DrawInterface == nullptr || !bEnabled)
 	{
 		return;
 	}
@@ -109,8 +107,8 @@ FRigUnit_DebugArcItemSpace_Execute()
 	FTransform DrawTransform = Transform;
 	if (Space.IsValid())
 	{
-		DrawTransform = DrawTransform * Context.Hierarchy->GetGlobalTransform(Space);
+		DrawTransform = DrawTransform * ExecuteContext.Hierarchy->GetGlobalTransform(Space);
 	}
 
-	Context.DrawInterface->DrawArc(WorldOffset, DrawTransform, Radius, FMath::DegreesToRadians(MinimumDegrees), FMath::DegreesToRadians(MaximumDegrees), Color, Thickness, Detail);
+	ExecuteContext.UnitContext.DrawInterface->DrawArc(WorldOffset, DrawTransform, Radius, FMath::DegreesToRadians(MinimumDegrees), FMath::DegreesToRadians(MaximumDegrees), Color, Thickness, Detail);
 }
diff --git a/Engine/Plugins/Animation/ControlRig/Source/ControlRig/Private/Units/Debug/RigUnit_DebugPrimitives.h b/Engine/Plugins/Animation/ControlRig/Source/ControlRig/Private/Units/Debug/RigUnit_DebugPrimitives.h
--- a/Engine/Plugins/Animation/ControlRig/Source/ControlRig/Private/Units/Debug/RigUnit_DebugPrimitives.h
+++ b/Engine/Plugins/Animation/ControlRig/Source/ControlRig/Private/Units/Debug/RigUnit_DebugPrimitives.h
@@ -21,7 +21,7 @@ struct CONTROLRIG_API FRigUnit_DebugRectangle : public FRigUnit_DebugBaseMutable
 	}
 
 	RIGVM_METHOD()
-	virtual void Execute(const FRigUnitContext& Context) override;
+	virtual void Execute() override;
 
 	UPROPERTY(meta = (Input))
 	FTransform Transform;
@@ -68,7 +68,7 @@ struct CONTROLRIG_API FRigUnit_DebugRectangleItemSpace : public FRigUnit_DebugBa
 	}
 
 	RIGVM_METHOD()
-	virtual void Execute(const FRigUnitContext& Context) override;
+	virtual void Execute() override;
 
 	UPROPERTY(meta = (Input))
 	FTransform Transform;
@@ -111,7 +111,7 @@ struct CONTROLRIG_API FRigUnit_DebugArc : public FRigUnit_DebugBaseMutable
 	}
 
 	RIGVM_METHOD()
-	virtual void Execute(const FRigUnitContext& Context) override;
+	virtual void Execute() override;
 
 	UPROPERTY(meta = (Input))
 	FTransform Transform;
@@ -169,7 +169,7 @@ struct CONTROLRIG_API FRigUnit_DebugArcItemSpace : public FRigUnit_DebugBaseMuta
 	}
 
 	RIGVM_METHOD()
-	virtual void Execute(const FRigUnitContext& Context) override;
+	virtual void Execute() override;
 
 	UPROPERTY(meta = (Input))
 	FTransform Transform;
diff --git a/Engine/Plugins/Animation/ControlRig/Source/ControlRig/Private/Units/Debug/RigUnit_DebugTransform.cpp b/Engine/Plugins/Animation/ControlRig/Source/ControlRig/Private/Units/Debug/RigUnit_DebugTransform.cpp
--- a/Engine/Plugins/Animation/ControlRig/Source/ControlRig/Private/Units/Debug/RigUnit_DebugTransform.cpp
+++ b/Engine/Plugins/Animation/ControlRig/Source/ControlRig/Private/Units/Debug/RigUnit_DebugTransform.cpp
@@ -9,38 +9,38 @@
 FRigUnit_DebugTransform_Execute()
 {
     DECLARE_SCOPE_HIERARCHICAL_COUNTER_RIGUNIT()
-	if (Context.State == EControlRigState::Init)
+	if (ExecuteContext.UnitContext.State == EControlRigState::Init)
 	{
 		return;
 	}
 
-	if (Context.DrawInterface == nullptr || !bEnabled)
+	if (ExecuteContext.UnitContext.DrawInterface == nullptr || !bEnabled)
 	{
 		return;
 	}
 
 	FTransform DrawTransform = Transform;
-	if (Space != NAME_None && Context.Hierarchy != nullptr)
+	if (Space != NAME_None && ExecuteContext.Hierarchy != nullptr)
 	{
-		DrawTransform = Transform * Context.Hierarchy->GetGlobalTransform(FRigElementKey(Space, ERigElementType::Bone));
+		DrawTransform = Transform * ExecuteContext.Hierarchy->GetGlobalTransform(FRigElementKey(Space, ERigElementType::Bone));
 	}
 
 	switch (Mode)
 	{
 		case ERigUnitDebugTransformMode::Axes:
 		{
-			Context.DrawInterface->DrawAxes(WorldOffset, DrawTransform, Scale, Thickness);
+			ExecuteContext.UnitContext.DrawInterface->DrawAxes(WorldOffset, DrawTransform, Scale, Thickness);
 			break;
 		}
 		case ERigUnitDebugTransformMode::Point:
 		{
-			Context.DrawInterface->DrawPoint(WorldOffset, DrawTransform.GetLocation(), Scale, Color);
+			ExecuteContext.UnitContext.DrawInterface->DrawPoint(WorldOffset, DrawTransform.GetLocation(), Scale, Color);
 			break;
 		}
 		case ERigUnitDebugTransformMode::Box:
 		{
 			DrawTransform.SetScale3D(DrawTransform.GetScale3D() * Scale);
-			Context.DrawInterface->DrawBox(WorldOffset, DrawTransform, Color, Thickness);
+			ExecuteContext.UnitContext.DrawInterface->DrawBox(WorldOffset, DrawTransform, Color, Thickness);
 			break;
 		}
 	}
@@ -73,8 +73,7 @@ FRigUnit_DebugTransformMutable_Execute()
 		Scale,
 		FRigElementKey(Space, ERigElementType::Bone), 
 		WorldOffset, 
-		bEnabled,
-		Context);
+		bEnabled);
 }
 
 FRigVMStructUpgradeInfo FRigUnit_DebugTransformMutable::GetUpgradeInfo() const
@@ -97,38 +96,38 @@ FRigVMStructUpgradeInfo FRigUnit_DebugTransformMutable::GetUpgradeInfo() const
 FRigUnit_DebugTransformMutableItemSpace_Execute()
 {
     DECLARE_SCOPE_HIERARCHICAL_COUNTER_RIGUNIT()
-	if (Context.State == EControlRigState::Init)
+	if (ExecuteContext.UnitContext.State == EControlRigState::Init)
 	{
 		return;
 	}
 
-	if (Context.DrawInterface == nullptr || !bEnabled)
+	if (ExecuteContext.UnitContext.DrawInterface == nullptr || !bEnabled)
 	{
 		return;
 	}
 
 	FTransform DrawTransform = Transform;
-	if (Space.IsValid() && Context.Hierarchy != nullptr)
+	if (Space.IsValid() && ExecuteContext.Hierarchy != nullptr)
 	{
-		DrawTransform = Transform * Context.Hierarchy->GetGlobalTransform(Space);
+		DrawTransform = Transform * ExecuteContext.Hierarchy->GetGlobalTransform(Space);
 	}
 
 	switch (Mode)
 	{
 		case ERigUnitDebugTransformMode::Axes:
 		{
-			Context.DrawInterface->DrawAxes(WorldOffset, DrawTransform, Scale, Thickness);
+			ExecuteContext.UnitContext.DrawInterface->DrawAxes(WorldOffset, DrawTransform, Scale, Thickness);
 			break;
 		}
 		case ERigUnitDebugTransformMode::Point:
 		{
-			Context.DrawInterface->DrawPoint(WorldOffset, DrawTransform.GetLocation(), Scale, Color);
+			ExecuteContext.UnitContext.DrawInterface->DrawPoint(WorldOffset, DrawTransform.GetLocation(), Scale, Color);
 			break;
 		}
 		case ERigUnitDebugTransformMode::Box:
 		{
 			DrawTransform.SetScale3D(DrawTransform.GetScale3D() * Scale);
-			Context.DrawInterface->DrawBox(WorldOffset, DrawTransform, Color, Thickness);
+			ExecuteContext.UnitContext.DrawInterface->DrawBox(WorldOffset, DrawTransform, Color, Thickness);
 			break;
 		}
 	}
@@ -137,20 +136,20 @@ FRigUnit_DebugTransformMutableItemSpace_Execute()
 FRigUnit_DebugTransformArrayMutable_Execute()
 {
     DECLARE_SCOPE_HIERARCHICAL_COUNTER_RIGUNIT()
-	if (Context.State == EControlRigState::Init)
+	if (ExecuteContext.UnitContext.State == EControlRigState::Init)
 	{
 		return;
 	}
 
-	if (Context.DrawInterface == nullptr || !bEnabled)
+	if (ExecuteContext.UnitContext.DrawInterface == nullptr || !bEnabled)
 	{
 		return;
 	}
 
 	WorkData.DrawTransforms.SetNumUninitialized(Transforms.Num());
-	if (Space != NAME_None && Context.Hierarchy)
+	if (Space != NAME_None && ExecuteContext.Hierarchy)
 	{
-		FTransform SpaceTransform = Context.Hierarchy->GetGlobalTransform(FRigElementKey(Space, ERigElementType::Bone));
+		FTransform SpaceTransform = ExecuteContext.Hierarchy->GetGlobalTransform(FRigElementKey(Space, ERigElementType::Bone));
 		for(int32 Index=0;Index<Transforms.Num();Index++)
 		{
 			WorkData.DrawTransforms[Index] = Transforms[Index] * SpaceTransform;
@@ -170,18 +169,18 @@ FRigUnit_DebugTransformArrayMutable_Execute()
 		{
 			case ERigUnitDebugTransformMode::Axes:
 			{
-				Context.DrawInterface->DrawAxes(WorldOffset, DrawTransform, Scale, Thickness);
+				ExecuteContext.UnitContext.DrawInterface->DrawAxes(WorldOffset, DrawTransform, Scale, Thickness);
 				break;
 			}
 			case ERigUnitDebugTransformMode::Point:
 			{
-				Context.DrawInterface->DrawPoint(WorldOffset, DrawTransform.GetLocation(), Scale, Color);
+				ExecuteContext.UnitContext.DrawInterface->DrawPoint(WorldOffset, DrawTransform.GetLocation(), Scale, Color);
 				break;
 			}
 			case ERigUnitDebugTransformMode::Box:
 			{
 				DrawTransform.SetScale3D(DrawTransform.GetScale3D() * Scale);
-				Context.DrawInterface->DrawBox(WorldOffset, DrawTransform, Color, Thickness);
+				ExecuteContext.UnitContext.DrawInterface->DrawBox(WorldOffset, DrawTransform, Color, Thickness);
 				break;
 			}
 		}
@@ -208,12 +207,12 @@ FRigVMStructUpgradeInfo FRigUnit_DebugTransformArrayMutable::GetUpgradeInfo() co
 FRigUnit_DebugTransformArrayMutableItemSpace_Execute()
 {
 	DECLARE_SCOPE_HIERARCHICAL_COUNTER_RIGUNIT()
-	if (Context.State == EControlRigState::Init)
+	if (ExecuteContext.UnitContext.State == EControlRigState::Init)
 	{
 		return;
 	}
 
-	if (Context.DrawInterface == nullptr || !bEnabled || Transforms.Num() == 0)
+	if (ExecuteContext.UnitContext.DrawInterface == nullptr || !bEnabled || Transforms.Num() == 0)
 	{
 		return;
 	}
@@ -229,8 +228,7 @@ FRigUnit_DebugTransformArrayMutableItemSpace_Execute()
 			Scale,
 			Space, 
 			WorldOffset, 
-			bEnabled,
-			Context);
+			bEnabled);
 	}
 
 	if(ParentIndices.Num() == Transforms.Num())
@@ -239,7 +237,7 @@ FRigUnit_DebugTransformArrayMutableItemSpace_Execute()
 		{
 			if(Transforms.IsValidIndex(ParentIndices[Index]))
 			{
-				Context.DrawInterface->DrawLine(
+				ExecuteContext.UnitContext.DrawInterface->DrawLine(
 					WorldOffset,
 					Transforms[Index].GetTranslation(),
 					Transforms[ParentIndices[Index]].GetTranslation(),
diff --git a/Engine/Plugins/Animation/ControlRig/Source/ControlRig/Private/Units/Debug/RigUnit_DebugTransform.h b/Engine/Plugins/Animation/ControlRig/Source/ControlRig/Private/Units/Debug/RigUnit_DebugTransform.h
--- a/Engine/Plugins/Animation/ControlRig/Source/ControlRig/Private/Units/Debug/RigUnit_DebugTransform.h
+++ b/Engine/Plugins/Animation/ControlRig/Source/ControlRig/Private/Units/Debug/RigUnit_DebugTransform.h
@@ -37,7 +37,7 @@ struct CONTROLRIG_API FRigUnit_DebugTransform : public FRigUnit_DebugBase
 	}
 
 	RIGVM_METHOD()
-	virtual void Execute(const FRigUnitContext& Context) override;
+	virtual void Execute() override;
 
 	UPROPERTY(meta = (Input, Output))
 	FTransform Transform;
@@ -83,7 +83,7 @@ struct CONTROLRIG_API FRigUnit_DebugTransformMutable : public FRigUnit_DebugBase
 	}
 
 	RIGVM_METHOD()
-	virtual void Execute(const FRigUnitContext& Context) override;
+	virtual void Execute() override;
 
 	UPROPERTY(meta = (Input))
 	FTransform Transform;
@@ -133,7 +133,7 @@ struct CONTROLRIG_API FRigUnit_DebugTransformMutableItemSpace : public FRigUnit_
 	}
 
 	RIGVM_METHOD()
-	virtual void Execute(const FRigUnitContext& Context) override;
+	virtual void Execute() override;
 
 	UPROPERTY(meta = (Input))
 	FTransform Transform;
@@ -185,7 +185,7 @@ struct CONTROLRIG_API FRigUnit_DebugTransformArrayMutable : public FRigUnit_Debu
 	}
 
 	RIGVM_METHOD()
-	virtual void Execute(const FRigUnitContext& Context) override;
+	virtual void Execute() override;
 
 	UPROPERTY(meta = (Input))
 	TArray<FTransform> Transforms;
@@ -238,7 +238,7 @@ struct CONTROLRIG_API FRigUnit_DebugTransformArrayMutableItemSpace : public FRig
 	}
 
 	RIGVM_METHOD()
-	virtual void Execute(const FRigUnitContext& Context) override;
+	virtual void Execute() override;
 
 	UPROPERTY(meta = (Input))
 	TArray<FTransform> Transforms;
diff --git a/Engine/Plugins/Animation/ControlRig/Source/ControlRig/Private/Units/Debug/RigUnit_ProfilingBracket.cpp b/Engine/Plugins/Animation/ControlRig/Source/ControlRig/Private/Units/Debug/RigUnit_ProfilingBracket.cpp
--- a/Engine/Plugins/Animation/ControlRig/Source/ControlRig/Private/Units/Debug/RigUnit_ProfilingBracket.cpp
+++ b/Engine/Plugins/Animation/ControlRig/Source/ControlRig/Private/Units/Debug/RigUnit_ProfilingBracket.cpp
@@ -9,7 +9,7 @@
 FRigUnit_StartProfilingTimer_Execute()
 {
     DECLARE_SCOPE_HIERARCHICAL_COUNTER_RIGUNIT()
-	if (Context.State == EControlRigState::Update)
+	if (ExecuteContext.UnitContext.State == EControlRigState::Update)
 	{
 		UKismetAnimationLibrary::K2_StartProfilingTimer();
 	}
@@ -18,14 +18,14 @@ FRigUnit_StartProfilingTimer_Execute()
 FRigUnit_EndProfilingTimer_Execute()
 {
     DECLARE_SCOPE_HIERARCHICAL_COUNTER_RIGUNIT()
-	if (Context.State == EControlRigState::Init)
+	if (ExecuteContext.UnitContext.State == EControlRigState::Init)
 	{
 		AccumulatedTime = 0.f;
 		MeasurementsLeft = FMath::Max(NumberOfMeasurements, 1);
 		return;
 	}
 
-	if (Context.State == EControlRigState::Update)
+	if (ExecuteContext.UnitContext.State == EControlRigState::Update)
 	{
 		float Delta = UKismetAnimationLibrary::K2_EndProfilingTimer(false);
 		if(MeasurementsLeft > 0)
diff --git a/Engine/Plugins/Animation/ControlRig/Source/ControlRig/Private/Units/Debug/RigUnit_VisualDebug.cpp b/Engine/Plugins/Animation/ControlRig/Source/ControlRig/Private/Units/Debug/RigUnit_VisualDebug.cpp
--- a/Engine/Plugins/Animation/ControlRig/Source/ControlRig/Private/Units/Debug/RigUnit_VisualDebug.cpp
+++ b/Engine/Plugins/Animation/ControlRig/Source/ControlRig/Private/Units/Debug/RigUnit_VisualDebug.cpp
@@ -7,7 +7,7 @@
 
 FRigUnit_VisualDebugVector_Execute()
 {
-	FRigUnit_VisualDebugVectorItemSpace::StaticExecute(ExecuteContext, Value, bEnabled, Mode, Color, Thickness, Scale, FRigElementKey(BoneSpace, ERigElementType::Bone), Context);
+	FRigUnit_VisualDebugVectorItemSpace::StaticExecute(ExecuteContext, Value, bEnabled, Mode, Color, Thickness, Scale, FRigElementKey(BoneSpace, ERigElementType::Bone));
 }
 
 FRigVMStructUpgradeInfo FRigUnit_VisualDebugVector::GetUpgradeInfo() const
@@ -28,12 +28,12 @@ FRigVMStructUpgradeInfo FRigUnit_VisualDebugVector::GetUpgradeInfo() const
 FRigUnit_VisualDebugVectorItemSpace_Execute()
 {
     DECLARE_SCOPE_HIERARCHICAL_COUNTER_RIGUNIT()
-	if (Context.State == EControlRigState::Init)
+	if (ExecuteContext.UnitContext.State == EControlRigState::Init)
 	{
 		return;
 	}
 
-	if (Context.DrawInterface == nullptr || !bEnabled)
+	if (ExecuteContext.UnitContext.DrawInterface == nullptr || !bEnabled)
 	{
 		return;
 	}
@@ -41,19 +41,19 @@ FRigUnit_VisualDebugVectorItemSpace_Execute()
 	FTransform WorldOffset = FTransform::Identity;
 	if (Space.IsValid())
 	{
-		WorldOffset = Context.Hierarchy->GetGlobalTransform(Space);
+		WorldOffset = ExecuteContext.Hierarchy->GetGlobalTransform(Space);
 	}
 
 	switch(Mode)
 	{
 		case ERigUnitVisualDebugPointMode::Point:
 		{
-			Context.DrawInterface->DrawPoint(WorldOffset, Value, Thickness, Color);
+			ExecuteContext.UnitContext.DrawInterface->DrawPoint(WorldOffset, Value, Thickness, Color);
 			break;
 		}
 		case ERigUnitVisualDebugPointMode::Vector:
 		{
-			Context.DrawInterface->DrawLine(WorldOffset, FVector::ZeroVector, Value * Scale, Color, Thickness);
+			ExecuteContext.UnitContext.DrawInterface->DrawLine(WorldOffset, FVector::ZeroVector, Value * Scale, Color, Thickness);
 			break;
 		}
 		default:
@@ -65,7 +65,7 @@ FRigUnit_VisualDebugVectorItemSpace_Execute()
 
 FRigUnit_VisualDebugQuat_Execute()
 {
-	FRigUnit_VisualDebugQuatItemSpace::StaticExecute(ExecuteContext, Value, bEnabled, Thickness, Scale, FRigElementKey(BoneSpace, ERigElementType::Bone), Context);
+	FRigUnit_VisualDebugQuatItemSpace::StaticExecute(ExecuteContext, Value, bEnabled, Thickness, Scale, FRigElementKey(BoneSpace, ERigElementType::Bone));
 }
 
 FRigVMStructUpgradeInfo FRigUnit_VisualDebugQuat::GetUpgradeInfo() const
@@ -89,12 +89,12 @@ FRigUnit_VisualDebugQuatItemSpace_Execute()
     FTransform Transform = FTransform::Identity;
     Transform.SetRotation(Value);
 
-	FRigUnit_VisualDebugTransformItemSpace::StaticExecute(ExecuteContext, Transform, bEnabled, Thickness, Scale, Space, Context);
+	FRigUnit_VisualDebugTransformItemSpace::StaticExecute(ExecuteContext, Transform, bEnabled, Thickness, Scale, Space);
 }
 
 FRigUnit_VisualDebugTransform_Execute()
 {
-	FRigUnit_VisualDebugTransformItemSpace::StaticExecute(ExecuteContext, Value, bEnabled, Thickness, Scale, FRigElementKey(BoneSpace, ERigElementType::Bone), Context);
+	FRigUnit_VisualDebugTransformItemSpace::StaticExecute(ExecuteContext, Value, bEnabled, Thickness, Scale, FRigElementKey(BoneSpace, ERigElementType::Bone));
 }
 
 FRigVMStructUpgradeInfo FRigUnit_VisualDebugTransform::GetUpgradeInfo() const
@@ -114,12 +114,12 @@ FRigVMStructUpgradeInfo FRigUnit_VisualDebugTransform::GetUpgradeInfo() const
 FRigUnit_VisualDebugTransformItemSpace_Execute()
 {
     DECLARE_SCOPE_HIERARCHICAL_COUNTER_RIGUNIT()
-	if (Context.State == EControlRigState::Init)
+	if (ExecuteContext.UnitContext.State == EControlRigState::Init)
 	{
 		return;
 	}
 
-	if (Context.DrawInterface == nullptr || !bEnabled)
+	if (ExecuteContext.UnitContext.DrawInterface == nullptr || !bEnabled)
 	{
 		return;
 	}
@@ -127,8 +127,8 @@ FRigUnit_VisualDebugTransformItemSpace_Execute()
 	FTransform WorldOffset = FTransform::Identity;
 	if (Space.IsValid())
 	{
-		WorldOffset = Context.Hierarchy->GetGlobalTransform(Space);
+		WorldOffset = ExecuteContext.Hierarchy->GetGlobalTransform(Space);
 	}
 
-	Context.DrawInterface->DrawAxes(WorldOffset, Value, Scale, Thickness);
+	ExecuteContext.UnitContext.DrawInterface->DrawAxes(WorldOffset, Value, Scale, Thickness);
 }
diff --git a/Engine/Plugins/Animation/ControlRig/Source/ControlRig/Private/Units/Debug/RigUnit_VisualDebug.h b/Engine/Plugins/Animation/ControlRig/Source/ControlRig/Private/Units/Debug/RigUnit_VisualDebug.h
--- a/Engine/Plugins/Animation/ControlRig/Source/ControlRig/Private/Units/Debug/RigUnit_VisualDebug.h
+++ b/Engine/Plugins/Animation/ControlRig/Source/ControlRig/Private/Units/Debug/RigUnit_VisualDebug.h
@@ -35,7 +35,7 @@ struct CONTROLRIG_API FRigUnit_VisualDebugVector : public FRigUnit_DebugBase
 	}
 
 	RIGVM_METHOD()
-	virtual void Execute(const FRigUnitContext& Context) override;
+	virtual void Execute() override;
 
 	UPROPERTY(meta = (Input, Output))
 	FVector Value;
@@ -87,7 +87,7 @@ struct CONTROLRIG_API FRigUnit_VisualDebugVectorItemSpace : public FRigUnit_Debu
 	}
 
 	RIGVM_METHOD()
-	virtual void Execute(const FRigUnitContext& Context) override;
+	virtual void Execute() override;
 
 	UPROPERTY(meta = (Input, Output))
 	FVector Value;
@@ -126,7 +126,7 @@ struct CONTROLRIG_API FRigUnit_VisualDebugQuat : public FRigUnit_DebugBase
 	}
 
 	RIGVM_METHOD()
-	virtual void Execute(const FRigUnitContext& Context) override;
+	virtual void Execute() override;
 
 	UPROPERTY(meta = (Input, Output))
 	FQuat Value;
@@ -170,7 +170,7 @@ struct CONTROLRIG_API FRigUnit_VisualDebugQuatItemSpace : public FRigUnit_DebugB
 	}
 
 	RIGVM_METHOD()
-	virtual void Execute(const FRigUnitContext& Context) override;
+	virtual void Execute() override;
 
 	UPROPERTY(meta = (Input, Output))
 	FQuat Value;
@@ -203,7 +203,7 @@ struct CONTROLRIG_API FRigUnit_VisualDebugTransform : public FRigUnit_DebugBase
 	}
 
 	RIGVM_METHOD()
-	virtual void Execute(const FRigUnitContext& Context) override;
+	virtual void Execute() override;
 
 	UPROPERTY(meta = (Input, Output))
 	FTransform Value;
@@ -247,7 +247,7 @@ struct CONTROLRIG_API FRigUnit_VisualDebugTransformItemSpace : public FRigUnit_D
 	}
 
 	RIGVM_METHOD()
-	virtual void Execute(const FRigUnitContext& Context) override;
+	virtual void Execute() override;
 
 	UPROPERTY(meta = (Input, Output))
 	FTransform Value;
diff --git a/Engine/Plugins/Animation/ControlRig/Source/ControlRig/Private/Units/Deprecated/Math/RigUnit_Converter.h b/Engine/Plugins/Animation/ControlRig/Source/ControlRig/Private/Units/Deprecated/Math/RigUnit_Converter.h
--- a/Engine/Plugins/Animation/ControlRig/Source/ControlRig/Private/Units/Deprecated/Math/RigUnit_Converter.h
+++ b/Engine/Plugins/Animation/ControlRig/Source/ControlRig/Private/Units/Deprecated/Math/RigUnit_Converter.h
@@ -18,7 +18,7 @@ struct CONTROLRIG_API FRigUnit_ConvertTransform : public FRigUnit
 	FEulerTransform Result;
 
 	RIGVM_METHOD()
-	virtual void Execute(const FRigUnitContext& Context) override;
+	virtual void Execute() override;
 
 	RIGVM_METHOD()
 	virtual FRigVMStructUpgradeInfo GetUpgradeInfo() const override;
@@ -36,7 +36,7 @@ struct CONTROLRIG_API FRigUnit_ConvertEulerTransform : public FRigUnit
 	FTransform Result;
 
 	RIGVM_METHOD()
-	virtual void Execute(const FRigUnitContext& Context) override;
+	virtual void Execute() override;
 
 	RIGVM_METHOD()
 	virtual FRigVMStructUpgradeInfo GetUpgradeInfo() const override;
@@ -54,7 +54,7 @@ struct CONTROLRIG_API FRigUnit_ConvertRotation : public FRigUnit
 	FQuat	Result = FQuat::Identity;
 
 	RIGVM_METHOD()
-	virtual void Execute(const FRigUnitContext& Context) override;
+	virtual void Execute() override;
 
 	RIGVM_METHOD()
 	virtual FRigVMStructUpgradeInfo GetUpgradeInfo() const override;
@@ -78,7 +78,7 @@ struct CONTROLRIG_API FRigUnit_ConvertQuaternion: public FRigUnit
 	FRotator	Result = FRotator(0.f);
 
 	RIGVM_METHOD()
-	virtual void Execute(const FRigUnitContext& Context) override;
+	virtual void Execute() override;
 
 	RIGVM_METHOD()
 	virtual FRigVMStructUpgradeInfo GetUpgradeInfo() const override;
@@ -96,7 +96,7 @@ struct CONTROLRIG_API FRigUnit_ConvertVectorToRotation: public FRigUnit
 	FRotator	Result = FRotator(0.f);
 
 	RIGVM_METHOD()
-	virtual void Execute(const FRigUnitContext& Context) override;
+	virtual void Execute() override;
 
 	RIGVM_METHOD()
 	virtual FRigVMStructUpgradeInfo GetUpgradeInfo() const override;
@@ -114,7 +114,7 @@ struct CONTROLRIG_API FRigUnit_ConvertVectorToQuaternion: public FRigUnit
 	FQuat	Result = FQuat::Identity;
 
 	RIGVM_METHOD()
-	virtual void Execute(const FRigUnitContext& Context) override;
+	virtual void Execute() override;
 
 	RIGVM_METHOD()
 	virtual FRigVMStructUpgradeInfo GetUpgradeInfo() const override;
@@ -132,7 +132,7 @@ struct CONTROLRIG_API FRigUnit_ConvertRotationToVector: public FRigUnit
 	FVector Result = FVector(0.f);
 
 	RIGVM_METHOD()
-	virtual void Execute(const FRigUnitContext& Context) override;
+	virtual void Execute() override;
 
 	RIGVM_METHOD()
 	virtual FRigVMStructUpgradeInfo GetUpgradeInfo() const override;
@@ -150,7 +150,7 @@ struct CONTROLRIG_API FRigUnit_ConvertQuaternionToVector: public FRigUnit
 	FVector Result = FVector(0.f);
 
 	RIGVM_METHOD()
-	virtual void Execute(const FRigUnitContext& Context) override;
+	virtual void Execute() override;
 
 	RIGVM_METHOD()
 	virtual FRigVMStructUpgradeInfo GetUpgradeInfo() const override;
@@ -182,7 +182,7 @@ struct CONTROLRIG_API FRigUnit_ToSwingAndTwist : public FRigUnit
 	}
 
 	RIGVM_METHOD()
-	virtual void Execute(const FRigUnitContext& Context) override;
+	virtual void Execute() override;
 
 	RIGVM_METHOD()
 	virtual FRigVMStructUpgradeInfo GetUpgradeInfo() const override;
diff --git a/Engine/Plugins/Animation/ControlRig/Source/ControlRig/Private/Units/Deprecated/Math/RigUnit_Float.h b/Engine/Plugins/Animation/ControlRig/Source/ControlRig/Private/Units/Deprecated/Math/RigUnit_Float.h
--- a/Engine/Plugins/Animation/ControlRig/Source/ControlRig/Private/Units/Deprecated/Math/RigUnit_Float.h
+++ b/Engine/Plugins/Animation/ControlRig/Source/ControlRig/Private/Units/Deprecated/Math/RigUnit_Float.h
@@ -28,7 +28,7 @@ struct CONTROLRIG_API FRigUnit_Multiply_FloatFloat : public FRigUnit_BinaryFloat
 	GENERATED_BODY()
 
 	RIGVM_METHOD()
-	virtual void Execute(const FRigUnitContext& Context) override;
+	virtual void Execute() override;
 
 	RIGVM_METHOD()
 	virtual FRigVMStructUpgradeInfo GetUpgradeInfo() const override;
@@ -40,7 +40,7 @@ struct CONTROLRIG_API FRigUnit_Add_FloatFloat : public FRigUnit_BinaryFloatOp
 	GENERATED_BODY()
 
 	RIGVM_METHOD()
-	virtual void Execute(const FRigUnitContext& Context) override;
+	virtual void Execute() override;
 
 	RIGVM_METHOD()
 	virtual FRigVMStructUpgradeInfo GetUpgradeInfo() const override;
@@ -52,7 +52,7 @@ struct CONTROLRIG_API FRigUnit_Subtract_FloatFloat : public FRigUnit_BinaryFloat
 	GENERATED_BODY()
 
 	RIGVM_METHOD()
-	virtual void Execute(const FRigUnitContext& Context) override;
+	virtual void Execute() override;
 
 	RIGVM_METHOD()
 	virtual FRigVMStructUpgradeInfo GetUpgradeInfo() const override;
@@ -64,7 +64,7 @@ struct CONTROLRIG_API FRigUnit_Divide_FloatFloat : public FRigUnit_BinaryFloatOp
 	GENERATED_BODY()
 
 	RIGVM_METHOD()
-	virtual void Execute(const FRigUnitContext& Context) override;
+	virtual void Execute() override;
 
 	RIGVM_METHOD()
 	virtual FRigVMStructUpgradeInfo GetUpgradeInfo() const override;
@@ -89,7 +89,7 @@ struct CONTROLRIG_API FRigUnit_Clamp_Float: public FRigUnit
 	float Result = 0.f;
 
 	RIGVM_METHOD()
-	virtual void Execute(const FRigUnitContext& Context) override;
+	virtual void Execute() override;
 
 	RIGVM_METHOD()
 	virtual FRigVMStructUpgradeInfo GetUpgradeInfo() const override;
@@ -120,7 +120,7 @@ struct CONTROLRIG_API FRigUnit_MapRange_Float: public FRigUnit
 	float Result = 0.f;
 
 	RIGVM_METHOD()
-	virtual void Execute(const FRigUnitContext& Context) override;
+	virtual void Execute() override;
 
 	RIGVM_METHOD()
 	virtual FRigVMStructUpgradeInfo GetUpgradeInfo() const override;
diff --git a/Engine/Plugins/Animation/ControlRig/Source/ControlRig/Private/Units/Deprecated/Math/RigUnit_Quaternion.h b/Engine/Plugins/Animation/ControlRig/Source/ControlRig/Private/Units/Deprecated/Math/RigUnit_Quaternion.h
--- a/Engine/Plugins/Animation/ControlRig/Source/ControlRig/Private/Units/Deprecated/Math/RigUnit_Quaternion.h
+++ b/Engine/Plugins/Animation/ControlRig/Source/ControlRig/Private/Units/Deprecated/Math/RigUnit_Quaternion.h
@@ -34,7 +34,7 @@ struct CONTROLRIG_API FRigUnit_MultiplyQuaternion : public FRigUnit_BinaryQuater
 	GENERATED_BODY()
 
 	RIGVM_METHOD()
-	virtual void Execute(const FRigUnitContext& Context) override;
+	virtual void Execute() override;
 
 	RIGVM_METHOD()
 	virtual FRigVMStructUpgradeInfo GetUpgradeInfo() const override;
@@ -59,7 +59,7 @@ struct CONTROLRIG_API FRigUnit_InverseQuaterion: public FRigUnit_UnaryQuaternion
 	GENERATED_BODY()
 
 	RIGVM_METHOD()
-	virtual void Execute(const FRigUnitContext& Context) override;
+	virtual void Execute() override;
 
 	RIGVM_METHOD()
 	virtual FRigVMStructUpgradeInfo GetUpgradeInfo() const override;
@@ -80,7 +80,7 @@ struct CONTROLRIG_API FRigUnit_QuaternionToAxisAndAngle : public FRigUnit
 	float Angle = 0.f;
 
 	RIGVM_METHOD()
-	virtual void Execute(const FRigUnitContext& Context) override;
+	virtual void Execute() override;
 
 	RIGVM_METHOD()
 	virtual FRigVMStructUpgradeInfo GetUpgradeInfo() const override;
@@ -107,7 +107,7 @@ struct CONTROLRIG_API FRigUnit_QuaternionFromAxisAndAngle : public FRigUnit
 	FQuat Result;
 
 	RIGVM_METHOD()
-	virtual void Execute(const FRigUnitContext& Context) override;
+	virtual void Execute() override;
 
 	RIGVM_METHOD()
 	virtual FRigVMStructUpgradeInfo GetUpgradeInfo() const override;
@@ -134,7 +134,7 @@ struct CONTROLRIG_API FRigUnit_QuaternionToAngle : public FRigUnit
 	float Angle;
 
 	RIGVM_METHOD()
-	virtual void Execute(const FRigUnitContext& Context) override;
+	virtual void Execute() override;
 	
 	RIGVM_METHOD()
 	virtual FRigVMStructUpgradeInfo GetUpgradeInfo() const override;
diff --git a/Engine/Plugins/Animation/ControlRig/Source/ControlRig/Private/Units/Deprecated/RigUnit_AimConstraint.cpp b/Engine/Plugins/Animation/ControlRig/Source/ControlRig/Private/Units/Deprecated/RigUnit_AimConstraint.cpp
--- a/Engine/Plugins/Animation/ControlRig/Source/ControlRig/Private/Units/Deprecated/RigUnit_AimConstraint.cpp
+++ b/Engine/Plugins/Animation/ControlRig/Source/ControlRig/Private/Units/Deprecated/RigUnit_AimConstraint.cpp
@@ -13,7 +13,7 @@ FRigUnit_AimConstraint_Execute()
 
 	TArray<FConstraintData>& ConstraintData = WorkData.ConstraintData;
 
-	if (Context.State == EControlRigState::Init)
+	if (ExecuteContext.UnitContext.State == EControlRigState::Init)
 	{
 		ConstraintData.Reset();
 
@@ -43,7 +43,7 @@ FRigUnit_AimConstraint_Execute()
 			}
 		}
 	}
-	else if (Context.State == EControlRigState::Update)
+	else if (ExecuteContext.UnitContext.State == EControlRigState::Update)
 	{
 		URigHierarchy* Hierarchy = ExecuteContext.Hierarchy;
 		if (Hierarchy)
diff --git a/Engine/Plugins/Animation/ControlRig/Source/ControlRig/Private/Units/Deprecated/RigUnit_ApplyFK.cpp b/Engine/Plugins/Animation/ControlRig/Source/ControlRig/Private/Units/Deprecated/RigUnit_ApplyFK.cpp
--- a/Engine/Plugins/Animation/ControlRig/Source/ControlRig/Private/Units/Deprecated/RigUnit_ApplyFK.cpp
+++ b/Engine/Plugins/Animation/ControlRig/Source/ControlRig/Private/Units/Deprecated/RigUnit_ApplyFK.cpp
@@ -10,11 +10,11 @@ FRigUnit_ApplyFK_Execute()
 {
     DECLARE_SCOPE_HIERARCHICAL_COUNTER_RIGUNIT()
 	
-	if (Context.State == EControlRigState::Init)
+	if (ExecuteContext.UnitContext.State == EControlRigState::Init)
 	{
 		return;
 	}
-	else if (Context.State == EControlRigState::Update)
+	else if (ExecuteContext.UnitContext.State == EControlRigState::Update)
 	{
 		URigHierarchy* Hierarchy = ExecuteContext.Hierarchy;
 		if (Hierarchy)
diff --git a/Engine/Plugins/Animation/ControlRig/Source/ControlRig/Private/Units/Deprecated/RigUnit_GetJointTransform.cpp b/Engine/Plugins/Animation/ControlRig/Source/ControlRig/Private/Units/Deprecated/RigUnit_GetJointTransform.cpp
--- a/Engine/Plugins/Animation/ControlRig/Source/ControlRig/Private/Units/Deprecated/RigUnit_GetJointTransform.cpp
+++ b/Engine/Plugins/Animation/ControlRig/Source/ControlRig/Private/Units/Deprecated/RigUnit_GetJointTransform.cpp
@@ -11,7 +11,7 @@ FRigUnit_GetJointTransform_Execute()
 {
     DECLARE_SCOPE_HIERARCHICAL_COUNTER_RIGUNIT()
     
-	if (URigHierarchy* Hierarchy = Context.Hierarchy)
+	if (URigHierarchy* Hierarchy = ExecuteContext.Hierarchy)
 	{
 		const FRigElementKey Key(Joint, ERigElementType::Bone);
 		
diff --git a/Engine/Plugins/Animation/ControlRig/Source/ControlRig/Private/Units/Deprecated/RigUnit_TwoBoneIKFK.cpp b/Engine/Plugins/Animation/ControlRig/Source/ControlRig/Private/Units/Deprecated/RigUnit_TwoBoneIKFK.cpp
--- a/Engine/Plugins/Animation/ControlRig/Source/ControlRig/Private/Units/Deprecated/RigUnit_TwoBoneIKFK.cpp
+++ b/Engine/Plugins/Animation/ControlRig/Source/ControlRig/Private/Units/Deprecated/RigUnit_TwoBoneIKFK.cpp
@@ -11,7 +11,7 @@ FRigUnit_TwoBoneIKFK_Execute()
 {
     DECLARE_SCOPE_HIERARCHICAL_COUNTER_RIGUNIT()
 
-	if (Context.State == EControlRigState::Init)
+	if (ExecuteContext.UnitContext.State == EControlRigState::Init)
 	{
 		URigHierarchy* Hierarchy = ExecuteContext.Hierarchy;
 		if (Hierarchy)
@@ -57,7 +57,7 @@ FRigUnit_TwoBoneIKFK_Execute()
 			}
 		}
 	}
-	else  if (Context.State == EControlRigState::Update)
+	else  if (ExecuteContext.UnitContext.State == EControlRigState::Update)
 	{
 		if (StartJointIndex != INDEX_NONE && MidJointIndex != INDEX_NONE && EndJointIndex != INDEX_NONE)
 		{
@@ -76,7 +76,7 @@ FRigUnit_TwoBoneIKFK_Execute()
 			else if (FMath::IsNearlyEqual(IKBlend, 1.f))
 			{
 				// update transform before going through IK
-				const URigHierarchy* Hierarchy = Context.Hierarchy;
+				const URigHierarchy* Hierarchy = ExecuteContext.Hierarchy;
 				check(Hierarchy);
 
 				StartJointIKTransform = Hierarchy->GetGlobalTransform(StartJointIndex);
@@ -93,7 +93,7 @@ FRigUnit_TwoBoneIKFK_Execute()
 			else
 			{
 				// update transform before going through IK
-				const URigHierarchy* Hierarchy = Context.Hierarchy;
+				const URigHierarchy* Hierarchy = ExecuteContext.Hierarchy;
 				check(Hierarchy);
 
 				StartJointIKTransform = Hierarchy->GetGlobalTransform(StartJointIndex);
diff --git a/Engine/Plugins/Animation/ControlRig/Source/ControlRig/Private/Units/Drawing/RigUnit_DrawContainer.cpp b/Engine/Plugins/Animation/ControlRig/Source/ControlRig/Private/Units/Drawing/RigUnit_DrawContainer.cpp
--- a/Engine/Plugins/Animation/ControlRig/Source/ControlRig/Private/Units/Drawing/RigUnit_DrawContainer.cpp
+++ b/Engine/Plugins/Animation/ControlRig/Source/ControlRig/Private/Units/Drawing/RigUnit_DrawContainer.cpp
@@ -9,19 +9,19 @@ FRigUnit_DrawContainerGetInstruction_Execute()
 {
     DECLARE_SCOPE_HIERARCHICAL_COUNTER_RIGUNIT()
 
-    if(Context.DrawContainer == nullptr)
+    if(ExecuteContext.UnitContext.DrawContainer == nullptr)
     {
     	return;
     }
 	
-	switch(Context.State)
+	switch(ExecuteContext.UnitContext.State)
 	{
 		case EControlRigState::Update:
 		{
-			int32 Index = Context.DrawContainer->GetIndex(InstructionName);
+			int32 Index = ExecuteContext.UnitContext.DrawContainer->GetIndex(InstructionName);
 			if (Index != INDEX_NONE)
 			{
-				const FControlRigDrawInstruction& Instruction = (*Context.DrawContainer)[Index];
+				const FControlRigDrawInstruction& Instruction = (*ExecuteContext.UnitContext.DrawContainer)[Index];
 				Color = Instruction.Color;
 				Transform = Instruction.Transform;
 			}
@@ -45,19 +45,19 @@ FRigUnit_DrawContainerSetThickness_Execute()
 {
 	DECLARE_SCOPE_HIERARCHICAL_COUNTER_RIGUNIT()
 
-	if (Context.DrawContainer == nullptr)
+	if (ExecuteContext.UnitContext.DrawContainer == nullptr)
 	{
 		return;
 	}
 
-	switch (Context.State)
+	switch (ExecuteContext.UnitContext.State)
 	{
 		case EControlRigState::Update:
 		{
-			int32 Index = Context.DrawContainer->GetIndex(InstructionName);
+			int32 Index = ExecuteContext.UnitContext.DrawContainer->GetIndex(InstructionName);
 			if (Index != INDEX_NONE)
 			{
-				FControlRigDrawInstruction& Instruction = (*Context.DrawContainer)[Index];
+				FControlRigDrawInstruction& Instruction = (*ExecuteContext.UnitContext.DrawContainer)[Index];
 				Instruction.Thickness = Thickness;
 			}
 			break;
@@ -75,19 +75,19 @@ FRigUnit_DrawContainerSetColor_Execute()
 {
 	DECLARE_SCOPE_HIERARCHICAL_COUNTER_RIGUNIT()
 
-	if (Context.DrawContainer == nullptr)
+	if (ExecuteContext.UnitContext.DrawContainer == nullptr)
 	{
 		return;
 	}
 
-	switch (Context.State)
+	switch (ExecuteContext.UnitContext.State)
 	{
 		case EControlRigState::Update:
 		{
-			int32 Index = Context.DrawContainer->GetIndex(InstructionName);
+			int32 Index = ExecuteContext.UnitContext.DrawContainer->GetIndex(InstructionName);
 			if (Index != INDEX_NONE)
 			{
-				FControlRigDrawInstruction& Instruction = (*Context.DrawContainer)[Index];
+				FControlRigDrawInstruction& Instruction = (*ExecuteContext.UnitContext.DrawContainer)[Index];
 				Instruction.Color = Color;
 			}
 			break;
@@ -105,19 +105,19 @@ FRigUnit_DrawContainerSetTransform_Execute()
 {
 	DECLARE_SCOPE_HIERARCHICAL_COUNTER_RIGUNIT()
 
-	if (Context.DrawContainer == nullptr)
+	if (ExecuteContext.UnitContext.DrawContainer == nullptr)
 	{
 		return;
 	}
 
-	switch (Context.State)
+	switch (ExecuteContext.UnitContext.State)
 	{
 		case EControlRigState::Update:
 		{
-			int32 Index = Context.DrawContainer->GetIndex(InstructionName);
+			int32 Index = ExecuteContext.UnitContext.DrawContainer->GetIndex(InstructionName);
 			if (Index != INDEX_NONE)
 			{
-				FControlRigDrawInstruction& Instruction = (*Context.DrawContainer)[Index];
+				FControlRigDrawInstruction& Instruction = (*ExecuteContext.UnitContext.DrawContainer)[Index];
 				Instruction.Transform = Transform;
 			}
 			break;
diff --git a/Engine/Plugins/Animation/ControlRig/Source/ControlRig/Private/Units/Drawing/RigUnit_DrawContainer.h b/Engine/Plugins/Animation/ControlRig/Source/ControlRig/Private/Units/Drawing/RigUnit_DrawContainer.h
--- a/Engine/Plugins/Animation/ControlRig/Source/ControlRig/Private/Units/Drawing/RigUnit_DrawContainer.h
+++ b/Engine/Plugins/Animation/ControlRig/Source/ControlRig/Private/Units/Drawing/RigUnit_DrawContainer.h
@@ -21,7 +21,7 @@ struct CONTROLRIG_API FRigUnit_DrawContainerGetInstruction : public FRigUnit
 	}
 
 	RIGVM_METHOD()
-	virtual void Execute(const FRigUnitContext& Context) override;
+	virtual void Execute() override;
 
 	UPROPERTY(meta = (Input, Constant, CustomWidget = "DrawingName"))
 	FName InstructionName;
@@ -51,7 +51,7 @@ struct CONTROLRIG_API FRigUnit_DrawContainerSetColor: public FRigUnitMutable
 	}
 
 	RIGVM_METHOD()
-	virtual void Execute(const FRigUnitContext& Context) override;
+	virtual void Execute() override;
 
 	UPROPERTY(meta = (Input, Constant, CustomWidget = "DrawingName"))
 	FName InstructionName;
@@ -78,7 +78,7 @@ struct CONTROLRIG_API FRigUnit_DrawContainerSetThickness : public FRigUnitMutabl
 	}
 
 	RIGVM_METHOD()
-	virtual void Execute(const FRigUnitContext& Context) override;
+	virtual void Execute() override;
 
 	UPROPERTY(meta = (Input, Constant, CustomWidget = "DrawingName"))
 	FName InstructionName;
@@ -105,7 +105,7 @@ struct CONTROLRIG_API FRigUnit_DrawContainerSetTransform : public FRigUnitMutabl
 	}
 
 	RIGVM_METHOD()
-	virtual void Execute(const FRigUnitContext& Context) override;
+	virtual void Execute() override;
 
 	UPROPERTY(meta = (Input, Constant, CustomWidget = "DrawingName"))
 	FName InstructionName;
diff --git a/Engine/Plugins/Animation/ControlRig/Source/ControlRig/Private/Units/Execution/RigUnit_BeginExecution.cpp b/Engine/Plugins/Animation/ControlRig/Source/ControlRig/Private/Units/Execution/RigUnit_BeginExecution.cpp
--- a/Engine/Plugins/Animation/ControlRig/Source/ControlRig/Private/Units/Execution/RigUnit_BeginExecution.cpp
+++ b/Engine/Plugins/Animation/ControlRig/Source/ControlRig/Private/Units/Execution/RigUnit_BeginExecution.cpp
@@ -10,7 +10,6 @@ FName FRigUnit_BeginExecution::EventName = TEXT("Forwards Solve");
 FRigUnit_BeginExecution_Execute()
 {
     DECLARE_SCOPE_HIERARCHICAL_COUNTER_RIGUNIT()
-	ExecuteContext.Hierarchy = Context.Hierarchy;
 	ExecuteContext.EventName = FRigUnit_BeginExecution::EventName;
 }
 
diff --git a/Engine/Plugins/Animation/ControlRig/Source/ControlRig/Private/Units/Execution/RigUnit_Collection.cpp b/Engine/Plugins/Animation/ControlRig/Source/ControlRig/Private/Units/Execution/RigUnit_Collection.cpp
--- a/Engine/Plugins/Animation/ControlRig/Source/ControlRig/Private/Units/Execution/RigUnit_Collection.cpp
+++ b/Engine/Plugins/Animation/ControlRig/Source/ControlRig/Private/Units/Execution/RigUnit_Collection.cpp
@@ -25,7 +25,7 @@
 
 FRigUnit_CollectionChain_Execute()
 {
-	FRigUnit_CollectionChainArray::StaticExecute(ExecuteContext, FirstItem, LastItem, Reverse, Collection.Keys, Context);
+	FRigUnit_CollectionChainArray::StaticExecute(ExecuteContext, FirstItem, LastItem, Reverse, Collection.Keys);
 }
 
 FRigVMStructUpgradeInfo FRigUnit_CollectionChain::GetUpgradeInfo() const
@@ -40,37 +40,37 @@ FRigVMStructUpgradeInfo FRigUnit_CollectionChain::GetUpgradeInfo() const
 
 FRigUnit_CollectionChainArray_Execute()
 {
-	uint32 Hash = FRigUnit_CollectionChain_Hash + Context.Hierarchy->GetTopologyVersion() * 17;
+	uint32 Hash = FRigUnit_CollectionChain_Hash + ExecuteContext.Hierarchy->GetTopologyVersion() * 17;
 	Hash = HashCombine(Hash, GetTypeHash(FirstItem));
 	Hash = HashCombine(Hash, GetTypeHash(LastItem));
 	Hash = HashCombine(Hash, Reverse ? 1 : 0);
 
 	FRigElementKeyCollection Collection;
-	if(const FRigElementKeyCollection* Cache = Context.Hierarchy->FindCachedCollection(Hash))
+	if(const FRigElementKeyCollection* Cache = ExecuteContext.Hierarchy->FindCachedCollection(Hash))
 	{
 		Collection = *Cache;
 	}
 	else
 	{
-		Collection = FRigElementKeyCollection::MakeFromChain(Context.Hierarchy, FirstItem, LastItem, Reverse);
+		Collection = FRigElementKeyCollection::MakeFromChain(ExecuteContext.Hierarchy, FirstItem, LastItem, Reverse);
 		if (Collection.IsEmpty())
 		{
-			if (Context.Hierarchy->GetIndex(FirstItem) == INDEX_NONE)
+			if (ExecuteContext.Hierarchy->GetIndex(FirstItem) == INDEX_NONE)
 			{
-				if(Context.State != EControlRigState::Init)
+				if(ExecuteContext.UnitContext.State != EControlRigState::Init)
 				{
 					UE_CONTROLRIG_RIGUNIT_REPORT_WARNING(TEXT("First Item '%s' is not valid."), *FirstItem.ToString());
 				}
 			}
-			if (Context.Hierarchy->GetIndex(LastItem) == INDEX_NONE)
+			if (ExecuteContext.Hierarchy->GetIndex(LastItem) == INDEX_NONE)
 			{
-				if(Context.State != EControlRigState::Init)
+				if(ExecuteContext.UnitContext.State != EControlRigState::Init)
 				{
 					UE_CONTROLRIG_RIGUNIT_REPORT_WARNING(TEXT("Last Item '%s' is not valid."), *LastItem.ToString());
 				}
 			}
 		}
-		Context.Hierarchy->AddCachedCollection(Hash, Collection);
+		ExecuteContext.Hierarchy->AddCachedCollection(Hash, Collection);
 	}
 
 	Items = Collection.Keys;
@@ -78,7 +78,7 @@ FRigUnit_CollectionChainArray_Execute()
 
 FRigUnit_CollectionNameSearch_Execute()
 {
-	FRigUnit_CollectionNameSearchArray::StaticExecute(ExecuteContext, PartialName, TypeToSearch, Collection.Keys, Context);
+	FRigUnit_CollectionNameSearchArray::StaticExecute(ExecuteContext, PartialName, TypeToSearch, Collection.Keys);
 }
 
 FRigVMStructUpgradeInfo FRigUnit_CollectionNameSearch::GetUpgradeInfo() const
@@ -94,19 +94,19 @@ FRigUnit_CollectionNameSearchArray_Execute()
 {
 	DECLARE_SCOPE_HIERARCHICAL_COUNTER_RIGUNIT()
 
-	uint32 Hash = FRigUnit_CollectionNameSearch_Hash + Context.Hierarchy->GetTopologyVersion() * 17;
+	uint32 Hash = FRigUnit_CollectionNameSearch_Hash + ExecuteContext.Hierarchy->GetTopologyVersion() * 17;
 	Hash = HashCombine(Hash, GetTypeHash(PartialName));
 	Hash = HashCombine(Hash, (int32)TypeToSearch * 8);
 
 	FRigElementKeyCollection Collection;
-	if(const FRigElementKeyCollection* Cache = Context.Hierarchy->FindCachedCollection(Hash))
+	if(const FRigElementKeyCollection* Cache = ExecuteContext.Hierarchy->FindCachedCollection(Hash))
 	{
 		Collection = *Cache;
 	}
 	else
 	{
-		Collection = FRigElementKeyCollection::MakeFromName(Context.Hierarchy, PartialName, (uint8)TypeToSearch);
-		Context.Hierarchy->AddCachedCollection(Hash, Collection);
+		Collection = FRigElementKeyCollection::MakeFromName(ExecuteContext.Hierarchy, PartialName, (uint8)TypeToSearch);
+		ExecuteContext.Hierarchy->AddCachedCollection(Hash, Collection);
 	}
 
 	Items = Collection.Keys;
@@ -114,7 +114,7 @@ FRigUnit_CollectionNameSearchArray_Execute()
 
 FRigUnit_CollectionChildren_Execute()
 {
-	FRigUnit_CollectionChildrenArray::StaticExecute(ExecuteContext, Parent, bIncludeParent, bRecursive, TypeToSearch, Collection.Keys, Context);
+	FRigUnit_CollectionChildrenArray::StaticExecute(ExecuteContext, Parent, bIncludeParent, bRecursive, TypeToSearch, Collection.Keys);
 }
 
 FRigVMStructUpgradeInfo FRigUnit_CollectionChildren::GetUpgradeInfo() const
@@ -134,28 +134,28 @@ FRigUnit_CollectionChildrenArray_Execute()
 {
 	DECLARE_SCOPE_HIERARCHICAL_COUNTER_RIGUNIT()
 
-	uint32 Hash = FRigUnit_CollectionChildren_Hash + Context.Hierarchy->GetTopologyVersion() * 17;
+	uint32 Hash = FRigUnit_CollectionChildren_Hash + ExecuteContext.Hierarchy->GetTopologyVersion() * 17;
 	Hash = HashCombine(Hash, GetTypeHash(Parent));
 	Hash = HashCombine(Hash, bRecursive ? 2 : 0);
 	Hash = HashCombine(Hash, bIncludeParent ? 1 : 0);
 	Hash = HashCombine(Hash, (int32)TypeToSearch * 8);
 
 	FRigElementKeyCollection Collection;
-	if(const FRigElementKeyCollection* Cache = Context.Hierarchy->FindCachedCollection(Hash))
+	if(const FRigElementKeyCollection* Cache = ExecuteContext.Hierarchy->FindCachedCollection(Hash))
 	{
 		Collection = *Cache;
 	}
 	else
 	{
-		Collection = FRigElementKeyCollection::MakeFromChildren(Context.Hierarchy, Parent, bRecursive, bIncludeParent, (uint8)TypeToSearch);
+		Collection = FRigElementKeyCollection::MakeFromChildren(ExecuteContext.Hierarchy, Parent, bRecursive, bIncludeParent, (uint8)TypeToSearch);
 		if (Collection.IsEmpty())
 		{
-			if (Context.Hierarchy->GetIndex(Parent) == INDEX_NONE)
+			if (ExecuteContext.Hierarchy->GetIndex(Parent) == INDEX_NONE)
 			{
 				UE_CONTROLRIG_RIGUNIT_REPORT_WARNING(TEXT("Parent '%s' is not valid."), *Parent.ToString());
 			}
 		}
-		Context.Hierarchy->AddCachedCollection(Hash, Collection);
+		ExecuteContext.Hierarchy->AddCachedCollection(Hash, Collection);
 	}
 
 	Items = Collection.Keys;
@@ -165,17 +165,17 @@ FRigUnit_CollectionGetAll_Execute()
 {
 	DECLARE_SCOPE_HIERARCHICAL_COUNTER_RIGUNIT()
 
-	uint32 Hash = FRigUnit_CollectionGetAll_Hash + Context.Hierarchy->GetTopologyVersion() * 17;
+	uint32 Hash = FRigUnit_CollectionGetAll_Hash + ExecuteContext.Hierarchy->GetTopologyVersion() * 17;
 	Hash = HashCombine(Hash, (int32)TypeToSearch * 8);
 
 	FRigElementKeyCollection Collection;
-	if(const FRigElementKeyCollection* Cache = Context.Hierarchy->FindCachedCollection(Hash))
+	if(const FRigElementKeyCollection* Cache = ExecuteContext.Hierarchy->FindCachedCollection(Hash))
 	{
 		Collection = *Cache;
 	}
 	else
 	{
-		Context.Hierarchy->Traverse([&Collection, TypeToSearch](FRigBaseElement* InElement, bool &bContinue)
+		ExecuteContext.Hierarchy->Traverse([&Collection, TypeToSearch](FRigBaseElement* InElement, bool &bContinue)
 		{
 			bContinue = true;
 			
@@ -190,7 +190,7 @@ FRigUnit_CollectionGetAll_Execute()
 		{
 			UE_CONTROLRIG_RIGUNIT_REPORT_WARNING(TEXT("%s"), TEXT("No elements found for given filter."));
 		}
-		Context.Hierarchy->AddCachedCollection(Hash, Collection);
+		ExecuteContext.Hierarchy->AddCachedCollection(Hash, Collection);
 	}
 
 	Items = Collection.Keys;
@@ -237,7 +237,7 @@ IMPLEMENT_RIGUNIT_AUTOMATION_TEST(FRigUnit_CollectionChildren)
 
 FRigUnit_CollectionReplaceItems_Execute()
 {
-	FRigUnit_CollectionReplaceItemsArray::StaticExecute(ExecuteContext, Items.Keys, Old, New, RemoveInvalidItems, bAllowDuplicates, Collection.Keys, Context);
+	FRigUnit_CollectionReplaceItemsArray::StaticExecute(ExecuteContext, Items.Keys, Old, New, RemoveInvalidItems, bAllowDuplicates, Collection.Keys);
 }
 
 FRigVMStructUpgradeInfo FRigUnit_CollectionReplaceItems::GetUpgradeInfo() const
@@ -256,14 +256,14 @@ FRigUnit_CollectionReplaceItemsArray_Execute()
 {
 	DECLARE_SCOPE_HIERARCHICAL_COUNTER_RIGUNIT()
 
-	uint32 Hash = FRigUnit_CollectionReplaceItems_Hash + Context.Hierarchy->GetTopologyVersion() * 17;
+	uint32 Hash = FRigUnit_CollectionReplaceItems_Hash + ExecuteContext.Hierarchy->GetTopologyVersion() * 17;
 	Hash = HashCombine(Hash, GetTypeHash(Items));
 	Hash = HashCombine(Hash, 12 * GetTypeHash(Old));
 	Hash = HashCombine(Hash, 13 * GetTypeHash(New));
 	Hash = HashCombine(Hash, RemoveInvalidItems ? 14 : 0);
 
 	FRigElementKeyCollection Collection;
-	if(const FRigElementKeyCollection* Cache = Context.Hierarchy->FindCachedCollection(Hash))
+	if(const FRigElementKeyCollection* Cache = ExecuteContext.Hierarchy->FindCachedCollection(Hash))
 	{
 		Collection = *Cache;
 	}
@@ -274,9 +274,9 @@ FRigUnit_CollectionReplaceItemsArray_Execute()
 		for (int32 Index = 0; Index < Items.Num(); Index++)
 		{
 			FRigElementKey Key = Items[Index];
-			FRigUnit_ItemReplace::StaticExecute(ExecuteContext, Key, Old, New, Key, Context);
+			FRigUnit_ItemReplace::StaticExecute(ExecuteContext, Key, Old, New, Key);
 
-			if (Context.Hierarchy->GetIndex(Key) != INDEX_NONE)
+			if (ExecuteContext.Hierarchy->GetIndex(Key) != INDEX_NONE)
 			{
 				if(bAllowDuplicates)
 				{
@@ -293,7 +293,7 @@ FRigUnit_CollectionReplaceItemsArray_Execute()
 			}
 		}
 
-		Context.Hierarchy->AddCachedCollection(Hash, Collection);
+		ExecuteContext.Hierarchy->AddCachedCollection(Hash, Collection);
 	}
 	
 	Result = Collection.Keys;
@@ -326,7 +326,7 @@ FRigUnit_CollectionGetItems_Execute()
 
 FRigUnit_CollectionGetParentIndices_Execute()
 {
-	FRigUnit_CollectionGetParentIndicesItemArray::StaticExecute(ExecuteContext, Collection.Keys, ParentIndices, Context);
+	FRigUnit_CollectionGetParentIndicesItemArray::StaticExecute(ExecuteContext, Collection.Keys, ParentIndices);
 }
 
 FRigVMStructUpgradeInfo FRigUnit_CollectionGetParentIndices::GetUpgradeInfo() const
@@ -343,7 +343,7 @@ FRigUnit_CollectionGetParentIndicesItemArray_Execute()
 {
 	DECLARE_SCOPE_HIERARCHICAL_COUNTER_RIGUNIT()
 
-	if(Context.Hierarchy == nullptr)
+	if(ExecuteContext.Hierarchy == nullptr)
 	{
 		ParentIndices.Reset();
 		return;
@@ -355,7 +355,7 @@ FRigUnit_CollectionGetParentIndicesItemArray_Execute()
 	{
 		ParentIndices[Index] = INDEX_NONE;
 
-		const int32 ItemIndex = Context.Hierarchy->GetIndex(Items[Index]);
+		const int32 ItemIndex = ExecuteContext.Hierarchy->GetIndex(Items[Index]);
 		if(ItemIndex == INDEX_NONE)
 		{
 			continue;;
@@ -369,17 +369,17 @@ FRigUnit_CollectionGetParentIndicesItemArray_Execute()
 			}
 			case ERigElementType::Bone:
 			{
-				ParentIndices[Index] = Context.Hierarchy->GetFirstParent(ItemIndex);
+				ParentIndices[Index] = ExecuteContext.Hierarchy->GetFirstParent(ItemIndex);
 				break;
 			}
 			default:
 			{
-				if(const FRigBaseElement* ChildElement = Context.Hierarchy->Get(ItemIndex))
+				if(const FRigBaseElement* ChildElement = ExecuteContext.Hierarchy->Get(ItemIndex))
 				{
-					TArray<int32> ItemParents = Context.Hierarchy->GetParents(ItemIndex);
+					TArray<int32> ItemParents = ExecuteContext.Hierarchy->GetParents(ItemIndex);
 					for(int32 ParentIndex = 0; ParentIndex < ItemParents.Num(); ParentIndex++)
 					{
-						const FRigElementWeight Weight = Context.Hierarchy->GetParentWeight(ChildElement, ParentIndex, false);
+						const FRigElementWeight Weight = ExecuteContext.Hierarchy->GetParentWeight(ChildElement, ParentIndex, false);
 						if(!Weight.IsAlmostZero())
 						{
 							ParentIndices[Index] = ItemParents[ParentIndex];
@@ -397,8 +397,8 @@ FRigUnit_CollectionGetParentIndicesItemArray_Execute()
 
 			while(ParentIndices[Index] == INDEX_NONE && ParentIndex != INDEX_NONE)
 			{
-				ParentIndices[Index] = Items.Find(Context.Hierarchy->GetKey(ParentIndex));
-				ParentIndex = Context.Hierarchy->GetFirstParent(ParentIndex);
+				ParentIndices[Index] = Items.Find(ExecuteContext.Hierarchy->GetKey(ParentIndex));
+				ParentIndex = ExecuteContext.Hierarchy->GetFirstParent(ParentIndex);
 			}
 		}
 	}
diff --git a/Engine/Plugins/Animation/ControlRig/Source/ControlRig/Private/Units/Execution/RigUnit_Collection.h b/Engine/Plugins/Animation/ControlRig/Source/ControlRig/Private/Units/Execution/RigUnit_Collection.h
--- a/Engine/Plugins/Animation/ControlRig/Source/ControlRig/Private/Units/Execution/RigUnit_Collection.h
+++ b/Engine/Plugins/Animation/ControlRig/Source/ControlRig/Private/Units/Execution/RigUnit_Collection.h
@@ -34,7 +34,7 @@ struct CONTROLRIG_API FRigUnit_CollectionChain : public FRigUnit_CollectionBase
 	}
 
 	RIGVM_METHOD()
-	virtual void Execute(const FRigUnitContext& Context) override;
+	virtual void Execute() override;
 
 	UPROPERTY(meta = (Input, ExpandByDefault))
 	FRigElementKey FirstItem;
@@ -68,7 +68,7 @@ struct CONTROLRIG_API FRigUnit_CollectionChainArray : public FRigUnit_Collection
 	}
 
 	RIGVM_METHOD()
-	virtual void Execute(const FRigUnitContext& Context) override;
+	virtual void Execute() override;
 
 	UPROPERTY(meta = (Input, ExpandByDefault))
 	FRigElementKey FirstItem;
@@ -99,7 +99,7 @@ struct CONTROLRIG_API FRigUnit_CollectionNameSearch : public FRigUnit_Collection
 	}
 
 	RIGVM_METHOD()
-	virtual void Execute(const FRigUnitContext& Context) override;
+	virtual void Execute() override;
 
 	UPROPERTY(meta = (Input))
 	FName PartialName;
@@ -130,7 +130,7 @@ struct CONTROLRIG_API FRigUnit_CollectionNameSearchArray : public FRigUnit_Colle
 	}
 
 	RIGVM_METHOD()
-	virtual void Execute(const FRigUnitContext& Context) override;
+	virtual void Execute() override;
 
 	UPROPERTY(meta = (Input))
 	FName PartialName;
@@ -160,7 +160,7 @@ struct CONTROLRIG_API FRigUnit_CollectionChildren : public FRigUnit_CollectionBa
 	}
 
 	RIGVM_METHOD()
-	virtual void Execute(const FRigUnitContext& Context) override;
+	virtual void Execute() override;
 
 	UPROPERTY(meta = (Input, ExpandByDefault))
 	FRigElementKey Parent;
@@ -199,7 +199,7 @@ struct CONTROLRIG_API FRigUnit_CollectionChildrenArray : public FRigUnit_Collect
 	}
 
 	RIGVM_METHOD()
-	virtual void Execute(const FRigUnitContext& Context) override;
+	virtual void Execute() override;
 
 	UPROPERTY(meta = (Input, ExpandByDefault))
 	FRigElementKey Parent;
@@ -231,7 +231,7 @@ struct CONTROLRIG_API FRigUnit_CollectionGetAll : public FRigUnit_CollectionBase
 	}
 
 	RIGVM_METHOD()
-	virtual void Execute(const FRigUnitContext& Context) override;
+	virtual void Execute() override;
 
 	UPROPERTY(meta = (Input))
 	ERigElementType TypeToSearch;
@@ -256,7 +256,7 @@ struct CONTROLRIG_API FRigUnit_CollectionReplaceItems : public FRigUnit_Collecti
 	}
 
 	RIGVM_METHOD()
-	virtual void Execute(const FRigUnitContext& Context) override;
+	virtual void Execute() override;
 
 	UPROPERTY(meta = (Input))
 	FRigElementKeyCollection Items;
@@ -296,7 +296,7 @@ struct CONTROLRIG_API FRigUnit_CollectionReplaceItemsArray : public FRigUnit_Col
 	}
 
 	RIGVM_METHOD()
-	virtual void Execute(const FRigUnitContext& Context) override;
+	virtual void Execute() override;
 
 	UPROPERTY(meta = (Input))
 	TArray<FRigElementKey> Items;
@@ -332,7 +332,7 @@ struct CONTROLRIG_API FRigUnit_CollectionItems : public FRigUnit_CollectionBase
 	}
 
 	RIGVM_METHOD()
-	virtual void Execute(const FRigUnitContext& Context) override;
+	virtual void Execute() override;
 
 	UPROPERTY(meta = (Input, ExpandByDefault))
 	TArray<FRigElementKey> Items;
@@ -357,7 +357,7 @@ struct CONTROLRIG_API FRigUnit_CollectionGetItems : public FRigUnit_CollectionBa
 	}
 
 	RIGVM_METHOD()
-	virtual void Execute(const FRigUnitContext& Context) override;
+	virtual void Execute() override;
 
 	UPROPERTY(meta = (Input))
 	FRigElementKeyCollection Collection;
@@ -382,7 +382,7 @@ struct CONTROLRIG_API FRigUnit_CollectionGetParentIndices : public FRigUnit_Coll
 	}
 
 	RIGVM_METHOD()
-	virtual void Execute(const FRigUnitContext& Context) override;
+	virtual void Execute() override;
 
 	UPROPERTY(meta = (Input))
 	FRigElementKeyCollection Collection;
@@ -410,7 +410,7 @@ struct CONTROLRIG_API FRigUnit_CollectionGetParentIndicesItemArray : public FRig
 	}
 
 	RIGVM_METHOD()
-	virtual void Execute(const FRigUnitContext& Context) override;
+	virtual void Execute() override;
 
 	UPROPERTY(meta = (Input))
 	TArray<FRigElementKey> Items;
@@ -434,7 +434,7 @@ struct CONTROLRIG_API FRigUnit_CollectionUnion : public FRigUnit_CollectionBase
 	}
 
 	RIGVM_METHOD()
-	virtual void Execute(const FRigUnitContext& Context) override;
+	virtual void Execute() override;
 
 	UPROPERTY(meta = (Input))
 	FRigElementKeyCollection A;
@@ -466,7 +466,7 @@ struct CONTROLRIG_API FRigUnit_CollectionIntersection : public FRigUnit_Collecti
 	}
 
 	RIGVM_METHOD()
-	virtual void Execute(const FRigUnitContext& Context) override;
+	virtual void Execute() override;
 
 	UPROPERTY(meta = (Input))
 	FRigElementKeyCollection A;
@@ -495,7 +495,7 @@ struct CONTROLRIG_API FRigUnit_CollectionDifference : public FRigUnit_Collection
 	}
 
 	RIGVM_METHOD()
-	virtual void Execute(const FRigUnitContext& Context) override;
+	virtual void Execute() override;
 
 	UPROPERTY(meta = (Input))
 	FRigElementKeyCollection A;
@@ -523,7 +523,7 @@ struct CONTROLRIG_API FRigUnit_CollectionReverse : public FRigUnit_CollectionBas
 	}
 
 	RIGVM_METHOD()
-	virtual void Execute(const FRigUnitContext& Context) override;
+	virtual void Execute() override;
 
 	UPROPERTY(meta = (Input))
 	FRigElementKeyCollection Collection;
@@ -550,7 +550,7 @@ struct CONTROLRIG_API FRigUnit_CollectionCount : public FRigUnit_CollectionBase
 	}
 
 	RIGVM_METHOD()
-	virtual void Execute(const FRigUnitContext& Context) override;
+	virtual void Execute() override;
 
 	UPROPERTY(meta = (Input))
 	FRigElementKeyCollection Collection;
@@ -578,7 +578,7 @@ struct CONTROLRIG_API FRigUnit_CollectionItemAtIndex : public FRigUnit_Collectio
 	}
 
 	RIGVM_METHOD()
-	virtual void Execute(const FRigUnitContext& Context) override;
+	virtual void Execute() override;
 
 	UPROPERTY(meta = (Input))
 	FRigElementKeyCollection Collection;
@@ -619,7 +619,7 @@ struct CONTROLRIG_API FRigUnit_CollectionLoop : public FRigUnit_CollectionBaseMu
 	virtual int32 GetNumSlices() const override { return Count; }
 
 	RIGVM_METHOD()
-	virtual void Execute(const FRigUnitContext& Context) override;
+	virtual void Execute() override;
 
 	UPROPERTY(meta = (Singleton))
 	FName BlockToRun;
@@ -666,7 +666,7 @@ struct CONTROLRIG_API FRigUnit_CollectionAddItem : public FRigUnit_CollectionBas
 	}
 
 	RIGVM_METHOD()
-    virtual void Execute(const FRigUnitContext& Context) override;
+    virtual void Execute() override;
 
 	UPROPERTY(meta = (Input))
 	FRigElementKeyCollection Collection;
