RigVM: Remove support for opaque attributes

diff --git a/Engine/Plugins/Animation/ControlRig/Source/ControlRig/Private/Units/Execution/RigUnit_DynamicHierarchy.cpp b/Engine/Plugins/Animation/ControlRig/Source/ControlRig/Private/Units/Execution/RigUnit_DynamicHierarchy.cpp
--- a/Engine/Plugins/Animation/ControlRig/Source/ControlRig/Private/Units/Execution/RigUnit_DynamicHierarchy.cpp
+++ b/Engine/Plugins/Animation/ControlRig/Source/ControlRig/Private/Units/Execution/RigUnit_DynamicHierarchy.cpp
@@ -11,13 +11,11 @@
 #include UE_INLINE_GENERATED_CPP_BY_NAME(RigUnit_DynamicHierarchy)
 
 bool FRigUnit_DynamicHierarchyBase::IsValidToRunInContext(
-	const FRigUnitContext& InContext,
 	const FControlRigExecuteContext& InExecuteContext,
 	bool bAllowOnlyConstructionEvent,
 	FString* OutErrorMessage)
 {
-	if((InContext.State != EControlRigState::Update) ||
-		(InContext.Hierarchy == nullptr) ||
+	if((InExecuteContext.UnitContext.State != EControlRigState::Update) ||
 		(InExecuteContext.Hierarchy == nullptr))
 	{
 		return false;
@@ -36,7 +34,7 @@ bool FRigUnit_DynamicHierarchyBase::IsValidToRunInContext(
 		}
 	}
 
-	if(InContext.Hierarchy->Num() >= InContext.HierarchySettings.ProceduralElementLimit)
+	if(InExecuteContext.Hierarchy->Num() >= InExecuteContext.UnitContext.HierarchySettings.ProceduralElementLimit)
 	{
 		if(OutErrorMessage)
 		{
@@ -51,7 +49,7 @@ bool FRigUnit_DynamicHierarchyBase::IsValidToRunInContext(
 
 FRigUnit_AddParent_Execute()
 {
-	if(!FRigUnit_DynamicHierarchyBase::IsValidToRunInContext(Context, ExecuteContext, true))
+	if(!FRigUnit_DynamicHierarchyBase::IsValidToRunInContext(ExecuteContext, true))
 	{
 		return;
 	}
@@ -79,7 +77,7 @@ FRigUnit_AddParent_Execute()
 
 FRigUnit_SetDefaultParent_Execute()
 {
-	if(!FRigUnit_DynamicHierarchyBase::IsValidToRunInContext(Context, ExecuteContext, true))
+	if(!FRigUnit_DynamicHierarchyBase::IsValidToRunInContext(ExecuteContext, true))
 	{
 		return;
 	}
@@ -107,7 +105,7 @@ FRigUnit_SetDefaultParent_Execute()
 
 FRigUnit_SwitchParent_Execute()
 {
-	if(!FRigUnit_DynamicHierarchyBase::IsValidToRunInContext(Context, ExecuteContext, false))
+	if(!FRigUnit_DynamicHierarchyBase::IsValidToRunInContext(ExecuteContext, false))
 	{
 		return;
 	}
@@ -200,7 +198,7 @@ FRigUnit_SwitchParent_Execute()
 
 FRigUnit_HierarchyGetParentWeights_Execute()
 {
-	FRigUnit_HierarchyGetParentWeightsArray::StaticExecute(ExecuteContext, Child, Weights, Parents.Keys, Context);
+	FRigUnit_HierarchyGetParentWeightsArray::StaticExecute(ExecuteContext, Child, Weights, Parents.Keys);
 }
 
 FRigVMStructUpgradeInfo FRigUnit_HierarchyGetParentWeights::GetUpgradeInfo() const
@@ -215,30 +213,30 @@ FRigVMStructUpgradeInfo FRigUnit_HierarchyGetParentWeights::GetUpgradeInfo() con
 
 FRigUnit_HierarchyGetParentWeightsArray_Execute()
 {
-	if((Context.State != EControlRigState::Update) || (Context.Hierarchy == nullptr))
+	if((ExecuteContext.UnitContext.State != EControlRigState::Update) || (ExecuteContext.Hierarchy == nullptr))
 	{
 		return;
 	}
 
-	FRigBaseElement* ChildElement = Context.Hierarchy->Find(Child);
+	FRigBaseElement* ChildElement = ExecuteContext.Hierarchy->Find(Child);
 	if(ChildElement == nullptr)
 	{
 		UE_CONTROLRIG_RIGUNIT_REPORT_WARNING(TEXT("Item %s does not exist."), *Child.ToString())
 		return;
 	}
 	
-	Weights = Context.Hierarchy->GetParentWeightArray(ChildElement, false);
-	Parents = Context.Hierarchy->GetParents(ChildElement->GetKey(), false);
+	Weights = ExecuteContext.Hierarchy->GetParentWeightArray(ChildElement, false);
+	Parents = ExecuteContext.Hierarchy->GetParents(ChildElement->GetKey(), false);
 }
 
 FRigUnit_HierarchySetParentWeights_Execute()
 {
-	if(!FRigUnit_DynamicHierarchyBase::IsValidToRunInContext(Context, ExecuteContext, false))
+	if(!FRigUnit_DynamicHierarchyBase::IsValidToRunInContext(ExecuteContext, false))
 	{
 		return;
 	}
 
-	FRigBaseElement* ChildElement = Context.Hierarchy->Find(Child);
+	FRigBaseElement* ChildElement = ExecuteContext.Hierarchy->Find(Child);
 	if(ChildElement == nullptr)
 	{
 		UE_CONTROLRIG_RIGUNIT_REPORT_WARNING(TEXT("Item %s does not exist."), *Child.ToString())
@@ -263,7 +261,7 @@ FRigUnit_HierarchySetParentWeights_Execute()
 FRigUnit_HierarchyReset_Execute()
 {
 	FString ErrorMessage;
-	if(!FRigUnit_DynamicHierarchyBase::IsValidToRunInContext(Context, ExecuteContext, true, &ErrorMessage))
+	if(!FRigUnit_DynamicHierarchyBase::IsValidToRunInContext(ExecuteContext, true, &ErrorMessage))
 	{
 		if(!ErrorMessage.IsEmpty())
 		{
@@ -277,7 +275,7 @@ FRigUnit_HierarchyReset_Execute()
 FRigUnit_HierarchyImportFromSkeleton_Execute()
 {
 	FString ErrorMessage;
-	if(!FRigUnit_DynamicHierarchyBase::IsValidToRunInContext(Context, ExecuteContext, true, &ErrorMessage))
+	if(!FRigUnit_DynamicHierarchyBase::IsValidToRunInContext(ExecuteContext, true, &ErrorMessage))
 	{
 		if(!ErrorMessage.IsEmpty())
 		{
@@ -292,7 +290,7 @@ FRigUnit_HierarchyImportFromSkeleton_Execute()
 	{
 		FRigHierarchyControllerInstructionBracket InstructionBracket(Controller, ExecuteContext.GetInstructionIndex());
 		
-		if(const USkeletalMeshComponent* SkelMeshComponent = Context.DataSourceRegistry->RequestSource<USkeletalMeshComponent>(UControlRig::OwnerComponent))
+		if(const USkeletalMeshComponent* SkelMeshComponent = ExecuteContext.UnitContext.DataSourceRegistry->RequestSource<USkeletalMeshComponent>(UControlRig::OwnerComponent))
 		{
 			if(SkelMeshComponent->GetSkeletalMeshAsset())
 			{
@@ -313,7 +311,7 @@ FRigUnit_HierarchyImportFromSkeleton_Execute()
 FRigUnit_HierarchyRemoveElement_Execute()
 {
 	FString ErrorMessage;
-	if(!FRigUnit_DynamicHierarchyBase::IsValidToRunInContext(Context, ExecuteContext, true, &ErrorMessage))
+	if(!FRigUnit_DynamicHierarchyBase::IsValidToRunInContext(ExecuteContext, true, &ErrorMessage))
 	{
 		if(!ErrorMessage.IsEmpty())
 		{
@@ -334,7 +332,7 @@ FRigUnit_HierarchyRemoveElement_Execute()
 FRigUnit_HierarchyAddBone_Execute()
 {
 	FString ErrorMessage;
-	if(!FRigUnit_DynamicHierarchyBase::IsValidToRunInContext(Context, ExecuteContext, true, &ErrorMessage))
+	if(!FRigUnit_DynamicHierarchyBase::IsValidToRunInContext(ExecuteContext, true, &ErrorMessage))
 	{
 		if(!ErrorMessage.IsEmpty())
 		{
@@ -355,7 +353,7 @@ FRigUnit_HierarchyAddBone_Execute()
 FRigUnit_HierarchyAddNull_Execute()
 {
 	FString ErrorMessage;
-	if(!FRigUnit_DynamicHierarchyBase::IsValidToRunInContext(Context, ExecuteContext, true, &ErrorMessage))
+	if(!FRigUnit_DynamicHierarchyBase::IsValidToRunInContext(ExecuteContext, true, &ErrorMessage))
 	{
 		if(!ErrorMessage.IsEmpty())
 		{
@@ -416,7 +414,7 @@ void FRigUnit_HierarchyAddControlFloat_Settings::Configure(FRigControlSettings&
 FRigUnit_HierarchyAddControlFloat_Execute()
 {
 	FString ErrorMessage;
-	if(!FRigUnit_DynamicHierarchyBase::IsValidToRunInContext(Context, ExecuteContext, true, &ErrorMessage))
+	if(!FRigUnit_DynamicHierarchyBase::IsValidToRunInContext(ExecuteContext, true, &ErrorMessage))
 	{
 		if(!ErrorMessage.IsEmpty())
 		{
@@ -464,7 +462,7 @@ void FRigUnit_HierarchyAddControlInteger_Settings::Configure(FRigControlSettings
 FRigUnit_HierarchyAddControlInteger_Execute()
 {
 	FString ErrorMessage;
-	if(!FRigUnit_DynamicHierarchyBase::IsValidToRunInContext(Context, ExecuteContext, true, &ErrorMessage))
+	if(!FRigUnit_DynamicHierarchyBase::IsValidToRunInContext(ExecuteContext, true, &ErrorMessage))
 	{
 		if(!ErrorMessage.IsEmpty())
 		{
@@ -514,7 +512,7 @@ void FRigUnit_HierarchyAddControlVector2D_Settings::Configure(FRigControlSetting
 FRigUnit_HierarchyAddControlVector2D_Execute()
 {
 	FString ErrorMessage;
-	if(!FRigUnit_DynamicHierarchyBase::IsValidToRunInContext(Context, ExecuteContext, true, &ErrorMessage))
+	if(!FRigUnit_DynamicHierarchyBase::IsValidToRunInContext(ExecuteContext, true, &ErrorMessage))
 	{
 		if(!ErrorMessage.IsEmpty())
 		{
@@ -564,7 +562,7 @@ void FRigUnit_HierarchyAddControlVector_Settings::Configure(FRigControlSettings&
 FRigUnit_HierarchyAddControlVector_Execute()
 {
 	FString ErrorMessage;
-	if(!FRigUnit_DynamicHierarchyBase::IsValidToRunInContext(Context, ExecuteContext, true, &ErrorMessage))
+	if(!FRigUnit_DynamicHierarchyBase::IsValidToRunInContext(ExecuteContext, true, &ErrorMessage))
 	{
 		if(!ErrorMessage.IsEmpty())
 		{
@@ -614,7 +612,7 @@ void FRigUnit_HierarchyAddControlRotator_Settings::Configure(FRigControlSettings
 FRigUnit_HierarchyAddControlRotator_Execute()
 {
 	FString ErrorMessage;
-	if(!FRigUnit_DynamicHierarchyBase::IsValidToRunInContext(Context, ExecuteContext, true, &ErrorMessage))
+	if(!FRigUnit_DynamicHierarchyBase::IsValidToRunInContext(ExecuteContext, true, &ErrorMessage))
 	{
 		if(!ErrorMessage.IsEmpty())
 		{
@@ -652,7 +650,7 @@ void FRigUnit_HierarchyAddControlTransform_Settings::Configure(FRigControlSettin
 FRigUnit_HierarchyAddControlTransform_Execute()
 {
 	FString ErrorMessage;
-	if(!FRigUnit_DynamicHierarchyBase::IsValidToRunInContext(Context, ExecuteContext, true, &ErrorMessage))
+	if(!FRigUnit_DynamicHierarchyBase::IsValidToRunInContext(ExecuteContext, true, &ErrorMessage))
 	{
 		if(!ErrorMessage.IsEmpty())
 		{
@@ -680,7 +678,7 @@ FRigUnit_HierarchyAddControlTransform_Execute()
 FRigUnit_HierarchyAddAnimationChannelBool_Execute()
 {
 	FString ErrorMessage;
-	if(!FRigUnit_DynamicHierarchyBase::IsValidToRunInContext(Context, ExecuteContext, true, &ErrorMessage))
+	if(!FRigUnit_DynamicHierarchyBase::IsValidToRunInContext(ExecuteContext, true, &ErrorMessage))
 	{
 		if(!ErrorMessage.IsEmpty())
 		{
@@ -719,7 +717,7 @@ FRigUnit_HierarchyAddAnimationChannelBool_Execute()
 FRigUnit_HierarchyAddAnimationChannelFloat_Execute()
 {
 	FString ErrorMessage;
-	if(!FRigUnit_DynamicHierarchyBase::IsValidToRunInContext(Context, ExecuteContext, true, &ErrorMessage))
+	if(!FRigUnit_DynamicHierarchyBase::IsValidToRunInContext(ExecuteContext, true, &ErrorMessage))
 	{
 		if(!ErrorMessage.IsEmpty())
 		{
@@ -758,7 +756,7 @@ FRigUnit_HierarchyAddAnimationChannelFloat_Execute()
 FRigUnit_HierarchyAddAnimationChannelInteger_Execute()
 {
 	FString ErrorMessage;
-	if(!FRigUnit_DynamicHierarchyBase::IsValidToRunInContext(Context, ExecuteContext, true, &ErrorMessage))
+	if(!FRigUnit_DynamicHierarchyBase::IsValidToRunInContext(ExecuteContext, true, &ErrorMessage))
 	{
 		if(!ErrorMessage.IsEmpty())
 		{
@@ -797,7 +795,7 @@ FRigUnit_HierarchyAddAnimationChannelInteger_Execute()
 FRigUnit_HierarchyAddAnimationChannelVector2D_Execute()
 {
 	FString ErrorMessage;
-	if(!FRigUnit_DynamicHierarchyBase::IsValidToRunInContext(Context, ExecuteContext, true, &ErrorMessage))
+	if(!FRigUnit_DynamicHierarchyBase::IsValidToRunInContext(ExecuteContext, true, &ErrorMessage))
 	{
 		if(!ErrorMessage.IsEmpty())
 		{
@@ -836,7 +834,7 @@ FRigUnit_HierarchyAddAnimationChannelVector2D_Execute()
 FRigUnit_HierarchyAddAnimationChannelVector_Execute()
 {
 	FString ErrorMessage;
-	if(!FRigUnit_DynamicHierarchyBase::IsValidToRunInContext(Context, ExecuteContext, true, &ErrorMessage))
+	if(!FRigUnit_DynamicHierarchyBase::IsValidToRunInContext(ExecuteContext, true, &ErrorMessage))
 	{
 		if(!ErrorMessage.IsEmpty())
 		{
@@ -875,7 +873,7 @@ FRigUnit_HierarchyAddAnimationChannelVector_Execute()
 FRigUnit_HierarchyAddAnimationChannelRotator_Execute()
 {
 	FString ErrorMessage;
-	if(!FRigUnit_DynamicHierarchyBase::IsValidToRunInContext(Context, ExecuteContext, true, &ErrorMessage))
+	if(!FRigUnit_DynamicHierarchyBase::IsValidToRunInContext(ExecuteContext, true, &ErrorMessage))
 	{
 		if(!ErrorMessage.IsEmpty())
 		{
@@ -915,7 +913,7 @@ FRigUnit_HierarchyGetShapeSettings_Execute()
 {
 	Settings = FRigUnit_HierarchyAddControl_ShapeSettings();
 
-	if(URigHierarchy* Hierarchy = Context.Hierarchy)
+	if(URigHierarchy* Hierarchy = ExecuteContext.Hierarchy)
 	{
 		if(const FRigControlElement* ControlElement = Hierarchy->Find<FRigControlElement>(Item))
 		{
diff --git a/Engine/Plugins/Animation/ControlRig/Source/ControlRig/Private/Units/Execution/RigUnit_DynamicHierarchy.h b/Engine/Plugins/Animation/ControlRig/Source/ControlRig/Private/Units/Execution/RigUnit_DynamicHierarchy.h
--- a/Engine/Plugins/Animation/ControlRig/Source/ControlRig/Private/Units/Execution/RigUnit_DynamicHierarchy.h
+++ b/Engine/Plugins/Animation/ControlRig/Source/ControlRig/Private/Units/Execution/RigUnit_DynamicHierarchy.h
@@ -12,7 +12,6 @@ struct CONTROLRIG_API FRigUnit_DynamicHierarchyBase : public FRigUnit
 	GENERATED_BODY()
 
 	static bool IsValidToRunInContext(
-		const FRigUnitContext& InContext,
 		const FControlRigExecuteContext& InExecuteContext,
 		bool bAllowOnlyConstructionEvent,
 		FString* OutErrorMessage = nullptr);
@@ -39,7 +38,7 @@ struct CONTROLRIG_API FRigUnit_AddParent : public FRigUnit_DynamicHierarchyBaseM
 	}
 
 	RIGVM_METHOD()
-	virtual void Execute(const FRigUnitContext& Context) override;
+	virtual void Execute() override;
 
 	/*
 	 * The child to be parented under the new parent
@@ -68,7 +67,7 @@ struct CONTROLRIG_API FRigUnit_SetDefaultParent : public FRigUnit_DynamicHierarc
 	}
 
 	RIGVM_METHOD()
-	virtual void Execute(const FRigUnitContext& Context) override;
+	virtual void Execute() override;
 
 	/*
 	 * The child to be parented under the new default parent
@@ -112,7 +111,7 @@ struct CONTROLRIG_API FRigUnit_SwitchParent : public FRigUnit_DynamicHierarchyBa
 	}
 
 	RIGVM_METHOD()
-	virtual void Execute(const FRigUnitContext& Context) override;
+	virtual void Execute() override;
 
 	/* Depending on this the child will switch to the world,
 	 * back to its default or to the item provided by the Parent pin
@@ -149,7 +148,7 @@ struct CONTROLRIG_API FRigUnit_HierarchyGetParentWeights : public FRigUnit_Dynam
 	}
 
 	RIGVM_METHOD()
-	virtual void Execute(const FRigUnitContext& Context) override;
+	virtual void Execute() override;
 
 	/*
 	 * The child to retrieve the weights for
@@ -187,7 +186,7 @@ struct CONTROLRIG_API FRigUnit_HierarchyGetParentWeightsArray : public FRigUnit_
 	}
 
 	RIGVM_METHOD()
-	virtual void Execute(const FRigUnitContext& Context) override;
+	virtual void Execute() override;
 
 	/*
 	 * The child to retrieve the weights for
@@ -222,7 +221,7 @@ struct CONTROLRIG_API FRigUnit_HierarchySetParentWeights : public FRigUnit_Dynam
 	}
 
 	RIGVM_METHOD()
-	virtual void Execute(const FRigUnitContext& Context) override;
+	virtual void Execute() override;
 
 	/*
 	 * The child to set the parents' weights for
@@ -252,7 +251,7 @@ struct FRigUnit_HierarchyReset : public FRigUnit_DynamicHierarchyBaseMutable
 	}
 
 	RIGVM_METHOD()
-	virtual void Execute(const FRigUnitContext& Context) override;
+	virtual void Execute() override;
 };
 
 /**
@@ -271,7 +270,7 @@ struct FRigUnit_HierarchyImportFromSkeleton : public FRigUnit_DynamicHierarchyBa
 	}
 
 	RIGVM_METHOD()
-	virtual void Execute(const FRigUnitContext& Context) override;
+	virtual void Execute() override;
 
 	UPROPERTY(meta = (Input))
 	FName NameSpace;
@@ -299,7 +298,7 @@ struct CONTROLRIG_API FRigUnit_HierarchyRemoveElement : public FRigUnit_DynamicH
 	}
 
 	RIGVM_METHOD()
-	virtual void Execute(const FRigUnitContext& Context) override;
+	virtual void Execute() override;
 
 	/*
 	 * The item to remove
@@ -377,7 +376,7 @@ struct CONTROLRIG_API FRigUnit_HierarchyAddBone : public FRigUnit_HierarchyAddEl
 	EBoneGetterSetterMode Space;
 
 	RIGVM_METHOD()
-	virtual void Execute(const FRigUnitContext& Context) override;
+	virtual void Execute() override;
 };
 
 /**
@@ -411,7 +410,7 @@ struct CONTROLRIG_API FRigUnit_HierarchyAddNull : public FRigUnit_HierarchyAddEl
 	EBoneGetterSetterMode Space;
 
 	RIGVM_METHOD()
-	virtual void Execute(const FRigUnitContext& Context) override;
+	virtual void Execute() override;
 };
 
 USTRUCT(BlueprintType)
@@ -584,7 +583,7 @@ struct CONTROLRIG_API FRigUnit_HierarchyAddControlFloat : public FRigUnit_Hierar
 	FRigUnit_HierarchyAddControlFloat_Settings Settings;
 
 	RIGVM_METHOD()
-	virtual void Execute(const FRigUnitContext& Context) override;
+	virtual void Execute() override;
 };
 
 USTRUCT(BlueprintType)
@@ -669,7 +668,7 @@ struct CONTROLRIG_API FRigUnit_HierarchyAddControlInteger : public FRigUnit_Hier
 	FRigUnit_HierarchyAddControlInteger_Settings Settings;
 
 	RIGVM_METHOD()
-	virtual void Execute(const FRigUnitContext& Context) override;
+	virtual void Execute() override;
 };
 
 USTRUCT(BlueprintType)
@@ -763,7 +762,7 @@ struct CONTROLRIG_API FRigUnit_HierarchyAddControlVector2D : public FRigUnit_Hie
 	FRigUnit_HierarchyAddControlVector2D_Settings Settings;
 
 	RIGVM_METHOD()
-	virtual void Execute(const FRigUnitContext& Context) override;
+	virtual void Execute() override;
 };
 
 USTRUCT(BlueprintType)
@@ -864,7 +863,7 @@ struct CONTROLRIG_API FRigUnit_HierarchyAddControlVector : public FRigUnit_Hiera
 	FRigUnit_HierarchyAddControlVector_Settings Settings;
 
 	RIGVM_METHOD()
-	virtual void Execute(const FRigUnitContext& Context) override;
+	virtual void Execute() override;
 };
 
 USTRUCT(BlueprintType)
@@ -958,7 +957,7 @@ struct CONTROLRIG_API FRigUnit_HierarchyAddControlRotator : public FRigUnit_Hier
 	FRigUnit_HierarchyAddControlRotator_Settings Settings;
 
 	RIGVM_METHOD()
-	virtual void Execute(const FRigUnitContext& Context) override;
+	virtual void Execute() override;
 };
 
 USTRUCT(BlueprintType)
@@ -1014,7 +1013,7 @@ struct CONTROLRIG_API FRigUnit_HierarchyAddControlTransform : public FRigUnit_Hi
 	FRigUnit_HierarchyAddControlTransform_Settings Settings;
 
 	RIGVM_METHOD()
-	virtual void Execute(const FRigUnitContext& Context) override;
+	virtual void Execute() override;
 };
 
 /**
@@ -1053,7 +1052,7 @@ struct CONTROLRIG_API FRigUnit_HierarchyAddAnimationChannelBool : public FRigUni
 	bool MaximumValue;
 
 	RIGVM_METHOD()
-	virtual void Execute(const FRigUnitContext& Context) override;
+	virtual void Execute() override;
 };
 
 /**
@@ -1092,7 +1091,7 @@ struct CONTROLRIG_API FRigUnit_HierarchyAddAnimationChannelFloat : public FRigUn
 	float MaximumValue;
 
 	RIGVM_METHOD()
-	virtual void Execute(const FRigUnitContext& Context) override;
+	virtual void Execute() override;
 };
 
 /**
@@ -1131,7 +1130,7 @@ struct CONTROLRIG_API FRigUnit_HierarchyAddAnimationChannelInteger : public FRig
 	int32 MaximumValue;
 
 	RIGVM_METHOD()
-	virtual void Execute(const FRigUnitContext& Context) override;
+	virtual void Execute() override;
 };
 
 /**
@@ -1170,7 +1169,7 @@ struct CONTROLRIG_API FRigUnit_HierarchyAddAnimationChannelVector2D : public FRi
 	FVector2D MaximumValue;
 
 	RIGVM_METHOD()
-	virtual void Execute(const FRigUnitContext& Context) override;
+	virtual void Execute() override;
 };
 
 /**
@@ -1209,7 +1208,7 @@ struct CONTROLRIG_API FRigUnit_HierarchyAddAnimationChannelVector : public FRigU
 	FVector MaximumValue;
 
 	RIGVM_METHOD()
-	virtual void Execute(const FRigUnitContext& Context) override;
+	virtual void Execute() override;
 };
 
 /**
@@ -1248,7 +1247,7 @@ struct CONTROLRIG_API FRigUnit_HierarchyAddAnimationChannelRotator : public FRig
 	FRotator MaximumValue;
 
 	RIGVM_METHOD()
-	virtual void Execute(const FRigUnitContext& Context) override;
+	virtual void Execute() override;
 };
 
 /**
@@ -1277,7 +1276,7 @@ struct CONTROLRIG_API FRigUnit_HierarchyGetShapeSettings : public FRigUnit_Dynam
 	FRigUnit_HierarchyAddControl_ShapeSettings Settings;
 
 	RIGVM_METHOD()
-	virtual void Execute(const FRigUnitContext& Context) override;
+	virtual void Execute() override;
 };
 
 /**
@@ -1307,5 +1306,5 @@ struct CONTROLRIG_API FRigUnit_HierarchySetShapeSettings : public FRigUnit_Dynam
 	FRigUnit_HierarchyAddControl_ShapeSettings Settings;
 
 	RIGVM_METHOD()
-	virtual void Execute(const FRigUnitContext& Context) override;
+	virtual void Execute() override;
 };
diff --git a/Engine/Plugins/Animation/ControlRig/Source/ControlRig/Private/Units/Execution/RigUnit_Hierarchy.cpp b/Engine/Plugins/Animation/ControlRig/Source/ControlRig/Private/Units/Execution/RigUnit_Hierarchy.cpp
--- a/Engine/Plugins/Animation/ControlRig/Source/ControlRig/Private/Units/Execution/RigUnit_Hierarchy.cpp
+++ b/Engine/Plugins/Animation/ControlRig/Source/ControlRig/Private/Units/Execution/RigUnit_Hierarchy.cpp
@@ -11,13 +11,13 @@ FRigUnit_HierarchyGetParent_Execute()
 {
     DECLARE_SCOPE_HIERARCHICAL_COUNTER_RIGUNIT()
 
-	if (Context.State == EControlRigState::Init)
+	if (ExecuteContext.UnitContext.State == EControlRigState::Init)
 	{
 		CachedChild.Reset();
 		CachedParent.Reset();
 	}
 
-	if(CachedChild.IsIdentical(Child, Context.Hierarchy))
+	if(CachedChild.IsIdentical(Child, ExecuteContext.Hierarchy))
 	{
 		Parent = CachedParent.GetKey();
 	}
@@ -26,12 +26,12 @@ FRigUnit_HierarchyGetParent_Execute()
 		Parent.Reset();
 		CachedParent.Reset();
 
-		if(CachedChild.UpdateCache(Child, Context.Hierarchy))
+		if(CachedChild.UpdateCache(Child, ExecuteContext.Hierarchy))
 		{
-			Parent = Context.Hierarchy->GetFirstParent(Child);
+			Parent = ExecuteContext.Hierarchy->GetFirstParent(Child);
 			if(Parent.IsValid())
 			{
-				CachedParent.UpdateCache(Parent, Context.Hierarchy);
+				CachedParent.UpdateCache(Parent, ExecuteContext.Hierarchy);
 			}
 		}
 	}
@@ -39,7 +39,7 @@ FRigUnit_HierarchyGetParent_Execute()
 
 FRigUnit_HierarchyGetParents_Execute()
 {
-	FRigUnit_HierarchyGetParentsItemArray::StaticExecute(ExecuteContext, Child, bIncludeChild, bReverse, Parents.Keys, CachedChild, CachedParents, Context);
+	FRigUnit_HierarchyGetParentsItemArray::StaticExecute(ExecuteContext, Child, bIncludeChild, bReverse, Parents.Keys, CachedChild, CachedParents);
 }
 
 FRigVMStructUpgradeInfo FRigUnit_HierarchyGetParents::GetUpgradeInfo() const
@@ -56,17 +56,17 @@ FRigUnit_HierarchyGetParentsItemArray_Execute()
 {
     DECLARE_SCOPE_HIERARCHICAL_COUNTER_RIGUNIT()
 
-	if (Context.State == EControlRigState::Init)
+	if (ExecuteContext.UnitContext.State == EControlRigState::Init)
 	{
 		CachedChild.Reset();
 		CachedParents.Reset();
 	}
 
-	if(!CachedChild.IsIdentical(Child, Context.Hierarchy))
+	if(!CachedChild.IsIdentical(Child, ExecuteContext.Hierarchy))
 	{
 		CachedParents.Reset();
 
-		if(CachedChild.UpdateCache(Child, Context.Hierarchy))
+		if(CachedChild.UpdateCache(Child, ExecuteContext.Hierarchy))
 		{
 			TArray<FRigElementKey> Keys;
 			FRigElementKey Parent = Child;
@@ -76,7 +76,7 @@ FRigUnit_HierarchyGetParentsItemArray_Execute()
 				{
 					Keys.Add(Parent);
 				}
-				Parent = Context.Hierarchy->GetFirstParent(Parent);
+				Parent = ExecuteContext.Hierarchy->GetFirstParent(Parent);
 			}
 			while(Parent.IsValid());
 
@@ -95,17 +95,17 @@ FRigUnit_HierarchyGetChildren_Execute()
 {
     DECLARE_SCOPE_HIERARCHICAL_COUNTER_RIGUNIT()
 
-	if (Context.State == EControlRigState::Init)
+	if (ExecuteContext.UnitContext.State == EControlRigState::Init)
 	{
 		CachedParent.Reset();
 		CachedChildren.Reset();
 	}
 
-	if(!CachedParent.IsIdentical(Parent, Context.Hierarchy))
+	if(!CachedParent.IsIdentical(Parent, ExecuteContext.Hierarchy))
 	{
 		CachedChildren.Reset();
 
-		if(CachedParent.UpdateCache(Parent, Context.Hierarchy))
+		if(CachedParent.UpdateCache(Parent, ExecuteContext.Hierarchy))
 		{
 			TArray<FRigElementKey> Keys;
 
@@ -115,7 +115,7 @@ FRigUnit_HierarchyGetChildren_Execute()
 			}
 
 			
-			Keys.Append(Context.Hierarchy->GetChildren(Parent, bRecursive));
+			Keys.Append(ExecuteContext.Hierarchy->GetChildren(Parent, bRecursive));
 
 			CachedChildren = FRigElementKeyCollection(Keys);
 		}
@@ -136,7 +136,7 @@ FRigVMStructUpgradeInfo FRigUnit_HierarchyGetChildren::GetUpgradeInfo() const
 
 FRigUnit_HierarchyGetSiblings_Execute()
 {
-	FRigUnit_HierarchyGetSiblingsItemArray::StaticExecute(ExecuteContext, Item, bIncludeItem, Siblings.Keys, CachedItem, CachedSiblings, Context);
+	FRigUnit_HierarchyGetSiblingsItemArray::StaticExecute(ExecuteContext, Item, bIncludeItem, Siblings.Keys, CachedItem, CachedSiblings);
 }
 
 FRigVMStructUpgradeInfo FRigUnit_HierarchyGetSiblings::GetUpgradeInfo() const
@@ -151,24 +151,24 @@ FRigUnit_HierarchyGetSiblingsItemArray_Execute()
 {
     DECLARE_SCOPE_HIERARCHICAL_COUNTER_RIGUNIT()
 
-	if (Context.State == EControlRigState::Init)
+	if (ExecuteContext.UnitContext.State == EControlRigState::Init)
 	{
 		CachedItem.Reset();
 		CachedSiblings.Reset();
 	}
 
-	if(!CachedItem.IsIdentical(Item, Context.Hierarchy))
+	if(!CachedItem.IsIdentical(Item, ExecuteContext.Hierarchy))
 	{
 		CachedSiblings.Reset();
 
-		if(CachedItem.UpdateCache(Item, Context.Hierarchy))
+		if(CachedItem.UpdateCache(Item, ExecuteContext.Hierarchy))
 		{
 			TArray<FRigElementKey> Keys;
 
-			FRigElementKey Parent = Context.Hierarchy->GetFirstParent(Item);
+			FRigElementKey Parent = ExecuteContext.Hierarchy->GetFirstParent(Item);
 			if(Parent.IsValid())
 			{
-				TArray<FRigElementKey> Children = Context.Hierarchy->GetChildren(Parent, false);
+				TArray<FRigElementKey> Children = ExecuteContext.Hierarchy->GetChildren(Parent, false);
 				for(FRigElementKey Child : Children)
 				{
 					if(bIncludeItem || Child != Item)
@@ -192,7 +192,7 @@ FRigUnit_HierarchyGetSiblingsItemArray_Execute()
 
 FRigUnit_HierarchyGetPose_Execute()
 {
-	FRigUnit_HierarchyGetPoseItemArray::StaticExecute(ExecuteContext, Initial, ElementType, ItemsToGet.Keys, Pose, Context);
+	FRigUnit_HierarchyGetPoseItemArray::StaticExecute(ExecuteContext, Initial, ElementType, ItemsToGet.Keys, Pose);
 }
 
 FRigVMStructUpgradeInfo FRigUnit_HierarchyGetPose::GetUpgradeInfo() const
@@ -207,12 +207,12 @@ FRigVMStructUpgradeInfo FRigUnit_HierarchyGetPose::GetUpgradeInfo() const
 
 FRigUnit_HierarchyGetPoseItemArray_Execute()
 {
-	Pose = Context.Hierarchy->GetPose(Initial, ElementType, ItemsToGet);
+	Pose = ExecuteContext.Hierarchy->GetPose(Initial, ElementType, ItemsToGet);
 }
 
 FRigUnit_HierarchySetPose_Execute()
 {
-	FRigUnit_HierarchySetPoseItemArray::StaticExecute(ExecuteContext, Pose, ElementType, Space, ItemsToSet.Keys, Weight, Context);
+	FRigUnit_HierarchySetPoseItemArray::StaticExecute(ExecuteContext, Pose, ElementType, Space, ItemsToSet.Keys, Weight);
 }
 
 FRigVMStructUpgradeInfo FRigUnit_HierarchySetPose::GetUpgradeInfo() const
@@ -244,7 +244,7 @@ FRigUnit_PoseIsEmpty_Execute()
 
 FRigUnit_PoseGetItems_Execute()
 {
-	FRigUnit_PoseGetItemsItemArray::StaticExecute(ExecuteContext, Pose, ElementType, Items.Keys, Context);
+	FRigUnit_PoseGetItemsItemArray::StaticExecute(ExecuteContext, Pose, ElementType, Items.Keys);
 }
 
 FRigVMStructUpgradeInfo FRigUnit_PoseGetItems::GetUpgradeInfo() const
@@ -285,14 +285,14 @@ FRigUnit_PoseGetDelta_Execute()
 	if(PoseA.Num() == 0 && PoseB.Num() != 0)
 	{
 		PosesAreEqual = false;
-		FRigUnit_PoseGetItems::StaticExecute(ExecuteContext, PoseB, ElementType, ItemsWithDelta, Context);
+		FRigUnit_PoseGetItems::StaticExecute(ExecuteContext, PoseB, ElementType, ItemsWithDelta);
 		return;
 	}
 
 	if(PoseA.Num() != 0 && PoseB.Num() == 0)
 	{
 		PosesAreEqual = false;
-		FRigUnit_PoseGetItems::StaticExecute(ExecuteContext, PoseA, ElementType, ItemsWithDelta, Context);
+		FRigUnit_PoseGetItems::StaticExecute(ExecuteContext, PoseA, ElementType, ItemsWithDelta);
 		return;
 	}
 
@@ -490,7 +490,7 @@ bool FRigUnit_PoseGetDelta::AreCurvesEqual(float A, float B, float CurveU)
 
 FRigUnit_PoseGetTransform_Execute()
 {
-	if (Context.State == EControlRigState::Init)
+	if (ExecuteContext.UnitContext.State == EControlRigState::Init)
 	{
 		CachedPoseElementIndex = CachedPoseHash = INDEX_NONE;
 	}
@@ -552,7 +552,7 @@ FRigUnit_PoseGetTransformArray_Execute()
 
 FRigUnit_PoseGetCurve_Execute()
 {
-	if (Context.State == EControlRigState::Init)
+	if (ExecuteContext.UnitContext.State == EControlRigState::Init)
 	{
 		CachedPoseElementIndex = CachedPoseHash = INDEX_NONE;
 	}
diff --git a/Engine/Plugins/Animation/ControlRig/Source/ControlRig/Private/Units/Execution/RigUnit_Hierarchy.h b/Engine/Plugins/Animation/ControlRig/Source/ControlRig/Private/Units/Execution/RigUnit_Hierarchy.h
--- a/Engine/Plugins/Animation/ControlRig/Source/ControlRig/Private/Units/Execution/RigUnit_Hierarchy.h
+++ b/Engine/Plugins/Animation/ControlRig/Source/ControlRig/Private/Units/Execution/RigUnit_Hierarchy.h
@@ -33,7 +33,7 @@ struct CONTROLRIG_API FRigUnit_HierarchyGetParent : public FRigUnit_HierarchyBas
 	}
 
 	RIGVM_METHOD()
-	virtual void Execute(const FRigUnitContext& Context) override;
+	virtual void Execute() override;
 
 	UPROPERTY(meta = (Input, ExpandByDefault))
 	FRigElementKey Child;
@@ -68,7 +68,7 @@ struct CONTROLRIG_API FRigUnit_HierarchyGetParents : public FRigUnit_HierarchyBa
 	}
 
 	RIGVM_METHOD()
-	virtual void Execute(const FRigUnitContext& Context) override;
+	virtual void Execute() override;
 
 	UPROPERTY(meta = (Input, ExpandByDefault))
 	FRigElementKey Child;
@@ -112,7 +112,7 @@ struct CONTROLRIG_API FRigUnit_HierarchyGetParentsItemArray : public FRigUnit_Hi
 	}
 
 	RIGVM_METHOD()
-	virtual void Execute(const FRigUnitContext& Context) override;
+	virtual void Execute() override;
 
 	UPROPERTY(meta = (Input, ExpandByDefault))
 	FRigElementKey Child;
@@ -154,7 +154,7 @@ struct CONTROLRIG_API FRigUnit_HierarchyGetChildren : public FRigUnit_HierarchyB
 	}
 
 	RIGVM_METHOD()
-	virtual void Execute(const FRigUnitContext& Context) override;
+	virtual void Execute() override;
 
 	UPROPERTY(meta = (Input, ExpandByDefault))
 	FRigElementKey Parent;
@@ -197,7 +197,7 @@ struct CONTROLRIG_API FRigUnit_HierarchyGetSiblings : public FRigUnit_HierarchyB
 	}
 
 	RIGVM_METHOD()
-	virtual void Execute(const FRigUnitContext& Context) override;
+	virtual void Execute() override;
 
 	UPROPERTY(meta = (Input, ExpandByDefault))
 	FRigElementKey Item;
@@ -237,7 +237,7 @@ struct CONTROLRIG_API FRigUnit_HierarchyGetSiblingsItemArray : public FRigUnit_H
 	}
 
 	RIGVM_METHOD()
-	virtual void Execute(const FRigUnitContext& Context) override;
+	virtual void Execute() override;
 
 	UPROPERTY(meta = (Input, ExpandByDefault))
 	FRigElementKey Item;
@@ -274,7 +274,7 @@ struct CONTROLRIG_API FRigUnit_HierarchyGetPose : public FRigUnit_HierarchyBase
 	}
 
 	RIGVM_METHOD()
-	virtual void Execute(const FRigUnitContext& Context) override;
+	virtual void Execute() override;
 
 	UPROPERTY(meta = (Input))
 	bool Initial;
@@ -309,7 +309,7 @@ struct CONTROLRIG_API FRigUnit_HierarchyGetPoseItemArray : public FRigUnit_Hiera
 	}
 
 	RIGVM_METHOD()
-	virtual void Execute(const FRigUnitContext& Context) override;
+	virtual void Execute() override;
 
 	UPROPERTY(meta = (Input))
 	bool Initial;
@@ -344,7 +344,7 @@ struct CONTROLRIG_API FRigUnit_HierarchySetPose : public FRigUnit_HierarchyBaseM
 	}
 
 	RIGVM_METHOD()
-	virtual void Execute(const FRigUnitContext& Context) override;
+	virtual void Execute() override;
 
 	UPROPERTY(meta = (Input))
 	FRigPose Pose;
@@ -383,7 +383,7 @@ struct CONTROLRIG_API FRigUnit_HierarchySetPoseItemArray : public FRigUnit_Hiera
 	}
 
 	RIGVM_METHOD()
-	virtual void Execute(const FRigUnitContext& Context) override;
+	virtual void Execute() override;
 
 	UPROPERTY(meta = (Input))
 	FRigPose Pose;
@@ -417,7 +417,7 @@ struct CONTROLRIG_API FRigUnit_PoseIsEmpty : public FRigUnit_HierarchyBase
 	}
 
 	RIGVM_METHOD()
-	virtual void Execute(const FRigUnitContext& Context) override;
+	virtual void Execute() override;
 
 	UPROPERTY(meta = (Input))
 	FRigPose Pose;
@@ -442,7 +442,7 @@ struct CONTROLRIG_API FRigUnit_PoseGetItems : public FRigUnit_HierarchyBase
 	}
 
 	RIGVM_METHOD()
-	virtual void Execute(const FRigUnitContext& Context) override;
+	virtual void Execute() override;
 
 	UPROPERTY(meta = (Input))
 	FRigPose Pose;
@@ -472,7 +472,7 @@ struct CONTROLRIG_API FRigUnit_PoseGetItemsItemArray : public FRigUnit_Hierarchy
 	}
 
 	RIGVM_METHOD()
-	virtual void Execute(const FRigUnitContext& Context) override;
+	virtual void Execute() override;
 
 	UPROPERTY(meta = (Input))
 	FRigPose Pose;
@@ -504,7 +504,7 @@ struct CONTROLRIG_API FRigUnit_PoseGetDelta : public FRigUnit_HierarchyBase
 	}
 
 	RIGVM_METHOD()
-	virtual void Execute(const FRigUnitContext& Context) override;
+	virtual void Execute() override;
 
 	UPROPERTY(meta = (Input))
 	FRigPose PoseA;
@@ -587,7 +587,7 @@ struct CONTROLRIG_API FRigUnit_PoseGetTransform : public FRigUnit_HierarchyBase
 	}
 
 	RIGVM_METHOD()
-	virtual void Execute(const FRigUnitContext& Context) override;
+	virtual void Execute() override;
 
 	UPROPERTY(meta = (Input))
 	FRigPose Pose;
@@ -634,7 +634,7 @@ struct CONTROLRIG_API FRigUnit_PoseGetTransformArray : public FRigUnit_Hierarchy
 	}
 
 	RIGVM_METHOD()
-	virtual void Execute(const FRigUnitContext& Context) override;
+	virtual void Execute() override;
 
 	UPROPERTY(meta = (Input))
 	FRigPose Pose;
@@ -670,7 +670,7 @@ struct CONTROLRIG_API FRigUnit_PoseGetCurve : public FRigUnit_HierarchyBase
 	}
 
 	RIGVM_METHOD()
-	virtual void Execute(const FRigUnitContext& Context) override;
+	virtual void Execute() override;
 
 	UPROPERTY(meta = (Input))
 	FRigPose Pose;
@@ -721,7 +721,7 @@ struct CONTROLRIG_API FRigUnit_PoseLoop : public FRigUnit_HierarchyBaseMutable
 	virtual int32 GetNumSlices() const override { return Count; }
 
 	RIGVM_METHOD()
-	virtual void Execute(const FRigUnitContext& Context) override;
+	virtual void Execute() override;
 
 	UPROPERTY(meta = (Singleton))
 	FName BlockToRun;
diff --git a/Engine/Plugins/Animation/ControlRig/Source/ControlRig/Private/Units/Execution/RigUnit_InteractionExecution.cpp b/Engine/Plugins/Animation/ControlRig/Source/ControlRig/Private/Units/Execution/RigUnit_InteractionExecution.cpp
--- a/Engine/Plugins/Animation/ControlRig/Source/ControlRig/Private/Units/Execution/RigUnit_InteractionExecution.cpp
+++ b/Engine/Plugins/Animation/ControlRig/Source/ControlRig/Private/Units/Execution/RigUnit_InteractionExecution.cpp
@@ -10,7 +10,6 @@ FName FRigUnit_InteractionExecution::EventName = TEXT("Interaction");
 FRigUnit_InteractionExecution_Execute()
 {
     DECLARE_SCOPE_HIERARCHICAL_COUNTER_RIGUNIT()
-	ExecuteContext.Hierarchy = Context.Hierarchy;
 	ExecuteContext.EventName = FRigUnit_InteractionExecution::EventName;
 }
 
diff --git a/Engine/Plugins/Animation/ControlRig/Source/ControlRig/Private/Units/Execution/RigUnit_IsInteracting.cpp b/Engine/Plugins/Animation/ControlRig/Source/ControlRig/Private/Units/Execution/RigUnit_IsInteracting.cpp
--- a/Engine/Plugins/Animation/ControlRig/Source/ControlRig/Private/Units/Execution/RigUnit_IsInteracting.cpp
+++ b/Engine/Plugins/Animation/ControlRig/Source/ControlRig/Private/Units/Execution/RigUnit_IsInteracting.cpp
@@ -9,10 +9,10 @@ FRigUnit_IsInteracting_Execute()
 {
     DECLARE_SCOPE_HIERARCHICAL_COUNTER_RIGUNIT()
 
-	bIsInteracting = Context.IsInteracting();
-	bIsTranslating = ( Context.InteractionType & (uint8)EControlRigInteractionType::Translate) != 0;
-	bIsRotating = ( Context.InteractionType & (uint8)EControlRigInteractionType::Rotate) != 0;
-	bIsScaling = ( Context.InteractionType & (uint8)EControlRigInteractionType::Scale) != 0;
-	Items = Context.ElementsBeingInteracted;
+	bIsInteracting = ExecuteContext.UnitContext.IsInteracting();
+	bIsTranslating = ( ExecuteContext.UnitContext.InteractionType & (uint8)EControlRigInteractionType::Translate) != 0;
+	bIsRotating = ( ExecuteContext.UnitContext.InteractionType & (uint8)EControlRigInteractionType::Rotate) != 0;
+	bIsScaling = ( ExecuteContext.UnitContext.InteractionType & (uint8)EControlRigInteractionType::Scale) != 0;
+	Items = ExecuteContext.UnitContext.ElementsBeingInteracted;
 }
 
diff --git a/Engine/Plugins/Animation/ControlRig/Source/ControlRig/Private/Units/Execution/RigUnit_Item.cpp b/Engine/Plugins/Animation/ControlRig/Source/ControlRig/Private/Units/Execution/RigUnit_Item.cpp
--- a/Engine/Plugins/Animation/ControlRig/Source/ControlRig/Private/Units/Execution/RigUnit_Item.cpp
+++ b/Engine/Plugins/Animation/ControlRig/Source/ControlRig/Private/Units/Execution/RigUnit_Item.cpp
@@ -13,7 +13,7 @@ FRigUnit_ItemExists_Execute()
 
 	Exists = false;
 
-	switch (Context.State)
+	switch (ExecuteContext.UnitContext.State)
 	{
 		case EControlRigState::Init:
 		{
@@ -22,7 +22,7 @@ FRigUnit_ItemExists_Execute()
 		}
 		case EControlRigState::Update:
 		{
-			Exists = CachedIndex.UpdateCache(Item, Context.Hierarchy);
+			Exists = CachedIndex.UpdateCache(Item, ExecuteContext.Hierarchy);
 	    	break;
 	    }
 	    default:
diff --git a/Engine/Plugins/Animation/ControlRig/Source/ControlRig/Private/Units/Execution/RigUnit_Item.h b/Engine/Plugins/Animation/ControlRig/Source/ControlRig/Private/Units/Execution/RigUnit_Item.h
--- a/Engine/Plugins/Animation/ControlRig/Source/ControlRig/Private/Units/Execution/RigUnit_Item.h
+++ b/Engine/Plugins/Animation/ControlRig/Source/ControlRig/Private/Units/Execution/RigUnit_Item.h
@@ -34,7 +34,7 @@ struct CONTROLRIG_API FRigUnit_ItemExists : public FRigUnit_ItemBase
 	}
 
 	RIGVM_METHOD()
-	virtual void Execute(const FRigUnitContext& Context) override;
+	virtual void Execute() override;
 
 	UPROPERTY(meta = (Input, ExpandByDefault))
 	FRigElementKey Item;
@@ -62,7 +62,7 @@ struct CONTROLRIG_API FRigUnit_ItemReplace : public FRigUnit_ItemBase
 	}
 
 	RIGVM_METHOD()
-	virtual void Execute(const FRigUnitContext& Context) override;
+	virtual void Execute() override;
 
 	UPROPERTY(meta = (Input))
 	FRigElementKey Item;
@@ -92,7 +92,7 @@ struct CONTROLRIG_API FRigUnit_ItemEquals : public FRigUnit_ItemBase
 	}
 
 	RIGVM_METHOD()
-	virtual void Execute(const FRigUnitContext& Context) override;
+	virtual void Execute() override;
 
 	UPROPERTY(meta = (Input, ExpandByDefault))
 	FRigElementKey A;
@@ -122,7 +122,7 @@ struct CONTROLRIG_API FRigUnit_ItemNotEquals : public FRigUnit_ItemBase
 	}
 
 	RIGVM_METHOD()
-	virtual void Execute(const FRigUnitContext& Context) override;
+	virtual void Execute() override;
 
 	UPROPERTY(meta = (Input, ExpandByDefault))
 	FRigElementKey A;
@@ -152,7 +152,7 @@ struct CONTROLRIG_API FRigUnit_ItemTypeEquals : public FRigUnit_ItemBase
 	}
 
 	RIGVM_METHOD()
-	virtual void Execute(const FRigUnitContext& Context) override;
+	virtual void Execute() override;
 
 	UPROPERTY(meta = (Input, ExpandByDefault))
 	FRigElementKey A;
@@ -179,7 +179,7 @@ struct CONTROLRIG_API FRigUnit_ItemTypeNotEquals : public FRigUnit_ItemBase
 	}
 
 	RIGVM_METHOD()
-	virtual void Execute(const FRigUnitContext& Context) override;
+	virtual void Execute() override;
 
 	UPROPERTY(meta = (Input, ExpandByDefault))
 	FRigElementKey A;
@@ -194,7 +194,7 @@ struct CONTROLRIG_API FRigUnit_ItemTypeNotEquals : public FRigUnit_ItemBase
 /**
  * Casts the provided item key to its name
  */
-USTRUCT(meta=(DisplayName="To Name", Keywords="", TemplateName="Cast"))
+USTRUCT(meta=(DisplayName="To Name", Keywords="", TemplateName="Cast", ExecuteContext="FRigVMExecuteContext"))
 struct CONTROLRIG_API FRigUnit_ItemToName : public FRigUnit_ItemBase
 {
 	GENERATED_BODY()
@@ -206,7 +206,7 @@ struct CONTROLRIG_API FRigUnit_ItemToName : public FRigUnit_ItemBase
 	}
 
 	RIGVM_METHOD()
-	virtual void Execute(const FRigUnitContext& Context) override;
+	virtual void Execute() override;
 
 	UPROPERTY(meta = (Input, ExpandByDefault))
 	FRigElementKey Value;
diff --git a/Engine/Plugins/Animation/ControlRig/Source/ControlRig/Private/Units/Hierarchy/RigUnit_AddBoneTransform.cpp b/Engine/Plugins/Animation/ControlRig/Source/ControlRig/Private/Units/Hierarchy/RigUnit_AddBoneTransform.cpp
--- a/Engine/Plugins/Animation/ControlRig/Source/ControlRig/Private/Units/Hierarchy/RigUnit_AddBoneTransform.cpp
+++ b/Engine/Plugins/Animation/ControlRig/Source/ControlRig/Private/Units/Hierarchy/RigUnit_AddBoneTransform.cpp
@@ -12,7 +12,7 @@ FRigUnit_AddBoneTransform_Execute()
 	URigHierarchy* Hierarchy = ExecuteContext.Hierarchy;
 	if (Hierarchy)
 	{
-		switch (Context.State)
+		switch (ExecuteContext.UnitContext.State)
 		{
 			case EControlRigState::Init:
 			{
diff --git a/Engine/Plugins/Animation/ControlRig/Source/ControlRig/Private/Units/Hierarchy/RigUnit_ControlChannel.cpp b/Engine/Plugins/Animation/ControlRig/Source/ControlRig/Private/Units/Hierarchy/RigUnit_ControlChannel.cpp
--- a/Engine/Plugins/Animation/ControlRig/Source/ControlRig/Private/Units/Hierarchy/RigUnit_ControlChannel.cpp
+++ b/Engine/Plugins/Animation/ControlRig/Source/ControlRig/Private/Units/Hierarchy/RigUnit_ControlChannel.cpp
@@ -45,16 +45,16 @@ FRigUnit_GetBoolAnimationChannel_Execute()
 {
 	Value = false;
 	
-	if(!UpdateCache(Context.Hierarchy, Control, Channel, CachedChannelKey, CachedChannelHash))
+	if(!UpdateCache(ExecuteContext.Hierarchy, Control, Channel, CachedChannelKey, CachedChannelHash))
 	{
 		return;
 	}
 
-	if(const FRigControlElement* ChannelElement = Context.Hierarchy->Find<FRigControlElement>(CachedChannelKey))
+	if(const FRigControlElement* ChannelElement = ExecuteContext.Hierarchy->Find<FRigControlElement>(CachedChannelKey))
 	{
 		if(ChannelElement->Settings.ControlType == ERigControlType::Bool)
 		{
-			const FRigControlValue StoredValue = Context.Hierarchy->GetControlValueByIndex(ChannelElement->GetIndex(), bInitial ? ERigControlValueType::Initial : ERigControlValueType::Current);
+			const FRigControlValue StoredValue = ExecuteContext.Hierarchy->GetControlValueByIndex(ChannelElement->GetIndex(), bInitial ? ERigControlValueType::Initial : ERigControlValueType::Current);
 			Value = StoredValue.Get<bool>();
 		}
 		else
@@ -68,16 +68,16 @@ FRigUnit_GetFloatAnimationChannel_Execute()
 {
 	Value = 0.f;
 	
-	if(!UpdateCache(Context.Hierarchy, Control, Channel, CachedChannelKey, CachedChannelHash))
+	if(!UpdateCache(ExecuteContext.Hierarchy, Control, Channel, CachedChannelKey, CachedChannelHash))
 	{
 		return;
 	}
 
-	if(const FRigControlElement* ChannelElement = Context.Hierarchy->Find<FRigControlElement>(CachedChannelKey))
+	if(const FRigControlElement* ChannelElement = ExecuteContext.Hierarchy->Find<FRigControlElement>(CachedChannelKey))
 	{
 		if(ChannelElement->Settings.ControlType == ERigControlType::Float)
 		{
-			const FRigControlValue StoredValue = Context.Hierarchy->GetControlValueByIndex(ChannelElement->GetIndex(), bInitial ? ERigControlValueType::Initial : ERigControlValueType::Current);
+			const FRigControlValue StoredValue = ExecuteContext.Hierarchy->GetControlValueByIndex(ChannelElement->GetIndex(), bInitial ? ERigControlValueType::Initial : ERigControlValueType::Current);
 			Value = StoredValue.Get<float>();
 		}
 		else
@@ -91,16 +91,16 @@ FRigUnit_GetIntAnimationChannel_Execute()
 {
 	Value = 0;
 	
-	if(!UpdateCache(Context.Hierarchy, Control, Channel, CachedChannelKey, CachedChannelHash))
+	if(!UpdateCache(ExecuteContext.Hierarchy, Control, Channel, CachedChannelKey, CachedChannelHash))
 	{
 		return;
 	}
 
-	if(const FRigControlElement* ChannelElement = Context.Hierarchy->Find<FRigControlElement>(CachedChannelKey))
+	if(const FRigControlElement* ChannelElement = ExecuteContext.Hierarchy->Find<FRigControlElement>(CachedChannelKey))
 	{
 		if(ChannelElement->Settings.ControlType == ERigControlType::Integer)
 		{
-			const FRigControlValue StoredValue = Context.Hierarchy->GetControlValueByIndex(ChannelElement->GetIndex(), bInitial ? ERigControlValueType::Initial : ERigControlValueType::Current);
+			const FRigControlValue StoredValue = ExecuteContext.Hierarchy->GetControlValueByIndex(ChannelElement->GetIndex(), bInitial ? ERigControlValueType::Initial : ERigControlValueType::Current);
 			Value = StoredValue.Get<int32>();
 		}
 		else
@@ -114,16 +114,16 @@ FRigUnit_GetVector2DAnimationChannel_Execute()
 {
 	Value = FVector2D::ZeroVector;
 	
-	if(!UpdateCache(Context.Hierarchy, Control, Channel, CachedChannelKey, CachedChannelHash))
+	if(!UpdateCache(ExecuteContext.Hierarchy, Control, Channel, CachedChannelKey, CachedChannelHash))
 	{
 		return;
 	}
 
-	if(const FRigControlElement* ChannelElement = Context.Hierarchy->Find<FRigControlElement>(CachedChannelKey))
+	if(const FRigControlElement* ChannelElement = ExecuteContext.Hierarchy->Find<FRigControlElement>(CachedChannelKey))
 	{
 		if(ChannelElement->Settings.ControlType == ERigControlType::Vector2D)
 		{
-			const FRigControlValue StoredValue = Context.Hierarchy->GetControlValueByIndex(ChannelElement->GetIndex(), bInitial ? ERigControlValueType::Initial : ERigControlValueType::Current);
+			const FRigControlValue StoredValue = ExecuteContext.Hierarchy->GetControlValueByIndex(ChannelElement->GetIndex(), bInitial ? ERigControlValueType::Initial : ERigControlValueType::Current);
 			Value = FVector2D(StoredValue.Get<FVector2f>());
 		}
 		else
@@ -137,16 +137,16 @@ FRigUnit_GetVectorAnimationChannel_Execute()
 {
 	Value = FVector::ZeroVector;
 	
-	if(!UpdateCache(Context.Hierarchy, Control, Channel, CachedChannelKey, CachedChannelHash))
+	if(!UpdateCache(ExecuteContext.Hierarchy, Control, Channel, CachedChannelKey, CachedChannelHash))
 	{
 		return;
 	}
 
-	if(const FRigControlElement* ChannelElement = Context.Hierarchy->Find<FRigControlElement>(CachedChannelKey))
+	if(const FRigControlElement* ChannelElement = ExecuteContext.Hierarchy->Find<FRigControlElement>(CachedChannelKey))
 	{
 		if(ChannelElement->Settings.ControlType == ERigControlType::Position || ChannelElement->Settings.ControlType == ERigControlType::Scale)
 		{
-			const FRigControlValue StoredValue = Context.Hierarchy->GetControlValueByIndex(ChannelElement->GetIndex(), bInitial ? ERigControlValueType::Initial : ERigControlValueType::Current);
+			const FRigControlValue StoredValue = ExecuteContext.Hierarchy->GetControlValueByIndex(ChannelElement->GetIndex(), bInitial ? ERigControlValueType::Initial : ERigControlValueType::Current);
 			Value = FVector(StoredValue.Get<FVector3f>());
 		}
 		else
@@ -160,16 +160,16 @@ FRigUnit_GetRotatorAnimationChannel_Execute()
 {
 	Value = FRotator::ZeroRotator;
 	
-	if(!UpdateCache(Context.Hierarchy, Control, Channel, CachedChannelKey, CachedChannelHash))
+	if(!UpdateCache(ExecuteContext.Hierarchy, Control, Channel, CachedChannelKey, CachedChannelHash))
 	{
 		return;
 	}
 
-	if(const FRigControlElement* ChannelElement = Context.Hierarchy->Find<FRigControlElement>(CachedChannelKey))
+	if(const FRigControlElement* ChannelElement = ExecuteContext.Hierarchy->Find<FRigControlElement>(CachedChannelKey))
 	{
 		if(ChannelElement->Settings.ControlType == ERigControlType::Rotator)
 		{
-			const FRigControlValue StoredValue = Context.Hierarchy->GetControlValueByIndex(ChannelElement->GetIndex(), bInitial ? ERigControlValueType::Initial : ERigControlValueType::Current);
+			const FRigControlValue StoredValue = ExecuteContext.Hierarchy->GetControlValueByIndex(ChannelElement->GetIndex(), bInitial ? ERigControlValueType::Initial : ERigControlValueType::Current);
 			Value = FRotator::MakeFromEuler(FVector(StoredValue.Get<FVector3f>()));
 		}
 		else
@@ -183,21 +183,21 @@ FRigUnit_GetTransformAnimationChannel_Execute()
 {
 	Value = FTransform::Identity;
 	
-	if(!UpdateCache(Context.Hierarchy, Control, Channel, CachedChannelKey, CachedChannelHash))
+	if(!UpdateCache(ExecuteContext.Hierarchy, Control, Channel, CachedChannelKey, CachedChannelHash))
 	{
 		return;
 	}
 
-	if(const FRigControlElement* ChannelElement = Context.Hierarchy->Find<FRigControlElement>(CachedChannelKey))
+	if(const FRigControlElement* ChannelElement = ExecuteContext.Hierarchy->Find<FRigControlElement>(CachedChannelKey))
 	{
 		if(ChannelElement->Settings.ControlType == ERigControlType::Transform)
 		{
-			const FRigControlValue StoredValue = Context.Hierarchy->GetControlValueByIndex(ChannelElement->GetIndex(), bInitial ? ERigControlValueType::Initial : ERigControlValueType::Current);
+			const FRigControlValue StoredValue = ExecuteContext.Hierarchy->GetControlValueByIndex(ChannelElement->GetIndex(), bInitial ? ERigControlValueType::Initial : ERigControlValueType::Current);
 			Value = StoredValue.Get<FRigControlValue::FTransform_Float>().ToTransform();
 		}
 		else if(ChannelElement->Settings.ControlType == ERigControlType::EulerTransform)
 		{
-			const FRigControlValue StoredValue = Context.Hierarchy->GetControlValueByIndex(ChannelElement->GetIndex(), bInitial ? ERigControlValueType::Initial : ERigControlValueType::Current);
+			const FRigControlValue StoredValue = ExecuteContext.Hierarchy->GetControlValueByIndex(ChannelElement->GetIndex(), bInitial ? ERigControlValueType::Initial : ERigControlValueType::Current);
 			Value = StoredValue.Get<FRigControlValue::FEulerTransform_Float>().ToTransform().ToFTransform();
 		}
 		else
@@ -209,17 +209,17 @@ FRigUnit_GetTransformAnimationChannel_Execute()
 
 FRigUnit_SetBoolAnimationChannel_Execute()
 {
-	if(!UpdateCache(Context.Hierarchy, Control, Channel, CachedChannelKey, CachedChannelHash))
+	if(!UpdateCache(ExecuteContext.Hierarchy, Control, Channel, CachedChannelKey, CachedChannelHash))
 	{
 		return;
 	}
 
-	if(const FRigControlElement* ChannelElement = Context.Hierarchy->Find<FRigControlElement>(CachedChannelKey))
+	if(const FRigControlElement* ChannelElement = ExecuteContext.Hierarchy->Find<FRigControlElement>(CachedChannelKey))
 	{
 		if(ChannelElement->Settings.ControlType == ERigControlType::Bool)
 		{
 			const FRigControlValue ValueToStore = FRigControlValue::Make<bool>(Value);
-			Context.Hierarchy->SetControlValueByIndex(ChannelElement->GetIndex(), ValueToStore, bInitial ? ERigControlValueType::Initial : ERigControlValueType::Current);
+			ExecuteContext.Hierarchy->SetControlValueByIndex(ChannelElement->GetIndex(), ValueToStore, bInitial ? ERigControlValueType::Initial : ERigControlValueType::Current);
 		}
 		else
 		{
@@ -230,17 +230,17 @@ FRigUnit_SetBoolAnimationChannel_Execute()
 
 FRigUnit_SetFloatAnimationChannel_Execute()
 {
-	if(!UpdateCache(Context.Hierarchy, Control, Channel, CachedChannelKey, CachedChannelHash))
+	if(!UpdateCache(ExecuteContext.Hierarchy, Control, Channel, CachedChannelKey, CachedChannelHash))
 	{
 		return;
 	}
 
-	if(const FRigControlElement* ChannelElement = Context.Hierarchy->Find<FRigControlElement>(CachedChannelKey))
+	if(const FRigControlElement* ChannelElement = ExecuteContext.Hierarchy->Find<FRigControlElement>(CachedChannelKey))
 	{
 		if(ChannelElement->Settings.ControlType == ERigControlType::Float)
 		{
 			const FRigControlValue ValueToStore = FRigControlValue::Make<float>(Value);
-			Context.Hierarchy->SetControlValueByIndex(ChannelElement->GetIndex(), ValueToStore, bInitial ? ERigControlValueType::Initial : ERigControlValueType::Current);
+			ExecuteContext.Hierarchy->SetControlValueByIndex(ChannelElement->GetIndex(), ValueToStore, bInitial ? ERigControlValueType::Initial : ERigControlValueType::Current);
 		}
 		else
 		{
@@ -251,17 +251,17 @@ FRigUnit_SetFloatAnimationChannel_Execute()
 
 FRigUnit_SetIntAnimationChannel_Execute()
 {
-	if(!UpdateCache(Context.Hierarchy, Control, Channel, CachedChannelKey, CachedChannelHash))
+	if(!UpdateCache(ExecuteContext.Hierarchy, Control, Channel, CachedChannelKey, CachedChannelHash))
 	{
 		return;
 	}
 
-	if(const FRigControlElement* ChannelElement = Context.Hierarchy->Find<FRigControlElement>(CachedChannelKey))
+	if(const FRigControlElement* ChannelElement = ExecuteContext.Hierarchy->Find<FRigControlElement>(CachedChannelKey))
 	{
 		if(ChannelElement->Settings.ControlType == ERigControlType::Integer)
 		{
 			const FRigControlValue ValueToStore = FRigControlValue::Make<int32>(Value);
-			Context.Hierarchy->SetControlValueByIndex(ChannelElement->GetIndex(), ValueToStore, bInitial ? ERigControlValueType::Initial : ERigControlValueType::Current);
+			ExecuteContext.Hierarchy->SetControlValueByIndex(ChannelElement->GetIndex(), ValueToStore, bInitial ? ERigControlValueType::Initial : ERigControlValueType::Current);
 		}
 		else
 		{
@@ -272,17 +272,17 @@ FRigUnit_SetIntAnimationChannel_Execute()
 
 FRigUnit_SetVector2DAnimationChannel_Execute()
 {
-	if(!UpdateCache(Context.Hierarchy, Control, Channel, CachedChannelKey, CachedChannelHash))
+	if(!UpdateCache(ExecuteContext.Hierarchy, Control, Channel, CachedChannelKey, CachedChannelHash))
 	{
 		return;
 	}
 
-	if(const FRigControlElement* ChannelElement = Context.Hierarchy->Find<FRigControlElement>(CachedChannelKey))
+	if(const FRigControlElement* ChannelElement = ExecuteContext.Hierarchy->Find<FRigControlElement>(CachedChannelKey))
 	{
 		if(ChannelElement->Settings.ControlType == ERigControlType::Vector2D)
 		{
 			const FRigControlValue ValueToStore = FRigControlValue::Make<FVector2f>(FVector2f(Value));
-			Context.Hierarchy->SetControlValueByIndex(ChannelElement->GetIndex(), ValueToStore, bInitial ? ERigControlValueType::Initial : ERigControlValueType::Current);
+			ExecuteContext.Hierarchy->SetControlValueByIndex(ChannelElement->GetIndex(), ValueToStore, bInitial ? ERigControlValueType::Initial : ERigControlValueType::Current);
 		}
 		else
 		{
@@ -293,17 +293,17 @@ FRigUnit_SetVector2DAnimationChannel_Execute()
 
 FRigUnit_SetVectorAnimationChannel_Execute()
 {
-	if(!UpdateCache(Context.Hierarchy, Control, Channel, CachedChannelKey, CachedChannelHash))
+	if(!UpdateCache(ExecuteContext.Hierarchy, Control, Channel, CachedChannelKey, CachedChannelHash))
 	{
 		return;
 	}
 
-	if(const FRigControlElement* ChannelElement = Context.Hierarchy->Find<FRigControlElement>(CachedChannelKey))
+	if(const FRigControlElement* ChannelElement = ExecuteContext.Hierarchy->Find<FRigControlElement>(CachedChannelKey))
 	{
 		if(ChannelElement->Settings.ControlType == ERigControlType::Position || ChannelElement->Settings.ControlType == ERigControlType::Scale)
 		{
 			const FRigControlValue ValueToStore = FRigControlValue::Make<FVector3f>(FVector3f(Value));
-			Context.Hierarchy->SetControlValueByIndex(ChannelElement->GetIndex(), ValueToStore, bInitial ? ERigControlValueType::Initial : ERigControlValueType::Current);
+			ExecuteContext.Hierarchy->SetControlValueByIndex(ChannelElement->GetIndex(), ValueToStore, bInitial ? ERigControlValueType::Initial : ERigControlValueType::Current);
 		}
 		else
 		{
@@ -314,17 +314,17 @@ FRigUnit_SetVectorAnimationChannel_Execute()
 
 FRigUnit_SetRotatorAnimationChannel_Execute()
 {
-	if(!UpdateCache(Context.Hierarchy, Control, Channel, CachedChannelKey, CachedChannelHash))
+	if(!UpdateCache(ExecuteContext.Hierarchy, Control, Channel, CachedChannelKey, CachedChannelHash))
 	{
 		return;
 	}
 
-	if(const FRigControlElement* ChannelElement = Context.Hierarchy->Find<FRigControlElement>(CachedChannelKey))
+	if(const FRigControlElement* ChannelElement = ExecuteContext.Hierarchy->Find<FRigControlElement>(CachedChannelKey))
 	{
 		if(ChannelElement->Settings.ControlType == ERigControlType::Rotator)
 		{
 			const FRigControlValue ValueToStore = FRigControlValue::Make<FVector3f>(FVector3f(Value.Euler()));
-			Context.Hierarchy->SetControlValueByIndex(ChannelElement->GetIndex(), ValueToStore, bInitial ? ERigControlValueType::Initial : ERigControlValueType::Current);
+			ExecuteContext.Hierarchy->SetControlValueByIndex(ChannelElement->GetIndex(), ValueToStore, bInitial ? ERigControlValueType::Initial : ERigControlValueType::Current);
 		}
 		else
 		{
@@ -335,22 +335,22 @@ FRigUnit_SetRotatorAnimationChannel_Execute()
 
 FRigUnit_SetTransformAnimationChannel_Execute()
 {
-	if(!UpdateCache(Context.Hierarchy, Control, Channel, CachedChannelKey, CachedChannelHash))
+	if(!UpdateCache(ExecuteContext.Hierarchy, Control, Channel, CachedChannelKey, CachedChannelHash))
 	{
 		return;
 	}
 
-	if(const FRigControlElement* ChannelElement = Context.Hierarchy->Find<FRigControlElement>(CachedChannelKey))
+	if(const FRigControlElement* ChannelElement = ExecuteContext.Hierarchy->Find<FRigControlElement>(CachedChannelKey))
 	{
 		if(ChannelElement->Settings.ControlType == ERigControlType::Transform)
 		{
 			const FRigControlValue ValueToStore = FRigControlValue::Make<FRigControlValue::FTransform_Float>(Value);
-			Context.Hierarchy->SetControlValueByIndex(ChannelElement->GetIndex(), ValueToStore, bInitial ? ERigControlValueType::Initial : ERigControlValueType::Current);
+			ExecuteContext.Hierarchy->SetControlValueByIndex(ChannelElement->GetIndex(), ValueToStore, bInitial ? ERigControlValueType::Initial : ERigControlValueType::Current);
 		}
 		else if(ChannelElement->Settings.ControlType == ERigControlType::EulerTransform)
 		{
 			const FRigControlValue ValueToStore = FRigControlValue::Make<FRigControlValue::FEulerTransform_Float>(FEulerTransform(Value));
-			Context.Hierarchy->SetControlValueByIndex(ChannelElement->GetIndex(), ValueToStore, bInitial ? ERigControlValueType::Initial : ERigControlValueType::Current);
+			ExecuteContext.Hierarchy->SetControlValueByIndex(ChannelElement->GetIndex(), ValueToStore, bInitial ? ERigControlValueType::Initial : ERigControlValueType::Current);
 		}
 		else
 		{
diff --git a/Engine/Plugins/Animation/ControlRig/Source/ControlRig/Private/Units/Hierarchy/RigUnit_ControlChannel.h b/Engine/Plugins/Animation/ControlRig/Source/ControlRig/Private/Units/Hierarchy/RigUnit_ControlChannel.h
--- a/Engine/Plugins/Animation/ControlRig/Source/ControlRig/Private/Units/Hierarchy/RigUnit_ControlChannel.h
+++ b/Engine/Plugins/Animation/ControlRig/Source/ControlRig/Private/Units/Hierarchy/RigUnit_ControlChannel.h
@@ -64,7 +64,7 @@ struct CONTROLRIG_API FRigUnit_GetBoolAnimationChannel : public FRigUnit_GetAnim
 	{}
 
 	RIGVM_METHOD()
-	virtual void Execute(const FRigUnitContext& Context) override;
+	virtual void Execute() override;
 
 	// The current value of the animation channel
 	UPROPERTY(meta=(Output))
@@ -85,7 +85,7 @@ struct CONTROLRIG_API FRigUnit_GetFloatAnimationChannel : public FRigUnit_GetAni
 	{}
 
 	RIGVM_METHOD()
-	virtual void Execute(const FRigUnitContext& Context) override;
+	virtual void Execute() override;
 
 	// The current value of the animation channel
 	UPROPERTY(meta=(Output, UIMin=0, UIMax=1))
@@ -106,7 +106,7 @@ struct CONTROLRIG_API FRigUnit_GetIntAnimationChannel : public FRigUnit_GetAnima
 	{}
 
 	RIGVM_METHOD()
-	virtual void Execute(const FRigUnitContext& Context) override;
+	virtual void Execute() override;
 
 	// The current value of the animation channel
 	UPROPERTY(meta=(Output))
@@ -127,7 +127,7 @@ struct CONTROLRIG_API FRigUnit_GetVector2DAnimationChannel : public FRigUnit_Get
 	{}
 
 	RIGVM_METHOD()
-	virtual void Execute(const FRigUnitContext& Context) override;
+	virtual void Execute() override;
 
 	// The current value of the animation channel
 	UPROPERTY(meta=(Output))
@@ -148,7 +148,7 @@ struct CONTROLRIG_API FRigUnit_GetVectorAnimationChannel : public FRigUnit_GetAn
 	{}
 
 	RIGVM_METHOD()
-	virtual void Execute(const FRigUnitContext& Context) override;
+	virtual void Execute() override;
 
 	// The current value of the animation channel
 	UPROPERTY(meta=(Output))
@@ -169,7 +169,7 @@ struct CONTROLRIG_API FRigUnit_GetRotatorAnimationChannel : public FRigUnit_GetA
 	{}
 
 	RIGVM_METHOD()
-	virtual void Execute(const FRigUnitContext& Context) override;
+	virtual void Execute() override;
 
 	// The current value of the animation channel
 	UPROPERTY(meta=(Output))
@@ -190,7 +190,7 @@ struct CONTROLRIG_API FRigUnit_GetTransformAnimationChannel : public FRigUnit_Ge
 	{}
 
 	RIGVM_METHOD()
-	virtual void Execute(const FRigUnitContext& Context) override;
+	virtual void Execute() override;
 
 	// The current value of the animation channel
 	UPROPERTY(meta=(Output))
@@ -228,7 +228,7 @@ struct CONTROLRIG_API FRigUnit_SetBoolAnimationChannel : public FRigUnit_SetAnim
 	{}
 
 	RIGVM_METHOD()
-	virtual void Execute(const FRigUnitContext& Context) override;
+	virtual void Execute() override;
 
 	// The new value of the animation channel
 	UPROPERTY(meta=(Input))
@@ -249,7 +249,7 @@ struct CONTROLRIG_API FRigUnit_SetFloatAnimationChannel : public FRigUnit_SetAni
 	{}
 
 	RIGVM_METHOD()
-	virtual void Execute(const FRigUnitContext& Context) override;
+	virtual void Execute() override;
 
 	// The new value of the animation channel
 	UPROPERTY(meta=(Input, UIMin=0, UIMax=1))
@@ -270,7 +270,7 @@ struct CONTROLRIG_API FRigUnit_SetIntAnimationChannel : public FRigUnit_SetAnima
 	{}
 
 	RIGVM_METHOD()
-	virtual void Execute(const FRigUnitContext& Context) override;
+	virtual void Execute() override;
 
 	// The new value of the animation channel
 	UPROPERTY(meta=(Input))
@@ -291,7 +291,7 @@ struct CONTROLRIG_API FRigUnit_SetVector2DAnimationChannel : public FRigUnit_Set
 	{}
 
 	RIGVM_METHOD()
-	virtual void Execute(const FRigUnitContext& Context) override;
+	virtual void Execute() override;
 
 	// The new value of the animation channel
 	UPROPERTY(meta=(Input))
@@ -312,7 +312,7 @@ struct CONTROLRIG_API FRigUnit_SetVectorAnimationChannel : public FRigUnit_SetAn
 	{}
 
 	RIGVM_METHOD()
-	virtual void Execute(const FRigUnitContext& Context) override;
+	virtual void Execute() override;
 
 	// The new value of the animation channel
 	UPROPERTY(meta=(Input))
@@ -333,7 +333,7 @@ struct CONTROLRIG_API FRigUnit_SetRotatorAnimationChannel : public FRigUnit_SetA
 	{}
 
 	RIGVM_METHOD()
-	virtual void Execute(const FRigUnitContext& Context) override;
+	virtual void Execute() override;
 
 	// The new value of the animation channel
 	UPROPERTY(meta=(Input))
@@ -354,7 +354,7 @@ struct CONTROLRIG_API FRigUnit_SetTransformAnimationChannel : public FRigUnit_Se
 	{}
 
 	RIGVM_METHOD()
-	virtual void Execute(const FRigUnitContext& Context) override;
+	virtual void Execute() override;
 
 	// The new value of the animation channel
 	UPROPERTY(meta=(Input))
diff --git a/Engine/Plugins/Animation/ControlRig/Source/ControlRig/Private/Units/Hierarchy/RigUnit_GetBoneTransform.cpp b/Engine/Plugins/Animation/ControlRig/Source/ControlRig/Private/Units/Hierarchy/RigUnit_GetBoneTransform.cpp
--- a/Engine/Plugins/Animation/ControlRig/Source/ControlRig/Private/Units/Hierarchy/RigUnit_GetBoneTransform.cpp
+++ b/Engine/Plugins/Animation/ControlRig/Source/ControlRig/Private/Units/Hierarchy/RigUnit_GetBoneTransform.cpp
@@ -10,14 +10,14 @@
 FRigUnit_GetBoneTransform_Execute()
 {
     DECLARE_SCOPE_HIERARCHICAL_COUNTER_RIGUNIT()
-	const URigHierarchy* Hierarchy = Context.Hierarchy;
+	const URigHierarchy* Hierarchy = ExecuteContext.Hierarchy;
 	if (Hierarchy)
 	{
-		switch (Context.State)
+		switch (ExecuteContext.UnitContext.State)
 		{
 			case EControlRigState::Init:
 			{
-				FRigUnit_GetInitialBoneTransform::StaticExecute(ExecuteContext, Bone, Space, Transform, CachedBone, Context);;
+				FRigUnit_GetInitialBoneTransform::StaticExecute(ExecuteContext, Bone, Space, Transform, CachedBone);;
 				break;
 			}
 			case EControlRigState::Update:
diff --git a/Engine/Plugins/Animation/ControlRig/Source/ControlRig/Private/Units/Hierarchy/RigUnit_GetControlInitialTransform.cpp b/Engine/Plugins/Animation/ControlRig/Source/ControlRig/Private/Units/Hierarchy/RigUnit_GetControlInitialTransform.cpp
--- a/Engine/Plugins/Animation/ControlRig/Source/ControlRig/Private/Units/Hierarchy/RigUnit_GetControlInitialTransform.cpp
+++ b/Engine/Plugins/Animation/ControlRig/Source/ControlRig/Private/Units/Hierarchy/RigUnit_GetControlInitialTransform.cpp
@@ -9,10 +9,10 @@
 FRigUnit_GetControlInitialTransform_Execute()
 {
     DECLARE_SCOPE_HIERARCHICAL_COUNTER_RIGUNIT()
-	const URigHierarchy* Hierarchy = Context.Hierarchy;
+	const URigHierarchy* Hierarchy = ExecuteContext.Hierarchy;
 	if (Hierarchy)
 	{
-		switch (Context.State)
+		switch (ExecuteContext.UnitContext.State)
 		{
 			case EControlRigState::Init:
 			{
diff --git a/Engine/Plugins/Animation/ControlRig/Source/ControlRig/Private/Units/Hierarchy/RigUnit_GetControlTransform.cpp b/Engine/Plugins/Animation/ControlRig/Source/ControlRig/Private/Units/Hierarchy/RigUnit_GetControlTransform.cpp
--- a/Engine/Plugins/Animation/ControlRig/Source/ControlRig/Private/Units/Hierarchy/RigUnit_GetControlTransform.cpp
+++ b/Engine/Plugins/Animation/ControlRig/Source/ControlRig/Private/Units/Hierarchy/RigUnit_GetControlTransform.cpp
@@ -8,10 +8,10 @@
 FRigUnit_GetControlBool_Execute()
 {
     DECLARE_SCOPE_HIERARCHICAL_COUNTER_RIGUNIT()
-	const URigHierarchy* Hierarchy = Context.Hierarchy;
+	const URigHierarchy* Hierarchy = ExecuteContext.Hierarchy;
 	if (Hierarchy)
 	{
-		switch (Context.State)
+		switch (ExecuteContext.UnitContext.State)
 		{
 			case EControlRigState::Init:
 			{
@@ -39,10 +39,10 @@ FRigUnit_GetControlBool_Execute()
 FRigUnit_GetControlFloat_Execute()
 {
     DECLARE_SCOPE_HIERARCHICAL_COUNTER_RIGUNIT()
-	const URigHierarchy* Hierarchy = Context.Hierarchy;
+	const URigHierarchy* Hierarchy = ExecuteContext.Hierarchy;
 	if (Hierarchy)
 	{
-		switch (Context.State)
+		switch (ExecuteContext.UnitContext.State)
 		{
 			case EControlRigState::Init:
 			{
@@ -73,10 +73,10 @@ FRigUnit_GetControlFloat_Execute()
 FRigUnit_GetControlInteger_Execute()
 {
     DECLARE_SCOPE_HIERARCHICAL_COUNTER_RIGUNIT()
-	const URigHierarchy* Hierarchy = Context.Hierarchy;
+	const URigHierarchy* Hierarchy = ExecuteContext.Hierarchy;
 	if (Hierarchy)
 	{
-		switch (Context.State)
+		switch (ExecuteContext.UnitContext.State)
 		{
 			case EControlRigState::Init:
 			{
@@ -106,10 +106,10 @@ FRigUnit_GetControlInteger_Execute()
 FRigUnit_GetControlVector2D_Execute()
 {
     DECLARE_SCOPE_HIERARCHICAL_COUNTER_RIGUNIT()
-	const URigHierarchy* Hierarchy = Context.Hierarchy;
+	const URigHierarchy* Hierarchy = ExecuteContext.Hierarchy;
 	if (Hierarchy)
 	{
-		switch (Context.State)
+		switch (ExecuteContext.UnitContext.State)
 		{
 			case EControlRigState::Init:
 			{
@@ -142,10 +142,10 @@ FRigUnit_GetControlVector2D_Execute()
 FRigUnit_GetControlVector_Execute()
 {
     DECLARE_SCOPE_HIERARCHICAL_COUNTER_RIGUNIT()
-	const URigHierarchy* Hierarchy = Context.Hierarchy;
+	const URigHierarchy* Hierarchy = ExecuteContext.Hierarchy;
 	if (Hierarchy)
 	{
-		switch (Context.State)
+		switch (ExecuteContext.UnitContext.State)
 		{
 			case EControlRigState::Init:
 			{
@@ -204,10 +204,10 @@ FRigUnit_GetControlVector_Execute()
 FRigUnit_GetControlRotator_Execute()
 {
     DECLARE_SCOPE_HIERARCHICAL_COUNTER_RIGUNIT()
-	const URigHierarchy* Hierarchy = Context.Hierarchy;
+	const URigHierarchy* Hierarchy = ExecuteContext.Hierarchy;
 	if (Hierarchy)
 	{
-		switch (Context.State)
+		switch (ExecuteContext.UnitContext.State)
 		{
 			case EControlRigState::Init:
 			{
@@ -259,10 +259,10 @@ FRigUnit_GetControlRotator_Execute()
 FRigUnit_GetControlTransform_Execute()
 {
     DECLARE_SCOPE_HIERARCHICAL_COUNTER_RIGUNIT()
-	const URigHierarchy* Hierarchy = Context.Hierarchy;
+	const URigHierarchy* Hierarchy = ExecuteContext.Hierarchy;
 	if (Hierarchy)
 	{
-		switch (Context.State)
+		switch (ExecuteContext.UnitContext.State)
 		{
 			case EControlRigState::Init:
 			{
diff --git a/Engine/Plugins/Animation/ControlRig/Source/ControlRig/Private/Units/Hierarchy/RigUnit_GetControlTransform.h b/Engine/Plugins/Animation/ControlRig/Source/ControlRig/Private/Units/Hierarchy/RigUnit_GetControlTransform.h
--- a/Engine/Plugins/Animation/ControlRig/Source/ControlRig/Private/Units/Hierarchy/RigUnit_GetControlTransform.h
+++ b/Engine/Plugins/Animation/ControlRig/Source/ControlRig/Private/Units/Hierarchy/RigUnit_GetControlTransform.h
@@ -19,7 +19,7 @@ struct CONTROLRIG_API FRigUnit_GetControlBool : public FRigUnit
 	{}
 
 	RIGVM_METHOD()
-	virtual void Execute(const FRigUnitContext& Context) override;
+	virtual void Execute() override;
 
 	/**
 	 * The name of the Control to retrieve the bool for.
@@ -52,7 +52,7 @@ struct CONTROLRIG_API FRigUnit_GetControlFloat : public FRigUnit
 	{}
 
 	RIGVM_METHOD()
-	virtual void Execute(const FRigUnitContext& Context) override;
+	virtual void Execute() override;
 
 	/**
 	 * The name of the Control to retrieve the Float for.
@@ -93,7 +93,7 @@ struct CONTROLRIG_API FRigUnit_GetControlInteger : public FRigUnit
 	{}
 
 	RIGVM_METHOD()
-	virtual void Execute(const FRigUnitContext& Context) override;
+	virtual void Execute() override;
 
 	/**
 	 * The name of the Control to retrieve the Integer for.
@@ -134,7 +134,7 @@ struct CONTROLRIG_API FRigUnit_GetControlVector2D : public FRigUnit
 	{}
 
 	RIGVM_METHOD()
-	virtual void Execute(const FRigUnitContext& Context) override;
+	virtual void Execute() override;
 
 	/**
 	 * The name of the Control to retrieve the Vector2D for.
@@ -176,7 +176,7 @@ struct CONTROLRIG_API FRigUnit_GetControlVector : public FRigUnit
 	{}
 
 	RIGVM_METHOD()
-	virtual void Execute(const FRigUnitContext& Context) override;
+	virtual void Execute() override;
 
 	/**
 	 * The name of the Control to retrieve the Vector for.
@@ -226,7 +226,7 @@ struct CONTROLRIG_API FRigUnit_GetControlRotator : public FRigUnit
 	}
 
 	RIGVM_METHOD()
-	virtual void Execute(const FRigUnitContext& Context) override;
+	virtual void Execute() override;
 
 	/**
 	 * The name of the Control to retrieve the Rotator for.
@@ -275,7 +275,7 @@ struct CONTROLRIG_API FRigUnit_GetControlTransform : public FRigUnit
 	{}
 
 	RIGVM_METHOD()
-	virtual void Execute(const FRigUnitContext& Context) override;
+	virtual void Execute() override;
 
 	/**
 	 * The name of the Control to retrieve the transform for.
diff --git a/Engine/Plugins/Animation/ControlRig/Source/ControlRig/Private/Units/Hierarchy/RigUnit_GetCurveValue.cpp b/Engine/Plugins/Animation/ControlRig/Source/ControlRig/Private/Units/Hierarchy/RigUnit_GetCurveValue.cpp
--- a/Engine/Plugins/Animation/ControlRig/Source/ControlRig/Private/Units/Hierarchy/RigUnit_GetCurveValue.cpp
+++ b/Engine/Plugins/Animation/ControlRig/Source/ControlRig/Private/Units/Hierarchy/RigUnit_GetCurveValue.cpp
@@ -8,9 +8,9 @@
 FRigUnit_GetCurveValue_Execute()
 {
     DECLARE_SCOPE_HIERARCHICAL_COUNTER_RIGUNIT()
-    if (const URigHierarchy* Hierarchy = Context.Hierarchy)
+    if (const URigHierarchy* Hierarchy = ExecuteContext.Hierarchy)
 	{
-		switch (Context.State)
+		switch (ExecuteContext.UnitContext.State)
 		{
 			case EControlRigState::Init:
 			{
diff --git a/Engine/Plugins/Animation/ControlRig/Source/ControlRig/Private/Units/Hierarchy/RigUnit_GetInitialBoneTransform.cpp b/Engine/Plugins/Animation/ControlRig/Source/ControlRig/Private/Units/Hierarchy/RigUnit_GetInitialBoneTransform.cpp
--- a/Engine/Plugins/Animation/ControlRig/Source/ControlRig/Private/Units/Hierarchy/RigUnit_GetInitialBoneTransform.cpp
+++ b/Engine/Plugins/Animation/ControlRig/Source/ControlRig/Private/Units/Hierarchy/RigUnit_GetInitialBoneTransform.cpp
@@ -9,10 +9,10 @@
 FRigUnit_GetInitialBoneTransform_Execute()
 {
     DECLARE_SCOPE_HIERARCHICAL_COUNTER_RIGUNIT()
-	const URigHierarchy* Hierarchy = Context.Hierarchy;
+	const URigHierarchy* Hierarchy = ExecuteContext.Hierarchy;
 	if (Hierarchy)
 	{
-		switch (Context.State)
+		switch (ExecuteContext.UnitContext.State)
 		{
 			case EControlRigState::Init:
 			{
diff --git a/Engine/Plugins/Animation/ControlRig/Source/ControlRig/Private/Units/Hierarchy/RigUnit_GetRelativeTransform.cpp b/Engine/Plugins/Animation/ControlRig/Source/ControlRig/Private/Units/Hierarchy/RigUnit_GetRelativeTransform.cpp
--- a/Engine/Plugins/Animation/ControlRig/Source/ControlRig/Private/Units/Hierarchy/RigUnit_GetRelativeTransform.cpp
+++ b/Engine/Plugins/Animation/ControlRig/Source/ControlRig/Private/Units/Hierarchy/RigUnit_GetRelativeTransform.cpp
@@ -14,7 +14,7 @@ FRigUnit_GetRelativeTransformForItem_Execute()
 	FTransform ChildTransform = FTransform::Identity;
 	FTransform ParentTransform = FTransform::Identity;
 
-	FRigUnit_GetTransform::StaticExecute(ExecuteContext, Child, EBoneGetterSetterMode::GlobalSpace, bChildInitial, ChildTransform, CachedChild, Context);
-	FRigUnit_GetTransform::StaticExecute(ExecuteContext, Parent, EBoneGetterSetterMode::GlobalSpace, bParentInitial, ParentTransform, CachedParent, Context);
-	FRigUnit_MathTransformMakeRelative::StaticExecute(ExecuteContext, ChildTransform, ParentTransform, RelativeTransform, Context);
+	FRigUnit_GetTransform::StaticExecute(ExecuteContext, Child, EBoneGetterSetterMode::GlobalSpace, bChildInitial, ChildTransform, CachedChild);
+	FRigUnit_GetTransform::StaticExecute(ExecuteContext, Parent, EBoneGetterSetterMode::GlobalSpace, bParentInitial, ParentTransform, CachedParent);
+	FRigUnit_MathTransformMakeRelative::StaticExecute(ExecuteContext, ChildTransform, ParentTransform, RelativeTransform);
 }
diff --git a/Engine/Plugins/Animation/ControlRig/Source/ControlRig/Private/Units/Hierarchy/RigUnit_GetTransform.cpp b/Engine/Plugins/Animation/ControlRig/Source/ControlRig/Private/Units/Hierarchy/RigUnit_GetTransform.cpp
--- a/Engine/Plugins/Animation/ControlRig/Source/ControlRig/Private/Units/Hierarchy/RigUnit_GetTransform.cpp
+++ b/Engine/Plugins/Animation/ControlRig/Source/ControlRig/Private/Units/Hierarchy/RigUnit_GetTransform.cpp
@@ -16,9 +16,9 @@ FRigUnit_GetTransform_Execute()
 {
     DECLARE_SCOPE_HIERARCHICAL_COUNTER_RIGUNIT()
 
-	if (URigHierarchy* Hierarchy = Context.Hierarchy)
+	if (URigHierarchy* Hierarchy = ExecuteContext.Hierarchy)
 	{
-		switch (Context.State)
+		switch (ExecuteContext.UnitContext.State)
 		{
 			case EControlRigState::Init:
 			{
@@ -35,7 +35,7 @@ FRigUnit_GetTransform_Execute()
 				}
 				else
 				{
-					if(bInitial || Context.State == EControlRigState::Init)
+					if(bInitial || ExecuteContext.UnitContext.State == EControlRigState::Init)
 					{
 						switch (Space)
 						{
@@ -87,7 +87,7 @@ FRigUnit_GetTransform_Execute()
 
 FRigUnit_GetTransformArray_Execute()
 {
-	FRigUnit_GetTransformItemArray::StaticExecute(ExecuteContext, Items.Keys, Space, bInitial, Transforms, CachedIndex, Context);
+	FRigUnit_GetTransformItemArray::StaticExecute(ExecuteContext, Items.Keys, Space, bInitial, Transforms, CachedIndex);
 }
 
 FRigVMStructUpgradeInfo FRigUnit_GetTransformArray::GetUpgradeInfo() const
@@ -112,7 +112,7 @@ FRigUnit_GetTransformItemArray_Execute()
 	Transforms.SetNum(Items.Num());
 	for(int32 Index=0;Index<Items.Num();Index++)
 	{
-		FRigUnit_GetTransform::StaticExecute(ExecuteContext, Items[Index], Space, bInitial, Transforms[Index], CachedIndex[Index], Context);
+		FRigUnit_GetTransform::StaticExecute(ExecuteContext, Items[Index], Space, bInitial, Transforms[Index], CachedIndex[Index]);
 
 	}
 }
diff --git a/Engine/Plugins/Animation/ControlRig/Source/ControlRig/Private/Units/Hierarchy/RigUnit_GetTransform.h b/Engine/Plugins/Animation/ControlRig/Source/ControlRig/Private/Units/Hierarchy/RigUnit_GetTransform.h
--- a/Engine/Plugins/Animation/ControlRig/Source/ControlRig/Private/Units/Hierarchy/RigUnit_GetTransform.h
+++ b/Engine/Plugins/Animation/ControlRig/Source/ControlRig/Private/Units/Hierarchy/RigUnit_GetTransform.h
@@ -36,7 +36,7 @@ struct CONTROLRIG_API FRigUnit_GetTransform : public FRigUnit
 	}
 
 	RIGVM_METHOD()
-	virtual void Execute(const FRigUnitContext& Context) override;
+	virtual void Execute() override;
 
 	/**
 	 * The item to retrieve the transform for
@@ -83,7 +83,7 @@ struct CONTROLRIG_API FRigUnit_GetTransformArray : public FRigUnit
 	{}
 
 	RIGVM_METHOD()
-	virtual void Execute(const FRigUnitContext& Context) override;
+	virtual void Execute() override;
 
 	/**
 	* The items to retrieve the transforms for
@@ -133,7 +133,7 @@ struct CONTROLRIG_API FRigUnit_GetTransformItemArray : public FRigUnit
 	{}
 
 	RIGVM_METHOD()
-	virtual void Execute(const FRigUnitContext& Context) override;
+	virtual void Execute() override;
 
 	/**
 	* The items to retrieve the transforms for
diff --git a/Engine/Plugins/Animation/ControlRig/Source/ControlRig/Private/Units/Hierarchy/RigUnit_Metadata.cpp b/Engine/Plugins/Animation/ControlRig/Source/ControlRig/Private/Units/Hierarchy/RigUnit_Metadata.cpp
--- a/Engine/Plugins/Animation/ControlRig/Source/ControlRig/Private/Units/Hierarchy/RigUnit_Metadata.cpp
+++ b/Engine/Plugins/Animation/ControlRig/Source/ControlRig/Private/Units/Hierarchy/RigUnit_Metadata.cpp
@@ -170,10 +170,10 @@ TArray<FRigVMTemplateArgument> FRigDispatch_GetMetadata::GetArguments() const
 FRigBaseMetadata* FRigDispatch_GetMetadata::FindMetadata(const FRigVMExtendedExecuteContext& InContext,
                                                          const FRigElementKey& InKey, const FName& InName, ERigMetadataType InType, FCachedRigElement& Cache)
 {
-	const FRigUnitContext& Context = GetRigUnitContext(InContext);
-	if(Cache.UpdateCache(InKey, Context.Hierarchy))
+	const FControlRigExecuteContext& ExecuteContext = InContext.GetPublicData<FControlRigExecuteContext>();
+	if(Cache.UpdateCache(InKey, ExecuteContext.Hierarchy))
 	{
-		if(FRigBaseElement* Element = Context.Hierarchy->Get(Cache.GetIndex()))
+		if(FRigBaseElement* Element = ExecuteContext.Hierarchy->Get(Cache.GetIndex()))
 		{
 			return Element->GetMetadata(InName, InType);
 		}
@@ -422,7 +422,7 @@ FRigUnit_RemoveAllMetadata_Execute()
 
 FRigUnit_HasMetadata_Execute()
 {
-	const URigHierarchy* Hierarchy = Context.Hierarchy;
+	const URigHierarchy* Hierarchy = ExecuteContext.Hierarchy;
 
 	Found = false;
 
@@ -437,7 +437,7 @@ FRigUnit_HasMetadata_Execute()
 
 FRigUnit_FindItemsWithMetadata_Execute()
 {
-	const URigHierarchy* Hierarchy = Context.Hierarchy;
+	const URigHierarchy* Hierarchy = ExecuteContext.Hierarchy;
 
 	Items.Reset();
 
@@ -455,7 +455,7 @@ FRigUnit_GetMetadataTags_Execute()
 {
 	Tags.Reset();
 	
-	const URigHierarchy* Hierarchy = Context.Hierarchy;
+	const URigHierarchy* Hierarchy = ExecuteContext.Hierarchy;
 	if(CachedIndex.UpdateCache(Item, Hierarchy))
 	{
 		if(const FRigBaseElement* Element = Hierarchy->Get(CachedIndex))
@@ -530,7 +530,7 @@ FRigUnit_HasMetadataTag_Execute()
 {
 	Found = false;
 	
-	const URigHierarchy* Hierarchy = Context.Hierarchy;
+	const URigHierarchy* Hierarchy = ExecuteContext.Hierarchy;
 	if(CachedIndex.UpdateCache(Item, Hierarchy))
 	{
 		if(const FRigBaseElement* Element = Hierarchy->Get(CachedIndex))
@@ -547,7 +547,7 @@ FRigUnit_HasMetadataTagArray_Execute()
 {
 	Found = false;
 	
-	const URigHierarchy* Hierarchy = Context.Hierarchy;
+	const URigHierarchy* Hierarchy = ExecuteContext.Hierarchy;
 	if(CachedIndex.UpdateCache(Item, Hierarchy))
 	{
 		if(const FRigBaseElement* Element = Hierarchy->Get(CachedIndex))
@@ -570,7 +570,7 @@ FRigUnit_HasMetadataTagArray_Execute()
 
 FRigUnit_FindItemsWithMetadataTag_Execute()
 {
-	const URigHierarchy* Hierarchy = Context.Hierarchy;
+	const URigHierarchy* Hierarchy = ExecuteContext.Hierarchy;
 
 	Items.Reset();
 
@@ -589,7 +589,7 @@ FRigUnit_FindItemsWithMetadataTag_Execute()
 
 FRigUnit_FindItemsWithMetadataTagArray_Execute()
 {
-	const URigHierarchy* Hierarchy = Context.Hierarchy;
+	const URigHierarchy* Hierarchy = ExecuteContext.Hierarchy;
 
 	Items.Reset();
 
@@ -618,7 +618,7 @@ FRigUnit_FindItemsWithMetadataTagArray_Execute()
 
 FRigUnit_FilterItemsByMetadataTags_Execute()
 {
-	const URigHierarchy* Hierarchy = Context.Hierarchy;
+	const URigHierarchy* Hierarchy = ExecuteContext.Hierarchy;
 
 	Result.Reset();
 
diff --git a/Engine/Plugins/Animation/ControlRig/Source/ControlRig/Private/Units/Hierarchy/RigUnit_Metadata.h b/Engine/Plugins/Animation/ControlRig/Source/ControlRig/Private/Units/Hierarchy/RigUnit_Metadata.h
--- a/Engine/Plugins/Animation/ControlRig/Source/ControlRig/Private/Units/Hierarchy/RigUnit_Metadata.h
+++ b/Engine/Plugins/Animation/ControlRig/Source/ControlRig/Private/Units/Hierarchy/RigUnit_Metadata.h
@@ -183,7 +183,7 @@ struct CONTROLRIG_API FRigUnit_RemoveMetadata : public FRigUnitMutable
 	{}
 
 	RIGVM_METHOD()
-	virtual void Execute(const FRigUnitContext& Context) override;
+	virtual void Execute() override;
 
 	/**
 	 * The item to remove the metadata from 
@@ -221,7 +221,7 @@ struct CONTROLRIG_API FRigUnit_RemoveAllMetadata : public FRigUnitMutable
 	{}
 
 	RIGVM_METHOD()
-	virtual void Execute(const FRigUnitContext& Context) override;
+	virtual void Execute() override;
 
 	/**
 	 * The item to remove the metadata from 
@@ -255,7 +255,7 @@ struct CONTROLRIG_API FRigUnit_HasMetadata : public FRigUnit
 	{}
 
 	RIGVM_METHOD()
-	virtual void Execute(const FRigUnitContext& Context) override;
+	virtual void Execute() override;
 
 	/**
 	 * The item to check the metadata for
@@ -299,7 +299,7 @@ struct CONTROLRIG_API FRigUnit_FindItemsWithMetadata : public FRigUnit
 	{}
 
 	RIGVM_METHOD()
-	virtual void Execute(const FRigUnitContext& Context) override;
+	virtual void Execute() override;
 
 	/**
 	 * The name of the metadata to find
@@ -333,7 +333,7 @@ struct CONTROLRIG_API FRigUnit_GetMetadataTags : public FRigUnit
 	{}
 
 	RIGVM_METHOD()
-	virtual void Execute(const FRigUnitContext& Context) override;
+	virtual void Execute() override;
 
 	/**
 	 * The item to check the metadata for
@@ -367,7 +367,7 @@ struct CONTROLRIG_API FRigUnit_SetMetadataTag : public FRigUnitMutable
 	{}
 
 	RIGVM_METHOD()
-	virtual void Execute(const FRigUnitContext& Context) override;
+	virtual void Execute() override;
 
 	/**
 	 * The item to set the metadata for
@@ -401,7 +401,7 @@ struct CONTROLRIG_API FRigUnit_SetMetadataTagArray : public FRigUnitMutable
 	{}
 
 	RIGVM_METHOD()
-	virtual void Execute(const FRigUnitContext& Context) override;
+	virtual void Execute() override;
 
 	/**
 	 * The item to set the metadata for
@@ -436,7 +436,7 @@ struct CONTROLRIG_API FRigUnit_RemoveMetadataTag : public FRigUnitMutable
 	{}
 
 	RIGVM_METHOD()
-	virtual void Execute(const FRigUnitContext& Context) override;
+	virtual void Execute() override;
 
 	/**
 	 * The item to set the metadata for
@@ -477,7 +477,7 @@ struct CONTROLRIG_API FRigUnit_HasMetadataTag : public FRigUnit
 	{}
 
 	RIGVM_METHOD()
-	virtual void Execute(const FRigUnitContext& Context) override;
+	virtual void Execute() override;
 
 	/**
 	 * The item to check the metadata for
@@ -516,7 +516,7 @@ struct CONTROLRIG_API FRigUnit_HasMetadataTagArray : public FRigUnit
 	{}
 
 	RIGVM_METHOD()
-	virtual void Execute(const FRigUnitContext& Context) override;
+	virtual void Execute() override;
 
 	/**
 	 * The item to check the metadata for
@@ -553,7 +553,7 @@ struct CONTROLRIG_API FRigUnit_FindItemsWithMetadataTag : public FRigUnit
 	{}
 
 	RIGVM_METHOD()
-	virtual void Execute(const FRigUnitContext& Context) override;
+	virtual void Execute() override;
 
 	/**
 	 * The name of the tag to find
@@ -580,7 +580,7 @@ struct CONTROLRIG_API FRigUnit_FindItemsWithMetadataTagArray : public FRigUnit
 	{}
 
 	RIGVM_METHOD()
-	virtual void Execute(const FRigUnitContext& Context) override;
+	virtual void Execute() override;
 
 	/**
 	 * The tags to find
@@ -609,7 +609,7 @@ struct CONTROLRIG_API FRigUnit_FilterItemsByMetadataTags : public FRigUnit
 	{}
 
 	RIGVM_METHOD()
-	virtual void Execute(const FRigUnitContext& Context) override;
+	virtual void Execute() override;
 
 	// The items to filter
 	UPROPERTY(meta=(Input))
diff --git a/Engine/Plugins/Animation/ControlRig/Source/ControlRig/Private/Units/Hierarchy/RigUnit_OffsetTransform.cpp b/Engine/Plugins/Animation/ControlRig/Source/ControlRig/Private/Units/Hierarchy/RigUnit_OffsetTransform.cpp
--- a/Engine/Plugins/Animation/ControlRig/Source/ControlRig/Private/Units/Hierarchy/RigUnit_OffsetTransform.cpp
+++ b/Engine/Plugins/Animation/ControlRig/Source/ControlRig/Private/Units/Hierarchy/RigUnit_OffsetTransform.cpp
@@ -20,7 +20,7 @@ FRigUnit_OffsetTransformForItem_Execute()
 	FTransform PreviousTransform = FTransform::Identity;
 	FTransform GlobalTransform = FTransform::Identity;
 
-	FRigUnit_GetTransform::StaticExecute(ExecuteContext, Item, EBoneGetterSetterMode::GlobalSpace, false, PreviousTransform, CachedIndex, Context);
-	FRigUnit_MathTransformMakeAbsolute::StaticExecute(ExecuteContext, OffsetTransform, PreviousTransform, GlobalTransform, Context);
-	FRigUnit_SetTransform::StaticExecute(ExecuteContext, Item, EBoneGetterSetterMode::GlobalSpace, false, GlobalTransform, Weight, bPropagateToChildren, CachedIndex, Context);
+	FRigUnit_GetTransform::StaticExecute(ExecuteContext, Item, EBoneGetterSetterMode::GlobalSpace, false, PreviousTransform, CachedIndex);
+	FRigUnit_MathTransformMakeAbsolute::StaticExecute(ExecuteContext, OffsetTransform, PreviousTransform, GlobalTransform);
+	FRigUnit_SetTransform::StaticExecute(ExecuteContext, Item, EBoneGetterSetterMode::GlobalSpace, false, GlobalTransform, Weight, bPropagateToChildren, CachedIndex);
 }
diff --git a/Engine/Plugins/Animation/ControlRig/Source/ControlRig/Private/Units/Hierarchy/RigUnit_ParentSwitchConstraint.cpp b/Engine/Plugins/Animation/ControlRig/Source/ControlRig/Private/Units/Hierarchy/RigUnit_ParentSwitchConstraint.cpp
--- a/Engine/Plugins/Animation/ControlRig/Source/ControlRig/Private/Units/Hierarchy/RigUnit_ParentSwitchConstraint.cpp
+++ b/Engine/Plugins/Animation/ControlRig/Source/ControlRig/Private/Units/Hierarchy/RigUnit_ParentSwitchConstraint.cpp
@@ -7,7 +7,7 @@
 
 FRigUnit_ParentSwitchConstraint_Execute()
 {
-	FRigUnit_ParentSwitchConstraintArray::StaticExecute(ExecuteContext, Subject, ParentIndex, Parents.Keys, InitialGlobalTransform, Weight, Transform, Switched, CachedSubject, CachedParent, RelativeOffset, Context);	
+	FRigUnit_ParentSwitchConstraintArray::StaticExecute(ExecuteContext, Subject, ParentIndex, Parents.Keys, InitialGlobalTransform, Weight, Transform, Switched, CachedSubject, CachedParent, RelativeOffset);	
 }
 
 FRigVMStructUpgradeInfo FRigUnit_ParentSwitchConstraint::GetUpgradeInfo() const
@@ -31,7 +31,7 @@ FRigUnit_ParentSwitchConstraintArray_Execute()
 
 	if (URigHierarchy* Hierarchy = ExecuteContext.Hierarchy)
 	{
-		switch (Context.State)
+		switch (ExecuteContext.UnitContext.State)
 		{
 			case EControlRigState::Init:
 			{
diff --git a/Engine/Plugins/Animation/ControlRig/Source/ControlRig/Private/Units/Hierarchy/RigUnit_ProjectTransformToNewParent.cpp b/Engine/Plugins/Animation/ControlRig/Source/ControlRig/Private/Units/Hierarchy/RigUnit_ProjectTransformToNewParent.cpp
--- a/Engine/Plugins/Animation/ControlRig/Source/ControlRig/Private/Units/Hierarchy/RigUnit_ProjectTransformToNewParent.cpp
+++ b/Engine/Plugins/Animation/ControlRig/Source/ControlRig/Private/Units/Hierarchy/RigUnit_ProjectTransformToNewParent.cpp
@@ -16,10 +16,10 @@ FRigUnit_ProjectTransformToNewParent_Execute()
 	FTransform NewParentTransform = FTransform::Identity;
 	FTransform RelativeTransform = FTransform::Identity;
 
-	FRigUnit_GetTransform::StaticExecute(ExecuteContext, Child, EBoneGetterSetterMode::GlobalSpace, bChildInitial, ChildTransform, CachedChild, Context);
-	FRigUnit_GetTransform::StaticExecute(ExecuteContext, OldParent, EBoneGetterSetterMode::GlobalSpace, bOldParentInitial, OldParentTransform, CachedOldParent, Context);
-	FRigUnit_GetTransform::StaticExecute(ExecuteContext, NewParent, EBoneGetterSetterMode::GlobalSpace, bNewParentInitial, NewParentTransform, CachedNewParent, Context);
-	FRigUnit_MathTransformMakeRelative::StaticExecute(ExecuteContext, ChildTransform, OldParentTransform, RelativeTransform, Context);
-	FRigUnit_MathTransformMakeAbsolute::StaticExecute(ExecuteContext, RelativeTransform, NewParentTransform, Transform, Context);
+	FRigUnit_GetTransform::StaticExecute(ExecuteContext, Child, EBoneGetterSetterMode::GlobalSpace, bChildInitial, ChildTransform, CachedChild);
+	FRigUnit_GetTransform::StaticExecute(ExecuteContext, OldParent, EBoneGetterSetterMode::GlobalSpace, bOldParentInitial, OldParentTransform, CachedOldParent);
+	FRigUnit_GetTransform::StaticExecute(ExecuteContext, NewParent, EBoneGetterSetterMode::GlobalSpace, bNewParentInitial, NewParentTransform, CachedNewParent);
+	FRigUnit_MathTransformMakeRelative::StaticExecute(ExecuteContext, ChildTransform, OldParentTransform, RelativeTransform);
+	FRigUnit_MathTransformMakeAbsolute::StaticExecute(ExecuteContext, RelativeTransform, NewParentTransform, Transform);
 }
 
diff --git a/Engine/Plugins/Animation/ControlRig/Source/ControlRig/Private/Units/Hierarchy/RigUnit_PropagateTransform.cpp b/Engine/Plugins/Animation/ControlRig/Source/ControlRig/Private/Units/Hierarchy/RigUnit_PropagateTransform.cpp
--- a/Engine/Plugins/Animation/ControlRig/Source/ControlRig/Private/Units/Hierarchy/RigUnit_PropagateTransform.cpp
+++ b/Engine/Plugins/Animation/ControlRig/Source/ControlRig/Private/Units/Hierarchy/RigUnit_PropagateTransform.cpp
@@ -22,7 +22,7 @@ FRigUnit_PropagateTransform_Execute()
     	/*
     	 * This node doesn't do anything anymore now that the hierarchy is lazy
     	 *
-		switch (Context.State)
+		switch (ExecuteContext.UnitContext.State)
 		{
 			case EControlRigState::Init:
 			{
diff --git a/Engine/Plugins/Animation/ControlRig/Source/ControlRig/Private/Units/Hierarchy/RigUnit_SendEvent.cpp b/Engine/Plugins/Animation/ControlRig/Source/ControlRig/Private/Units/Hierarchy/RigUnit_SendEvent.cpp
--- a/Engine/Plugins/Animation/ControlRig/Source/ControlRig/Private/Units/Hierarchy/RigUnit_SendEvent.cpp
+++ b/Engine/Plugins/Animation/ControlRig/Source/ControlRig/Private/Units/Hierarchy/RigUnit_SendEvent.cpp
@@ -14,7 +14,7 @@ FRigUnit_SendEvent_Execute()
     	return;
     }
 
-	if (bOnlyDuringInteraction && !Context.IsInteracting())
+	if (bOnlyDuringInteraction && !ExecuteContext.UnitContext.IsInteracting())
 	{
 		return;
 	}
@@ -22,7 +22,7 @@ FRigUnit_SendEvent_Execute()
 	URigHierarchy* Hierarchy = ExecuteContext.Hierarchy;
 	if (Hierarchy)
 	{
-		switch (Context.State)
+		switch (ExecuteContext.UnitContext.State)
 		{
 			case EControlRigState::Init:
 			{
@@ -34,7 +34,7 @@ FRigUnit_SendEvent_Execute()
 				EventContext.Key = Item;
 				EventContext.Event = Event;
 				EventContext.SourceEventName = ExecuteContext.GetEventName();
-				EventContext.LocalTime = Context.AbsoluteTime + OffsetInSeconds;
+				EventContext.LocalTime = ExecuteContext.UnitContext.AbsoluteTime + OffsetInSeconds;
 				Hierarchy->SendEvent(EventContext, false /* async */); //needs to be false for sequencer keying to work
 				break;
 			}
diff --git a/Engine/Plugins/Animation/ControlRig/Source/ControlRig/Private/Units/Hierarchy/RigUnit_SetControlColor.cpp b/Engine/Plugins/Animation/ControlRig/Source/ControlRig/Private/Units/Hierarchy/RigUnit_SetControlColor.cpp
--- a/Engine/Plugins/Animation/ControlRig/Source/ControlRig/Private/Units/Hierarchy/RigUnit_SetControlColor.cpp
+++ b/Engine/Plugins/Animation/ControlRig/Source/ControlRig/Private/Units/Hierarchy/RigUnit_SetControlColor.cpp
@@ -11,10 +11,10 @@ FRigUnit_GetControlColor_Execute()
 
     Color = FLinearColor::Black;
 
-	const URigHierarchy* Hierarchy = Context.Hierarchy;
+	const URigHierarchy* Hierarchy = ExecuteContext.Hierarchy;
 	if (Hierarchy)
 	{
-		switch (Context.State)
+		switch (ExecuteContext.UnitContext.State)
 		{
 			case EControlRigState::Init:
 			{
@@ -44,7 +44,7 @@ FRigUnit_SetControlColor_Execute()
 	URigHierarchy* Hierarchy = ExecuteContext.Hierarchy;
 	if (Hierarchy)
 	{
-		switch (Context.State)
+		switch (ExecuteContext.UnitContext.State)
 		{
 			case EControlRigState::Init:
 			{
diff --git a/Engine/Plugins/Animation/ControlRig/Source/ControlRig/Private/Units/Hierarchy/RigUnit_SetControlDrivenList.cpp b/Engine/Plugins/Animation/ControlRig/Source/ControlRig/Private/Units/Hierarchy/RigUnit_SetControlDrivenList.cpp
--- a/Engine/Plugins/Animation/ControlRig/Source/ControlRig/Private/Units/Hierarchy/RigUnit_SetControlDrivenList.cpp
+++ b/Engine/Plugins/Animation/ControlRig/Source/ControlRig/Private/Units/Hierarchy/RigUnit_SetControlDrivenList.cpp
@@ -13,10 +13,10 @@ FRigUnit_GetControlDrivenList_Execute()
 
     Driven.Reset();
 
-	const URigHierarchy* Hierarchy = Context.Hierarchy;
+	const URigHierarchy* Hierarchy = ExecuteContext.Hierarchy;
 	if (Hierarchy)
 	{
-		switch (Context.State)
+		switch (ExecuteContext.UnitContext.State)
 		{
 			case EControlRigState::Init:
 			{
@@ -46,7 +46,7 @@ FRigUnit_SetControlDrivenList_Execute()
 	URigHierarchy* Hierarchy = ExecuteContext.Hierarchy;
 	if (Hierarchy)
 	{
-		switch (Context.State)
+		switch (ExecuteContext.UnitContext.State)
 		{
 			case EControlRigState::Init:
 			{
diff --git a/Engine/Plugins/Animation/ControlRig/Source/ControlRig/Private/Units/Hierarchy/RigUnit_SetControlOffset.cpp b/Engine/Plugins/Animation/ControlRig/Source/ControlRig/Private/Units/Hierarchy/RigUnit_SetControlOffset.cpp
--- a/Engine/Plugins/Animation/ControlRig/Source/ControlRig/Private/Units/Hierarchy/RigUnit_SetControlOffset.cpp
+++ b/Engine/Plugins/Animation/ControlRig/Source/ControlRig/Private/Units/Hierarchy/RigUnit_SetControlOffset.cpp
@@ -16,7 +16,7 @@ FRigUnit_SetControlOffset_Execute()
 	URigHierarchy* Hierarchy = ExecuteContext.Hierarchy;
 	if (Hierarchy)
 	{
-		switch (Context.State)
+		switch (ExecuteContext.UnitContext.State)
 		{
 			case EControlRigState::Init:
 			{
@@ -55,10 +55,10 @@ FRigUnit_GetShapeTransform_Execute()
 {
 	DECLARE_SCOPE_HIERARCHICAL_COUNTER_RIGUNIT()
 
-	const URigHierarchy* Hierarchy = Context.Hierarchy;
+	const URigHierarchy* Hierarchy = ExecuteContext.Hierarchy;
 	if (Hierarchy)
 	{
-		switch (Context.State)
+		switch (ExecuteContext.UnitContext.State)
 		{
 			case EControlRigState::Init:
 			{
@@ -92,7 +92,7 @@ FRigUnit_SetShapeTransform_Execute()
 	URigHierarchy* Hierarchy = ExecuteContext.Hierarchy;
 	if (Hierarchy)
 	{
-		switch (Context.State)
+		switch (ExecuteContext.UnitContext.State)
 		{
 			case EControlRigState::Init:
 			{
diff --git a/Engine/Plugins/Animation/ControlRig/Source/ControlRig/Private/Units/Hierarchy/RigUnit_SetControlTransform.cpp b/Engine/Plugins/Animation/ControlRig/Source/ControlRig/Private/Units/Hierarchy/RigUnit_SetControlTransform.cpp
--- a/Engine/Plugins/Animation/ControlRig/Source/ControlRig/Private/Units/Hierarchy/RigUnit_SetControlTransform.cpp
+++ b/Engine/Plugins/Animation/ControlRig/Source/ControlRig/Private/Units/Hierarchy/RigUnit_SetControlTransform.cpp
@@ -13,7 +13,7 @@ FRigUnit_SetControlBool_Execute()
 	URigHierarchy* Hierarchy = ExecuteContext.Hierarchy;
 	if (Hierarchy)
 	{
-		switch (Context.State)
+		switch (ExecuteContext.UnitContext.State)
 		{
 			case EControlRigState::Init:
 			{
@@ -45,7 +45,7 @@ FRigUnit_SetMultiControlBool_Execute()
 	{
 		CachedControlIndices.SetNum(Entries.Num());
 
-		switch (Context.State)
+		switch (ExecuteContext.UnitContext.State)
 		{
 			case EControlRigState::Update:
 			{
@@ -53,7 +53,7 @@ FRigUnit_SetMultiControlBool_Execute()
 				{
 					bool BoolValue = Entries[EntryIndex].BoolValue;
 
-					FRigUnit_SetControlBool::StaticExecute(ExecuteContext, Entries[EntryIndex].Control, BoolValue, CachedControlIndices[EntryIndex], Context);
+					FRigUnit_SetControlBool::StaticExecute(ExecuteContext, Entries[EntryIndex].Control, BoolValue, CachedControlIndices[EntryIndex]);
 				}
 				break;
 			}
@@ -71,7 +71,7 @@ FRigUnit_SetControlFloat_Execute()
 	URigHierarchy* Hierarchy = ExecuteContext.Hierarchy;
 	if (Hierarchy)
 	{
-		switch (Context.State)
+		switch (ExecuteContext.UnitContext.State)
 		{
 			case EControlRigState::Init:
 			{
@@ -111,7 +111,7 @@ FRigUnit_SetMultiControlFloat_Execute()
 	{ 
 		CachedControlIndices.SetNum(Entries.Num());
 	 
-		switch (Context.State)
+		switch (ExecuteContext.UnitContext.State)
 		{
 			case EControlRigState::Update:
 			{ 
@@ -119,7 +119,7 @@ FRigUnit_SetMultiControlFloat_Execute()
 				{
 					float FloatValue = Entries[EntryIndex].FloatValue;
 
-					FRigUnit_SetControlFloat::StaticExecute(ExecuteContext, Entries[EntryIndex].Control, Weight, FloatValue, CachedControlIndices[EntryIndex], Context); 
+					FRigUnit_SetControlFloat::StaticExecute(ExecuteContext, Entries[EntryIndex].Control, Weight, FloatValue, CachedControlIndices[EntryIndex]); 
 				}
 				break;
 			}
@@ -137,7 +137,7 @@ FRigUnit_SetControlInteger_Execute()
 	URigHierarchy* Hierarchy = ExecuteContext.Hierarchy;
 	if (Hierarchy)
 	{
-		switch (Context.State)
+		switch (ExecuteContext.UnitContext.State)
 		{
 			case EControlRigState::Init:
 			{
@@ -177,7 +177,7 @@ FRigUnit_SetMultiControlInteger_Execute()
 	{
 		CachedControlIndices.SetNum(Entries.Num());
 
-		switch (Context.State)
+		switch (ExecuteContext.UnitContext.State)
 		{
 			case EControlRigState::Update:
 			{
@@ -185,7 +185,7 @@ FRigUnit_SetMultiControlInteger_Execute()
 				{
 					int32 IntegerValue = Entries[EntryIndex].IntegerValue;
 
-					FRigUnit_SetControlInteger::StaticExecute(ExecuteContext, Entries[EntryIndex].Control, Weight, IntegerValue, CachedControlIndices[EntryIndex], Context);
+					FRigUnit_SetControlInteger::StaticExecute(ExecuteContext, Entries[EntryIndex].Control, Weight, IntegerValue, CachedControlIndices[EntryIndex]);
 				}
 				break;
 			}
@@ -203,7 +203,7 @@ FRigUnit_SetControlVector2D_Execute()
 	URigHierarchy* Hierarchy = ExecuteContext.Hierarchy;
 	if (Hierarchy)
 	{
-		switch (Context.State)
+		switch (ExecuteContext.UnitContext.State)
 		{
 			case EControlRigState::Init:
 			{
@@ -244,7 +244,7 @@ FRigUnit_SetMultiControlVector2D_Execute()
 	{
 		CachedControlIndices.SetNum(Entries.Num());
 
-		switch (Context.State)
+		switch (ExecuteContext.UnitContext.State)
 		{
 			case EControlRigState::Update:
 			{
@@ -252,7 +252,7 @@ FRigUnit_SetMultiControlVector2D_Execute()
 				{
 					FVector2D Vector = Entries[EntryIndex].Vector;
 
-					FRigUnit_SetControlVector2D::StaticExecute(ExecuteContext, Entries[EntryIndex].Control, Weight, Vector, CachedControlIndices[EntryIndex], Context);
+					FRigUnit_SetControlVector2D::StaticExecute(ExecuteContext, Entries[EntryIndex].Control, Weight, Vector, CachedControlIndices[EntryIndex]);
 				}
 				break;
 			}
@@ -270,7 +270,7 @@ FRigUnit_SetControlVector_Execute()
 	URigHierarchy* Hierarchy = ExecuteContext.Hierarchy;
 	if (Hierarchy)
 	{
-		switch (Context.State)
+		switch (ExecuteContext.UnitContext.State)
 		{
 			case EControlRigState::Init:
 			{
@@ -348,7 +348,7 @@ FRigUnit_SetControlRotator_Execute()
 	URigHierarchy* Hierarchy = ExecuteContext.Hierarchy;
 	if (Hierarchy)
 	{
-		switch (Context.State)
+		switch (ExecuteContext.UnitContext.State)
 		{
 			case EControlRigState::Init:
 			{
@@ -413,7 +413,7 @@ FRigUnit_SetMultiControlRotator_Execute()
 	{
 		CachedControlIndices.SetNum(Entries.Num());
 
-		switch (Context.State)
+		switch (ExecuteContext.UnitContext.State)
 		{
 			case EControlRigState::Update:
 			{
@@ -421,7 +421,7 @@ FRigUnit_SetMultiControlRotator_Execute()
 				{
 					FRotator Rotator = Entries[EntryIndex].Rotator;
 
-					FRigUnit_SetControlRotator::StaticExecute(ExecuteContext, Entries[EntryIndex].Control, Weight, Rotator, Entries[EntryIndex].Space , CachedControlIndices[EntryIndex], Context);
+					FRigUnit_SetControlRotator::StaticExecute(ExecuteContext, Entries[EntryIndex].Control, Weight, Rotator, Entries[EntryIndex].Space , CachedControlIndices[EntryIndex]);
 				}
 				break;
 			}
@@ -439,7 +439,7 @@ FRigUnit_SetControlTransform_Execute()
 	URigHierarchy* Hierarchy = ExecuteContext.Hierarchy;
 	if (Hierarchy)
 	{
-		switch (Context.State)
+		switch (ExecuteContext.UnitContext.State)
 		{
 			case EControlRigState::Init:
 			{
