RigVM: Remove support for opaque attributes

diff --git a/Engine/Plugins/Animation/ControlRig/Source/ControlRig/Private/Units/Hierarchy/RigUnit_SetControlTransform.h b/Engine/Plugins/Animation/ControlRig/Source/ControlRig/Private/Units/Hierarchy/RigUnit_SetControlTransform.h
--- a/Engine/Plugins/Animation/ControlRig/Source/ControlRig/Private/Units/Hierarchy/RigUnit_SetControlTransform.h
+++ b/Engine/Plugins/Animation/ControlRig/Source/ControlRig/Private/Units/Hierarchy/RigUnit_SetControlTransform.h
@@ -19,7 +19,7 @@ struct CONTROLRIG_API FRigUnit_SetControlBool : public FRigUnitMutable
 	{}
 
 	RIGVM_METHOD()
-	virtual void Execute(const FRigUnitContext& Context) override;
+	virtual void Execute() override;
 
 	/**
 	 * The name of the Control to set the bool for.
@@ -72,7 +72,7 @@ struct CONTROLRIG_API FRigUnit_SetMultiControlBool : public FRigUnitMutable
 	}
 
 	RIGVM_METHOD()
-	virtual void Execute(const FRigUnitContext& Context) override;
+	virtual void Execute() override;
 
 	/**
 	 * The array of control-Bool pairs to be processed
@@ -100,7 +100,7 @@ struct CONTROLRIG_API FRigUnit_SetControlFloat : public FRigUnitMutable
 	{}
 
 	RIGVM_METHOD()
-	virtual void Execute(const FRigUnitContext& Context) override;
+	virtual void Execute() override;
 
 	/**
 	 * The name of the Control to set the transform for.
@@ -161,7 +161,7 @@ struct CONTROLRIG_API FRigUnit_SetMultiControlFloat : public FRigUnitMutable
 	}
 
 	RIGVM_METHOD()
-	virtual void Execute(const FRigUnitContext& Context) override;
+	virtual void Execute() override;
 
 	/**
 	 * The array of control-float pairs to be processed
@@ -197,7 +197,7 @@ struct CONTROLRIG_API FRigUnit_SetControlInteger : public FRigUnitMutable
 	{}
 
 	RIGVM_METHOD()
-	virtual void Execute(const FRigUnitContext& Context) override;
+	virtual void Execute() override;
 
 	/**
 	 * The name of the Control to set the transform for.
@@ -258,7 +258,7 @@ struct CONTROLRIG_API FRigUnit_SetMultiControlInteger : public FRigUnitMutable
 	}
 
 	RIGVM_METHOD()
-	virtual void Execute(const FRigUnitContext& Context) override;
+	virtual void Execute() override;
 
 	/**
 	 * The array of control-integer pairs to be processed
@@ -292,7 +292,7 @@ struct CONTROLRIG_API FRigUnit_SetControlVector2D : public FRigUnitMutable
 	{}
 
 	RIGVM_METHOD()
-	virtual void Execute(const FRigUnitContext& Context) override;
+	virtual void Execute() override;
 
 	/**
 	 * The name of the Control to set the transform for.
@@ -354,7 +354,7 @@ struct CONTROLRIG_API FRigUnit_SetMultiControlVector2D : public FRigUnitMutable
 	}
 
 	RIGVM_METHOD()
-	virtual void Execute(const FRigUnitContext& Context) override;
+	virtual void Execute() override;
 
 	/**
 	 * The array of control-vector2D pairs to be processed
@@ -389,7 +389,7 @@ struct CONTROLRIG_API FRigUnit_SetControlVector : public FRigUnitMutable
 	{}
 
 	RIGVM_METHOD()
-	virtual void Execute(const FRigUnitContext& Context) override;
+	virtual void Execute() override;
 
 	/**
 	 * The name of the Control to set the transform for.
@@ -437,7 +437,7 @@ struct CONTROLRIG_API FRigUnit_SetControlRotator : public FRigUnitMutable
 	{}
 
 	RIGVM_METHOD()
-	virtual void Execute(const FRigUnitContext& Context) override;
+	virtual void Execute() override;
 
 	/**
 	 * The name of the Control to set the transform for.
@@ -515,7 +515,7 @@ struct CONTROLRIG_API FRigUnit_SetMultiControlRotator : public FRigUnitMutable
 	}
 
 	RIGVM_METHOD()
-	virtual void Execute(const FRigUnitContext& Context) override;
+	virtual void Execute() override;
 
 	/**
 	 * The array of control-rotator pairs to be processed
@@ -549,7 +549,7 @@ struct CONTROLRIG_API FRigUnit_SetControlTransform : public FRigUnitMutable
 	{}
 
 	RIGVM_METHOD()
-	virtual void Execute(const FRigUnitContext& Context) override;
+	virtual void Execute() override;
 
 	/**
 	 * The name of the Control to set the transform for.
diff --git a/Engine/Plugins/Animation/ControlRig/Source/ControlRig/Private/Units/Hierarchy/RigUnit_SetRelativeTransform.cpp b/Engine/Plugins/Animation/ControlRig/Source/ControlRig/Private/Units/Hierarchy/RigUnit_SetRelativeTransform.cpp
--- a/Engine/Plugins/Animation/ControlRig/Source/ControlRig/Private/Units/Hierarchy/RigUnit_SetRelativeTransform.cpp
+++ b/Engine/Plugins/Animation/ControlRig/Source/ControlRig/Private/Units/Hierarchy/RigUnit_SetRelativeTransform.cpp
@@ -21,9 +21,9 @@ FRigUnit_SetRelativeTransformForItem_Execute()
 	FTransform ParentTransform = FTransform::Identity;
 	FTransform GlobalTransform = FTransform::Identity;
 
-	FRigUnit_GetTransform::StaticExecute(ExecuteContext, Parent, EBoneGetterSetterMode::GlobalSpace, bParentInitial, ParentTransform, CachedParent, Context);
-	FRigUnit_MathTransformMakeAbsolute::StaticExecute(ExecuteContext, Value, ParentTransform, GlobalTransform, Context);
-	FRigUnit_SetTransform::StaticExecute(ExecuteContext, Child, EBoneGetterSetterMode::GlobalSpace, false, GlobalTransform, Weight, bPropagateToChildren, CachedChild, Context);
+	FRigUnit_GetTransform::StaticExecute(ExecuteContext, Parent, EBoneGetterSetterMode::GlobalSpace, bParentInitial, ParentTransform, CachedParent);
+	FRigUnit_MathTransformMakeAbsolute::StaticExecute(ExecuteContext, Value, ParentTransform, GlobalTransform);
+	FRigUnit_SetTransform::StaticExecute(ExecuteContext, Child, EBoneGetterSetterMode::GlobalSpace, false, GlobalTransform, Weight, bPropagateToChildren, CachedChild);
 }
 
 FRigUnit_SetRelativeTranslationForItem_Execute()
@@ -39,11 +39,11 @@ FRigUnit_SetRelativeTranslationForItem_Execute()
 	FTransform LocalTransform = FTransform::Identity;
 	FTransform GlobalTransform = FTransform::Identity;
 
-	FRigUnit_GetTransform::StaticExecute(ExecuteContext, Parent, EBoneGetterSetterMode::GlobalSpace, bParentInitial, ParentTransform, CachedParent, Context);
-	FRigUnit_GetRelativeTransformForItem::StaticExecute(ExecuteContext, Child, false, Parent, bParentInitial, LocalTransform, CachedChild, CachedParent, Context);
+	FRigUnit_GetTransform::StaticExecute(ExecuteContext, Parent, EBoneGetterSetterMode::GlobalSpace, bParentInitial, ParentTransform, CachedParent);
+	FRigUnit_GetRelativeTransformForItem::StaticExecute(ExecuteContext, Child, false, Parent, bParentInitial, LocalTransform, CachedChild, CachedParent);
 	LocalTransform.SetTranslation(Value);
-	FRigUnit_MathTransformMakeAbsolute::StaticExecute(ExecuteContext, LocalTransform, ParentTransform, GlobalTransform, Context);
-	FRigUnit_SetTransform::StaticExecute(ExecuteContext, Child, EBoneGetterSetterMode::GlobalSpace, false, GlobalTransform, Weight, bPropagateToChildren, CachedChild, Context);
+	FRigUnit_MathTransformMakeAbsolute::StaticExecute(ExecuteContext, LocalTransform, ParentTransform, GlobalTransform);
+	FRigUnit_SetTransform::StaticExecute(ExecuteContext, Child, EBoneGetterSetterMode::GlobalSpace, false, GlobalTransform, Weight, bPropagateToChildren, CachedChild);
 }
 
 FRigUnit_SetRelativeRotationForItem_Execute()
@@ -59,11 +59,11 @@ FRigUnit_SetRelativeRotationForItem_Execute()
 	FTransform LocalTransform = FTransform::Identity;
 	FTransform GlobalTransform = FTransform::Identity;
 
-	FRigUnit_GetTransform::StaticExecute(ExecuteContext, Parent, EBoneGetterSetterMode::GlobalSpace, bParentInitial, ParentTransform, CachedParent, Context);
-	FRigUnit_GetRelativeTransformForItem::StaticExecute(ExecuteContext, Child, false, Parent, bParentInitial, LocalTransform, CachedChild, CachedParent, Context);
+	FRigUnit_GetTransform::StaticExecute(ExecuteContext, Parent, EBoneGetterSetterMode::GlobalSpace, bParentInitial, ParentTransform, CachedParent);
+	FRigUnit_GetRelativeTransformForItem::StaticExecute(ExecuteContext, Child, false, Parent, bParentInitial, LocalTransform, CachedChild, CachedParent);
 	LocalTransform.SetRotation(Value);
 	LocalTransform.NormalizeRotation();
-	FRigUnit_MathTransformMakeAbsolute::StaticExecute(ExecuteContext, LocalTransform, ParentTransform, GlobalTransform, Context);
-	FRigUnit_SetTransform::StaticExecute(ExecuteContext, Child, EBoneGetterSetterMode::GlobalSpace, false, GlobalTransform, Weight, bPropagateToChildren, CachedChild, Context);
+	FRigUnit_MathTransformMakeAbsolute::StaticExecute(ExecuteContext, LocalTransform, ParentTransform, GlobalTransform);
+	FRigUnit_SetTransform::StaticExecute(ExecuteContext, Child, EBoneGetterSetterMode::GlobalSpace, false, GlobalTransform, Weight, bPropagateToChildren, CachedChild);
 }
 
diff --git a/Engine/Plugins/Animation/ControlRig/Source/ControlRig/Private/Units/Hierarchy/RigUnit_SetTransform.cpp b/Engine/Plugins/Animation/ControlRig/Source/ControlRig/Private/Units/Hierarchy/RigUnit_SetTransform.cpp
--- a/Engine/Plugins/Animation/ControlRig/Source/ControlRig/Private/Units/Hierarchy/RigUnit_SetTransform.cpp
+++ b/Engine/Plugins/Animation/ControlRig/Source/ControlRig/Private/Units/Hierarchy/RigUnit_SetTransform.cpp
@@ -27,7 +27,7 @@ FRigUnit_SetTransform_Execute()
 
 	if (URigHierarchy* Hierarchy = ExecuteContext.Hierarchy)
 	{
-		switch (Context.State)
+		switch (ExecuteContext.UnitContext.State)
 		{
 			case EControlRigState::Init:
 			{
@@ -58,7 +58,7 @@ FRigUnit_SetTransform_Execute()
 					if(bInitial && (CachedIndex.GetKey().Type == ERigElementType::Control))
 					{
 						FTransform TransformMutable = Value;
-						FRigUnit_SetControlOffset::StaticExecute(ExecuteContext, CachedIndex.GetKey().Name, TransformMutable, Space, CachedIndex, Context);
+						FRigUnit_SetControlOffset::StaticExecute(ExecuteContext, CachedIndex.GetKey().Name, TransformMutable, Space, CachedIndex);
 						
 						if (ExecuteContext.GetEventName() == FRigUnit_PrepareForExecution::EventName)
 						{
@@ -138,9 +138,9 @@ FString FRigUnit_SetTranslation::GetUnitLabel() const
 FRigUnit_SetTranslation_Execute()
 {
 	FTransform Transform = FTransform::Identity;
-	FRigUnit_GetTransform::StaticExecute(ExecuteContext, Item, Space, bInitial, Transform, CachedIndex, Context);
+	FRigUnit_GetTransform::StaticExecute(ExecuteContext, Item, Space, bInitial, Transform, CachedIndex);
 	Transform.SetLocation(Value);
-	FRigUnit_SetTransform::StaticExecute(ExecuteContext, Item, Space, bInitial, Transform, Weight, bPropagateToChildren, CachedIndex, Context);
+	FRigUnit_SetTransform::StaticExecute(ExecuteContext, Item, Space, bInitial, Transform, Weight, bPropagateToChildren, CachedIndex);
 }
 
 FString FRigUnit_SetRotation::GetUnitLabel() const
@@ -152,9 +152,9 @@ FString FRigUnit_SetRotation::GetUnitLabel() const
 FRigUnit_SetRotation_Execute()
 {
 	FTransform Transform = FTransform::Identity;
-	FRigUnit_GetTransform::StaticExecute(ExecuteContext, Item, Space, bInitial, Transform, CachedIndex, Context);
+	FRigUnit_GetTransform::StaticExecute(ExecuteContext, Item, Space, bInitial, Transform, CachedIndex);
 	Transform.SetRotation(Value);
-	FRigUnit_SetTransform::StaticExecute(ExecuteContext, Item, Space, bInitial, Transform, Weight, bPropagateToChildren, CachedIndex, Context);
+	FRigUnit_SetTransform::StaticExecute(ExecuteContext, Item, Space, bInitial, Transform, Weight, bPropagateToChildren, CachedIndex);
 }
 
 FString FRigUnit_SetScale::GetUnitLabel() const
@@ -166,15 +166,15 @@ FString FRigUnit_SetScale::GetUnitLabel() const
 FRigUnit_SetScale_Execute()
 {
 	FTransform Transform = FTransform::Identity;
-	FRigUnit_GetTransform::StaticExecute(ExecuteContext, Item, Space, bInitial, Transform, CachedIndex, Context);
+	FRigUnit_GetTransform::StaticExecute(ExecuteContext, Item, Space, bInitial, Transform, CachedIndex);
 	Transform.SetScale3D(Scale);
-	FRigUnit_SetTransform::StaticExecute(ExecuteContext, Item, Space, bInitial, Transform, Weight, bPropagateToChildren, CachedIndex, Context);
+	FRigUnit_SetTransform::StaticExecute(ExecuteContext, Item, Space, bInitial, Transform, Weight, bPropagateToChildren, CachedIndex);
 }
 
 
 FRigUnit_SetTransformArray_Execute()
 {
-	FRigUnit_SetTransformItemArray::StaticExecute(ExecuteContext, Items.Keys, Space, bInitial, Transforms, Weight, bPropagateToChildren, CachedIndex, Context);
+	FRigUnit_SetTransformItemArray::StaticExecute(ExecuteContext, Items.Keys, Space, bInitial, Transforms, Weight, bPropagateToChildren, CachedIndex);
 }
 
 FRigVMStructUpgradeInfo FRigUnit_SetTransformArray::GetUpgradeInfo() const
@@ -208,6 +208,6 @@ FRigUnit_SetTransformItemArray_Execute()
 
 	for(int32 Index=0;Index<Items.Num();Index++)
 	{
-		FRigUnit_SetTransform::StaticExecute(ExecuteContext, Items[Index], Space, bInitial, Transforms[Index], Weight, bPropagateToChildren, CachedIndex[Index], Context);
+		FRigUnit_SetTransform::StaticExecute(ExecuteContext, Items[Index], Space, bInitial, Transforms[Index], Weight, bPropagateToChildren, CachedIndex[Index]);
 	}
 }
diff --git a/Engine/Plugins/Animation/ControlRig/Source/ControlRig/Private/Units/Hierarchy/RigUnit_WorldSpace.cpp b/Engine/Plugins/Animation/ControlRig/Source/ControlRig/Private/Units/Hierarchy/RigUnit_WorldSpace.cpp
--- a/Engine/Plugins/Animation/ControlRig/Source/ControlRig/Private/Units/Hierarchy/RigUnit_WorldSpace.cpp
+++ b/Engine/Plugins/Animation/ControlRig/Source/ControlRig/Private/Units/Hierarchy/RigUnit_WorldSpace.cpp
@@ -8,35 +8,35 @@
 FRigUnit_ToWorldSpace_Transform_Execute()
 {
     DECLARE_SCOPE_HIERARCHICAL_COUNTER_RIGUNIT()
-    World = Context.ToWorldSpace(Value);
+    World = ExecuteContext.UnitContext.ToWorldSpace(Value);
 }
 
 FRigUnit_ToRigSpace_Transform_Execute()
 {
     DECLARE_SCOPE_HIERARCHICAL_COUNTER_RIGUNIT()
-    Global = Context.ToRigSpace(Value);
+    Global = ExecuteContext.UnitContext.ToRigSpace(Value);
 }
 
 FRigUnit_ToWorldSpace_Location_Execute()
 {
     DECLARE_SCOPE_HIERARCHICAL_COUNTER_RIGUNIT()
-    World = Context.ToWorldSpace(Value);
+    World = ExecuteContext.UnitContext.ToWorldSpace(Value);
 }
 
 FRigUnit_ToRigSpace_Location_Execute()
 {
     DECLARE_SCOPE_HIERARCHICAL_COUNTER_RIGUNIT()
-    Global = Context.ToRigSpace(Value);
+    Global = ExecuteContext.UnitContext.ToRigSpace(Value);
 }
 
 FRigUnit_ToWorldSpace_Rotation_Execute()
 {
     DECLARE_SCOPE_HIERARCHICAL_COUNTER_RIGUNIT()
-    World = Context.ToWorldSpace(Value);
+    World = ExecuteContext.UnitContext.ToWorldSpace(Value);
 }
 
 FRigUnit_ToRigSpace_Rotation_Execute()
 {
     DECLARE_SCOPE_HIERARCHICAL_COUNTER_RIGUNIT()
-    Global = Context.ToRigSpace(Value);
+    Global = ExecuteContext.UnitContext.ToRigSpace(Value);
 }
diff --git a/Engine/Plugins/Animation/ControlRig/Source/ControlRig/Private/Units/Highlevel/Harmonics/RigUnit_BoneHarmonics.cpp b/Engine/Plugins/Animation/ControlRig/Source/ControlRig/Private/Units/Highlevel/Harmonics/RigUnit_BoneHarmonics.cpp
--- a/Engine/Plugins/Animation/ControlRig/Source/ControlRig/Private/Units/Highlevel/Harmonics/RigUnit_BoneHarmonics.cpp
+++ b/Engine/Plugins/Animation/ControlRig/Source/ControlRig/Private/Units/Highlevel/Harmonics/RigUnit_BoneHarmonics.cpp
@@ -30,8 +30,7 @@ FRigUnit_BoneHarmonics_Execute()
 		WaveMinimum,
 		WaveMaximum,
 		RotationOrder,
-		WorkData,
-		Context);
+		WorkData);
 }
 
 FRigVMStructUpgradeInfo FRigUnit_BoneHarmonics::GetUpgradeInfo() const
@@ -78,7 +77,7 @@ FRigUnit_ItemHarmonics_Execute()
 	TArray<FCachedRigElement>& CachedItems = WorkData.CachedItems;
 	FVector& WaveTime = WorkData.WaveTime;
 
-	if (Context.State == EControlRigState::Init ||
+	if (ExecuteContext.UnitContext.State == EControlRigState::Init ||
 		CachedItems.Num() != Targets.Num())
 	{
 		CachedItems.Reset();
@@ -127,6 +126,6 @@ FRigUnit_ItemHarmonics_Execute()
 		}
 	}
 
-	WaveTime += WaveSpeed * Context.DeltaTime;
+	WaveTime += WaveSpeed * ExecuteContext.UnitContext.DeltaTime;
 }
 
diff --git a/Engine/Plugins/Animation/ControlRig/Source/ControlRig/Private/Units/Highlevel/Harmonics/RigUnit_ChainHarmonics.cpp b/Engine/Plugins/Animation/ControlRig/Source/ControlRig/Private/Units/Highlevel/Harmonics/RigUnit_ChainHarmonics.cpp
--- a/Engine/Plugins/Animation/ControlRig/Source/ControlRig/Private/Units/Highlevel/Harmonics/RigUnit_ChainHarmonics.cpp
+++ b/Engine/Plugins/Animation/ControlRig/Source/ControlRig/Private/Units/Highlevel/Harmonics/RigUnit_ChainHarmonics.cpp
@@ -18,8 +18,7 @@ FRigUnit_ChainHarmonics_Execute()
 		Pendulum,
 		bDrawDebug,
 		DrawWorldOffset,
-		WorkData,
-		Context);
+		WorkData);
 }
 
 FRigVMStructUpgradeInfo FRigUnit_ChainHarmonics::GetUpgradeInfo() const
@@ -58,7 +57,7 @@ FRigUnit_ChainHarmonicsPerItem_Execute()
 	TArray<FVector>& HierarchyLine = WorkData.HierarchyLine;
 	TArray<FVector>& VelocityLines = WorkData.VelocityLines;
 	
-	if (Context.State == EControlRigState::Init)
+	if (ExecuteContext.UnitContext.State == EControlRigState::Init)
 	{
 		Time = FVector::ZeroVector;
 		Items.Reset();
@@ -189,15 +188,15 @@ FRigUnit_ChainHarmonicsPerItem_Execute()
 			FVector Velocity = Pendulum.PendulumGravity;
 			Velocity += Stiffness * Pendulum.PendulumStiffness;
 
-			if (Context.DeltaTime > 0.f)
+			if (ExecuteContext.UnitContext.DeltaTime > 0.f)
 			{
 				PendulumVelocity[Index] = FMath::Lerp<FVector>(PendulumVelocity[Index], Velocity, FMath::Clamp<float>(Pendulum.PendulumBlend, 0.f, 0.999f));
 				PendulumVelocity[Index] = PendulumVelocity[Index] * Pendulum.PendulumDrag;
 
 				FVector PrevPosition = PendulumPosition[Index];
-				PendulumPosition[Index] = PendulumPosition[Index] + PendulumVelocity[Index] * Context.DeltaTime;
+				PendulumPosition[Index] = PendulumPosition[Index] + PendulumVelocity[Index] * ExecuteContext.UnitContext.DeltaTime;
 				PendulumPosition[Index] = GlobalTransform.GetLocation() + (PendulumPosition[Index] - GlobalTransform.GetLocation()).GetSafeNormal() * Length;
-				PendulumVelocity[Index] = (PendulumPosition[Index] - PrevPosition) / Context.DeltaTime;
+				PendulumVelocity[Index] = (PendulumPosition[Index] - PrevPosition) / ExecuteContext.UnitContext.DeltaTime;
 			}
 
 			VelocityLines[Index * 2 + 0] = PendulumPosition[Index];
@@ -221,9 +220,9 @@ FRigUnit_ChainHarmonicsPerItem_Execute()
 		ParentTransform = GlobalTransform;
 	}
 
-	Time = Time + Speed * Context.DeltaTime;
+	Time = Time + Speed * ExecuteContext.UnitContext.DeltaTime;
 
-	if (Context.DrawInterface != nullptr && bDrawDebug)
+	if (ExecuteContext.UnitContext.DrawInterface != nullptr && bDrawDebug)
 	{
 		HierarchyLine.SetNum(Items.Num());
 		for (int32 Index = 0; Index < Items.Num(); Index++)
@@ -231,9 +230,9 @@ FRigUnit_ChainHarmonicsPerItem_Execute()
 			HierarchyLine[Index] = Hierarchy->GetGlobalTransform(Items[Index]).GetLocation();
 		}
 
-		Context.DrawInterface->DrawLineStrip(DrawWorldOffset, HierarchyLine, FLinearColor::Yellow, 0.f);
-		Context.DrawInterface->DrawLines(DrawWorldOffset, VelocityLines, FLinearColor(0.3f, 0.3f, 1.f), 0.f);
-		Context.DrawInterface->DrawPoints(DrawWorldOffset, PendulumPosition, 3.f, FLinearColor::Blue);
+		ExecuteContext.UnitContext.DrawInterface->DrawLineStrip(DrawWorldOffset, HierarchyLine, FLinearColor::Yellow, 0.f);
+		ExecuteContext.UnitContext.DrawInterface->DrawLines(DrawWorldOffset, VelocityLines, FLinearColor(0.3f, 0.3f, 1.f), 0.f);
+		ExecuteContext.UnitContext.DrawInterface->DrawPoints(DrawWorldOffset, PendulumPosition, 3.f, FLinearColor::Blue);
 		return;
 	}
 
diff --git a/Engine/Plugins/Animation/ControlRig/Source/ControlRig/Private/Units/Highlevel/Hierarchy/RigUnit_AimBone.cpp b/Engine/Plugins/Animation/ControlRig/Source/ControlRig/Private/Units/Highlevel/Hierarchy/RigUnit_AimBone.cpp
--- a/Engine/Plugins/Animation/ControlRig/Source/ControlRig/Private/Units/Highlevel/Hierarchy/RigUnit_AimBone.cpp
+++ b/Engine/Plugins/Animation/ControlRig/Source/ControlRig/Private/Units/Highlevel/Hierarchy/RigUnit_AimBone.cpp
@@ -13,13 +13,13 @@ FRigUnit_AimBoneMath_Execute()
 
 	Result = InputTransform;
 
-	URigHierarchy* Hierarchy = Context.Hierarchy;
+	URigHierarchy* Hierarchy = ExecuteContext.Hierarchy;
 	if (Hierarchy == nullptr)
 	{
 		return;
 	}
 
-	if (Context.State == EControlRigState::Init)
+	if (ExecuteContext.UnitContext.State == EControlRigState::Init)
 	{
 		PrimaryCachedSpace.Reset();
 		SecondaryCachedSpace.Reset();
@@ -48,17 +48,17 @@ FRigUnit_AimBoneMath_Execute()
 			}
 		}
 
-		if (Context.DrawInterface != nullptr && DebugSettings.bEnabled)
+		if (ExecuteContext.UnitContext.DrawInterface != nullptr && DebugSettings.bEnabled)
 		{
 			const FLinearColor Color = FLinearColor(0.f, 1.f, 1.f, 1.f);
 			if (Primary.Kind == EControlRigVectorKind::Direction)
 			{
-				Context.DrawInterface->DrawLine(DebugSettings.WorldOffset, Result.GetLocation(), Result.GetLocation() + Target * DebugSettings.Scale, Color);
+				ExecuteContext.UnitContext.DrawInterface->DrawLine(DebugSettings.WorldOffset, Result.GetLocation(), Result.GetLocation() + Target * DebugSettings.Scale, Color);
 			}
 			else
 			{
-				Context.DrawInterface->DrawLine(DebugSettings.WorldOffset, Result.GetLocation(), Target, Color);
-				Context.DrawInterface->DrawBox(DebugSettings.WorldOffset, FTransform(FQuat::Identity, Target, FVector(1.f, 1.f, 1.f) * DebugSettings.Scale * 0.1f), Color);
+				ExecuteContext.UnitContext.DrawInterface->DrawLine(DebugSettings.WorldOffset, Result.GetLocation(), Target, Color);
+				ExecuteContext.UnitContext.DrawInterface->DrawBox(DebugSettings.WorldOffset, FTransform(FQuat::Identity, Target, FVector(1.f, 1.f, 1.f) * DebugSettings.Scale * 0.1f), Color);
 			}
 		}
 
@@ -102,17 +102,17 @@ FRigUnit_AimBoneMath_Execute()
 			}
 		}
 
-		if (Context.DrawInterface != nullptr && DebugSettings.bEnabled)
+		if (ExecuteContext.UnitContext.DrawInterface != nullptr && DebugSettings.bEnabled)
 		{
 			const FLinearColor Color = FLinearColor(0.f, 0.2f, 1.f, 1.f);
 			if (Secondary.Kind == EControlRigVectorKind::Direction)
 			{
-				Context.DrawInterface->DrawLine(DebugSettings.WorldOffset, Result.GetLocation(), Result.GetLocation() + Target * DebugSettings.Scale, Color);
+				ExecuteContext.UnitContext.DrawInterface->DrawLine(DebugSettings.WorldOffset, Result.GetLocation(), Result.GetLocation() + Target * DebugSettings.Scale, Color);
 			}
 			else
 			{
-				Context.DrawInterface->DrawLine(DebugSettings.WorldOffset, Result.GetLocation(), Target, Color);
-				Context.DrawInterface->DrawBox(DebugSettings.WorldOffset, FTransform(FQuat::Identity, Target, FVector(1.f, 1.f, 1.f) * DebugSettings.Scale * 0.1f), Color);
+				ExecuteContext.UnitContext.DrawInterface->DrawLine(DebugSettings.WorldOffset, Result.GetLocation(), Target, Color);
+				ExecuteContext.UnitContext.DrawInterface->DrawBox(DebugSettings.WorldOffset, FTransform(FQuat::Identity, Target, FVector(1.f, 1.f, 1.f) * DebugSettings.Scale * 0.1f), Color);
 			}
 		}
 
@@ -183,8 +183,7 @@ FRigUnit_AimBone_Execute()
 		DebugSettings,
 		CachedBoneIndex,
 		PrimaryCachedSpace,
-		SecondaryCachedSpace,
-		Context);
+		SecondaryCachedSpace);
 }
 
 FRigVMStructUpgradeInfo FRigUnit_AimBone::GetUpgradeInfo() const
@@ -220,7 +219,7 @@ FRigUnit_AimItem_Execute()
 		return;
 	}
 
-	if (Context.State == EControlRigState::Init)
+	if (ExecuteContext.UnitContext.State == EControlRigState::Init)
 	{
 		CachedItem.Reset();
 		PrimaryCachedSpace.Reset();
@@ -250,8 +249,7 @@ FRigUnit_AimItem_Execute()
 		Transform,
 		DebugSettings,
 		PrimaryCachedSpace,
-		SecondaryCachedSpace,
-		Context);
+		SecondaryCachedSpace);
 
 	Hierarchy->SetGlobalTransform(CachedItem, Transform);
 }
@@ -266,7 +264,7 @@ FRigUnit_AimConstraintLocalSpaceOffset_Execute()
 		return;
 	}
 	
-	if(Context.State == EControlRigState::Init)
+	if(ExecuteContext.UnitContext.State == EControlRigState::Init)
 	{
 		WorldUpSpaceCache.Reset();
 		ChildCache.Reset();
@@ -392,8 +390,7 @@ FRigUnit_AimConstraintLocalSpaceOffset_Execute()
 					InitialAimResult,
 					DummyDebugSettings,
 					PrimaryCachedSpace,
-					SecondaryCachedSpace,
-					Context);
+					SecondaryCachedSpace);
 
 				FTransform ChildParentInitialGlobalTransform = Hierarchy->GetParentTransformByIndex(ChildCache, true);
 				FQuat MixedInitialLocalRotation = ChildParentInitialGlobalTransform.GetRotation().Inverse() * InitialAimResult.GetRotation();
@@ -461,8 +458,7 @@ FRigUnit_AimConstraintLocalSpaceOffset_Execute()
 				AimResult,
 				AdvancedSettings.DebugSettings,
 				PrimaryCachedSpace,
-				WorldUpSpaceCache,
-				Context);	
+				WorldUpSpaceCache);	
 
 			// handle filtering, performed in local space
 			FTransform ChildParentGlobalTransform = Hierarchy->GetParentTransformByIndex(ChildCache, false);
diff --git a/Engine/Plugins/Animation/ControlRig/Source/ControlRig/Private/Units/Highlevel/Hierarchy/RigUnit_CCDIK.cpp b/Engine/Plugins/Animation/ControlRig/Source/ControlRig/Private/Units/Highlevel/Hierarchy/RigUnit_CCDIK.cpp
--- a/Engine/Plugins/Animation/ControlRig/Source/ControlRig/Private/Units/Highlevel/Hierarchy/RigUnit_CCDIK.cpp
+++ b/Engine/Plugins/Animation/ControlRig/Source/ControlRig/Private/Units/Highlevel/Hierarchy/RigUnit_CCDIK.cpp
@@ -7,7 +7,7 @@
 
 FRigUnit_CCDIK_Execute()
 {
-	if (Context.State == EControlRigState::Init)
+	if (ExecuteContext.UnitContext.State == EControlRigState::Init)
 	{
 		WorkData.CachedItems.Reset();
 		return;
@@ -17,7 +17,7 @@ FRigUnit_CCDIK_Execute()
 	if(WorkData.CachedItems.Num() == 0)
 	{
 		Items = FRigElementKeyCollection::MakeFromChain(
-			Context.Hierarchy,
+			ExecuteContext.Hierarchy,
 			FRigElementKey(StartBone, ERigElementType::Bone),
 			FRigElementKey(EffectorBone, ERigElementType::Bone),
 			false /* reverse */
@@ -45,8 +45,7 @@ FRigUnit_CCDIK_Execute()
 		BaseRotationLimit,
 		RotationLimitsPerItem,
 		bPropagateToChildren,
-		WorkData,
-		Context);
+		WorkData);
 }
 
 FRigVMStructUpgradeInfo FRigUnit_CCDIK::GetUpgradeInfo() const
@@ -57,7 +56,7 @@ FRigVMStructUpgradeInfo FRigUnit_CCDIK::GetUpgradeInfo() const
 
 FRigUnit_CCDIKPerItem_Execute()
 {
-	FRigUnit_CCDIKItemArray::StaticExecute(ExecuteContext, Items.Keys, EffectorTransform, Precision, Weight, MaxIterations, bStartFromTail, BaseRotationLimit, RotationLimits, bPropagateToChildren, WorkData, Context);
+	FRigUnit_CCDIKItemArray::StaticExecute(ExecuteContext, Items.Keys, EffectorTransform, Precision, Weight, MaxIterations, bStartFromTail, BaseRotationLimit, RotationLimits, bPropagateToChildren, WorkData);
 }
 
 FRigVMStructUpgradeInfo FRigUnit_CCDIKPerItem::GetUpgradeInfo() const
@@ -91,7 +90,7 @@ FRigUnit_CCDIKItemArray_Execute()
 	TArray<float>& RotationLimitsPerItem = WorkData.RotationLimitsPerItem;
 	FCachedRigElement& CachedEffector = WorkData.CachedEffector;
 
-	if ((Context.State == EControlRigState::Init) ||
+	if ((ExecuteContext.UnitContext.State == EControlRigState::Init) ||
 		(
 			(RotationLimits.Num() != RotationLimitIndex.Num()) &&
 			(RotationLimitIndex.Num() > 0))
@@ -104,7 +103,7 @@ FRigUnit_CCDIKItemArray_Execute()
 		return;
 	}
 	
-	if (Context.State == EControlRigState::Update)
+	if (ExecuteContext.UnitContext.State == EControlRigState::Update)
 	{
 		if (CachedItems.Num() == 0 && Items.Num() > 0)
 		{
diff --git a/Engine/Plugins/Animation/ControlRig/Source/ControlRig/Private/Units/Highlevel/Hierarchy/RigUnit_DistributeRotation.cpp b/Engine/Plugins/Animation/ControlRig/Source/ControlRig/Private/Units/Highlevel/Hierarchy/RigUnit_DistributeRotation.cpp
--- a/Engine/Plugins/Animation/ControlRig/Source/ControlRig/Private/Units/Highlevel/Hierarchy/RigUnit_DistributeRotation.cpp
+++ b/Engine/Plugins/Animation/ControlRig/Source/ControlRig/Private/Units/Highlevel/Hierarchy/RigUnit_DistributeRotation.cpp
@@ -54,8 +54,7 @@ FRigUnit_DistributeRotation_Execute()
 		Rotations,
 		RotationEaseType,
 		Weight,
-		WorkData,
-		Context);
+		WorkData);
 }
 
 FRigVMStructUpgradeInfo FRigUnit_DistributeRotation::GetUpgradeInfo() const
@@ -66,7 +65,7 @@ FRigVMStructUpgradeInfo FRigUnit_DistributeRotation::GetUpgradeInfo() const
 
 FRigUnit_DistributeRotationForCollection_Execute()
 {
-	FRigUnit_DistributeRotationForItemArray::StaticExecute(ExecuteContext, Items.Keys, Rotations, RotationEaseType, Weight, WorkData, Context);
+	FRigUnit_DistributeRotationForItemArray::StaticExecute(ExecuteContext, Items.Keys, Rotations, RotationEaseType, Weight, WorkData);
 }
 
 FRigVMStructUpgradeInfo FRigUnit_DistributeRotationForCollection::GetUpgradeInfo() const
@@ -107,7 +106,7 @@ FRigUnit_DistributeRotationForItemArray_Execute()
 		}
 	}
 
-	if (Context.State == EControlRigState::Init || (CachedItems.Num() > 0 && CachedItems.Num() != Items.Num()))
+	if (ExecuteContext.UnitContext.State == EControlRigState::Init || (CachedItems.Num() > 0 && CachedItems.Num() != Items.Num()))
 	{
 		CachedItems.Reset();
 		ItemRotationA.Reset();
diff --git a/Engine/Plugins/Animation/ControlRig/Source/ControlRig/Private/Units/Highlevel/Hierarchy/RigUnit_FABRIK.cpp b/Engine/Plugins/Animation/ControlRig/Source/ControlRig/Private/Units/Highlevel/Hierarchy/RigUnit_FABRIK.cpp
--- a/Engine/Plugins/Animation/ControlRig/Source/ControlRig/Private/Units/Highlevel/Hierarchy/RigUnit_FABRIK.cpp
+++ b/Engine/Plugins/Animation/ControlRig/Source/ControlRig/Private/Units/Highlevel/Hierarchy/RigUnit_FABRIK.cpp
@@ -7,7 +7,7 @@
 
 FRigUnit_FABRIK_Execute()
 {
-	if (Context.State == EControlRigState::Init)
+	if (ExecuteContext.UnitContext.State == EControlRigState::Init)
 	{
 		WorkData.CachedItems.Reset();
 		return;
@@ -17,7 +17,7 @@ FRigUnit_FABRIK_Execute()
 	if(WorkData.CachedItems.Num() == 0)
 	{
 		Items = FRigElementKeyCollection::MakeFromChain(
-			Context.Hierarchy,
+			ExecuteContext.Hierarchy,
 			FRigElementKey(StartBone, ERigElementType::Bone),
 			FRigElementKey(EffectorBone, ERigElementType::Bone),
 			false /* reverse */
@@ -33,8 +33,7 @@ FRigUnit_FABRIK_Execute()
 		bPropagateToChildren,
 		MaxIterations,
 		WorkData,
-		bSetEffectorTransform,
-		Context);
+		bSetEffectorTransform);
 }
 
 FRigVMStructUpgradeInfo FRigUnit_FABRIK::GetUpgradeInfo() const
@@ -45,7 +44,7 @@ FRigVMStructUpgradeInfo FRigUnit_FABRIK::GetUpgradeInfo() const
 
 FRigUnit_FABRIKPerItem_Execute()
 {
-	FRigUnit_FABRIKItemArray::StaticExecute(ExecuteContext, Items.Keys, EffectorTransform, Precision, Weight, bPropagateToChildren, MaxIterations, WorkData, bSetEffectorTransform, Context);
+	FRigUnit_FABRIKItemArray::StaticExecute(ExecuteContext, Items.Keys, EffectorTransform, Precision, Weight, bPropagateToChildren, MaxIterations, WorkData, bSetEffectorTransform);
 }
 
 FRigVMStructUpgradeInfo FRigUnit_FABRIKPerItem::GetUpgradeInfo() const
@@ -74,14 +73,14 @@ FRigUnit_FABRIKItemArray_Execute()
 	TArray<FCachedRigElement>& CachedItems = WorkData.CachedItems;
 	FCachedRigElement& CachedEffector = WorkData.CachedEffector;
 
-	if (Context.State == EControlRigState::Init)
+	if (ExecuteContext.UnitContext.State == EControlRigState::Init)
 	{
 		CachedItems.Reset();
 		CachedEffector.Reset();
 		return;
 	}
 
-	if (Context.State == EControlRigState::Update)
+	if (ExecuteContext.UnitContext.State == EControlRigState::Update)
 	{
 		if (CachedItems.Num() == 0 && Items.Num() > 0)
 		{
diff --git a/Engine/Plugins/Animation/ControlRig/Source/ControlRig/Private/Units/Highlevel/Hierarchy/RigUnit_FitChainToCurve.cpp b/Engine/Plugins/Animation/ControlRig/Source/ControlRig/Private/Units/Highlevel/Hierarchy/RigUnit_FitChainToCurve.cpp
--- a/Engine/Plugins/Animation/ControlRig/Source/ControlRig/Private/Units/Highlevel/Hierarchy/RigUnit_FitChainToCurve.cpp
+++ b/Engine/Plugins/Animation/ControlRig/Source/ControlRig/Private/Units/Highlevel/Hierarchy/RigUnit_FitChainToCurve.cpp
@@ -7,7 +7,7 @@
 
 FRigUnit_FitChainToCurve_Execute()
 {
-	if (Context.State == EControlRigState::Init)
+	if (ExecuteContext.UnitContext.State == EControlRigState::Init)
 	{
 		WorkData.CachedItems.Reset();
 		return;
@@ -17,7 +17,7 @@ FRigUnit_FitChainToCurve_Execute()
 	if(WorkData.CachedItems.Num() == 0)
 	{
 		Items = FRigElementKeyCollection::MakeFromChain(
-			Context.Hierarchy,
+			ExecuteContext.Hierarchy,
 			FRigElementKey(StartBone, ERigElementType::Bone),
 			FRigElementKey(EndBone, ERigElementType::Bone),
 			false /* reverse */
@@ -40,8 +40,7 @@ FRigUnit_FitChainToCurve_Execute()
 		Weight,
 		bPropagateToChildren,
 		DebugSettings,
-		WorkData,
-		Context);
+		WorkData);
 }
 
 FRigVMStructUpgradeInfo FRigUnit_FitChainToCurve::GetUpgradeInfo() const
@@ -68,8 +67,7 @@ FRigUnit_FitChainToCurvePerItem_Execute()
 		Weight,
 		bPropagateToChildren,
 		DebugSettings,
-		WorkData,
-		Context);
+		WorkData);
 }
 
 FRigVMStructUpgradeInfo FRigUnit_FitChainToCurvePerItem::GetUpgradeInfo() const
@@ -99,7 +97,7 @@ FRigUnit_FitChainToCurveItemArray_Execute()
 	TArray<float>& ItemRotationT = WorkData.ItemRotationT;
 	TArray<FTransform>& ItemLocalTransforms = WorkData.ItemLocalTransforms;
 
-	if (Context.State == EControlRigState::Init)
+	if (ExecuteContext.UnitContext.State == EControlRigState::Init)
 	{
 		CachedItems.Reset();
 		return;
@@ -454,17 +452,17 @@ FRigUnit_FitChainToCurveItemArray_Execute()
 		}
 	}
 
-	if (Context.DrawInterface != nullptr && DebugSettings.bEnabled)
+	if (ExecuteContext.UnitContext.DrawInterface != nullptr && DebugSettings.bEnabled)
 	{
-		Context.DrawInterface->DrawBezier(DebugSettings.WorldOffset, Bezier, 0.f, 1.f, DebugSettings.CurveColor, DebugSettings.Scale, 64);
-		Context.DrawInterface->DrawPoint(DebugSettings.WorldOffset, Bezier.A, DebugSettings.Scale * 6, DebugSettings.CurveColor);
-		Context.DrawInterface->DrawPoint(DebugSettings.WorldOffset, Bezier.B, DebugSettings.Scale * 6, DebugSettings.CurveColor);
-		Context.DrawInterface->DrawPoint(DebugSettings.WorldOffset, Bezier.C, DebugSettings.Scale * 6, DebugSettings.CurveColor);
-		Context.DrawInterface->DrawPoint(DebugSettings.WorldOffset, Bezier.D, DebugSettings.Scale * 6, DebugSettings.CurveColor);
-		Context.DrawInterface->DrawLineStrip(DebugSettings.WorldOffset, CurvePositions, DebugSettings.SegmentsColor, DebugSettings.Scale);
-		Context.DrawInterface->DrawPoints(DebugSettings.WorldOffset, CurvePositions, DebugSettings.Scale * 4.f, DebugSettings.SegmentsColor);
-		// Context.DrawInterface->DrawPoints(DebugSettings.WorldOffset, CurvePositions, DebugSettings.Scale * 3.f, FLinearColor::Blue);
-		// Context.DrawInterface->DrawPoints(DebugSettings.WorldOffset, ItemPositions, DebugSettings.Scale * 6.f, DebugSettings.SegmentsColor);
+		ExecuteContext.UnitContext.DrawInterface->DrawBezier(DebugSettings.WorldOffset, Bezier, 0.f, 1.f, DebugSettings.CurveColor, DebugSettings.Scale, 64);
+		ExecuteContext.UnitContext.DrawInterface->DrawPoint(DebugSettings.WorldOffset, Bezier.A, DebugSettings.Scale * 6, DebugSettings.CurveColor);
+		ExecuteContext.UnitContext.DrawInterface->DrawPoint(DebugSettings.WorldOffset, Bezier.B, DebugSettings.Scale * 6, DebugSettings.CurveColor);
+		ExecuteContext.UnitContext.DrawInterface->DrawPoint(DebugSettings.WorldOffset, Bezier.C, DebugSettings.Scale * 6, DebugSettings.CurveColor);
+		ExecuteContext.UnitContext.DrawInterface->DrawPoint(DebugSettings.WorldOffset, Bezier.D, DebugSettings.Scale * 6, DebugSettings.CurveColor);
+		ExecuteContext.UnitContext.DrawInterface->DrawLineStrip(DebugSettings.WorldOffset, CurvePositions, DebugSettings.SegmentsColor, DebugSettings.Scale);
+		ExecuteContext.UnitContext.DrawInterface->DrawPoints(DebugSettings.WorldOffset, CurvePositions, DebugSettings.Scale * 4.f, DebugSettings.SegmentsColor);
+		// ExecuteContext.UnitContext.DrawInterface->DrawPoints(DebugSettings.WorldOffset, CurvePositions, DebugSettings.Scale * 3.f, FLinearColor::Blue);
+		// ExecuteContext.UnitContext.DrawInterface->DrawPoints(DebugSettings.WorldOffset, ItemPositions, DebugSettings.Scale * 6.f, DebugSettings.SegmentsColor);
 	}
 }
 
diff --git a/Engine/Plugins/Animation/ControlRig/Source/ControlRig/Private/Units/Highlevel/Hierarchy/RigUnit_ModifyBoneTransforms.cpp b/Engine/Plugins/Animation/ControlRig/Source/ControlRig/Private/Units/Highlevel/Hierarchy/RigUnit_ModifyBoneTransforms.cpp
--- a/Engine/Plugins/Animation/ControlRig/Source/ControlRig/Private/Units/Highlevel/Hierarchy/RigUnit_ModifyBoneTransforms.cpp
+++ b/Engine/Plugins/Animation/ControlRig/Source/ControlRig/Private/Units/Highlevel/Hierarchy/RigUnit_ModifyBoneTransforms.cpp
@@ -23,8 +23,7 @@ FRigUnit_ModifyBoneTransforms_Execute()
 		WeightMinimum,
 		WeightMaximum,
 		Mode,
-		WorkData,
-		Context);
+		WorkData);
 }
 
 FRigVMStructUpgradeInfo FRigUnit_ModifyBoneTransforms::GetUpgradeInfo() const
diff --git a/Engine/Plugins/Animation/ControlRig/Source/ControlRig/Private/Units/Highlevel/Hierarchy/RigUnit_MultiFABRIK.cpp b/Engine/Plugins/Animation/ControlRig/Source/ControlRig/Private/Units/Highlevel/Hierarchy/RigUnit_MultiFABRIK.cpp
--- a/Engine/Plugins/Animation/ControlRig/Source/ControlRig/Private/Units/Highlevel/Hierarchy/RigUnit_MultiFABRIK.cpp
+++ b/Engine/Plugins/Animation/ControlRig/Source/ControlRig/Private/Units/Highlevel/Hierarchy/RigUnit_MultiFABRIK.cpp
@@ -228,7 +228,7 @@ FRigUnit_MultiFABRIK_Execute()
 	TArray<FRigUnit_MultiFABRIK_BoneWorkingData>&	BoneTree = WorkData.BoneTree;
 	TArray<FCachedRigElement>&				        EffectorIndices = WorkData.EffectorBoneIndices;
 	
-	if (Context.State == EControlRigState::Init)
+	if (ExecuteContext.UnitContext.State == EControlRigState::Init)
 	{
 		EffectorIndices.Reset();
 		BoneTree.Reset();
@@ -236,7 +236,7 @@ FRigUnit_MultiFABRIK_Execute()
 		return;
 	}
 
-	if (Context.State == EControlRigState::Update)
+	if (ExecuteContext.UnitContext.State == EControlRigState::Update)
 	{
 		if(BoneTree.Num() == 0)
 		{
diff --git a/Engine/Plugins/Animation/ControlRig/Source/ControlRig/Private/Units/Highlevel/Hierarchy/RigUnit_SlideChain.cpp b/Engine/Plugins/Animation/ControlRig/Source/ControlRig/Private/Units/Highlevel/Hierarchy/RigUnit_SlideChain.cpp
--- a/Engine/Plugins/Animation/ControlRig/Source/ControlRig/Private/Units/Highlevel/Hierarchy/RigUnit_SlideChain.cpp
+++ b/Engine/Plugins/Animation/ControlRig/Source/ControlRig/Private/Units/Highlevel/Hierarchy/RigUnit_SlideChain.cpp
@@ -8,7 +8,7 @@
 
 FRigUnit_SlideChain_Execute()
 {
-	if (Context.State == EControlRigState::Init)
+	if (ExecuteContext.UnitContext.State == EControlRigState::Init)
 	{
 		WorkData.CachedItems.Reset();
 		return;
@@ -18,7 +18,7 @@ FRigUnit_SlideChain_Execute()
 	if(WorkData.CachedItems.Num() == 0)
 	{
 		Items = FRigElementKeyCollection::MakeFromChain(
-			Context.Hierarchy,
+			ExecuteContext.Hierarchy,
 			FRigElementKey(StartBone, ERigElementType::Bone),
 			FRigElementKey(EndBone, ERigElementType::Bone),
 			false /* reverse */
@@ -30,8 +30,7 @@ FRigUnit_SlideChain_Execute()
 		Items,
 		SlideAmount,
 		bPropagateToChildren,
-		WorkData,
-		Context);
+		WorkData);
 }
 
 FRigVMStructUpgradeInfo FRigUnit_SlideChain::GetUpgradeInfo() const
@@ -42,7 +41,7 @@ FRigVMStructUpgradeInfo FRigUnit_SlideChain::GetUpgradeInfo() const
 
 FRigUnit_SlideChainPerItem_Execute()
 {
-	FRigUnit_SlideChainItemArray::StaticExecute(ExecuteContext, Items.Keys, SlideAmount, bPropagateToChildren, WorkData, Context);
+	FRigUnit_SlideChainItemArray::StaticExecute(ExecuteContext, Items.Keys, SlideAmount, bPropagateToChildren, WorkData);
 }
 
 FRigVMStructUpgradeInfo FRigUnit_SlideChainPerItem::GetUpgradeInfo() const
@@ -69,7 +68,7 @@ FRigUnit_SlideChainItemArray_Execute()
 	TArray<FTransform>& Transforms = WorkData.Transforms;
 	TArray<FTransform>& BlendedTransforms = WorkData.BlendedTransforms;
 
-	if (Context.State == EControlRigState::Init)
+	if (ExecuteContext.UnitContext.State == EControlRigState::Init)
 	{
 		CachedItems.Reset();
 		return;
diff --git a/Engine/Plugins/Animation/ControlRig/Source/ControlRig/Private/Units/Highlevel/Hierarchy/RigUnit_SphericalPoseReader.cpp b/Engine/Plugins/Animation/ControlRig/Source/ControlRig/Private/Units/Highlevel/Hierarchy/RigUnit_SphericalPoseReader.cpp
--- a/Engine/Plugins/Animation/ControlRig/Source/ControlRig/Private/Units/Highlevel/Hierarchy/RigUnit_SphericalPoseReader.cpp
+++ b/Engine/Plugins/Animation/ControlRig/Source/ControlRig/Private/Units/Highlevel/Hierarchy/RigUnit_SphericalPoseReader.cpp
@@ -10,13 +10,13 @@ FRigUnit_SphericalPoseReader_Execute()
 {
 	DECLARE_SCOPE_HIERARCHICAL_COUNTER_RIGUNIT()
 
-	URigHierarchy* Hierarchy = Context.Hierarchy;
+	URigHierarchy* Hierarchy = ExecuteContext.Hierarchy;
 	if (Hierarchy == nullptr)
 	{
 		return;
 	}
 
-	if (Context.State == EControlRigState::Init)
+	if (ExecuteContext.UnitContext.State == EControlRigState::Init)
 	{
 		DriverCache.Reset();
 		OptionalParentCache.Reset();
@@ -101,7 +101,7 @@ FRigUnit_SphericalPoseReader_Execute()
 	{
 		// avoid NaNs from DistanceToEllipse, and guaranteed to be inside inner ellipse
 		OutputParam = 1.0f;
-		Debug.DrawDebug(WorldOffset, Context.DrawInterface, InnerRegion, OuterRegion, DriverNormal, Driver2D, OutputParam);
+		Debug.DrawDebug(WorldOffset, ExecuteContext.UnitContext.DrawInterface, InnerRegion, OuterRegion, DriverNormal, Driver2D, OutputParam);
 		return;
 	}
 	
@@ -125,7 +125,7 @@ FRigUnit_SphericalPoseReader_Execute()
 	OutputParam = CalcOutputParam(InnerEllipseResults, OuterEllipseResults);
 
 	// do all debug drawing
-	Debug.DrawDebug(WorldOffset, Context.DrawInterface, InnerRegion, OuterRegion, DriverNormal, Driver2D, OutputParam);
+	Debug.DrawDebug(WorldOffset, ExecuteContext.UnitContext.DrawInterface, InnerRegion, OuterRegion, DriverNormal, Driver2D, OutputParam);
 }
 
 void FRigUnit_SphericalPoseReader::RemapAndConvertInputs(
diff --git a/Engine/Plugins/Animation/ControlRig/Source/ControlRig/Private/Units/Highlevel/Hierarchy/RigUnit_SpringIK.cpp b/Engine/Plugins/Animation/ControlRig/Source/ControlRig/Private/Units/Highlevel/Hierarchy/RigUnit_SpringIK.cpp
--- a/Engine/Plugins/Animation/ControlRig/Source/ControlRig/Private/Units/Highlevel/Hierarchy/RigUnit_SpringIK.cpp
+++ b/Engine/Plugins/Animation/ControlRig/Source/ControlRig/Private/Units/Highlevel/Hierarchy/RigUnit_SpringIK.cpp
@@ -19,7 +19,7 @@ FRigUnit_SpringIK_Execute()
 	TArray<FTransform>& Transforms = WorkData.Transforms;
 	FCRSimPointContainer& Simulation = WorkData.Simulation;
 
-	if (Context.State == EControlRigState::Init)
+	if (ExecuteContext.UnitContext.State == EControlRigState::Init)
 	{
 		CachedBones.Reset();
 		CachedPoleVector.Reset();
@@ -215,7 +215,7 @@ FRigUnit_SpringIK_Execute()
 		}
 	}
 
-	Simulation.StepSemiExplicitEuler((bLiveSimulation ? Context.DeltaTime : Simulation.TimeStep) * (float)FMath::Clamp<int32>(Iterations, 1, 32));
+	Simulation.StepSemiExplicitEuler((bLiveSimulation ? ExecuteContext.UnitContext.DeltaTime : Simulation.TimeStep) * (float)FMath::Clamp<int32>(Iterations, 1, 32));
 
 	FVector AccumulatedTarget = FVector::ZeroVector;
 	FVector LastPrimaryTarget = FVector::ZeroVector;
@@ -308,11 +308,11 @@ FRigUnit_SpringIK_Execute()
 		Hierarchy->SetGlobalTransform(CachedBones[PointIndex], Transform, bPropagateToChildren);
 	}
 
-	if (Context.DrawInterface != nullptr && DebugSettings.bEnabled)
+	if (ExecuteContext.UnitContext.DrawInterface != nullptr && DebugSettings.bEnabled)
 	{
-		Context.DrawInterface->DrawPointSimulation(DebugSettings.WorldOffset, Simulation, DebugSettings.Color, DebugSettings.Scale * 0.25f);
-		Context.DrawInterface->DrawLine(DebugSettings.WorldOffset, PoleTarget, FirstPoint, DebugSettings.Color, 0.f);
-		Context.DrawInterface->DrawLine(DebugSettings.WorldOffset, PoleTarget, LastPoint, DebugSettings.Color, 0.f);
-		Context.DrawInterface->DrawBox(DebugSettings.WorldOffset, FTransform(FQuat::Identity, PoleTarget, FVector::OneVector * DebugSettings.Scale * 10.f), DebugSettings.Color);
+		ExecuteContext.UnitContext.DrawInterface->DrawPointSimulation(DebugSettings.WorldOffset, Simulation, DebugSettings.Color, DebugSettings.Scale * 0.25f);
+		ExecuteContext.UnitContext.DrawInterface->DrawLine(DebugSettings.WorldOffset, PoleTarget, FirstPoint, DebugSettings.Color, 0.f);
+		ExecuteContext.UnitContext.DrawInterface->DrawLine(DebugSettings.WorldOffset, PoleTarget, LastPoint, DebugSettings.Color, 0.f);
+		ExecuteContext.UnitContext.DrawInterface->DrawBox(DebugSettings.WorldOffset, FTransform(FQuat::Identity, PoleTarget, FVector::OneVector * DebugSettings.Scale * 10.f), DebugSettings.Color);
 	}
 }
diff --git a/Engine/Plugins/Animation/ControlRig/Source/ControlRig/Private/Units/Highlevel/Hierarchy/RigUnit_TransformConstraint.cpp b/Engine/Plugins/Animation/ControlRig/Source/ControlRig/Private/Units/Highlevel/Hierarchy/RigUnit_TransformConstraint.cpp
--- a/Engine/Plugins/Animation/ControlRig/Source/ControlRig/Private/Units/Highlevel/Hierarchy/RigUnit_TransformConstraint.cpp
+++ b/Engine/Plugins/Animation/ControlRig/Source/ControlRig/Private/Units/Highlevel/Hierarchy/RigUnit_TransformConstraint.cpp
@@ -17,8 +17,7 @@ FRigUnit_TransformConstraint_Execute()
 		FRigElementKey(BaseBone, ERigElementType::Bone),
 		Targets,
 		bUseInitialTransforms,
-		WorkData,
-		Context);
+		WorkData);
 }
 
 FRigVMStructUpgradeInfo FRigUnit_TransformConstraint::GetUpgradeInfo() const
@@ -111,7 +110,7 @@ FRigUnit_TransformConstraintPerItem_Execute()
 		}
 	};
 
-	if (Context.State == EControlRigState::Init)
+	if (ExecuteContext.UnitContext.State == EControlRigState::Init)
 	{
 		ConstraintData.Reset();
 		ConstraintDataToTargets.Reset();
@@ -122,7 +121,7 @@ FRigUnit_TransformConstraintPerItem_Execute()
 			SetupConstraintData();
 		}
 	}
-	else if (Context.State == EControlRigState::Update)
+	else if (ExecuteContext.UnitContext.State == EControlRigState::Update)
 	{
 		URigHierarchy* Hierarchy = ExecuteContext.Hierarchy;
 		if (Hierarchy)
@@ -196,7 +195,7 @@ FRigUnit_ParentConstraint_Execute()
 		return;
 	}
 
-	if(Context.State == EControlRigState::Init)
+	if(ExecuteContext.UnitContext.State == EControlRigState::Init)
 	{
 		ChildCache.Reset();
 		ParentCaches.Reset();
@@ -637,7 +636,7 @@ FRigUnit_PositionConstraintLocalSpaceOffset_Execute()
 		return;
 	}
 
-	if(Context.State == EControlRigState::Init)
+	if(ExecuteContext.UnitContext.State == EControlRigState::Init)
 	{
 		ChildCache.Reset();
 		ParentCaches.Reset();
@@ -1119,7 +1118,7 @@ FRigUnit_RotationConstraintLocalSpaceOffset_Execute()
 		return;
 	}
 	
-	if(Context.State == EControlRigState::Init)
+	if(ExecuteContext.UnitContext.State == EControlRigState::Init)
 	{
 		ChildCache.Reset();
 		ParentCaches.Reset();
@@ -1492,7 +1491,7 @@ IMPLEMENT_RIGUNIT_AUTOMATION_TEST(FRigUnit_RotationConstraintLocalSpaceOffset)
 
 FRigUnit_ScaleConstraint_Execute()
 {
-	TFunction<FVector (const FVector&)> GetNonZeroScale([ExecuteContext, Context](const FVector& InScale)
+	TFunction<FVector (const FVector&)> GetNonZeroScale([ExecuteContext](const FVector& InScale)
 	{
 		FVector NonZeroScale = InScale;
        
@@ -1698,7 +1697,7 @@ IMPLEMENT_RIGUNIT_AUTOMATION_TEST(FRigUnit_ScaleConstraint)
 
 FRigUnit_ScaleConstraintLocalSpaceOffset_Execute()
 {
-	TFunction<FVector (const FVector&)> GetNonZeroScale([ExecuteContext, Context](const FVector& InScale)
+	TFunction<FVector (const FVector&)> GetNonZeroScale([ExecuteContext](const FVector& InScale)
 	{
 		FVector NonZeroScale = InScale;
        
@@ -1732,7 +1731,7 @@ FRigUnit_ScaleConstraintLocalSpaceOffset_Execute()
 		return;
 	}
 	
-	if(Context.State == EControlRigState::Init)
+	if(ExecuteContext.UnitContext.State == EControlRigState::Init)
 	{
 		ChildCache.Reset();
 		ParentCaches.Reset();
diff --git a/Engine/Plugins/Animation/ControlRig/Source/ControlRig/Private/Units/Highlevel/Hierarchy/RigUnit_TransformConstraint.h b/Engine/Plugins/Animation/ControlRig/Source/ControlRig/Private/Units/Highlevel/Hierarchy/RigUnit_TransformConstraint.h
--- a/Engine/Plugins/Animation/ControlRig/Source/ControlRig/Private/Units/Highlevel/Hierarchy/RigUnit_TransformConstraint.h
+++ b/Engine/Plugins/Animation/ControlRig/Source/ControlRig/Private/Units/Highlevel/Hierarchy/RigUnit_TransformConstraint.h
@@ -79,7 +79,7 @@ struct CONTROLRIG_API FRigUnit_TransformConstraint : public FRigUnit_HighlevelBa
 	}
 
 	RIGVM_METHOD()
-	virtual void Execute(const FRigUnitContext& Context) override;
+	virtual void Execute() override;
 
 	UPROPERTY(EditAnywhere, Category = "Constraint", meta = (Input))
 	FName Bone;
@@ -149,7 +149,7 @@ struct CONTROLRIG_API FRigUnit_TransformConstraintPerItem : public FRigUnit_High
 	}
 
 	RIGVM_METHOD()
-	virtual void Execute(const FRigUnitContext& Context) override;
+	virtual void Execute() override;
 
 	UPROPERTY(EditAnywhere, Category = "Constraint", meta = (Input, ExpandByDefault))
 	FRigElementKey Item;
@@ -262,7 +262,7 @@ struct FRigUnit_ParentConstraint : public FRigUnit_HighlevelBaseMutable
 	}
 	
 	RIGVM_METHOD()
-    virtual void Execute(const FRigUnitContext& Context) override;
+    virtual void Execute() override;
 
 	UPROPERTY(EditAnywhere, Category = "Constraint", meta = (Input, ExpandByDefault))
 	FRigElementKey Child;
@@ -306,7 +306,7 @@ struct FRigUnit_PositionConstraint : public FRigUnit_HighlevelBaseMutable
 	}
 	
 	RIGVM_METHOD()
-    virtual void Execute(const FRigUnitContext& Context) override;
+    virtual void Execute() override;
 
 	UPROPERTY(EditAnywhere, Category = "Constraint", meta = (Input, ExpandByDefault))
 	FRigElementKey Child;
@@ -346,7 +346,7 @@ struct FRigUnit_PositionConstraintLocalSpaceOffset : public FRigUnit_HighlevelBa
 	}
 	
 	RIGVM_METHOD()
-	virtual void Execute(const FRigUnitContext& Context) override;
+	virtual void Execute() override;
 
 	UPROPERTY(EditAnywhere, Category = "Constraint", meta = (Input, ExpandByDefault))
 	FRigElementKey Child;
@@ -411,7 +411,7 @@ struct FRigUnit_RotationConstraint : public FRigUnit_HighlevelBaseMutable
 	}
 	
 	RIGVM_METHOD()
-    virtual void Execute(const FRigUnitContext& Context) override;
+    virtual void Execute() override;
 
 	UPROPERTY(EditAnywhere, Category = "Constraint", meta = (Input, ExpandByDefault))
 	FRigElementKey Child;
@@ -454,7 +454,7 @@ struct FRigUnit_RotationConstraintLocalSpaceOffset : public FRigUnit_HighlevelBa
 	}
 	
 	RIGVM_METHOD()
-	virtual void Execute(const FRigUnitContext& Context) override;
+	virtual void Execute() override;
 
 	UPROPERTY(EditAnywhere, Category = "Constraint", meta = (Input, ExpandByDefault))
 	FRigElementKey Child;
@@ -498,7 +498,7 @@ struct FRigUnit_ScaleConstraint : public FRigUnit_HighlevelBaseMutable
 	}
 	
 	RIGVM_METHOD()
-    virtual void Execute(const FRigUnitContext& Context) override;
+    virtual void Execute() override;
 
 	UPROPERTY(EditAnywhere, Category = "Constraint", meta = (Input, ExpandByDefault))
 	FRigElementKey Child;
@@ -538,7 +538,7 @@ struct FRigUnit_ScaleConstraintLocalSpaceOffset : public FRigUnit_HighlevelBaseM
 	}
 	
 	RIGVM_METHOD()
-	virtual void Execute(const FRigUnitContext& Context) override;
+	virtual void Execute() override;
 
 	UPROPERTY(EditAnywhere, Category = "Constraint", meta = (Input, ExpandByDefault))
 	FRigElementKey Child;
diff --git a/Engine/Plugins/Animation/ControlRig/Source/ControlRig/Private/Units/Highlevel/Hierarchy/RigUnit_TwistBones.cpp b/Engine/Plugins/Animation/ControlRig/Source/ControlRig/Private/Units/Highlevel/Hierarchy/RigUnit_TwistBones.cpp
--- a/Engine/Plugins/Animation/ControlRig/Source/ControlRig/Private/Units/Highlevel/Hierarchy/RigUnit_TwistBones.cpp
+++ b/Engine/Plugins/Animation/ControlRig/Source/ControlRig/Private/Units/Highlevel/Hierarchy/RigUnit_TwistBones.cpp
@@ -7,7 +7,7 @@
 
 FRigUnit_TwistBones_Execute()
 {
-	if (Context.State == EControlRigState::Init)
+	if (ExecuteContext.UnitContext.State == EControlRigState::Init)
 	{
 		WorkData.CachedItems.Reset();
 		return;
@@ -17,7 +17,7 @@ FRigUnit_TwistBones_Execute()
 	if(WorkData.CachedItems.Num() == 0)
 	{
 		Items = FRigElementKeyCollection::MakeFromChain(
-			Context.Hierarchy,
+			ExecuteContext.Hierarchy,
 			FRigElementKey(StartBone, ERigElementType::Bone),
 			FRigElementKey(EndBone, ERigElementType::Bone),
 			false /* reverse */
@@ -32,8 +32,7 @@ FRigUnit_TwistBones_Execute()
 		TwistEaseType,
 		Weight,
 		bPropagateToChildren,
-		WorkData,
-		Context);
+		WorkData);
 }
 
 FRigVMStructUpgradeInfo FRigUnit_TwistBones::GetUpgradeInfo() const
@@ -55,7 +54,7 @@ FRigUnit_TwistBonesPerItem_Execute()
 	TArray<float>& ItemRatios = WorkData.ItemRatios;
 	TArray<FTransform>& ItemTransforms = WorkData.ItemTransforms;
 
-	if (Context.State == EControlRigState::Init)
+	if (ExecuteContext.UnitContext.State == EControlRigState::Init)
 	{
 		CachedItems.Reset();
 		return;
diff --git a/Engine/Plugins/Animation/ControlRig/Source/ControlRig/Private/Units/Highlevel/Hierarchy/RigUnit_TwoBoneIKSimple.cpp b/Engine/Plugins/Animation/ControlRig/Source/ControlRig/Private/Units/Highlevel/Hierarchy/RigUnit_TwoBoneIKSimple.cpp
--- a/Engine/Plugins/Animation/ControlRig/Source/ControlRig/Private/Units/Highlevel/Hierarchy/RigUnit_TwoBoneIKSimple.cpp
+++ b/Engine/Plugins/Animation/ControlRig/Source/ControlRig/Private/Units/Highlevel/Hierarchy/RigUnit_TwoBoneIKSimple.cpp
@@ -32,8 +32,7 @@ FRigUnit_TwoBoneIKSimple_Execute()
 		CachedBoneAIndex,
 		CachedBoneBIndex,
 		CachedEffectorBoneIndex,
-		CachedPoleVectorSpaceIndex,
-		Context);
+		CachedPoleVectorSpaceIndex);
 }
 
 FRigVMStructUpgradeInfo FRigUnit_TwoBoneIKSimple::GetUpgradeInfo() const
@@ -76,7 +75,7 @@ FRigUnit_TwoBoneIKSimplePerItem_Execute()
 		return;
 	}
 
-	if (Context.State == EControlRigState::Init)
+	if (ExecuteContext.UnitContext.State == EControlRigState::Init)
 	{
 		CachedItemAIndex.Reset();
 		CachedItemBIndex.Reset();
@@ -155,14 +154,14 @@ FRigUnit_TwoBoneIKSimplePerItem_Execute()
 
 	FControlRigMathLibrary::SolveBasicTwoBoneIK(TransformA, TransformB, TransformC, PoleTarget, PrimaryAxis, SecondaryAxis, SecondaryAxisWeight, LengthA, LengthB, bEnableStretch, StretchStartRatio, StretchMaximumRatio);
 
-	if (Context.DrawInterface != nullptr && DebugSettings.bEnabled)
+	if (ExecuteContext.UnitContext.DrawInterface != nullptr && DebugSettings.bEnabled)
 	{
 		const FLinearColor Dark = FLinearColor(0.f, 0.2f, 1.f, 1.f);
 		const FLinearColor Bright = FLinearColor(0.f, 1.f, 1.f, 1.f);
-		Context.DrawInterface->DrawLine(DebugSettings.WorldOffset, TransformA.GetLocation(), TransformB.GetLocation(), Dark);
-		Context.DrawInterface->DrawLine(DebugSettings.WorldOffset, TransformB.GetLocation(), TransformC.GetLocation(), Dark);
-		Context.DrawInterface->DrawLine(DebugSettings.WorldOffset, TransformB.GetLocation(), PoleTarget, Bright);
-		Context.DrawInterface->DrawBox(DebugSettings.WorldOffset, FTransform(FQuat::Identity, PoleTarget, FVector(1.f, 1.f, 1.f) * DebugSettings.Scale * 0.1f), Bright);
+		ExecuteContext.UnitContext.DrawInterface->DrawLine(DebugSettings.WorldOffset, TransformA.GetLocation(), TransformB.GetLocation(), Dark);
+		ExecuteContext.UnitContext.DrawInterface->DrawLine(DebugSettings.WorldOffset, TransformB.GetLocation(), TransformC.GetLocation(), Dark);
+		ExecuteContext.UnitContext.DrawInterface->DrawLine(DebugSettings.WorldOffset, TransformB.GetLocation(), PoleTarget, Bright);
+		ExecuteContext.UnitContext.DrawInterface->DrawBox(DebugSettings.WorldOffset, FTransform(FQuat::Identity, PoleTarget, FVector(1.f, 1.f, 1.f) * DebugSettings.Scale * 0.1f), Bright);
 	}
 
 	if (Weight < 1.0f - SMALL_NUMBER)
diff --git a/Engine/Plugins/Animation/ControlRig/Source/ControlRig/Private/Units/Math/RigUnit_MathBool.cpp b/Engine/Plugins/Animation/ControlRig/Source/ControlRig/Private/Units/Math/RigUnit_MathBool.cpp
--- a/Engine/Plugins/Animation/ControlRig/Source/ControlRig/Private/Units/Math/RigUnit_MathBool.cpp
+++ b/Engine/Plugins/Animation/ControlRig/Source/ControlRig/Private/Units/Math/RigUnit_MathBool.cpp
@@ -85,7 +85,7 @@ FRigUnit_MathBoolToggled_Execute()
 {
 	DECLARE_SCOPE_HIERARCHICAL_COUNTER_RIGUNIT()
 
-	if(Context.State == EControlRigState::Init)
+	if(ExecuteContext.UnitContext.State == EControlRigState::Init)
 	{
 		Initialized = false;
 		LastValue = Value;
@@ -110,14 +110,14 @@ FRigUnit_MathBoolFlipFlop_Execute()
 {
 	DECLARE_SCOPE_HIERARCHICAL_COUNTER_RIGUNIT()
 
-	if(Context.State == EControlRigState::Init)
+	if(ExecuteContext.UnitContext.State == EControlRigState::Init)
 	{
 		Result = LastValue = !StartValue;
 		TimeLeft = 0.f;
 		return;
 	}
 
-	TimeLeft = TimeLeft - Context.DeltaTime;
+	TimeLeft = TimeLeft - ExecuteContext.UnitContext.DeltaTime;
 	if (TimeLeft <= 0.f)
 	{
 		LastValue = !LastValue;
@@ -131,7 +131,7 @@ FRigUnit_MathBoolOnce_Execute()
 {
 	DECLARE_SCOPE_HIERARCHICAL_COUNTER_RIGUNIT()
 
-	if(Context.State == EControlRigState::Init)
+	if(ExecuteContext.UnitContext.State == EControlRigState::Init)
 	{
 		Result = false;
 		LastValue = true;
@@ -146,7 +146,7 @@ FRigUnit_MathBoolOnce_Execute()
 
 	Result = LastValue;
 
-	TimeLeft -= Context.DeltaTime;
+	TimeLeft -= ExecuteContext.UnitContext.DeltaTime;
 	if(TimeLeft <= 0.f)
 	{
 		LastValue = false;
diff --git a/Engine/Plugins/Animation/ControlRig/Source/ControlRig/Private/Units/Math/RigUnit_MathBool.h b/Engine/Plugins/Animation/ControlRig/Source/ControlRig/Private/Units/Math/RigUnit_MathBool.h
--- a/Engine/Plugins/Animation/ControlRig/Source/ControlRig/Private/Units/Math/RigUnit_MathBool.h
+++ b/Engine/Plugins/Animation/ControlRig/Source/ControlRig/Private/Units/Math/RigUnit_MathBool.h
@@ -96,7 +96,7 @@ struct CONTROLRIG_API FRigUnit_MathBoolMake : public FRigUnit_MathBoolBase
 	}
 	
 	RIGVM_METHOD()
-	virtual void Execute(const FRigUnitContext& Context) override;
+	virtual void Execute() override;
 
 	UPROPERTY(meta=(Input, Output))
 	bool Value;
@@ -116,7 +116,7 @@ struct CONTROLRIG_API FRigUnit_MathBoolConstTrue : public FRigUnit_MathBoolConst
 	}
 
 	RIGVM_METHOD()
-	virtual void Execute(const FRigUnitContext& Context) override;
+	virtual void Execute() override;
 };
 
 /**
@@ -133,7 +133,7 @@ struct CONTROLRIG_API FRigUnit_MathBoolConstFalse : public FRigUnit_MathBoolCons
 	}
 
 	RIGVM_METHOD()
-	virtual void Execute(const FRigUnitContext& Context) override;
+	virtual void Execute() override;
 };
 
 /**
@@ -145,7 +145,7 @@ struct CONTROLRIG_API FRigUnit_MathBoolNot : public FRigUnit_MathBoolUnaryOp
 	GENERATED_BODY()
 
 	RIGVM_METHOD()
-	virtual void Execute(const FRigUnitContext& Context) override;
+	virtual void Execute() override;
 };
 
 /**
@@ -157,7 +157,7 @@ struct CONTROLRIG_API FRigUnit_MathBoolAnd : public FRigUnit_MathBoolBinaryAggre
 	GENERATED_BODY()
 
 	RIGVM_METHOD()
-	virtual void Execute(const FRigUnitContext& Context) override;
+	virtual void Execute() override;
 };
 
 /**
@@ -169,7 +169,7 @@ struct CONTROLRIG_API FRigUnit_MathBoolNand : public FRigUnit_MathBoolBinaryOp
 	GENERATED_BODY()
 
 	RIGVM_METHOD()
-	virtual void Execute(const FRigUnitContext& Context) override;
+	virtual void Execute() override;
 
 	RIGVM_METHOD()
 	virtual FRigVMStructUpgradeInfo GetUpgradeInfo() const override;
@@ -184,7 +184,7 @@ struct CONTROLRIG_API FRigUnit_MathBoolNand2 : public FRigUnit_MathBoolBinaryOp
 	GENERATED_BODY()
 
 	RIGVM_METHOD()
-	virtual void Execute(const FRigUnitContext& Context) override;
+	virtual void Execute() override;
 };
 
 /**
@@ -196,7 +196,7 @@ struct CONTROLRIG_API FRigUnit_MathBoolOr : public FRigUnit_MathBoolBinaryAggreg
 	GENERATED_BODY()
 
 	RIGVM_METHOD()
-	virtual void Execute(const FRigUnitContext& Context) override;
+	virtual void Execute() override;
 };
 
 /**
@@ -207,7 +207,7 @@ struct CONTROLRIG_API FRigUnit_MathBoolEquals : public FRigUnit_MathBoolBase
 {
 	GENERATED_BODY()
 	RIGVM_METHOD()
-	virtual void Execute(const FRigUnitContext& Context) override;
+	virtual void Execute() override;
 
 	FRigUnit_MathBoolEquals()
 	{
@@ -236,7 +236,7 @@ struct CONTROLRIG_API FRigUnit_MathBoolNotEquals : public FRigUnit_MathBoolBase
 {
 	GENERATED_BODY()
 	RIGVM_METHOD()
-	virtual void Execute(const FRigUnitContext& Context) override;
+	virtual void Execute() override;
 
 	FRigUnit_MathBoolNotEquals()
 	{
@@ -265,7 +265,7 @@ struct CONTROLRIG_API FRigUnit_MathBoolToggled : public FRigUnit_MathBoolBase
 {
 	GENERATED_BODY()
 	RIGVM_METHOD()
-	virtual void Execute(const FRigUnitContext& Context) override;
+	virtual void Execute() override;
 
 	FRigUnit_MathBoolToggled()
 	{
@@ -293,7 +293,7 @@ struct CONTROLRIG_API FRigUnit_MathBoolFlipFlop : public FRigUnit_MathBoolBase
 {
 	GENERATED_BODY()
 	RIGVM_METHOD()
-	virtual void Execute(const FRigUnitContext& Context) override;
+	virtual void Execute() override;
 
 	FRigUnit_MathBoolFlipFlop()
 	{
@@ -330,7 +330,7 @@ struct CONTROLRIG_API FRigUnit_MathBoolOnce : public FRigUnit_MathBoolBase
 {
 	GENERATED_BODY()
 	RIGVM_METHOD()
-	virtual void Execute(const FRigUnitContext& Context) override;
+	virtual void Execute() override;
 
 	FRigUnit_MathBoolOnce()
 	{
@@ -358,12 +358,12 @@ struct CONTROLRIG_API FRigUnit_MathBoolOnce : public FRigUnit_MathBoolBase
 /**
  * Turns the given bool into a float value
  */
-USTRUCT(meta=(DisplayName="To Float", TemplateName="Cast"))
+USTRUCT(meta=(DisplayName="To Float", TemplateName="Cast", ExecuteContext="FRigVMExecuteContext"))
 struct CONTROLRIG_API FRigUnit_MathBoolToFloat : public FRigUnit_MathBoolBase
 {
 	GENERATED_BODY()
 	RIGVM_METHOD()
-	virtual void Execute(const FRigUnitContext& Context) override;
+	virtual void Execute() override;
 
 	FRigUnit_MathBoolToFloat()
 	{
@@ -381,12 +381,12 @@ struct CONTROLRIG_API FRigUnit_MathBoolToFloat : public FRigUnit_MathBoolBase
 /**
  * Turns the given bool into an integer value
  */
-USTRUCT(meta=(DisplayName="To Integer", TemplateName="Cast"))
+USTRUCT(meta=(DisplayName="To Integer", TemplateName="Cast", ExecuteContext="FRigVMExecuteContext"))
 struct CONTROLRIG_API FRigUnit_MathBoolToInteger : public FRigUnit_MathBoolBase
 {
 	GENERATED_BODY()
 	RIGVM_METHOD()
-	virtual void Execute(const FRigUnitContext& Context) override;
+	virtual void Execute() override;
 
 	FRigUnit_MathBoolToInteger()
 	{
diff --git a/Engine/Plugins/Animation/ControlRig/Source/ControlRig/Private/Units/Math/RigUnit_MathColor.h b/Engine/Plugins/Animation/ControlRig/Source/ControlRig/Private/Units/Math/RigUnit_MathColor.h
--- a/Engine/Plugins/Animation/ControlRig/Source/ControlRig/Private/Units/Math/RigUnit_MathColor.h
+++ b/Engine/Plugins/Animation/ControlRig/Source/ControlRig/Private/Units/Math/RigUnit_MathColor.h
@@ -60,7 +60,7 @@ struct CONTROLRIG_API FRigUnit_MathColorMake : public FRigUnit_MathColorBase
 {
 	GENERATED_BODY()
 	RIGVM_METHOD()
-	virtual void Execute(const FRigUnitContext& Context) override;
+	virtual void Execute() override;
 
 	FRigUnit_MathColorMake()
 	{
@@ -88,12 +88,12 @@ struct CONTROLRIG_API FRigUnit_MathColorMake : public FRigUnit_MathColorBase
 /**
  * Makes a vector from a single float
  */
-USTRUCT(meta=(DisplayName="From Float", TemplateName="Cast", Keywords="Make,Construct"))
+USTRUCT(meta=(DisplayName="From Float", TemplateName="Cast", ExecuteContext="FRigVMExecuteContext", Keywords="Make,Construct"))
 struct CONTROLRIG_API FRigUnit_MathColorFromFloat : public FRigUnit_MathColorBase
 {
 	GENERATED_BODY()
 	RIGVM_METHOD()
-	virtual void Execute(const FRigUnitContext& Context) override;
+	virtual void Execute() override;
 
 	FRigUnit_MathColorFromFloat()
 	{
@@ -111,12 +111,12 @@ struct CONTROLRIG_API FRigUnit_MathColorFromFloat : public FRigUnit_MathColorBas
 /**
  * Makes a vector from a single double
  */
-USTRUCT(meta=(DisplayName="From Double", TemplateName="Cast", Keywords="Make,Construct"))
+USTRUCT(meta=(DisplayName="From Double", TemplateName="Cast", ExecuteContext="FRigVMExecuteContext", Keywords="Make,Construct"))
 struct CONTROLRIG_API FRigUnit_MathColorFromDouble : public FRigUnit_MathColorBase
 {
 	GENERATED_BODY()
 	RIGVM_METHOD()
-	virtual void Execute(const FRigUnitContext& Context) override;
+	virtual void Execute() override;
 
 	FRigUnit_MathColorFromDouble()
 	{
@@ -140,7 +140,7 @@ struct CONTROLRIG_API FRigUnit_MathColorAdd : public FRigUnit_MathColorBinaryAgg
 	GENERATED_BODY()
 
 	RIGVM_METHOD()
-	virtual void Execute(const FRigUnitContext& Context) override;
+	virtual void Execute() override;
 };
 
 /**
@@ -152,7 +152,7 @@ struct CONTROLRIG_API FRigUnit_MathColorSub : public FRigUnit_MathColorBinaryOp
 	GENERATED_BODY()
 
 	RIGVM_METHOD()
-	virtual void Execute(const FRigUnitContext& Context) override;
+	virtual void Execute() override;
 };
 
 /**
@@ -169,7 +169,7 @@ struct CONTROLRIG_API FRigUnit_MathColorMul : public FRigUnit_MathColorBinaryAgg
 	}
 
 	RIGVM_METHOD()
-	virtual void Execute(const FRigUnitContext& Context) override;
+	virtual void Execute() override;
 };
 
 /**
@@ -188,7 +188,7 @@ struct CONTROLRIG_API FRigUnit_MathColorLerp : public FRigUnit_MathColorBase
 	}
 
 	RIGVM_METHOD()
-	virtual void Execute(const FRigUnitContext& Context) override;
+	virtual void Execute() override;
 
 	UPROPERTY(meta=(Input))
 	FLinearColor A;
diff --git a/Engine/Plugins/Animation/ControlRig/Source/ControlRig/Private/Units/Math/RigUnit_MathDouble.h b/Engine/Plugins/Animation/ControlRig/Source/ControlRig/Private/Units/Math/RigUnit_MathDouble.h
--- a/Engine/Plugins/Animation/ControlRig/Source/ControlRig/Private/Units/Math/RigUnit_MathDouble.h
+++ b/Engine/Plugins/Animation/ControlRig/Source/ControlRig/Private/Units/Math/RigUnit_MathDouble.h
@@ -97,7 +97,7 @@ struct CONTROLRIG_API FRigUnit_MathDoubleMake : public FRigUnit_MathDoubleBase
 	}
 	
 	RIGVM_METHOD()
-	virtual void Execute(const FRigUnitContext& Context) override;
+	virtual void Execute() override;
 
 	UPROPERTY(meta=(Input, Output))
 	double Value;
@@ -112,7 +112,7 @@ struct CONTROLRIG_API FRigUnit_MathDoubleConstPi : public FRigUnit_MathDoubleCon
 	GENERATED_BODY()
 	
 	RIGVM_METHOD()
-	virtual void Execute(const FRigUnitContext& Context) override;
+	virtual void Execute() override;
 };
 
 /**
@@ -124,7 +124,7 @@ struct CONTROLRIG_API FRigUnit_MathDoubleConstHalfPi : public FRigUnit_MathDoubl
 	GENERATED_BODY()
 
 	RIGVM_METHOD()
-	virtual void Execute(const FRigUnitContext& Context) override;
+	virtual void Execute() override;
 };
 
 /**
@@ -136,7 +136,7 @@ struct CONTROLRIG_API FRigUnit_MathDoubleConstTwoPi : public FRigUnit_MathDouble
 	GENERATED_BODY()
 
 	RIGVM_METHOD()
-	virtual void Execute(const FRigUnitContext& Context) override;
+	virtual void Execute() override;
 };
 
 /**
@@ -148,7 +148,7 @@ struct CONTROLRIG_API FRigUnit_MathDoubleConstE : public FRigUnit_MathDoubleCons
 	GENERATED_BODY()
 	
 	RIGVM_METHOD()
-	virtual void Execute(const FRigUnitContext& Context) override;
+	virtual void Execute() override;
 };
 
 /**
@@ -160,7 +160,7 @@ struct CONTROLRIG_API FRigUnit_MathDoubleAdd : public FRigUnit_MathDoubleBinaryA
 	GENERATED_BODY()
 
 	RIGVM_METHOD()
-	virtual void Execute(const FRigUnitContext& Context) override;
+	virtual void Execute() override;
 };
 
 /**
@@ -172,7 +172,7 @@ struct CONTROLRIG_API FRigUnit_MathDoubleSub : public FRigUnit_MathDoubleBinaryO
 	GENERATED_BODY()
 
 	RIGVM_METHOD()
-	virtual void Execute(const FRigUnitContext& Context) override;
+	virtual void Execute() override;
 };
 
 /**
@@ -190,7 +190,7 @@ struct CONTROLRIG_API FRigUnit_MathDoubleMul : public FRigUnit_MathDoubleBinaryA
 	}
 
 	RIGVM_METHOD()
-	virtual void Execute(const FRigUnitContext& Context) override;
+	virtual void Execute() override;
 };
 
 /**
@@ -207,7 +207,7 @@ struct CONTROLRIG_API FRigUnit_MathDoubleDiv : public FRigUnit_MathDoubleBinaryO
 	}
 
 	RIGVM_METHOD()
-	virtual void Execute(const FRigUnitContext& Context) override;
+	virtual void Execute() override;
 };
 
 /**
@@ -225,7 +225,7 @@ struct CONTROLRIG_API FRigUnit_MathDoubleMod : public FRigUnit_MathDoubleBinaryO
 	}
 
 	RIGVM_METHOD()
-	virtual void Execute(const FRigUnitContext& Context) override;
+	virtual void Execute() override;
 };
 
 /**
@@ -237,7 +237,7 @@ struct CONTROLRIG_API FRigUnit_MathDoubleMin : public FRigUnit_MathDoubleBinaryA
 	GENERATED_BODY()
 
 	RIGVM_METHOD()
-	virtual void Execute(const FRigUnitContext& Context) override;
+	virtual void Execute() override;
 };
 
 /**
@@ -249,7 +249,7 @@ struct CONTROLRIG_API FRigUnit_MathDoubleMax : public FRigUnit_MathDoubleBinaryA
 	GENERATED_BODY()
 
 	RIGVM_METHOD()
-	virtual void Execute(const FRigUnitContext& Context) override;
+	virtual void Execute() override;
 };
 
 /**
@@ -267,7 +267,7 @@ struct CONTROLRIG_API FRigUnit_MathDoublePow : public FRigUnit_MathDoubleBinaryO
 	}
 
 	RIGVM_METHOD()
-	virtual void Execute(const FRigUnitContext& Context) override;
+	virtual void Execute() override;
 };
 
 /**
@@ -279,7 +279,7 @@ struct CONTROLRIG_API FRigUnit_MathDoubleSqrt : public FRigUnit_MathDoubleUnaryO
 	GENERATED_BODY()
 
 	RIGVM_METHOD()
-	virtual void Execute(const FRigUnitContext& Context) override;
+	virtual void Execute() override;
 };
 
 /**
@@ -291,7 +291,7 @@ struct CONTROLRIG_API FRigUnit_MathDoubleNegate : public FRigUnit_MathDoubleUnar
 	GENERATED_BODY()
 
 	RIGVM_METHOD()
-	virtual void Execute(const FRigUnitContext& Context) override;
+	virtual void Execute() override;
 };
 
 /**
@@ -303,7 +303,7 @@ struct CONTROLRIG_API FRigUnit_MathDoubleAbs : public FRigUnit_MathDoubleUnaryOp
 	GENERATED_BODY()
 
 	RIGVM_METHOD()
-	virtual void Execute(const FRigUnitContext& Context) override;
+	virtual void Execute() override;
 };
 
 /**
@@ -321,7 +321,7 @@ struct CONTROLRIG_API FRigUnit_MathDoubleFloor : public FRigUnit_MathDoubleBase
 	}
 
 	RIGVM_METHOD()
-	virtual void Execute(const FRigUnitContext& Context) override;
+	virtual void Execute() override;
 
 	UPROPERTY(meta=(Input))
 	double Value;
@@ -348,7 +348,7 @@ struct CONTROLRIG_API FRigUnit_MathDoubleCeil : public FRigUnit_MathDoubleBase
 	}
 
 	RIGVM_METHOD()
-	virtual void Execute(const FRigUnitContext& Context) override;
+	virtual void Execute() override;
 
 	UPROPERTY(meta=(Input))
 	double Value;
@@ -375,7 +375,7 @@ struct CONTROLRIG_API FRigUnit_MathDoubleRound : public FRigUnit_MathDoubleBase
 	}
 
 	RIGVM_METHOD()
-	virtual void Execute(const FRigUnitContext& Context) override;
+	virtual void Execute() override;
 
 	UPROPERTY(meta=(Input))
 	double Value;
@@ -390,7 +390,7 @@ struct CONTROLRIG_API FRigUnit_MathDoubleRound : public FRigUnit_MathDoubleBase
 /**
  * Returns the double cast to an int (this uses floor)
  */
-USTRUCT(meta=(DisplayName="To Int", TemplateName="Cast"))
+USTRUCT(meta=(DisplayName="To Int", TemplateName="Cast", ExecuteContext="FRigVMExecuteContext"))
 struct CONTROLRIG_API FRigUnit_MathDoubleToInt : public FRigUnit_MathDoubleBase
 {
 	GENERATED_BODY()
@@ -402,7 +402,7 @@ struct CONTROLRIG_API FRigUnit_MathDoubleToInt : public FRigUnit_MathDoubleBase
 	}
 
 	RIGVM_METHOD()
-	virtual void Execute(const FRigUnitContext& Context) override;
+	virtual void Execute() override;
 
 	UPROPERTY(meta=(Input))
 	double Value;
@@ -420,7 +420,7 @@ struct CONTROLRIG_API FRigUnit_MathDoubleSign : public FRigUnit_MathDoubleUnaryO
 	GENERATED_BODY()
 
 	RIGVM_METHOD()
-	virtual void Execute(const FRigUnitContext& Context) override;
+	virtual void Execute() override;
 };
 
 /**
@@ -437,7 +437,7 @@ struct CONTROLRIG_API FRigUnit_MathDoubleClamp : public FRigUnit_MathDoubleBase
 	}
 
 	RIGVM_METHOD()
-	virtual void Execute(const FRigUnitContext& Context) override;
+	virtual void Execute() override;
 
 	UPROPERTY(meta=(Input))
 	double Value;
@@ -467,7 +467,7 @@ struct CONTROLRIG_API FRigUnit_MathDoubleLerp : public FRigUnit_MathDoubleBase
 	}
 
 	RIGVM_METHOD()
-	virtual void Execute(const FRigUnitContext& Context) override;
+	virtual void Execute() override;
 
 	UPROPERTY(meta=(Input))
 	double A;
@@ -498,7 +498,7 @@ struct CONTROLRIG_API FRigUnit_MathDoubleRemap : public FRigUnit_MathDoubleBase
 	}
 
 	RIGVM_METHOD()
-	virtual void Execute(const FRigUnitContext& Context) override;
+	virtual void Execute() override;
 
 	UPROPERTY(meta=(Input))
 	double Value;
@@ -531,7 +531,7 @@ struct CONTROLRIG_API FRigUnit_MathDoubleEquals : public FRigUnit_MathDoubleBase
 {
 	GENERATED_BODY()
 	RIGVM_METHOD()
-	virtual void Execute(const FRigUnitContext& Context) override;
+	virtual void Execute() override;
 
 	FRigUnit_MathDoubleEquals()
 	{
@@ -560,7 +560,7 @@ struct CONTROLRIG_API FRigUnit_MathDoubleNotEquals : public FRigUnit_MathDoubleB
 {
 	GENERATED_BODY()
 	RIGVM_METHOD()
-	virtual void Execute(const FRigUnitContext& Context) override;
+	virtual void Execute() override;
 
 	FRigUnit_MathDoubleNotEquals()
 	{
@@ -596,7 +596,7 @@ struct CONTROLRIG_API FRigUnit_MathDoubleGreater : public FRigUnit_MathDoubleBas
 	}
 
 	RIGVM_METHOD()
-	virtual void Execute(const FRigUnitContext& Context) override;
+	virtual void Execute() override;
 
 	UPROPERTY(meta=(Input))
 	double A;
@@ -623,7 +623,7 @@ struct CONTROLRIG_API FRigUnit_MathDoubleLess : public FRigUnit_MathDoubleBase
 	}
 
 	RIGVM_METHOD()
-	virtual void Execute(const FRigUnitContext& Context) override;
+	virtual void Execute() override;
 
 	UPROPERTY(meta=(Input))
 	double A;
@@ -650,7 +650,7 @@ struct CONTROLRIG_API FRigUnit_MathDoubleGreaterEqual : public FRigUnit_MathDoub
 	}
 
 	RIGVM_METHOD()
-	virtual void Execute(const FRigUnitContext& Context) override;
+	virtual void Execute() override;
 
 	UPROPERTY(meta=(Input))
 	double A;
@@ -677,7 +677,7 @@ struct CONTROLRIG_API FRigUnit_MathDoubleLessEqual : public FRigUnit_MathDoubleB
 	}
 
 	RIGVM_METHOD()
-	virtual void Execute(const FRigUnitContext& Context) override;
+	virtual void Execute() override;
 
 	UPROPERTY(meta=(Input))
 	double A;
@@ -704,7 +704,7 @@ struct CONTROLRIG_API FRigUnit_MathDoubleIsNearlyZero : public FRigUnit_MathDoub
 	}
 	
 	RIGVM_METHOD()
-	virtual void Execute(const FRigUnitContext& Context) override;
+	virtual void Execute() override;
 
 	UPROPERTY(meta=(Input))
 	double Value;
@@ -731,7 +731,7 @@ struct CONTROLRIG_API FRigUnit_MathDoubleIsNearlyEqual : public FRigUnit_MathDou
 	}
 
 	RIGVM_METHOD()
-	virtual void Execute(const FRigUnitContext& Context) override;
+	virtual void Execute() override;
 
 	UPROPERTY(meta=(Input))
 	double A;
@@ -755,7 +755,7 @@ struct CONTROLRIG_API FRigUnit_MathDoubleDeg : public FRigUnit_MathDoubleUnaryOp
 	GENERATED_BODY()
 
 	RIGVM_METHOD()
-	virtual void Execute(const FRigUnitContext& Context) override;
+	virtual void Execute() override;
 };
 
 /**
@@ -767,7 +767,7 @@ struct CONTROLRIG_API FRigUnit_MathDoubleRad : public FRigUnit_MathDoubleUnaryOp
 	GENERATED_BODY()
 
 	RIGVM_METHOD()
-	virtual void Execute(const FRigUnitContext& Context) override;
+	virtual void Execute() override;
 };
 
 /**
@@ -779,7 +779,7 @@ struct CONTROLRIG_API FRigUnit_MathDoubleSin : public FRigUnit_MathDoubleUnaryOp
 	GENERATED_BODY()
 
 	RIGVM_METHOD()
-	virtual void Execute(const FRigUnitContext& Context) override;
+	virtual void Execute() override;
 };
 
 /**
@@ -791,7 +791,7 @@ struct CONTROLRIG_API FRigUnit_MathDoubleCos : public FRigUnit_MathDoubleUnaryOp
 	GENERATED_BODY()
 
 	RIGVM_METHOD()
-	virtual void Execute(const FRigUnitContext& Context) override;
+	virtual void Execute() override;
 };
 
 /**
@@ -803,7 +803,7 @@ struct CONTROLRIG_API FRigUnit_MathDoubleTan : public FRigUnit_MathDoubleUnaryOp
 	GENERATED_BODY()
 
 	RIGVM_METHOD()
-	virtual void Execute(const FRigUnitContext& Context) override;
+	virtual void Execute() override;
 };
 
 /**
@@ -815,7 +815,7 @@ struct CONTROLRIG_API FRigUnit_MathDoubleAsin : public FRigUnit_MathDoubleUnaryO
 	GENERATED_BODY()
 
 	RIGVM_METHOD()
-	virtual void Execute(const FRigUnitContext& Context) override;
+	virtual void Execute() override;
 };
 
 /**
@@ -827,7 +827,7 @@ struct CONTROLRIG_API FRigUnit_MathDoubleAcos : public FRigUnit_MathDoubleUnaryO
 	GENERATED_BODY()
 
 	RIGVM_METHOD()
-	virtual void Execute(const FRigUnitContext& Context) override;
+	virtual void Execute() override;
 };
 
 /**
@@ -839,7 +839,7 @@ struct CONTROLRIG_API FRigUnit_MathDoubleAtan : public FRigUnit_MathDoubleUnaryO
 	GENERATED_BODY()
 
 	RIGVM_METHOD()
-	virtual void Execute(const FRigUnitContext& Context) override;
+	virtual void Execute() override;
 };
 
 /**
@@ -857,7 +857,7 @@ struct CONTROLRIG_API FRigUnit_MathDoubleLawOfCosine : public FRigUnit_MathDoubl
 	}
 
 	RIGVM_METHOD()
-	virtual void Execute(const FRigUnitContext& Context) override;
+	virtual void Execute() override;
 
 	UPROPERTY(meta = (Input))
 	double A;
@@ -890,5 +890,5 @@ struct CONTROLRIG_API FRigUnit_MathDoubleExponential : public FRigUnit_MathDoubl
 	GENERATED_BODY()
 
 	RIGVM_METHOD()
-	virtual void Execute(const FRigUnitContext& Context) override;
+	virtual void Execute() override;
 };
\ No newline at end of file
diff --git a/Engine/Plugins/Animation/ControlRig/Source/ControlRig/Private/Units/Math/RigUnit_MathFloat.h b/Engine/Plugins/Animation/ControlRig/Source/ControlRig/Private/Units/Math/RigUnit_MathFloat.h
--- a/Engine/Plugins/Animation/ControlRig/Source/ControlRig/Private/Units/Math/RigUnit_MathFloat.h
+++ b/Engine/Plugins/Animation/ControlRig/Source/ControlRig/Private/Units/Math/RigUnit_MathFloat.h
@@ -96,7 +96,7 @@ struct CONTROLRIG_API FRigUnit_MathFloatMake : public FRigUnit_MathFloatBase
 	}
 	
 	RIGVM_METHOD()
-	virtual void Execute(const FRigUnitContext& Context) override;
+	virtual void Execute() override;
 
 	UPROPERTY(meta=(Input, Output))
 	float Value;
@@ -111,7 +111,7 @@ struct CONTROLRIG_API FRigUnit_MathFloatConstPi : public FRigUnit_MathFloatConst
 	GENERATED_BODY()
 	
 	RIGVM_METHOD()
-	virtual void Execute(const FRigUnitContext& Context) override;
+	virtual void Execute() override;
 };
 
 /**
@@ -123,7 +123,7 @@ struct CONTROLRIG_API FRigUnit_MathFloatConstHalfPi : public FRigUnit_MathFloatC
 	GENERATED_BODY()
 
 	RIGVM_METHOD()
-	virtual void Execute(const FRigUnitContext& Context) override;
+	virtual void Execute() override;
 };
 
 /**
@@ -135,7 +135,7 @@ struct CONTROLRIG_API FRigUnit_MathFloatConstTwoPi : public FRigUnit_MathFloatCo
 	GENERATED_BODY()
 
 	RIGVM_METHOD()
-	virtual void Execute(const FRigUnitContext& Context) override;
+	virtual void Execute() override;
 };
 
 /**
@@ -147,7 +147,7 @@ struct CONTROLRIG_API FRigUnit_MathFloatConstE : public FRigUnit_MathFloatConsta
 	GENERATED_BODY()
 	
 	RIGVM_METHOD()
-	virtual void Execute(const FRigUnitContext& Context) override;
+	virtual void Execute() override;
 };
 
 /**
@@ -159,7 +159,7 @@ struct CONTROLRIG_API FRigUnit_MathFloatAdd : public FRigUnit_MathFloatBinaryAgg
 	GENERATED_BODY()
 
 	RIGVM_METHOD()
-	virtual void Execute(const FRigUnitContext& Context) override;
+	virtual void Execute() override;
 };
 
 /**
@@ -171,7 +171,7 @@ struct CONTROLRIG_API FRigUnit_MathFloatSub : public FRigUnit_MathFloatBinaryOp
 	GENERATED_BODY()
 
 	RIGVM_METHOD()
-	virtual void Execute(const FRigUnitContext& Context) override;
+	virtual void Execute() override;
 };
 
 /**
@@ -189,7 +189,7 @@ struct CONTROLRIG_API FRigUnit_MathFloatMul : public FRigUnit_MathFloatBinaryAgg
 	}
 
 	RIGVM_METHOD()
-	virtual void Execute(const FRigUnitContext& Context) override;
+	virtual void Execute() override;
 };
 
 /**
@@ -206,7 +206,7 @@ struct CONTROLRIG_API FRigUnit_MathFloatDiv : public FRigUnit_MathFloatBinaryOp
 	}
 
 	RIGVM_METHOD()
-	virtual void Execute(const FRigUnitContext& Context) override;
+	virtual void Execute() override;
 };
 
 /**
@@ -224,7 +224,7 @@ struct CONTROLRIG_API FRigUnit_MathFloatMod : public FRigUnit_MathFloatBinaryOp
 	}
 
 	RIGVM_METHOD()
-	virtual void Execute(const FRigUnitContext& Context) override;
+	virtual void Execute() override;
 };
 
 /**
@@ -236,7 +236,7 @@ struct CONTROLRIG_API FRigUnit_MathFloatMin : public FRigUnit_MathFloatBinaryAgg
 	GENERATED_BODY()
 
 	RIGVM_METHOD()
-	virtual void Execute(const FRigUnitContext& Context) override;
+	virtual void Execute() override;
 };
 
 /**
@@ -248,7 +248,7 @@ struct CONTROLRIG_API FRigUnit_MathFloatMax : public FRigUnit_MathFloatBinaryAgg
 	GENERATED_BODY()
 
 	RIGVM_METHOD()
-	virtual void Execute(const FRigUnitContext& Context) override;
+	virtual void Execute() override;
 };
 
 /**
@@ -266,7 +266,7 @@ struct CONTROLRIG_API FRigUnit_MathFloatPow : public FRigUnit_MathFloatBinaryOp
 	}
 
 	RIGVM_METHOD()
-	virtual void Execute(const FRigUnitContext& Context) override;
+	virtual void Execute() override;
 };
 
 /**
@@ -278,7 +278,7 @@ struct CONTROLRIG_API FRigUnit_MathFloatSqrt : public FRigUnit_MathFloatUnaryOp
 	GENERATED_BODY()
 
 	RIGVM_METHOD()
-	virtual void Execute(const FRigUnitContext& Context) override;
+	virtual void Execute() override;
 };
 
 /**
@@ -290,7 +290,7 @@ struct CONTROLRIG_API FRigUnit_MathFloatNegate : public FRigUnit_MathFloatUnaryO
 	GENERATED_BODY()
 
 	RIGVM_METHOD()
-	virtual void Execute(const FRigUnitContext& Context) override;
+	virtual void Execute() override;
 };
 
 /**
@@ -302,7 +302,7 @@ struct CONTROLRIG_API FRigUnit_MathFloatAbs : public FRigUnit_MathFloatUnaryOp
 	GENERATED_BODY()
 
 	RIGVM_METHOD()
-	virtual void Execute(const FRigUnitContext& Context) override;
+	virtual void Execute() override;
 };
 
 /**
@@ -320,7 +320,7 @@ struct CONTROLRIG_API FRigUnit_MathFloatFloor : public FRigUnit_MathFloatBase
 	}
 
 	RIGVM_METHOD()
-	virtual void Execute(const FRigUnitContext& Context) override;
+	virtual void Execute() override;
 
 	UPROPERTY(meta=(Input))
 	float Value;
@@ -347,7 +347,7 @@ struct CONTROLRIG_API FRigUnit_MathFloatCeil : public FRigUnit_MathFloatBase
 	}
 
 	RIGVM_METHOD()
-	virtual void Execute(const FRigUnitContext& Context) override;
+	virtual void Execute() override;
 
 	UPROPERTY(meta=(Input))
 	float Value;
@@ -374,7 +374,7 @@ struct CONTROLRIG_API FRigUnit_MathFloatRound : public FRigUnit_MathFloatBase
 	}
 
 	RIGVM_METHOD()
-	virtual void Execute(const FRigUnitContext& Context) override;
+	virtual void Execute() override;
 
 	UPROPERTY(meta=(Input))
 	float Value;
@@ -389,7 +389,7 @@ struct CONTROLRIG_API FRigUnit_MathFloatRound : public FRigUnit_MathFloatBase
 /**
  * Returns the float cast to an int (this uses floor)
  */
-USTRUCT(meta=(DisplayName="To Int", TemplateName="Cast"))
+USTRUCT(meta=(DisplayName="To Int", TemplateName="Cast", ExecuteContext="FRigVMExecuteContext"))
 struct CONTROLRIG_API FRigUnit_MathFloatToInt : public FRigUnit_MathFloatBase
 {
 	GENERATED_BODY()
@@ -401,7 +401,7 @@ struct CONTROLRIG_API FRigUnit_MathFloatToInt : public FRigUnit_MathFloatBase
 	}
 
 	RIGVM_METHOD()
-	virtual void Execute(const FRigUnitContext& Context) override;
+	virtual void Execute() override;
 
 	UPROPERTY(meta=(Input))
 	float Value;
@@ -419,7 +419,7 @@ struct CONTROLRIG_API FRigUnit_MathFloatSign : public FRigUnit_MathFloatUnaryOp
 	GENERATED_BODY()
 
 	RIGVM_METHOD()
-	virtual void Execute(const FRigUnitContext& Context) override;
+	virtual void Execute() override;
 };
 
 /**
@@ -436,7 +436,7 @@ struct CONTROLRIG_API FRigUnit_MathFloatClamp : public FRigUnit_MathFloatBase
 	}
 
 	RIGVM_METHOD()
-	virtual void Execute(const FRigUnitContext& Context) override;
+	virtual void Execute() override;
 
 	UPROPERTY(meta=(Input))
 	float Value;
@@ -466,7 +466,7 @@ struct CONTROLRIG_API FRigUnit_MathFloatLerp : public FRigUnit_MathFloatBase
 	}
 
 	RIGVM_METHOD()
-	virtual void Execute(const FRigUnitContext& Context) override;
+	virtual void Execute() override;
 
 	UPROPERTY(meta=(Input))
 	float A;
@@ -497,7 +497,7 @@ struct CONTROLRIG_API FRigUnit_MathFloatRemap : public FRigUnit_MathFloatBase
 	}
 
 	RIGVM_METHOD()
-	virtual void Execute(const FRigUnitContext& Context) override;
+	virtual void Execute() override;
 
 	UPROPERTY(meta=(Input))
 	float Value;
@@ -530,7 +530,7 @@ struct CONTROLRIG_API FRigUnit_MathFloatEquals : public FRigUnit_MathFloatBase
 {
 	GENERATED_BODY()
 	RIGVM_METHOD()
-	virtual void Execute(const FRigUnitContext& Context) override;
+	virtual void Execute() override;
 
 	FRigUnit_MathFloatEquals()
 	{
@@ -559,7 +559,7 @@ struct CONTROLRIG_API FRigUnit_MathFloatNotEquals : public FRigUnit_MathFloatBas
 {
 	GENERATED_BODY()
 	RIGVM_METHOD()
-	virtual void Execute(const FRigUnitContext& Context) override;
+	virtual void Execute() override;
 
 	FRigUnit_MathFloatNotEquals()
 	{
@@ -595,7 +595,7 @@ struct CONTROLRIG_API FRigUnit_MathFloatGreater : public FRigUnit_MathFloatBase
 	}
 
 	RIGVM_METHOD()
-	virtual void Execute(const FRigUnitContext& Context) override;
+	virtual void Execute() override;
 
 	UPROPERTY(meta=(Input))
 	float A;
@@ -622,7 +622,7 @@ struct CONTROLRIG_API FRigUnit_MathFloatLess : public FRigUnit_MathFloatBase
 	}
 
 	RIGVM_METHOD()
-	virtual void Execute(const FRigUnitContext& Context) override;
+	virtual void Execute() override;
 
 	UPROPERTY(meta=(Input))
 	float A;
@@ -649,7 +649,7 @@ struct CONTROLRIG_API FRigUnit_MathFloatGreaterEqual : public FRigUnit_MathFloat
 	}
 
 	RIGVM_METHOD()
-	virtual void Execute(const FRigUnitContext& Context) override;
+	virtual void Execute() override;
 
 	UPROPERTY(meta=(Input))
 	float A;
@@ -676,7 +676,7 @@ struct CONTROLRIG_API FRigUnit_MathFloatLessEqual : public FRigUnit_MathFloatBas
 	}
 
 	RIGVM_METHOD()
-	virtual void Execute(const FRigUnitContext& Context) override;
+	virtual void Execute() override;
 
 	UPROPERTY(meta=(Input))
 	float A;
@@ -703,7 +703,7 @@ struct CONTROLRIG_API FRigUnit_MathFloatIsNearlyZero : public FRigUnit_MathFloat
 	}
 	
 	RIGVM_METHOD()
-	virtual void Execute(const FRigUnitContext& Context) override;
+	virtual void Execute() override;
 
 	UPROPERTY(meta=(Input))
 	float Value;
@@ -730,7 +730,7 @@ struct CONTROLRIG_API FRigUnit_MathFloatIsNearlyEqual : public FRigUnit_MathFloa
 	}
 
 	RIGVM_METHOD()
-	virtual void Execute(const FRigUnitContext& Context) override;
+	virtual void Execute() override;
 
 	UPROPERTY(meta=(Input))
 	float A;
@@ -760,7 +760,7 @@ struct CONTROLRIG_API FRigUnit_MathFloatSelectBool : public FRigUnit_MathFloatBa
 	}
 
 	RIGVM_METHOD()
-	virtual void Execute(const FRigUnitContext& Context) override;
+	virtual void Execute() override;
 
 	UPROPERTY(meta=(Input))
 	bool Condition;
@@ -787,7 +787,7 @@ struct CONTROLRIG_API FRigUnit_MathFloatDeg : public FRigUnit_MathFloatUnaryOp
 	GENERATED_BODY()
 
 	RIGVM_METHOD()
-	virtual void Execute(const FRigUnitContext& Context) override;
+	virtual void Execute() override;
 };
 
 /**
@@ -799,7 +799,7 @@ struct CONTROLRIG_API FRigUnit_MathFloatRad : public FRigUnit_MathFloatUnaryOp
 	GENERATED_BODY()
 
 	RIGVM_METHOD()
-	virtual void Execute(const FRigUnitContext& Context) override;
+	virtual void Execute() override;
 };
 
 /**
@@ -811,7 +811,7 @@ struct CONTROLRIG_API FRigUnit_MathFloatSin : public FRigUnit_MathFloatUnaryOp
 	GENERATED_BODY()
 
 	RIGVM_METHOD()
-	virtual void Execute(const FRigUnitContext& Context) override;
+	virtual void Execute() override;
 };
 
 /**
@@ -823,7 +823,7 @@ struct CONTROLRIG_API FRigUnit_MathFloatCos : public FRigUnit_MathFloatUnaryOp
 	GENERATED_BODY()
 
 	RIGVM_METHOD()
-	virtual void Execute(const FRigUnitContext& Context) override;
+	virtual void Execute() override;
 };
 
 /**
@@ -835,7 +835,7 @@ struct CONTROLRIG_API FRigUnit_MathFloatTan : public FRigUnit_MathFloatUnaryOp
 	GENERATED_BODY()
 
 	RIGVM_METHOD()
-	virtual void Execute(const FRigUnitContext& Context) override;
+	virtual void Execute() override;
 };
 
 /**
@@ -847,7 +847,7 @@ struct CONTROLRIG_API FRigUnit_MathFloatAsin : public FRigUnit_MathFloatUnaryOp
 	GENERATED_BODY()
 
 	RIGVM_METHOD()
-	virtual void Execute(const FRigUnitContext& Context) override;
+	virtual void Execute() override;
 };
 
 /**
@@ -859,7 +859,7 @@ struct CONTROLRIG_API FRigUnit_MathFloatAcos : public FRigUnit_MathFloatUnaryOp
 	GENERATED_BODY()
 
 	RIGVM_METHOD()
-	virtual void Execute(const FRigUnitContext& Context) override;
+	virtual void Execute() override;
 };
 
 /**
@@ -871,7 +871,7 @@ struct CONTROLRIG_API FRigUnit_MathFloatAtan : public FRigUnit_MathFloatUnaryOp
 	GENERATED_BODY()
 
 	RIGVM_METHOD()
-	virtual void Execute(const FRigUnitContext& Context) override;
+	virtual void Execute() override;
 };
 
 /**
@@ -889,7 +889,7 @@ struct CONTROLRIG_API FRigUnit_MathFloatLawOfCosine : public FRigUnit_MathFloatB
 	}
 
 	RIGVM_METHOD()
-	virtual void Execute(const FRigUnitContext& Context) override;
+	virtual void Execute() override;
 
 	UPROPERTY(meta = (Input))
 	float A;
@@ -922,5 +922,5 @@ struct CONTROLRIG_API FRigUnit_MathFloatExponential : public FRigUnit_MathFloatU
 	GENERATED_BODY()
 
 	RIGVM_METHOD()
-	virtual void Execute(const FRigUnitContext& Context) override;
+	virtual void Execute() override;
 };
\ No newline at end of file
diff --git a/Engine/Plugins/Animation/ControlRig/Source/ControlRig/Private/Units/Math/RigUnit_MathInt.h b/Engine/Plugins/Animation/ControlRig/Source/ControlRig/Private/Units/Math/RigUnit_MathInt.h
--- a/Engine/Plugins/Animation/ControlRig/Source/ControlRig/Private/Units/Math/RigUnit_MathInt.h
+++ b/Engine/Plugins/Animation/ControlRig/Source/ControlRig/Private/Units/Math/RigUnit_MathInt.h
@@ -82,7 +82,7 @@ struct CONTROLRIG_API FRigUnit_MathIntMake : public FRigUnit_MathIntBase
 	}
 	
 	RIGVM_METHOD()
-	virtual void Execute(const FRigUnitContext& Context) override;
+	virtual void Execute() override;
 
 	UPROPERTY(meta=(Input, Output))
 	int32 Value;
@@ -97,7 +97,7 @@ struct CONTROLRIG_API FRigUnit_MathIntAdd : public FRigUnit_MathIntBinaryAggrega
 	GENERATED_BODY()
 
 	RIGVM_METHOD()
-	virtual void Execute(const FRigUnitContext& Context) override;
+	virtual void Execute() override;
 };
 
 /**
@@ -109,7 +109,7 @@ struct CONTROLRIG_API FRigUnit_MathIntSub : public FRigUnit_MathIntBinaryOp
 	GENERATED_BODY()
 
 	RIGVM_METHOD()
-	virtual void Execute(const FRigUnitContext& Context) override;
+	virtual void Execute() override;
 };
 
 /**
@@ -127,7 +127,7 @@ struct CONTROLRIG_API FRigUnit_MathIntMul : public FRigUnit_MathIntBinaryAggrega
 	}
 
 	RIGVM_METHOD()
-	virtual void Execute(const FRigUnitContext& Context) override;
+	virtual void Execute() override;
 };
 
 /**
@@ -144,7 +144,7 @@ struct CONTROLRIG_API FRigUnit_MathIntDiv : public FRigUnit_MathIntBinaryOp
 	}
 
 	RIGVM_METHOD()
-	virtual void Execute(const FRigUnitContext& Context) override;
+	virtual void Execute() override;
 };
 
 /**
@@ -162,7 +162,7 @@ struct CONTROLRIG_API FRigUnit_MathIntMod : public FRigUnit_MathIntBinaryOp
 	}
 
 	RIGVM_METHOD()
-	virtual void Execute(const FRigUnitContext& Context) override;
+	virtual void Execute() override;
 };
 
 /**
@@ -174,7 +174,7 @@ struct CONTROLRIG_API FRigUnit_MathIntMin : public FRigUnit_MathIntBinaryAggrega
 	GENERATED_BODY()
 
 	RIGVM_METHOD()
-	virtual void Execute(const FRigUnitContext& Context) override;
+	virtual void Execute() override;
 };
 
 /**
@@ -186,7 +186,7 @@ struct CONTROLRIG_API FRigUnit_MathIntMax : public FRigUnit_MathIntBinaryAggrega
 	GENERATED_BODY()
 
 	RIGVM_METHOD()
-	virtual void Execute(const FRigUnitContext& Context) override;
+	virtual void Execute() override;
 };
 
 /**
@@ -204,7 +204,7 @@ struct CONTROLRIG_API FRigUnit_MathIntPow : public FRigUnit_MathIntBinaryOp
 	}
 
 	RIGVM_METHOD()
-	virtual void Execute(const FRigUnitContext& Context) override;
+	virtual void Execute() override;
 };
 
 /**
@@ -216,7 +216,7 @@ struct CONTROLRIG_API FRigUnit_MathIntNegate : public FRigUnit_MathIntUnaryOp
 	GENERATED_BODY()
 
 	RIGVM_METHOD()
-	virtual void Execute(const FRigUnitContext& Context) override;
+	virtual void Execute() override;
 };
 
 /**
@@ -228,13 +228,13 @@ struct CONTROLRIG_API FRigUnit_MathIntAbs : public FRigUnit_MathIntUnaryOp
 	GENERATED_BODY()
 
 	RIGVM_METHOD()
-	virtual void Execute(const FRigUnitContext& Context) override;
+	virtual void Execute() override;
 };
 
 /**
  * Returns the int cast to a float
  */
-USTRUCT(meta=(DisplayName="To Float", TemplateName="Cast"))
+USTRUCT(meta=(DisplayName="To Float", TemplateName="Cast", ExecuteContext="FRigVMExecuteContext"))
 struct CONTROLRIG_API FRigUnit_MathIntToFloat : public FRigUnit_MathIntBase
 {
 	GENERATED_BODY()
@@ -246,7 +246,7 @@ struct CONTROLRIG_API FRigUnit_MathIntToFloat : public FRigUnit_MathIntBase
 	}
 
 	RIGVM_METHOD()
-	virtual void Execute(const FRigUnitContext& Context) override;
+	virtual void Execute() override;
 
 	UPROPERTY(meta=(Input))
 	int32 Value;
@@ -258,7 +258,7 @@ struct CONTROLRIG_API FRigUnit_MathIntToFloat : public FRigUnit_MathIntBase
 /**
  * Returns the int cast to a float
  */
-USTRUCT(meta=(DisplayName="To Double", TemplateName="Cast"))
+USTRUCT(meta=(DisplayName="To Double", TemplateName="Cast", ExecuteContext="FRigVMExecuteContext"))
 struct CONTROLRIG_API FRigUnit_MathIntToDouble : public FRigUnit_MathIntBase
 {
 	GENERATED_BODY()
@@ -270,7 +270,7 @@ struct CONTROLRIG_API FRigUnit_MathIntToDouble : public FRigUnit_MathIntBase
 	}
 
 	RIGVM_METHOD()
-	virtual void Execute(const FRigUnitContext& Context) override;
+	virtual void Execute() override;
 
 	UPROPERTY(meta=(Input))
 	int32 Value;
@@ -288,7 +288,7 @@ struct CONTROLRIG_API FRigUnit_MathIntSign : public FRigUnit_MathIntUnaryOp
 	GENERATED_BODY()
 
 	RIGVM_METHOD()
-	virtual void Execute(const FRigUnitContext& Context) override;
+	virtual void Execute() override;
 };
 
 /**
@@ -305,7 +305,7 @@ struct CONTROLRIG_API FRigUnit_MathIntClamp : public FRigUnit_MathIntBase
 	}
 
 	RIGVM_METHOD()
-	virtual void Execute(const FRigUnitContext& Context) override;
+	virtual void Execute() override;
 
 	UPROPERTY(meta=(Input))
 	int32 Value;
@@ -328,7 +328,7 @@ struct CONTROLRIG_API FRigUnit_MathIntEquals : public FRigUnit_MathIntBase
 {
 	GENERATED_BODY()
 	RIGVM_METHOD()
-	virtual void Execute(const FRigUnitContext& Context) override;
+	virtual void Execute() override;
 
 	FRigUnit_MathIntEquals()
 	{
@@ -357,7 +357,7 @@ struct CONTROLRIG_API FRigUnit_MathIntNotEquals : public FRigUnit_MathIntBase
 {
 	GENERATED_BODY()
 	RIGVM_METHOD()
-	virtual void Execute(const FRigUnitContext& Context) override;
+	virtual void Execute() override;
 
 	FRigUnit_MathIntNotEquals()
 	{
@@ -393,7 +393,7 @@ struct CONTROLRIG_API FRigUnit_MathIntGreater : public FRigUnit_MathIntBase
 	}
 
 	RIGVM_METHOD()
-	virtual void Execute(const FRigUnitContext& Context) override;
+	virtual void Execute() override;
 
 	UPROPERTY(meta=(Input))
 	int32 A;
@@ -420,7 +420,7 @@ struct CONTROLRIG_API FRigUnit_MathIntLess : public FRigUnit_MathIntBase
 	}
 
 	RIGVM_METHOD()
-	virtual void Execute(const FRigUnitContext& Context) override;
+	virtual void Execute() override;
 
 	UPROPERTY(meta=(Input))
 	int32 A;
@@ -447,7 +447,7 @@ struct CONTROLRIG_API FRigUnit_MathIntGreaterEqual : public FRigUnit_MathIntBase
 	}
 
 	RIGVM_METHOD()
-	virtual void Execute(const FRigUnitContext& Context) override;
+	virtual void Execute() override;
 
 	UPROPERTY(meta=(Input))
 	int32 A;
@@ -474,7 +474,7 @@ struct CONTROLRIG_API FRigUnit_MathIntLessEqual : public FRigUnit_MathIntBase
 	}
 
 	RIGVM_METHOD()
-	virtual void Execute(const FRigUnitContext& Context) override;
+	virtual void Execute() override;
 
 	UPROPERTY(meta=(Input))
 	int32 A;
diff --git a/Engine/Plugins/Animation/ControlRig/Source/ControlRig/Private/Units/Math/RigUnit_MathMatrix.h b/Engine/Plugins/Animation/ControlRig/Source/ControlRig/Private/Units/Math/RigUnit_MathMatrix.h
--- a/Engine/Plugins/Animation/ControlRig/Source/ControlRig/Private/Units/Math/RigUnit_MathMatrix.h
+++ b/Engine/Plugins/Animation/ControlRig/Source/ControlRig/Private/Units/Math/RigUnit_MathMatrix.h
@@ -73,7 +73,7 @@ struct CONTROLRIG_API FRigUnit_MathMatrixBinaryAggregateOp : public FRigUnit_Mat
 /**
 * Makes a transform from a matrix
 */
-USTRUCT(meta=(DisplayName="To Transform", TemplateName="Cast", Keywords="Make,Construct"))
+USTRUCT(meta=(DisplayName="To Transform", TemplateName="Cast", ExecuteContext="FRigVMExecuteContext", Keywords="Make,Construct"))
 struct CONTROLRIG_API FRigUnit_MathMatrixToTransform : public FRigUnit_MathMatrixBase
 {
 	GENERATED_BODY()
@@ -85,7 +85,7 @@ struct CONTROLRIG_API FRigUnit_MathMatrixToTransform : public FRigUnit_MathMatri
 	}
 
 	RIGVM_METHOD()
-	virtual void Execute(const FRigUnitContext& Context) override;
+	virtual void Execute() override;
 
 	UPROPERTY(meta=(Input))
 	FMatrix Value;
@@ -109,7 +109,7 @@ struct CONTROLRIG_API FRigUnit_MathMatrixFromTransform : public FRigUnit_MathMat
 	}
 
 	RIGVM_METHOD()
-	virtual void Execute(const FRigUnitContext& Context) override;
+	virtual void Execute() override;
 
 	UPROPERTY(meta=(Input))
 	FTransform Transform;
@@ -124,7 +124,7 @@ struct CONTROLRIG_API FRigUnit_MathMatrixFromTransform : public FRigUnit_MathMat
 /**
  * Makes a matrix from a transform
  */
-USTRUCT(meta=(DisplayName="From Transform", TemplateName="Cast", Keywords="Make,Construct"))
+USTRUCT(meta=(DisplayName="From Transform", TemplateName="Cast", ExecuteContext="FRigVMExecuteContext", Keywords="Make,Construct"))
 struct CONTROLRIG_API FRigUnit_MathMatrixFromTransformV2 : public FRigUnit_MathMatrixBase
 {
 	GENERATED_BODY()
@@ -136,7 +136,7 @@ struct CONTROLRIG_API FRigUnit_MathMatrixFromTransformV2 : public FRigUnit_MathM
 	}
 
 	RIGVM_METHOD()
-	virtual void Execute(const FRigUnitContext& Context) override;
+	virtual void Execute() override;
 
 	UPROPERTY(meta=(Input))
 	FTransform Value;
@@ -163,7 +163,7 @@ struct CONTROLRIG_API FRigUnit_MathMatrixToVectors : public FRigUnit_MathMatrixB
 	}
 
 	RIGVM_METHOD()
-	virtual void Execute(const FRigUnitContext& Context) override;
+	virtual void Execute() override;
 
 	UPROPERTY(meta=(Input))
 	FMatrix Value;
@@ -199,7 +199,7 @@ struct CONTROLRIG_API FRigUnit_MathMatrixFromVectors : public FRigUnit_MathMatri
 	}
 
 	RIGVM_METHOD()
-	virtual void Execute(const FRigUnitContext& Context) override;
+	virtual void Execute() override;
 
 	UPROPERTY(meta=(Input))
 	FVector Origin;
@@ -225,7 +225,7 @@ struct CONTROLRIG_API FRigUnit_MathMatrixMul : public FRigUnit_MathMatrixBinaryA
 {
 	GENERATED_BODY()
 	RIGVM_METHOD()
-	virtual void Execute(const FRigUnitContext& Context) override;
+	virtual void Execute() override;
 };
 
 /**
@@ -236,5 +236,5 @@ struct CONTROLRIG_API FRigUnit_MathMatrixInverse : public FRigUnit_MathMatrixUna
 {
 	GENERATED_BODY()
 	RIGVM_METHOD()
-	virtual void Execute(const FRigUnitContext& Context) override;
+	virtual void Execute() override;
 };
diff --git a/Engine/Plugins/Animation/ControlRig/Source/ControlRig/Private/Units/Math/RigUnit_MathQuaternion.cpp b/Engine/Plugins/Animation/ControlRig/Source/ControlRig/Private/Units/Math/RigUnit_MathQuaternion.cpp
--- a/Engine/Plugins/Animation/ControlRig/Source/ControlRig/Private/Units/Math/RigUnit_MathQuaternion.cpp
+++ b/Engine/Plugins/Animation/ControlRig/Source/ControlRig/Private/Units/Math/RigUnit_MathQuaternion.cpp
@@ -264,7 +264,7 @@ FRigUnit_MathQuaternionMirrorTransform_Execute()
 {
 	FTransform Transform = FTransform::Identity;
 	Transform.SetRotation(Value);
-	FRigUnit_MathTransformMirrorTransform::StaticExecute(ExecuteContext, Transform, MirrorAxis, AxisToFlip, CentralTransform, Transform, Context);
+	FRigUnit_MathTransformMirrorTransform::StaticExecute(ExecuteContext, Transform, MirrorAxis, AxisToFlip, CentralTransform, Transform);
 	Result = Transform.GetRotation();
 }
 
