RigVM: Remove support for opaque attributes

diff --git a/Engine/Plugins/Animation/ControlRig/Source/ControlRig/Private/Units/Math/RigUnit_MathQuaternion.h b/Engine/Plugins/Animation/ControlRig/Source/ControlRig/Private/Units/Math/RigUnit_MathQuaternion.h
--- a/Engine/Plugins/Animation/ControlRig/Source/ControlRig/Private/Units/Math/RigUnit_MathQuaternion.h
+++ b/Engine/Plugins/Animation/ControlRig/Source/ControlRig/Private/Units/Math/RigUnit_MathQuaternion.h
@@ -77,7 +77,7 @@ struct CONTROLRIG_API FRigUnit_MathQuaternionMake : public FRigUnit_MathQuaterni
 {
 	GENERATED_BODY()
 	RIGVM_METHOD()
-	virtual void Execute(const FRigUnitContext& Context) override;
+	virtual void Execute() override;
 
 	FRigUnit_MathQuaternionMake()
 	{
@@ -117,7 +117,7 @@ struct CONTROLRIG_API FRigUnit_MathQuaternionFromAxisAndAngle : public FRigUnit_
 	}
 
 	RIGVM_METHOD()
-	virtual void Execute(const FRigUnitContext& Context) override;
+	virtual void Execute() override;
 
 	UPROPERTY(meta=(Input))
 	FVector Axis;
@@ -145,7 +145,7 @@ struct CONTROLRIG_API FRigUnit_MathQuaternionFromEuler : public FRigUnit_MathQua
 	}
 
 	RIGVM_METHOD()
-	virtual void Execute(const FRigUnitContext& Context) override;
+	virtual void Execute() override;
 
 	UPROPERTY(meta = (Input))
 	FVector Euler;
@@ -172,7 +172,7 @@ struct CONTROLRIG_API FRigUnit_MathQuaternionFromRotator : public FRigUnit_MathQ
 	}
 
 	RIGVM_METHOD()
-	virtual void Execute(const FRigUnitContext& Context) override;
+	virtual void Execute() override;
 
 	UPROPERTY(meta=(Input))
 	FRotator Rotator;
@@ -187,7 +187,7 @@ struct CONTROLRIG_API FRigUnit_MathQuaternionFromRotator : public FRigUnit_MathQ
 /**
  * Makes a quaternion from a rotator
  */
-USTRUCT(meta=(DisplayName="From Rotator", TemplateName="Cast", Keywords="Make,Construct"))
+USTRUCT(meta=(DisplayName="From Rotator", TemplateName="Cast", ExecuteContext="FRigVMExecuteContext", Keywords="Make,Construct"))
 struct CONTROLRIG_API FRigUnit_MathQuaternionFromRotatorV2 : public FRigUnit_MathQuaternionBase
 {
 	GENERATED_BODY()
@@ -199,7 +199,7 @@ struct CONTROLRIG_API FRigUnit_MathQuaternionFromRotatorV2 : public FRigUnit_Mat
 	}
 
 	RIGVM_METHOD()
-	virtual void Execute(const FRigUnitContext& Context) override;
+	virtual void Execute() override;
 
 	UPROPERTY(meta=(Input))
 	FRotator Value;
@@ -223,7 +223,7 @@ struct CONTROLRIG_API FRigUnit_MathQuaternionFromTwoVectors : public FRigUnit_Ma
 	}
 
 	RIGVM_METHOD()
-	virtual void Execute(const FRigUnitContext& Context) override;
+	virtual void Execute() override;
 
 	UPROPERTY(meta=(Input))
 	FVector A;
@@ -251,7 +251,7 @@ struct CONTROLRIG_API FRigUnit_MathQuaternionToAxisAndAngle : public FRigUnit_Ma
 	}
 
 	RIGVM_METHOD()
-	virtual void Execute(const FRigUnitContext& Context) override;
+	virtual void Execute() override;
 
 	UPROPERTY(meta=(Input))
 	FQuat Value;
@@ -278,7 +278,7 @@ struct CONTROLRIG_API FRigUnit_MathQuaternionScale : public FRigUnit_MathQuatern
 	}
 
 	RIGVM_METHOD()
-	virtual void Execute(const FRigUnitContext& Context) override;
+	virtual void Execute() override;
 
 	UPROPERTY(meta=(Input, Output))
 	FQuat Value;
@@ -305,7 +305,7 @@ struct CONTROLRIG_API FRigUnit_MathQuaternionScaleV2 : public FRigUnit_MathQuate
 	}
 
 	RIGVM_METHOD()
-	virtual void Execute(const FRigUnitContext& Context) override;
+	virtual void Execute() override;
 
 	UPROPERTY(meta=(Input))
 	FQuat Value;
@@ -333,7 +333,7 @@ struct CONTROLRIG_API FRigUnit_MathQuaternionToEuler : public FRigUnit_MathQuate
 	}
 
 	RIGVM_METHOD()
-	virtual void Execute(const FRigUnitContext& Context) override;
+	virtual void Execute() override;
 
 	UPROPERTY(meta=(Input))
 	FQuat Value;
@@ -348,7 +348,7 @@ struct CONTROLRIG_API FRigUnit_MathQuaternionToEuler : public FRigUnit_MathQuate
 /**
  * Retrieves the rotator
  */
-USTRUCT(meta=(DisplayName="To Rotator", TemplateName="Cast", Keywords="Make,Construct"))
+USTRUCT(meta=(DisplayName="To Rotator", TemplateName="Cast", ExecuteContext="FRigVMExecuteContext", Keywords="Make,Construct"))
 struct CONTROLRIG_API FRigUnit_MathQuaternionToRotator : public FRigUnit_MathQuaternionBase
 {
 	GENERATED_BODY()
@@ -360,7 +360,7 @@ struct CONTROLRIG_API FRigUnit_MathQuaternionToRotator : public FRigUnit_MathQua
 	}
 
 	RIGVM_METHOD()
-	virtual void Execute(const FRigUnitContext& Context) override;
+	virtual void Execute() override;
 
 	UPROPERTY(meta=(Input))
 	FQuat Value;
@@ -377,7 +377,7 @@ struct CONTROLRIG_API FRigUnit_MathQuaternionMul : public FRigUnit_MathQuaternio
 {
 	GENERATED_BODY()
 	RIGVM_METHOD()
-	virtual void Execute(const FRigUnitContext& Context) override;
+	virtual void Execute() override;
 };
 
 /**
@@ -388,7 +388,7 @@ struct CONTROLRIG_API FRigUnit_MathQuaternionInverse : public FRigUnit_MathQuate
 {
 	GENERATED_BODY()
 	RIGVM_METHOD()
-	virtual void Execute(const FRigUnitContext& Context) override;
+	virtual void Execute() override;
 };
 
 /**
@@ -406,7 +406,7 @@ struct CONTROLRIG_API FRigUnit_MathQuaternionSlerp : public FRigUnit_MathQuatern
 	}
 
 	RIGVM_METHOD()
-	virtual void Execute(const FRigUnitContext& Context) override;
+	virtual void Execute() override;
 
 	UPROPERTY(meta=(Input))
 	FQuat A;
@@ -436,7 +436,7 @@ struct CONTROLRIG_API FRigUnit_MathQuaternionEquals : public FRigUnit_MathQuater
 	}	
 
 	RIGVM_METHOD()
-	virtual void Execute(const FRigUnitContext& Context) override;
+	virtual void Execute() override;
 
 	UPROPERTY(meta=(Input))
 	FQuat A;
@@ -466,7 +466,7 @@ struct CONTROLRIG_API FRigUnit_MathQuaternionNotEquals : public FRigUnit_MathQua
 	}
 
 	RIGVM_METHOD()
-	virtual void Execute(const FRigUnitContext& Context) override;
+	virtual void Execute() override;
 
 	UPROPERTY(meta=(Input))
 	FQuat A;
@@ -496,7 +496,7 @@ struct CONTROLRIG_API FRigUnit_MathQuaternionSelectBool : public FRigUnit_MathQu
 	}
 
 	RIGVM_METHOD()
-	virtual void Execute(const FRigUnitContext& Context) override;
+	virtual void Execute() override;
 
 	UPROPERTY(meta=(Input))
 	bool Condition;
@@ -529,7 +529,7 @@ struct CONTROLRIG_API FRigUnit_MathQuaternionDot : public FRigUnit_MathQuaternio
 	}
 
 	RIGVM_METHOD()
-	virtual void Execute(const FRigUnitContext& Context) override;
+	virtual void Execute() override;
 
 	UPROPERTY(meta=(Input))
 	FQuat A;
@@ -549,7 +549,7 @@ struct CONTROLRIG_API FRigUnit_MathQuaternionUnit : public FRigUnit_MathQuaterni
 {
 	GENERATED_BODY()
 	RIGVM_METHOD()
-	virtual void Execute(const FRigUnitContext& Context) override;
+	virtual void Execute() override;
 };
 
 /**
@@ -567,7 +567,7 @@ struct CONTROLRIG_API FRigUnit_MathQuaternionRotateVector : public FRigUnit_Math
 	}
 
 	RIGVM_METHOD()
-	virtual void Execute(const FRigUnitContext& Context) override;
+	virtual void Execute() override;
 
 	UPROPERTY(meta=(Input))
 	FQuat Transform;
@@ -595,7 +595,7 @@ struct CONTROLRIG_API FRigUnit_MathQuaternionGetAxis: public FRigUnit_MathQuater
 	}
 
 	RIGVM_METHOD()
-	virtual void Execute(const FRigUnitContext& Context) override;
+	virtual void Execute() override;
 
 	UPROPERTY(meta = (Input))
 	FQuat Quaternion;
@@ -623,7 +623,7 @@ struct CONTROLRIG_API FRigUnit_MathQuaternionSwingTwist : public FRigUnit_MathQu
 	}
 
 	RIGVM_METHOD()
-	virtual void Execute(const FRigUnitContext& Context) override;
+	virtual void Execute() override;
 
 	UPROPERTY(meta = (Input))
 	FQuat	Input;
@@ -656,7 +656,7 @@ struct CONTROLRIG_API FRigUnit_MathQuaternionRotationOrder : public FRigUnit_Mat
 
 
 	RIGVM_METHOD()
-	virtual void Execute(const FRigUnitContext& Context) override;
+	virtual void Execute() override;
 };
 
 /**
@@ -673,7 +673,7 @@ struct CONTROLRIG_API FRigUnit_MathQuaternionMakeRelative : public FRigUnit_Math
 	}
 
 	RIGVM_METHOD()
-	virtual void Execute(const FRigUnitContext& Context) override;
+	virtual void Execute() override;
 
 	UPROPERTY(meta=(Input))
 	FQuat Global;
@@ -699,7 +699,7 @@ struct CONTROLRIG_API FRigUnit_MathQuaternionMakeAbsolute : public FRigUnit_Math
 	}
 
 	RIGVM_METHOD()
-	virtual void Execute(const FRigUnitContext& Context) override;
+	virtual void Execute() override;
 
 	UPROPERTY(meta = (Input))
 	FQuat Local;
@@ -728,7 +728,7 @@ struct CONTROLRIG_API FRigUnit_MathQuaternionMirrorTransform : public FRigUnit_M
 	}
 
 	RIGVM_METHOD()
-	virtual void Execute(const FRigUnitContext& Context) override;
+	virtual void Execute() override;
 
 	UPROPERTY(meta=(Input))
 	FQuat Value;
diff --git a/Engine/Plugins/Animation/ControlRig/Source/ControlRig/Private/Units/Math/RigUnit_MathQuaternionTests.cpp b/Engine/Plugins/Animation/ControlRig/Source/ControlRig/Private/Units/Math/RigUnit_MathQuaternionTests.cpp
--- a/Engine/Plugins/Animation/ControlRig/Source/ControlRig/Private/Units/Math/RigUnit_MathQuaternionTests.cpp
+++ b/Engine/Plugins/Animation/ControlRig/Source/ControlRig/Private/Units/Math/RigUnit_MathQuaternionTests.cpp
@@ -224,13 +224,13 @@ IMPLEMENT_RIGUNIT_AUTOMATION_TEST(FRigUnit_MathQuaternionMakeRelative)
 	TransformUnit.Global.SetRotation(Unit.Global);
 	TransformUnit.Parent.SetRotation(Unit.Parent);
 
-	Context.State = EControlRigState::Init;
-	Unit.Execute(Context);
-	TransformUnit.Execute(Context);
+	ExecuteContext.UnitContext.State = EControlRigState::Init;
+	Unit.Execute();
+	TransformUnit.Execute();
 
-	Context.State = EControlRigState::Update;
-	Unit.Execute(Context);
-	TransformUnit.Execute(Context);
+	ExecuteContext.UnitContext.State = EControlRigState::Update;
+	Unit.Execute();
+	TransformUnit.Execute();
 
 	AddErrorIfFalse(FRigUnit_MathQuatTest_Utils::IsNearlyEqual(Unit.Local, TransformUnit.Local.GetRotation()), TEXT("unexpected result"));
 	return true;
@@ -245,13 +245,13 @@ IMPLEMENT_RIGUNIT_AUTOMATION_TEST(FRigUnit_MathQuaternionMakeAbsolute)
 	TransformUnit.Local.SetRotation(Unit.Local);
 	TransformUnit.Parent.SetRotation(Unit.Parent);
 
-	Context.State = EControlRigState::Init;
-	Unit.Execute(Context);
-	TransformUnit.Execute(Context);
+	ExecuteContext.UnitContext.State = EControlRigState::Init;
+	Unit.Execute();
+	TransformUnit.Execute();
 
-	Context.State = EControlRigState::Update;
-	Unit.Execute(Context);
-	TransformUnit.Execute(Context);
+	ExecuteContext.UnitContext.State = EControlRigState::Update;
+	Unit.Execute();
+	TransformUnit.Execute();
 
 	AddErrorIfFalse(FRigUnit_MathQuatTest_Utils::IsNearlyEqual(Unit.Global, TransformUnit.Global.GetRotation()), TEXT("unexpected result"));
 	return true;
diff --git a/Engine/Plugins/Animation/ControlRig/Source/ControlRig/Private/Units/Math/RigUnit_MathRBFInterpolate.cpp b/Engine/Plugins/Animation/ControlRig/Source/ControlRig/Private/Units/Math/RigUnit_MathRBFInterpolate.cpp
--- a/Engine/Plugins/Animation/ControlRig/Source/ControlRig/Private/Units/Math/RigUnit_MathRBFInterpolate.cpp
+++ b/Engine/Plugins/Animation/ControlRig/Source/ControlRig/Private/Units/Math/RigUnit_MathRBFInterpolate.cpp
@@ -301,7 +301,7 @@ FRigUnit_MathRBFInterpolateQuatFloat_Execute()
 	const TArrayView<const FMathRBFInterpolateQuatFloat_Target> TargetsView(Targets.GetData(), Targets.Num());
 	TArray<float> Weights;
 	GetInterpolatedWeights<FMathRBFInterpolateQuatFloat_Target>(
-		Context.State, WorkData, TargetsView, Input,  
+		ExecuteContext.UnitContext.State, WorkData, TargetsView, Input,  
 		DistanceFunction, SmoothingFunction, SmoothingAngle, 
 		bNormalizeOutput, TwistAxis, Weights);
 
@@ -320,7 +320,7 @@ FRigUnit_MathRBFInterpolateQuatVector_Execute()
 	const TArrayView<const FMathRBFInterpolateQuatVector_Target> TargetsView(Targets.GetData(), Targets.Num());
 	TArray<float> Weights;
 	GetInterpolatedWeights<FMathRBFInterpolateQuatVector_Target>(
-		Context.State, WorkData, TargetsView, Input,
+		ExecuteContext.UnitContext.State, WorkData, TargetsView, Input,
 		DistanceFunction, SmoothingFunction, SmoothingAngle, 
 		bNormalizeOutput, TwistAxis, Weights);
 
@@ -339,7 +339,7 @@ FRigUnit_MathRBFInterpolateQuatColor_Execute()
 	const TArrayView<const FMathRBFInterpolateQuatColor_Target> TargetsView(Targets.GetData(), Targets.Num());
 	TArray<float> Weights;
 	GetInterpolatedWeights<FMathRBFInterpolateQuatColor_Target>(
-		Context.State, WorkData, TargetsView, Input,
+		ExecuteContext.UnitContext.State, WorkData, TargetsView, Input,
 		DistanceFunction, SmoothingFunction, SmoothingAngle, 
 		bNormalizeOutput, TwistAxis, Weights);
 
@@ -357,7 +357,7 @@ FRigUnit_MathRBFInterpolateQuatQuat_Execute()
 	const TArrayView<const FMathRBFInterpolateQuatQuat_Target> TargetsView(Targets.GetData(), Targets.Num());
 	TArray<float> Weights;
 	GetInterpolatedWeights<FMathRBFInterpolateQuatQuat_Target>(
-		Context.State, WorkData, TargetsView, Input,
+		ExecuteContext.UnitContext.State, WorkData, TargetsView, Input,
 		DistanceFunction, SmoothingFunction, SmoothingAngle,
 		bNormalizeOutput, TwistAxis, Weights);
 	
@@ -377,7 +377,7 @@ FRigUnit_MathRBFInterpolateQuatXform_Execute()
 	const TArrayView<const FMathRBFInterpolateQuatXform_Target> TargetsView(Targets.GetData(), Targets.Num());
 	TArray<float> Weights;
 	GetInterpolatedWeights<FMathRBFInterpolateQuatXform_Target>(
-		Context.State, WorkData, TargetsView, Input,
+		ExecuteContext.UnitContext.State, WorkData, TargetsView, Input,
 		DistanceFunction, SmoothingFunction, SmoothingAngle,
 		bNormalizeOutput, TwistAxis, Weights);
 
@@ -407,7 +407,7 @@ FRigUnit_MathRBFInterpolateVectorFloat_Execute()
 	const TArrayView<const FMathRBFInterpolateVectorFloat_Target> TargetsView(Targets.GetData(), Targets.Num());
 	TArray<float> Weights;
 	GetInterpolatedWeights<FMathRBFInterpolateVectorFloat_Target>(
-		Context.State, WorkData, TargetsView, Input,
+		ExecuteContext.UnitContext.State, WorkData, TargetsView, Input,
 		DistanceFunction, SmoothingFunction, SmoothingRadius, 
 		bNormalizeOutput, Weights);
 
@@ -426,7 +426,7 @@ FRigUnit_MathRBFInterpolateVectorVector_Execute()
 	const TArrayView<const FMathRBFInterpolateVectorVector_Target> TargetsView(Targets.GetData(), Targets.Num());
 	TArray<float> Weights;
 	GetInterpolatedWeights<FMathRBFInterpolateVectorVector_Target>(
-		Context.State, WorkData, TargetsView, Input,
+		ExecuteContext.UnitContext.State, WorkData, TargetsView, Input,
 		DistanceFunction, SmoothingFunction, SmoothingRadius, 
 		bNormalizeOutput, Weights);
 
@@ -445,7 +445,7 @@ FRigUnit_MathRBFInterpolateVectorColor_Execute()
 	const TArrayView<const FMathRBFInterpolateVectorColor_Target> TargetsView(Targets.GetData(), Targets.Num());
 	TArray<float> Weights;
 	GetInterpolatedWeights<FMathRBFInterpolateVectorColor_Target>(
-		Context.State, WorkData, TargetsView, Input,
+		ExecuteContext.UnitContext.State, WorkData, TargetsView, Input,
 		DistanceFunction, SmoothingFunction, SmoothingRadius, 
 		bNormalizeOutput, Weights);
 
@@ -464,7 +464,7 @@ FRigUnit_MathRBFInterpolateVectorQuat_Execute()
 	const TArrayView<const FMathRBFInterpolateVectorQuat_Target> TargetsView(Targets.GetData(), Targets.Num());
 	TArray<float> Weights;
 	GetInterpolatedWeights<FMathRBFInterpolateVectorQuat_Target>(
-		Context.State, WorkData, TargetsView, Input,
+		ExecuteContext.UnitContext.State, WorkData, TargetsView, Input,
 		DistanceFunction, SmoothingFunction, SmoothingRadius, 
 		bNormalizeOutput, Weights);
 
@@ -484,7 +484,7 @@ FRigUnit_MathRBFInterpolateVectorXform_Execute()
 	const TArrayView<const FMathRBFInterpolateVectorXform_Target> TargetsView(Targets.GetData(), Targets.Num());
 	TArray<float> Weights;
 	GetInterpolatedWeights<FMathRBFInterpolateVectorXform_Target>(
-		Context.State, WorkData, TargetsView, Input,
+		ExecuteContext.UnitContext.State, WorkData, TargetsView, Input,
 		DistanceFunction, SmoothingFunction, SmoothingRadius, 
 		bNormalizeOutput, Weights);
 
diff --git a/Engine/Plugins/Animation/ControlRig/Source/ControlRig/Private/Units/Math/RigUnit_MathRBFInterpolate.h b/Engine/Plugins/Animation/ControlRig/Source/ControlRig/Private/Units/Math/RigUnit_MathRBFInterpolate.h
--- a/Engine/Plugins/Animation/ControlRig/Source/ControlRig/Private/Units/Math/RigUnit_MathRBFInterpolate.h
+++ b/Engine/Plugins/Animation/ControlRig/Source/ControlRig/Private/Units/Math/RigUnit_MathRBFInterpolate.h
@@ -195,7 +195,7 @@ struct CONTROLRIG_API FRigUnit_MathRBFInterpolateQuatFloat :
 	float Output = 0.0f;
 
 	RIGVM_METHOD()
-	void Execute(const FRigUnitContext& Context) override;
+	void Execute() override;
 };
 
 
@@ -224,7 +224,7 @@ struct CONTROLRIG_API FRigUnit_MathRBFInterpolateQuatVector :
 	FVector Output = FVector::ZeroVector;
 
 	RIGVM_METHOD()
-	void Execute(const FRigUnitContext& Context) override;
+	void Execute() override;
 };
 
 
@@ -253,7 +253,7 @@ struct CONTROLRIG_API FRigUnit_MathRBFInterpolateQuatColor :
 	FLinearColor Output = FLinearColor::Transparent;
 
 	RIGVM_METHOD()
-	void Execute(const FRigUnitContext& Context) override;
+	void Execute() override;
 };
 
 
@@ -282,7 +282,7 @@ struct CONTROLRIG_API FRigUnit_MathRBFInterpolateQuatQuat :
 	FQuat Output = FQuat::Identity;
 
 	RIGVM_METHOD()
-	void Execute(const FRigUnitContext& Context) override;
+	void Execute() override;
 };
 
 
@@ -311,7 +311,7 @@ struct CONTROLRIG_API FRigUnit_MathRBFInterpolateQuatXform :
 	FTransform Output = FTransform::Identity;
 
 	RIGVM_METHOD()
-	void Execute(const FRigUnitContext& Context) override;
+	void Execute() override;
 };
 
 
@@ -342,7 +342,7 @@ struct CONTROLRIG_API FRigUnit_MathRBFInterpolateVectorFloat :
 	float Output = 0.0f;
 
 	RIGVM_METHOD()
-	void Execute(const FRigUnitContext& Context) override;
+	void Execute() override;
 };
 
 
@@ -371,7 +371,7 @@ struct CONTROLRIG_API FRigUnit_MathRBFInterpolateVectorVector :
 	FVector Output = FVector::ZeroVector;
 
 	RIGVM_METHOD()
-	void Execute(const FRigUnitContext& Context) override;
+	void Execute() override;
 };
 
 
@@ -400,7 +400,7 @@ struct CONTROLRIG_API FRigUnit_MathRBFInterpolateVectorColor :
 	FLinearColor Output = FLinearColor::Transparent;
 
 	RIGVM_METHOD()
-	void Execute(const FRigUnitContext& Context) override;
+	void Execute() override;
 };
 
 
@@ -430,7 +430,7 @@ struct CONTROLRIG_API FRigUnit_MathRBFInterpolateVectorQuat :
 	FQuat Output = FQuat::Identity;
 
 	RIGVM_METHOD()
-	void Execute(const FRigUnitContext& Context) override;
+	void Execute() override;
 };
 
 USTRUCT()
@@ -458,5 +458,5 @@ struct CONTROLRIG_API FRigUnit_MathRBFInterpolateVectorXform :
 	FTransform Output = FTransform::Identity;
 
 	RIGVM_METHOD()
-	void Execute(const FRigUnitContext& Context) override;
+	void Execute() override;
 };
diff --git a/Engine/Plugins/Animation/ControlRig/Source/ControlRig/Private/Units/Math/RigUnit_MathTransform.cpp b/Engine/Plugins/Animation/ControlRig/Source/ControlRig/Private/Units/Math/RigUnit_MathTransform.cpp
--- a/Engine/Plugins/Animation/ControlRig/Source/ControlRig/Private/Units/Math/RigUnit_MathTransform.cpp
+++ b/Engine/Plugins/Animation/ControlRig/Source/ControlRig/Private/Units/Math/RigUnit_MathTransform.cpp
@@ -196,7 +196,7 @@ FRigUnit_MathTransformClampSpatially_Execute()
 {
 	DECLARE_SCOPE_HIERARCHICAL_COUNTER_RIGUNIT()
 	FVector Position;
-	FRigUnit_MathVectorClampSpatially::StaticExecute(ExecuteContext, Value.GetTranslation(), Axis, Type, Minimum, Maximum, Space, bDrawDebug, DebugColor, DebugThickness, Position, Context);
+	FRigUnit_MathVectorClampSpatially::StaticExecute(ExecuteContext, Value.GetTranslation(), Axis, Type, Minimum, Maximum, Space, bDrawDebug, DebugColor, DebugThickness, Position);
 	Result = Value;
 	Result.SetTranslation(Position);
 }
@@ -208,8 +208,8 @@ FRigUnit_MathTransformMirrorTransform_Execute()
 	MirrorSettings.AxisToFlip = AxisToFlip;
 
 	FTransform Local = FTransform::Identity;
-	FRigUnit_MathTransformMakeRelative::StaticExecute(ExecuteContext, Value, CentralTransform, Local, Context);
+	FRigUnit_MathTransformMakeRelative::StaticExecute(ExecuteContext, Value, CentralTransform, Local);
 	Local = MirrorSettings.MirrorTransform(Local);
-	FRigUnit_MathTransformMakeAbsolute::StaticExecute(ExecuteContext, Local, CentralTransform, Result, Context);
+	FRigUnit_MathTransformMakeAbsolute::StaticExecute(ExecuteContext, Local, CentralTransform, Result);
 }
 
diff --git a/Engine/Plugins/Animation/ControlRig/Source/ControlRig/Private/Units/Math/RigUnit_MathTransform.h b/Engine/Plugins/Animation/ControlRig/Source/ControlRig/Private/Units/Math/RigUnit_MathTransform.h
--- a/Engine/Plugins/Animation/ControlRig/Source/ControlRig/Private/Units/Math/RigUnit_MathTransform.h
+++ b/Engine/Plugins/Animation/ControlRig/Source/ControlRig/Private/Units/Math/RigUnit_MathTransform.h
@@ -84,7 +84,7 @@ struct CONTROLRIG_API FRigUnit_MathTransformMake : public FRigUnit_MathTransform
 {
 	GENERATED_BODY()
 	RIGVM_METHOD()
-	virtual void Execute(const FRigUnitContext& Context) override;
+	virtual void Execute() override;
 
 	FRigUnit_MathTransformMake()
 	{
@@ -122,7 +122,7 @@ struct CONTROLRIG_API FRigUnit_MathTransformFromEulerTransform : public FRigUnit
 	}
 
 	RIGVM_METHOD()
-	virtual void Execute(const FRigUnitContext& Context) override;
+	virtual void Execute() override;
 
 	UPROPERTY(meta=(Input))
 	FEulerTransform EulerTransform;
@@ -137,7 +137,7 @@ struct CONTROLRIG_API FRigUnit_MathTransformFromEulerTransform : public FRigUnit
 /**
  * Makes a quaternion based transform from a euler based transform
  */
-USTRUCT(meta=(DisplayName="To Transform", TemplateName="Cast", Keywords="Make,Construct"))
+USTRUCT(meta=(DisplayName="To Transform", TemplateName="Cast", ExecuteContext="FRigVMExecuteContext", Keywords="Make,Construct"))
 struct CONTROLRIG_API FRigUnit_MathTransformFromEulerTransformV2 : public FRigUnit_MathTransformBase
 {
 	GENERATED_BODY()
@@ -149,7 +149,7 @@ struct CONTROLRIG_API FRigUnit_MathTransformFromEulerTransformV2 : public FRigUn
 	}
 
 	RIGVM_METHOD()
-	virtual void Execute(const FRigUnitContext& Context) override;
+	virtual void Execute() override;
 
 	UPROPERTY(meta=(Input))
 	FEulerTransform Value;
@@ -161,7 +161,7 @@ struct CONTROLRIG_API FRigUnit_MathTransformFromEulerTransformV2 : public FRigUn
 /**
  * Retrieves a euler based transform from a quaternion based transform
  */
-USTRUCT(meta=(DisplayName="To Euler Transform", TemplateName="Cast", Keywords="Make,Construct"))
+USTRUCT(meta=(DisplayName="To Euler Transform", TemplateName="Cast", ExecuteContext="FRigVMExecuteContext", Keywords="Make,Construct"))
 struct CONTROLRIG_API FRigUnit_MathTransformToEulerTransform : public FRigUnit_MathTransformBase
 {
 	GENERATED_BODY()
@@ -173,7 +173,7 @@ struct CONTROLRIG_API FRigUnit_MathTransformToEulerTransform : public FRigUnit_M
 	}
 
 	RIGVM_METHOD()
-	virtual void Execute(const FRigUnitContext& Context) override;
+	virtual void Execute() override;
 
 	UPROPERTY(meta=(Input))
 	FTransform Value;
@@ -190,7 +190,7 @@ struct CONTROLRIG_API FRigUnit_MathTransformMul : public FRigUnit_MathTransformB
 {
 	GENERATED_BODY()
 	RIGVM_METHOD()
-	virtual void Execute(const FRigUnitContext& Context) override;
+	virtual void Execute() override;
 };
 
 /**
@@ -207,7 +207,7 @@ struct CONTROLRIG_API FRigUnit_MathTransformMakeRelative : public FRigUnit_MathT
 	}
 
 	RIGVM_METHOD()
-	virtual void Execute(const FRigUnitContext& Context) override;
+	virtual void Execute() override;
 
 	UPROPERTY(meta=(Input))
 	FTransform Global;
@@ -233,7 +233,7 @@ struct CONTROLRIG_API FRigUnit_MathTransformMakeAbsolute : public FRigUnit_MathT
 	}
 
 	RIGVM_METHOD()
-	virtual void Execute(const FRigUnitContext& Context) override;
+	virtual void Execute() override;
 
 	UPROPERTY(meta = (Input))
 	FTransform Local;
@@ -264,7 +264,7 @@ struct CONTROLRIG_API FRigUnit_MathTransformAccumulateArray : public FRigUnit_Ma
 	virtual FString GetUnitLabel() const override;
 
 	RIGVM_METHOD()
-	virtual void Execute(const FRigUnitContext& Context) override;
+	virtual void Execute() override;
 
 	UPROPERTY(meta=(Input, Output))
 	TArray<FTransform> Transforms;
@@ -297,7 +297,7 @@ struct CONTROLRIG_API FRigUnit_MathTransformInverse : public FRigUnit_MathTransf
 {
 	GENERATED_BODY()
 	RIGVM_METHOD()
-	virtual void Execute(const FRigUnitContext& Context) override;
+	virtual void Execute() override;
 };
 
 /**
@@ -315,7 +315,7 @@ struct CONTROLRIG_API FRigUnit_MathTransformLerp : public FRigUnit_MathTransform
 	}
 
 	RIGVM_METHOD()
-	virtual void Execute(const FRigUnitContext& Context) override;
+	virtual void Execute() override;
 
 	UPROPERTY(meta=(Input))
 	FTransform A;
@@ -345,7 +345,7 @@ struct CONTROLRIG_API FRigUnit_MathTransformSelectBool : public FRigUnit_MathTra
 	}
 
 	RIGVM_METHOD()
-	virtual void Execute(const FRigUnitContext& Context) override;
+	virtual void Execute() override;
 
 	UPROPERTY(meta=(Input))
 	bool Condition;
@@ -378,7 +378,7 @@ struct CONTROLRIG_API FRigUnit_MathTransformRotateVector : public FRigUnit_MathT
 	}
 
 	RIGVM_METHOD()
-	virtual void Execute(const FRigUnitContext& Context) override;
+	virtual void Execute() override;
 
 	UPROPERTY(meta=(Input))
 	FTransform Transform;
@@ -405,7 +405,7 @@ struct CONTROLRIG_API FRigUnit_MathTransformTransformVector : public FRigUnit_Ma
 	}
 	
 	RIGVM_METHOD()
-	virtual void Execute(const FRigUnitContext& Context) override;
+	virtual void Execute() override;
 
 	UPROPERTY(meta=(Input))
 	FTransform Transform;
@@ -436,7 +436,7 @@ struct CONTROLRIG_API FRigUnit_MathTransformFromSRT : public FRigUnit_MathTransf
 	}
 	
 	RIGVM_METHOD()
-	virtual void Execute(const FRigUnitContext& Context) override;
+	virtual void Execute() override;
 
 	UPROPERTY(meta=(Input))
 	FVector Location;
@@ -471,7 +471,7 @@ struct CONTROLRIG_API FRigUnit_MathTransformArrayToSRT : public FRigUnit_MathTra
 	}
 
 	RIGVM_METHOD()
-		virtual void Execute(const FRigUnitContext& Context) override;
+		virtual void Execute() override;
 
 	UPROPERTY(meta = (Input))
 	TArray<FTransform> Transforms;
@@ -508,7 +508,7 @@ struct CONTROLRIG_API FRigUnit_MathTransformClampSpatially : public FRigUnit_Mat
 	}
 
 	RIGVM_METHOD()
-	virtual void Execute(const FRigUnitContext& Context) override;
+	virtual void Execute() override;
 
 	UPROPERTY(meta = (Input))
 	FTransform Value;
@@ -560,7 +560,7 @@ struct CONTROLRIG_API FRigUnit_MathTransformMirrorTransform : public FRigUnit_Ma
 	}
 
 	RIGVM_METHOD()
-	virtual void Execute(const FRigUnitContext& Context) override;
+	virtual void Execute() override;
 
 	UPROPERTY(meta=(Input))
 	FTransform Value;
diff --git a/Engine/Plugins/Animation/ControlRig/Source/ControlRig/Private/Units/Math/RigUnit_MathVector.cpp b/Engine/Plugins/Animation/ControlRig/Source/ControlRig/Private/Units/Math/RigUnit_MathVector.cpp
--- a/Engine/Plugins/Animation/ControlRig/Source/ControlRig/Private/Units/Math/RigUnit_MathVector.cpp
+++ b/Engine/Plugins/Animation/ControlRig/Source/ControlRig/Private/Units/Math/RigUnit_MathVector.cpp
@@ -433,7 +433,7 @@ FRigUnit_MathVectorClampSpatially_Execute()
 	DECLARE_SCOPE_HIERARCHICAL_COUNTER_RIGUNIT()
 	Result = FControlRigMathLibrary::ClampSpatially(Value, Axis, Type, Minimum, Maximum, Space);
 
-	if (Context.DrawInterface != nullptr && bDrawDebug)
+	if (ExecuteContext.UnitContext.DrawInterface != nullptr && bDrawDebug)
 	{
 		switch (Type)
 		{
@@ -464,8 +464,8 @@ FRigUnit_MathVectorClampSpatially_Execute()
 					}
 				}
 
-				Context.DrawInterface->DrawLine(Space, Points[0], Points[1], DebugColor, DebugThickness);
-				Context.DrawInterface->DrawPoints(Space, Points, DebugThickness * 8.f, DebugColor);
+				ExecuteContext.UnitContext.DrawInterface->DrawLine(Space, Points[0], Points[1], DebugColor, DebugThickness);
+				ExecuteContext.UnitContext.DrawInterface->DrawPoints(Space, Points, DebugThickness * 8.f, DebugColor);
 
 				break;
 			}
@@ -491,9 +491,9 @@ FRigUnit_MathVectorClampSpatially_Execute()
 				}
 				if (Minimum > SMALL_NUMBER)
 				{
-					Context.DrawInterface->DrawArc(Space, CircleTransform, Minimum, 0.f, PI * 2.f, DebugColor, DebugThickness, 16);
+					ExecuteContext.UnitContext.DrawInterface->DrawArc(Space, CircleTransform, Minimum, 0.f, PI * 2.f, DebugColor, DebugThickness, 16);
 				}
-				Context.DrawInterface->DrawArc(Space, CircleTransform, Maximum, 0.f, PI * 2.f, DebugColor, DebugThickness, 16);
+				ExecuteContext.UnitContext.DrawInterface->DrawArc(Space, CircleTransform, Maximum, 0.f, PI * 2.f, DebugColor, DebugThickness, 16);
 				break;
 			}
 			default:
@@ -502,21 +502,21 @@ FRigUnit_MathVectorClampSpatially_Execute()
 				FTransform CircleTransform = FTransform::Identity;
 				if (Minimum > SMALL_NUMBER)
 				{
-					Context.DrawInterface->DrawArc(Space, CircleTransform, Minimum, 0.f, PI * 2.f, DebugColor, DebugThickness, 16);
+					ExecuteContext.UnitContext.DrawInterface->DrawArc(Space, CircleTransform, Minimum, 0.f, PI * 2.f, DebugColor, DebugThickness, 16);
 				}
-				Context.DrawInterface->DrawArc(Space, CircleTransform, Maximum, 0.f, PI * 2.f, DebugColor, DebugThickness, 16);
+				ExecuteContext.UnitContext.DrawInterface->DrawArc(Space, CircleTransform, Maximum, 0.f, PI * 2.f, DebugColor, DebugThickness, 16);
 				CircleTransform.SetRotation(FQuat(FVector(0.f, 1.f, 0.f), PI * 0.5f));
 				if (Minimum > SMALL_NUMBER)
 				{
-					Context.DrawInterface->DrawArc(Space, CircleTransform, Minimum, 0.f, PI * 2.f, DebugColor, DebugThickness, 16);
+					ExecuteContext.UnitContext.DrawInterface->DrawArc(Space, CircleTransform, Minimum, 0.f, PI * 2.f, DebugColor, DebugThickness, 16);
 				}
-				Context.DrawInterface->DrawArc(Space, CircleTransform, Maximum, 0.f, PI * 2.f, DebugColor, DebugThickness, 16);
+				ExecuteContext.UnitContext.DrawInterface->DrawArc(Space, CircleTransform, Maximum, 0.f, PI * 2.f, DebugColor, DebugThickness, 16);
 				CircleTransform.SetRotation(FQuat(FVector(1.f, 0.f, 0.f), PI * 0.5f));
 				if (Minimum > SMALL_NUMBER)
 				{
-					Context.DrawInterface->DrawArc(Space, CircleTransform, Minimum, 0.f, PI * 2.f, DebugColor, DebugThickness, 16);
+					ExecuteContext.UnitContext.DrawInterface->DrawArc(Space, CircleTransform, Minimum, 0.f, PI * 2.f, DebugColor, DebugThickness, 16);
 				}
-				Context.DrawInterface->DrawArc(Space, CircleTransform, Maximum, 0.f, PI * 2.f, DebugColor, DebugThickness, 16);
+				ExecuteContext.UnitContext.DrawInterface->DrawArc(Space, CircleTransform, Maximum, 0.f, PI * 2.f, DebugColor, DebugThickness, 16);
 				break;
 			}
 
@@ -570,6 +570,6 @@ FRigUnit_MathVectorMirrorTransform_Execute()
 {
 	FTransform Transform = FTransform::Identity;
 	Transform.SetTranslation(Value);
-	FRigUnit_MathTransformMirrorTransform::StaticExecute(ExecuteContext, Transform, MirrorAxis, AxisToFlip, CentralTransform, Transform, Context);
+	FRigUnit_MathTransformMirrorTransform::StaticExecute(ExecuteContext, Transform, MirrorAxis, AxisToFlip, CentralTransform, Transform);
 	Result = Transform.GetTranslation();
 }
diff --git a/Engine/Plugins/Animation/ControlRig/Source/ControlRig/Private/Units/Math/RigUnit_MathVector.h b/Engine/Plugins/Animation/ControlRig/Source/ControlRig/Private/Units/Math/RigUnit_MathVector.h
--- a/Engine/Plugins/Animation/ControlRig/Source/ControlRig/Private/Units/Math/RigUnit_MathVector.h
+++ b/Engine/Plugins/Animation/ControlRig/Source/ControlRig/Private/Units/Math/RigUnit_MathVector.h
@@ -77,7 +77,7 @@ struct CONTROLRIG_API FRigUnit_MathVectorMake : public FRigUnit_MathVectorBase
 {
 	GENERATED_BODY()
 	RIGVM_METHOD()
-	virtual void Execute(const FRigUnitContext& Context) override;
+	virtual void Execute() override;
 
 	FRigUnit_MathVectorMake()
 	{
@@ -101,12 +101,12 @@ struct CONTROLRIG_API FRigUnit_MathVectorMake : public FRigUnit_MathVectorBase
 /**
  * Makes a vector from a single float
  */
-USTRUCT(meta=(DisplayName="From Float", TemplateName="Cast", Keywords="Make,Construct"))
+USTRUCT(meta=(DisplayName="From Float", TemplateName="Cast", ExecuteContext="FRigVMExecuteContext", Keywords="Make,Construct"))
 struct CONTROLRIG_API FRigUnit_MathVectorFromFloat : public FRigUnit_MathVectorBase
 {
 	GENERATED_BODY()
 	RIGVM_METHOD()
-	virtual void Execute(const FRigUnitContext& Context) override;
+	virtual void Execute() override;
 
 	FRigUnit_MathVectorFromFloat()
 	{
@@ -124,12 +124,12 @@ struct CONTROLRIG_API FRigUnit_MathVectorFromFloat : public FRigUnit_MathVectorB
 /**
  * Makes a vector from a single double
  */
-USTRUCT(meta=(DisplayName="From Float", TemplateName="Cast", Keywords="Make,Construct"))
+USTRUCT(meta=(DisplayName="From Float", TemplateName="Cast", ExecuteContext="FRigVMExecuteContext", Keywords="Make,Construct"))
 struct CONTROLRIG_API FRigUnit_MathVectorFromDouble : public FRigUnit_MathVectorBase
 {
 	GENERATED_BODY()
 	RIGVM_METHOD()
-	virtual void Execute(const FRigUnitContext& Context) override;
+	virtual void Execute() override;
 
 	FRigUnit_MathVectorFromDouble()
 	{
@@ -153,7 +153,7 @@ struct CONTROLRIG_API FRigUnit_MathVectorAdd : public FRigUnit_MathVectorBinaryA
 	GENERATED_BODY()
 
 	RIGVM_METHOD()
-	virtual void Execute(const FRigUnitContext& Context) override;
+	virtual void Execute() override;
 };
 
 /**
@@ -165,7 +165,7 @@ struct CONTROLRIG_API FRigUnit_MathVectorSub : public FRigUnit_MathVectorBinaryO
 	GENERATED_BODY()
 
 	RIGVM_METHOD()
-	virtual void Execute(const FRigUnitContext& Context) override;
+	virtual void Execute() override;
 };
 
 /**
@@ -182,7 +182,7 @@ struct CONTROLRIG_API FRigUnit_MathVectorMul : public FRigUnit_MathVectorBinaryA
 	}
 
 	RIGVM_METHOD()
-	virtual void Execute(const FRigUnitContext& Context) override;
+	virtual void Execute() override;
 };
 
 /**
@@ -200,7 +200,7 @@ struct CONTROLRIG_API FRigUnit_MathVectorScale : public FRigUnit_MathVectorBase
 	}
 
 	RIGVM_METHOD()
-	virtual void Execute(const FRigUnitContext& Context) override;
+	virtual void Execute() override;
 
 	UPROPERTY(meta = (Input))
 	FVector Value;
@@ -226,7 +226,7 @@ struct CONTROLRIG_API FRigUnit_MathVectorDiv : public FRigUnit_MathVectorBinaryO
 	}
 
 	RIGVM_METHOD()
-	virtual void Execute(const FRigUnitContext& Context) override;
+	virtual void Execute() override;
 };
 
 /**
@@ -244,7 +244,7 @@ struct CONTROLRIG_API FRigUnit_MathVectorMod : public FRigUnit_MathVectorBinaryO
 	}
 
 	RIGVM_METHOD()
-	virtual void Execute(const FRigUnitContext& Context) override;
+	virtual void Execute() override;
 };
 
 /**
@@ -256,7 +256,7 @@ struct CONTROLRIG_API FRigUnit_MathVectorMin : public FRigUnit_MathVectorBinaryA
 	GENERATED_BODY()
 
 	RIGVM_METHOD()
-	virtual void Execute(const FRigUnitContext& Context) override;
+	virtual void Execute() override;
 };
 
 /**
@@ -268,7 +268,7 @@ struct CONTROLRIG_API FRigUnit_MathVectorMax : public FRigUnit_MathVectorBinaryA
 	GENERATED_BODY()
 
 	RIGVM_METHOD()
-	virtual void Execute(const FRigUnitContext& Context) override;
+	virtual void Execute() override;
 };
 
 /**
@@ -280,7 +280,7 @@ struct CONTROLRIG_API FRigUnit_MathVectorNegate : public FRigUnit_MathVectorUnar
 	GENERATED_BODY()
 
 	RIGVM_METHOD()
-	virtual void Execute(const FRigUnitContext& Context) override;
+	virtual void Execute() override;
 };
 
 /**
@@ -292,7 +292,7 @@ struct CONTROLRIG_API FRigUnit_MathVectorAbs : public FRigUnit_MathVectorUnaryOp
 	GENERATED_BODY()
 
 	RIGVM_METHOD()
-	virtual void Execute(const FRigUnitContext& Context) override;
+	virtual void Execute() override;
 };
 
 /**
@@ -304,7 +304,7 @@ struct CONTROLRIG_API FRigUnit_MathVectorFloor : public FRigUnit_MathVectorUnary
 	GENERATED_BODY()
 
 	RIGVM_METHOD()
-	virtual void Execute(const FRigUnitContext& Context) override;
+	virtual void Execute() override;
 };
 
 /**
@@ -316,7 +316,7 @@ struct CONTROLRIG_API FRigUnit_MathVectorCeil : public FRigUnit_MathVectorUnaryO
 	GENERATED_BODY()
 
 	RIGVM_METHOD()
-	virtual void Execute(const FRigUnitContext& Context) override;
+	virtual void Execute() override;
 };
 
 /**
@@ -328,7 +328,7 @@ struct CONTROLRIG_API FRigUnit_MathVectorRound : public FRigUnit_MathVectorUnary
 	GENERATED_BODY()
 
 	RIGVM_METHOD()
-	virtual void Execute(const FRigUnitContext& Context) override;
+	virtual void Execute() override;
 };
 
 /**
@@ -340,7 +340,7 @@ struct CONTROLRIG_API FRigUnit_MathVectorSign : public FRigUnit_MathVectorUnaryO
 	GENERATED_BODY()
 
 	RIGVM_METHOD()
-	virtual void Execute(const FRigUnitContext& Context) override;
+	virtual void Execute() override;
 };
 
 /**
@@ -358,7 +358,7 @@ struct CONTROLRIG_API FRigUnit_MathVectorClamp : public FRigUnit_MathVectorBase
 	}
 
 	RIGVM_METHOD()
-	virtual void Execute(const FRigUnitContext& Context) override;
+	virtual void Execute() override;
 
 	UPROPERTY(meta=(Input))
 	FVector Value;
@@ -389,7 +389,7 @@ struct CONTROLRIG_API FRigUnit_MathVectorLerp : public FRigUnit_MathVectorBase
 	}
 
 	RIGVM_METHOD()
-	virtual void Execute(const FRigUnitContext& Context) override;
+	virtual void Execute() override;
 
 	UPROPERTY(meta=(Input))
 	FVector A;
@@ -420,7 +420,7 @@ struct CONTROLRIG_API FRigUnit_MathVectorRemap : public FRigUnit_MathVectorBase
 	}
 
 	RIGVM_METHOD()
-	virtual void Execute(const FRigUnitContext& Context) override;
+	virtual void Execute() override;
 
 	UPROPERTY(meta=(Input))
 	FVector Value;
@@ -460,7 +460,7 @@ struct CONTROLRIG_API FRigUnit_MathVectorEquals : public FRigUnit_MathVectorBase
 	}
 
 	RIGVM_METHOD()
-	virtual void Execute(const FRigUnitContext& Context) override;
+	virtual void Execute() override;
 
 	UPROPERTY(meta=(Input))
 	FVector A;
@@ -490,7 +490,7 @@ struct CONTROLRIG_API FRigUnit_MathVectorNotEquals : public FRigUnit_MathVectorB
 	}
 
 	RIGVM_METHOD()
-	virtual void Execute(const FRigUnitContext& Context) override;
+	virtual void Execute() override;
 
 	UPROPERTY(meta=(Input))
 	FVector A;
@@ -521,7 +521,7 @@ struct CONTROLRIG_API FRigUnit_MathVectorIsNearlyZero : public FRigUnit_MathVect
 	}
 
 	RIGVM_METHOD()
-	virtual void Execute(const FRigUnitContext& Context) override;
+	virtual void Execute() override;
 
 	UPROPERTY(meta=(Input))
 	FVector Value;
@@ -549,7 +549,7 @@ struct CONTROLRIG_API FRigUnit_MathVectorIsNearlyEqual : public FRigUnit_MathVec
 	}
 
 	RIGVM_METHOD()
-	virtual void Execute(const FRigUnitContext& Context) override;
+	virtual void Execute() override;
 
 	UPROPERTY(meta=(Input))
 	FVector A;
@@ -579,7 +579,7 @@ struct CONTROLRIG_API FRigUnit_MathVectorSelectBool : public FRigUnit_MathVector
 	}
 
 	RIGVM_METHOD()
-	virtual void Execute(const FRigUnitContext& Context) override;
+	virtual void Execute() override;
 
 	UPROPERTY(meta=(Input))
 	bool Condition;
@@ -606,7 +606,7 @@ struct CONTROLRIG_API FRigUnit_MathVectorDeg : public FRigUnit_MathVectorUnaryOp
 	GENERATED_BODY()
 
 	RIGVM_METHOD()
-	virtual void Execute(const FRigUnitContext& Context) override;
+	virtual void Execute() override;
 };
 
 /**
@@ -618,7 +618,7 @@ struct CONTROLRIG_API FRigUnit_MathVectorRad : public FRigUnit_MathVectorUnaryOp
 	GENERATED_BODY()
 
 	RIGVM_METHOD()
-	virtual void Execute(const FRigUnitContext& Context) override;
+	virtual void Execute() override;
 };
 
 /**
@@ -636,7 +636,7 @@ struct CONTROLRIG_API FRigUnit_MathVectorLengthSquared : public FRigUnit_MathVec
 	}
 
 	RIGVM_METHOD()
-	virtual void Execute(const FRigUnitContext& Context) override;
+	virtual void Execute() override;
 
 	UPROPERTY(meta=(Input))
 	FVector Value;
@@ -660,7 +660,7 @@ struct CONTROLRIG_API FRigUnit_MathVectorLength : public FRigUnit_MathVectorBase
 	}
 
 	RIGVM_METHOD()
-	virtual void Execute(const FRigUnitContext& Context) override;
+	virtual void Execute() override;
 
 	UPROPERTY(meta=(Input))
 	FVector Value;
@@ -684,7 +684,7 @@ struct CONTROLRIG_API FRigUnit_MathVectorDistance : public FRigUnit_MathVectorBa
 	}
 
 	RIGVM_METHOD()
-	virtual void Execute(const FRigUnitContext& Context) override;
+	virtual void Execute() override;
 
 	UPROPERTY(meta=(Input))
 	FVector A;
@@ -704,7 +704,7 @@ struct CONTROLRIG_API FRigUnit_MathVectorCross : public FRigUnit_MathVectorBinar
 {
 	GENERATED_BODY()
 	RIGVM_METHOD()
-	virtual void Execute(const FRigUnitContext& Context) override;
+	virtual void Execute() override;
 };
 
 /**
@@ -722,7 +722,7 @@ struct CONTROLRIG_API FRigUnit_MathVectorDot : public FRigUnit_MathVectorBase
 	}
 
 	RIGVM_METHOD()
-	virtual void Execute(const FRigUnitContext& Context) override;
+	virtual void Execute() override;
 
 	UPROPERTY(meta=(Input))
 	FVector A;
@@ -743,7 +743,7 @@ struct CONTROLRIG_API FRigUnit_MathVectorUnit : public FRigUnit_MathVectorUnaryO
 	GENERATED_BODY()
 
 	RIGVM_METHOD()
-	virtual void Execute(const FRigUnitContext& Context) override;
+	virtual void Execute() override;
 };
 
 /**
@@ -761,7 +761,7 @@ struct CONTROLRIG_API FRigUnit_MathVectorSetLength: public FRigUnit_MathVectorBa
 	}
 
 	RIGVM_METHOD()
-	virtual void Execute(const FRigUnitContext& Context) override;
+	virtual void Execute() override;
 
 	UPROPERTY(meta = (Input))
 	FVector Value;
@@ -790,7 +790,7 @@ struct CONTROLRIG_API FRigUnit_MathVectorClampLength: public FRigUnit_MathVector
 
 
 	RIGVM_METHOD()
-	virtual void Execute(const FRigUnitContext& Context) override;
+	virtual void Execute() override;
 
 	UPROPERTY(meta = (Input))
 	FVector Value;
@@ -820,7 +820,7 @@ struct CONTROLRIG_API FRigUnit_MathVectorMirror : public FRigUnit_MathVectorBase
 	}
 
 	RIGVM_METHOD()
-	virtual void Execute(const FRigUnitContext& Context) override;
+	virtual void Execute() override;
 
 	UPROPERTY(meta=(Input))
 	FVector Value;
@@ -847,7 +847,7 @@ struct CONTROLRIG_API FRigUnit_MathVectorAngle : public FRigUnit_MathVectorBase
 	}
 
 	RIGVM_METHOD()
-	virtual void Execute(const FRigUnitContext& Context) override;
+	virtual void Execute() override;
 
 	UPROPERTY(meta=(Input))
 	FVector A;
@@ -874,7 +874,7 @@ struct CONTROLRIG_API FRigUnit_MathVectorParallel : public FRigUnit_MathVectorBa
 	}
 
 	RIGVM_METHOD()
-	virtual void Execute(const FRigUnitContext& Context) override;
+	virtual void Execute() override;
 
 	UPROPERTY(meta=(Input))
 	FVector A;
@@ -902,7 +902,7 @@ struct CONTROLRIG_API FRigUnit_MathVectorOrthogonal : public FRigUnit_MathVector
 	}
 
 	RIGVM_METHOD()
-	virtual void Execute(const FRigUnitContext& Context) override;
+	virtual void Execute() override;
 
 	UPROPERTY(meta=(Input))
 	FVector A;
@@ -930,7 +930,7 @@ struct CONTROLRIG_API FRigUnit_MathVectorBezierFourPoint : public FRigUnit_MathV
 	}
 
 	RIGVM_METHOD()
-	virtual void Execute(const FRigUnitContext& Context) override;
+	virtual void Execute() override;
 
 	UPROPERTY(meta=(Input))
 	FCRFourPointBezier Bezier;
@@ -962,7 +962,7 @@ struct CONTROLRIG_API FRigUnit_MathVectorMakeBezierFourPoint : public FRigUnit_M
 	}
 
 	RIGVM_METHOD()
-	virtual void Execute(const FRigUnitContext& Context) override;
+	virtual void Execute() override;
 
 	UPROPERTY(meta = (Input, Output))
 	FCRFourPointBezier Bezier;
@@ -993,7 +993,7 @@ struct CONTROLRIG_API FRigUnit_MathVectorClampSpatially: public FRigUnit_MathVec
 	}
 
 	RIGVM_METHOD()
-	virtual void Execute(const FRigUnitContext& Context) override;
+	virtual void Execute() override;
 
 	UPROPERTY(meta = (Input))
 	FVector Value;
@@ -1044,7 +1044,7 @@ struct CONTROLRIG_API FRigUnit_MathIntersectPlane : public FRigUnit_MathVectorBa
 	}
 
 	RIGVM_METHOD()
-	virtual void Execute(const FRigUnitContext& Context) override;
+	virtual void Execute() override;
 
 	UPROPERTY(meta = (Input))
 	FVector Start;
@@ -1081,7 +1081,7 @@ struct CONTROLRIG_API FRigUnit_MathDistanceToPlane : public FRigUnit_MathVectorB
 	}
 
 	RIGVM_METHOD()
-	virtual void Execute(const FRigUnitContext& Context) override;
+	virtual void Execute() override;
 
 	UPROPERTY(meta = (Input))
 	FVector Point;
@@ -1113,7 +1113,7 @@ struct CONTROLRIG_API FRigUnit_MathVectorMakeRelative : public FRigUnit_MathVect
 	}
 
 	RIGVM_METHOD()
-	virtual void Execute(const FRigUnitContext& Context) override;
+	virtual void Execute() override;
 
 	UPROPERTY(meta=(Input))
 	FVector Global;
@@ -1139,7 +1139,7 @@ struct CONTROLRIG_API FRigUnit_MathVectorMakeAbsolute : public FRigUnit_MathVect
 	}
 
 	RIGVM_METHOD()
-	virtual void Execute(const FRigUnitContext& Context) override;
+	virtual void Execute() override;
 
 	UPROPERTY(meta = (Input))
 	FVector Local;
@@ -1168,7 +1168,7 @@ struct CONTROLRIG_API FRigUnit_MathVectorMirrorTransform : public FRigUnit_MathV
 	}
 
 	RIGVM_METHOD()
-	virtual void Execute(const FRigUnitContext& Context) override;
+	virtual void Execute() override;
 
 	UPROPERTY(meta=(Input))
 	FVector Value;
diff --git a/Engine/Plugins/Animation/ControlRig/Source/ControlRig/Private/Units/Math/RigUnit_Noise.cpp b/Engine/Plugins/Animation/ControlRig/Source/ControlRig/Private/Units/Math/RigUnit_Noise.cpp
--- a/Engine/Plugins/Animation/ControlRig/Source/ControlRig/Private/Units/Math/RigUnit_Noise.cpp
+++ b/Engine/Plugins/Animation/ControlRig/Source/ControlRig/Private/Units/Math/RigUnit_Noise.cpp
@@ -8,7 +8,7 @@
 FRigUnit_NoiseFloat_Execute()
 {
     DECLARE_SCOPE_HIERARCHICAL_COUNTER_RIGUNIT()
-	if (Context.State == EControlRigState::Init)
+	if (ExecuteContext.UnitContext.State == EControlRigState::Init)
 	{
 		Time = 0.f;
 		return;
@@ -16,13 +16,13 @@ FRigUnit_NoiseFloat_Execute()
 
 	const float Noise = FMath::PerlinNoise1D(Value * Frequency + Time) + 0.5f;
 	Result = FMath::Lerp<float>(Minimum, Maximum, Noise);
-	Time = Time + Speed * Context.DeltaTime;
+	Time = Time + Speed * ExecuteContext.UnitContext.DeltaTime;
 }
 
 FRigUnit_NoiseDouble_Execute()
 {
 	DECLARE_SCOPE_HIERARCHICAL_COUNTER_RIGUNIT()
-	if (Context.State == EControlRigState::Init)
+	if (ExecuteContext.UnitContext.State == EControlRigState::Init)
 	{
 		Time = 0.f;
 		return;
@@ -30,12 +30,12 @@ FRigUnit_NoiseDouble_Execute()
 
 	const double Noise = double(FMath::PerlinNoise1D(Value * Frequency + Time)) + 0.5;
 	Result = FMath::Lerp<double>(Minimum, Maximum, Noise);
-	Time = Time + Speed * Context.DeltaTime;
+	Time = Time + Speed * ExecuteContext.UnitContext.DeltaTime;
 }
 
 FRigUnit_NoiseVector_Execute()
 {
-	FRigUnit_NoiseVector2::StaticExecute(ExecuteContext, Position, Speed, Frequency, Minimum, Maximum, Result, Time, Context);
+	FRigUnit_NoiseVector2::StaticExecute(ExecuteContext, Position, Speed, Frequency, Minimum, Maximum, Result, Time);
 }
 
 FRigVMStructUpgradeInfo FRigUnit_NoiseVector::GetUpgradeInfo() const
@@ -55,7 +55,7 @@ FRigVMStructUpgradeInfo FRigUnit_NoiseVector::GetUpgradeInfo() const
 FRigUnit_NoiseVector2_Execute()
 {
 	DECLARE_SCOPE_HIERARCHICAL_COUNTER_RIGUNIT()
-	if (Context.State == EControlRigState::Init)
+	if (ExecuteContext.UnitContext.State == EControlRigState::Init)
 	{
 		Time = FVector::ZeroVector;
 		return;
@@ -67,6 +67,6 @@ FRigUnit_NoiseVector2_Execute()
 	Result.X = FMath::Lerp<double>(Minimum, Maximum, NoiseX);
 	Result.Y = FMath::Lerp<double>(Minimum, Maximum, NoiseY);
 	Result.Z = FMath::Lerp<double>(Minimum, Maximum, NoiseZ);
-	Time = Time + Speed * Context.DeltaTime;
+	Time = Time + Speed * ExecuteContext.UnitContext.DeltaTime;
 }
 
diff --git a/Engine/Plugins/Animation/ControlRig/Source/ControlRig/Private/Units/Math/RigUnit_Random.cpp b/Engine/Plugins/Animation/ControlRig/Source/ControlRig/Private/Units/Math/RigUnit_Random.cpp
--- a/Engine/Plugins/Animation/ControlRig/Source/ControlRig/Private/Units/Math/RigUnit_Random.cpp
+++ b/Engine/Plugins/Animation/ControlRig/Source/ControlRig/Private/Units/Math/RigUnit_Random.cpp
@@ -20,7 +20,7 @@ float FRigUnit_Random_Helper(int32& Seed)
 FRigUnit_RandomFloat_Execute()
 {
     DECLARE_SCOPE_HIERARCHICAL_COUNTER_RIGUNIT()
-	if (Context.State == EControlRigState::Init)
+	if (ExecuteContext.UnitContext.State == EControlRigState::Init)
 	{
 		LastSeed = BaseSeed = Seed;
 		TimeLeft = 0.f;
@@ -32,7 +32,7 @@ FRigUnit_RandomFloat_Execute()
 		LastSeed = BaseSeed = Seed;
 	}
 
-	TimeLeft = TimeLeft - Context.DeltaTime;
+	TimeLeft = TimeLeft - ExecuteContext.UnitContext.DeltaTime;
 	if (TimeLeft > 0.f)
 	{
 		Result = LastResult;
@@ -48,7 +48,7 @@ FRigUnit_RandomFloat_Execute()
 FRigUnit_RandomVector_Execute()
 {
     DECLARE_SCOPE_HIERARCHICAL_COUNTER_RIGUNIT()
-	if (Context.State == EControlRigState::Init)
+	if (ExecuteContext.UnitContext.State == EControlRigState::Init)
 	{
 		LastSeed = BaseSeed = Seed;
 		TimeLeft = 0.f;
@@ -60,7 +60,7 @@ FRigUnit_RandomVector_Execute()
 		LastSeed = BaseSeed = Seed;
 	}
 
-	TimeLeft = TimeLeft - Context.DeltaTime;
+	TimeLeft = TimeLeft - ExecuteContext.UnitContext.DeltaTime;
 	if (TimeLeft > 0.f)
 	{
 		Result = LastResult;
diff --git a/Engine/Plugins/Animation/ControlRig/Source/ControlRig/Private/Units/Simulation/RigUnit_Accumulate.cpp b/Engine/Plugins/Animation/ControlRig/Source/ControlRig/Private/Units/Simulation/RigUnit_Accumulate.cpp
--- a/Engine/Plugins/Animation/ControlRig/Source/ControlRig/Private/Units/Simulation/RigUnit_Accumulate.cpp
+++ b/Engine/Plugins/Animation/ControlRig/Source/ControlRig/Private/Units/Simulation/RigUnit_Accumulate.cpp
@@ -8,13 +8,13 @@
 FRigUnit_AccumulateFloatAdd_Execute()
 {
     DECLARE_SCOPE_HIERARCHICAL_COUNTER_RIGUNIT()
-	if (Context.State == EControlRigState::Init)
+	if (ExecuteContext.UnitContext.State == EControlRigState::Init)
 	{
 		Result = AccumulatedValue = InitialValue;
 	}
 	else if(bIntegrateDeltaTime)
 	{
-		Result = AccumulatedValue = AccumulatedValue + Increment * Context.DeltaTime;
+		Result = AccumulatedValue = AccumulatedValue + Increment * ExecuteContext.UnitContext.DeltaTime;
 	}
 	else
 	{
@@ -25,13 +25,13 @@ FRigUnit_AccumulateFloatAdd_Execute()
 FRigUnit_AccumulateVectorAdd_Execute()
 {
     DECLARE_SCOPE_HIERARCHICAL_COUNTER_RIGUNIT()
-	if (Context.State == EControlRigState::Init)
+	if (ExecuteContext.UnitContext.State == EControlRigState::Init)
 	{
 		Result = AccumulatedValue = InitialValue;
 	}
 	else if (bIntegrateDeltaTime)
 	{
-		Result = AccumulatedValue = AccumulatedValue + Increment * Context.DeltaTime;
+		Result = AccumulatedValue = AccumulatedValue + Increment * ExecuteContext.UnitContext.DeltaTime;
 	}
 	else
 	{
@@ -42,7 +42,7 @@ FRigUnit_AccumulateVectorAdd_Execute()
 FRigUnit_AccumulateFloatMul_Execute()
 {
     DECLARE_SCOPE_HIERARCHICAL_COUNTER_RIGUNIT()
-	if (Context.State == EControlRigState::Init)
+	if (ExecuteContext.UnitContext.State == EControlRigState::Init)
 	{
 		Result = AccumulatedValue = InitialValue;
 	}
@@ -51,7 +51,7 @@ FRigUnit_AccumulateFloatMul_Execute()
 		float Factor = Multiplier;
 		if (bIntegrateDeltaTime)
 		{
-			Factor = FMath::Lerp<float>(1.f, Factor, Context.DeltaTime);
+			Factor = FMath::Lerp<float>(1.f, Factor, ExecuteContext.UnitContext.DeltaTime);
 		}
 		Result = AccumulatedValue = AccumulatedValue * Factor;
 	}
@@ -60,7 +60,7 @@ FRigUnit_AccumulateFloatMul_Execute()
 FRigUnit_AccumulateVectorMul_Execute()
 {
     DECLARE_SCOPE_HIERARCHICAL_COUNTER_RIGUNIT()
-	if (Context.State == EControlRigState::Init)
+	if (ExecuteContext.UnitContext.State == EControlRigState::Init)
 	{
 		Result = AccumulatedValue = InitialValue;
 	}
@@ -69,7 +69,7 @@ FRigUnit_AccumulateVectorMul_Execute()
 		FVector Factor = Multiplier;
 		if (bIntegrateDeltaTime)
 		{
-			Factor = FMath::Lerp<FVector>(FVector::OneVector, Factor, Context.DeltaTime);
+			Factor = FMath::Lerp<FVector>(FVector::OneVector, Factor, ExecuteContext.UnitContext.DeltaTime);
 		}
 		Result = AccumulatedValue = AccumulatedValue * Factor;
 	}
@@ -78,7 +78,7 @@ FRigUnit_AccumulateVectorMul_Execute()
 FRigUnit_AccumulateQuatMul_Execute()
 {
     DECLARE_SCOPE_HIERARCHICAL_COUNTER_RIGUNIT()
-	if (Context.State == EControlRigState::Init)
+	if (ExecuteContext.UnitContext.State == EControlRigState::Init)
 	{
 		Result = AccumulatedValue = InitialValue;
 	}
@@ -87,7 +87,7 @@ FRigUnit_AccumulateQuatMul_Execute()
 		FQuat Factor = Multiplier;
 		if (bIntegrateDeltaTime)
 		{
-			Factor = FQuat::Slerp(FQuat::Identity, Factor, Context.DeltaTime);
+			Factor = FQuat::Slerp(FQuat::Identity, Factor, ExecuteContext.UnitContext.DeltaTime);
 		}
 		if (bFlipOrder)
 		{
@@ -103,7 +103,7 @@ FRigUnit_AccumulateQuatMul_Execute()
 FRigUnit_AccumulateTransformMul_Execute()
 {
     DECLARE_SCOPE_HIERARCHICAL_COUNTER_RIGUNIT()
-	if (Context.State == EControlRigState::Init)
+	if (ExecuteContext.UnitContext.State == EControlRigState::Init)
 	{
 		Result = AccumulatedValue = InitialValue;
 	}
@@ -112,9 +112,9 @@ FRigUnit_AccumulateTransformMul_Execute()
 		FTransform Factor = Multiplier;
 		if (bIntegrateDeltaTime)
 		{
-			Factor.SetTranslation(FMath::Lerp<FVector>(FVector::OneVector, Factor.GetTranslation(), Context.DeltaTime));
-			Factor.SetRotation(FQuat::Slerp(FQuat::Identity, Factor.GetRotation(), Context.DeltaTime));
-			Factor.SetScale3D(FMath::Lerp<FVector>(FVector::OneVector, Factor.GetScale3D(), Context.DeltaTime));
+			Factor.SetTranslation(FMath::Lerp<FVector>(FVector::OneVector, Factor.GetTranslation(), ExecuteContext.UnitContext.DeltaTime));
+			Factor.SetRotation(FQuat::Slerp(FQuat::Identity, Factor.GetRotation(), ExecuteContext.UnitContext.DeltaTime));
+			Factor.SetScale3D(FMath::Lerp<FVector>(FVector::OneVector, Factor.GetScale3D(), ExecuteContext.UnitContext.DeltaTime));
 		}
 		if (bFlipOrder)
 		{
@@ -130,52 +130,52 @@ FRigUnit_AccumulateTransformMul_Execute()
 FRigUnit_AccumulateFloatLerp_Execute()
 {
     DECLARE_SCOPE_HIERARCHICAL_COUNTER_RIGUNIT()
-	if (Context.State == EControlRigState::Init)
+	if (ExecuteContext.UnitContext.State == EControlRigState::Init)
 	{
 		Result = AccumulatedValue = InitialValue;
 	}
 	else
 	{
-		Result = AccumulatedValue = FMath::Lerp<float>(AccumulatedValue, TargetValue, FMath::Clamp<float>(bIntegrateDeltaTime ? Context.DeltaTime * Blend : Blend, 0.f, 1.f));
+		Result = AccumulatedValue = FMath::Lerp<float>(AccumulatedValue, TargetValue, FMath::Clamp<float>(bIntegrateDeltaTime ? ExecuteContext.UnitContext.DeltaTime * Blend : Blend, 0.f, 1.f));
 	}
 }
 
 FRigUnit_AccumulateVectorLerp_Execute()
 {
     DECLARE_SCOPE_HIERARCHICAL_COUNTER_RIGUNIT()
-	if (Context.State == EControlRigState::Init)
+	if (ExecuteContext.UnitContext.State == EControlRigState::Init)
 	{
 		Result = AccumulatedValue = InitialValue;
 	}
 	else
 	{
-		Result = AccumulatedValue = FMath::Lerp<FVector>(AccumulatedValue, TargetValue, FMath::Clamp<float>(bIntegrateDeltaTime ? Context.DeltaTime * Blend : Blend, 0.f, 1.f));
+		Result = AccumulatedValue = FMath::Lerp<FVector>(AccumulatedValue, TargetValue, FMath::Clamp<float>(bIntegrateDeltaTime ? ExecuteContext.UnitContext.DeltaTime * Blend : Blend, 0.f, 1.f));
 	}
 }
 
 FRigUnit_AccumulateQuatLerp_Execute()
 {
     DECLARE_SCOPE_HIERARCHICAL_COUNTER_RIGUNIT()
-	if (Context.State == EControlRigState::Init)
+	if (ExecuteContext.UnitContext.State == EControlRigState::Init)
 	{
 		Result = AccumulatedValue = InitialValue;
 	}
 	else
 	{
-		Result = AccumulatedValue = FQuat::Slerp(AccumulatedValue, TargetValue, FMath::Clamp<float>(bIntegrateDeltaTime ? Context.DeltaTime * Blend : Blend, 0.f, 1.f));
+		Result = AccumulatedValue = FQuat::Slerp(AccumulatedValue, TargetValue, FMath::Clamp<float>(bIntegrateDeltaTime ? ExecuteContext.UnitContext.DeltaTime * Blend : Blend, 0.f, 1.f));
 	}
 }
 
 FRigUnit_AccumulateTransformLerp_Execute()
 {
     DECLARE_SCOPE_HIERARCHICAL_COUNTER_RIGUNIT()
-	if (Context.State == EControlRigState::Init)
+	if (ExecuteContext.UnitContext.State == EControlRigState::Init)
 	{
 		Result = AccumulatedValue = InitialValue;
 	}
 	else
 	{
-		float B = FMath::Clamp<float>(bIntegrateDeltaTime ? Context.DeltaTime * Blend : Blend, 0.f, 1.f);
+		float B = FMath::Clamp<float>(bIntegrateDeltaTime ? ExecuteContext.UnitContext.DeltaTime * Blend : Blend, 0.f, 1.f);
 		AccumulatedValue.SetTranslation(FMath::Lerp<FVector>(AccumulatedValue.GetTranslation(), TargetValue.GetTranslation(), B));
 		AccumulatedValue.SetRotation(FQuat::Slerp(AccumulatedValue.GetRotation(), TargetValue.GetRotation(), B));
 		AccumulatedValue.SetScale3D(FMath::Lerp<FVector>(AccumulatedValue.GetScale3D(), TargetValue.GetScale3D(), B));
@@ -186,7 +186,7 @@ FRigUnit_AccumulateTransformLerp_Execute()
 FRigUnit_AccumulateFloatRange_Execute()
 {
     DECLARE_SCOPE_HIERARCHICAL_COUNTER_RIGUNIT()
-	if (Context.State == EControlRigState::Init)
+	if (ExecuteContext.UnitContext.State == EControlRigState::Init)
 	{
 		Minimum = AccumulatedMinimum = Maximum = AccumulatedMaximum = Value;
 	}
@@ -200,7 +200,7 @@ FRigUnit_AccumulateFloatRange_Execute()
 FRigUnit_AccumulateVectorRange_Execute()
 {
     DECLARE_SCOPE_HIERARCHICAL_COUNTER_RIGUNIT()
-	if (Context.State == EControlRigState::Init)
+	if (ExecuteContext.UnitContext.State == EControlRigState::Init)
 	{
 		Minimum = AccumulatedMinimum = Maximum = AccumulatedMaximum = Value;
 	}
@@ -220,7 +220,7 @@ FRigUnit_AccumulateVectorRange_Execute()
 
 IMPLEMENT_RIGUNIT_AUTOMATION_TEST(FRigUnit_AccumulateFloatAdd)
 {
-	Context.DeltaTime = 0.5f;
+	ExecuteContext.UnitContext.DeltaTime = 0.5f;
 
 	Unit.Increment = 1.f;
 	InitAndExecute();
@@ -240,7 +240,7 @@ IMPLEMENT_RIGUNIT_AUTOMATION_TEST(FRigUnit_AccumulateFloatAdd)
 IMPLEMENT_RIGUNIT_AUTOMATION_TEST(FRigUnit_AccumulateVectorAdd)
 {
 	using TFloat = decltype(Unit.Result.X);
-	Context.DeltaTime = 0.5f;
+	ExecuteContext.UnitContext.DeltaTime = 0.5f;
 
 	Unit.Increment = FVector(1.0f, 0.f, 0.f);
 	InitAndExecute();
diff --git a/Engine/Plugins/Animation/ControlRig/Source/ControlRig/Private/Units/Simulation/RigUnit_Accumulate.h b/Engine/Plugins/Animation/ControlRig/Source/ControlRig/Private/Units/Simulation/RigUnit_Accumulate.h
--- a/Engine/Plugins/Animation/ControlRig/Source/ControlRig/Private/Units/Simulation/RigUnit_Accumulate.h
+++ b/Engine/Plugins/Animation/ControlRig/Source/ControlRig/Private/Units/Simulation/RigUnit_Accumulate.h
@@ -26,7 +26,7 @@ struct CONTROLRIG_API FRigUnit_AccumulateFloatAdd : public FRigUnit_AccumulateBa
 	}
 
 	RIGVM_METHOD()
-	virtual void Execute(const FRigUnitContext& Context) override;
+	virtual void Execute() override;
 
 	UPROPERTY(meta=(Input))
 	float Increment;
@@ -59,7 +59,7 @@ struct CONTROLRIG_API FRigUnit_AccumulateVectorAdd : public FRigUnit_AccumulateB
 	}
 
 	RIGVM_METHOD()
-	virtual void Execute(const FRigUnitContext& Context) override;
+	virtual void Execute() override;
 
 	UPROPERTY(meta = (Input))
 	FVector Increment;
@@ -92,7 +92,7 @@ struct CONTROLRIG_API FRigUnit_AccumulateFloatMul : public FRigUnit_AccumulateBa
 	}
 
 	RIGVM_METHOD()
-	virtual void Execute(const FRigUnitContext& Context) override;
+	virtual void Execute() override;
 
 	UPROPERTY(meta=(Input))
 	float Multiplier;
@@ -125,7 +125,7 @@ struct CONTROLRIG_API FRigUnit_AccumulateVectorMul : public FRigUnit_AccumulateB
 	}
 
 	RIGVM_METHOD()
-	virtual void Execute(const FRigUnitContext& Context) override;
+	virtual void Execute() override;
 
 	UPROPERTY(meta = (Input))
 	FVector Multiplier;
@@ -158,7 +158,7 @@ struct CONTROLRIG_API FRigUnit_AccumulateQuatMul : public FRigUnit_AccumulateBas
 	}
 
 	RIGVM_METHOD()
-	virtual void Execute(const FRigUnitContext& Context) override;
+	virtual void Execute() override;
 
 	UPROPERTY(meta=(Input))
 	FQuat Multiplier;
@@ -194,7 +194,7 @@ struct CONTROLRIG_API FRigUnit_AccumulateTransformMul : public FRigUnit_Accumula
 	}
 
 	RIGVM_METHOD()
-	virtual void Execute(const FRigUnitContext& Context) override;
+	virtual void Execute() override;
 
 	UPROPERTY(meta = (Input))
 	FTransform Multiplier;
@@ -230,7 +230,7 @@ struct CONTROLRIG_API FRigUnit_AccumulateFloatLerp : public FRigUnit_AccumulateB
 	}
 
 	RIGVM_METHOD()
-	virtual void Execute(const FRigUnitContext& Context) override;
+	virtual void Execute() override;
 
 	UPROPERTY(meta=(Input))
 	float TargetValue;
@@ -267,7 +267,7 @@ struct CONTROLRIG_API FRigUnit_AccumulateVectorLerp : public FRigUnit_Accumulate
 	}
 
 	RIGVM_METHOD()
-	virtual void Execute(const FRigUnitContext& Context) override;
+	virtual void Execute() override;
 
 	UPROPERTY(meta = (Input))
 	FVector TargetValue;
@@ -304,7 +304,7 @@ struct CONTROLRIG_API FRigUnit_AccumulateQuatLerp : public FRigUnit_AccumulateBa
 	}
 
 	RIGVM_METHOD()
-	virtual void Execute(const FRigUnitContext& Context) override;
+	virtual void Execute() override;
 
 	UPROPERTY(meta=(Input))
 	FQuat TargetValue;
@@ -341,7 +341,7 @@ struct CONTROLRIG_API FRigUnit_AccumulateTransformLerp : public FRigUnit_Accumul
 	}
 
 	RIGVM_METHOD()
-	virtual void Execute(const FRigUnitContext& Context) override;
+	virtual void Execute() override;
 
 	UPROPERTY(meta = (Input))
 	FTransform TargetValue;
@@ -376,7 +376,7 @@ struct CONTROLRIG_API FRigUnit_AccumulateFloatRange : public FRigUnit_Accumulate
 	}
 
 	RIGVM_METHOD()
-	virtual void Execute(const FRigUnitContext& Context) override;
+	virtual void Execute() override;
 
 	UPROPERTY(meta = (Input))
 	float Value;
@@ -408,7 +408,7 @@ struct CONTROLRIG_API FRigUnit_AccumulateVectorRange : public FRigUnit_Accumulat
 	}
 
 	RIGVM_METHOD()
-	virtual void Execute(const FRigUnitContext& Context) override;
+	virtual void Execute() override;
 
 	UPROPERTY(meta = (Input))
 	FVector Value;
diff --git a/Engine/Plugins/Animation/ControlRig/Source/ControlRig/Private/Units/Simulation/RigUnit_AlphaInterp.cpp b/Engine/Plugins/Animation/ControlRig/Source/ControlRig/Private/Units/Simulation/RigUnit_AlphaInterp.cpp
--- a/Engine/Plugins/Animation/ControlRig/Source/ControlRig/Private/Units/Simulation/RigUnit_AlphaInterp.cpp
+++ b/Engine/Plugins/Animation/ControlRig/Source/ControlRig/Private/Units/Simulation/RigUnit_AlphaInterp.cpp
@@ -13,7 +13,7 @@ FRigUnit_AlphaInterp_Execute()
 	ScaleBiasClamp.bClampResult = bClampResult;
 	ScaleBiasClamp.bInterpResult = bInterpResult;
 
-	if (Context.State == EControlRigState::Init)
+	if (ExecuteContext.UnitContext.State == EControlRigState::Init)
 	{
 		ScaleBiasClamp.Reinitialize();
 	}
@@ -28,7 +28,7 @@ FRigUnit_AlphaInterp_Execute()
 		ScaleBiasClamp.InterpSpeedIncreasing = InterpSpeedIncreasing;
 		ScaleBiasClamp.InterpSpeedDecreasing = InterpSpeedDecreasing;
 
-		Result = ScaleBiasClamp.ApplyTo(Value, Context.DeltaTime);
+		Result = ScaleBiasClamp.ApplyTo(Value, ExecuteContext.UnitContext.DeltaTime);
 	}
 }
 
@@ -63,7 +63,7 @@ FRigUnit_AlphaInterpVector_Execute()
 	ScaleBiasClamp.bClampResult = bClampResult;
 	ScaleBiasClamp.bInterpResult = bInterpResult;
 
-	if (Context.State == EControlRigState::Init)
+	if (ExecuteContext.UnitContext.State == EControlRigState::Init)
 	{
 		ScaleBiasClamp.Reinitialize();
 	}
@@ -78,9 +78,9 @@ FRigUnit_AlphaInterpVector_Execute()
 		ScaleBiasClamp.InterpSpeedIncreasing = InterpSpeedIncreasing;
 		ScaleBiasClamp.InterpSpeedDecreasing = InterpSpeedDecreasing;
 
-		Result.X = ScaleBiasClamp.ApplyTo(Value.X, Context.DeltaTime);
-		Result.Y = ScaleBiasClamp.ApplyTo(Value.Y, Context.DeltaTime);
-		Result.Z = ScaleBiasClamp.ApplyTo(Value.Z, Context.DeltaTime);
+		Result.X = ScaleBiasClamp.ApplyTo(Value.X, ExecuteContext.UnitContext.DeltaTime);
+		Result.Y = ScaleBiasClamp.ApplyTo(Value.Y, ExecuteContext.UnitContext.DeltaTime);
+		Result.Z = ScaleBiasClamp.ApplyTo(Value.Z, ExecuteContext.UnitContext.DeltaTime);
 	}
 }
 
@@ -115,7 +115,7 @@ FRigUnit_AlphaInterpQuat_Execute()
 	ScaleBiasClamp.bClampResult = bClampResult;
 	ScaleBiasClamp.bInterpResult = bInterpResult;
 
-	if (Context.State == EControlRigState::Init)
+	if (ExecuteContext.UnitContext.State == EControlRigState::Init)
 	{
 		ScaleBiasClamp.Reinitialize();
 	}
@@ -130,7 +130,7 @@ FRigUnit_AlphaInterpQuat_Execute()
 		ScaleBiasClamp.InterpSpeedIncreasing = InterpSpeedIncreasing;
 		ScaleBiasClamp.InterpSpeedDecreasing = InterpSpeedDecreasing;
 
-		const float T = ScaleBiasClamp.ApplyTo(1.f, Context.DeltaTime);
+		const float T = ScaleBiasClamp.ApplyTo(1.f, ExecuteContext.UnitContext.DeltaTime);
 		Result = FQuat::Slerp(FQuat::Identity, Value, T);
 	}
 }
diff --git a/Engine/Plugins/Animation/ControlRig/Source/ControlRig/Private/Units/Simulation/RigUnit_DeltaFromPrevious.cpp b/Engine/Plugins/Animation/ControlRig/Source/ControlRig/Private/Units/Simulation/RigUnit_DeltaFromPrevious.cpp
--- a/Engine/Plugins/Animation/ControlRig/Source/ControlRig/Private/Units/Simulation/RigUnit_DeltaFromPrevious.cpp
+++ b/Engine/Plugins/Animation/ControlRig/Source/ControlRig/Private/Units/Simulation/RigUnit_DeltaFromPrevious.cpp
@@ -8,7 +8,7 @@
 FRigUnit_DeltaFromPreviousFloat_Execute()
 {
     DECLARE_SCOPE_HIERARCHICAL_COUNTER_RIGUNIT()
-	if (Context.State == EControlRigState::Init)
+	if (ExecuteContext.UnitContext.State == EControlRigState::Init)
 	{
 		Delta = 0.f;
 		PreviousValue = Cache = Value;
@@ -18,7 +18,7 @@ FRigUnit_DeltaFromPreviousFloat_Execute()
 	PreviousValue = Cache;
 	Delta = Cache - Value;
 
-	if (FMath::Abs(Context.DeltaTime) > SMALL_NUMBER)
+	if (FMath::Abs(ExecuteContext.UnitContext.DeltaTime) > SMALL_NUMBER)
 	{
 		Cache = Value;
 	}
@@ -27,7 +27,7 @@ FRigUnit_DeltaFromPreviousFloat_Execute()
 FRigUnit_DeltaFromPreviousVector_Execute()
 {
     DECLARE_SCOPE_HIERARCHICAL_COUNTER_RIGUNIT()
-	if (Context.State == EControlRigState::Init)
+	if (ExecuteContext.UnitContext.State == EControlRigState::Init)
 	{
 		Delta = FVector::ZeroVector;
 		PreviousValue = Cache = Value;
@@ -37,7 +37,7 @@ FRigUnit_DeltaFromPreviousVector_Execute()
 	PreviousValue = Cache;
 	Delta = Cache - Value;
 
-	if (FMath::Abs(Context.DeltaTime) > SMALL_NUMBER)
+	if (FMath::Abs(ExecuteContext.UnitContext.DeltaTime) > SMALL_NUMBER)
 	{
 		Cache = Value;
 	}
@@ -46,7 +46,7 @@ FRigUnit_DeltaFromPreviousVector_Execute()
 FRigUnit_DeltaFromPreviousQuat_Execute()
 {
     DECLARE_SCOPE_HIERARCHICAL_COUNTER_RIGUNIT()
-	if (Context.State == EControlRigState::Init)
+	if (ExecuteContext.UnitContext.State == EControlRigState::Init)
 	{
 		Delta = FQuat::Identity;
 		PreviousValue = Cache = Value;
@@ -56,7 +56,7 @@ FRigUnit_DeltaFromPreviousQuat_Execute()
 	PreviousValue = Cache;
 	Delta = Cache.Inverse() * Value;
 
-	if (FMath::Abs(Context.DeltaTime) > SMALL_NUMBER)
+	if (FMath::Abs(ExecuteContext.UnitContext.DeltaTime) > SMALL_NUMBER)
 	{
 		Cache = Value;
 	}
@@ -65,7 +65,7 @@ FRigUnit_DeltaFromPreviousQuat_Execute()
 FRigUnit_DeltaFromPreviousTransform_Execute()
 {
     DECLARE_SCOPE_HIERARCHICAL_COUNTER_RIGUNIT()
-	if (Context.State == EControlRigState::Init)
+	if (ExecuteContext.UnitContext.State == EControlRigState::Init)
 	{
 		Delta = FTransform::Identity;
 		PreviousValue = Cache = Value;
@@ -75,7 +75,7 @@ FRigUnit_DeltaFromPreviousTransform_Execute()
 	PreviousValue = Cache;
 	Delta = Value.GetRelativeTransform(Cache);
 
-	if (FMath::Abs(Context.DeltaTime) > SMALL_NUMBER)
+	if (FMath::Abs(ExecuteContext.UnitContext.DeltaTime) > SMALL_NUMBER)
 	{
 		Cache = Value;
 	}
@@ -86,7 +86,7 @@ FRigUnit_DeltaFromPreviousTransform_Execute()
 
 IMPLEMENT_RIGUNIT_AUTOMATION_TEST(FRigUnit_DeltaFromPreviousFloat)
 {
-	Context.DeltaTime = 0.1f;
+	ExecuteContext.UnitContext.DeltaTime = 0.1f;
 	Unit.Value = 1.f;
 	InitAndExecute();
 	AddErrorIfFalse(FMath::IsNearlyEqual(Unit.Delta, 0.f), TEXT("unexpected average result"));
diff --git a/Engine/Plugins/Animation/ControlRig/Source/ControlRig/Private/Units/Simulation/RigUnit_Kalman.cpp b/Engine/Plugins/Animation/ControlRig/Source/ControlRig/Private/Units/Simulation/RigUnit_Kalman.cpp
--- a/Engine/Plugins/Animation/ControlRig/Source/ControlRig/Private/Units/Simulation/RigUnit_Kalman.cpp
+++ b/Engine/Plugins/Animation/ControlRig/Source/ControlRig/Private/Units/Simulation/RigUnit_Kalman.cpp
@@ -10,7 +10,7 @@ FRigUnit_KalmanFloat_Execute()
     DECLARE_SCOPE_HIERARCHICAL_COUNTER_RIGUNIT()
 	if(BufferSize <= 0)
 	{
-		if (Context.State == EControlRigState::Init)
+		if (ExecuteContext.UnitContext.State == EControlRigState::Init)
 		{
 			UE_CONTROLRIG_RIGUNIT_REPORT_WARNING(TEXT("BufferSize is too small."));
 		}
@@ -21,7 +21,7 @@ FRigUnit_KalmanFloat_Execute()
 
 	Buffer.SetNum(BufferSize);
 
-	if (Context.State == EControlRigState::Init)
+	if (ExecuteContext.UnitContext.State == EControlRigState::Init)
 	{
 		for (float& Element : Buffer)
 		{
@@ -61,7 +61,7 @@ FRigUnit_KalmanVector_Execute()
     DECLARE_SCOPE_HIERARCHICAL_COUNTER_RIGUNIT()
 	if(BufferSize <= 0)
 	{
-		if (Context.State == EControlRigState::Init)
+		if (ExecuteContext.UnitContext.State == EControlRigState::Init)
 		{
 			UE_CONTROLRIG_RIGUNIT_REPORT_WARNING(TEXT("BufferSize is too small."));
 		}
@@ -72,7 +72,7 @@ FRigUnit_KalmanVector_Execute()
 
 	Buffer.SetNum(BufferSize);
 
-	if (Context.State == EControlRigState::Init)
+	if (ExecuteContext.UnitContext.State == EControlRigState::Init)
 	{
 		for (FVector& Element : Buffer)
 		{
@@ -112,7 +112,7 @@ FRigUnit_KalmanTransform_Execute()
     DECLARE_SCOPE_HIERARCHICAL_COUNTER_RIGUNIT()
 	if(BufferSize <= 0)
 	{
-		if (Context.State == EControlRigState::Init)
+		if (ExecuteContext.UnitContext.State == EControlRigState::Init)
 		{
 			UE_CONTROLRIG_RIGUNIT_REPORT_WARNING(TEXT("BufferSize is too small."));
 		}
@@ -123,7 +123,7 @@ FRigUnit_KalmanTransform_Execute()
 
 	Buffer.SetNum(BufferSize);
 
-	if (Context.State == EControlRigState::Init)
+	if (ExecuteContext.UnitContext.State == EControlRigState::Init)
 	{
 		for (FTransform& Element : Buffer)
 		{
diff --git a/Engine/Plugins/Animation/ControlRig/Source/ControlRig/Private/Units/Simulation/RigUnit_PointSimulation.cpp b/Engine/Plugins/Animation/ControlRig/Source/ControlRig/Private/Units/Simulation/RigUnit_PointSimulation.cpp
--- a/Engine/Plugins/Animation/ControlRig/Source/ControlRig/Private/Units/Simulation/RigUnit_PointSimulation.cpp
+++ b/Engine/Plugins/Animation/ControlRig/Source/ControlRig/Private/Units/Simulation/RigUnit_PointSimulation.cpp
@@ -10,12 +10,12 @@ FRigUnit_PointSimulation_Execute()
     DECLARE_SCOPE_HIERARCHICAL_COUNTER_RIGUNIT()
 	URigHierarchy* Hierarchy = ExecuteContext.Hierarchy;
 
-	float DeltaTime = Context.DeltaTime;
+	float DeltaTime = ExecuteContext.UnitContext.DeltaTime;
 
 	FCRSimPointContainer& Simulation = WorkData.Simulation;
 	TArray<FCachedRigElement>& BoneIndices = WorkData.BoneIndices;
 
-	if (Context.State == EControlRigState::Init ||
+	if (ExecuteContext.UnitContext.State == EControlRigState::Init ||
 		Simulation.Points.Num() != Points.Num() ||
 		Simulation.Springs.Num() != Links.Num() ||
 		Simulation.Forces.Num() != Forces.Num() ||
@@ -251,9 +251,9 @@ FRigUnit_PointSimulation_Execute()
 		Bezier.D = Simulation.GetPointInterpolated(3).Position;
 	}
 
-	if (Context.DrawInterface != nullptr && DebugSettings.bEnabled)
+	if (ExecuteContext.UnitContext.DrawInterface != nullptr && DebugSettings.bEnabled)
 	{
-		Context.DrawInterface->DrawPointSimulation(DebugSettings.WorldOffset, Simulation, DebugSettings.Color, DebugSettings.Scale, DebugSettings.CollisionScale, DebugSettings.bDrawPointsAsSpheres);
+		ExecuteContext.UnitContext.DrawInterface->DrawPointSimulation(DebugSettings.WorldOffset, Simulation, DebugSettings.Color, DebugSettings.Scale, DebugSettings.CollisionScale, DebugSettings.bDrawPointsAsSpheres);
 	}
 }
 
diff --git a/Engine/Plugins/Animation/ControlRig/Source/ControlRig/Private/Units/Simulation/RigUnit_SpringInterp.cpp b/Engine/Plugins/Animation/ControlRig/Source/ControlRig/Private/Units/Simulation/RigUnit_SpringInterp.cpp
--- a/Engine/Plugins/Animation/ControlRig/Source/ControlRig/Private/Units/Simulation/RigUnit_SpringInterp.cpp
+++ b/Engine/Plugins/Animation/ControlRig/Source/ControlRig/Private/Units/Simulation/RigUnit_SpringInterp.cpp
@@ -16,14 +16,14 @@ FRigUnit_SpringInterp_Execute()
 {
 	DECLARE_SCOPE_HIERARCHICAL_COUNTER_RIGUNIT()
  
-	if (Context.State == EControlRigState::Init)
+	if (ExecuteContext.UnitContext.State == EControlRigState::Init)
 	{
 		SpringState.Reset();
 	}
 	else
 	{
 		// Clamp to avoid large time deltas.
-		float RemainingTime = FMath::Min(Context.DeltaTime, RigUnitSpringInterpConstants::MaxTimeStep);
+		float RemainingTime = FMath::Min(ExecuteContext.UnitContext.DeltaTime, RigUnitSpringInterpConstants::MaxTimeStep);
  
 		Result = Current;
 		while (RemainingTime >= RigUnitSpringInterpConstants::FixedTimeStep)
@@ -46,14 +46,14 @@ FRigUnit_SpringInterpVector_Execute()
 {
 	DECLARE_SCOPE_HIERARCHICAL_COUNTER_RIGUNIT()
  
-	if (Context.State == EControlRigState::Init)
+	if (ExecuteContext.UnitContext.State == EControlRigState::Init)
 	{
 		SpringState.Reset();
 	}
 	else
 	{
 		// Clamp to avoid large time deltas.
-		float RemainingTime = FMath::Min(Context.DeltaTime, RigUnitSpringInterpConstants::MaxTimeStep);
+		float RemainingTime = FMath::Min(ExecuteContext.UnitContext.DeltaTime, RigUnitSpringInterpConstants::MaxTimeStep);
  
 		Result = Current;
 		while (RemainingTime >= RigUnitSpringInterpConstants::FixedTimeStep)
@@ -76,7 +76,7 @@ FRigUnit_SpringInterpV2_Execute()
 {
 	DECLARE_SCOPE_HIERARCHICAL_COUNTER_RIGUNIT()
 
-	if (Context.State == EControlRigState::Init)
+	if (ExecuteContext.UnitContext.State == EControlRigState::Init)
 	{
 		SpringState.Reset();
 		Result = Target;
@@ -93,11 +93,11 @@ FRigUnit_SpringInterpV2_Execute()
 		}
 		else
 		{
-			SpringState.Velocity += Force * (Context.DeltaTime / RigUnitSpringInterpConstants::Mass);
+			SpringState.Velocity += Force * (ExecuteContext.UnitContext.DeltaTime / RigUnitSpringInterpConstants::Mass);
 		}
 		SimulatedResult = UKismetMathLibrary::FloatSpringInterp(
 			bUseCurrentInput ? Current : SimulatedResult, AdjustedTarget, SpringState, Stiffness, CriticalDamping,
-			Context.DeltaTime, RigUnitSpringInterpConstants::Mass, TargetVelocityAmount, 
+			ExecuteContext.UnitContext.DeltaTime, RigUnitSpringInterpConstants::Mass, TargetVelocityAmount, 
 			false, 0.0f, 0.0f, !bUseCurrentInput || bInitializeFromTarget);
 
 		Result = SimulatedResult;
@@ -109,7 +109,7 @@ FRigUnit_SpringInterpVectorV2_Execute()
 {
     DECLARE_SCOPE_HIERARCHICAL_COUNTER_RIGUNIT()
 
-	if (Context.State == EControlRigState::Init)
+	if (ExecuteContext.UnitContext.State == EControlRigState::Init)
 	{
 		SpringState.Reset();
 		Result = Target;
@@ -126,11 +126,11 @@ FRigUnit_SpringInterpVectorV2_Execute()
 		}
 		else
 		{
-			SpringState.Velocity += Force * (Context.DeltaTime / RigUnitSpringInterpConstants::Mass);
+			SpringState.Velocity += Force * (ExecuteContext.UnitContext.DeltaTime / RigUnitSpringInterpConstants::Mass);
 		}
 		SimulatedResult = UKismetMathLibrary::VectorSpringInterp(
 			bUseCurrentInput ? Current : SimulatedResult, AdjustedTarget, SpringState, Stiffness, CriticalDamping,
-			Context.DeltaTime, RigUnitSpringInterpConstants::Mass, TargetVelocityAmount,
+			ExecuteContext.UnitContext.DeltaTime, RigUnitSpringInterpConstants::Mass, TargetVelocityAmount,
 			false, FVector(), FVector(), !bUseCurrentInput || bInitializeFromTarget);
 		Result = SimulatedResult;
 	}
@@ -141,7 +141,7 @@ FRigUnit_SpringInterpQuaternionV2_Execute()
 {
     DECLARE_SCOPE_HIERARCHICAL_COUNTER_RIGUNIT()
 
-	if (Context.State == EControlRigState::Init)
+	if (ExecuteContext.UnitContext.State == EControlRigState::Init)
 	{
 		SpringState.Reset();
 		Result = Target;
@@ -151,10 +151,10 @@ FRigUnit_SpringInterpQuaternionV2_Execute()
 		// Treat the input as a frequency in Hz
 		float AngularFrequency = Strength * 2.0f * PI;
 		float Stiffness = AngularFrequency * AngularFrequency;
-		SpringState.AngularVelocity += Torque * (Context.DeltaTime / RigUnitSpringInterpConstants::Mass);
+		SpringState.AngularVelocity += Torque * (ExecuteContext.UnitContext.DeltaTime / RigUnitSpringInterpConstants::Mass);
 		SimulatedResult = UKismetMathLibrary::QuaternionSpringInterp(
 			bUseCurrentInput ? Current : SimulatedResult, Target, SpringState, Stiffness, CriticalDamping,
-			Context.DeltaTime, RigUnitSpringInterpConstants::Mass, TargetVelocityAmount, 
+			ExecuteContext.UnitContext.DeltaTime, RigUnitSpringInterpConstants::Mass, TargetVelocityAmount, 
 			!bUseCurrentInput || bInitializeFromTarget);
 		Result = SimulatedResult;
 		AngularVelocity = SpringState.AngularVelocity;
diff --git a/Engine/Plugins/Animation/ControlRig/Source/ControlRig/Private/Units/Simulation/RigUnit_TimeOffset.cpp b/Engine/Plugins/Animation/ControlRig/Source/ControlRig/Private/Units/Simulation/RigUnit_TimeOffset.cpp
--- a/Engine/Plugins/Animation/ControlRig/Source/ControlRig/Private/Units/Simulation/RigUnit_TimeOffset.cpp
+++ b/Engine/Plugins/Animation/ControlRig/Source/ControlRig/Private/Units/Simulation/RigUnit_TimeOffset.cpp
@@ -12,7 +12,7 @@ FRigUnit_TimeOffsetFloat_Execute()
 
 	if(BufferSize <= 0)
 	{
-		if (Context.State == EControlRigState::Init)
+		if (ExecuteContext.UnitContext.State == EControlRigState::Init)
 		{
 			UE_CONTROLRIG_RIGUNIT_REPORT_WARNING(TEXT("BufferSize is too small."));
 		}
@@ -22,7 +22,7 @@ FRigUnit_TimeOffsetFloat_Execute()
 
 	if (TimeRange <= 0)
 	{
-		if (Context.State == EControlRigState::Init)
+		if (ExecuteContext.UnitContext.State == EControlRigState::Init)
 		{
 			UE_CONTROLRIG_RIGUNIT_REPORT_WARNING(TEXT("TimeRange is too small."));
 		}
@@ -35,7 +35,7 @@ FRigUnit_TimeOffsetFloat_Execute()
 	DeltaTimes.SetNum(MaxSize);
 	MaxSize = FMath::Min<int32>(MaxSize, Buffer.Num());
 
-	if (Context.State == EControlRigState::Init)
+	if (ExecuteContext.UnitContext.State == EControlRigState::Init)
 	{
 		UpperBound = 0;
 		Result = Value;
@@ -80,35 +80,35 @@ FRigUnit_TimeOffsetFloat_Execute()
 		Result = FMath::Lerp<float>(A, B, T);
 	}
 
-	if (Context.DeltaTime > SMALL_NUMBER)
+	if (ExecuteContext.UnitContext.DeltaTime > SMALL_NUMBER)
 	{
 		if (UpperBound == 0)
 		{
 			LastInsertIndex = UpperBound;
 			Buffer[UpperBound] = Value;
-			DeltaTimes[UpperBound++] = Context.DeltaTime;
+			DeltaTimes[UpperBound++] = ExecuteContext.UnitContext.DeltaTime;
 		}
 		else
 		{
 			float SecondsPerEntry = TimeRange / float(MaxSize - 1);
-			if (DeltaTimes[LastInsertIndex] > SecondsPerEntry - Context.DeltaTime * 0.5f)
+			if (DeltaTimes[LastInsertIndex] > SecondsPerEntry - ExecuteContext.UnitContext.DeltaTime * 0.5f)
 			{
 				if (UpperBound < MaxSize)
 				{
 					LastInsertIndex = UpperBound;
 					Buffer[UpperBound] = Value;
-					DeltaTimes[UpperBound++] = Context.DeltaTime;
+					DeltaTimes[UpperBound++] = ExecuteContext.UnitContext.DeltaTime;
 				}
 				else
 				{
 					LastInsertIndex = (LastInsertIndex + 1) % UpperBound;
 					Buffer[LastInsertIndex] = Value;
-					DeltaTimes[LastInsertIndex] = Context.DeltaTime;
+					DeltaTimes[LastInsertIndex] = ExecuteContext.UnitContext.DeltaTime;
 				}
 			}
 			else
 			{
-				DeltaTimes[LastInsertIndex] += Context.DeltaTime;
+				DeltaTimes[LastInsertIndex] += ExecuteContext.UnitContext.DeltaTime;
 			}
 		}
 	}
@@ -121,7 +121,7 @@ FRigUnit_TimeOffsetVector_Execute()
 
 	if(BufferSize <= 0)
 	{
-		if (Context.State == EControlRigState::Init)
+		if (ExecuteContext.UnitContext.State == EControlRigState::Init)
 		{
 			UE_CONTROLRIG_RIGUNIT_REPORT_WARNING(TEXT("BufferSize is too small."));
 		}
@@ -131,7 +131,7 @@ FRigUnit_TimeOffsetVector_Execute()
 
 	if (TimeRange <= 0)
 	{
-		if (Context.State == EControlRigState::Init)
+		if (ExecuteContext.UnitContext.State == EControlRigState::Init)
 		{
 			UE_CONTROLRIG_RIGUNIT_REPORT_WARNING(TEXT("TimeRange is too small."));
 		}
@@ -144,7 +144,7 @@ FRigUnit_TimeOffsetVector_Execute()
 	DeltaTimes.SetNum(MaxSize);
 	MaxSize = FMath::Min<int32>(MaxSize, Buffer.Num());
 
-	if (Context.State == EControlRigState::Init)
+	if (ExecuteContext.UnitContext.State == EControlRigState::Init)
 	{
 		UpperBound = 0;
 		Result = Value;
@@ -189,35 +189,35 @@ FRigUnit_TimeOffsetVector_Execute()
 		Result = FMath::Lerp<FVector>(A, B, T);
 	}
 
-	if (Context.DeltaTime > SMALL_NUMBER)
+	if (ExecuteContext.UnitContext.DeltaTime > SMALL_NUMBER)
 	{
 		if (UpperBound == 0)
 		{
 			LastInsertIndex = UpperBound;
 			Buffer[UpperBound] = Value;
-			DeltaTimes[UpperBound++] = Context.DeltaTime;
+			DeltaTimes[UpperBound++] = ExecuteContext.UnitContext.DeltaTime;
 		}
 		else
 		{
 			float SecondsPerEntry = TimeRange / float(MaxSize - 1);
-			if (DeltaTimes[LastInsertIndex] > SecondsPerEntry - Context.DeltaTime * 0.5f)
+			if (DeltaTimes[LastInsertIndex] > SecondsPerEntry - ExecuteContext.UnitContext.DeltaTime * 0.5f)
 			{
 				if (UpperBound < MaxSize)
 				{
 					LastInsertIndex = UpperBound;
 					Buffer[UpperBound] = Value;
-					DeltaTimes[UpperBound++] = Context.DeltaTime;
+					DeltaTimes[UpperBound++] = ExecuteContext.UnitContext.DeltaTime;
 				}
 				else
 				{
 					LastInsertIndex = (LastInsertIndex + 1) % UpperBound;
 					Buffer[LastInsertIndex] = Value;
-					DeltaTimes[LastInsertIndex] = Context.DeltaTime;
+					DeltaTimes[LastInsertIndex] = ExecuteContext.UnitContext.DeltaTime;
 				}
 			}
 			else
 			{
-				DeltaTimes[LastInsertIndex] += Context.DeltaTime;
+				DeltaTimes[LastInsertIndex] += ExecuteContext.UnitContext.DeltaTime;
 			}
 		}
 	}
@@ -230,7 +230,7 @@ FRigUnit_TimeOffsetTransform_Execute()
 
 	if(BufferSize <= 0)
 	{
-		if (Context.State == EControlRigState::Init)
+		if (ExecuteContext.UnitContext.State == EControlRigState::Init)
 		{
 			UE_CONTROLRIG_RIGUNIT_REPORT_WARNING(TEXT("BufferSize is too small."));
 		}
@@ -240,7 +240,7 @@ FRigUnit_TimeOffsetTransform_Execute()
 
 	if (TimeRange <= 0)
 	{
-		if (Context.State == EControlRigState::Init)
+		if (ExecuteContext.UnitContext.State == EControlRigState::Init)
 		{
 			UE_CONTROLRIG_RIGUNIT_REPORT_WARNING(TEXT("TimeRange is too small."));
 		}
@@ -253,7 +253,7 @@ FRigUnit_TimeOffsetTransform_Execute()
 	DeltaTimes.SetNum(MaxSize);
 	MaxSize = FMath::Min<int32>(MaxSize, Buffer.Num());
 
-	if (Context.State == EControlRigState::Init)
+	if (ExecuteContext.UnitContext.State == EControlRigState::Init)
 	{
 		UpperBound = 0;
 		Result = Value;
@@ -300,35 +300,35 @@ FRigUnit_TimeOffsetTransform_Execute()
 		Result.SetScale3D(FMath::Lerp<FVector>(A.GetScale3D(), B.GetScale3D(), T));
 	}
 
-	if (Context.DeltaTime > SMALL_NUMBER)
+	if (ExecuteContext.UnitContext.DeltaTime > SMALL_NUMBER)
 	{
 		if (UpperBound == 0)
 		{
 			LastInsertIndex = UpperBound;
 			Buffer[UpperBound] = Value;
-			DeltaTimes[UpperBound++] = Context.DeltaTime;
+			DeltaTimes[UpperBound++] = ExecuteContext.UnitContext.DeltaTime;
 		}
 		else
 		{
 			float SecondsPerEntry = TimeRange / float(MaxSize - 1);
-			if (DeltaTimes[LastInsertIndex] > SecondsPerEntry - Context.DeltaTime * 0.5f)
+			if (DeltaTimes[LastInsertIndex] > SecondsPerEntry - ExecuteContext.UnitContext.DeltaTime * 0.5f)
 			{
 				if (UpperBound < MaxSize)
 				{
 					LastInsertIndex = UpperBound;
 					Buffer[UpperBound] = Value;
-					DeltaTimes[UpperBound++] = Context.DeltaTime;
+					DeltaTimes[UpperBound++] = ExecuteContext.UnitContext.DeltaTime;
 				}
 				else
 				{
 					LastInsertIndex = (LastInsertIndex + 1) % UpperBound;
 					Buffer[LastInsertIndex] = Value;
-					DeltaTimes[LastInsertIndex] = Context.DeltaTime;
+					DeltaTimes[LastInsertIndex] = ExecuteContext.UnitContext.DeltaTime;
 				}
 			}
 			else
 			{
-				DeltaTimes[LastInsertIndex] += Context.DeltaTime;
+				DeltaTimes[LastInsertIndex] += ExecuteContext.UnitContext.DeltaTime;
 			}
 		}
 	}
@@ -339,7 +339,7 @@ FRigUnit_TimeOffsetTransform_Execute()
 
 IMPLEMENT_RIGUNIT_AUTOMATION_TEST(FRigUnit_TimeOffsetFloat)
 {
-	Context.DeltaTime = 1.f;
+	ExecuteContext.UnitContext.DeltaTime = 1.f;
 	Unit.SecondsAgo = 0.5;
 	Unit.TimeRange = 5.f;
 	Unit.BufferSize = 16;
@@ -365,7 +365,7 @@ IMPLEMENT_RIGUNIT_AUTOMATION_TEST(FRigUnit_TimeOffsetFloat)
 	Execute();
 	AddErrorIfFalse(FMath::IsNearlyEqual(Unit.Result, 11.f), TEXT("unexpected previous result"));
 
-	Context.DeltaTime = 0.f;
+	ExecuteContext.UnitContext.DeltaTime = 0.f;
 	Unit.SecondsAgo = 1.5;
 	Execute();
 	AddErrorIfFalse(FMath::IsNearlyEqual(Unit.Result, 9.f), TEXT("unexpected previous result"));
@@ -379,7 +379,7 @@ IMPLEMENT_RIGUNIT_AUTOMATION_TEST(FRigUnit_TimeOffsetFloat)
 	Execute();
 	AddErrorIfFalse(FMath::IsNearlyEqual(Unit.Result, 3.f), TEXT("unexpected previous result"));
 
-	Context.DeltaTime = 1.f;
+	ExecuteContext.UnitContext.DeltaTime = 1.f;
 	Unit.SecondsAgo = 0.5;
 	Unit.Value = 14.f;
 	Execute();
@@ -418,7 +418,7 @@ IMPLEMENT_RIGUNIT_AUTOMATION_TEST(FRigUnit_TimeOffsetFloat)
 	Execute();
 	AddErrorIfFalse(FMath::IsNearlyEqual(Unit.Result, 35.f), TEXT("unexpected previous result"));
 
-	Context.DeltaTime = 0.f;
+	ExecuteContext.UnitContext.DeltaTime = 0.f;
 	Unit.SecondsAgo = 1.5;
 	Execute();
 	AddErrorIfFalse(FMath::IsNearlyEqual(Unit.Result, 33.f), TEXT("unexpected previous result"));
@@ -433,7 +433,7 @@ IMPLEMENT_RIGUNIT_AUTOMATION_TEST(FRigUnit_TimeOffsetFloat)
 	AddErrorIfFalse(FMath::IsNearlyEqual(Unit.Result, 27.f), TEXT("unexpected previous result"));
 
 
-	Context.DeltaTime = 1.f;
+	ExecuteContext.UnitContext.DeltaTime = 1.f;
 	Unit.TimeRange = 8.f;
 	Unit.BufferSize = 3;
 	Unit.SecondsAgo = 0.5f;
@@ -448,7 +448,7 @@ IMPLEMENT_RIGUNIT_AUTOMATION_TEST(FRigUnit_TimeOffsetFloat)
 	Execute();
 	AddErrorIfFalse(FMath::IsNearlyEqual(Unit.Result, 5.f), TEXT("unexpected previous result"));
 
-	Context.DeltaTime = 0.f;
+	ExecuteContext.UnitContext.DeltaTime = 0.f;
 	Unit.SecondsAgo = 1.5f;
 	Execute();
 	AddErrorIfFalse(FMath::IsNearlyEqual(Unit.Result, 4.f), TEXT("unexpected previous result"));
diff --git a/Engine/Plugins/Animation/ControlRig/Source/ControlRig/Private/Units/Simulation/RigUnit_Timeline.cpp b/Engine/Plugins/Animation/ControlRig/Source/ControlRig/Private/Units/Simulation/RigUnit_Timeline.cpp
--- a/Engine/Plugins/Animation/ControlRig/Source/ControlRig/Private/Units/Simulation/RigUnit_Timeline.cpp
+++ b/Engine/Plugins/Animation/ControlRig/Source/ControlRig/Private/Units/Simulation/RigUnit_Timeline.cpp
@@ -8,13 +8,13 @@
 FRigUnit_Timeline_Execute()
 {
     DECLARE_SCOPE_HIERARCHICAL_COUNTER_RIGUNIT()
-	if (Context.State == EControlRigState::Init)
+	if (ExecuteContext.UnitContext.State == EControlRigState::Init)
 	{
 		Time = AccumulatedValue = 0.f;
 		return;
 	}
 
-	Time = AccumulatedValue = AccumulatedValue + Context.DeltaTime * Speed;
+	Time = AccumulatedValue = AccumulatedValue + ExecuteContext.UnitContext.DeltaTime * Speed;
 }
 
 #if WITH_DEV_AUTOMATION_TESTS
@@ -22,7 +22,7 @@ FRigUnit_Timeline_Execute()
 
 IMPLEMENT_RIGUNIT_AUTOMATION_TEST(FRigUnit_Timeline)
 {
-	Context.DeltaTime = 1.f;
+	ExecuteContext.UnitContext.DeltaTime = 1.f;
 	InitAndExecute();
 	AddErrorIfFalse(FMath::IsNearlyEqual(Unit.Time, 1.f), TEXT("unexpected time"));
 	InitAndExecute();
@@ -42,7 +42,7 @@ IMPLEMENT_RIGUNIT_AUTOMATION_TEST(FRigUnit_Timeline)
 FRigUnit_TimeLoop_Execute()
 {
 	DECLARE_SCOPE_HIERARCHICAL_COUNTER_RIGUNIT()
-	if (Context.State == EControlRigState::Init)
+	if (ExecuteContext.UnitContext.State == EControlRigState::Init)
 	{
 		Absolute = Relative = FlipFlop = AccumulatedAbsolute = AccumulatedRelative = 0.f;
 		NumIterations = 0;
@@ -51,7 +51,7 @@ FRigUnit_TimeLoop_Execute()
 	}
 
 	const float DurationClamped = FMath::Max(Duration, 0.0001f);
-	const float Increment = Context.DeltaTime * Speed;
+	const float Increment = ExecuteContext.UnitContext.DeltaTime * Speed;
 	Absolute = AccumulatedAbsolute = AccumulatedAbsolute + Increment;
 
 	AccumulatedRelative = AccumulatedRelative + Increment;
diff --git a/Engine/Plugins/Animation/ControlRig/Source/ControlRig/Private/Units/Simulation/RigUnit_Verlet.cpp b/Engine/Plugins/Animation/ControlRig/Source/ControlRig/Private/Units/Simulation/RigUnit_Verlet.cpp
--- a/Engine/Plugins/Animation/ControlRig/Source/ControlRig/Private/Units/Simulation/RigUnit_Verlet.cpp
+++ b/Engine/Plugins/Animation/ControlRig/Source/ControlRig/Private/Units/Simulation/RigUnit_Verlet.cpp
@@ -8,7 +8,7 @@
 FRigUnit_VerletIntegrateVector_Execute()
 {
     DECLARE_SCOPE_HIERARCHICAL_COUNTER_RIGUNIT()
-	if (Context.State == EControlRigState::Init)
+	if (ExecuteContext.UnitContext.State == EControlRigState::Init)
 	{
 		bInitialized = false;
 		return;
@@ -24,13 +24,13 @@ FRigUnit_VerletIntegrateVector_Execute()
 	}
 
 	Point.LinearDamping = Damp;
-	if (Context.DeltaTime > SMALL_NUMBER)
+	if (ExecuteContext.UnitContext.DeltaTime > SMALL_NUMBER)
 	{
-		float U = FMath::Clamp<float>(Blend * Context.DeltaTime, 0.f, 1.f);
-		const FVector CombinedForce = (Target - Point.Position) * FMath::Max(Strength, 0.0001f) + Force * Context.DeltaTime * 60.f;
+		float U = FMath::Clamp<float>(Blend * ExecuteContext.UnitContext.DeltaTime, 0.f, 1.f);
+		const FVector CombinedForce = (Target - Point.Position) * FMath::Max(Strength, 0.0001f) + Force * ExecuteContext.UnitContext.DeltaTime * 60.f;
 		const FVector PreviousVelocity = Point.LinearVelocity;
-		Point = Point.IntegrateVerlet(CombinedForce, Blend, Context.DeltaTime);
-		Acceleration = (Point.LinearVelocity - PreviousVelocity) / Context.DeltaTime;
+		Point = Point.IntegrateVerlet(CombinedForce, Blend, ExecuteContext.UnitContext.DeltaTime);
+		Acceleration = (Point.LinearVelocity - PreviousVelocity) / ExecuteContext.UnitContext.DeltaTime;
 		Position = Point.Position;
 		Velocity = Point.LinearVelocity;
 	}
diff --git a/Engine/Plugins/Animation/ControlRig/Source/ControlRig/Public/ControlRigDefines.h b/Engine/Plugins/Animation/ControlRig/Source/ControlRig/Public/ControlRigDefines.h
--- a/Engine/Plugins/Animation/ControlRig/Source/ControlRig/Public/ControlRigDefines.h
+++ b/Engine/Plugins/Animation/ControlRig/Source/ControlRig/Public/ControlRigDefines.h
@@ -13,28 +13,6 @@
 //#define DECLARE_SCOPE_HIERARCHICAL_COUNTER_RIGUNIT() \
 //	DECLARE_SCOPE_HIERARCHICAL_COUNTER_FUNC()
 
-USTRUCT(BlueprintType)
-struct FControlRigExecuteContext : public FRigVMExecuteContext
-{
-	GENERATED_BODY()
-
-	FControlRigExecuteContext()
-		: FRigVMExecuteContext()
-		, Hierarchy(nullptr)
-	{
-	}
-
-	virtual void Copy(const FRigVMExecuteContext* InOtherContext) override
-	{
-		Super::Copy(InOtherContext);
-
-		const FControlRigExecuteContext* OtherContext = (const FControlRigExecuteContext*)InOtherContext; 
-		Hierarchy = OtherContext->Hierarchy;
-	}
-		
-	URigHierarchy* Hierarchy;
-};
-
 UENUM()
 enum class ETransformSpaceMode : uint8
 {
diff --git a/Engine/Plugins/Animation/ControlRig/Source/ControlRig/Public/Units/RigDispatchFactory.h b/Engine/Plugins/Animation/ControlRig/Source/ControlRig/Public/Units/RigDispatchFactory.h
--- a/Engine/Plugins/Animation/ControlRig/Source/ControlRig/Public/Units/RigDispatchFactory.h
+++ b/Engine/Plugins/Animation/ControlRig/Source/ControlRig/Public/Units/RigDispatchFactory.h
@@ -26,14 +26,6 @@ struct CONTROLRIG_API FRigDispatchFactory : public FRigVMDispatchFactory
     	FRigVMRegistry::Get().FindOrAddType(FCachedRigElement::StaticStruct());
 	}
 
-	virtual TArray<TPair<FName,FString>> GetOpaqueArguments() const override
-	{
-		static const TArray<TPair<FName,FString>> OpaqueArguments = {
-			TPair<FName,FString>(TEXT("Context"), TEXT("const FRigUnitContext&"))
-		};
-		return OpaqueArguments;
-	}
-
 #if WITH_EDITOR
 
 	virtual FString GetArgumentDefaultValue(const FName& InArgumentName, TRigVMTypeIndex InTypeIndex) const override;
@@ -42,7 +34,7 @@ struct CONTROLRIG_API FRigDispatchFactory : public FRigVMDispatchFactory
 
 	static const FRigUnitContext& GetRigUnitContext(const FRigVMExtendedExecuteContext& InContext)
 	{
-		return *(const FRigUnitContext*)InContext.OpaqueArguments[0];
+		return InContext.GetPublicData<FControlRigExecuteContext>().UnitContext;
 	}
 };
 
diff --git a/Engine/Plugins/Animation/ControlRig/Source/ControlRig/Public/Units/RigUnit.h b/Engine/Plugins/Animation/ControlRig/Source/ControlRig/Public/Units/RigUnit.h
--- a/Engine/Plugins/Animation/ControlRig/Source/ControlRig/Public/Units/RigUnit.h
+++ b/Engine/Plugins/Animation/ControlRig/Source/ControlRig/Public/Units/RigUnit.h
@@ -5,12 +5,13 @@
 #include "ControlRigDefines.h"
 #include "RigVMCore/RigVMStruct.h"
 #include "RigVMCore/RigVMRegistry.h"
+#include "RigUnitContext.h"
 #include "RigUnit.generated.h"
 
 struct FRigUnitContext;
 
 /** Base class for all rig units */
-USTRUCT(BlueprintType, meta=(Abstract, NodeColor = "0.1 0.1 0.1"), BlueprintInternalUseOnlyHierarchical)
+USTRUCT(BlueprintType, meta=(Abstract, NodeColor = "0.1 0.1 0.1", ExecuteContext="FControlRigExecuteContext"), BlueprintInternalUseOnlyHierarchical)
 struct CONTROLRIG_API FRigUnit : public FRigVMStruct
 {
 	GENERATED_BODY()
@@ -25,7 +26,7 @@ struct CONTROLRIG_API FRigUnit : public FRigVMStruct
 	virtual FString GetUnitLabel() const { return FString(); }
 
 	/** Execute logic for this rig unit */
-	virtual void Execute(const FRigUnitContext& Context) {}
+	virtual void Execute() {}
 
 	virtual FRigElementKey DetermineSpaceForPin(const FString& InPinPath, void* InUserContext) const { return FRigElementKey(); }
 	
diff --git a/Engine/Plugins/Animation/ControlRig/Source/ControlRig/Public/Units/RigUnitContext.h b/Engine/Plugins/Animation/ControlRig/Source/ControlRig/Public/Units/RigUnitContext.h
--- a/Engine/Plugins/Animation/ControlRig/Source/ControlRig/Public/Units/RigUnitContext.h
+++ b/Engine/Plugins/Animation/ControlRig/Source/ControlRig/Public/Units/RigUnitContext.h
@@ -64,7 +64,6 @@ struct FRigUnitContext
 		, DeltaTime(0.f)
 		, AbsoluteTime(0.f)
 		, State(EControlRigState::Invalid)
-		, Hierarchy(nullptr)
 		, InteractionType((uint8)EControlRigInteractionType::None)
 		, ElementsBeingInteracted()
 		, ToWorldSpaceTransform(FTransform::Identity)
@@ -101,9 +100,6 @@ struct FRigUnitContext
 	/** Current execution context */
 	EControlRigState State;
 
-	/** The current hierarchy being executed */
-	URigHierarchy* Hierarchy;
-
 	/** The current hierarchy settings */
 	FRigHierarchySettings HierarchySettings;
 
@@ -203,11 +199,34 @@ struct FRigUnitContext
 	}
 };
 
+USTRUCT(BlueprintType)
+struct FControlRigExecuteContext : public FRigVMExecuteContext
+{
+	GENERATED_BODY()
+
+	FControlRigExecuteContext()
+		: FRigVMExecuteContext()
+		, Hierarchy(nullptr)
+	{
+	}
+
+	virtual void Copy(const FRigVMExecuteContext* InOtherContext) override
+	{
+		Super::Copy(InOtherContext);
+
+		const FControlRigExecuteContext* OtherContext = (const FControlRigExecuteContext*)InOtherContext; 
+		Hierarchy = OtherContext->Hierarchy;
+	}
+
+	FRigUnitContext UnitContext;
+	URigHierarchy* Hierarchy;
+};
+
 #if WITH_EDITOR
 #define UE_CONTROLRIG_RIGUNIT_REPORT(Severity, Format, ...) \
-if(Context.Log != nullptr) \
+if(ExecuteContext.UnitContext.Log != nullptr) \
 { \
-	Context.Log->Report(EMessageSeverity::Severity, ExecuteContext.GetFunctionName(), ExecuteContext.GetInstructionIndex(), FString::Printf((Format), ##__VA_ARGS__)); \
+	ExecuteContext.UnitContext.Log->Report(EMessageSeverity::Severity, ExecuteContext.GetFunctionName(), ExecuteContext.GetInstructionIndex(), FString::Printf((Format), ##__VA_ARGS__)); \
 }
 #define UE_CONTROLRIG_RIGUNIT_LOG_MESSAGE(Format, ...) UE_CONTROLRIG_RIGUNIT_REPORT(Info, (Format), ##__VA_ARGS__)
 #define UE_CONTROLRIG_RIGUNIT_REPORT_WARNING(Format, ...) UE_CONTROLRIG_RIGUNIT_REPORT(Warning, (Format), ##__VA_ARGS__)
diff --git a/Engine/Plugins/Animation/ControlRig/Source/ControlRig/Public/Units/RigUnitTest.h b/Engine/Plugins/Animation/ControlRig/Source/ControlRig/Public/Units/RigUnitTest.h
--- a/Engine/Plugins/Animation/ControlRig/Source/ControlRig/Public/Units/RigUnitTest.h
+++ b/Engine/Plugins/Animation/ControlRig/Source/ControlRig/Public/Units/RigUnitTest.h
@@ -40,13 +40,11 @@ public:
 			Hierarchy->AddToRoot();
 			// we no longer add the controller to root since controller is now part of the hierarchy
 		}
-		Context.Hierarchy = Hierarchy;;
 		ExecuteContext.Hierarchy = Hierarchy;
 	}
 
 	URigHierarchy* Hierarchy;
 	URigHierarchyController* Controller;
-	FRigUnitContext Context;
 };
 
 #define CONTROLRIG_RIGUNIT_STRINGIFY(Content) #Content
@@ -80,13 +78,13 @@ public:
 		virtual FString GetBeautifiedTestName() const override { return TEXT(CONTROLRIG_RIGUNIT_STRINGIFY(ControlRig.Units.TUnitStruct)); } \
 		void Init() \
 		{ \
-			Context.State = EControlRigState::Init; \
-			Unit.Execute(Context); \
+			ExecuteContext.UnitContext.State = EControlRigState::Init; \
+			Unit.Execute(ExecuteContext); \
 		} \
 		void Execute() \
 		{ \
-			Context.State = EControlRigState::Update; \
-			Unit.Execute(Context); \
+			ExecuteContext.UnitContext.State = EControlRigState::Update; \
+			Unit.Execute(ExecuteContext); \
 		} \
 		void InitAndExecute() { Init(); Execute(); } \
 	}; \
diff --git a/Engine/Plugins/Animation/ControlRig/Source/ControlRigDeveloper/Private/ControlRigBlueprint.cpp b/Engine/Plugins/Animation/ControlRig/Source/ControlRigDeveloper/Private/ControlRigBlueprint.cpp
--- a/Engine/Plugins/Animation/ControlRig/Source/ControlRigDeveloper/Private/ControlRigBlueprint.cpp
+++ b/Engine/Plugins/Animation/ControlRig/Source/ControlRigDeveloper/Private/ControlRigBlueprint.cpp
@@ -132,6 +132,7 @@ UControlRigBlueprint::UControlRigBlueprint(const FObjectInitializer& ObjectIniti
 		RigVMClient.GetOrCreateFunctionLibrary(false, &ObjectInitializer, false);
 		RigVMClient.AddModel(RigVMModelPrefix, false, &ObjectInitializer, false);
 	}
+	RigVMClient.SetExecuteContextStruct(FControlRigExecuteContext::StaticStruct());
 	
 	FunctionLibraryEdGraph = ObjectInitializer.CreateDefaultSubobject<UControlRigGraph>(this, TEXT("RigVMFunctionLibraryEdGraph"));
 	FunctionLibraryEdGraph->Schema = UControlRigGraphSchema::StaticClass();
@@ -375,8 +376,6 @@ void UControlRigBlueprint::HandleRigVMGraphAdded(const FRigVMClient* InClient, c
 {
 	if(URigVMGraph* Model = InClient->GetModel(InNodePath))
 	{
-		Model->SetExecuteContextStruct(FControlRigExecuteContext::StaticStruct());
-
 		if(!HasAnyFlags(RF_ClassDefaultObject | RF_NeedInitialization | RF_NeedLoad | RF_NeedPostLoad) &&
 			GetOuter() != GetTransientPackage())
 		{
@@ -1184,20 +1183,6 @@ void UControlRigBlueprint::RecompileVM()
 		}
 		CDO->VM->Reset();
 
-		FRigUnitContext InitContext;
-		InitContext.State = EControlRigState::Init;
-		InitContext.Hierarchy = CDO->DynamicHierarchy;
-
-		FRigUnitContext UpdateContext = InitContext;
-		UpdateContext.State = EControlRigState::Update;
-
-		void* InitContextPtr = &InitContext;
-		void* UpdateContextPtr = &UpdateContext;
-
-		TArray<FRigVMUserDataArray> UserData;
-		UserData.Add(FRigVMUserDataArray(&InitContextPtr, 1));
-		UserData.Add(FRigVMUserDataArray(&UpdateContextPtr, 1));
-
 		// Clear all Errors
 		CompileLog.Messages.Reset();
 		CompileLog.NumErrors = CompileLog.NumWarnings = 0;
@@ -1223,9 +1208,9 @@ void UControlRigBlueprint::RecompileVM()
 		}
 
 		URigVMCompiler* Compiler = URigVMCompiler::StaticClass()->GetDefaultObject<URigVMCompiler>();
-		VMCompileSettings.SetExecuteContextStruct(FControlRigExecuteContext::StaticStruct());
+		VMCompileSettings.SetExecuteContextStruct(RigVMClient.GetExecuteContextStruct());
 		Compiler->Settings = (bCompileInDebugMode) ? FRigVMCompileSettings::Fast(VMCompileSettings.GetExecuteContextStruct()) : VMCompileSettings;
-		Compiler->Compile(RigVMClient.GetAllModels(false, false), GetOrCreateController(), CDO->VM, CDO->GetExternalVariablesImpl(false), UserData, &PinToOperandMap);
+		Compiler->Compile(RigVMClient.GetAllModels(false, false), GetOrCreateController(), CDO->VM, CDO->GetExternalVariablesImpl(false), &PinToOperandMap);
 
 		if (bErrorsDuringCompilation)
 		{
@@ -2350,7 +2335,7 @@ URigVMGraph* UControlRigBlueprint::GetTemplateModel(bool bIsFunctionLibrary)
 			TemplateModel = NewObject<URigVMGraph>(this, TEXT("TemplateModel"));
 		}
 		TemplateModel->SetFlags(RF_Transient);
-		TemplateModel->SetExecuteContextStruct(FControlRigExecuteContext::StaticStruct());
+		TemplateModel->SetExecuteContextStruct(RigVMClient.GetExecuteContextStruct());
 	}
 	return TemplateModel;
 #else
