RigVM: Fixed array pin support for select node

diff --git a/Engine/Plugins/Animation/ControlRig/Source/ControlRigDeveloper/Private/Graph/ControlRigGraphNode.cpp b/Engine/Plugins/Animation/ControlRig/Source/ControlRigDeveloper/Private/Graph/ControlRigGraphNode.cpp
--- a/Engine/Plugins/Animation/ControlRig/Source/ControlRigDeveloper/Private/Graph/ControlRigGraphNode.cpp
+++ b/Engine/Plugins/Animation/ControlRig/Source/ControlRigDeveloper/Private/Graph/ControlRigGraphNode.cpp
@@ -1704,7 +1704,7 @@ void UControlRigGraphNode::ConfigurePin(UEdGraphPin* EdGraphPin, const URigVMPin
 
 	EdGraphPin->bHidden = bHidden;
 	EdGraphPin->PinFriendlyName = FText::FromName(ModelPin->GetDisplayName());
-	EdGraphPin->bNotConnectable = !bConnectable;
+	EdGraphPin->bNotConnectable = !bConnectable  || ModelPin->IsFixedSizeArray();
 	EdGraphPin->bOrphanedPin = ModelPin->IsOrphanPin() ? 1 : 0; 
 	EdGraphPin->bDisplayAsMutableRef = ModelPin->IsWildCard();
 }
diff --git a/Engine/Plugins/Animation/ControlRig/Source/ControlRigEditor/Private/ControlRigEditorModule.cpp b/Engine/Plugins/Animation/ControlRig/Source/ControlRigEditor/Private/ControlRigEditorModule.cpp
--- a/Engine/Plugins/Animation/ControlRig/Source/ControlRigEditor/Private/ControlRigEditorModule.cpp
+++ b/Engine/Plugins/Animation/ControlRig/Source/ControlRigEditor/Private/ControlRigEditorModule.cpp
@@ -1483,26 +1483,35 @@ void FControlRigEditorModule::GetContextMenuActions(const UControlRigGraphSchema
 						{
 							FToolMenuSection& TemplatesSection = Menu->AddSection("EdGraphSchemaTemplates", LOCTEXT("TemplatesHeader", "Templates"));
 
-							if(ModelPin->IsRootPin())
+							if(const FRigVMTemplate* Template = TemplateNode->GetTemplate())
 							{
-								TSharedRef<SControlRigChangePinType> ChangePinTypeWidget =
-								SNew(SControlRigChangePinType)
-								.Blueprint(RigBlueprint)
-								.ModelPins({ModelPin});
+								if(!ModelPin->IsExecuteContext())
+								{
+									if(const FRigVMTemplateArgument* Argument = Template->FindArgument(ModelPin->GetRootPin()->GetFName()))
+									{
+										if(!Argument->IsSingleton())
+										{
+											TSharedRef<SControlRigChangePinType> ChangePinTypeWidget =
+											SNew(SControlRigChangePinType)
+											.Blueprint(RigBlueprint)
+											.ModelPins({ModelPin});
 
-								TemplatesSection.AddEntry(FToolMenuEntry::InitWidget("ChangePinTypeWidget", ChangePinTypeWidget, FText(), true));
+											TemplatesSection.AddEntry(FToolMenuEntry::InitWidget("ChangePinTypeWidget", ChangePinTypeWidget, FText(), true));
+										}
+									}
+								}
+							
+								TemplatesSection.AddMenuEntry(
+									"Unresolve Template Node",
+									LOCTEXT("UnresolveTemplateNode", "Unresolve Template Node"),
+									LOCTEXT("UnresolveTemplateNode_Tooltip", "Removes any type information from the template node"),
+									FSlateIcon(),
+									FUIAction(FExecuteAction::CreateLambda([Controller, ModelPin]() {
+										const TArray<FName> Nodes = ModelPin->GetGraph()->GetSelectNodes();
+										Controller->UnresolveTemplateNodes(Nodes, true, true);
+									})
+								));
 							}
-						
-							TemplatesSection.AddMenuEntry(
-								"Unresolve Template Node",
-								LOCTEXT("UnresolveTemplateNode", "Unresolve Template Node"),
-								LOCTEXT("UnresolveTemplateNode_Tooltip", "Removes any type information from the template node"),
-								FSlateIcon(),
-								FUIAction(FExecuteAction::CreateLambda([Controller, ModelPin]() {
-									const TArray<FName> Nodes = ModelPin->GetGraph()->GetSelectNodes();
-									Controller->UnresolveTemplateNodes(Nodes, true, true);
-								})
-							));
 						}
 					}
 
diff --git a/Engine/Plugins/Animation/ControlRig/Source/ControlRigEditor/Private/Graph/SControlRigGraphNode.cpp b/Engine/Plugins/Animation/ControlRig/Source/ControlRigEditor/Private/Graph/SControlRigGraphNode.cpp
--- a/Engine/Plugins/Animation/ControlRig/Source/ControlRigEditor/Private/Graph/SControlRigGraphNode.cpp
+++ b/Engine/Plugins/Animation/ControlRig/Source/ControlRigEditor/Private/Graph/SControlRigGraphNode.cpp
@@ -235,31 +235,53 @@ void SControlRigGraphNode::UpdateCompactNode()
 	CreatePinWidgets();
 }
 
-void SControlRigGraphNode::CreateAggregateAddPinButton()
+void SControlRigGraphNode::CreateAddPinButton()
 {
-	if (LeftNodeBox.IsValid() && ModelNode.IsValid() && (ModelNode->IsAggregate() || ModelNode->IsA<URigVMAggregateNode>()))
+	if (LeftNodeBox.IsValid() && ModelNode.IsValid())
 	{
-		const bool bInputAggregate = ModelNode->IsInputAggregate();
-		const TSharedRef<SWidget> AddPinButton = AddPinButtonContent(
-		   LOCTEXT("ControlRigAggregateNodeAddPinButton", "Add pin"),
-		   bInputAggregate ? 
-			LOCTEXT("ControlRigAggregateNodeAddInputPinButton_Tooltip", "Adds an input pin to the node") :
-			LOCTEXT("ControlRigAggregateNodeAddOutputPinButton_Tooltip", "Adds an output pin to the node"),
-		   !bInputAggregate);
-
-		FMargin AddPinPadding = bInputAggregate ? Settings->GetInputPinPadding() : Settings->GetOutputPinPadding();
+		TSharedPtr<SWidget> AddPinButton;
+		FMargin AddPinPadding = Settings->GetInputPinPadding();
 		AddPinPadding.Top += 2.0f;
-		AddPinPadding.Left -= bInputAggregate ? 2.f : 0.f;
-		AddPinPadding.Right -= bInputAggregate ? 0.f : 2.f;
+		AddPinPadding.Left -= 2.0f;
+		EHorizontalAlignment HorizontalAlignment = HAlign_Left;
+		
+		if(ModelNode->IsAggregate() || ModelNode->IsA<URigVMAggregateNode>())
+		{
+			const bool bInputAggregate = ModelNode->IsInputAggregate();
+			AddPinButton = AddPinButtonContent(
+			   LOCTEXT("ControlRigAggregateNodeAddPinButton", "Add pin"),
+			   bInputAggregate ? 
+				LOCTEXT("ControlRigAggregateNodeAddInputPinButton_Tooltip", "Adds an input pin to the node") :
+				LOCTEXT("ControlRigAggregateNodeAddOutputPinButton_Tooltip", "Adds an output pin to the node"),
+			   !bInputAggregate);
+		
+			AddPinPadding = bInputAggregate ? Settings->GetInputPinPadding() : Settings->GetOutputPinPadding();
+			AddPinPadding.Top += 2.0f;
+			AddPinPadding.Left -= bInputAggregate ? 2.f : 0.f;
+			AddPinPadding.Right -= bInputAggregate ? 0.f : 2.f;
 
-		LeftNodeBox->AddSlot()
-			.AutoHeight()
-			.VAlign(VAlign_Center)
-			.HAlign(bInputAggregate ? HAlign_Left : HAlign_Right)
-			.Padding(AddPinPadding)
-			[
-				AddPinButton
-			];
+			HorizontalAlignment = bInputAggregate ? HAlign_Left : HAlign_Right;
+		}
+
+		else if(ModelNode->GetPins().ContainsByPredicate([](const URigVMPin* Pin) { return Pin->IsFixedSizeArray(); }))
+		{
+			AddPinButton = AddPinButtonContent(
+			   LOCTEXT("ControlRigFixedArrayAddPinButton", "Add pin"),
+				LOCTEXT("ControlRigFixedArrayAddPinButton_Tooltip", "Adds an input pin to the node"),
+			   false);
+		}
+
+		if(AddPinButton.IsValid())
+		{
+			LeftNodeBox->AddSlot()
+				.AutoHeight()
+				.VAlign(VAlign_Center)
+				.HAlign(HorizontalAlignment)
+				.Padding(AddPinPadding)
+				[
+					AddPinButton.ToSharedRef()
+				];
+		}
 	}
 }
 
@@ -269,7 +291,22 @@ FReply SControlRigGraphNode::OnAddPin()
 	{
 		if (UControlRigGraphNode* ControlRigGraphNode = Cast<UControlRigGraphNode>(GraphNode))
 		{
-			ControlRigGraphNode->HandleAddAggregateElement(ModelNode->GetNodePath());
+			if (ModelNode->IsA<URigVMAggregateNode>())
+			{
+				ControlRigGraphNode->HandleAddAggregateElement(ModelNode->GetNodePath());
+			}
+			else
+			{
+				// we assume the model node has a fixed size array pin
+				for(URigVMPin* Pin : ModelNode->GetPins())
+				{
+					if(Pin->IsFixedSizeArray())
+					{
+						ControlRigGraphNode->HandleAddArrayElement(Pin->GetPinPath());
+						break;
+					}
+				}
+			}
 		}
 	}
 	return FReply::Handled();
@@ -1202,14 +1239,19 @@ int32 SControlRigGraphNode::GetNodeTopologyVersion() const
 	return INDEX_NONE;
 }
 
-EVisibility SControlRigGraphNode::GetPinVisibility(int32 InPinInfoIndex) const
+EVisibility SControlRigGraphNode::GetPinVisibility(int32 InPinInfoIndex, bool bAskingForSubPin) const
 {
 	if(PinInfos.IsValidIndex(InPinInfoIndex))
 	{
+		if(PinInfos[InPinInfoIndex].bFixedArray)
+		{
+			return bAskingForSubPin ? EVisibility::Visible : EVisibility::Collapsed;
+		}
+		
 		const int32 ParentPinIndex = PinInfos[InPinInfoIndex].ParentIndex;
 		if(ParentPinIndex != INDEX_NONE)
 		{
-			const EVisibility ParentPinVisibility = GetPinVisibility(ParentPinIndex); 
+			const EVisibility ParentPinVisibility = GetPinVisibility(ParentPinIndex, true); 
 			if(ParentPinVisibility != EVisibility::Visible)
 			{
 				return ParentPinVisibility;
@@ -1481,6 +1523,7 @@ void SControlRigGraphNode::UpdatePinTreeView()
 		PinInfo.bExpanded = ModelPin->IsExpanded();
 		PinInfo.ModelPinPath = ModelPin->GetPinPath();
 		PinInfo.bAutoHeight = false;
+		PinInfo.bFixedArray = ModelPin->IsFixedSizeArray();
 
 		if(!bSupportSubPins)
 		{
@@ -1515,7 +1558,7 @@ void SControlRigGraphNode::UpdatePinTreeView()
 			PinInfo.Depth = PinInfos[PinInfo.ParentIndex].Depth + 1;
 		}
 
-		TAttribute<EVisibility> PinVisibilityAttribute = TAttribute<EVisibility>::CreateSP(this, &SControlRigGraphNode::GetPinVisibility, PinInfo.Index);
+		TAttribute<EVisibility> PinVisibilityAttribute = TAttribute<EVisibility>::CreateSP(this, &SControlRigGraphNode::GetPinVisibility, PinInfo.Index, false);
 
 		bool bPinWidgetForExpanderLeft = false;
 		TSharedPtr<SGraphPin> PinWidgetForExpander;
@@ -1779,7 +1822,7 @@ void SControlRigGraphNode::UpdatePinTreeView()
 				.MaxHeight(InputPinInfo.bAutoHeight ? TAttribute<float>() : MaxHeight)
                 [
                     SAssignNew(SlotLayout, SHorizontalBox)
-                    .Visibility(this, &SControlRigGraphNode::GetPinVisibility, InputPinInfo.Index)
+                    .Visibility(this, &SControlRigGraphNode::GetPinVisibility, InputPinInfo.Index, false)
 		            
                     +SHorizontalBox::Slot()
                     .Expose(FirstSlot)
@@ -1817,7 +1860,7 @@ void SControlRigGraphNode::UpdatePinTreeView()
 				.MaxHeight(InputPinInfo.bAutoHeight ? TAttribute<float>() : MaxHeight)
                 [
                     SAssignNew(SlotLayout, SHorizontalBox)
-                    .Visibility(this, &SControlRigGraphNode::GetPinVisibility, OutputPinInfo.Index)
+                    .Visibility(this, &SControlRigGraphNode::GetPinVisibility, OutputPinInfo.Index, false)
 
                     +SHorizontalBox::Slot()
 					.Expose(FirstSlot)
@@ -1861,7 +1904,7 @@ void SControlRigGraphNode::UpdatePinTreeView()
 			.MaxHeight(OutputPinInfo.bAutoHeight ? TAttribute<float>() : MaxHeight)
             [
             	SNew(SHorizontalBox)
-	            .Visibility(this, &SControlRigGraphNode::GetPinVisibility, OutputPinInfo.Index)
+	            .Visibility(this, &SControlRigGraphNode::GetPinVisibility, OutputPinInfo.Index, false)
 	            
 				+SHorizontalBox::Slot()
 				.FillWidth(1.f)
@@ -1919,7 +1962,7 @@ void SControlRigGraphNode::UpdatePinTreeView()
 		}
 	}
 
-	CreateAggregateAddPinButton();
+	CreateAddPinButton();
 	
 	// add spacer widget at the end
 	LeftNodeBox->AddSlot()
diff --git a/Engine/Plugins/Animation/ControlRig/Source/ControlRigEditor/Private/Graph/SControlRigGraphNode.h b/Engine/Plugins/Animation/ControlRig/Source/ControlRigEditor/Private/Graph/SControlRigGraphNode.h
--- a/Engine/Plugins/Animation/ControlRig/Source/ControlRigEditor/Private/Graph/SControlRigGraphNode.h
+++ b/Engine/Plugins/Animation/ControlRig/Source/ControlRigEditor/Private/Graph/SControlRigGraphNode.h
@@ -57,7 +57,7 @@ public:
 	void UpdateStandardNode();
 	void UpdateCompactNode();
 
-	void CreateAggregateAddPinButton();
+	void CreateAddPinButton();
 
 	/** Callback function executed when Add pin button is clicked */
 	virtual FReply OnAddPin() override;
@@ -88,7 +88,7 @@ private:
 private:
 
 	int32 GetNodeTopologyVersion() const;
-	EVisibility GetPinVisibility(int32 InPinInfoIndex) const;
+	EVisibility GetPinVisibility(int32 InPinInfoIndex, bool bAskingForSubPin) const;
 	const FSlateBrush * GetExpanderImage(int32 InPinInfoIndex, bool bLeft, bool bHovered) const;
 	FReply OnExpanderArrowClicked(int32 InPinInfoIndex);
 	void HandleModifiedEvent(ERigVMGraphNotifType InNotifType, URigVMGraph* InGraph, UObject* InSubject);
@@ -126,6 +126,7 @@ private:
 		TSharedPtr<SGraphPin> OutputPinWidget;
 		bool bExpanded;
 		bool bAutoHeight;
+		bool bFixedArray;
 	};
 
 	
diff --git a/Engine/Source/Developer/RigVMDeveloper/Private/RigVMModel/RigVMPin.cpp b/Engine/Source/Developer/RigVMDeveloper/Private/RigVMModel/RigVMPin.cpp
--- a/Engine/Source/Developer/RigVMDeveloper/Private/RigVMModel/RigVMPin.cpp
+++ b/Engine/Source/Developer/RigVMDeveloper/Private/RigVMModel/RigVMPin.cpp
@@ -17,10 +17,12 @@
 #include "Misc/StringBuilder.h"
 #include "Logging/LogScopedVerbosityOverride.h"
 #include "RigVMModel/Nodes/RigVMCollapseNode.h"
+#include "RigVMModel/Nodes/RigVMDispatchNode.h"
 #include "RigVMModel/Nodes/RigVMFunctionReferenceNode.h"
 #include "RigVMModel/Nodes/RigVMFunctionEntryNode.h"
 #include "RigVMModel/Nodes/RigVMFunctionReturnNode.h"
 #include "RigVMModel/Nodes/RigVMInvokeEntryNode.h"
+#include "RigVMModel/Nodes/RigVMSelectNode.h"
 
 #include UE_INLINE_GENERATED_CPP_BY_NAME(RigVMPin)
 
@@ -567,6 +569,16 @@ bool URigVMPin::IsDynamicArray() const
 
 bool URigVMPin::IsLazy() const
 {
+	// fixed array elements are treated as lazy elements
+	// if the original argument is also marked as lazy
+	if(const URigVMPin* ParentPin = GetParentPin())
+	{
+		if(ParentPin->IsFixedSizeArray())
+		{
+			return ParentPin->IsLazy();
+		}
+	}
+	
 	if(GetDirection() == ERigVMPinDirection::Input)
 	{
 		if(const URigVMNode* Node = GetNode())
@@ -709,6 +721,40 @@ bool URigVMPin::ContainsWildCardSubPin() const
 	return false;
 }
 
+bool URigVMPin::IsFixedSizeArray() const
+{
+#if WITH_EDITOR
+
+	if(IsArray() && IsRootPin())
+	{
+		if(const URigVMNode* Node = GetNode())
+		{
+			if(const URigVMUnitNode* UnitNode = Cast<URigVMUnitNode>(Node))
+			{
+				if(const UScriptStruct* Struct = UnitNode->GetScriptStruct())
+				{
+					if(const FProperty* Property = Struct->FindPropertyByName(GetFName()))
+					{
+						return Property->HasMetaData(FRigVMStruct::FixedSizeArrayMetaName);
+					}
+				}
+			}
+			else if(const URigVMDispatchNode* DispatchNode = Cast<URigVMDispatchNode>(Node))
+			{
+				if(const FRigVMDispatchFactory* Factory = DispatchNode->GetFactory())
+				{
+					return Factory->HasArgumentMetaData(GetFName(), FRigVMStruct::FixedSizeArrayMetaName);
+				}
+			}
+			else if(Node->IsA<URigVMSelectNode>())
+			{
+				return GetFName().ToString() == URigVMSelectNode::ValueName;
+			}
+		}
+	}
+#endif
+	return false;
+}
 
 FString URigVMPin::GetDefaultValue() const
 {
diff --git a/Engine/Source/Developer/RigVMDeveloper/Public/RigVMModel/Nodes/RigVMSelectNode.h b/Engine/Source/Developer/RigVMDeveloper/Public/RigVMModel/Nodes/RigVMSelectNode.h
--- a/Engine/Source/Developer/RigVMDeveloper/Public/RigVMModel/Nodes/RigVMSelectNode.h
+++ b/Engine/Source/Developer/RigVMDeveloper/Public/RigVMModel/Nodes/RigVMSelectNode.h
@@ -52,5 +52,6 @@ private:
 	friend class FRigVMParserAST;
 	friend class FRigVMSelectExprAST;
 	friend struct FRigVMRemoveNodeAction;
+	friend class URigVMPin;
 };
 
diff --git a/Engine/Source/Developer/RigVMDeveloper/Public/RigVMModel/RigVMPin.h b/Engine/Source/Developer/RigVMDeveloper/Public/RigVMModel/RigVMPin.h
--- a/Engine/Source/Developer/RigVMDeveloper/Public/RigVMModel/RigVMPin.h
+++ b/Engine/Source/Developer/RigVMDeveloper/Public/RigVMModel/RigVMPin.h
@@ -271,6 +271,10 @@ public:
 	UFUNCTION(BlueprintCallable, Category = RigVMPin)
 	bool ContainsWildCardSubPin() const;
 
+	// Returns true if this pin is an array that should be displayed as elements only
+	UFUNCTION(BlueprintCallable, Category = RigVMPin)
+	bool IsFixedSizeArray() const;
+	
 	// Returns the default value of the Pin as a string.
 	// Note that this value is computed based on the Pin's
 	// SubPins - so for example for a FVector typed Pin
diff --git a/Engine/Source/Runtime/RigVM/Private/RigVMCore/RigVMStruct.cpp b/Engine/Source/Runtime/RigVM/Private/RigVMCore/RigVMStruct.cpp
--- a/Engine/Source/Runtime/RigVM/Private/RigVMCore/RigVMStruct.cpp
+++ b/Engine/Source/Runtime/RigVM/Private/RigVMCore/RigVMStruct.cpp
@@ -294,6 +294,7 @@ const FName FRigVMStruct::NodeColorMetaName("NodeColor");
 const FName FRigVMStruct::IconMetaName("Icon");
 const FName FRigVMStruct::KeywordsMetaName("Keywords");
 const FName FRigVMStruct::TemplateNameMetaName = FRigVMRegistry::TemplateNameMetaName;
+const FName FRigVMStruct::FixedSizeArrayMetaName("FixedSizeArray");
 const FName FRigVMStruct::AggregateMetaName("Aggregate");
 const FName FRigVMStruct::ExpandPinByDefaultMetaName("ExpandByDefault");
 const FName FRigVMStruct::DefaultArraySizeMetaName("DefaultArraySize");
diff --git a/Engine/Source/Runtime/RigVM/Public/RigVMCore/RigVMStruct.h b/Engine/Source/Runtime/RigVM/Public/RigVMCore/RigVMStruct.h
--- a/Engine/Source/Runtime/RigVM/Public/RigVMCore/RigVMStruct.h
+++ b/Engine/Source/Runtime/RigVM/Public/RigVMCore/RigVMStruct.h
@@ -240,6 +240,7 @@ public:
 	static const FName IconMetaName;
 	static const FName KeywordsMetaName;
 	static const FName TemplateNameMetaName;
+	static const FName FixedSizeArrayMetaName;
 	static const FName AggregateMetaName;
 	static const FName ExpandPinByDefaultMetaName;
 	static const FName DefaultArraySizeMetaName;
diff --git a/Engine/Source/Runtime/RigVM/Public/RigVMCore/RigVMTemplate.h b/Engine/Source/Runtime/RigVM/Public/RigVMCore/RigVMTemplate.h
--- a/Engine/Source/Runtime/RigVM/Public/RigVMCore/RigVMTemplate.h
+++ b/Engine/Source/Runtime/RigVM/Public/RigVMCore/RigVMTemplate.h
@@ -85,6 +85,11 @@ struct RIGVM_API FRigVMTemplateArgumentType
 	{
 	}
 
+	bool IsValid() const
+	{
+		return !CPPType.IsNone();
+	}
+
 	static FRigVMTemplateArgumentType Array()
 	{
 		return FRigVMTemplateArgumentType(RigVMTypeUtils::GetWildCardArrayCPPTypeName(), RigVMTypeUtils::GetWildCardCPPTypeObject());
