Sequencer: Add filter button to content browsers in Sequencer

diff --git a/Engine/Plugins/Animation/ControlRig/Source/ControlRigEditor/Private/Sequencer/ControlRigParameterTrackEditor.cpp b/Engine/Plugins/Animation/ControlRig/Source/ControlRigEditor/Private/Sequencer/ControlRigParameterTrackEditor.cpp
--- a/Engine/Plugins/Animation/ControlRig/Source/ControlRigEditor/Private/Sequencer/ControlRigParameterTrackEditor.cpp
+++ b/Engine/Plugins/Animation/ControlRig/Source/ControlRigEditor/Private/Sequencer/ControlRigParameterTrackEditor.cpp
@@ -4728,6 +4728,7 @@ void FControlRigParameterSection::AddAnimationSubMenuForFK(FMenuBuilder& MenuBui
 		AssetPickerConfig.OnAssetSelected = FOnAssetSelected::CreateRaw(this, &FControlRigParameterSection::OnAnimationAssetSelectedForFK, ObjectBinding, Section);
 		AssetPickerConfig.OnAssetEnterPressed = FOnAssetEnterPressed::CreateRaw(this, &FControlRigParameterSection::OnAnimationAssetEnterPressedForFK, ObjectBinding, Section);
 		AssetPickerConfig.bAllowNullSelection = false;
+		AssetPickerConfig.bAddFilterUI = true;
 		AssetPickerConfig.InitialAssetViewType = EAssetViewType::List;
 		AssetPickerConfig.OnShouldFilterAsset = FOnShouldFilterAsset::CreateRaw(this, &FControlRigParameterSection::ShouldFilterAssetForFK);
 		AssetPickerConfig.Filter.bRecursiveClasses = true;
diff --git a/Engine/Plugins/Media/MediaCompositing/Source/MediaCompositingEditor/Private/Sequencer/MediaPlayerPropertyTrackEditor.cpp b/Engine/Plugins/Media/MediaCompositing/Source/MediaCompositingEditor/Private/Sequencer/MediaPlayerPropertyTrackEditor.cpp
--- a/Engine/Plugins/Media/MediaCompositing/Source/MediaCompositingEditor/Private/Sequencer/MediaPlayerPropertyTrackEditor.cpp
+++ b/Engine/Plugins/Media/MediaCompositing/Source/MediaCompositingEditor/Private/Sequencer/MediaPlayerPropertyTrackEditor.cpp
@@ -154,6 +154,7 @@ TSharedPtr<SWidget> FMediaPlayerPropertyTrackEditor::BuildOutlinerEditWidget(con
 			AssetPickerConfig.OnAssetSelected     = FOnAssetSelected::CreateRaw(this,     &FMediaPlayerPropertyTrackEditor::AddNewSection,             MediaTrack);
 			AssetPickerConfig.OnAssetEnterPressed = FOnAssetEnterPressed::CreateRaw(this, &FMediaPlayerPropertyTrackEditor::AddNewSectionEnterPressed, MediaTrack);
 			AssetPickerConfig.bAllowNullSelection = false;
+			AssetPickerConfig.bAddFilterUI = true;
 			AssetPickerConfig.InitialAssetViewType = EAssetViewType::List;
 			AssetPickerConfig.Filter.bRecursiveClasses = true;
 			AssetPickerConfig.Filter.ClassPaths.Add(UMediaSource::StaticClass()->GetClassPathName());
diff --git a/Engine/Plugins/Media/MediaCompositing/Source/MediaCompositingEditor/Private/Sequencer/MediaTrackEditor.cpp b/Engine/Plugins/Media/MediaCompositing/Source/MediaCompositingEditor/Private/Sequencer/MediaTrackEditor.cpp
--- a/Engine/Plugins/Media/MediaCompositing/Source/MediaCompositingEditor/Private/Sequencer/MediaTrackEditor.cpp
+++ b/Engine/Plugins/Media/MediaCompositing/Source/MediaCompositingEditor/Private/Sequencer/MediaTrackEditor.cpp
@@ -103,6 +103,7 @@ TSharedPtr<SWidget> FMediaTrackEditor::BuildOutlinerEditWidget(const FGuid& Obje
 			AssetPickerConfig.OnAssetSelected = FOnAssetSelected::CreateRaw(this, &FMediaTrackEditor::AddNewSection, MediaTrack);
 			AssetPickerConfig.OnAssetEnterPressed = FOnAssetEnterPressed::CreateRaw(this, &FMediaTrackEditor::AddNewSectionEnterPressed, MediaTrack);
 			AssetPickerConfig.bAllowNullSelection = false;
+			AssetPickerConfig.bAddFilterUI = true;
 			AssetPickerConfig.InitialAssetViewType = EAssetViewType::List;
 			AssetPickerConfig.Filter.bRecursiveClasses = true;
 			AssetPickerConfig.Filter.ClassPaths.Add(UMediaSource::StaticClass()->GetClassPathName());
diff --git a/Engine/Plugins/MovieScene/TemplateSequence/Source/TemplateSequenceEditor/Private/TrackEditors/TemplateSequenceTrackEditor.cpp b/Engine/Plugins/MovieScene/TemplateSequence/Source/TemplateSequenceEditor/Private/TrackEditors/TemplateSequenceTrackEditor.cpp
--- a/Engine/Plugins/MovieScene/TemplateSequence/Source/TemplateSequenceEditor/Private/TrackEditors/TemplateSequenceTrackEditor.cpp
+++ b/Engine/Plugins/MovieScene/TemplateSequence/Source/TemplateSequenceEditor/Private/TrackEditors/TemplateSequenceTrackEditor.cpp
@@ -136,6 +136,7 @@ public:
 			AssetPickerConfig.OnAssetEnterPressed = FOnAssetEnterPressed::CreateSP(TrackEditorRef, &FTemplateSequenceTrackEditor::OnTemplateSequenceAssetEnterPressed, ObjectBindings);
 			AssetPickerConfig.RefreshAssetViewDelegates.Add(&RefreshAssetViewDelegate);
 			AssetPickerConfig.bAllowNullSelection = false;
+			AssetPickerConfig.bAddFilterUI = true;
 			AssetPickerConfig.InitialAssetViewType = EAssetViewType::List;
 			AssetPickerConfig.Filter.bRecursiveClasses = true;
 			AssetPickerConfig.Filter.ClassPaths.Add(UTemplateSequence::StaticClass()->GetClassPathName());
diff --git a/Engine/Plugins/VirtualProduction/DMX/DMXEngine/Source/DMXEditor/Private/Sequencer/DMXLibraryTrackEditor.cpp b/Engine/Plugins/VirtualProduction/DMX/DMXEngine/Source/DMXEditor/Private/Sequencer/DMXLibraryTrackEditor.cpp
--- a/Engine/Plugins/VirtualProduction/DMX/DMXEngine/Source/DMXEditor/Private/Sequencer/DMXLibraryTrackEditor.cpp
+++ b/Engine/Plugins/VirtualProduction/DMX/DMXEngine/Source/DMXEditor/Private/Sequencer/DMXLibraryTrackEditor.cpp
@@ -46,6 +46,7 @@ TSharedRef<SWidget> CreateAssetPicker(FOnAssetSelected OnAssetSelected, FOnAsset
 		AssetPickerConfig.OnAssetSelected = OnAssetSelected;
 		AssetPickerConfig.OnAssetEnterPressed = OnAssetEnterPressed;
 		AssetPickerConfig.bAllowNullSelection = false;
+		AssetPickerConfig.bAddFilterUI = true;
 		AssetPickerConfig.InitialAssetViewType = EAssetViewType::List;
 		AssetPickerConfig.Filter.bRecursiveClasses = true;
 		AssetPickerConfig.Filter.ClassPaths.Add(UDMXLibrary::StaticClass()->GetClassPathName());
diff --git a/Engine/Source/Editor/MovieSceneTools/Private/TrackEditors/AudioTrackEditor.cpp b/Engine/Source/Editor/MovieSceneTools/Private/TrackEditors/AudioTrackEditor.cpp
--- a/Engine/Source/Editor/MovieSceneTools/Private/TrackEditors/AudioTrackEditor.cpp
+++ b/Engine/Source/Editor/MovieSceneTools/Private/TrackEditors/AudioTrackEditor.cpp
@@ -1295,6 +1295,7 @@ TSharedRef<SWidget> FAudioTrackEditor::BuildAudioSubMenu(FOnAssetSelected OnAsse
 		AssetPickerConfig.OnAssetSelected = OnAssetSelected;
 		AssetPickerConfig.OnAssetEnterPressed = OnAssetEnterPressed;
 		AssetPickerConfig.bAllowNullSelection = false;
+		AssetPickerConfig.bAddFilterUI = true;
 		AssetPickerConfig.InitialAssetViewType = EAssetViewType::List;
 		for (FTopLevelAssetPath ClassName : DerivedClassNames)
 		{
diff --git a/Engine/Source/Editor/MovieSceneTools/Private/TrackEditors/CameraShakeSourceShakeTrackEditor.cpp b/Engine/Source/Editor/MovieSceneTools/Private/TrackEditors/CameraShakeSourceShakeTrackEditor.cpp
--- a/Engine/Source/Editor/MovieSceneTools/Private/TrackEditors/CameraShakeSourceShakeTrackEditor.cpp
+++ b/Engine/Source/Editor/MovieSceneTools/Private/TrackEditors/CameraShakeSourceShakeTrackEditor.cpp
@@ -486,6 +486,7 @@ void FCameraShakeSourceShakeTrackEditor::AddOtherCameraShakeBrowserSubMenu(FMenu
 		AssetPickerConfig.OnAssetEnterPressed = FOnAssetEnterPressed::CreateSP(this, &FCameraShakeSourceShakeTrackEditor::OnCameraShakeAssetEnterPressed, ObjectBindings);
 		AssetPickerConfig.OnShouldFilterAsset = FOnShouldFilterAsset::CreateSP(this, &FCameraShakeSourceShakeTrackEditor::OnShouldFilterCameraShake);
 		AssetPickerConfig.bAllowNullSelection = false;
+		AssetPickerConfig.bAddFilterUI = true;
 		AssetPickerConfig.InitialAssetViewType = EAssetViewType::List;
 		AssetPickerConfig.Filter.ClassPaths.Add(UBlueprint::StaticClass()->GetClassPathName());
 		AssetPickerConfig.SaveSettingsName = TEXT("SequencerAssetPicker");
diff --git a/Engine/Source/Editor/MovieSceneTools/Private/TrackEditors/CameraShakeTrackEditor.cpp b/Engine/Source/Editor/MovieSceneTools/Private/TrackEditors/CameraShakeTrackEditor.cpp
--- a/Engine/Source/Editor/MovieSceneTools/Private/TrackEditors/CameraShakeTrackEditor.cpp
+++ b/Engine/Source/Editor/MovieSceneTools/Private/TrackEditors/CameraShakeTrackEditor.cpp
@@ -166,6 +166,7 @@ void FCameraShakeTrackEditor::AddCameraShakeSubMenu(FMenuBuilder& MenuBuilder, T
 		AssetPickerConfig.OnAssetSelected = FOnAssetSelected::CreateRaw(this, &FCameraShakeTrackEditor::OnCameraShakeAssetSelected, ObjectBindings);
 		AssetPickerConfig.OnAssetEnterPressed = FOnAssetEnterPressed::CreateRaw(this, &FCameraShakeTrackEditor::OnCameraShakeAssetEnterPressed, ObjectBindings);
 		AssetPickerConfig.bAllowNullSelection = false;
+		AssetPickerConfig.bAddFilterUI = true;
 		AssetPickerConfig.InitialAssetViewType = EAssetViewType::List;
 		AssetPickerConfig.Filter.ClassPaths.Add(UBlueprint::StaticClass()->GetClassPathName());
 		AssetPickerConfig.SaveSettingsName = TEXT("SequencerAssetPicker");
diff --git a/Engine/Source/Editor/MovieSceneTools/Private/TrackEditors/CinematicShotTrackEditor.cpp b/Engine/Source/Editor/MovieSceneTools/Private/TrackEditors/CinematicShotTrackEditor.cpp
--- a/Engine/Source/Editor/MovieSceneTools/Private/TrackEditors/CinematicShotTrackEditor.cpp
+++ b/Engine/Source/Editor/MovieSceneTools/Private/TrackEditors/CinematicShotTrackEditor.cpp
@@ -626,6 +626,7 @@ TSharedRef<SWidget> FCinematicShotTrackEditor::HandleAddCinematicShotComboButton
 		AssetPickerConfig.OnAssetSelected = FOnAssetSelected::CreateRaw( this, &FCinematicShotTrackEditor::HandleAddCinematicShotComboButtonMenuEntryExecute);
 		AssetPickerConfig.OnAssetEnterPressed = FOnAssetEnterPressed::CreateRaw( this, &FCinematicShotTrackEditor::HandleAddCinematicShotComboButtonMenuEntryEnterPressed);
 		AssetPickerConfig.bAllowNullSelection = false;
+		AssetPickerConfig.bAddFilterUI = true;
 		AssetPickerConfig.InitialAssetViewType = EAssetViewType::List;
 		AssetPickerConfig.Filter.ClassPaths.Add(FTopLevelAssetPath(TEXT("/Script/LevelSequence"), TEXT("LevelSequence")));
 		AssetPickerConfig.SaveSettingsName = TEXT("SequencerAssetPicker");
diff --git a/Engine/Source/Editor/MovieSceneTools/Private/TrackEditors/MaterialParameterCollectionTrackEditor.cpp b/Engine/Source/Editor/MovieSceneTools/Private/TrackEditors/MaterialParameterCollectionTrackEditor.cpp
--- a/Engine/Source/Editor/MovieSceneTools/Private/TrackEditors/MaterialParameterCollectionTrackEditor.cpp
+++ b/Engine/Source/Editor/MovieSceneTools/Private/TrackEditors/MaterialParameterCollectionTrackEditor.cpp
@@ -83,6 +83,7 @@ TSharedRef<SWidget> CreateAssetPicker(FOnAssetSelected OnAssetSelected, FOnAsset
 		AssetPickerConfig.OnAssetSelected = OnAssetSelected;
 		AssetPickerConfig.OnAssetEnterPressed = OnAssetEnterPressed;
 		AssetPickerConfig.bAllowNullSelection = false;
+		AssetPickerConfig.bAddFilterUI = true;
 		AssetPickerConfig.InitialAssetViewType = EAssetViewType::List;
 		AssetPickerConfig.Filter.bRecursiveClasses = true;
 		AssetPickerConfig.Filter.ClassPaths.Add(UMaterialParameterCollection::StaticClass()->GetClassPathName());
diff --git a/Engine/Source/Editor/MovieSceneTools/Private/TrackEditors/SkeletalAnimationTrackEditor.cpp b/Engine/Source/Editor/MovieSceneTools/Private/TrackEditors/SkeletalAnimationTrackEditor.cpp
--- a/Engine/Source/Editor/MovieSceneTools/Private/TrackEditors/SkeletalAnimationTrackEditor.cpp
+++ b/Engine/Source/Editor/MovieSceneTools/Private/TrackEditors/SkeletalAnimationTrackEditor.cpp
@@ -1738,6 +1738,7 @@ void FSkeletalAnimationTrackEditor::AddAnimationSubMenu(FMenuBuilder& MenuBuilde
 		AssetPickerConfig.OnAssetSelected = FOnAssetSelected::CreateRaw( this, &FSkeletalAnimationTrackEditor::OnAnimationAssetSelected, ObjectBindings, Track);
 		AssetPickerConfig.OnAssetEnterPressed = FOnAssetEnterPressed::CreateRaw( this, &FSkeletalAnimationTrackEditor::OnAnimationAssetEnterPressed, ObjectBindings, Track);
 		AssetPickerConfig.bAllowNullSelection = false;
+		AssetPickerConfig.bAddFilterUI = true;
 		AssetPickerConfig.InitialAssetViewType = EAssetViewType::List;
 		AssetPickerConfig.Filter.bRecursiveClasses = true;
 		AssetPickerConfig.Filter.ClassPaths.Add(UAnimSequenceBase::StaticClass()->GetClassPathName());
diff --git a/Engine/Source/Editor/MovieSceneTools/Private/TrackEditors/SubTrackEditor.cpp b/Engine/Source/Editor/MovieSceneTools/Private/TrackEditors/SubTrackEditor.cpp
--- a/Engine/Source/Editor/MovieSceneTools/Private/TrackEditors/SubTrackEditor.cpp
+++ b/Engine/Source/Editor/MovieSceneTools/Private/TrackEditors/SubTrackEditor.cpp
@@ -511,6 +511,7 @@ TSharedRef<SWidget> FSubTrackEditor::HandleAddSubSequenceComboButtonGetMenuConte
 			AssetPickerConfig.OnAssetSelected = FOnAssetSelected::CreateRaw( this, &FSubTrackEditor::HandleAddSubSequenceComboButtonMenuEntryExecute, InTrack);
 			AssetPickerConfig.OnAssetEnterPressed = FOnAssetEnterPressed::CreateRaw( this, &FSubTrackEditor::HandleAddSubSequenceComboButtonMenuEntryEnterPressed, InTrack);
 			AssetPickerConfig.bAllowNullSelection = false;
+			AssetPickerConfig.bAddFilterUI = true;
 			AssetPickerConfig.InitialAssetViewType = EAssetViewType::List;
 			GetSupportedSequenceClassPaths(AssetPickerConfig.Filter.ClassPaths);
 			AssetPickerConfig.SaveSettingsName = TEXT("SequencerAssetPicker");
diff --git a/Engine/Source/Editor/MovieSceneTools/Private/TrackEditors/TransformTrackEditor.cpp b/Engine/Source/Editor/MovieSceneTools/Private/TrackEditors/TransformTrackEditor.cpp
--- a/Engine/Source/Editor/MovieSceneTools/Private/TrackEditors/TransformTrackEditor.cpp
+++ b/Engine/Source/Editor/MovieSceneTools/Private/TrackEditors/TransformTrackEditor.cpp
@@ -129,6 +129,7 @@ void F3DTransformTrackEditor::BuildTrackContextMenu( FMenuBuilder& MenuBuilder,
 		UMovieSceneSequence* Sequence = InSequencer->GetFocusedMovieSceneSequence();
 
 		FAssetPickerConfig AssetPickerConfig;
+		AssetPickerConfig.bAddFilterUI = true;
 		AssetPickerConfig.SelectionMode = ESelectionMode::Single;
 		AssetPickerConfig.Filter.ClassPaths.Add(UAnimSequence::StaticClass()->GetClassPathName());
 		AssetPickerConfig.OnAssetSelected = FOnAssetSelected::CreateStatic(&F3DTransformTrackEditor::ImportAnimSequenceTransforms, InSequencer, InTransformTrack);
