Control Rig: Fix crash in CitySample

diff --git a/Engine/Plugins/Animation/ControlRig/Source/ControlRigDeveloper/Private/ControlRigBlueprint.cpp b/Engine/Plugins/Animation/ControlRig/Source/ControlRigDeveloper/Private/ControlRigBlueprint.cpp
--- a/Engine/Plugins/Animation/ControlRig/Source/ControlRigDeveloper/Private/ControlRigBlueprint.cpp
+++ b/Engine/Plugins/Animation/ControlRig/Source/ControlRigDeveloper/Private/ControlRigBlueprint.cpp
@@ -905,7 +905,7 @@ void UControlRigBlueprint::PostLoad()
 						{
 							if(URigVMController* Controller = GetController(GraphToValidate))
 							{
-								FRigVMControllerNotifGuard(Controller, true);
+								FRigVMControllerNotifGuard NotifGuard(Controller, true);
 								Controller->SetPinDefaultValue(Pin->GetPinPath(), TEXT("Null"), false, false, false);
 							}
 						}
@@ -1307,7 +1307,7 @@ void UControlRigBlueprint::RefreshAllModels(EControlRigBlueprintLoadType InLoadT
 			URigVMController* Controller = GetOrCreateController(GraphToDetach);
 			// temporarily disable default value validation during load time, serialized values should always be accepted
 			TGuardValue<bool> PerGraphDisablePinDefaultValueValidation(Controller->bValidatePinDefaults, false);
-			FRigVMControllerNotifGuard(Controller, true);
+			FRigVMControllerNotifGuard NotifGuard(Controller, true);
 			Controller->DetachLinksFromPinObjects();
 			TArray<URigVMNode*> Nodes = GraphToDetach->GetNodes();
 			for (URigVMNode* Node : Nodes)
@@ -1323,7 +1323,7 @@ void UControlRigBlueprint::RefreshAllModels(EControlRigBlueprintLoadType InLoadT
 		URigVMController* Controller = GetOrCreateController(GraphToDetach);
 		// at this stage, allow all links to be reattached,
 		// RecomputeAllTemplateFilteredPermutations() later should break any invalid links
-		FRigVMControllerNotifGuard(Controller, true);
+		FRigVMControllerNotifGuard NotifGuard(Controller, true);
 		Controller->ReattachLinksToPinObjects(true /* follow redirectors */, nullptr, true, true);
 	}
 
@@ -1396,7 +1396,7 @@ void UControlRigBlueprint::RefreshAllModels(EControlRigBlueprintLoadType InLoadT
 		URigVMController* Controller = GetOrCreateController(GraphToClean);
 		TGuardValue<bool> RecomputeGuard(Controller->bSuspendRecomputingOuterTemplateFilters, true);
 		TGuardValue<bool> GuardEditGraph(GraphToClean->bEditable, true);
-		FRigVMControllerNotifGuard(Controller, true);
+		FRigVMControllerNotifGuard NotifGuard(Controller, true);
 		
 		for(URigVMNode* ModelNode : GraphToClean->GetNodes())
 		{
@@ -4098,7 +4098,7 @@ void UControlRigBlueprint::PatchVariableNodesOnLoad()
 
 		{
 			URigVMController* Controller = GetOrCreateController();
-			FRigVMControllerNotifGuard(Controller, true);
+			FRigVMControllerNotifGuard NotifGuard(Controller, true);
 			Controller->ReattachLinksToPinObjects();
 		}
 
@@ -4154,7 +4154,7 @@ void UControlRigBlueprint::PatchRigElementKeyCacheOnLoad()
 		{
 			URigVMController* Controller = GetOrCreateController(Graph);
 			TGuardValue<bool> DisablePinDefaultValueValidation(Controller->bValidatePinDefaults, false);
-			FRigVMControllerNotifGuard(Controller, true);
+			FRigVMControllerNotifGuard NotifGuard(Controller, true);
 			for (URigVMNode* Node : Graph->GetNodes())
 			{
 				if (URigVMUnitNode* UnitNode = Cast<URigVMUnitNode>(Node))
@@ -4365,7 +4365,7 @@ void UControlRigBlueprint::PatchPropagateToChildren()
 		
 		if (URigVMController* Controller = GetOrCreateController(Graph))
 		{
-			FRigVMControllerNotifGuard(Controller, true);
+			FRigVMControllerNotifGuard NotifGuard(Controller, true);
 			for (const URigVMPin* Pin: PinsToUpdate)
 			{
 				Controller->SetPinDefaultValue(Pin->GetPinPath(), TEXT("True"), false, false, false);
@@ -4386,7 +4386,7 @@ void UControlRigBlueprint::PatchParameterNodesOnLoad()
 
 		{
 			URigVMController* Controller = GetOrCreateController();
-			FRigVMControllerNotifGuard(Controller, true);
+			FRigVMControllerNotifGuard NotifGuard(Controller, true);
 			Controller->ReattachLinksToPinObjects();
 		}
 
diff --git a/Engine/Plugins/Animation/ControlRig/Source/ControlRigEditor/Private/Editor/ControlRigEditor.cpp b/Engine/Plugins/Animation/ControlRig/Source/ControlRigEditor/Private/Editor/ControlRigEditor.cpp
--- a/Engine/Plugins/Animation/ControlRig/Source/ControlRigEditor/Private/Editor/ControlRigEditor.cpp
+++ b/Engine/Plugins/Animation/ControlRig/Source/ControlRigEditor/Private/Editor/ControlRigEditor.cpp
@@ -5005,7 +5005,7 @@ void FControlRigEditor::OnWrappedPropertyChangedChainEvent(UDetailsViewWrapperOb
 			}
 			else if (PropertyPath == TEXT("DefaultValue"))
 			{
-				FRigVMControllerNotifGuard(Controller, true);
+				FRigVMControllerNotifGuard NotifGuard(Controller, true);
 				for (FRigVMGraphVariableDescription& Variable : Graph->GetLocalVariables())
 				{
 					if (Variable.Name == VariableDescription.Name)
diff --git a/Engine/Plugins/Animation/ControlRig/Source/ControlRigEditor/Private/Graph/NodeSpawners/ControlRigFunctionRefNodeSpawner.cpp b/Engine/Plugins/Animation/ControlRig/Source/ControlRigEditor/Private/Graph/NodeSpawners/ControlRigFunctionRefNodeSpawner.cpp
--- a/Engine/Plugins/Animation/ControlRig/Source/ControlRigEditor/Private/Graph/NodeSpawners/ControlRigFunctionRefNodeSpawner.cpp
+++ b/Engine/Plugins/Animation/ControlRig/Source/ControlRigEditor/Private/Graph/NodeSpawners/ControlRigFunctionRefNodeSpawner.cpp
@@ -385,7 +385,7 @@ UControlRigGraphNode* UControlRigFunctionRefNodeSpawner::SpawnNode(UEdGraph* Par
 			{
 				// If the package is a template, do not remove the ControlRigGraphNode
 				// We might be spawning a node to populate the PROTO_ context menu for function declarations.
-				FRigVMControllerNotifGuard(Controller, true);
+				FRigVMControllerNotifGuard NotifGuard(Controller, true);
 				Controller->RemoveNode(ModelNode, false);
 			}
 
