RigVM: Deal with duplicatoin of function graph store assets

diff --git a/Engine/Plugins/Animation/ControlRig/Source/ControlRigDeveloper/Private/ControlRigBlueprint.cpp b/Engine/Plugins/Animation/ControlRig/Source/ControlRigDeveloper/Private/ControlRigBlueprint.cpp
--- a/Engine/Plugins/Animation/ControlRig/Source/ControlRigDeveloper/Private/ControlRigBlueprint.cpp
+++ b/Engine/Plugins/Animation/ControlRig/Source/ControlRigDeveloper/Private/ControlRigBlueprint.cpp
@@ -107,6 +107,8 @@ bool FControlRigPublicFunctionData::IsMutable() const
 }
 
 TArray<UControlRigBlueprint*> UControlRigBlueprint::sCurrentlyOpenedRigBlueprints;
+FSoftObjectPath UControlRigBlueprint::PreDuplicateAssetPath;
+FSoftObjectPath UControlRigBlueprint::PreDuplicateHostPath;
 
 UControlRigBlueprint::UControlRigBlueprint(const FObjectInitializer& ObjectInitializer)
 {
@@ -2629,6 +2631,21 @@ void UControlRigBlueprint::ReplaceDeprecatedNodes()
 	Super::ReplaceDeprecatedNodes();
 }
 
+void UControlRigBlueprint::PreDuplicate(FObjectDuplicationParameters& DupParams)
+{
+	Super::PreDuplicate(DupParams);
+	PreDuplicateAssetPath = GetPathName();
+	if(UControlRigBlueprintGeneratedClass* CRGeneratedClass = GetControlRigBlueprintGeneratedClass())
+	{
+		PreDuplicateHostPath = FSoftObjectPath(CRGeneratedClass->GetPathName());
+	}
+	else
+	{
+		PreDuplicateHostPath.Reset();
+	}
+
+}
+
 void UControlRigBlueprint::PostDuplicate(bool bDuplicateForPIE)
 {
 	RefreshAllModels();
@@ -2641,6 +2658,33 @@ void UControlRigBlueprint::PostDuplicate(bool bDuplicateForPIE)
 		Controller->OnModified().AddUObject(this, &UControlRigBlueprint::HandleHierarchyModified);
 	}
 
+	auto UpdateFunctionHeaders = [this](const FString& InOldPath, const FString& InNewPath)
+	{
+		if(InOldPath.IsEmpty() || InNewPath.IsEmpty())
+		{
+			return;
+		}
+		if(!InNewPath.Equals(InOldPath, ESearchCase::CaseSensitive))
+		{
+			if(UControlRigBlueprintGeneratedClass* CRGeneratedClass = GetControlRigBlueprintGeneratedClass())
+			{
+				FRigVMGraphFunctionStore& Store = CRGeneratedClass->GraphFunctionStore;
+				Store.PostDuplicateHost(InOldPath, InNewPath);
+			}
+			RigVMClient.PostDuplicateHost(InOldPath, InNewPath);
+		}
+	};
+
+	// update the paths once for the blueprint and once for the generated class
+	UpdateFunctionHeaders(PreDuplicateAssetPath.ToString(), GetPathName());
+	if(const UControlRigBlueprintGeneratedClass* CRGeneratedClass = GetControlRigBlueprintGeneratedClass())
+	{
+		UpdateFunctionHeaders(PreDuplicateHostPath.ToString(), CRGeneratedClass->GetPathName());
+	}
+
+	PreDuplicateAssetPath.Reset();
+	PreDuplicateHostPath.Reset();
+
 	FBlueprintEditorUtils::MarkBlueprintAsStructurallyModified(this);
 	RecompileVM();
 }
diff --git a/Engine/Plugins/Animation/ControlRig/Source/ControlRigDeveloper/Public/ControlRigBlueprint.h b/Engine/Plugins/Animation/ControlRig/Source/ControlRigDeveloper/Public/ControlRigBlueprint.h
--- a/Engine/Plugins/Animation/ControlRig/Source/ControlRigDeveloper/Public/ControlRigBlueprint.h
+++ b/Engine/Plugins/Animation/ControlRig/Source/ControlRigDeveloper/Public/ControlRigBlueprint.h
@@ -243,6 +243,7 @@ public:
 	virtual bool IsPostLoadThreadSafe() const override { return false; }
 	virtual void PostTransacted(const FTransactionObjectEvent& TransactionEvent) override;
 	virtual void ReplaceDeprecatedNodes() override;
+	virtual void PreDuplicate(FObjectDuplicationParameters& DupParams) override;
 	virtual void PostDuplicate(bool bDuplicateForPIE) override;
 
 	virtual bool SupportsGlobalVariables() const override { return true; }
@@ -782,6 +783,11 @@ public:
 
 	static constexpr TCHAR RigVMModelPrefix[] = TEXT("RigVMModel");
 
+protected:
+
+	static FSoftObjectPath PreDuplicateAssetPath;
+	static FSoftObjectPath PreDuplicateHostPath;
+
 private:
 	bool bDirtyDuringLoad;
 	bool bErrorsDuringCompilation;
diff --git a/Engine/Source/Developer/RigVMDeveloper/Private/RigVMModel/RigVMClient.cpp b/Engine/Source/Developer/RigVMDeveloper/Private/RigVMModel/RigVMClient.cpp
--- a/Engine/Source/Developer/RigVMDeveloper/Private/RigVMModel/RigVMClient.cpp
+++ b/Engine/Source/Developer/RigVMDeveloper/Private/RigVMModel/RigVMClient.cpp
@@ -671,6 +671,16 @@ void FRigVMClient::PatchModelsOnLoad()
 	}
 }
 
+void FRigVMClient::PostDuplicateHost(const FString& InOldPathName, const FString& InNewPathName)
+{
+	TArray<URigVMGraph*> AllModels = GetAllModels(true, true);
+	for(URigVMGraph* Model : AllModels)
+	{
+		URigVMController* Controller = GetOrCreateController(Model);
+		Controller->PostDuplicateHost(InOldPathName, InNewPathName);
+	}
+}
+
 void FRigVMClient::HandleGraphModifiedEvent(ERigVMGraphNotifType InNotifType, URigVMGraph* InGraph, UObject* InSubject)
 {
 	if (bIgnoreModelNotifications)
diff --git a/Engine/Source/Developer/RigVMDeveloper/Private/RigVMModel/RigVMController.cpp b/Engine/Source/Developer/RigVMDeveloper/Private/RigVMModel/RigVMController.cpp
--- a/Engine/Source/Developer/RigVMDeveloper/Private/RigVMModel/RigVMController.cpp
+++ b/Engine/Source/Developer/RigVMDeveloper/Private/RigVMModel/RigVMController.cpp
@@ -19753,6 +19753,23 @@ void URigVMController::PatchIfSelectNodesOnLoad()
 	}
 }
 
+void URigVMController::PostDuplicateHost(const FString& InOldPathName, const FString& InNewPathName)
+{
+	if (const URigVMGraph* Graph = GetGraph())
+	{
+		TArray<URigVMNode*> FunctionRefNodes = Graph->GetNodes().FilterByPredicate([](URigVMNode* Node)
+		{
+			return Node->IsA<URigVMFunctionReferenceNode>();
+		});
+
+		for(URigVMNode* Node : FunctionRefNodes)
+		{
+			URigVMFunctionReferenceNode* FunctionReferenceNode = CastChecked<URigVMFunctionReferenceNode>(Node);
+			FunctionReferenceNode->ReferencedFunctionHeader.PostDuplicateHost(InOldPathName, InNewPathName);
+		}
+	}
+}
+
 #endif
 
 URigVMControllerSettings::URigVMControllerSettings(const FObjectInitializer& Initializer)
diff --git a/Engine/Source/Developer/RigVMDeveloper/Public/RigVMModel/RigVMClient.h b/Engine/Source/Developer/RigVMDeveloper/Public/RigVMModel/RigVMClient.h
--- a/Engine/Source/Developer/RigVMDeveloper/Public/RigVMModel/RigVMClient.h
+++ b/Engine/Source/Developer/RigVMDeveloper/Public/RigVMModel/RigVMClient.h
@@ -142,6 +142,9 @@ public:
 	// backwards compatibility
 	void PatchModelsOnLoad();
 
+	// work to be done after a duplication of the source asset
+	void PostDuplicateHost(const FString& InOldPathName, const FString& InNewPathName);
+
 	void HandleGraphModifiedEvent(ERigVMGraphNotifType InNotifType, URigVMGraph* InGraph, UObject* InSubject);
 
 	FRigVMGraphFunctionStore* FindFunctionStore(const URigVMLibraryNode* InLibraryNode);
diff --git a/Engine/Source/Developer/RigVMDeveloper/Public/RigVMModel/RigVMController.h b/Engine/Source/Developer/RigVMDeveloper/Public/RigVMModel/RigVMController.h
--- a/Engine/Source/Developer/RigVMDeveloper/Public/RigVMModel/RigVMController.h
+++ b/Engine/Source/Developer/RigVMDeveloper/Public/RigVMModel/RigVMController.h
@@ -1168,6 +1168,9 @@ protected:
 	void PatchDispatchNodesOnLoad();
 	void PatchBranchNodesOnLoad();
 	void PatchIfSelectNodesOnLoad();
+
+	// work to do after a duplication of the host asset
+	void PostDuplicateHost(const FString& InOldPathName, const FString& InNewPathName);
 	
 private: 
 	UPROPERTY(transient)
diff --git a/Engine/Source/Runtime/RigVM/Private/RigVMCore/RigVMGraphFunctionDefinition.cpp b/Engine/Source/Runtime/RigVM/Private/RigVMCore/RigVMGraphFunctionDefinition.cpp
--- a/Engine/Source/Runtime/RigVM/Private/RigVMCore/RigVMGraphFunctionDefinition.cpp
+++ b/Engine/Source/Runtime/RigVM/Private/RigVMCore/RigVMGraphFunctionDefinition.cpp
@@ -61,7 +61,35 @@ FRigVMGraphFunctionData* FRigVMGraphFunctionHeader::GetFunctionData() const
 	return nullptr;
 }
 
+void FRigVMGraphFunctionHeader::PostDuplicateHost(const FString& InOldPathName, const FString& InNewPathName)
+{
+	const FString OldPathName = InOldPathName + TEXT(":");
+	const FString NewPathName = InNewPathName + TEXT(":");
+
+	auto ReplacePathName = [InOldPathName, InNewPathName, OldPathName, NewPathName](FSoftObjectPath& InOutObjectPath)
+	{
+		FString PathName = InOutObjectPath.ToString();
+		if(PathName.Equals(InOldPathName, ESearchCase::CaseSensitive))
+		{
+			InOutObjectPath = FSoftObjectPath(InNewPathName);
+		}
+		else if(PathName.StartsWith(OldPathName, ESearchCase::CaseSensitive))
+		{
+			PathName = NewPathName + PathName.Mid(OldPathName.Len());
+			InOutObjectPath = FSoftObjectPath(PathName);
+		}
+	};
+
+	ReplacePathName(LibraryPointer.LibraryNode);
+	ReplacePathName(LibraryPointer.HostObject);
+}
+
 bool FRigVMGraphFunctionData::IsMutable() const
 {
 	return Header.IsMutable();
 }
+
+void FRigVMGraphFunctionData::PostDuplicateHost(const FString& InOldHostPathName, const FString& InNewHostPathName)
+{
+	Header.PostDuplicateHost(InOldHostPathName, InNewHostPathName);
+}
diff --git a/Engine/Source/Runtime/RigVM/Private/RigVMCore/RigVMGraphFunctionHost.cpp b/Engine/Source/Runtime/RigVM/Private/RigVMCore/RigVMGraphFunctionHost.cpp
--- a/Engine/Source/Runtime/RigVM/Private/RigVMCore/RigVMGraphFunctionHost.cpp
+++ b/Engine/Source/Runtime/RigVM/Private/RigVMCore/RigVMGraphFunctionHost.cpp
@@ -308,3 +308,16 @@ bool FRigVMGraphFunctionStore::RemoveAllCompilationData()
 
 	return true;
 }
+
+void FRigVMGraphFunctionStore::PostDuplicateHost(const FString& InOldPathName, const FString& InNewPathName)
+{
+	for (FRigVMGraphFunctionData& Data : PublicFunctions)
+	{
+		Data.Header.PostDuplicateHost(InOldPathName, InNewPathName);
+	}
+
+	for (FRigVMGraphFunctionData& Data : PrivateFunctions)
+	{
+		Data.Header.PostDuplicateHost(InOldPathName, InNewPathName);
+	}
+}
diff --git a/Engine/Source/Runtime/RigVM/Public/RigVMCore/RigVMGraphFunctionDefinition.h b/Engine/Source/Runtime/RigVM/Public/RigVMCore/RigVMGraphFunctionDefinition.h
--- a/Engine/Source/Runtime/RigVM/Public/RigVMCore/RigVMGraphFunctionDefinition.h
+++ b/Engine/Source/Runtime/RigVM/Public/RigVMCore/RigVMGraphFunctionDefinition.h
@@ -403,6 +403,8 @@ struct RIGVM_API FRigVMGraphFunctionHeader
 		Ar << Data.ExternalVariables;
 		return Ar;
 	}
+
+	void PostDuplicateHost(const FString& InOldPathName, const FString& InNewPathName);
 };
 
 USTRUCT(BlueprintType)
@@ -436,4 +438,6 @@ struct RIGVM_API FRigVMGraphFunctionData
 		Ar << Data.CompilationData;
 		return Ar;
 	}
+
+	void PostDuplicateHost(const FString& InOldPathName, const FString& InNewPathName);
 };
diff --git a/Engine/Source/Runtime/RigVM/Public/RigVMCore/RigVMGraphFunctionHost.h b/Engine/Source/Runtime/RigVM/Public/RigVMCore/RigVMGraphFunctionHost.h
--- a/Engine/Source/Runtime/RigVM/Public/RigVMCore/RigVMGraphFunctionHost.h
+++ b/Engine/Source/Runtime/RigVM/Public/RigVMCore/RigVMGraphFunctionHost.h
@@ -68,6 +68,8 @@ public:
 
 	bool RemoveAllCompilationData();
 
+	void PostDuplicateHost(const FString& InOldPathName, const FString& InNewPathName);
+
 	friend FArchive& operator<<(FArchive& Ar, FRigVMGraphFunctionStore& Host)
 	{
 		Ar << Host.PublicFunctions;
