ControlRig: Ground truth test assets for mannequin, encapsulated mannequin and basiccontrols

diff --git a/Engine/Plugins/Animation/ControlRig/Source/ControlRig/Public/ControlRigTestData.h b/Engine/Plugins/Animation/ControlRig/Source/ControlRig/Public/ControlRigTestData.h
--- a/Engine/Plugins/Animation/ControlRig/Source/ControlRig/Public/ControlRigTestData.h
+++ b/Engine/Plugins/Animation/ControlRig/Source/ControlRig/Public/ControlRigTestData.h
@@ -79,7 +79,8 @@ class CONTROLRIG_API UControlRigTestData : public UObject
 public:
 
 	UControlRigTestData()
-		: LastFrameIndex(INDEX_NONE)
+		: Tolerance(0.001)
+		, LastFrameIndex(INDEX_NONE)
 		, DesiredRecordingDuration(0.0)
 		, TimeAtStartOfRecording(0.0)
 		, bIsApplyingOutputs(false)
@@ -126,6 +127,12 @@ public:
 	UPROPERTY(VisibleAnywhere, BlueprintReadOnly, Category = "ControlRigTestData")
 	TArray<FControlRigTestDataFrame> OutputFrames;
 
+	UPROPERTY(EditAnywhere, BlueprintReadOnly, Category = "ControlRigTestData")
+	TArray<int32> FramesToSkip;
+
+	UPROPERTY(EditAnywhere, BlueprintReadOnly, Category = "ControlRigTestData")
+	double Tolerance;
+
 private:
 
 	void ClearDelegates(UControlRig* InControlRig);
