always add constraints keys on the main section

diff --git a/Engine/Plugins/Animation/ControlRig/Source/ControlRigEditor/Private/Sequencer/ConstraintChannelHelper.cpp b/Engine/Plugins/Animation/ControlRig/Source/ControlRigEditor/Private/Sequencer/ConstraintChannelHelper.cpp
--- a/Engine/Plugins/Animation/ControlRig/Source/ControlRigEditor/Private/Sequencer/ConstraintChannelHelper.cpp
+++ b/Engine/Plugins/Animation/ControlRig/Source/ControlRigEditor/Private/Sequencer/ConstraintChannelHelper.cpp
@@ -20,6 +20,9 @@
 #include "Channels/MovieSceneDoubleChannel.h"
 #include "MovieSceneConstraintChannelHelper.h"
 #include "MovieSceneSpawnableAnnotation.h"
+#include "EntitySystem/MovieSceneDecompositionQuery.h"
+#include "Systems/MovieScenePropertyInstantiator.h"
+#include "Tracks/MovieScene3DTransformTrack.h"
 
 
 #define LOCTEXT_NAMESPACE "Constraints"
@@ -55,6 +58,224 @@ namespace
 
 		return false;
 	}
+
+	// NOTE, because of restriction for hotfix changes this code some of the following functions are duplicated in
+	// ConstraintBaker.cpp. This will be factorized as part of a refactor for 5.2
+	
+	UMovieScene* GetMovieScene(const TSharedPtr<ISequencer>& InSequencer)
+	{
+		const UMovieSceneSequence* MovieSceneSequence = InSequencer ? InSequencer->GetFocusedMovieSceneSequence() : nullptr;
+		return MovieSceneSequence ? MovieSceneSequence->GetMovieScene() : nullptr;
+	}
+	
+	UMovieSceneControlRigParameterSection* GetControlConstraintSection(
+		const UTransformableControlHandle* InHandle,
+		const TSharedPtr<ISequencer>& InSequencer)
+	{
+		const UMovieScene* MovieScene = GetMovieScene(InSequencer);
+		if (!MovieScene)
+		{
+			return nullptr;
+		}
+
+		auto GetControlRigTrack = [InHandle, MovieScene]()->UMovieSceneControlRigParameterTrack*
+		{
+			const TWeakObjectPtr<UControlRig> ControlRig = InHandle->ControlRig.LoadSynchronous();
+			if (ControlRig.IsValid())
+			{	
+				const TArray<FMovieSceneBinding>& Bindings = MovieScene->GetBindings();
+				for (const FMovieSceneBinding& Binding : Bindings)
+				{
+					UMovieSceneTrack* Track = MovieScene->FindTrack(UMovieSceneControlRigParameterTrack::StaticClass(), Binding.GetObjectGuid());
+					UMovieSceneControlRigParameterTrack* ControlRigTrack = Cast<UMovieSceneControlRigParameterTrack>(Track);
+					if (ControlRigTrack && ControlRigTrack->GetControlRig() == ControlRig)
+					{
+						return ControlRigTrack;
+					}
+				}
+			}
+			return nullptr;
+		};
+
+		UMovieSceneControlRigParameterTrack* ControlRigTrack = GetControlRigTrack();
+		if (!ControlRigTrack)
+		{
+			return nullptr;
+		}
+
+		const TArray<UMovieSceneSection*>& AllSections = ControlRigTrack->GetAllSections();
+		UMovieSceneSection* Section = AllSections.IsEmpty() ? ControlRigTrack->FindSection(0) : AllSections[0];
+		return Cast<UMovieSceneControlRigParameterSection>(Section);
+	}
+	
+	UMovieScene3DTransformSection* GetComponentConstraintSection(
+		const TSharedPtr<ISequencer>& InSequencer,
+		const FGuid& InGuid,
+		const FTransform& InDefaultTransform)
+	{
+		UMovieScene* MovieScene = GetMovieScene(InSequencer);
+		if (!MovieScene)
+		{
+			return nullptr;
+		}
+
+		UMovieScene3DTransformTrack* TransformTrack = MovieScene->FindTrack<UMovieScene3DTransformTrack>(InGuid);
+		if (!TransformTrack)
+		{
+			MovieScene->Modify();
+			TransformTrack = MovieScene->AddTrack<UMovieScene3DTransformTrack>(InGuid);
+		}
+		TransformTrack->Modify();
+
+		const TArray<UMovieSceneSection*>& AllSections = TransformTrack->GetAllSections();
+
+		static constexpr FFrameNumber Frame0;
+		bool bSectionAdded = false;
+		UMovieScene3DTransformSection* TransformSection = AllSections.IsEmpty() ?
+			Cast<UMovieScene3DTransformSection>(TransformTrack->FindOrAddSection(Frame0, bSectionAdded)) :
+			Cast<UMovieScene3DTransformSection>(AllSections[0]);
+		if (!TransformSection)
+		{
+			return nullptr;
+		}
+
+		TransformSection->Modify();
+		if (bSectionAdded)
+		{
+			TransformSection->SetRange(TRange<FFrameNumber>::All());
+
+			const FMovieSceneChannelProxy& SectionChannelProxy = TransformSection->GetChannelProxy();
+			const TMovieSceneChannelHandle<FMovieSceneDoubleChannel> DoubleChannels[] = {
+				SectionChannelProxy.GetChannelByName<FMovieSceneDoubleChannel>("Location.X"),
+				SectionChannelProxy.GetChannelByName<FMovieSceneDoubleChannel>("Location.Y"),
+				SectionChannelProxy.GetChannelByName<FMovieSceneDoubleChannel>("Location.Z"),
+				SectionChannelProxy.GetChannelByName<FMovieSceneDoubleChannel>("Rotation.X"),
+				SectionChannelProxy.GetChannelByName<FMovieSceneDoubleChannel>("Rotation.Y"),
+				SectionChannelProxy.GetChannelByName<FMovieSceneDoubleChannel>("Rotation.Z"),
+				SectionChannelProxy.GetChannelByName<FMovieSceneDoubleChannel>("Scale.X"),
+				SectionChannelProxy.GetChannelByName<FMovieSceneDoubleChannel>("Scale.Y"),
+				SectionChannelProxy.GetChannelByName<FMovieSceneDoubleChannel>("Scale.Z")
+			};
+
+			const FVector Location0 = InDefaultTransform.GetLocation();
+			const FRotator Rotation0 = InDefaultTransform.GetRotation().Rotator();
+			const FVector Scale3D0 = InDefaultTransform.GetScale3D();
+			const FTransform::FReal DefaultValues[] = { Location0.X, Location0.Y, Location0.Z,
+														Rotation0.Roll, Rotation0.Pitch, Rotation0.Yaw,
+														Scale3D0.X, Scale3D0.Y, Scale3D0.Z };
+			for (int32 Index = 0; Index < 9; Index++)
+			{
+				if (FMovieSceneDoubleChannel* Channel = DoubleChannels[Index].Get())
+				{
+					Channel->SetDefault(DefaultValues[Index]);
+				}
+			}
+		}
+
+		return TransformSection;
+	}
+
+	// NOTE this has to be moved to a public header like MovieSceneToolsHelpers.h (or similar) for 5.2
+	// see F3DTransformTrackEditor::RecomposeTransform for more details
+	void RecomposeTransforms(
+		const TSharedPtr<ISequencer>& InSequencer, USceneComponent* SceneComponent, UMovieSceneSection* Section,
+		const TArray<FFrameNumber>& InFrames, TArray<FTransform>& InOutTransforms)
+	{
+		using namespace UE::MovieScene;
+
+		FMovieSceneRootEvaluationTemplateInstance& EvaluationTemplate = InSequencer->GetEvaluationTemplate();
+		UMovieSceneEntitySystemLinker* EntityLinker = EvaluationTemplate.GetEntitySystemLinker();
+		if (!EntityLinker)
+		{
+			return;
+		}
+
+		UMovieScenePropertyInstantiatorSystem* System = EntityLinker->FindSystem<UMovieScenePropertyInstantiatorSystem>();
+		if (!System)
+		{
+			return;
+		}
+
+		TGuardValue<FEntityManager*> DebugVizGuard(GEntityManagerForDebuggingVisualizers, &EntityLinker->EntityManager);
+		TArray<FMovieSceneEntityID> ImportedEntityIDs;
+		EvaluationTemplate.FindEntitiesFromOwner(Section, InSequencer->GetFocusedTemplateID(), ImportedEntityIDs);
+		if (ImportedEntityIDs.Num())
+		{
+			const FBuiltInComponentTypes* BuiltInComponents = FBuiltInComponentTypes::Get();
+			FMovieSceneTracksComponentTypes* TrackComponents = FMovieSceneTracksComponentTypes::Get();
+
+			UMovieScene* MovieScene = InSequencer->GetFocusedMovieSceneSequence()->GetMovieScene();
+			const FFrameRate TickResolution = MovieScene->GetTickResolution();
+			const EMovieScenePlayerStatus::Type PlaybackStatus = InSequencer->GetPlaybackStatus();
+
+			TArray<FMovieSceneEntityID> EntityIDs;
+			FDecompositionQuery Query;
+			Query.Object   = SceneComponent;
+			Query.bConvertFromSourceEntityIDs = false;  // We already pass the children entity IDs
+			
+			// add keys
+			for (int32 Index = 0; Index < InFrames.Num(); ++Index)
+			{
+				const FFrameNumber& FrameNumber = InFrames[Index];
+				const FMovieSceneEvaluationRange EvaluationRange = FMovieSceneEvaluationRange(FFrameTime(FrameNumber), TickResolution);
+				const FMovieSceneContext Context = FMovieSceneContext(EvaluationRange, PlaybackStatus).SetHasJumped(true);
+
+				EvaluationTemplate.EvaluateSynchronousBlocking(Context, *InSequencer);
+
+				if (EntityIDs.IsEmpty())
+				{
+					// In order to check for the result channels later, we need to look up the children entities that are
+					// bound to the given animated object. Imported entities generally don't have the result channels.
+					FEntityTaskBuilder()
+					.ReadEntityIDs()
+					.Read(BuiltInComponents->ParentEntity)
+					.Read(BuiltInComponents->BoundObject)
+					.FilterAll({ TrackComponents->ComponentTransform.PropertyTag })
+					.Iterate_PerEntity(
+						&EntityLinker->EntityManager, 
+						[SceneComponent, ImportedEntityIDs, &EntityIDs](FMovieSceneEntityID EntityID, FMovieSceneEntityID ParentEntityID, UObject* BoundObject)
+						{
+							if (SceneComponent == BoundObject && ImportedEntityIDs.Contains(ParentEntityID))
+							{
+								EntityIDs.Add(EntityID);
+							}
+						});
+					Query.Entities = MakeArrayView(EntityIDs);
+				}
+				
+				FTransform& Transform = InOutTransforms[Index];
+
+				const FIntermediate3DTransform CurrentValue(Transform.GetTranslation(), Transform.GetRotation().Rotator(), Transform.GetScale3D());
+
+				TRecompositionResult<FIntermediate3DTransform> TransformData = System->RecomposeBlendOperational(TrackComponents->ComponentTransform, Query, CurrentValue);
+
+				double CurrentTransformChannels[9] = { 0, 0, 0, 0, 0, 0, 0, 0, 0 };
+				EMovieSceneTransformChannel ChannelsObtained(EMovieSceneTransformChannel::None);
+				check(EntityIDs.Num() == TransformData.Values.Num());
+				for (int32 EntityIndex = 0; EntityIndex < TransformData.Values.Num(); ++EntityIndex)
+				{
+					const FMovieSceneEntityID& EntityID = EntityIDs[EntityIndex];
+					
+					const FIntermediate3DTransform& EntityTransformData = TransformData.Values[EntityIndex];
+					FComponentMask EntityType = EntityLinker->EntityManager.GetEntityType(EntityID);
+					for (int32 CompositeIndex = 0; CompositeIndex < 9; ++CompositeIndex)
+					{
+						EMovieSceneTransformChannel ChannelMask = (EMovieSceneTransformChannel)(1 << CompositeIndex);
+						if (!EnumHasAnyFlags(ChannelsObtained, ChannelMask) && EntityType.Contains(BuiltInComponents->DoubleResult[CompositeIndex]))
+						{
+							EnumAddFlags(ChannelsObtained, (EMovieSceneTransformChannel)(1 << CompositeIndex));
+							CurrentTransformChannels[CompositeIndex] = EntityTransformData[CompositeIndex];
+						}
+					}
+				}
+
+				Transform = FTransform(
+					FRotator(CurrentTransformChannels[4], CurrentTransformChannels[5], CurrentTransformChannels[3]), // pitch yaw roll
+					FVector(CurrentTransformChannels[0], CurrentTransformChannels[1], CurrentTransformChannels[2]),
+					FVector(CurrentTransformChannels[6], CurrentTransformChannels[7], CurrentTransformChannels[8]));
+			}
+		}
+	}
 }
 
 bool FConstraintChannelHelper::IsKeyframingAvailable()
@@ -204,111 +425,116 @@ bool FConstraintChannelHelper::SmartControlConstraintKey(
 	{
 		return false;
 	}
-	
-	if (UMovieSceneControlRigParameterSection* Section = GetControlSection(ControlHandle, InSequencer))
+
+	UMovieSceneControlRigParameterSection* ConstraintSection = GetControlConstraintSection(ControlHandle, InSequencer);
+	UMovieSceneControlRigParameterSection* TransformSection = GetControlSection(ControlHandle, InSequencer);
+	if ((!ConstraintSection) || (!TransformSection))
 	{
-		FScopedTransaction Transaction(LOCTEXT("KeyConstraintaKehy", "Key Constraint Key"));
-		Section->Modify();
+		return false;
+	}
+	
+	FScopedTransaction Transaction(LOCTEXT("KeyConstraintaKehy", "Key Constraint Key"));
+	ConstraintSection->Modify();
+	TransformSection->Modify();
 
-		// set constraint as dynamic
-		InConstraint->bDynamicOffset = true;
-		
-		// add the channel
-		Section->AddConstraintChannel(InConstraint);
+	// set constraint as dynamic
+	InConstraint->bDynamicOffset = true;
+	
+	// add the channel
+	ConstraintSection->AddConstraintChannel(InConstraint);
 
-		// add key if needed
-		if (FConstraintAndActiveChannel* Channel = Section->GetConstraintChannel(InConstraint->GetFName()))
+	// add key if needed
+	if (FConstraintAndActiveChannel* Channel = ConstraintSection->GetConstraintChannel(InConstraint->GetFName()))
+	{
+		bool ActiveValueToBeSet = false;
+		//add key if we can and make sure the key we are setting is what we want
+		if (CanAddKey(Channel->ActiveChannel, Time, ActiveValueToBeSet) && (bActive.IsSet() == false || bActive.GetValue() == ActiveValueToBeSet))
 		{
+			const FFrameRate TickResolution = InSequencer->GetFocusedTickResolution();
 
-			bool ActiveValueToBeSet = false;
-			//add key if we can and make sure the key we are setting is what we want
-			if (CanAddKey(Channel->ActiveChannel, Time, ActiveValueToBeSet) && (bActive.IsSet() == false || bActive.GetValue() == ActiveValueToBeSet))
-			{
-				const FFrameRate TickResolution = InSequencer->GetFocusedTickResolution();
-
-				const bool bNeedsCompensation = InConstraint->NeedsCompensation();
+			const bool bNeedsCompensation = InConstraint->NeedsCompensation();
 				
-				TGuardValue<bool> CompensateGuard(FMovieSceneConstraintChannelHelper::bDoNotCompensate, true);
+			TGuardValue<bool> CompensateGuard(FMovieSceneConstraintChannelHelper::bDoNotCompensate, true);
 				
-				UControlRig* ControlRig = ControlHandle->ControlRig.Get();
-				const FName& ControlName = ControlHandle->ControlName;
+			UControlRig* ControlRig = ControlHandle->ControlRig.Get();
+			const FName& ControlName = ControlHandle->ControlName;
 				
-				// store the frames to compensate
-				const TArrayView<FMovieSceneFloatChannel*> Channels = ControlHandle->GetFloatChannels(Section);
-				TArray<FFrameNumber> FramesToCompensate;
-				if (bNeedsCompensation)
-				{
-					FMovieSceneConstraintChannelHelper::GetFramesToCompensate<FMovieSceneFloatChannel>(Channel->ActiveChannel, ActiveValueToBeSet, Time, Channels, FramesToCompensate);
-				}
-				else
-				{
-					FramesToCompensate.Add(Time);
-				}
+			// store the frames to compensate
+			const TArrayView<FMovieSceneFloatChannel*> Channels = ControlHandle->GetFloatChannels(TransformSection);
+			TArray<FFrameNumber> FramesToCompensate;
+			if (bNeedsCompensation)
+			{
+				FMovieSceneConstraintChannelHelper::GetFramesToCompensate<FMovieSceneFloatChannel>(Channel->ActiveChannel, ActiveValueToBeSet, Time, Channels, FramesToCompensate);
+			}
+			else
+			{
+				FramesToCompensate.Add(Time);
+			}
 
-				// store child and space transforms for these frames
-				FCompensationEvaluator Evaluator(InConstraint);
-				Evaluator.ComputeLocalTransforms(ControlRig->GetWorld(), InSequencer, FramesToCompensate, ActiveValueToBeSet);
-				TArray<FTransform>& ChildLocals = Evaluator.ChildLocals;
-				
-				// store tangents at this time
-				TArray<FMovieSceneTangentData> Tangents;
-				int32 ChannelIndex = 0, NumChannels = 0;
 
-				FChannelMapInfo* pChannelIndex = nullptr;
-				FRigControlElement* ControlElement = nullptr;
-				Tie(ControlElement, pChannelIndex) = FControlRigSpaceChannelHelpers::GetControlAndChannelInfo(ControlRig, Section, ControlName);
+			// store child and space transforms for these frames
+			FCompensationEvaluator Evaluator(InConstraint);
+			Evaluator.ComputeLocalTransforms(ControlRig->GetWorld(), InSequencer, FramesToCompensate, ActiveValueToBeSet);
+			TArray<FTransform>& ChildLocals = Evaluator.ChildLocals;
+			
+			// store tangents at this time
+			TArray<FMovieSceneTangentData> Tangents;
+			int32 ChannelIndex = 0, NumChannels = 0;
 
-				if (pChannelIndex && ControlElement)
+			FChannelMapInfo* pChannelIndex = nullptr;
+			FRigControlElement* ControlElement = nullptr;
+			Tie(ControlElement, pChannelIndex) = FControlRigSpaceChannelHelpers::GetControlAndChannelInfo(ControlRig, TransformSection, ControlName);
+
+			if (pChannelIndex && ControlElement)
+			{
+				// get the number of float channels to treat
+				NumChannels = FControlRigSpaceChannelHelpers::GetNumFloatChannels(ControlElement->Settings.ControlType);
+				if (bNeedsCompensation && NumChannels > 0)
 				{
-					// get the number of float channels to treat
-					NumChannels = FControlRigSpaceChannelHelpers::GetNumFloatChannels(ControlElement->Settings.ControlType);
-					if (bNeedsCompensation && NumChannels > 0)
-					{
-						ChannelIndex = pChannelIndex->ChannelIndex;
-						EvaluateTangentAtThisTime<FMovieSceneFloatChannel>(ChannelIndex, NumChannels, Section, Time, Tangents);
-					}
+					ChannelIndex = pChannelIndex->ChannelIndex;
+					EvaluateTangentAtThisTime<FMovieSceneFloatChannel>(ChannelIndex, NumChannels, TransformSection, Time, Tangents);
 				}
+			}
+		
+			const EMovieSceneTransformChannel ChannelsToKey =InConstraint->GetChannelsToKey();
 			
-				const EMovieSceneTransformChannel ChannelsToKey =InConstraint->GetChannelsToKey();
-				
-				// add child's transform key at Time-1 to keep animation
-				if (bNeedsCompensation)
-				{
-					const FFrameNumber TimeMinusOne(Time - 1);
+			// add child's transform key at Time-1 to keep animation
+			if (bNeedsCompensation)
+			{
+				const FFrameNumber TimeMinusOne(Time - 1);
 
-					ControlHandle->AddTransformKeys({ TimeMinusOne },
-						{ ChildLocals[0] }, ChannelsToKey, TickResolution, nullptr,true);
+				ControlHandle->AddTransformKeys({ TimeMinusOne },
+					{ ChildLocals[0] }, ChannelsToKey, TickResolution, nullptr,true);
 
-					// set tangents at Time-1
-					if (NumChannels > 0)
-					{
-						SetTangentsAtThisTime<FMovieSceneFloatChannel>(ChannelIndex, NumChannels, Section, TimeMinusOne, Tangents);
-					}
-				}
-
-				// add active key
+				// set tangents at Time-1
+				if (NumChannels > 0)
 				{
-					TMovieSceneChannelData<bool> ChannelData = Channel->ActiveChannel.GetData();
-					ChannelData.AddKey(Time, ActiveValueToBeSet);
+					SetTangentsAtThisTime<FMovieSceneFloatChannel>(ChannelIndex, NumChannels, TransformSection, TimeMinusOne, Tangents);
 				}
+			}
 
-				// compensate
-				{
-					// we need to remove the first transforms as we store NumFrames+1 transforms
-					ChildLocals.RemoveAt(0);
+			// add active key
+			{
+				TMovieSceneChannelData<bool> ChannelData = Channel->ActiveChannel.GetData();
+				ChannelData.AddKey(Time, ActiveValueToBeSet);
+			}
+
+			// compensate
+			{
+				// we need to remove the first transforms as we store NumFrames+1 transforms
+				ChildLocals.RemoveAt(0);
 
-					// add keys
-					ControlHandle->AddTransformKeys(FramesToCompensate,
-						ChildLocals, ChannelsToKey, TickResolution, nullptr,true);
+				// add keys
+				ControlHandle->AddTransformKeys(FramesToCompensate,
+					ChildLocals, ChannelsToKey, TickResolution, nullptr,true);
 
-					// set tangents at Time
-					if (bNeedsCompensation && NumChannels > 0)
-					{
-						SetTangentsAtThisTime<FMovieSceneFloatChannel>(ChannelIndex, NumChannels, Section, Time, Tangents);
-					}
+				// set tangents at Time
+				if (bNeedsCompensation && NumChannels > 0)
+				{
+					SetTangentsAtThisTime<FMovieSceneFloatChannel>(ChannelIndex, NumChannels, TransformSection, Time, Tangents);
 				}
-				return true;
 			}
+			return true;
 		}
 	}
 	return false;
@@ -338,117 +564,130 @@ bool FConstraintChannelHelper::SmartComponentConstraintKey(
 		return false;
 	}
 
-	if (UMovieScene3DTransformSection* Section = MovieSceneToolHelpers::GetTransformSection(InSequencer.Get(), Guid, LocalTransform))
+	UMovieScene3DTransformSection* TransformSection = MovieSceneToolHelpers::GetTransformSection(InSequencer.Get(), Guid, LocalTransform);
+	UMovieScene3DTransformSection* ConstraintSection = GetComponentConstraintSection(InSequencer, Guid, LocalTransform);
+	if ((!ConstraintSection) || (!TransformSection))
 	{
-		FScopedTransaction Transaction(LOCTEXT("KeyConstraintaKehy", "Key Constraint Key"));
-		Section->Modify();
+		return false;
+	}
 
-		// set constraint as dynamic
-		InConstraint->bDynamicOffset = true;
-		
-		// add the channel
-		Section->AddConstraintChannel(InConstraint);
+	FScopedTransaction Transaction(LOCTEXT("KeyConstraintaKehy", "Key Constraint Key"));
+	TransformSection->Modify();
+	ConstraintSection->Modify();
+
+	// set constraint as dynamic
+	InConstraint->bDynamicOffset = true;
+	
+	// add the channel
+	ConstraintSection->AddConstraintChannel(InConstraint);
 
-		// add key if needed
-		if (FConstraintAndActiveChannel* Channel = Section->GetConstraintChannel(InConstraint->GetFName()))
+	// add key if needed
+	if (FConstraintAndActiveChannel* Channel = ConstraintSection->GetConstraintChannel(InConstraint->GetFName()))
+	{
+		bool ActiveValueToBeSet = false;
+		//add key if we can and make sure the key we are setting is what we want
+		if (CanAddKey(Channel->ActiveChannel, Time, ActiveValueToBeSet) && (bActive.IsSet() == false || bActive.GetValue() == ActiveValueToBeSet))
 		{
-			bool ActiveValueToBeSet = false;
-			//add key if we can and make sure the key we are setting is what we want
-			if (CanAddKey(Channel->ActiveChannel, Time, ActiveValueToBeSet) && (bActive.IsSet() == false || bActive.GetValue() == ActiveValueToBeSet))
+			const FFrameRate TickResolution = InSequencer->GetFocusedTickResolution();
+			const bool bNeedsCompensation = InConstraint->NeedsCompensation();
+
+			//new for compensation
+
+			TGuardValue<bool> CompensateGuard(FMovieSceneConstraintChannelHelper::bDoNotCompensate, true);
+			// store the frames to compensate
+			const TArrayView<FMovieSceneDoubleChannel*> Channels = ComponentHandle->GetDoubleChannels(TransformSection);
+			TArray<FFrameNumber> FramesToCompensate;
+			if (bNeedsCompensation)
 			{
-				const FFrameRate TickResolution = InSequencer->GetFocusedTickResolution();
-				const bool bNeedsCompensation = InConstraint->NeedsCompensation();
-				
-				Section->Modify();
+				FMovieSceneConstraintChannelHelper::GetFramesToCompensate<FMovieSceneDoubleChannel>(Channel->ActiveChannel, ActiveValueToBeSet, Time, Channels, FramesToCompensate);
+			}
+			else
+			{
+				FramesToCompensate.Add(Time);
+			}
 
-				//new for compensation
+			// store child and space transforms for these frames
+			FCompensationEvaluator Evaluator(InConstraint);
+			Evaluator.ComputeLocalTransforms(Actor->GetWorld(), InSequencer, FramesToCompensate, ActiveValueToBeSet);
+			TArray<FTransform>& ChildLocals = Evaluator.ChildLocals;
 
-				TGuardValue<bool> CompensateGuard(FMovieSceneConstraintChannelHelper::bDoNotCompensate, true);
+			// store tangents at this time
+			TArray<FMovieSceneTangentData> Tangents;
+			const int32 ChannelIndex = 0;
+			const int32 NumChannels = 9;
 
-				// store the frames to compensate
-				const TArrayView<FMovieSceneDoubleChannel*> Channels = ComponentHandle->GetDoubleChannels(Section);
-				TArray<FFrameNumber> FramesToCompensate;
-				if (bNeedsCompensation)
-				{
-					FMovieSceneConstraintChannelHelper::GetFramesToCompensate<FMovieSceneDoubleChannel>(Channel->ActiveChannel, ActiveValueToBeSet, Time, Channels, FramesToCompensate);
-				}
-				else
-				{
-					FramesToCompensate.Add(Time);
-				}
+			if (bNeedsCompensation)
+			{
+				EvaluateTangentAtThisTime<FMovieSceneDoubleChannel>(ChannelIndex, NumChannels, TransformSection, Time, Tangents);
+			}
 
-				// store child and space transforms for these frames
-				FCompensationEvaluator Evaluator(InConstraint);
-				Evaluator.ComputeLocalTransforms(Actor->GetWorld(), InSequencer, FramesToCompensate, ActiveValueToBeSet);
-				TArray<FTransform>& ChildLocals = Evaluator.ChildLocals;
+			const EMovieSceneTransformChannel ChannelsToKey = InConstraint->GetChannelsToKey();
 
-				// store tangents at this time
-				TArray<FMovieSceneTangentData> Tangents;
-				const int32 ChannelIndex = 0;
-				const int32 NumChannels = 9;
+			// add child's transform key at Time-1 to keep animation
+			if (bNeedsCompensation)
+			{
+				const FFrameNumber TimeMinusOne(Time - 1);
 
-				if (bNeedsCompensation)
+				TArray<FTransform> TransformMinusOne({ChildLocals[0]} );
+				if (ConstraintSection != TransformSection)
 				{
-					EvaluateTangentAtThisTime<FMovieSceneDoubleChannel>(ChannelIndex, NumChannels, Section, Time, Tangents);
+					RecomposeTransforms(InSequencer, ComponentHandle->Component.Get(), TransformSection, { TimeMinusOne },TransformMinusOne);
 				}
+				MovieSceneToolHelpers::AddTransformKeys(TransformSection, { TimeMinusOne }, TransformMinusOne, ChannelsToKey);
 
-				const EMovieSceneTransformChannel ChannelsToKey = InConstraint->GetChannelsToKey();
-
-				// add child's transform key at Time-1 to keep animation
-				if (bNeedsCompensation)
-				{
-					const FFrameNumber TimeMinusOne(Time - 1);
+				SetTangentsAtThisTime<FMovieSceneDoubleChannel>(ChannelIndex, NumChannels, TransformSection, TimeMinusOne, Tangents);
+			}
 
-					MovieSceneToolHelpers::AddTransformKeys(Section, { TimeMinusOne }, { ChildLocals[0] }, ChannelsToKey);
+			// add active key
+			{
+				TMovieSceneChannelData<bool> ChannelData = Channel->ActiveChannel.GetData();
+				ChannelData.AddKey(Time, ActiveValueToBeSet);
+			}
 
-					SetTangentsAtThisTime<FMovieSceneDoubleChannel>(ChannelIndex, NumChannels, Section, TimeMinusOne, Tangents);
-				}
+			// compensate
+			{
+				// we need to remove the first transforms as we store NumFrames+1 transforms
+				ChildLocals.RemoveAt(0);
 
-				// add active key
+				if (ConstraintSection != TransformSection)
 				{
-					TMovieSceneChannelData<bool> ChannelData = Channel->ActiveChannel.GetData();
-					ChannelData.AddKey(Time, ActiveValueToBeSet);
+					RecomposeTransforms(InSequencer, ComponentHandle->Component.Get(), TransformSection, FramesToCompensate,ChildLocals);
 				}
 
-				// compensate
-				{
-					// we need to remove the first transforms as we store NumFrames+1 transforms
-					ChildLocals.RemoveAt(0);
-
-					// add keys
-					MovieSceneToolHelpers::AddTransformKeys(Section, FramesToCompensate, ChildLocals, ChannelsToKey);
+				// add keys
+				MovieSceneToolHelpers::AddTransformKeys(TransformSection, FramesToCompensate, ChildLocals, ChannelsToKey);
 
-					// set tangents at Time
-					if (bNeedsCompensation)
-					{
-						SetTangentsAtThisTime<FMovieSceneDoubleChannel>(ChannelIndex,NumChannels, Section, Time, Tangents);
-					}
-				}
-				// evaluate the constraint, this is needed so the global transform will be set up on the component 
-				//Todo do we need to evaluate all constraints?
-				InConstraint->SetActive(true); //will be false
-				InConstraint->Evaluate();
-
-				//need to fire this event so the transform values set by the constraint propragate to the transform section
-				//first turn off autokey though
-				EAutoChangeMode AutoChangeMode = InSequencer->GetAutoChangeMode();
-				if (AutoChangeMode == EAutoChangeMode::AutoKey || AutoChangeMode == EAutoChangeMode::All)
-				{
-					InSequencer->SetAutoChangeMode(EAutoChangeMode::None);
-				};
-				FProperty* TransformProperty = FindFProperty<FProperty>(USceneComponent::StaticClass(), USceneComponent::GetRelativeLocationPropertyName());
-				FEditPropertyChain PropertyChain;
-				PropertyChain.AddHead(TransformProperty);
-				FCoreUObjectDelegates::OnPreObjectPropertyChanged.Broadcast(Actor, PropertyChain);
-				FPropertyChangedEvent PropertyChangedEvent(TransformProperty, EPropertyChangeType::ValueSet);
-				FCoreUObjectDelegates::OnObjectPropertyChanged.Broadcast(Actor, PropertyChangedEvent);
-				InSequencer->RequestEvaluate();
-				if (AutoChangeMode == EAutoChangeMode::AutoKey || AutoChangeMode == EAutoChangeMode::All)
+				// set tangents at Time
+				if (bNeedsCompensation)
 				{
-					InSequencer->SetAutoChangeMode(AutoChangeMode);
+					SetTangentsAtThisTime<FMovieSceneDoubleChannel>(ChannelIndex,NumChannels, TransformSection, Time, Tangents);
 				}
-				return true;
 			}
+			// evaluate the constraint, this is needed so the global transform will be set up on the component 
+			//Todo do we need to evaluate all constraints?
+			InConstraint->SetActive(true); //will be false
+			InConstraint->Evaluate();
+
+			//need to fire this event so the transform values set by the constraint propragate to the transform section
+			//first turn off autokey though
+			EAutoChangeMode AutoChangeMode = InSequencer->GetAutoChangeMode();
+			if (AutoChangeMode == EAutoChangeMode::AutoKey || AutoChangeMode == EAutoChangeMode::All)
+			{
+				InSequencer->SetAutoChangeMode(EAutoChangeMode::None);
+			};
+			FProperty* TransformProperty = FindFProperty<FProperty>(USceneComponent::StaticClass(), USceneComponent::GetRelativeLocationPropertyName());
+			FEditPropertyChain PropertyChain;
+			PropertyChain.AddHead(TransformProperty);
+			FCoreUObjectDelegates::OnPreObjectPropertyChanged.Broadcast(Actor, PropertyChain);
+			FPropertyChangedEvent PropertyChangedEvent(TransformProperty, EPropertyChangeType::ValueSet);
+			FCoreUObjectDelegates::OnObjectPropertyChanged.Broadcast(Actor, PropertyChangedEvent);
+			InSequencer->RequestEvaluate();
+
+			if (AutoChangeMode == EAutoChangeMode::AutoKey || AutoChangeMode == EAutoChangeMode::All)
+			{
+				InSequencer->SetAutoChangeMode(AutoChangeMode);
+			}
+			return true;
 		}
 	}
 	return false;
@@ -467,47 +706,47 @@ void FConstraintChannelHelper::Compensate(UTickableTransformConstraint* InConstr
 	UWorld* World = nullptr;
 	if (const UTransformableComponentHandle* ComponentHandle = Cast<UTransformableComponentHandle>(InConstraint->ChildTRSHandle))
 	{
-
 		AActor* Actor = ComponentHandle->Component->GetOwner();
 		if (!Actor)
 		{
 			return;
 		}
-		World = Actor->GetWorld();
-		const FTransform LocalTransform = ComponentHandle->GetLocalTransform();
+
 		const FGuid Guid = Sequencer->GetHandleToObject(Actor, true);
 		if (!Guid.IsValid())
 		{
 			return;
 		}
 
-		if (UMovieScene3DTransformSection* TransformSection = MovieSceneToolHelpers::GetTransformSection(Sequencer.Get(), Guid, LocalTransform))
-		{
-			Section = TransformSection;
-		}
+		World = Actor->GetWorld();
+		
+		const FTransform LocalTransform = ComponentHandle->GetLocalTransform();
+		Section = GetComponentConstraintSection(Sequencer, Guid, LocalTransform);
 	}
 
 	if (const UTransformableControlHandle* ControlHandle = Cast<UTransformableControlHandle>(InConstraint->ChildTRSHandle))
 	{
-		UControlRig* ControlRig = ControlHandle->ControlRig.LoadSynchronous();
+		const UControlRig* ControlRig = ControlHandle->ControlRig.LoadSynchronous();
 		if (!ControlRig)
 		{
 			return;
 		}
 		World = ControlRig->GetWorld();
-		if (UMovieSceneControlRigParameterSection* ControlSection = GetControlSection(ControlHandle,Sequencer))
-		{
-			Section = ControlSection;
-		}
+		Section = GetControlConstraintSection(ControlHandle,Sequencer);
 	}
 
+	if (!Section)
+	{
+		return;
+	}
+	
 	CompensateIfNeeded(World, Sequencer, Section, OptTime);
 }
 
 void FConstraintChannelHelper::CompensateIfNeeded(
 	UWorld* InWorld,
 	const TSharedPtr<ISequencer>& InSequencer,
-	IMovieSceneConstrainedSection* Section,
+	IMovieSceneConstrainedSection* ConstraintSection,
 	const TOptional<FFrameNumber>& OptionalTime)
 {
 	if (FMovieSceneConstraintChannelHelper::bDoNotCompensate)
@@ -536,7 +775,7 @@ void FConstraintChannelHelper::CompensateIfNeeded(
 	bool bNeedsEvaluation = false;
 
 	// gather all transform constraints
-	TArray<FConstraintAndActiveChannel>& ConstraintChannels = Section->GetConstraintsChannels();
+	TArray<FConstraintAndActiveChannel>& ConstraintChannels = ConstraintSection->GetConstraintsChannels();
 	TArray<FConstraintAndActiveChannel> TransformConstraintsChannels;
 	Algo::CopyIf(ConstraintChannels, TransformConstraintsChannels,
 		[](const FConstraintAndActiveChannel& InChannel)
diff --git a/Engine/Plugins/Animation/ControlRig/Source/ControlRigEditor/Private/Tools/ConstraintBaker.cpp b/Engine/Plugins/Animation/ControlRig/Source/ControlRigEditor/Private/Tools/ConstraintBaker.cpp
--- a/Engine/Plugins/Animation/ControlRig/Source/ControlRigEditor/Private/Tools/ConstraintBaker.cpp
+++ b/Engine/Plugins/Animation/ControlRig/Source/ControlRigEditor/Private/Tools/ConstraintBaker.cpp
@@ -20,12 +20,21 @@
 #include "ConstraintChannelHelper.h"
 #include "ConstraintChannel.h"
 #include "ConstraintChannelHelper.inl"
+#include "EntitySystem/MovieSceneDecompositionQuery.h"
+#include "Sequencer/MovieSceneControlRigParameterTrack.h"
+#include "Tracks/MovieScene3DTransformTrack.h"
+#include "Systems/MovieScenePropertyInstantiator.h"
 
 #define LOCTEXT_NAMESPACE "ConstraintBaker"
 
 namespace
 {
-
+	UMovieScene* GetMovieScene(const TSharedPtr<ISequencer>& InSequencer)
+	{
+		const UMovieSceneSequence* MovieSceneSequence = InSequencer ? InSequencer->GetFocusedMovieSceneSequence() : nullptr;
+		return MovieSceneSequence ? MovieSceneSequence->GetMovieScene() : nullptr;
+	}
+	
 	UMovieScene3DTransformSection* GetTransformSection(
 		const TSharedPtr<ISequencer>& InSequencer,
 		AActor* InActor,
@@ -45,6 +54,177 @@ namespace
 		return MovieSceneToolHelpers::GetTransformSection(InSequencer.Get(), Guid, InTransform0);
 	}
 
+	// NOTE, because of restriction for hotfix changes this code some of the following functions are duplicated in
+	// ConstraintChannelHelper.cpp. This will be factorized as part of a refactor for 5.2
+	UMovieScene3DTransformSection* GetComponentConstraintSection(
+		const TSharedPtr<ISequencer>& InSequencer,
+		const FGuid& InGuid,
+		const FTransform& InDefaultTransform)
+	{
+		UMovieScene* MovieScene = GetMovieScene(InSequencer);
+		if (!MovieScene)
+		{
+			return nullptr;
+		}
+
+		UMovieScene3DTransformTrack* TransformTrack = MovieScene->FindTrack<UMovieScene3DTransformTrack>(InGuid);
+		if (!TransformTrack)
+		{
+			MovieScene->Modify();
+			TransformTrack = MovieScene->AddTrack<UMovieScene3DTransformTrack>(InGuid);
+		}
+		TransformTrack->Modify();
+
+		const TArray<UMovieSceneSection*>& AllSections = TransformTrack->GetAllSections();
+
+		static constexpr FFrameNumber Frame0;
+		bool bSectionAdded = false;
+		UMovieScene3DTransformSection* TransformSection = AllSections.IsEmpty() ?
+			Cast<UMovieScene3DTransformSection>(TransformTrack->FindOrAddSection(Frame0, bSectionAdded)) :
+			Cast<UMovieScene3DTransformSection>(AllSections[0]);
+		if (!TransformSection)
+		{
+			return nullptr;
+		}
+
+		TransformSection->Modify();
+		if (bSectionAdded)
+		{
+			TransformSection->SetRange(TRange<FFrameNumber>::All());
+
+			const FMovieSceneChannelProxy& SectionChannelProxy = TransformSection->GetChannelProxy();
+			const TMovieSceneChannelHandle<FMovieSceneDoubleChannel> DoubleChannels[] = {
+				SectionChannelProxy.GetChannelByName<FMovieSceneDoubleChannel>("Location.X"),
+				SectionChannelProxy.GetChannelByName<FMovieSceneDoubleChannel>("Location.Y"),
+				SectionChannelProxy.GetChannelByName<FMovieSceneDoubleChannel>("Location.Z"),
+				SectionChannelProxy.GetChannelByName<FMovieSceneDoubleChannel>("Rotation.X"),
+				SectionChannelProxy.GetChannelByName<FMovieSceneDoubleChannel>("Rotation.Y"),
+				SectionChannelProxy.GetChannelByName<FMovieSceneDoubleChannel>("Rotation.Z"),
+				SectionChannelProxy.GetChannelByName<FMovieSceneDoubleChannel>("Scale.X"),
+				SectionChannelProxy.GetChannelByName<FMovieSceneDoubleChannel>("Scale.Y"),
+				SectionChannelProxy.GetChannelByName<FMovieSceneDoubleChannel>("Scale.Z")
+			};
+
+			const FVector Location0 = InDefaultTransform.GetLocation();
+			const FRotator Rotation0 = InDefaultTransform.GetRotation().Rotator();
+			const FVector Scale3D0 = InDefaultTransform.GetScale3D();
+			const FTransform::FReal DefaultValues[] = { Location0.X, Location0.Y, Location0.Z,
+														Rotation0.Roll, Rotation0.Pitch, Rotation0.Yaw,
+														Scale3D0.X, Scale3D0.Y, Scale3D0.Z };
+			for (int32 Index = 0; Index < 9; Index++)
+			{
+				if (FMovieSceneDoubleChannel* Channel = DoubleChannels[Index].Get())
+				{
+					Channel->SetDefault(DefaultValues[Index]);
+				}
+			}
+		}
+
+		return TransformSection;
+	}
+
+	// NOTE this has to be moved to a public header like MovieSceneToolsHelpers.h (or similar) for 5.2
+	// see F3DTransformTrackEditor::RecomposeTransform for more details
+	void RecomposeTransforms(
+		const TSharedPtr<ISequencer>& InSequencer, USceneComponent* SceneComponent, UMovieSceneSection* Section,
+		const TArray<FFrameNumber>& InFrames, TArray<FTransform>& InOutTransforms)
+	{
+		using namespace UE::MovieScene;
+
+		FMovieSceneRootEvaluationTemplateInstance& EvaluationTemplate = InSequencer->GetEvaluationTemplate();
+		UMovieSceneEntitySystemLinker* EntityLinker = EvaluationTemplate.GetEntitySystemLinker();
+		if (!EntityLinker)
+		{
+			return;
+		}
+
+		UMovieScenePropertyInstantiatorSystem* System = EntityLinker->FindSystem<UMovieScenePropertyInstantiatorSystem>();
+		if (!System)
+		{
+			return;
+		}
+
+		TGuardValue<FEntityManager*> DebugVizGuard(GEntityManagerForDebuggingVisualizers, &EntityLinker->EntityManager);
+		TArray<FMovieSceneEntityID> ImportedEntityIDs;
+		EvaluationTemplate.FindEntitiesFromOwner(Section, InSequencer->GetFocusedTemplateID(), ImportedEntityIDs);
+		if (ImportedEntityIDs.Num())
+		{
+			const FBuiltInComponentTypes* BuiltInComponents = FBuiltInComponentTypes::Get();
+			FMovieSceneTracksComponentTypes* TrackComponents = FMovieSceneTracksComponentTypes::Get();
+
+			UMovieScene* MovieScene = InSequencer->GetFocusedMovieSceneSequence()->GetMovieScene();
+			const FFrameRate TickResolution = MovieScene->GetTickResolution();
+			const EMovieScenePlayerStatus::Type PlaybackStatus = InSequencer->GetPlaybackStatus();
+
+			TArray<FMovieSceneEntityID> EntityIDs;
+			FDecompositionQuery Query;
+			Query.Object   = SceneComponent;
+			Query.bConvertFromSourceEntityIDs = false;  // We already pass the children entity IDs
+			
+			// add keys
+			for (int32 Index = 0; Index < InFrames.Num(); ++Index)
+			{
+				const FFrameNumber& FrameNumber = InFrames[Index];
+				const FMovieSceneEvaluationRange EvaluationRange = FMovieSceneEvaluationRange(FFrameTime(FrameNumber), TickResolution);
+				const FMovieSceneContext Context = FMovieSceneContext(EvaluationRange, PlaybackStatus).SetHasJumped(true);
+
+				EvaluationTemplate.EvaluateSynchronousBlocking(Context, *InSequencer);
+
+				if (EntityIDs.IsEmpty())
+				{
+					// In order to check for the result channels later, we need to look up the children entities that are
+					// bound to the given animated object. Imported entities generally don't have the result channels.
+					FEntityTaskBuilder()
+					.ReadEntityIDs()
+					.Read(BuiltInComponents->ParentEntity)
+					.Read(BuiltInComponents->BoundObject)
+					.FilterAll({ TrackComponents->ComponentTransform.PropertyTag })
+					.Iterate_PerEntity(
+						&EntityLinker->EntityManager, 
+						[SceneComponent, ImportedEntityIDs, &EntityIDs](FMovieSceneEntityID EntityID, FMovieSceneEntityID ParentEntityID, UObject* BoundObject)
+						{
+							if (SceneComponent == BoundObject && ImportedEntityIDs.Contains(ParentEntityID))
+							{
+								EntityIDs.Add(EntityID);
+							}
+						});
+					Query.Entities = MakeArrayView(EntityIDs);
+				}
+				
+				FTransform& Transform = InOutTransforms[Index];
+
+				const FIntermediate3DTransform CurrentValue(Transform.GetTranslation(), Transform.GetRotation().Rotator(), Transform.GetScale3D());
+
+				TRecompositionResult<FIntermediate3DTransform> TransformData = System->RecomposeBlendOperational(TrackComponents->ComponentTransform, Query, CurrentValue);
+
+				double CurrentTransformChannels[9] = { 0, 0, 0, 0, 0, 0, 0, 0, 0 };
+				EMovieSceneTransformChannel ChannelsObtained(EMovieSceneTransformChannel::None);
+				check(EntityIDs.Num() == TransformData.Values.Num());
+				for (int32 EntityIndex = 0; EntityIndex < TransformData.Values.Num(); ++EntityIndex)
+				{
+					const FMovieSceneEntityID& EntityID = EntityIDs[EntityIndex];
+					
+					const FIntermediate3DTransform& EntityTransformData = TransformData.Values[EntityIndex];
+					FComponentMask EntityType = EntityLinker->EntityManager.GetEntityType(EntityID);
+					for (int32 CompositeIndex = 0; CompositeIndex < 9; ++CompositeIndex)
+					{
+						EMovieSceneTransformChannel ChannelMask = (EMovieSceneTransformChannel)(1 << CompositeIndex);
+						if (!EnumHasAnyFlags(ChannelsObtained, ChannelMask) && EntityType.Contains(BuiltInComponents->DoubleResult[CompositeIndex]))
+						{
+							EnumAddFlags(ChannelsObtained, (EMovieSceneTransformChannel)(1 << CompositeIndex));
+							CurrentTransformChannels[CompositeIndex] = EntityTransformData[CompositeIndex];
+						}
+					}
+				}
+
+				Transform = FTransform(
+					FRotator(CurrentTransformChannels[4], CurrentTransformChannels[5], CurrentTransformChannels[3]), // pitch yaw roll
+					FVector(CurrentTransformChannels[0], CurrentTransformChannels[1], CurrentTransformChannels[2]),
+					FVector(CurrentTransformChannels[6], CurrentTransformChannels[7], CurrentTransformChannels[8]));
+			}
+		}
+	}
+	
 	void BakeComponent(
 		const TSharedPtr<ISequencer>& InSequencer,
 		const UTransformableComponentHandle* InComponentHandle,
@@ -63,14 +243,28 @@ namespace
 		{
 			return;
 		}
+		
 		UMovieScene3DTransformSection* TransformSection = GetTransformSection(InSequencer, Actor, InTransforms[0]);
 		if (!TransformSection)
 		{
 			return;
 		}
+
 		const UMovieScene* MovieScene = InSequencer->GetFocusedMovieSceneSequence()->GetMovieScene();
-		InComponentHandle->AddTransformKeys(InFrames, InTransforms, InChannels, MovieScene->GetTickResolution(), TransformSection, true);
+		const FFrameRate TickResolution = MovieScene->GetTickResolution();
 
+		const FGuid Guid = InSequencer->GetHandleToObject(Actor, false);
+		const bool bNeedsRecomposition = TransformSection != GetComponentConstraintSection(InSequencer, Guid, InTransforms[0]);
+		if (bNeedsRecomposition)
+		{
+			TArray<FTransform> Transforms(InTransforms);
+			RecomposeTransforms(InSequencer, InComponentHandle->Component.Get(), TransformSection, InFrames, Transforms);
+			InComponentHandle->AddTransformKeys(InFrames, Transforms, InChannels, TickResolution, TransformSection, true);
+		}
+		else
+		{
+			InComponentHandle->AddTransformKeys(InFrames, InTransforms, InChannels, TickResolution, TransformSection, true);
+		}
 	}
 
 	void BakeControl(
@@ -90,7 +284,46 @@ namespace
 		const UMovieScene* MovieScene = InSequencer->GetFocusedMovieSceneSequence()->GetMovieScene();
 		UMovieSceneSection* Section = nullptr; //cnntrol rig doesn't need section it instead
 		InControlHandle->AddTransformKeys(InFrames, InLocalTransforms, InChannels, MovieScene->GetTickResolution(), Section, true);
+	}
+
+	UMovieSceneControlRigParameterSection* GetControlConstraintSection(
+	const UTransformableControlHandle* InHandle,
+	const TSharedPtr<ISequencer>& InSequencer)
+	{
+		const UMovieScene* MovieScene = GetMovieScene(InSequencer);
+		if (!MovieScene)
+		{
+			return nullptr;
+		}
+
+		auto GetControlRigTrack = [InHandle, MovieScene]()->UMovieSceneControlRigParameterTrack*
+		{
+			const TWeakObjectPtr<UControlRig> ControlRig = InHandle->ControlRig.LoadSynchronous();
+			if (ControlRig.IsValid())
+			{	
+				const TArray<FMovieSceneBinding>& Bindings = MovieScene->GetBindings();
+				for (const FMovieSceneBinding& Binding : Bindings)
+				{
+					UMovieSceneTrack* Track = MovieScene->FindTrack(UMovieSceneControlRigParameterTrack::StaticClass(), Binding.GetObjectGuid());
+					UMovieSceneControlRigParameterTrack* ControlRigTrack = Cast<UMovieSceneControlRigParameterTrack>(Track);
+					if (ControlRigTrack && ControlRigTrack->GetControlRig() == ControlRig)
+					{
+						return ControlRigTrack;
+					}
+				}
+			}
+			return nullptr;
+		};
+
+		UMovieSceneControlRigParameterTrack* ControlRigTrack = GetControlRigTrack();
+		if (!ControlRigTrack)
+		{
+			return nullptr;
+		}
 
+		const TArray<UMovieSceneSection*>& AllSections = ControlRigTrack->GetAllSections();
+		UMovieSceneSection* Section = AllSections.IsEmpty() ? ControlRigTrack->FindSection(0) : AllSections[0];
+		return Cast<UMovieSceneControlRigParameterSection>(Section);
 	}
 }
 
@@ -102,7 +335,6 @@ void FConstraintBaker::AddTransformKeys(
 	const TArray<FTransform>& InTransforms,
 	const EMovieSceneTransformChannel& InChannels)
 {
-
 	if (const UTransformableComponentHandle* ComponentHandle = Cast<UTransformableComponentHandle>(InHandle))
 	{
 		return BakeComponent(InSequencer, ComponentHandle, InFrames, InTransforms, InChannels); 
@@ -260,14 +492,18 @@ void FConstraintBaker::Bake(UWorld* InWorld,
 	{
 		return;
 	}
+
+	const TObjectPtr<UTransformableHandle>& ChildHandle = InConstraint->ChildTRSHandle;
 	
 	//get the section to be used later to delete the extra transform keys at the frame -1 times, abort if not there for some reason
-	UMovieSceneSection* Section = nullptr;
-	if (const UTransformableControlHandle* ControlHandle = Cast<UTransformableControlHandle>(InConstraint->ChildTRSHandle))
+	UMovieSceneSection* ConstraintSection = nullptr;
+	UMovieSceneSection* TransformSection = nullptr;
+	if (const UTransformableControlHandle* ControlHandle = Cast<UTransformableControlHandle>(ChildHandle))
 	{
-		Section = FConstraintChannelHelper::GetControlSection(ControlHandle, InSequencer);
+		TransformSection = FConstraintChannelHelper::GetControlSection(ControlHandle, InSequencer);
+		ConstraintSection = GetControlConstraintSection(ControlHandle, InSequencer);
 	}
-	else if (const UTransformableComponentHandle* ComponentHandle = Cast<UTransformableComponentHandle>(InConstraint->ChildTRSHandle))
+	else if (const UTransformableComponentHandle* ComponentHandle = Cast<UTransformableComponentHandle>(ChildHandle))
 	{ 
 		//todo move to function also used by SmartConstraintKey
 		AActor* Actor = ComponentHandle->Component->GetOwner();
@@ -282,17 +518,17 @@ void FConstraintBaker::Bake(UWorld* InWorld,
 		{
 			return;
 		}
-		Section = MovieSceneToolHelpers::GetTransformSection(InSequencer.Get(), Guid, LocalTransform);
-	}
-	if (Section == nullptr)
-	{
-		return;
+
+		TransformSection = MovieSceneToolHelpers::GetTransformSection(InSequencer.Get(), Guid, LocalTransform);
+		ConstraintSection = GetComponentConstraintSection(InSequencer, Guid, LocalTransform);
 	}
-	IMovieSceneConstrainedSection* ConstrainedSection = Cast<IMovieSceneConstrainedSection>(Section);
-	if (ConstrainedSection == nullptr)
+
+	IMovieSceneConstrainedSection* ConstrainedSection = Cast<IMovieSceneConstrainedSection>(ConstraintSection);
+	if (ConstrainedSection == nullptr || TransformSection == nullptr)
 	{
 		return;
 	}
+	
 	FConstraintAndActiveChannel* ActiveChannel = ConstrainedSection->GetConstraintChannel(InConstraint->GetFName());
 	if (ActiveChannel == nullptr)
 	{
@@ -308,6 +544,33 @@ void FConstraintBaker::Bake(UWorld* InWorld,
 	else
 	{
 		GetMinimalFramesToBake(InWorld, InConstraint, InSequencer, ConstrainedSection, FramesToBake);
+
+		// if it needs recomposition 
+		if (ConstraintSection != TransformSection && !FramesToBake.IsEmpty())
+		{
+			const TMovieSceneChannelData<const bool> ConstraintChannelData = ActiveChannel->ActiveChannel.GetData();
+			const TArrayView<const FFrameNumber> ConstraintFrames = ConstraintChannelData.GetTimes();
+			const TArrayView<const bool> ConstraintValues = ConstraintChannelData.GetValues();
+			for (int32 Index = 0; Index < ConstraintFrames.Num(); ++Index)
+			{
+				// if the key is active
+				if (ConstraintValues[Index])
+				{
+					const int32 FrameIndex = FramesToBake.IndexOfByKey(ConstraintFrames[Index]);
+					if (FrameIndex != INDEX_NONE)
+					{
+						// T-1 is not part of the frames to bake, then we need to compensate:
+						// in the context of additive, we need to add a 'zeroed' key to prevent from popping
+						const FFrameNumber FrameMinusOne(ConstraintFrames[Index] - 1);
+						const bool bAddMinusOne = (FrameIndex == 0) ? true : !ConstraintFrames.Contains(FrameMinusOne);  
+						if (bAddMinusOne)
+						{
+							FramesToBake.Insert(FrameMinusOne, FrameIndex);
+						}
+					}
+				}
+			}
+		}
 	}
 	
 	FCompensationEvaluator Evaluator(InConstraint);
@@ -318,15 +581,15 @@ void FConstraintBaker::Bake(UWorld* InWorld,
 		return;
 	}
 	
-	Section->Modify();
+	ConstraintSection->Modify();
 
 	// disable constraint and delete extra transform keys
 	TMovieSceneChannelData<bool> ConstraintChannelData = ActiveChannel->ActiveChannel.GetData();
 	const TArrayView<const FFrameNumber> ConstraintFrames = ConstraintChannelData.GetTimes();
 	
 	// get transform channels
-	const TArrayView<FMovieSceneFloatChannel*> FloatTransformChannels = InConstraint->ChildTRSHandle->GetFloatChannels(Section);
-	const TArrayView<FMovieSceneDoubleChannel*> DoubleTransformChannels = InConstraint->ChildTRSHandle->GetDoubleChannels(Section);
+	const TArrayView<FMovieSceneFloatChannel*> FloatTransformChannels = ChildHandle->GetFloatChannels(TransformSection);
+	const TArrayView<FMovieSceneDoubleChannel*> DoubleTransformChannels = ChildHandle->GetDoubleChannels(TransformSection);
 
 	for (int32 Index = 0; Index < ConstraintFrames.Num(); ++Index)
 	{
@@ -349,9 +612,10 @@ void FConstraintBaker::Bake(UWorld* InWorld,
 			}
 		}
 	}
+
 	// now bake to channel curves
 	const EMovieSceneTransformChannel Channels = InConstraint->GetChannelsToKey();
-	AddTransformKeys(InSequencer, InConstraint->ChildTRSHandle, FramesToBake, Transforms, Channels);
+	AddTransformKeys(InSequencer, ChildHandle, FramesToBake, Transforms, Channels);
 
 	// notify
 	InSequencer->RequestEvaluate();
