Sequencer: Need to make sure key handle maps transact in undo otherwise we will lose selection in the curve editor on undo since we don't save the key handles and there's a check there to only reselect after undo those key handles that match.

diff --git a/Engine/Plugins/Animation/ControlRig/Source/ControlRig/Public/Sequencer/MovieSceneControlRigSpaceChannel.h b/Engine/Plugins/Animation/ControlRig/Source/ControlRig/Public/Sequencer/MovieSceneControlRigSpaceChannel.h
--- a/Engine/Plugins/Animation/ControlRig/Source/ControlRig/Public/Sequencer/MovieSceneControlRigSpaceChannel.h
+++ b/Engine/Plugins/Animation/ControlRig/Source/ControlRig/Public/Sequencer/MovieSceneControlRigSpaceChannel.h
@@ -129,6 +129,8 @@ private:
 	UPROPERTY(meta = (KeyValues))
 	TArray<FMovieSceneControlRigSpaceBaseKey> KeyValues;
 
+	/** This needs to be a UPROPERTY so it gets saved into editor transactions but transient so it doesn't get saved into assets. */
+	UPROPERTY(Transient)
 	FMovieSceneKeyHandleMap KeyHandles;
 	
 	FMovieSceneControlRigSpaceChannelSpaceNoLongerUsedEvent SpaceNoLongerUsedEvent;
diff --git a/Engine/Source/Runtime/MovieScene/Public/Channels/MovieSceneAudioTriggerChannel.h b/Engine/Source/Runtime/MovieScene/Public/Channels/MovieSceneAudioTriggerChannel.h
--- a/Engine/Source/Runtime/MovieScene/Public/Channels/MovieSceneAudioTriggerChannel.h
+++ b/Engine/Source/Runtime/MovieScene/Public/Channels/MovieSceneAudioTriggerChannel.h
@@ -107,6 +107,8 @@ private:
 	UPROPERTY(meta = (KeyValues))
 	TArray<bool> Values;
 
+	/** This needs to be a UPROPERTY so it gets saved into editor transactions but transient so it doesn't get saved into assets. */
+	UPROPERTY(Transient)
 	FMovieSceneKeyHandleMap KeyHandles;
 };
 
diff --git a/Engine/Source/Runtime/MovieScene/Public/Channels/MovieSceneBoolChannel.h b/Engine/Source/Runtime/MovieScene/Public/Channels/MovieSceneBoolChannel.h
--- a/Engine/Source/Runtime/MovieScene/Public/Channels/MovieSceneBoolChannel.h
+++ b/Engine/Source/Runtime/MovieScene/Public/Channels/MovieSceneBoolChannel.h
@@ -167,6 +167,8 @@ protected:
 	UPROPERTY(meta=(KeyValues))
 	TArray<bool> Values;
 
+	/** This needs to be a UPROPERTY so it gets saved into editor transactions but transient so it doesn't get saved into assets. */
+	UPROPERTY(Transient)
 	FMovieSceneKeyHandleMap KeyHandles;
 };
 
diff --git a/Engine/Source/Runtime/MovieScene/Public/Channels/MovieSceneByteChannel.h b/Engine/Source/Runtime/MovieScene/Public/Channels/MovieSceneByteChannel.h
--- a/Engine/Source/Runtime/MovieScene/Public/Channels/MovieSceneByteChannel.h
+++ b/Engine/Source/Runtime/MovieScene/Public/Channels/MovieSceneByteChannel.h
@@ -187,6 +187,8 @@ private:
 	UPROPERTY()
 	TObjectPtr<UEnum> Enum;
 
+	/** This needs to be a UPROPERTY so it gets saved into editor transactions but transient so it doesn't get saved into assets. */
+	UPROPERTY(Transient)
 	FMovieSceneKeyHandleMap KeyHandles;
 };
 
diff --git a/Engine/Source/Runtime/MovieScene/Public/Channels/MovieSceneIntegerChannel.h b/Engine/Source/Runtime/MovieScene/Public/Channels/MovieSceneIntegerChannel.h
--- a/Engine/Source/Runtime/MovieScene/Public/Channels/MovieSceneIntegerChannel.h
+++ b/Engine/Source/Runtime/MovieScene/Public/Channels/MovieSceneIntegerChannel.h
@@ -190,6 +190,8 @@ private:
 	UPROPERTY(meta=(KeyValues))
 	TArray<int32> Values;
 
+	/** This needs to be a UPROPERTY so it gets saved into editor transactions but transient so it doesn't get saved into assets. */
+	UPROPERTY(Transient)
 	FMovieSceneKeyHandleMap KeyHandles;
 };
 
diff --git a/Engine/Source/Runtime/MovieScene/Public/Channels/MovieSceneObjectPathChannel.h b/Engine/Source/Runtime/MovieScene/Public/Channels/MovieSceneObjectPathChannel.h
--- a/Engine/Source/Runtime/MovieScene/Public/Channels/MovieSceneObjectPathChannel.h
+++ b/Engine/Source/Runtime/MovieScene/Public/Channels/MovieSceneObjectPathChannel.h
@@ -176,6 +176,8 @@ private:
 	UPROPERTY()
 	FMovieSceneObjectPathChannelKeyValue DefaultValue;
 
+	/** This needs to be a UPROPERTY so it gets saved into editor transactions but transient so it doesn't get saved into assets. */
+	UPROPERTY(Transient)
 	FMovieSceneKeyHandleMap KeyHandles;
 };
 
diff --git a/Engine/Source/Runtime/MovieSceneTracks/Public/Channels/MovieSceneCameraShakeSourceTriggerChannel.h b/Engine/Source/Runtime/MovieSceneTracks/Public/Channels/MovieSceneCameraShakeSourceTriggerChannel.h
--- a/Engine/Source/Runtime/MovieSceneTracks/Public/Channels/MovieSceneCameraShakeSourceTriggerChannel.h
+++ b/Engine/Source/Runtime/MovieSceneTracks/Public/Channels/MovieSceneCameraShakeSourceTriggerChannel.h
@@ -96,6 +96,8 @@ private:
 	UPROPERTY(meta=(KeyValues))
 	TArray<FMovieSceneCameraShakeSourceTrigger> KeyValues;
 
+	/** This needs to be a UPROPERTY so it gets saved into editor transactions but transient so it doesn't get saved into assets. */
+	UPROPERTY(Transient)
 	FMovieSceneKeyHandleMap KeyHandles;
 };
 
diff --git a/Engine/Source/Runtime/MovieSceneTracks/Public/Channels/MovieSceneEventChannel.h b/Engine/Source/Runtime/MovieSceneTracks/Public/Channels/MovieSceneEventChannel.h
--- a/Engine/Source/Runtime/MovieSceneTracks/Public/Channels/MovieSceneEventChannel.h
+++ b/Engine/Source/Runtime/MovieSceneTracks/Public/Channels/MovieSceneEventChannel.h
@@ -71,6 +71,8 @@ private:
 	UPROPERTY(meta=(KeyValues))
 	TArray<FMovieSceneEvent> KeyValues;
 
+	/** This needs to be a UPROPERTY so it gets saved into editor transactions but transient so it doesn't get saved into assets. */
+	UPROPERTY(Transient)
 	FMovieSceneKeyHandleMap KeyHandles;
 };
 
diff --git a/Engine/Source/Runtime/MovieSceneTracks/Public/Channels/MovieSceneStringChannel.h b/Engine/Source/Runtime/MovieSceneTracks/Public/Channels/MovieSceneStringChannel.h
--- a/Engine/Source/Runtime/MovieSceneTracks/Public/Channels/MovieSceneStringChannel.h
+++ b/Engine/Source/Runtime/MovieSceneTracks/Public/Channels/MovieSceneStringChannel.h
@@ -153,6 +153,8 @@ private:
 	UPROPERTY()
 	bool bHasDefaultValue;
 
+	/** This needs to be a UPROPERTY so it gets saved into editor transactions but transient so it doesn't get saved into assets. */
+	UPROPERTY(Transient)
 	FMovieSceneKeyHandleMap KeyHandles;
 };
 
diff --git a/Engine/Source/Runtime/MovieSceneTracks/Public/Sections/MovieSceneActorReferenceSection.h b/Engine/Source/Runtime/MovieSceneTracks/Public/Sections/MovieSceneActorReferenceSection.h
--- a/Engine/Source/Runtime/MovieSceneTracks/Public/Sections/MovieSceneActorReferenceSection.h
+++ b/Engine/Source/Runtime/MovieSceneTracks/Public/Sections/MovieSceneActorReferenceSection.h
@@ -172,6 +172,8 @@ private:
 	UPROPERTY(meta=(KeyValues))
 	TArray<FMovieSceneActorReferenceKey> KeyValues;
 
+	/** This needs to be a UPROPERTY so it gets saved into editor transactions but transient so it doesn't get saved into assets. */
+	UPROPERTY(Transient)
 	FMovieSceneKeyHandleMap KeyHandles;
 };
 
diff --git a/Engine/Source/Runtime/MovieSceneTracks/Public/Sections/MovieSceneEventSection.h b/Engine/Source/Runtime/MovieSceneTracks/Public/Sections/MovieSceneEventSection.h
--- a/Engine/Source/Runtime/MovieSceneTracks/Public/Sections/MovieSceneEventSection.h
+++ b/Engine/Source/Runtime/MovieSceneTracks/Public/Sections/MovieSceneEventSection.h
@@ -196,6 +196,8 @@ private:
 	UPROPERTY(meta=(KeyValues))
 	TArray<FEventPayload> KeyValues;
 
+	/** This needs to be a UPROPERTY so it gets saved into editor transactions but transient so it doesn't get saved into assets. */
+	UPROPERTY(Transient)
 	FMovieSceneKeyHandleMap KeyHandles;
 
 
