NonUnity Build file for anonymous namespace clashing.

diff --git a/Engine/Plugins/Animation/ControlRig/Source/ControlRigEditor/Private/Tools/ConstraintBaker.cpp b/Engine/Plugins/Animation/ControlRig/Source/ControlRigEditor/Private/Tools/ConstraintBaker.cpp
--- a/Engine/Plugins/Animation/ControlRig/Source/ControlRigEditor/Private/Tools/ConstraintBaker.cpp
+++ b/Engine/Plugins/Animation/ControlRig/Source/ControlRigEditor/Private/Tools/ConstraintBaker.cpp
@@ -27,7 +27,7 @@
 
 #define LOCTEXT_NAMESPACE "ConstraintBaker"
 
-namespace
+namespace MoveMeFor52
 {
 	UMovieScene* GetMovieScene(const TSharedPtr<ISequencer>& InSequencer)
 	{
@@ -61,7 +61,7 @@ namespace
 		const FGuid& InGuid,
 		const FTransform& InDefaultTransform)
 	{
-		UMovieScene* MovieScene = GetMovieScene(InSequencer);
+		UMovieScene* MovieScene = MoveMeFor52::GetMovieScene(InSequencer);
 		if (!MovieScene)
 		{
 			return nullptr;
@@ -244,7 +244,7 @@ namespace
 			return;
 		}
 		
-		UMovieScene3DTransformSection* TransformSection = GetTransformSection(InSequencer, Actor, InTransforms[0]);
+		UMovieScene3DTransformSection* TransformSection = MoveMeFor52::GetTransformSection(InSequencer, Actor, InTransforms[0]);
 		if (!TransformSection)
 		{
 			return;
@@ -254,11 +254,11 @@ namespace
 		const FFrameRate TickResolution = MovieScene->GetTickResolution();
 
 		const FGuid Guid = InSequencer->GetHandleToObject(Actor, false);
-		const bool bNeedsRecomposition = TransformSection != GetComponentConstraintSection(InSequencer, Guid, InTransforms[0]);
+		const bool bNeedsRecomposition = TransformSection != MoveMeFor52::GetComponentConstraintSection(InSequencer, Guid, InTransforms[0]);
 		if (bNeedsRecomposition)
 		{
 			TArray<FTransform> Transforms(InTransforms);
-			RecomposeTransforms(InSequencer, InComponentHandle->Component.Get(), TransformSection, InFrames, Transforms);
+			MoveMeFor52::RecomposeTransforms(InSequencer, InComponentHandle->Component.Get(), TransformSection, InFrames, Transforms);
 			InComponentHandle->AddTransformKeys(InFrames, Transforms, InChannels, TickResolution, TransformSection, true);
 		}
 		else
@@ -290,7 +290,7 @@ namespace
 	const UTransformableControlHandle* InHandle,
 	const TSharedPtr<ISequencer>& InSequencer)
 	{
-		const UMovieScene* MovieScene = GetMovieScene(InSequencer);
+		const UMovieScene* MovieScene = MoveMeFor52::GetMovieScene(InSequencer);
 		if (!MovieScene)
 		{
 			return nullptr;
@@ -337,12 +337,12 @@ void FConstraintBaker::AddTransformKeys(
 {
 	if (const UTransformableComponentHandle* ComponentHandle = Cast<UTransformableComponentHandle>(InHandle))
 	{
-		return BakeComponent(InSequencer, ComponentHandle, InFrames, InTransforms, InChannels); 
+		return MoveMeFor52::BakeComponent(InSequencer, ComponentHandle, InFrames, InTransforms, InChannels);
 	}
 	
 	if (const UTransformableControlHandle* ControlHandle = Cast<UTransformableControlHandle>(InHandle))
 	{
-		return BakeControl(InSequencer, ControlHandle, InFrames, InTransforms, InChannels); 
+		return MoveMeFor52::BakeControl(InSequencer, ControlHandle, InFrames, InTransforms, InChannels);
 	}
 }
 
@@ -501,7 +501,7 @@ void FConstraintBaker::Bake(UWorld* InWorld,
 	if (const UTransformableControlHandle* ControlHandle = Cast<UTransformableControlHandle>(ChildHandle))
 	{
 		TransformSection = FConstraintChannelHelper::GetControlSection(ControlHandle, InSequencer);
-		ConstraintSection = GetControlConstraintSection(ControlHandle, InSequencer);
+		ConstraintSection = MoveMeFor52::GetControlConstraintSection(ControlHandle, InSequencer);
 	}
 	else if (const UTransformableComponentHandle* ComponentHandle = Cast<UTransformableComponentHandle>(ChildHandle))
 	{ 
@@ -520,7 +520,7 @@ void FConstraintBaker::Bake(UWorld* InWorld,
 		}
 
 		TransformSection = MovieSceneToolHelpers::GetTransformSection(InSequencer.Get(), Guid, LocalTransform);
-		ConstraintSection = GetComponentConstraintSection(InSequencer, Guid, LocalTransform);
+		ConstraintSection = MoveMeFor52::GetComponentConstraintSection(InSequencer, Guid, LocalTransform);
 	}
 
 	IMovieSceneConstrainedSection* ConstrainedSection = Cast<IMovieSceneConstrainedSection>(ConstraintSection);
