PrimitiveComponent: added an editor only flag to allways allow translucent selection. This is needed such that systems using static meshes as gizmos won't have different behaviors when different materials are used. Control Rig is the one of the examples.

diff --git a/Engine/Plugins/Animation/ControlRig/Source/ControlRig/Private/ControlRigGizmoActor.cpp b/Engine/Plugins/Animation/ControlRig/Source/ControlRig/Private/ControlRigGizmoActor.cpp
--- a/Engine/Plugins/Animation/ControlRig/Source/ControlRig/Private/ControlRigGizmoActor.cpp
+++ b/Engine/Plugins/Animation/ControlRig/Source/ControlRig/Private/ControlRigGizmoActor.cpp
@@ -32,6 +32,10 @@ AControlRigShapeActor::AControlRigShapeActor(const FObjectInitializer& ObjectIni
 	StaticMeshComponent->HitProxyPriority = HPP_Wireframe;
 #endif
 
+#if WITH_EDITOR
+	StaticMeshComponent->bAlwaysAllowTranslucentSelect = true;
+#endif
+
 	RootComponent = ActorRootComponent;
 	StaticMeshComponent->SetupAttachment(RootComponent);
 	StaticMeshComponent->bCastStaticShadow = false;
diff --git a/Engine/Source/Runtime/Engine/Classes/Components/PrimitiveComponent.h b/Engine/Source/Runtime/Engine/Classes/Components/PrimitiveComponent.h
--- a/Engine/Source/Runtime/Engine/Classes/Components/PrimitiveComponent.h
+++ b/Engine/Source/Runtime/Engine/Classes/Components/PrimitiveComponent.h
@@ -657,7 +657,10 @@ protected:
 	uint8 bHasNoStreamableTextures : 1;
 
 #if WITH_EDITOR
-	bool bIgnoreBoundsForEditorFocus = false;
+	uint8 bIgnoreBoundsForEditorFocus : 1;
+
+public:
+	uint8 bAlwaysAllowTranslucentSelect : 1;
 #endif
 
 public:
diff --git a/Engine/Source/Runtime/Engine/Private/Components/PrimitiveComponent.cpp b/Engine/Source/Runtime/Engine/Private/Components/PrimitiveComponent.cpp
--- a/Engine/Source/Runtime/Engine/Private/Components/PrimitiveComponent.cpp
+++ b/Engine/Source/Runtime/Engine/Private/Components/PrimitiveComponent.cpp
@@ -354,6 +354,12 @@ UPrimitiveComponent::UPrimitiveComponent(const FObjectInitializer& ObjectInitial
 	bReplicatePhysicsToAutonomousProxy = true;
 
 	bReceiveMobileCSMShadows = true;
+
+#if WITH_EDITOR
+	bIgnoreBoundsForEditorFocus = false;
+	bAlwaysAllowTranslucentSelect = false;
+#endif
+
 #if WITH_EDITORONLY_DATA
 	bEnableAutoLODGeneration = true;
 	HitProxyPriority = HPP_World;
diff --git a/Engine/Source/Runtime/Engine/Private/EngineUtils.cpp b/Engine/Source/Runtime/Engine/Private/EngineUtils.cpp
--- a/Engine/Source/Runtime/Engine/Private/EngineUtils.cpp
+++ b/Engine/Source/Runtime/Engine/Private/EngineUtils.cpp
@@ -59,6 +59,15 @@ FTypedElementHandle HActor::GetElementHandle() const
 	return FTypedElementHandle();
 }
 
+bool HActor::AlwaysAllowsTranslucentPrimitives() const
+{
+#if WITH_EDITOR
+	return PrimComponent->bAlwaysAllowTranslucentSelect;
+#endif
+
+	return false;
+}
+
 #if !UE_BUILD_SHIPPING
 FContentComparisonHelper::FContentComparisonHelper()
 {
diff --git a/Engine/Source/Runtime/Engine/Public/EngineUtils.h b/Engine/Source/Runtime/Engine/Public/EngineUtils.h
--- a/Engine/Source/Runtime/Engine/Public/EngineUtils.h
+++ b/Engine/Source/Runtime/Engine/Public/EngineUtils.h
@@ -82,6 +82,8 @@ struct HActor : public HHitProxy
 	}
 
 	ENGINE_API virtual FTypedElementHandle GetElementHandle() const override;
+
+	ENGINE_API bool AlwaysAllowsTranslucentPrimitives() const override;
 };
 
 //
