Only reconstruct Control Rig instance, during EndLoad, on ControlRig sequencer track when its class is non-native

diff --git a/Engine/Plugins/Animation/ControlRig/Source/ControlRig/Private/Sequencer/MovieSceneControlRigParameterTrack.cpp b/Engine/Plugins/Animation/ControlRig/Source/ControlRig/Private/Sequencer/MovieSceneControlRigParameterTrack.cpp
--- a/Engine/Plugins/Animation/ControlRig/Source/ControlRig/Private/Sequencer/MovieSceneControlRigParameterTrack.cpp
+++ b/Engine/Plugins/Animation/ControlRig/Source/ControlRig/Private/Sequencer/MovieSceneControlRigParameterTrack.cpp
@@ -380,7 +380,7 @@ UMovieSceneSection* UMovieSceneControlRigParameterTrack::GetSectionToKey() const
 
 void UMovieSceneControlRigParameterTrack::ReconstructControlRig()
 {
-	if (ControlRig  && !ControlRig->HasAnyFlags(RF_NeedLoad | RF_NeedPostLoad | RF_NeedInitialization))
+	if (ControlRig && !ControlRig->HasAnyFlags(RF_NeedLoad | RF_NeedPostLoad | RF_NeedInitialization))
 	{
 		ControlRig->ConditionalPostLoad();
 		ControlRig->Initialize();
@@ -414,13 +414,17 @@ void UMovieSceneControlRigParameterTrack::PostLoad()
 
 #if WITH_EDITOR
 	FCoreUObjectDelegates::OnEndLoadPackage.AddUObject(this, &UMovieSceneControlRigParameterTrack::HandlePackageDone);
-	if (ControlRig)
+	// If we have a control Rig and it's not a native one, register OnEndLoadPackage callback on instance directly
+	if (ControlRig && !ControlRig->GetClass()->IsNative())
 	{
 		ControlRig->OnEndLoadPackage().AddUObject(this, &UMovieSceneControlRigParameterTrack::HandleControlRigPackageDone);
 	}
-#else
-	ReconstructControlRig();
+	// Otherwise try and reconstruct the control rig directly (which is fine for native classes)
+	else
 #endif
+	{		
+		ReconstructControlRig();
+	}
 }
 
 #if WITH_EDITORONLY_DATA
@@ -450,11 +454,12 @@ void UMovieSceneControlRigParameterTrack::HandlePackageDone(const FEndLoadPackag
 				return;
 			}
 		}
+
+		// Only reconstruct in case it is not a native ControlRig class
+		ReconstructControlRig();
 	}
 
 	FCoreUObjectDelegates::OnEndLoadPackage.RemoveAll(this);
-
-	ReconstructControlRig();
 }
 
 void UMovieSceneControlRigParameterTrack::HandleControlRigPackageDone(UControlRig* InControlRig)
