Control Rig: Add functionality to ControlRigLog

diff --git a/Engine/Plugins/Animation/ControlRig/Source/ControlRig/Private/ControlRigLog.cpp b/Engine/Plugins/Animation/ControlRig/Source/ControlRig/Private/ControlRigLog.cpp
--- a/Engine/Plugins/Animation/ControlRig/Source/ControlRig/Private/ControlRigLog.cpp
+++ b/Engine/Plugins/Animation/ControlRig/Source/ControlRig/Private/ControlRigLog.cpp
@@ -2,6 +2,20 @@
 
 #include "ControlRigLog.h"
 
+#if WITH_EDITOR
+TArray<FControlRigLog::FLogEntry> FControlRigLog::GetEntries(EMessageSeverity::Type InSeverity, bool bIncludeHigherSeverity)
+{
+	return Entries.FilterByPredicate([InSeverity, bIncludeHigherSeverity](const FControlRigLog::FLogEntry& Entry)
+	{
+		if (bIncludeHigherSeverity)
+		{
+			return Entry.Severity <= InSeverity;
+		}
+		return Entry.Severity == InSeverity;
+	});
+}
+#endif
+
 void FControlRigLog::Reset()
 {
 #if WITH_EDITOR
diff --git a/Engine/Plugins/Animation/ControlRig/Source/ControlRig/Public/ControlRigLog.h b/Engine/Plugins/Animation/ControlRig/Source/ControlRig/Public/ControlRigLog.h
--- a/Engine/Plugins/Animation/ControlRig/Source/ControlRig/Public/ControlRigLog.h
+++ b/Engine/Plugins/Animation/ControlRig/Source/ControlRig/Public/ControlRigLog.h
@@ -29,6 +29,8 @@ public:
 	};
 	TArray<FLogEntry> Entries;
 	TMap<FString, bool> KnownMessages;
+
+	TArray<FLogEntry> GetEntries(EMessageSeverity::Type InSeverity = EMessageSeverity::Info, bool bIncludeHigherSeverity = true);
 #endif
 
 	virtual void Reset();
