Control Rig: Ignore modification events when patching the model

diff --git a/Engine/Plugins/Animation/ControlRig/Source/ControlRigDeveloper/Private/ControlRigBlueprint.cpp b/Engine/Plugins/Animation/ControlRig/Source/ControlRigDeveloper/Private/ControlRigBlueprint.cpp
--- a/Engine/Plugins/Animation/ControlRig/Source/ControlRigDeveloper/Private/ControlRigBlueprint.cpp
+++ b/Engine/Plugins/Animation/ControlRig/Source/ControlRigDeveloper/Private/ControlRigBlueprint.cpp
@@ -857,25 +857,32 @@ void UControlRigBlueprint::PostLoad()
 		}
 		
 		InitializeModelIfRequired(false /* recompile vm */);
-		GetRigVMClient()->PatchModelsOnLoad();
-		PatchFunctionReferencesOnLoad();
-		PatchVariableNodesOnLoad();
-		PatchVariableNodesWithIncorrectType();
-		PatchRigElementKeyCacheOnLoad();
-		PatchBoundVariables();
-		PatchPropagateToChildren();
-		PatchParameterNodesOnLoad();
-		PatchLinksWithCast();
-		PatchFunctionsOnLoad();
+		{
+			TGuardValue<bool> GuardNotifications(bSuspendModelNotificationsForSelf, true);
+			GetRigVMClient()->PatchModelsOnLoad();
+			PatchFunctionReferencesOnLoad();
+			PatchVariableNodesOnLoad();
+			PatchVariableNodesWithIncorrectType();
+			PatchRigElementKeyCacheOnLoad();
+			PatchBoundVariables();
+			PatchPropagateToChildren();
+			PatchParameterNodesOnLoad();
+			PatchLinksWithCast();
+			PatchFunctionsOnLoad();
+		}
 
 #if WITH_EDITOR
 
-		// refresh the graph such that the pin hierarchies matches their CPPTypeObject
-		// this step is needed everytime we open a BP in the editor, b/c even after load
-		// model data can change while the Control Rig BP is not opened
-		// for example, if a user defined struct changed after BP load,
-		// any pin that references the struct needs to be regenerated
-		RefreshAllModels();
+		{
+			TGuardValue<bool> GuardNotifications(bSuspendModelNotificationsForSelf, true);
+
+			// refresh the graph such that the pin hierarchies matches their CPPTypeObject
+			// this step is needed everytime we open a BP in the editor, b/c even after load
+			// model data can change while the Control Rig BP is not opened
+			// for example, if a user defined struct changed after BP load,
+			// any pin that references the struct needs to be regenerated
+			RefreshAllModels();
+		}
 
 		GetControlRigBlueprintGeneratedClass()->GetRigVMGraphFunctionStore()->RemoveAllCompilationData();
 
diff --git a/Engine/Source/Developer/RigVMDeveloper/Private/RigVMModel/RigVMClient.cpp b/Engine/Source/Developer/RigVMDeveloper/Private/RigVMModel/RigVMClient.cpp
--- a/Engine/Source/Developer/RigVMDeveloper/Private/RigVMModel/RigVMClient.cpp
+++ b/Engine/Source/Developer/RigVMDeveloper/Private/RigVMModel/RigVMClient.cpp
@@ -662,6 +662,7 @@ void FRigVMClient::DestroyObject(UObject* InObject)
 void FRigVMClient::PatchModelsOnLoad()
 {
 	TArray<URigVMGraph*> AllModels = GetAllModels(true, true);
+	TGuardValue<bool> ClientIgnoreModificationsGuard(bIgnoreModelNotifications, true);
 	for(URigVMGraph* Model : AllModels)
 	{
 		URigVMController* Controller = GetOrCreateController(Model);
