Control Rig: Don't assume unit nodes to be backed up by FRigUnit

diff --git a/Engine/Plugins/Animation/ControlRig/Source/ControlRigEditor/Private/ControlRigEditorModule.cpp b/Engine/Plugins/Animation/ControlRig/Source/ControlRigEditor/Private/ControlRigEditorModule.cpp
--- a/Engine/Plugins/Animation/ControlRig/Source/ControlRigEditor/Private/ControlRigEditorModule.cpp
+++ b/Engine/Plugins/Animation/ControlRig/Source/ControlRigEditor/Private/ControlRigEditorModule.cpp
@@ -1776,29 +1776,32 @@ void FControlRigEditorModule::GetContextMenuActions(const UControlRigGraphSchema
 												{
 													if (TSharedPtr<FStructOnScope> DefaultStructScope = UnitNode->ConstructStructInstance())
 													{
-														FRigUnit* DefaultStruct = (FRigUnit*)DefaultStructScope->GetStructMemory();
+														if(DefaultStructScope->GetStruct()->IsChildOf(FRigUnit::StaticStruct()))
+														{
+															FRigUnit* DefaultStruct = (FRigUnit*)DefaultStructScope->GetStructMemory();
 
-														FString PinPath = ModelPin->GetPinPath();
-														FString Left, Right;
+															FString PinPath = ModelPin->GetPinPath();
+															FString Left, Right;
 
-														FRigElementKey SpaceKey;
-														if (URigVMPin::SplitPinPathAtStart(PinPath, Left, Right))
-														{
-															SpaceKey = DefaultStruct->DetermineSpaceForPin(Right, RigBlueprint->Hierarchy);
-														}
+															FRigElementKey SpaceKey;
+															if (URigVMPin::SplitPinPathAtStart(PinPath, Left, Right))
+															{
+																SpaceKey = DefaultStruct->DetermineSpaceForPin(Right, RigBlueprint->Hierarchy);
+															}
 
-														if (SpaceKey.IsValid())
-														{
-															if (URigVMPin* SpacePin = Injection->Node->FindPin(TEXT("Space")))
+															if (SpaceKey.IsValid())
 															{
-																if(URigVMPin* SpaceTypePin = SpacePin->FindSubPin(TEXT("Type")))
+																if (URigVMPin* SpacePin = Injection->Node->FindPin(TEXT("Space")))
 																{
-																	FString SpaceTypeStr = StaticEnum<ERigElementType>()->GetDisplayNameTextByValue((int64)SpaceKey.Type).ToString();
-																	Controller->SetPinDefaultValue(SpaceTypePin->GetPinPath(), SpaceTypeStr, true, true, false, true);
-																}
-																if(URigVMPin* SpaceNamePin = SpacePin->FindSubPin(TEXT("Name")))
-																{
-																	Controller->SetPinDefaultValue(SpaceNamePin->GetPinPath(), SpaceKey.Name.ToString(), true, true, false, true);
+																	if(URigVMPin* SpaceTypePin = SpacePin->FindSubPin(TEXT("Type")))
+																	{
+																		FString SpaceTypeStr = StaticEnum<ERigElementType>()->GetDisplayNameTextByValue((int64)SpaceKey.Type).ToString();
+																		Controller->SetPinDefaultValue(SpaceTypePin->GetPinPath(), SpaceTypeStr, true, true, false, true);
+																	}
+																	if(URigVMPin* SpaceNamePin = SpacePin->FindSubPin(TEXT("Name")))
+																	{
+																		Controller->SetPinDefaultValue(SpaceNamePin->GetPinPath(), SpaceKey.Name.ToString(), true, true, false, true);
+																	}
 																}
 															}
 														}
@@ -1878,13 +1881,16 @@ void FControlRigEditorModule::GetContextMenuActions(const UControlRigGraphSchema
 							if(ScriptStruct)
 							{
 								StructOnScope = UnitNode->ConstructStructInstance(false /* default */);
-								StructMemory = (FRigUnit*)StructOnScope->GetStructMemory();
+								if(StructOnScope->GetStruct()->IsChildOf(FRigUnit::StaticStruct()))
+								{
+									StructMemory = (FRigUnit*)StructOnScope->GetStructMemory();
 
-								//FRigUnitContext RigUnitContext;
-								//RigUnitContext.Hierarchy = TemporaryHierarchy;
-								//RigUnitContext.State = EControlRigState::Update;
-								
-								StructMemory->Execute();
+									//FRigUnitContext RigUnitContext;
+									//RigUnitContext.Hierarchy = TemporaryHierarchy;
+									//RigUnitContext.State = EControlRigState::Update;
+									
+									StructMemory->Execute();
+								}
 							}
 						}
 
