Constraints: switch the transform values in the parent space when adding keys using the channel data

diff --git a/Engine/Plugins/Animation/ControlRig/Source/ControlRig/Private/Sequencer/MovieSceneControlRigParameterSection.cpp b/Engine/Plugins/Animation/ControlRig/Source/ControlRig/Private/Sequencer/MovieSceneControlRigParameterSection.cpp
--- a/Engine/Plugins/Animation/ControlRig/Source/ControlRig/Private/Sequencer/MovieSceneControlRigParameterSection.cpp
+++ b/Engine/Plugins/Animation/ControlRig/Source/ControlRig/Private/Sequencer/MovieSceneControlRigParameterSection.cpp
@@ -589,6 +589,15 @@ struct FParameterTransformChannelEditorData
 					const FRigControlValue::FEulerTransform_Float Euler = 
 						ControlRig->GetHierarchy()
 						->GetControlValue(ControlElement, ERigControlValueType::Current).Get<FRigControlValue::FEulerTransform_Float>();
+
+					// switch translation to constraint space if needed
+					const uint32 ControlHash = UTransformableControlHandle::ComputeHash(ControlRig, ControlElement->GetName());
+					TOptional<FTransform> ConstraintSpaceTransform = FTransformConstraintUtils::GetRelativeTransform(ControlRig->GetWorld(), ControlHash);
+					if (ConstraintSpaceTransform)
+					{
+						return ConstraintSpaceTransform->GetTranslation();
+					}
+					
 					return FVector(Euler.GetTranslation());
 				}
 			}
@@ -598,12 +607,22 @@ struct FParameterTransformChannelEditorData
 
 	static TOptional<FRotator> GetRotator(UControlRig* ControlRig, FName ParameterName, UObject& InObject, FTrackInstancePropertyBindings* Bindings)
 	{
-
 		if (ControlRig)
 		{
 			FRigControlElement* ControlElement = ControlRig->FindControl(ParameterName);
 			if (ControlElement)
 			{
+				if (ControlElement->Settings.ControlType == ERigControlType::EulerTransform)
+				{
+					// switch rotation to constraint space if needed
+					const uint32 ControlHash = UTransformableControlHandle::ComputeHash(ControlRig, ControlElement->GetName());
+					TOptional<FTransform> ConstraintSpaceTransform = FTransformConstraintUtils::GetRelativeTransform(ControlRig->GetWorld(), ControlHash);
+					if (ConstraintSpaceTransform)
+					{
+						return ConstraintSpaceTransform->GetRotation().Rotator();
+					}
+				}
+				
 				return ControlRig->GetHierarchy()->GetControlPreferredRotator(ControlElement);
 			}
 		}
@@ -629,6 +648,15 @@ struct FParameterTransformChannelEditorData
 					const FRigControlValue::FEulerTransform_Float Transform = 
 						ControlRig->GetHierarchy()
 						->GetControlValue(ControlElement, ERigControlValueType::Current).Get<FRigControlValue::FEulerTransform_Float>();
+
+					// switch scale to constraint space if needed
+					const uint32 ControlHash = UTransformableControlHandle::ComputeHash(ControlRig, ControlElement->GetName());
+					TOptional<FTransform> ConstraintSpaceTransform = FTransformConstraintUtils::GetRelativeTransform(ControlRig->GetWorld(), ControlHash);
+					if (ConstraintSpaceTransform)
+					{
+						return ConstraintSpaceTransform->GetScale3D();
+					}
+					
 					return FVector(Transform.GetScale3D());
 				}
 			}
