ControlRig: Remove EControlRigState from units

diff --git a/Engine/Plugins/Animation/ControlRig/Source/ControlRig/Private/AnimNode_ControlRig.cpp b/Engine/Plugins/Animation/ControlRig/Source/ControlRig/Private/AnimNode_ControlRig.cpp
--- a/Engine/Plugins/Animation/ControlRig/Source/ControlRig/Private/AnimNode_ControlRig.cpp
+++ b/Engine/Plugins/Animation/ControlRig/Source/ControlRig/Private/AnimNode_ControlRig.cpp
@@ -36,7 +36,7 @@ FAnimNode_ControlRig::~FAnimNode_ControlRig()
 	}
 }
 
-void FAnimNode_ControlRig::HandleOnInitialized_AnyThread(UControlRig*, const EControlRigState, const FName&)
+void FAnimNode_ControlRig::HandleOnInitialized_AnyThread(UControlRig*, const FName&)
 {
 	RefPoseSetterHash.Reset();
 }
diff --git a/Engine/Plugins/Animation/ControlRig/Source/ControlRig/Private/AnimNode_ControlRigBase.cpp b/Engine/Plugins/Animation/ControlRig/Source/ControlRig/Private/AnimNode_ControlRigBase.cpp
--- a/Engine/Plugins/Animation/ControlRig/Source/ControlRig/Private/AnimNode_ControlRigBase.cpp
+++ b/Engine/Plugins/Animation/ControlRig/Source/ControlRig/Private/AnimNode_ControlRigBase.cpp
@@ -614,7 +614,7 @@ void FAnimNode_ControlRigBase::CacheBones_AnyThread(const FAnimationCacheBonesCo
 		if(ControlRig->IsConstructionModeEnabled() ||
 			(ControlRig->IsConstructionRequired() && (bControlRigRequiresInitialization || bIsLODChange)))
 		{
-			ControlRig->Execute(EControlRigState::Update, FRigUnit_PrepareForExecution::EventName);
+			ControlRig->Execute(FRigUnit_PrepareForExecution::EventName);
 			bControlRigRequiresInitialization = false;
 		}
 
@@ -805,8 +805,7 @@ void FAnimNode_ControlRigBase::CacheBones_AnyThread(const FAnimationCacheBonesCo
 		{
 			// re-init only if this is the first run
 			// and restore control values
-			FControlRigControlScope Scope(ControlRig);
-			ControlRig->Execute(EControlRigState::Init, FRigUnit_BeginExecution::EventName);
+			ControlRig->RequestInit();
 			bControlRigRequiresInitialization = false;
 		}
 		
diff --git a/Engine/Plugins/Animation/ControlRig/Source/ControlRig/Private/ControlRig.cpp b/Engine/Plugins/Animation/ControlRig/Source/ControlRig/Private/ControlRig.cpp
--- a/Engine/Plugins/Animation/ControlRig/Source/ControlRig/Private/ControlRig.cpp
+++ b/Engine/Plugins/Animation/ControlRig/Source/ControlRig/Private/ControlRig.cpp
@@ -59,19 +59,12 @@ static TAutoConsoleVariable<int32> CVarControlRigDisableExecutionAll(TEXT("Contr
 // CVar to disable swapping to nativized vms 
 static TAutoConsoleVariable<int32> CVarControlRigDisableNativizedVMs(TEXT("ControlRig.DisableNativizedVMs"), 1, TEXT("if nonzero we disable swapping to nativized VMs."));
 
-static bool bControlRigUseVMSnapshots = false;
-static FAutoConsoleVariableRef CVarControlRigUseVMSnapshots(
-	TEXT("ControlRig.UseVMSnapshots"),
-	bControlRigUseVMSnapshots,
-	TEXT("If True the VM will try to reuse previous initializations of the same rig."));
-
 UControlRig::UControlRig(const FObjectInitializer& ObjectInitializer)
 	: Super(ObjectInitializer)
 	, DeltaTime(0.0f)
 	, AbsoluteTime(0.0f)
 	, FramesPerSecond(0.0f)
 	, bAccumulateTime(true)
-	, LatestExecutedState(EControlRigState::Invalid)
 #if WITH_EDITOR
 	, ControlRigLog(nullptr)
 	, bEnableControlRigLogging(true)
@@ -155,20 +148,6 @@ void UControlRig::BeginDestroy()
 	}
 #endif
 
-	// on destruction clear out the initialized snapshots
-	if (HasAnyFlags(RF_ClassDefaultObject))
-	{
-		for (TPair<uint32, TObjectPtr<URigVM>>& Pair : InitializedVMSnapshots)
-		{
-			if (IsValid(Pair.Value))
-			{
-				TObjectPtr<URigVM> VMSnapshot = Pair.Value;
-				VMSnapshot->RemoveFromRoot();				
-			}
-		}
-		InitializedVMSnapshots.Reset();
-	}
-
 	TRACE_OBJECT_LIFETIME_END(this);
 }
 
@@ -198,7 +177,7 @@ UWorld* UControlRig::GetWorld() const
 	return nullptr;
 }
 
-void UControlRig::Initialize(bool bInitRigUnits)
+void UControlRig::Initialize(bool bRequestInit)
 {
 	TRACE_OBJECT_LIFETIME_BEGIN(this);
 
@@ -211,20 +190,6 @@ void UControlRig::Initialize(bool bInitRigUnits)
 		return;
 	}
 
-	// recompute the hash used to differentiate VMs based on their memory layout 
-	if (HasAnyFlags(RF_ClassDefaultObject))
-	{
-		CachedMemoryHash = 0;
-
-		if(VM)
-		{
-			CachedMemoryHash = HashCombine(
-				VM->GetLiteralMemory()->GetMemoryHash(),
-				VM->GetWorkMemory()->GetMemoryHash()
-			);
-		}
-	}
-
 	if (IsTemplate())
 	{
 		// don't initialize template class 
@@ -243,7 +208,7 @@ void UControlRig::Initialize(bool bInitRigUnits)
 	// should refresh mapping 
 	RequestConstruction();
 
-	if (bInitRigUnits)
+	if (bRequestInit)
 	{
 		RequestInit();
 	}
@@ -255,6 +220,61 @@ void UControlRig::Initialize(bool bInitRigUnits)
 	GetHierarchy()->UpdateVisibilityOnProxyControls();
 }
 
+bool UControlRig::InitializeVM(const FName& InEventName)
+{
+	// update the VM's external variables
+	VM->ClearExternalVariables();
+	TArray<FRigVMExternalVariable> ExternalVariables = GetExternalVariablesImpl(false);
+	for (FRigVMExternalVariable ExternalVariable : ExternalVariables)
+	{
+		VM->AddExternalVariable(ExternalVariable);
+	}
+
+	TArray<URigVMMemoryStorage*> LocalMemory = VM->GetLocalMemoryArray();
+	const bool bResult = VM->Initialize(LocalMemory);
+	if(bResult)
+	{
+		bRequiresInitExecution = false;
+	}
+	bRequiresConstructionEvent = true;
+
+#if WITH_EDITOR
+	// setup the hierarchy's controller log function
+	if(URigHierarchyController* HierarchyController = GetHierarchy()->GetController(true))
+	{
+		HierarchyController->LogFunction = [this](EMessageSeverity::Type InSeverity, const FString& Message)
+		{
+			const FRigVMExtendedExecuteContext& ExtendedExecuteContext = GetVM()->GetContext();
+			const FRigVMExecuteContext& PublicContext = ExtendedExecuteContext.GetPublicData<>(); 
+			if(ControlRigLog)
+			{
+				ControlRigLog->Report(InSeverity,PublicContext.FunctionName,PublicContext.InstructionIndex, Message);
+			}
+			else
+			{
+				LogOnce(InSeverity,PublicContext.InstructionIndex, Message);
+			}
+		};
+	}
+#endif
+
+	// reset the time and caches during init
+	AbsoluteTime = DeltaTime = 0.f;
+		
+	if(VM)
+	{
+		VM->GetContext().GetPublicDataSafe<FControlRigExecuteContext>().GetNameCache()->Reset();
+	}
+
+	if (InitializedEvent.IsBound())
+	{
+		FControlRigBracketScope BracketScope(InitBracket);
+		InitializedEvent.Broadcast(this, InEventName);
+	}
+
+	return bResult;
+}
+
 void UControlRig::InitializeFromCDO()
 {
 	DECLARE_SCOPE_HIERARCHICAL_COUNTER_FUNC()
@@ -347,7 +367,7 @@ void UControlRig::Evaluate_AnyThread()
 	
 	for (const FName& EventName : EventQueueToRun)
 	{
-		Execute(EControlRigState::Update, EventName);
+		Execute(EventName);
 
 #if WITH_EDITOR
 		if (VM)
@@ -664,21 +684,6 @@ void UControlRig::InstantiateVMFromCDO()
 		}
 	}
 
-	if(VM)
-	{
-		if(!VM->IsNativized())
-		{
-			CachedMemoryHash = HashCombine(
-				VM->GetLiteralMemory()->GetMemoryHash(),
-				VM->GetWorkMemory()->GetMemoryHash()
-			);
-		}
-		else
-		{
-			CachedMemoryHash = Cast<URigVMNativized>(VM)->GetVMHash();
-		}
-	}
-
 	RequestInit();
 }
 
@@ -701,23 +706,33 @@ void UControlRig::CopyExternalVariableDefaultValuesFromCDO()
 	}
 }
 
-bool UControlRig::Execute(const EControlRigState InState, const FName& InEventName)
+bool UControlRig::Execute(const FName& InEventName)
 {
 	if(!CanExecute())
 	{
 		return false;
 	}
 
+	bool bJustRanInit = false;
+	if(bRequiresInitExecution)
+	{
+		const TGuardValue<float> AbsoluteTimeGuard(AbsoluteTime, AbsoluteTime);
+		const TGuardValue<float> DeltaTimeGuard(DeltaTime, DeltaTime);
+		if(!InitializeVM(InEventName))
+		{
+			return false;
+		}
+		bJustRanInit = true;
+	}
+
 	if(EventQueueToRun.IsEmpty())
 	{
 		EventQueueToRun = EventQueue;
 	}
-	
+
 	const bool bIsEventInQueue = EventQueueToRun.Contains(InEventName);
 	const bool bIsEventFirstInQueue = !EventQueueToRun.IsEmpty() && EventQueueToRun[0] == InEventName; 
 	const bool bIsEventLastInQueue = !EventQueueToRun.IsEmpty() && EventQueueToRun.Last() == InEventName;
-	const bool bIsInitializingMemory = InState == EControlRigState::Init;
-	const bool bIsExecutingInstructions = InState == EControlRigState::Update;
 	const bool bIsConstructionEvent = InEventName == FRigUnit_PrepareForExecution::EventName;
 	const bool bIsForwardSolve = InEventName == FRigUnit_BeginExecution::EventName;
 	const bool bIsInteractionEvent = InEventName == FRigUnit_InteractionExecution::EventName;
@@ -727,12 +742,11 @@ bool UControlRig::Execute(const EControlRigState InState, const FName& InEventNa
 	DECLARE_SCOPE_HIERARCHICAL_COUNTER_FUNC()
 	QUICK_SCOPE_CYCLE_COUNTER(STAT_ControlRig_Execute);
 	
-	LatestExecutedState = InState;
-
 	FRigVMExtendedExecuteContext& ExtendedExecuteContext = GetVM()->GetContext();
 	FControlRigExecuteContext& PublicContext = ExtendedExecuteContext.GetPublicDataSafe<FControlRigExecuteContext>();
 	PublicContext.SetDeltaTime(DeltaTime);
 	PublicContext.SetAbsoluteTime(AbsoluteTime);
+	PublicContext.SetFramesPerSecond(GetCurrentFramesPerSecond());
 
 	if (VM)
 	{
@@ -741,36 +755,6 @@ bool UControlRig::Execute(const EControlRigState InState, const FName& InEventNa
 			InstantiateVMFromCDO();
 		}
 
-		if (bIsInitializingMemory)
-		{
-			VM->ClearExternalVariables();
-
-			TArray<FRigVMExternalVariable> ExternalVariables = GetExternalVariablesImpl(false);
-			for (FRigVMExternalVariable ExternalVariable : ExternalVariables)
-			{
-				VM->AddExternalVariable(ExternalVariable);
-			}
-			
-#if WITH_EDITOR
-			// setup the hierarchy's controller log function
-			if(URigHierarchyController* HierarchyController = GetHierarchy()->GetController(true))
-			{
-				HierarchyController->LogFunction = [this](EMessageSeverity::Type InSeverity, const FString& Message)
-				{
-					const FRigVMExtendedExecuteContext& ExtendedExecuteContext = GetVM()->GetContext();
-					const FRigVMExecuteContext& PublicContext = ExtendedExecuteContext.GetPublicData<>(); 
-					if(ControlRigLog)
-					{
-						ControlRigLog->Report(InSeverity,PublicContext.FunctionName,PublicContext.InstructionIndex, Message);
-					}
-					else
-					{
-						LogOnce(InSeverity,PublicContext.InstructionIndex, Message);
-					}
-				};
-			}
-#endif
-		}
 #if WITH_EDITOR
 		// default to always clear data after each execution
 		// only set a valid first entry event later when execution
@@ -806,22 +790,6 @@ bool UControlRig::Execute(const EControlRigState InState, const FName& InEventNa
 	}
 #endif
 
-	bool bJustRanInit = false;
-	if (bRequiresInitExecution)
-	{
-		bRequiresInitExecution = false;
-
-		if (!bIsInitializingMemory)
-		{
-			// if init is required we'll run init on the whole bytecode
-			if(!Execute(EControlRigState::Init, NAME_None))
-			{
-				return false;
-			}
-			bJustRanInit = true;
-		}
-	}
-
 	FRigUnitContext& Context = PublicContext.UnitContext;
 
 	// setup the draw interface for debug drawing
@@ -846,24 +814,9 @@ bool UControlRig::Execute(const EControlRigState InState, const FName& InEventNa
 	// setup the data source registry
 	Context.DataSourceRegistry = GetDataSourceRegistry();
 
-	// reset the time and caches during init
-	if (bIsInitializingMemory)
-	{
-		AbsoluteTime = DeltaTime = 0.f;
-		
-		if(VM)
-		{
-			VM->GetPublicData<>().GetNameCache()->Reset();
-		}
-	}
-
 	// setup the context with further fields
-	Context.DeltaTime = DeltaTime;
-	Context.AbsoluteTime = AbsoluteTime;
-	Context.FramesPerSecond = GetCurrentFramesPerSecond();
 	Context.InteractionType = InteractionType;
 	Context.ElementsBeingInteracted = ElementsBeingInteracted;
-	Context.State = InState;
 	PublicContext.Hierarchy = GetHierarchy();
 
 	// allow access to the hierarchy
@@ -984,7 +937,7 @@ bool UControlRig::Execute(const EControlRigState InState, const FName& InEventNa
 	bool bSuccess = true;
 
 	// we'll special case the construction event here
-	if (bIsConstructionEvent && !bIsInitializingMemory)
+	if (bIsConstructionEvent)
 	{
 		// remember the previous selection
 		const TArray<FRigElementKey> PreviousSelection = GetHierarchy()->GetSelectedKeys();
@@ -996,7 +949,7 @@ bool UControlRig::Execute(const EControlRigState InState, const FName& InEventNa
 			if (PreConstructionForUIEvent.IsBound())
 			{
 				FControlRigBracketScope BracketScope(PreConstructionBracket);
-				PreConstructionForUIEvent.Broadcast(this, EControlRigState::Update, FRigUnit_PrepareForExecution::EventName);
+				PreConstructionForUIEvent.Broadcast(this, FRigUnit_PrepareForExecution::EventName);
 			}
 
 #if WITH_EDITOR
@@ -1062,7 +1015,7 @@ bool UControlRig::Execute(const EControlRigState InState, const FName& InEventNa
 					if (PreConstructionEvent.IsBound())
 					{
 						FControlRigBracketScope BracketScope(PreConstructionBracket);
-						PreConstructionEvent.Broadcast(this, EControlRigState::Update, FRigUnit_PrepareForExecution::EventName);
+						PreConstructionEvent.Broadcast(this, FRigUnit_PrepareForExecution::EventName);
 					}
 
 					bSuccess = ExecuteUnits(Context, FRigUnit_PrepareForExecution::EventName);
@@ -1072,7 +1025,7 @@ bool UControlRig::Execute(const EControlRigState InState, const FName& InEventNa
 				if (PostConstructionEvent.IsBound())
 				{
 					FControlRigBracketScope BracketScope(PostConstructionBracket);
-					PostConstructionEvent.Broadcast(this, EControlRigState::Update, FRigUnit_PrepareForExecution::EventName);
+					PostConstructionEvent.Broadcast(this, FRigUnit_PrepareForExecution::EventName);
 				}
 			}
 			
@@ -1120,7 +1073,7 @@ bool UControlRig::Execute(const EControlRigState InState, const FName& InEventNa
 #if WITH_EDITOR
 		// only set a valid first entry event when execution
 		// has passed the initialization stage and there are multiple events present
-		if (EventQueueToRun.Num() >= 2 && VM && !bIsInitializingMemory)
+		if (EventQueueToRun.Num() >= 2 && VM)
 		{
 			VM->SetFirstEntryEventInEventQueue(EventQueueToRun[0]);
 		}
@@ -1134,12 +1087,12 @@ bool UControlRig::Execute(const EControlRigState InState, const FName& InEventNa
 		}
 #endif
 		
-		if (bIsExecutingInstructions && bIsForwardSolve)
+		if (bIsForwardSolve)
 		{
 			if (PreForwardsSolveEvent.IsBound())
 			{
 				FControlRigBracketScope BracketScope(PreForwardsSolveBracket);
-				PreForwardsSolveEvent.Broadcast(this, EControlRigState::Update, FRigUnit_BeginExecution::EventName);
+				PreForwardsSolveEvent.Broadcast(this, FRigUnit_BeginExecution::EventName);
 			}
 		}
 
@@ -1152,25 +1105,20 @@ bool UControlRig::Execute(const EControlRigState InState, const FName& InEventNa
 		}
 #endif
 
-		if (bIsExecutingInstructions && bIsForwardSolve)
+		if (bIsForwardSolve)
 		{
 			if (PostForwardsSolveEvent.IsBound())
 			{
 				FControlRigBracketScope BracketScope(PostForwardsSolveBracket);
-				PostForwardsSolveEvent.Broadcast(this, EControlRigState::Update, FRigUnit_BeginExecution::EventName);
+				PostForwardsSolveEvent.Broadcast(this, FRigUnit_BeginExecution::EventName);
 			}
 		}
-		
-		if (bIsInitializingMemory && !bIsForwardSolve)
-		{
-			bSuccess = ExecuteUnits(Context, FRigUnit_BeginExecution::EventName);
-		}
 	}
 
 #if WITH_EDITOR
 
 	// for the last event in the queue - clear the log message queue
-	if (ControlRigLog != nullptr && bEnableControlRigLogging && !bIsInitializingMemory)
+	if (ControlRigLog != nullptr && bEnableControlRigLogging)
 	{
 		if (bJustRanInit)
 		{
@@ -1200,26 +1148,15 @@ bool UControlRig::Execute(const EControlRigState InState, const FName& InEventNa
 	}
 #endif
 
-	if (bIsInitializingMemory)
+	if(!bIsEventInQueue || bIsEventLastInQueue) 
 	{
-		if (InitializedEvent.IsBound())
-		{
-			FControlRigBracketScope BracketScope(InitBracket);
-			InitializedEvent.Broadcast(this, EControlRigState::Init, InEventName);
-		}
+		DeltaTime = 0.f;
 	}
-	else if (bIsExecutingInstructions)
-	{
-		if(!bIsEventInQueue || bIsEventLastInQueue) 
-		{
-			DeltaTime = 0.f;
-		}
 
-		if (ExecutedEvent.IsBound())
-		{
-			FControlRigBracketScope BracketScope(UpdateBracket);
-			ExecutedEvent.Broadcast(this, EControlRigState::Update, InEventName);
-		}
+	if (ExecutedEvent.IsBound())
+	{
+		FControlRigBracketScope BracketScope(UpdateBracket);
+		ExecutedEvent.Broadcast(this, InEventName);
 	}
 
 	// close remaining undo brackets from hierarchy
@@ -1228,11 +1165,11 @@ bool UControlRig::Execute(const EControlRigState InState, const FName& InEventNa
 		FRigEventContext EventContext;
 		EventContext.Event = ERigEvent::CloseUndoBracket;
 		EventContext.SourceEventName = InEventName;
-		EventContext.LocalTime = Context.AbsoluteTime;
+		EventContext.LocalTime = PublicContext.GetAbsoluteTime();
 		HandleHierarchyEvent(GetHierarchy(), EventContext);
 	}
 
-	if (Context.DrawInterface && Context.DrawContainer && bIsEventLastInQueue && bIsExecutingInstructions) 
+	if (Context.DrawInterface && Context.DrawContainer && bIsEventLastInQueue) 
 	{
 		Context.DrawInterface->Instructions.Append(Context.DrawContainer->Instructions);
 
@@ -1510,118 +1447,56 @@ bool UControlRig::ExecuteUnits(FRigUnitContext& InOutContext, const FName& InEve
 		}
 #endif
 		
-		const bool bUseInitializationSnapshots = bControlRigUseVMSnapshots && !VM->IsNativized();
 		const bool bUseDebuggingSnapshots = !VM->IsNativized();
 		
 		TArray<URigVMMemoryStorage*> LocalMemory = VM->GetLocalMemoryArray();
 
-		bool bSuccess = true;
-
-		if (InOutContext.State == EControlRigState::Init)
+#if WITH_EDITOR
+		if(bUseDebuggingSnapshots)
 		{
-			if(IsInGameThread() && bUseInitializationSnapshots)
+			if(URigVM* SnapShotVM = GetSnapshotVM(false)) // don't create it for normal runs
 			{
-				const uint32 SnapshotHash = GetHashForInitializeVMSnapShot();
-				UControlRig* CDO = Cast<UControlRig>(GetClass()->GetDefaultObject());
+				const bool bIsEventFirstInQueue = !EventQueueToRun.IsEmpty() && EventQueueToRun[0] == InEventName; 
+				const bool bIsEventLastInQueue = !EventQueueToRun.IsEmpty() && EventQueueToRun.Last() == InEventName;
 
-				bool bIsValidSnapshot = false;
-				if(SnapshotHash != 0)
+				if (VM->GetHaltedAtBreakpoint().IsValid())
 				{
-					TObjectPtr<URigVM>* InitializedVMSnapshotPtr = CDO->InitializedVMSnapshots.Find(SnapshotHash);
-					if(InitializedVMSnapshotPtr && *InitializedVMSnapshotPtr)
+					if(bIsEventFirstInQueue)
 					{
-						const URigVM* InitializedVMSnapshot = InitializedVMSnapshotPtr->Get();
-
-						if(VM->WorkMemoryStorageObject->GetClass() == InitializedVMSnapshot->WorkMemoryStorageObject->GetClass() &&
-							InitializedVMSnapshot->WorkMemoryStorageObject->IsValidLowLevel())
-						{
-							InitializedVMSnapshots.Reset();
-
-							VM->WorkMemoryStorageObject->CopyFrom(InitializedVMSnapshot->WorkMemoryStorageObject);
-							VM->InvalidateCachedMemory();
-							VM->Initialize(LocalMemory, false);
-							bIsValidSnapshot = true;
-						}
-						else
-						{
-							CDO->InitializedVMSnapshots.Remove(SnapshotHash);
-						}
+						VM->CopyFrom(SnapShotVM, false, false, false, true, true);
 					}
 				}
-				
-				if(!bIsValidSnapshot)
+				else if(bIsEventLastInQueue)
 				{
-					VM->Initialize(LocalMemory);
-
-					// objects assigned to transient properties need transient flag, otherwise it might get exported during save
-					URigVM* InitializedVMSnapshot = NewObject<URigVM>(CDO, NAME_None, RF_Public | RF_Transient);
-					InitializedVMSnapshot->WorkMemoryStorageObject = NewObject<URigVMMemoryStorage>(InitializedVMSnapshot, VM->GetWorkMemory()->GetClass());
-					InitializedVMSnapshot->WorkMemoryStorageObject->CopyFrom(VM->WorkMemoryStorageObject);
-
-					CDO->InitializedVMSnapshots.Add(SnapshotHash, InitializedVMSnapshot);
-
-					// GC won't consider some subobjects that are created after the constructor as part of the CDO,
-					// so even if CDO is rooted and references these sub objects, 
-					// it is not enough to keep them alive.
-					// Hence, we have to add them to root here.
-					InitializedVMSnapshot->AddToRoot();
+					SnapShotVM->CopyFrom(VM, false, false, false, true, true);
 				}
 			}
-			else
-			{
-				bSuccess = VM->Initialize(LocalMemory);
-			}
-			bRequiresConstructionEvent = true;
 		}
-		else
-		{
-#if WITH_EDITOR
-			if(bUseDebuggingSnapshots)
-			{
-				if(URigVM* SnapShotVM = GetSnapshotVM(false)) // don't create it for normal runs
-				{
-					const bool bIsEventFirstInQueue = !EventQueueToRun.IsEmpty() && EventQueueToRun[0] == InEventName; 
-					const bool bIsEventLastInQueue = !EventQueueToRun.IsEmpty() && EventQueueToRun.Last() == InEventName;
-
-					if (VM->GetHaltedAtBreakpoint().IsValid())
-					{
-						if(bIsEventFirstInQueue)
-						{
-							VM->CopyFrom(SnapShotVM, false, false, false, true, true);
-						}
-					}
-					else if(bIsEventLastInQueue)
-					{
-						SnapShotVM->CopyFrom(VM, false, false, false, true, true);
-					}
-				}
-			}
 #endif
 
-			URigHierarchy* Hierarchy = GetHierarchy();
+		URigHierarchy* Hierarchy = GetHierarchy();
 #if WITH_EDITOR
 
-			bool bRecordTransformsAtRuntime = true;
-			if(const UObject* Outer = GetOuter())
-			{
-				if(Outer->IsA<UControlRigComponent>())
-				{
-					bRecordTransformsAtRuntime = false;
-				}
-			}
-			TGuardValue<bool> RecordTransformsPerInstructionGuard(Hierarchy->bRecordTransformsAtRuntime, bRecordTransformsAtRuntime);
-			
-			if(Hierarchy->bRecordTransformsAtRuntime)
+		bool bRecordTransformsAtRuntime = true;
+		if(const UObject* Outer = GetOuter())
+		{
+			if(Outer->IsA<UControlRigComponent>())
 			{
-				Hierarchy->ReadTransformsAtRuntime.Reset();
-				Hierarchy->WrittenTransformsAtRuntime.Reset();
+				bRecordTransformsAtRuntime = false;
 			}
-			
+		}
+		TGuardValue<bool> RecordTransformsPerInstructionGuard(Hierarchy->bRecordTransformsAtRuntime, bRecordTransformsAtRuntime);
+		
+		if(Hierarchy->bRecordTransformsAtRuntime)
+		{
+			Hierarchy->ReadTransformsAtRuntime.Reset();
+			Hierarchy->WrittenTransformsAtRuntime.Reset();
+		}
+		
 #endif
-			FRigHierarchyExecuteContextBracket HierarchyContextGuard(Hierarchy, &VM->GetContext());
+		FRigHierarchyExecuteContextBracket HierarchyContextGuard(Hierarchy, &VM->GetContext());
 
-			bSuccess = VM->Execute(LocalMemory, InEventName) != ERigVMExecuteResult::Failed;
-		}
+		const bool bSuccess = VM->Execute(LocalMemory, InEventName) != ERigVMExecuteResult::Failed;
 
 #if UE_CONTROLRIG_PROFILE_EXECUTE_UNITS_NUM
 		const uint64 EndCycles = FPlatformTime::Cycles64();
@@ -2008,7 +1883,7 @@ void UControlRig::HandleExecutionReachedExit(const FName& InEventName)
 	}
 #endif
 	
-	if (LatestExecutedState != EControlRigState::Init && bAccumulateTime)
+	if (bAccumulateTime)
 	{
 		AbsoluteTime += DeltaTime;
 	}
@@ -2768,7 +2643,7 @@ void UControlRig::SetInteractionRig(UControlRig* InInteractionRig)
 		InteractionRig->Initialize(true);
 		InteractionRig->CopyPoseFromOtherRig(this);
 		InteractionRig->RequestConstruction();
-		InteractionRig->Execute(EControlRigState::Update, FRigUnit_BeginExecution::EventName);
+		InteractionRig->Execute(FRigUnit_BeginExecution::EventName);
 
 		InteractionRig->ControlModified().AddUObject(this, &UControlRig::HandleInteractionRigControlModified);
 		InteractionRig->OnInitialized_AnyThread().AddUObject(this, &UControlRig::HandleInteractionRigInitialized);
@@ -2779,7 +2654,7 @@ void UControlRig::SetInteractionRig(UControlRig* InInteractionRig)
 		ControlSelected().AddUObject(ToRawPtr(InteractionRig), &UControlRig::HandleInteractionRigControlSelected, true);
 
 		FControlRigBracketScope BracketScope(InterRigSyncBracket);
-		InteractionRig->HandleInteractionRigExecuted(this, EControlRigState::Update, FRigUnit_BeginExecution::EventName);
+		InteractionRig->HandleInteractionRigExecuted(this, FRigUnit_BeginExecution::EventName);
 	}
 }
 
@@ -2980,7 +2855,7 @@ void UControlRig::HandleInteractionRigControlModified(UControlRig* Subject, FRig
 
 }
 
-void UControlRig::HandleInteractionRigInitialized(UControlRig* Subject, EControlRigState State, const FName& EventName)
+void UControlRig::HandleInteractionRigInitialized(UControlRig* Subject, const FName& EventName)
 {
 	check(Subject);
 
@@ -2992,7 +2867,7 @@ void UControlRig::HandleInteractionRigInitialized(UControlRig* Subject, EControl
 	RequestInit();
 }
 
-void UControlRig::HandleInteractionRigExecuted(UControlRig* Subject, EControlRigState State, const FName& EventName)
+void UControlRig::HandleInteractionRigExecuted(UControlRig* Subject, const FName& EventName)
 {
 	check(Subject);
 
@@ -3003,7 +2878,7 @@ void UControlRig::HandleInteractionRigExecuted(UControlRig* Subject, EControlRig
 	FControlRigBracketScope BracketScope(InterRigSyncBracket);
 
 	CopyPoseFromOtherRig(Subject);
-	Execute(EControlRigState::Update, FRigUnit_InverseExecution::EventName);
+	Execute(FRigUnit_InverseExecution::EventName);
 
 	FRigControlModifiedContext Context;
 	Context.EventName = FRigUnit_InverseExecution::EventName;
@@ -3468,33 +3343,9 @@ void UControlRig::PostInitInstance(UControlRig* InCDO)
 			VM->AddToRoot();
 			DynamicHierarchy->AddToRoot();
 		}
-
-		// Clear the initialized VM snapshots
-		InitializedVMSnapshots.Reset();
 	}
-}
 
-uint32 UControlRig::GetHashForInitializeVMSnapShot()
-{
-	if(CachedMemoryHash == 0)
-	{
-		return 0;
-	}
-	
-	uint32 Hash = GetHierarchy()->GetNameHash();
-
-	const TArray<FRigVMExternalVariable> ExternalVariables = GetExternalVariablesImpl(false);
-
-	Hash = HashCombine(Hash, GetTypeHash(ExternalVariables.Num()));
-
-	for(const FRigVMExternalVariable& ExternalVariable : ExternalVariables)
-	{
-		Hash = HashCombine(ExternalVariable.GetTypeHash(), Hash);
-	}
-
-	Hash = HashCombine(CachedMemoryHash, Hash);
-
-	return Hash;
+	RequestInit();
 }
 
 UTransformableControlHandle* UControlRig::CreateTransformableControlHandle(
diff --git a/Engine/Plugins/Animation/ControlRig/Source/ControlRig/Private/ControlRigComponent.cpp b/Engine/Plugins/Animation/ControlRig/Source/ControlRig/Private/ControlRigComponent.cpp
--- a/Engine/Plugins/Animation/ControlRig/Source/ControlRig/Private/ControlRigComponent.cpp
+++ b/Engine/Plugins/Animation/ControlRig/Source/ControlRig/Private/ControlRigComponent.cpp
@@ -1655,7 +1655,7 @@ FName UControlRigComponent::GetComponentNameWithinActor(UActorComponent* InCompo
 	return ComponentProperty;
 }
 
-void UControlRigComponent::HandleControlRigInitializedEvent(UControlRig* InControlRig, const EControlRigState InState, const FName& InEventName)
+void UControlRigComponent::HandleControlRigInitializedEvent(UControlRig* InControlRig, const FName& InEventName)
 {
 #if WITH_EDITOR
 	if (bUpdateInEditor)
@@ -1670,7 +1670,7 @@ void UControlRigComponent::HandleControlRigInitializedEvent(UControlRig* InContr
 	}
 }
 
-void UControlRigComponent::HandleControlRigPreConstructionEvent(UControlRig* InControlRig, const EControlRigState InState, const FName& InEventName)
+void UControlRigComponent::HandleControlRigPreConstructionEvent(UControlRig* InControlRig, const FName& InEventName)
 {
 	TArray<USkeletalMeshComponent*> ComponentsToTick;
 
@@ -1720,7 +1720,7 @@ void UControlRigComponent::HandleControlRigPreConstructionEvent(UControlRig* InC
 	}
 }
 
-void UControlRigComponent::HandleControlRigPostConstructionEvent(UControlRig* InControlRig, const EControlRigState InState, const FName& InEventName)
+void UControlRigComponent::HandleControlRigPostConstructionEvent(UControlRig* InControlRig, const FName& InEventName)
 {
 #if WITH_EDITOR
 	if (bUpdateInEditor)
@@ -1735,7 +1735,7 @@ void UControlRigComponent::HandleControlRigPostConstructionEvent(UControlRig* In
 	}
 }
 
-void UControlRigComponent::HandleControlRigPreForwardsSolveEvent(UControlRig* InControlRig, const EControlRigState InState, const FName& InEventName)
+void UControlRigComponent::HandleControlRigPreForwardsSolveEvent(UControlRig* InControlRig, const FName& InEventName)
 {
 #if WITH_EDITOR
 	if (bUpdateInEditor)
@@ -1750,7 +1750,7 @@ void UControlRigComponent::HandleControlRigPreForwardsSolveEvent(UControlRig* In
 	}
 }
 
-void UControlRigComponent::HandleControlRigPostForwardsSolveEvent(UControlRig* InControlRig, const EControlRigState InState, const FName& InEventName)
+void UControlRigComponent::HandleControlRigPostForwardsSolveEvent(UControlRig* InControlRig, const FName& InEventName)
 {
 #if WITH_EDITOR
 	if (bUpdateInEditor)
@@ -1765,7 +1765,7 @@ void UControlRigComponent::HandleControlRigPostForwardsSolveEvent(UControlRig* I
 	}
 }
 
-void UControlRigComponent::HandleControlRigExecutedEvent(UControlRig* InControlRig, const EControlRigState InState, const FName& InEventName)
+void UControlRigComponent::HandleControlRigExecutedEvent(UControlRig* InControlRig, const FName& InEventName)
 {
 	TransferOutputs();
 }
diff --git a/Engine/Plugins/Animation/ControlRig/Source/ControlRig/Private/ControlRigTestData.cpp b/Engine/Plugins/Animation/ControlRig/Source/ControlRig/Private/ControlRigTestData.cpp
--- a/Engine/Plugins/Animation/ControlRig/Source/ControlRig/Private/ControlRigTestData.cpp
+++ b/Engine/Plugins/Animation/ControlRig/Source/ControlRig/Private/ControlRigTestData.cpp
@@ -232,7 +232,7 @@ bool UControlRigTestData::Record(UControlRig* InControlRig, double InRecordingDu
 	{
 		InControlRig->RequestInit();
 		PreConstructionHandle = InControlRig->OnPreConstruction_AnyThread().AddLambda(
-			[this](UControlRig* InControlRig, const EControlRigState InState, const FName& InEventName)
+			[this](UControlRig* InControlRig, const FName& InEventName)
 			{
 				Initial.Store(InControlRig, true);
 			}
@@ -240,7 +240,7 @@ bool UControlRigTestData::Record(UControlRig* InControlRig, double InRecordingDu
 	}
 
 	PreForwardHandle = InControlRig->OnPreForwardsSolve_AnyThread().AddLambda(
-		[this](UControlRig* InControlRig, const EControlRigState InState, const FName& InEventName)
+		[this](UControlRig* InControlRig, const FName& InEventName)
 		{
 			FControlRigTestDataFrame Frame;
 			Frame.Store(InControlRig);
@@ -254,7 +254,7 @@ bool UControlRigTestData::Record(UControlRig* InControlRig, double InRecordingDu
 	);
 
 	PostForwardHandle = InControlRig->OnPostForwardsSolve_AnyThread().AddLambda(
-		[this](UControlRig* InControlRig, const EControlRigState InState, const FName& InEventName)
+		[this](UControlRig* InControlRig, const FName& InEventName)
 		{
 			FControlRigTestDataFrame Frame;
 			Frame.Store(InControlRig);
@@ -297,14 +297,14 @@ bool UControlRigTestData::SetupReplay(UControlRig* InControlRig, bool bGroundTru
 	
 	InControlRig->RequestInit();
 	PreConstructionHandle = InControlRig->OnPreConstruction_AnyThread().AddLambda(
-		[this](UControlRig* InControlRig, const EControlRigState InState, const FName& InEventName)
+		[this](UControlRig* InControlRig, const FName& InEventName)
 		{
 			Initial.Restore(InControlRig, true);
 		}
 	);
 
 	PreForwardHandle = InControlRig->OnPreForwardsSolve_AnyThread().AddLambda(
-		[this](UControlRig* InControlRig, const EControlRigState InState, const FName& InEventName)
+		[this](UControlRig* InControlRig, const FName& InEventName)
 		{
 			// loop the animation data
 			if(InControlRig->GetAbsoluteTime() < GetTimeRange().X - SMALL_NUMBER ||
@@ -325,7 +325,7 @@ bool UControlRigTestData::SetupReplay(UControlRig* InControlRig, bool bGroundTru
 	);
 
 	PostForwardHandle = InControlRig->OnPostForwardsSolve_AnyThread().AddLambda(
-		[this](UControlRig* InControlRig, const EControlRigState InState, const FName& InEventName)
+		[this](UControlRig* InControlRig, const FName& InEventName)
 		{
 			const FRigPose CurrentPose = InControlRig->GetHierarchy()->GetPose();
 
diff --git a/Engine/Plugins/Animation/ControlRig/Source/ControlRig/Private/ControlRigValidationPass.cpp b/Engine/Plugins/Animation/ControlRig/Source/ControlRig/Private/ControlRigValidationPass.cpp
--- a/Engine/Plugins/Animation/ControlRig/Source/ControlRig/Private/ControlRigValidationPass.cpp
+++ b/Engine/Plugins/Animation/ControlRig/Source/ControlRig/Private/ControlRigValidationPass.cpp
@@ -124,14 +124,14 @@ void UControlRigValidator::SetControlRig(UControlRig* InControlRig)
 
 	if (UControlRig* ControlRig = WeakControlRig.Get())
 	{
-		OnControlRigInitialized(ControlRig, EControlRigState::Init, FRigUnit_BeginExecution::EventName);
+		OnControlRigInitialized(ControlRig, FRigUnit_BeginExecution::EventName);
 		ControlRig->OnInitialized_AnyThread().AddUObject(this, &UControlRigValidator::OnControlRigInitialized);
 		ControlRig->OnExecuted_AnyThread().AddUObject(this, &UControlRigValidator::OnControlRigExecuted);
 		ValidationContext.DrawInterface = &ControlRig->DrawInterface;
 	}
 }
 
-void UControlRigValidator::OnControlRigInitialized(UControlRig* Subject, EControlRigState State, const FName& EventName)
+void UControlRigValidator::OnControlRigInitialized(UControlRig* Subject, const FName& EventName)
 {
 	if (UControlRig* ControlRig = WeakControlRig.Get())
 	{
@@ -149,7 +149,7 @@ void UControlRigValidator::OnControlRigInitialized(UControlRig* Subject, EContro
 	}
 }
 
-void UControlRigValidator::OnControlRigExecuted(UControlRig* Subject, EControlRigState State, const FName& EventName)
+void UControlRigValidator::OnControlRigExecuted(UControlRig* Subject, const FName& EventName)
 {
 	if (UControlRig* ControlRig = WeakControlRig.Get())
 	{
diff --git a/Engine/Plugins/Animation/ControlRig/Source/ControlRig/Private/Rigs/AdditiveControlRig.cpp b/Engine/Plugins/Animation/ControlRig/Source/ControlRig/Private/Rigs/AdditiveControlRig.cpp
--- a/Engine/Plugins/Animation/ControlRig/Source/ControlRig/Private/Rigs/AdditiveControlRig.cpp
+++ b/Engine/Plugins/Animation/ControlRig/Source/ControlRig/Private/Rigs/AdditiveControlRig.cpp
@@ -93,7 +93,6 @@ void UAdditiveControlRig::Initialize(bool bInitRigUnits /*= true*/)
 		return;
 	}
 
-	Execute(EControlRigState::Init, FRigUnit_BeginExecution::EventName);
 	FRigUnitContext DefaultContext;
 	ExecuteUnits(DefaultContext, FRigUnit_PrepareForExecution::EventName);
 }
diff --git a/Engine/Plugins/Animation/ControlRig/Source/ControlRig/Private/Rigs/FKControlRig.cpp b/Engine/Plugins/Animation/ControlRig/Source/ControlRig/Private/Rigs/FKControlRig.cpp
--- a/Engine/Plugins/Animation/ControlRig/Source/ControlRig/Private/Rigs/FKControlRig.cpp
+++ b/Engine/Plugins/Animation/ControlRig/Source/ControlRig/Private/Rigs/FKControlRig.cpp
@@ -106,13 +106,6 @@ FName UFKControlRig::GetControlTargetName(const FName& InName, const ERigElement
 
 bool UFKControlRig::ExecuteUnits(FRigUnitContext& InOutContext, const FName& InEventName)
 {
-	if (InOutContext.State != EControlRigState::Update)
-	{
-		// we don't need any initialization so
-		// it's fine to consider this done here and return true.
-		return true;
-	}
-
 #if WITH_EDITOR
 	if(URigHierarchy* Hierarchy = GetHierarchy())
 	{
@@ -307,9 +300,6 @@ void UFKControlRig::Initialize(bool bInitRigUnits /*= true*/)
 		const FReferenceSkeleton& RefSkeleton = Skeleton->GetReferenceSkeleton();
 		CreateRigElements(RefSkeleton, (Skeleton) ? Skeleton->GetSmartNameContainer(USkeleton::AnimCurveMappingName) : nullptr);
 	}
-
-	// execute init
-	Execute(EControlRigState::Init, FRigUnit_BeginExecution::EventName);
 }
 
 TArray<FName> UFKControlRig::GetControlNames()
diff --git a/Engine/Plugins/Animation/ControlRig/Source/ControlRig/Private/Sequencer/MovieSceneControlRigParameterSection.cpp b/Engine/Plugins/Animation/ControlRig/Source/ControlRig/Private/Sequencer/MovieSceneControlRigParameterSection.cpp
--- a/Engine/Plugins/Animation/ControlRig/Source/ControlRig/Private/Sequencer/MovieSceneControlRigParameterSection.cpp
+++ b/Engine/Plugins/Animation/ControlRig/Source/ControlRig/Private/Sequencer/MovieSceneControlRigParameterSection.cpp
@@ -2919,7 +2919,7 @@ bool UMovieSceneControlRigParameterSection::LoadAnimSequenceIntoThisSection(UAni
 	}
 	SourceBones.ResetTransforms();
 	SourceCurves.ResetValues();
-	ControlRig->Execute(EControlRigState::Update, TEXT("Setup"));
+	ControlRig->Execute(TEXT("Setup"));
 	*/
 	const IAnimationDataModel* DataModel = AnimSequence->GetDataModel();
 	const FAnimationCurveData& CurveData = DataModel->GetCurveData();
@@ -2984,9 +2984,9 @@ bool UMovieSceneControlRigParameterSection::LoadAnimSequenceIntoThisSection(UAni
 		if (Index == 0)
 		{
 			//to make sure the first frame looks good we need to do this first. UE-100069
-			ControlRig->Execute(EControlRigState::Update, FRigUnit_InverseExecution::EventName);
+			ControlRig->Execute(FRigUnit_InverseExecution::EventName);
 		}
-		ControlRig->Execute(EControlRigState::Update, FRigUnit_InverseExecution::EventName);
+		ControlRig->Execute(FRigUnit_InverseExecution::EventName);
 
 		const ERichCurveInterpMode InterpMode = bKeyReduce ? RCIM_Cubic : RCIM_Linear;
 		RecordControlRigKey(FrameNumber, Index == 0, InterpMode);
diff --git a/Engine/Plugins/Animation/ControlRig/Source/ControlRig/Private/Units/Animation/RigUnit_GetDeltaTime.cpp b/Engine/Plugins/Animation/ControlRig/Source/ControlRig/Private/Units/Animation/RigUnit_GetDeltaTime.cpp
--- a/Engine/Plugins/Animation/ControlRig/Source/ControlRig/Private/Units/Animation/RigUnit_GetDeltaTime.cpp
+++ b/Engine/Plugins/Animation/ControlRig/Source/ControlRig/Private/Units/Animation/RigUnit_GetDeltaTime.cpp
@@ -8,7 +8,7 @@
 FRigUnit_GetDeltaTime_Execute()
 {
     DECLARE_SCOPE_HIERARCHICAL_COUNTER_RIGUNIT()
-	Result = ExecuteContext.UnitContext.DeltaTime;
+	Result = ExecuteContext.GetDeltaTime();
 }
 
 #if WITH_DEV_AUTOMATION_TESTS
@@ -16,7 +16,7 @@ FRigUnit_GetDeltaTime_Execute()
 
 IMPLEMENT_RIGUNIT_AUTOMATION_TEST(FRigUnit_GetDeltaTime)
 {
-	ExecuteContext.UnitContext.DeltaTime = 0.2f;
+	ExecuteContext.SetDeltaTime(0.2f);
 	Execute();
 	AddErrorIfFalse(FMath::IsNearlyEqual(Unit.Result, 0.2f), TEXT("unexpected delta time"));
 	return true;
diff --git a/Engine/Plugins/Animation/ControlRig/Source/ControlRig/Private/Units/Animation/RigUnit_TimeConversion.cpp b/Engine/Plugins/Animation/ControlRig/Source/ControlRig/Private/Units/Animation/RigUnit_TimeConversion.cpp
--- a/Engine/Plugins/Animation/ControlRig/Source/ControlRig/Private/Units/Animation/RigUnit_TimeConversion.cpp
+++ b/Engine/Plugins/Animation/ControlRig/Source/ControlRig/Private/Units/Animation/RigUnit_TimeConversion.cpp
@@ -9,9 +9,9 @@ FRigUnit_FramesToSeconds_Execute()
 {
     DECLARE_SCOPE_HIERARCHICAL_COUNTER_RIGUNIT()
 
-    if(ExecuteContext.UnitContext.FramesPerSecond > SMALL_NUMBER)
+    if(ExecuteContext.GetFramesPerSecond() > SMALL_NUMBER)
     {
-		Seconds = Frames / ExecuteContext.UnitContext.FramesPerSecond;
+		Seconds = Frames / ExecuteContext.GetFramesPerSecond();
 	}
 	else
 	{
@@ -23,9 +23,9 @@ FRigUnit_SecondsToFrames_Execute()
 {
     DECLARE_SCOPE_HIERARCHICAL_COUNTER_RIGUNIT()
 
-    if(ExecuteContext.UnitContext.FramesPerSecond > SMALL_NUMBER)
+    if(ExecuteContext.GetFramesPerSecond() > SMALL_NUMBER)
     {
-		Frames = Seconds * ExecuteContext.UnitContext.FramesPerSecond;
+		Frames = Seconds * ExecuteContext.GetFramesPerSecond();
 	}
 	else
 	{
diff --git a/Engine/Plugins/Animation/ControlRig/Source/ControlRig/Private/Units/Control/RigUnit_Control.cpp b/Engine/Plugins/Animation/ControlRig/Source/ControlRig/Private/Units/Control/RigUnit_Control.cpp
--- a/Engine/Plugins/Animation/ControlRig/Source/ControlRig/Private/Units/Control/RigUnit_Control.cpp
+++ b/Engine/Plugins/Animation/ControlRig/Source/ControlRig/Private/Units/Control/RigUnit_Control.cpp
@@ -8,9 +8,11 @@
 FRigUnit_Control_Execute()
 {
     DECLARE_SCOPE_HIERARCHICAL_COUNTER_RIGUNIT()
-	if (ExecuteContext.UnitContext.State == EControlRigState::Init)
+
+	if (!bIsInitialized)
 	{
 		Transform.FromFTransform(InitTransform);
+		bIsInitialized = true;
 	}
 
 	Result = StaticGetResultantTransform(Transform, Base, Filter);
diff --git a/Engine/Plugins/Animation/ControlRig/Source/ControlRig/Private/Units/Control/RigUnit_Control_StaticMesh.cpp b/Engine/Plugins/Animation/ControlRig/Source/ControlRig/Private/Units/Control/RigUnit_Control_StaticMesh.cpp
--- a/Engine/Plugins/Animation/ControlRig/Source/ControlRig/Private/Units/Control/RigUnit_Control_StaticMesh.cpp
+++ b/Engine/Plugins/Animation/ControlRig/Source/ControlRig/Private/Units/Control/RigUnit_Control_StaticMesh.cpp
@@ -10,7 +10,7 @@ FRigUnit_Control_StaticMesh::FRigUnit_Control_StaticMesh()
 
 FRigUnit_Control_StaticMesh_Execute()
 {
-	FRigUnit_Control::StaticExecute(ExecuteContext, Transform, Base, InitTransform, Result, Filter);
+	FRigUnit_Control::StaticExecute(ExecuteContext, Transform, Base, InitTransform, Result, Filter, bIsInitialized);
 }
 
 FRigVMStructUpgradeInfo FRigUnit_Control_StaticMesh::GetUpgradeInfo() const
diff --git a/Engine/Plugins/Animation/ControlRig/Source/ControlRig/Private/Units/Debug/RigUnit_DebugBezier.cpp b/Engine/Plugins/Animation/ControlRig/Source/ControlRig/Private/Units/Debug/RigUnit_DebugBezier.cpp
--- a/Engine/Plugins/Animation/ControlRig/Source/ControlRig/Private/Units/Debug/RigUnit_DebugBezier.cpp
+++ b/Engine/Plugins/Animation/ControlRig/Source/ControlRig/Private/Units/Debug/RigUnit_DebugBezier.cpp
@@ -29,10 +29,6 @@ FRigVMStructUpgradeInfo FRigUnit_DebugBezier::GetUpgradeInfo() const
 FRigUnit_DebugBezierItemSpace_Execute()
 {
     DECLARE_SCOPE_HIERARCHICAL_COUNTER_RIGUNIT()
-	if (ExecuteContext.UnitContext.State == EControlRigState::Init)
-	{
-		return;
-	}
 
 	if (ExecuteContext.UnitContext.DrawInterface == nullptr || !bEnabled)
 	{
diff --git a/Engine/Plugins/Animation/ControlRig/Source/ControlRig/Private/Units/Debug/RigUnit_DebugHierarchy.cpp b/Engine/Plugins/Animation/ControlRig/Source/ControlRig/Private/Units/Debug/RigUnit_DebugHierarchy.cpp
--- a/Engine/Plugins/Animation/ControlRig/Source/ControlRig/Private/Units/Debug/RigUnit_DebugHierarchy.cpp
+++ b/Engine/Plugins/Animation/ControlRig/Source/ControlRig/Private/Units/Debug/RigUnit_DebugHierarchy.cpp
@@ -8,10 +8,6 @@
 FRigUnit_DebugHierarchy_Execute()
 {
     DECLARE_SCOPE_HIERARCHICAL_COUNTER_RIGUNIT()
-	if (ExecuteContext.UnitContext.State == EControlRigState::Init)
-	{
-		return;
-	}
 
 	if (ExecuteContext.UnitContext.DrawInterface == nullptr || !bEnabled)
 	{
@@ -28,10 +24,6 @@ FRigUnit_DebugHierarchy_Execute()
 FRigUnit_DebugPose_Execute()
 {
 	DECLARE_SCOPE_HIERARCHICAL_COUNTER_RIGUNIT()
-	if (ExecuteContext.UnitContext.State == EControlRigState::Init)
-	{
-		return;
-	}
 
 	if (ExecuteContext.UnitContext.DrawInterface == nullptr || !bEnabled)
 	{
diff --git a/Engine/Plugins/Animation/ControlRig/Source/ControlRig/Private/Units/Debug/RigUnit_DebugLine.cpp b/Engine/Plugins/Animation/ControlRig/Source/ControlRig/Private/Units/Debug/RigUnit_DebugLine.cpp
--- a/Engine/Plugins/Animation/ControlRig/Source/ControlRig/Private/Units/Debug/RigUnit_DebugLine.cpp
+++ b/Engine/Plugins/Animation/ControlRig/Source/ControlRig/Private/Units/Debug/RigUnit_DebugLine.cpp
@@ -37,10 +37,6 @@ FRigVMStructUpgradeInfo FRigUnit_DebugLine::GetUpgradeInfo() const
 FRigUnit_DebugLineItemSpace_Execute()
 {
     DECLARE_SCOPE_HIERARCHICAL_COUNTER_RIGUNIT()
-	if (ExecuteContext.UnitContext.State == EControlRigState::Init)
-	{
-		return;
-	}
 
 	if (ExecuteContext.UnitContext.DrawInterface == nullptr || !bEnabled)
 	{
diff --git a/Engine/Plugins/Animation/ControlRig/Source/ControlRig/Private/Units/Debug/RigUnit_DebugPoint.cpp b/Engine/Plugins/Animation/ControlRig/Source/ControlRig/Private/Units/Debug/RigUnit_DebugPoint.cpp
--- a/Engine/Plugins/Animation/ControlRig/Source/ControlRig/Private/Units/Debug/RigUnit_DebugPoint.cpp
+++ b/Engine/Plugins/Animation/ControlRig/Source/ControlRig/Private/Units/Debug/RigUnit_DebugPoint.cpp
@@ -9,10 +9,6 @@
 FRigUnit_DebugPoint_Execute()
 {
     DECLARE_SCOPE_HIERARCHICAL_COUNTER_RIGUNIT()
-	if (ExecuteContext.UnitContext.State == EControlRigState::Init)
-	{
-		return;
-	}
 
 	if (ExecuteContext.UnitContext.DrawInterface == nullptr || !bEnabled)
 	{
@@ -62,10 +58,6 @@ FRigVMStructUpgradeInfo FRigUnit_DebugPoint::GetUpgradeInfo() const
 FRigUnit_DebugPointMutable_Execute()
 {
     DECLARE_SCOPE_HIERARCHICAL_COUNTER_RIGUNIT()
-	if (ExecuteContext.UnitContext.State == EControlRigState::Init)
-	{
-		return;
-	}
 
 	if (ExecuteContext.UnitContext.DrawInterface == nullptr || !bEnabled)
 	{
diff --git a/Engine/Plugins/Animation/ControlRig/Source/ControlRig/Private/Units/Debug/RigUnit_DebugTransform.cpp b/Engine/Plugins/Animation/ControlRig/Source/ControlRig/Private/Units/Debug/RigUnit_DebugTransform.cpp
--- a/Engine/Plugins/Animation/ControlRig/Source/ControlRig/Private/Units/Debug/RigUnit_DebugTransform.cpp
+++ b/Engine/Plugins/Animation/ControlRig/Source/ControlRig/Private/Units/Debug/RigUnit_DebugTransform.cpp
@@ -9,10 +9,6 @@
 FRigUnit_DebugTransform_Execute()
 {
     DECLARE_SCOPE_HIERARCHICAL_COUNTER_RIGUNIT()
-	if (ExecuteContext.UnitContext.State == EControlRigState::Init)
-	{
-		return;
-	}
 
 	if (ExecuteContext.UnitContext.DrawInterface == nullptr || !bEnabled)
 	{
@@ -96,10 +92,6 @@ FRigVMStructUpgradeInfo FRigUnit_DebugTransformMutable::GetUpgradeInfo() const
 FRigUnit_DebugTransformMutableItemSpace_Execute()
 {
     DECLARE_SCOPE_HIERARCHICAL_COUNTER_RIGUNIT()
-	if (ExecuteContext.UnitContext.State == EControlRigState::Init)
-	{
-		return;
-	}
 
 	if (ExecuteContext.UnitContext.DrawInterface == nullptr || !bEnabled)
 	{
@@ -136,10 +128,6 @@ FRigUnit_DebugTransformMutableItemSpace_Execute()
 FRigUnit_DebugTransformArrayMutable_Execute()
 {
     DECLARE_SCOPE_HIERARCHICAL_COUNTER_RIGUNIT()
-	if (ExecuteContext.UnitContext.State == EControlRigState::Init)
-	{
-		return;
-	}
 
 	if (ExecuteContext.UnitContext.DrawInterface == nullptr || !bEnabled)
 	{
@@ -207,10 +195,6 @@ FRigVMStructUpgradeInfo FRigUnit_DebugTransformArrayMutable::GetUpgradeInfo() co
 FRigUnit_DebugTransformArrayMutableItemSpace_Execute()
 {
 	DECLARE_SCOPE_HIERARCHICAL_COUNTER_RIGUNIT()
-	if (ExecuteContext.UnitContext.State == EControlRigState::Init)
-	{
-		return;
-	}
 
 	if (ExecuteContext.UnitContext.DrawInterface == nullptr || !bEnabled || Transforms.Num() == 0)
 	{
diff --git a/Engine/Plugins/Animation/ControlRig/Source/ControlRig/Private/Units/Debug/RigUnit_ProfilingBracket.cpp b/Engine/Plugins/Animation/ControlRig/Source/ControlRig/Private/Units/Debug/RigUnit_ProfilingBracket.cpp
--- a/Engine/Plugins/Animation/ControlRig/Source/ControlRig/Private/Units/Debug/RigUnit_ProfilingBracket.cpp
+++ b/Engine/Plugins/Animation/ControlRig/Source/ControlRig/Private/Units/Debug/RigUnit_ProfilingBracket.cpp
@@ -9,40 +9,34 @@
 FRigUnit_StartProfilingTimer_Execute()
 {
     DECLARE_SCOPE_HIERARCHICAL_COUNTER_RIGUNIT()
-	if (ExecuteContext.UnitContext.State == EControlRigState::Update)
-	{
-		UKismetAnimationLibrary::K2_StartProfilingTimer();
-	}
+	UKismetAnimationLibrary::K2_StartProfilingTimer();
 }
 
 FRigUnit_EndProfilingTimer_Execute()
 {
     DECLARE_SCOPE_HIERARCHICAL_COUNTER_RIGUNIT()
-	if (ExecuteContext.UnitContext.State == EControlRigState::Init)
+	if (!bIsInitialized)
 	{
 		AccumulatedTime = 0.f;
 		MeasurementsLeft = FMath::Max(NumberOfMeasurements, 1);
-		return;
+		bIsInitialized = true;
 	}
 
-	if (ExecuteContext.UnitContext.State == EControlRigState::Update)
+	float Delta = UKismetAnimationLibrary::K2_EndProfilingTimer(false);
+	if(MeasurementsLeft > 0)
 	{
-		float Delta = UKismetAnimationLibrary::K2_EndProfilingTimer(false);
-		if(MeasurementsLeft > 0)
-		{
-			AccumulatedTime += Delta / float(FMath::Max(NumberOfMeasurements, 1));
-			MeasurementsLeft--;
+		AccumulatedTime += Delta / float(FMath::Max(NumberOfMeasurements, 1));
+		MeasurementsLeft--;
 
-			if (MeasurementsLeft == 0)
+		if (MeasurementsLeft == 0)
+		{
+			if (Prefix.IsEmpty())
+			{
+				UE_CONTROLRIG_RIGUNIT_LOG_MESSAGE(TEXT("%0.3f ms (%d runs)."), AccumulatedTime, FMath::Max(NumberOfMeasurements, 1));
+			}
+			else
 			{
-				if (Prefix.IsEmpty())
-				{
-					UE_CONTROLRIG_RIGUNIT_LOG_MESSAGE(TEXT("%0.3f ms (%d runs)."), AccumulatedTime, FMath::Max(NumberOfMeasurements, 1));
-				}
-				else
-				{
-					UE_CONTROLRIG_RIGUNIT_LOG_MESSAGE(TEXT("[%s] %0.3f ms (%d runs)."), *Prefix, AccumulatedTime, FMath::Max(NumberOfMeasurements, 1));
-				}
+				UE_CONTROLRIG_RIGUNIT_LOG_MESSAGE(TEXT("[%s] %0.3f ms (%d runs)."), *Prefix, AccumulatedTime, FMath::Max(NumberOfMeasurements, 1));
 			}
 		}
 	}
diff --git a/Engine/Plugins/Animation/ControlRig/Source/ControlRig/Private/Units/Debug/RigUnit_ProfilingBracket.h b/Engine/Plugins/Animation/ControlRig/Source/ControlRig/Private/Units/Debug/RigUnit_ProfilingBracket.h
--- a/Engine/Plugins/Animation/ControlRig/Source/ControlRig/Private/Units/Debug/RigUnit_ProfilingBracket.h
+++ b/Engine/Plugins/Animation/ControlRig/Source/ControlRig/Private/Units/Debug/RigUnit_ProfilingBracket.h
@@ -35,6 +35,7 @@ struct CONTROLRIG_API FRigUnit_EndProfilingTimer : public FRigUnit_DebugBaseMuta
 		AccumulatedTime = 0.f;
 		MeasurementsLeft = 0;
 		Prefix = TEXT("Timer");
+		bIsInitialized = false;
 	}
 
 	RIGVM_METHOD()
@@ -51,4 +52,7 @@ struct CONTROLRIG_API FRigUnit_EndProfilingTimer : public FRigUnit_DebugBaseMuta
 
 	UPROPERTY()
 	int32 MeasurementsLeft;
+
+	UPROPERTY()
+	bool bIsInitialized;
 };
diff --git a/Engine/Plugins/Animation/ControlRig/Source/ControlRig/Private/Units/Deprecated/RigUnit_AimConstraint.cpp b/Engine/Plugins/Animation/ControlRig/Source/ControlRig/Private/Units/Deprecated/RigUnit_AimConstraint.cpp
--- a/Engine/Plugins/Animation/ControlRig/Source/ControlRig/Private/Units/Deprecated/RigUnit_AimConstraint.cpp
+++ b/Engine/Plugins/Animation/ControlRig/Source/ControlRig/Private/Units/Deprecated/RigUnit_AimConstraint.cpp
@@ -13,7 +13,7 @@ FRigUnit_AimConstraint_Execute()
 
 	TArray<FConstraintData>& ConstraintData = WorkData.ConstraintData;
 
-	if (ExecuteContext.UnitContext.State == EControlRigState::Init)
+	if (ConstraintData.Num() != AimTargets.Num())
 	{
 		ConstraintData.Reset();
 
@@ -43,48 +43,46 @@ FRigUnit_AimConstraint_Execute()
 			}
 		}
 	}
-	else if (ExecuteContext.UnitContext.State == EControlRigState::Update)
+
+	URigHierarchy* Hierarchy = ExecuteContext.Hierarchy;
+	if (Hierarchy)
 	{
-		URigHierarchy* Hierarchy = ExecuteContext.Hierarchy;
-		if (Hierarchy)
+		const FRigElementKey Key(Joint, ERigElementType::Bone); 
+		int32 BoneIndex = Hierarchy->GetIndex(Key);
+		if (BoneIndex != INDEX_NONE)
 		{
-			const FRigElementKey Key(Joint, ERigElementType::Bone); 
-			int32 BoneIndex = Hierarchy->GetIndex(Key);
-			if (BoneIndex != INDEX_NONE)
+			const int32 TargetNum = AimTargets.Num();
+			if (TargetNum > 0)
 			{
-				const int32 TargetNum = AimTargets.Num();
-				if (TargetNum > 0)
+				for (int32 ConstraintIndex= 0; ConstraintIndex< ConstraintData.Num(); ++ConstraintIndex)
 				{
-					for (int32 ConstraintIndex= 0; ConstraintIndex< ConstraintData.Num(); ++ConstraintIndex)
-					{
-						FAimConstraintDescription* AimConstraintDesc = ConstraintData[ConstraintIndex].Constraint.GetTypedConstraint<FAimConstraintDescription>();
-						AimConstraintDesc->LookAt_Axis = FAxis(AimVector);
-
-						if (UpTargets.IsValidIndex(ConstraintIndex))
-						{
-							AimConstraintDesc->LookUp_Axis = FAxis(UpVector);
-							AimConstraintDesc->bUseLookUp = UpVector.Size() > 0.f;
-							AimConstraintDesc->LookUpTarget = UpTargets[ConstraintIndex].Transform.GetLocation();
-						}
-
-						ConstraintData[ConstraintIndex].CurrentTransform = AimTargets[ConstraintIndex].Transform;
-						ConstraintData[ConstraintIndex].Weight = AimTargets[ConstraintIndex].Weight;
-					}
+					FAimConstraintDescription* AimConstraintDesc = ConstraintData[ConstraintIndex].Constraint.GetTypedConstraint<FAimConstraintDescription>();
+					AimConstraintDesc->LookAt_Axis = FAxis(AimVector);
 
-					FTransform BaseTransform = FTransform::Identity;
-					int32 ParentIndex = Hierarchy->GetIndex(Hierarchy->GetFirstParent(Key));
-					if (ParentIndex != INDEX_NONE)
+					if (UpTargets.IsValidIndex(ConstraintIndex))
 					{
-						BaseTransform = Hierarchy->GetGlobalTransform(ParentIndex);
+						AimConstraintDesc->LookUp_Axis = FAxis(UpVector);
+						AimConstraintDesc->bUseLookUp = UpVector.Size() > 0.f;
+						AimConstraintDesc->LookUpTarget = UpTargets[ConstraintIndex].Transform.GetLocation();
 					}
 
-					FTransform SourceTransform = Hierarchy->GetGlobalTransform(BoneIndex);
-
-					// @todo: ignore maintain offset for now
-					FTransform ConstrainedTransform = AnimationCore::SolveConstraints(SourceTransform, BaseTransform, ConstraintData);
+					ConstraintData[ConstraintIndex].CurrentTransform = AimTargets[ConstraintIndex].Transform;
+					ConstraintData[ConstraintIndex].Weight = AimTargets[ConstraintIndex].Weight;
+				}
 
-					Hierarchy->SetGlobalTransform(BoneIndex, ConstrainedTransform);
+				FTransform BaseTransform = FTransform::Identity;
+				int32 ParentIndex = Hierarchy->GetIndex(Hierarchy->GetFirstParent(Key));
+				if (ParentIndex != INDEX_NONE)
+				{
+					BaseTransform = Hierarchy->GetGlobalTransform(ParentIndex);
 				}
+
+				FTransform SourceTransform = Hierarchy->GetGlobalTransform(BoneIndex);
+
+				// @todo: ignore maintain offset for now
+				FTransform ConstrainedTransform = AnimationCore::SolveConstraints(SourceTransform, BaseTransform, ConstraintData);
+
+				Hierarchy->SetGlobalTransform(BoneIndex, ConstrainedTransform);
 			}
 		}
 	}
diff --git a/Engine/Plugins/Animation/ControlRig/Source/ControlRig/Private/Units/Deprecated/RigUnit_ApplyFK.cpp b/Engine/Plugins/Animation/ControlRig/Source/ControlRig/Private/Units/Deprecated/RigUnit_ApplyFK.cpp
--- a/Engine/Plugins/Animation/ControlRig/Source/ControlRig/Private/Units/Deprecated/RigUnit_ApplyFK.cpp
+++ b/Engine/Plugins/Animation/ControlRig/Source/ControlRig/Private/Units/Deprecated/RigUnit_ApplyFK.cpp
@@ -10,47 +10,40 @@ FRigUnit_ApplyFK_Execute()
 {
     DECLARE_SCOPE_HIERARCHICAL_COUNTER_RIGUNIT()
 	
-	if (ExecuteContext.UnitContext.State == EControlRigState::Init)
+	URigHierarchy* Hierarchy = ExecuteContext.Hierarchy;
+	if (Hierarchy)
 	{
-		return;
-	}
-	else if (ExecuteContext.UnitContext.State == EControlRigState::Update)
-	{
-		URigHierarchy* Hierarchy = ExecuteContext.Hierarchy;
-		if (Hierarchy)
+		const FRigElementKey Key(Joint, ERigElementType::Bone);
+		int32 Index = Hierarchy->GetIndex(Key);
+		if (Index != INDEX_NONE)
 		{
-			const FRigElementKey Key(Joint, ERigElementType::Bone);
-			int32 Index = Hierarchy->GetIndex(Key);
-			if (Index != INDEX_NONE)
-			{
-				// first filter input transform
-				FTransform InputTransform = Transform;
-				Filter.FilterTransform(InputTransform);
+			// first filter input transform
+			FTransform InputTransform = Transform;
+			Filter.FilterTransform(InputTransform);
 
-				FTransform InputBaseTransform = UtilityHelpers::GetBaseTransformByMode(
-					ApplyTransformSpace,
-					[Hierarchy](const FRigElementKey& BoneKey) { return Hierarchy->GetGlobalTransform(BoneKey); },
-					Hierarchy->GetFirstParent(Key),
-					FRigElementKey(BaseJoint, ERigElementType::Bone),
-					BaseTransform
-				);
+			FTransform InputBaseTransform = UtilityHelpers::GetBaseTransformByMode(
+				ApplyTransformSpace,
+				[Hierarchy](const FRigElementKey& BoneKey) { return Hierarchy->GetGlobalTransform(BoneKey); },
+				Hierarchy->GetFirstParent(Key),
+				FRigElementKey(BaseJoint, ERigElementType::Bone),
+				BaseTransform
+			);
 
-				// now get override or additive
-				// whether I'd like to apply whole thing or not
-				if (ApplyTransformMode == EApplyTransformMode::Override)
-				{
-					// get base transform
-					FTransform ApplyTransform = InputTransform * InputBaseTransform;
-					Hierarchy->SetGlobalTransform(Index, ApplyTransform);
-				}
-				else
-				{
-					// if additive, we get current transform and calculate base transform and apply in their local space
-					FTransform CurrentTransform = Hierarchy->GetGlobalTransform(Index);
-					FTransform LocalTransform = InputTransform * CurrentTransform.GetRelativeTransform(InputBaseTransform);
-					// apply additive
-					Hierarchy->SetGlobalTransform(Index, LocalTransform * InputBaseTransform);
-				}
+			// now get override or additive
+			// whether I'd like to apply whole thing or not
+			if (ApplyTransformMode == EApplyTransformMode::Override)
+			{
+				// get base transform
+				FTransform ApplyTransform = InputTransform * InputBaseTransform;
+				Hierarchy->SetGlobalTransform(Index, ApplyTransform);
+			}
+			else
+			{
+				// if additive, we get current transform and calculate base transform and apply in their local space
+				FTransform CurrentTransform = Hierarchy->GetGlobalTransform(Index);
+				FTransform LocalTransform = InputTransform * CurrentTransform.GetRelativeTransform(InputBaseTransform);
+				// apply additive
+				Hierarchy->SetGlobalTransform(Index, LocalTransform * InputBaseTransform);
 			}
 		}
 	}
diff --git a/Engine/Plugins/Animation/ControlRig/Source/ControlRig/Private/Units/Deprecated/RigUnit_TwoBoneIKFK.cpp b/Engine/Plugins/Animation/ControlRig/Source/ControlRig/Private/Units/Deprecated/RigUnit_TwoBoneIKFK.cpp
--- a/Engine/Plugins/Animation/ControlRig/Source/ControlRig/Private/Units/Deprecated/RigUnit_TwoBoneIKFK.cpp
+++ b/Engine/Plugins/Animation/ControlRig/Source/ControlRig/Private/Units/Deprecated/RigUnit_TwoBoneIKFK.cpp
@@ -11,7 +11,7 @@ FRigUnit_TwoBoneIKFK_Execute()
 {
     DECLARE_SCOPE_HIERARCHICAL_COUNTER_RIGUNIT()
 
-	if (ExecuteContext.UnitContext.State == EControlRigState::Init)
+	if (!bIsInitialized)
 	{
 		URigHierarchy* Hierarchy = ExecuteContext.Hierarchy;
 		if (Hierarchy)
@@ -56,66 +56,65 @@ FRigUnit_TwoBoneIKFK_Execute()
 				}
 			}
 		}
+		bIsInitialized = true;
 	}
-	else  if (ExecuteContext.UnitContext.State == EControlRigState::Update)
+
+	if (StartJointIndex != INDEX_NONE && MidJointIndex != INDEX_NONE && EndJointIndex != INDEX_NONE)
 	{
-		if (StartJointIndex != INDEX_NONE && MidJointIndex != INDEX_NONE && EndJointIndex != INDEX_NONE)
-		{
-			FTransform StartJointTransform;
-			FTransform MidJointTransform;
-			FTransform EndJointTransform;
+		FTransform StartJointTransform;
+		FTransform MidJointTransform;
+		FTransform EndJointTransform;
 
-			// FK only
-			if (FMath::IsNearlyZero(IKBlend))
-			{
-				StartJointTransform = StartJointFKTransform;
-				MidJointTransform = MidJointFKTransform;
-				EndJointTransform = EndJointFKTransform;
-			}
-			// IK only
-			else if (FMath::IsNearlyEqual(IKBlend, 1.f))
-			{
-				// update transform before going through IK
-				const URigHierarchy* Hierarchy = ExecuteContext.Hierarchy;
-				check(Hierarchy);
+		// FK only
+		if (FMath::IsNearlyZero(IKBlend))
+		{
+			StartJointTransform = StartJointFKTransform;
+			MidJointTransform = MidJointFKTransform;
+			EndJointTransform = EndJointFKTransform;
+		}
+		// IK only
+		else if (FMath::IsNearlyEqual(IKBlend, 1.f))
+		{
+			// update transform before going through IK
+			const URigHierarchy* Hierarchy = ExecuteContext.Hierarchy;
+			check(Hierarchy);
 
-				StartJointIKTransform = Hierarchy->GetGlobalTransform(StartJointIndex);
-				MidJointIKTransform = Hierarchy->GetGlobalTransform(MidJointIndex);
-				EndJointIKTransform = Hierarchy->GetGlobalTransform(EndJointIndex);
+			StartJointIKTransform = Hierarchy->GetGlobalTransform(StartJointIndex);
+			MidJointIKTransform = Hierarchy->GetGlobalTransform(MidJointIndex);
+			EndJointIKTransform = Hierarchy->GetGlobalTransform(EndJointIndex);
 
-				AnimationCore::SolveTwoBoneIK(StartJointIKTransform, MidJointIKTransform, EndJointIKTransform, PoleTarget, EndEffector.GetLocation(), UpperLimbLength, LowerLimbLength, false, 1.0f, 1.05f);
-				EndJointIKTransform.SetRotation(EndEffector.GetRotation());
+			AnimationCore::SolveTwoBoneIK(StartJointIKTransform, MidJointIKTransform, EndJointIKTransform, PoleTarget, EndEffector.GetLocation(), UpperLimbLength, LowerLimbLength, false, 1.0f, 1.05f);
+			EndJointIKTransform.SetRotation(EndEffector.GetRotation());
 
-				StartJointTransform = StartJointIKTransform;
-				MidJointTransform = MidJointIKTransform;
-				EndJointTransform = EndJointIKTransform;
-			}
-			else
-			{
-				// update transform before going through IK
-				const URigHierarchy* Hierarchy = ExecuteContext.Hierarchy;
-				check(Hierarchy);
+			StartJointTransform = StartJointIKTransform;
+			MidJointTransform = MidJointIKTransform;
+			EndJointTransform = EndJointIKTransform;
+		}
+		else
+		{
+			// update transform before going through IK
+			const URigHierarchy* Hierarchy = ExecuteContext.Hierarchy;
+			check(Hierarchy);
 
-				StartJointIKTransform = Hierarchy->GetGlobalTransform(StartJointIndex);
-				MidJointIKTransform = Hierarchy->GetGlobalTransform(MidJointIndex);
-				EndJointIKTransform = Hierarchy->GetGlobalTransform(EndJointIndex);
+			StartJointIKTransform = Hierarchy->GetGlobalTransform(StartJointIndex);
+			MidJointIKTransform = Hierarchy->GetGlobalTransform(MidJointIndex);
+			EndJointIKTransform = Hierarchy->GetGlobalTransform(EndJointIndex);
 
-				AnimationCore::SolveTwoBoneIK(StartJointIKTransform, MidJointIKTransform, EndJointIKTransform, PoleTarget, EndEffector.GetLocation(), UpperLimbLength, LowerLimbLength, false, 1.0f, 1.05f);
-				EndJointIKTransform.SetRotation(EndEffector.GetRotation());
+			AnimationCore::SolveTwoBoneIK(StartJointIKTransform, MidJointIKTransform, EndJointIKTransform, PoleTarget, EndEffector.GetLocation(), UpperLimbLength, LowerLimbLength, false, 1.0f, 1.05f);
+			EndJointIKTransform.SetRotation(EndEffector.GetRotation());
 
-				StartJointTransform.Blend(StartJointFKTransform, StartJointIKTransform, IKBlend);
-				MidJointTransform.Blend(MidJointFKTransform, MidJointIKTransform, IKBlend);
-				EndJointTransform.Blend(MidJointFKTransform, EndJointIKTransform, IKBlend);
-			}
+			StartJointTransform.Blend(StartJointFKTransform, StartJointIKTransform, IKBlend);
+			MidJointTransform.Blend(MidJointFKTransform, MidJointIKTransform, IKBlend);
+			EndJointTransform.Blend(MidJointFKTransform, EndJointIKTransform, IKBlend);
+		}
 
-			URigHierarchy* Hierarchy = ExecuteContext.Hierarchy;
-			check(Hierarchy);
-			Hierarchy->SetGlobalTransform(StartJointIndex, StartJointTransform);
-			Hierarchy->SetGlobalTransform(MidJointIndex, MidJointTransform);
-			Hierarchy->SetGlobalTransform(EndJointIndex, EndJointTransform);
+		URigHierarchy* Hierarchy = ExecuteContext.Hierarchy;
+		check(Hierarchy);
+		Hierarchy->SetGlobalTransform(StartJointIndex, StartJointTransform);
+		Hierarchy->SetGlobalTransform(MidJointIndex, MidJointTransform);
+		Hierarchy->SetGlobalTransform(EndJointIndex, EndJointTransform);
 
-			PreviousFKIKBlend = IKBlend;
-		}
+		PreviousFKIKBlend = IKBlend;
 	}
 }
 
diff --git a/Engine/Plugins/Animation/ControlRig/Source/ControlRig/Private/Units/Drawing/RigUnit_DrawContainer.cpp b/Engine/Plugins/Animation/ControlRig/Source/ControlRig/Private/Units/Drawing/RigUnit_DrawContainer.cpp
--- a/Engine/Plugins/Animation/ControlRig/Source/ControlRig/Private/Units/Drawing/RigUnit_DrawContainer.cpp
+++ b/Engine/Plugins/Animation/ControlRig/Source/ControlRig/Private/Units/Drawing/RigUnit_DrawContainer.cpp
@@ -14,24 +14,17 @@ FRigUnit_DrawContainerGetInstruction_Execute()
     	return;
     }
 	
-	switch(ExecuteContext.UnitContext.State)
+	int32 Index = ExecuteContext.UnitContext.DrawContainer->GetIndex(InstructionName);
+	if (Index != INDEX_NONE)
 	{
-		case EControlRigState::Update:
-		{
-			int32 Index = ExecuteContext.UnitContext.DrawContainer->GetIndex(InstructionName);
-			if (Index != INDEX_NONE)
-			{
-				const FControlRigDrawInstruction& Instruction = (*ExecuteContext.UnitContext.DrawContainer)[Index];
-				Color = Instruction.Color;
-				Transform = Instruction.Transform;
-			}
-			else
-			{
-				Color = FLinearColor::Red;
-				Transform = FTransform::Identity;
-			}
-			break;
-		}
+		const FControlRigDrawInstruction& Instruction = (*ExecuteContext.UnitContext.DrawContainer)[Index];
+		Color = Instruction.Color;
+		Transform = Instruction.Transform;
+	}
+	else
+	{
+		Color = FLinearColor::Red;
+		Transform = FTransform::Identity;
 	}
 }
 
@@ -50,18 +43,11 @@ FRigUnit_DrawContainerSetThickness_Execute()
 		return;
 	}
 
-	switch (ExecuteContext.UnitContext.State)
+	int32 Index = ExecuteContext.UnitContext.DrawContainer->GetIndex(InstructionName);
+	if (Index != INDEX_NONE)
 	{
-		case EControlRigState::Update:
-		{
-			int32 Index = ExecuteContext.UnitContext.DrawContainer->GetIndex(InstructionName);
-			if (Index != INDEX_NONE)
-			{
-				FControlRigDrawInstruction& Instruction = (*ExecuteContext.UnitContext.DrawContainer)[Index];
-				Instruction.Thickness = Thickness;
-			}
-			break;
-		}
+		FControlRigDrawInstruction& Instruction = (*ExecuteContext.UnitContext.DrawContainer)[Index];
+		Instruction.Thickness = Thickness;
 	}
 }
 
@@ -80,18 +66,11 @@ FRigUnit_DrawContainerSetColor_Execute()
 		return;
 	}
 
-	switch (ExecuteContext.UnitContext.State)
+	int32 Index = ExecuteContext.UnitContext.DrawContainer->GetIndex(InstructionName);
+	if (Index != INDEX_NONE)
 	{
-		case EControlRigState::Update:
-		{
-			int32 Index = ExecuteContext.UnitContext.DrawContainer->GetIndex(InstructionName);
-			if (Index != INDEX_NONE)
-			{
-				FControlRigDrawInstruction& Instruction = (*ExecuteContext.UnitContext.DrawContainer)[Index];
-				Instruction.Color = Color;
-			}
-			break;
-		}
+		FControlRigDrawInstruction& Instruction = (*ExecuteContext.UnitContext.DrawContainer)[Index];
+		Instruction.Color = Color;
 	}
 }
 
@@ -110,18 +89,11 @@ FRigUnit_DrawContainerSetTransform_Execute()
 		return;
 	}
 
-	switch (ExecuteContext.UnitContext.State)
+	int32 Index = ExecuteContext.UnitContext.DrawContainer->GetIndex(InstructionName);
+	if (Index != INDEX_NONE)
 	{
-		case EControlRigState::Update:
-		{
-			int32 Index = ExecuteContext.UnitContext.DrawContainer->GetIndex(InstructionName);
-			if (Index != INDEX_NONE)
-			{
-				FControlRigDrawInstruction& Instruction = (*ExecuteContext.UnitContext.DrawContainer)[Index];
-				Instruction.Transform = Transform;
-			}
-			break;
-		}
+		FControlRigDrawInstruction& Instruction = (*ExecuteContext.UnitContext.DrawContainer)[Index];
+		Instruction.Transform = Transform;
 	}
 }
 
diff --git a/Engine/Plugins/Animation/ControlRig/Source/ControlRig/Private/Units/Execution/RigUnit_Collection.cpp b/Engine/Plugins/Animation/ControlRig/Source/ControlRig/Private/Units/Execution/RigUnit_Collection.cpp
--- a/Engine/Plugins/Animation/ControlRig/Source/ControlRig/Private/Units/Execution/RigUnit_Collection.cpp
+++ b/Engine/Plugins/Animation/ControlRig/Source/ControlRig/Private/Units/Execution/RigUnit_Collection.cpp
@@ -57,17 +57,11 @@ FRigUnit_CollectionChainArray_Execute()
 		{
 			if (ExecuteContext.Hierarchy->GetIndex(FirstItem) == INDEX_NONE)
 			{
-				if(ExecuteContext.UnitContext.State != EControlRigState::Init)
-				{
-					UE_CONTROLRIG_RIGUNIT_REPORT_WARNING(TEXT("First Item '%s' is not valid."), *FirstItem.ToString());
-				}
+				UE_CONTROLRIG_RIGUNIT_REPORT_WARNING(TEXT("First Item '%s' is not valid."), *FirstItem.ToString());
 			}
 			if (ExecuteContext.Hierarchy->GetIndex(LastItem) == INDEX_NONE)
 			{
-				if(ExecuteContext.UnitContext.State != EControlRigState::Init)
-				{
-					UE_CONTROLRIG_RIGUNIT_REPORT_WARNING(TEXT("Last Item '%s' is not valid."), *LastItem.ToString());
-				}
+				UE_CONTROLRIG_RIGUNIT_REPORT_WARNING(TEXT("Last Item '%s' is not valid."), *LastItem.ToString());
 			}
 		}
 		ExecuteContext.Hierarchy->AddCachedCollection(Hash, Collection);
diff --git a/Engine/Plugins/Animation/ControlRig/Source/ControlRig/Private/Units/Execution/RigUnit_DynamicHierarchy.cpp b/Engine/Plugins/Animation/ControlRig/Source/ControlRig/Private/Units/Execution/RigUnit_DynamicHierarchy.cpp
--- a/Engine/Plugins/Animation/ControlRig/Source/ControlRig/Private/Units/Execution/RigUnit_DynamicHierarchy.cpp
+++ b/Engine/Plugins/Animation/ControlRig/Source/ControlRig/Private/Units/Execution/RigUnit_DynamicHierarchy.cpp
@@ -15,8 +15,7 @@ bool FRigUnit_DynamicHierarchyBase::IsValidToRunInContext(
 	bool bAllowOnlyConstructionEvent,
 	FString* OutErrorMessage)
 {
-	if((InExecuteContext.UnitContext.State != EControlRigState::Update) ||
-		(InExecuteContext.Hierarchy == nullptr))
+	if(InExecuteContext.Hierarchy == nullptr)
 	{
 		return false;
 	}
@@ -213,7 +212,7 @@ FRigVMStructUpgradeInfo FRigUnit_HierarchyGetParentWeights::GetUpgradeInfo() con
 
 FRigUnit_HierarchyGetParentWeightsArray_Execute()
 {
-	if((ExecuteContext.UnitContext.State != EControlRigState::Update) || (ExecuteContext.Hierarchy == nullptr))
+	if(ExecuteContext.Hierarchy == nullptr)
 	{
 		return;
 	}
diff --git a/Engine/Plugins/Animation/ControlRig/Source/ControlRig/Private/Units/Execution/RigUnit_Hierarchy.cpp b/Engine/Plugins/Animation/ControlRig/Source/ControlRig/Private/Units/Execution/RigUnit_Hierarchy.cpp
--- a/Engine/Plugins/Animation/ControlRig/Source/ControlRig/Private/Units/Execution/RigUnit_Hierarchy.cpp
+++ b/Engine/Plugins/Animation/ControlRig/Source/ControlRig/Private/Units/Execution/RigUnit_Hierarchy.cpp
@@ -11,12 +11,6 @@ FRigUnit_HierarchyGetParent_Execute()
 {
     DECLARE_SCOPE_HIERARCHICAL_COUNTER_RIGUNIT()
 
-	if (ExecuteContext.UnitContext.State == EControlRigState::Init)
-	{
-		CachedChild.Reset();
-		CachedParent.Reset();
-	}
-
 	if(CachedChild.IsIdentical(Child, ExecuteContext.Hierarchy))
 	{
 		Parent = CachedParent.GetKey();
@@ -56,12 +50,6 @@ FRigUnit_HierarchyGetParentsItemArray_Execute()
 {
     DECLARE_SCOPE_HIERARCHICAL_COUNTER_RIGUNIT()
 
-	if (ExecuteContext.UnitContext.State == EControlRigState::Init)
-	{
-		CachedChild.Reset();
-		CachedParents.Reset();
-	}
-
 	if(!CachedChild.IsIdentical(Child, ExecuteContext.Hierarchy))
 	{
 		CachedParents.Reset();
@@ -95,12 +83,6 @@ FRigUnit_HierarchyGetChildren_Execute()
 {
     DECLARE_SCOPE_HIERARCHICAL_COUNTER_RIGUNIT()
 
-	if (ExecuteContext.UnitContext.State == EControlRigState::Init)
-	{
-		CachedParent.Reset();
-		CachedChildren.Reset();
-	}
-
 	if(!CachedParent.IsIdentical(Parent, ExecuteContext.Hierarchy))
 	{
 		CachedChildren.Reset();
@@ -151,12 +133,6 @@ FRigUnit_HierarchyGetSiblingsItemArray_Execute()
 {
     DECLARE_SCOPE_HIERARCHICAL_COUNTER_RIGUNIT()
 
-	if (ExecuteContext.UnitContext.State == EControlRigState::Init)
-	{
-		CachedItem.Reset();
-		CachedSiblings.Reset();
-	}
-
 	if(!CachedItem.IsIdentical(Item, ExecuteContext.Hierarchy))
 	{
 		CachedSiblings.Reset();
@@ -490,11 +466,6 @@ bool FRigUnit_PoseGetDelta::AreCurvesEqual(float A, float B, float CurveU)
 
 FRigUnit_PoseGetTransform_Execute()
 {
-	if (ExecuteContext.UnitContext.State == EControlRigState::Init)
-	{
-		CachedPoseElementIndex = CachedPoseHash = INDEX_NONE;
-	}
-
 	// set up defaults
 	Valid = false;
 	Transform = FTransform::Identity;
@@ -552,11 +523,6 @@ FRigUnit_PoseGetTransformArray_Execute()
 
 FRigUnit_PoseGetCurve_Execute()
 {
-	if (ExecuteContext.UnitContext.State == EControlRigState::Init)
-	{
-		CachedPoseElementIndex = CachedPoseHash = INDEX_NONE;
-	}
-
 	// set up defaults
 	Valid = false;
 	CurveValue = 0.f;
diff --git a/Engine/Plugins/Animation/ControlRig/Source/ControlRig/Private/Units/Execution/RigUnit_Item.cpp b/Engine/Plugins/Animation/ControlRig/Source/ControlRig/Private/Units/Execution/RigUnit_Item.cpp
--- a/Engine/Plugins/Animation/ControlRig/Source/ControlRig/Private/Units/Execution/RigUnit_Item.cpp
+++ b/Engine/Plugins/Animation/ControlRig/Source/ControlRig/Private/Units/Execution/RigUnit_Item.cpp
@@ -11,25 +11,7 @@ FRigUnit_ItemExists_Execute()
 {
     DECLARE_SCOPE_HIERARCHICAL_COUNTER_RIGUNIT()
 
-	Exists = false;
-
-	switch (ExecuteContext.UnitContext.State)
-	{
-		case EControlRigState::Init:
-		{
-			CachedIndex.Reset();
-			// fall through to update
-		}
-		case EControlRigState::Update:
-		{
-			Exists = CachedIndex.UpdateCache(Item, ExecuteContext.Hierarchy);
-	    	break;
-	    }
-	    default:
-	    {
-	    	break;
-	    }
-	}
+	Exists = CachedIndex.UpdateCache(Item, ExecuteContext.Hierarchy);
 }
 
 FRigUnit_ItemReplace_Execute()
diff --git a/Engine/Plugins/Animation/ControlRig/Source/ControlRig/Private/Units/Hierarchy/RigUnit_AddBoneTransform.cpp b/Engine/Plugins/Animation/ControlRig/Source/ControlRig/Private/Units/Hierarchy/RigUnit_AddBoneTransform.cpp
--- a/Engine/Plugins/Animation/ControlRig/Source/ControlRig/Private/Units/Hierarchy/RigUnit_AddBoneTransform.cpp
+++ b/Engine/Plugins/Animation/ControlRig/Source/ControlRig/Private/Units/Hierarchy/RigUnit_AddBoneTransform.cpp
@@ -12,46 +12,32 @@ FRigUnit_AddBoneTransform_Execute()
 	URigHierarchy* Hierarchy = ExecuteContext.Hierarchy;
 	if (Hierarchy)
 	{
-		switch (ExecuteContext.UnitContext.State)
+		const FRigElementKey Key(Bone, ERigElementType::Bone); 
+		if (!CachedBone.UpdateCache(Key, Hierarchy))
 		{
-			case EControlRigState::Init:
+			UE_CONTROLRIG_RIGUNIT_REPORT_WARNING(TEXT("Bone '%s' is not valid."), *Bone.ToString());
+		}
+		else
+		{
+			FTransform TargetTransform;
+			const FTransform PreviousTransform = Hierarchy->GetGlobalTransform(CachedBone);
+
+			if (bPostMultiply)
 			{
-				CachedBone.Reset();
+				TargetTransform = PreviousTransform * Transform;
 			}
-			case EControlRigState::Update:
+			else
 			{
-				const FRigElementKey Key(Bone, ERigElementType::Bone); 
-				if (!CachedBone.UpdateCache(Key, Hierarchy))
-				{
-					UE_CONTROLRIG_RIGUNIT_REPORT_WARNING(TEXT("Bone '%s' is not valid."), *Bone.ToString());
-				}
-				else
-				{
-					FTransform TargetTransform;
-					const FTransform PreviousTransform = Hierarchy->GetGlobalTransform(CachedBone);
-
-					if (bPostMultiply)
-					{
-						TargetTransform = PreviousTransform * Transform;
-					}
-					else
-					{
-						TargetTransform = Transform * PreviousTransform;
-					}
-
-					if (!FMath::IsNearlyEqual(Weight, 1.f))
-					{
-						float T = FMath::Clamp<float>(Weight, 0.f, 1.f);
-						TargetTransform = FControlRigMathLibrary::LerpTransform(PreviousTransform, TargetTransform, T);
-					}
-
-					Hierarchy->SetGlobalTransform(CachedBone, TargetTransform, bPropagateToChildren);
-				}
+				TargetTransform = Transform * PreviousTransform;
 			}
-			default:
+
+			if (!FMath::IsNearlyEqual(Weight, 1.f))
 			{
-				break;
+				float T = FMath::Clamp<float>(Weight, 0.f, 1.f);
+				TargetTransform = FControlRigMathLibrary::LerpTransform(PreviousTransform, TargetTransform, T);
 			}
+
+			Hierarchy->SetGlobalTransform(CachedBone, TargetTransform, bPropagateToChildren);
 		}
 	}
 }
diff --git a/Engine/Plugins/Animation/ControlRig/Source/ControlRig/Private/Units/Hierarchy/RigUnit_GetBoneTransform.cpp b/Engine/Plugins/Animation/ControlRig/Source/ControlRig/Private/Units/Hierarchy/RigUnit_GetBoneTransform.cpp
--- a/Engine/Plugins/Animation/ControlRig/Source/ControlRig/Private/Units/Hierarchy/RigUnit_GetBoneTransform.cpp
+++ b/Engine/Plugins/Animation/ControlRig/Source/ControlRig/Private/Units/Hierarchy/RigUnit_GetBoneTransform.cpp
@@ -13,43 +13,28 @@ FRigUnit_GetBoneTransform_Execute()
 	const URigHierarchy* Hierarchy = ExecuteContext.Hierarchy;
 	if (Hierarchy)
 	{
-		switch (ExecuteContext.UnitContext.State)
+		if (!CachedBone.UpdateCache(FRigElementKey(Bone, ERigElementType::Bone), Hierarchy))
 		{
-			case EControlRigState::Init:
-			{
-				FRigUnit_GetInitialBoneTransform::StaticExecute(ExecuteContext, Bone, Space, Transform, CachedBone);;
-				break;
-			}
-			case EControlRigState::Update:
+			UE_CONTROLRIG_RIGUNIT_REPORT_WARNING(TEXT("Bone '%s' is not valid."), *Bone.ToString());
+		}
+		else
+		{
+			switch (Space)
 			{
-				if (!CachedBone.UpdateCache(FRigElementKey(Bone, ERigElementType::Bone), Hierarchy))
+				case EBoneGetterSetterMode::GlobalSpace:
 				{
-					UE_CONTROLRIG_RIGUNIT_REPORT_WARNING(TEXT("Bone '%s' is not valid."), *Bone.ToString());
+					Transform = Hierarchy->GetGlobalTransform(CachedBone);
+					break;
 				}
-				else
+				case EBoneGetterSetterMode::LocalSpace:
 				{
-					switch (Space)
-					{
-						case EBoneGetterSetterMode::GlobalSpace:
-						{
-							Transform = Hierarchy->GetGlobalTransform(CachedBone);
-							break;
-						}
-						case EBoneGetterSetterMode::LocalSpace:
-						{
-							Transform = Hierarchy->GetLocalTransform(CachedBone);
-							break;
-						}
-						default:
-						{
-							break;
-						}
-					}
+					Transform = Hierarchy->GetLocalTransform(CachedBone);
+					break;
+				}
+				default:
+				{
+					break;
 				}
-			}
-			default:
-			{
-				break;
 			}
 		}
 	}
@@ -76,26 +61,26 @@ IMPLEMENT_RIGUNIT_AUTOMATION_TEST(FRigUnit_GetBoneTransform)
 
 	Unit.Bone = TEXT("Unknown");
 	Unit.Space = EBoneGetterSetterMode::GlobalSpace;
-	InitAndExecute();
+	Execute();
 	AddErrorIfFalse(Unit.Transform.GetTranslation().Equals(FVector(0.f, 0.f, 0.f)), TEXT("unexpected global transform"));
 	Unit.Space = EBoneGetterSetterMode::LocalSpace;
-	InitAndExecute();
+	Execute();
 	AddErrorIfFalse(Unit.Transform.GetTranslation().Equals(FVector(0.f, 0.f, 0.f)), TEXT("unexpected local transform"));
 
 	Unit.Bone = TEXT("Root");
 	Unit.Space = EBoneGetterSetterMode::GlobalSpace;
-	InitAndExecute();
+	Execute();
 	AddErrorIfFalse(Unit.Transform.GetTranslation().Equals(FVector(1.f, 0.f, 0.f)), TEXT("unexpected global transform"));
 	Unit.Space = EBoneGetterSetterMode::LocalSpace;
-	InitAndExecute();
+	Execute();
 	AddErrorIfFalse(Unit.Transform.GetTranslation().Equals(FVector(1.f, 0.f, 0.f)), TEXT("unexpected local transform"));
 
 	Unit.Bone = TEXT("BoneA");
 	Unit.Space = EBoneGetterSetterMode::GlobalSpace;
-	InitAndExecute();
+	Execute();
 	AddErrorIfFalse(Unit.Transform.GetTranslation().Equals(FVector(1.f, 2.f, 3.f)), TEXT("unexpected global transform"));
 	Unit.Space = EBoneGetterSetterMode::LocalSpace;
-	InitAndExecute();
+	Execute();
 	AddErrorIfFalse(Unit.Transform.GetTranslation().Equals(FVector(0.f, 2.f, 3.f)), TEXT("unexpected local transform"));
 
 	return true;
diff --git a/Engine/Plugins/Animation/ControlRig/Source/ControlRig/Private/Units/Hierarchy/RigUnit_GetControlInitialTransform.cpp b/Engine/Plugins/Animation/ControlRig/Source/ControlRig/Private/Units/Hierarchy/RigUnit_GetControlInitialTransform.cpp
--- a/Engine/Plugins/Animation/ControlRig/Source/ControlRig/Private/Units/Hierarchy/RigUnit_GetControlInitialTransform.cpp
+++ b/Engine/Plugins/Animation/ControlRig/Source/ControlRig/Private/Units/Hierarchy/RigUnit_GetControlInitialTransform.cpp
@@ -12,43 +12,29 @@ FRigUnit_GetControlInitialTransform_Execute()
 	const URigHierarchy* Hierarchy = ExecuteContext.Hierarchy;
 	if (Hierarchy)
 	{
-		switch (ExecuteContext.UnitContext.State)
+		const FRigElementKey Key(Control, ERigElementType::Control); 
+		if (!CachedControlIndex.UpdateCache(Key, Hierarchy))
 		{
-			case EControlRigState::Init:
-			{
-				CachedControlIndex.Reset();
-			}
-			case EControlRigState::Update:
+			UE_CONTROLRIG_RIGUNIT_REPORT_WARNING(TEXT("Control '%s' is not valid."), *Control.ToString());
+		}
+		else
+		{
+			switch (Space)
 			{
-				const FRigElementKey Key(Control, ERigElementType::Control); 
-				if (!CachedControlIndex.UpdateCache(Key, Hierarchy))
+				case EBoneGetterSetterMode::GlobalSpace:
 				{
-					UE_CONTROLRIG_RIGUNIT_REPORT_WARNING(TEXT("Control '%s' is not valid."), *Control.ToString());
+					Transform = Hierarchy->GetInitialGlobalTransform(CachedControlIndex);
+					break;
 				}
-				else
+				case EBoneGetterSetterMode::LocalSpace:
 				{
-					switch (Space)
-					{
-						case EBoneGetterSetterMode::GlobalSpace:
-						{
-							Transform = Hierarchy->GetInitialGlobalTransform(CachedControlIndex);
-							break;
-						}
-						case EBoneGetterSetterMode::LocalSpace:
-						{
-							Transform = Hierarchy->GetInitialLocalTransform(CachedControlIndex);
-							break;
-						}
-						default:
-						{
-							break;
-						}
-					}
+					Transform = Hierarchy->GetInitialLocalTransform(CachedControlIndex);
+					break;
+				}
+				default:
+				{
+					break;
 				}
-			}
-			default:
-			{
-				break;
 			}
 		}
 	}
diff --git a/Engine/Plugins/Animation/ControlRig/Source/ControlRig/Private/Units/Hierarchy/RigUnit_GetControlTransform.cpp b/Engine/Plugins/Animation/ControlRig/Source/ControlRig/Private/Units/Hierarchy/RigUnit_GetControlTransform.cpp
--- a/Engine/Plugins/Animation/ControlRig/Source/ControlRig/Private/Units/Hierarchy/RigUnit_GetControlTransform.cpp
+++ b/Engine/Plugins/Animation/ControlRig/Source/ControlRig/Private/Units/Hierarchy/RigUnit_GetControlTransform.cpp
@@ -11,27 +11,13 @@ FRigUnit_GetControlBool_Execute()
 	const URigHierarchy* Hierarchy = ExecuteContext.Hierarchy;
 	if (Hierarchy)
 	{
-		switch (ExecuteContext.UnitContext.State)
+		if (!CachedControlIndex.UpdateCache(FRigElementKey(Control, ERigElementType::Control), Hierarchy))
 		{
-			case EControlRigState::Init:
-			{
-				CachedControlIndex.Reset();
-			}
-			case EControlRigState::Update:
-			{
-				if (!CachedControlIndex.UpdateCache(FRigElementKey(Control, ERigElementType::Control), Hierarchy))
-				{
-					UE_CONTROLRIG_RIGUNIT_REPORT_WARNING(TEXT("Control '%s' is not valid."), *Control.ToString());
-				}
-				else
-				{
-					BoolValue = Hierarchy->GetControlValue(CachedControlIndex).Get<bool>();
-				}
-			}
-			default:
-			{
-				break;
-			}
+			UE_CONTROLRIG_RIGUNIT_REPORT_WARNING(TEXT("Control '%s' is not valid."), *Control.ToString());
+		}
+		else
+		{
+			BoolValue = Hierarchy->GetControlValue(CachedControlIndex).Get<bool>();
 		}
 	}
 }
@@ -42,29 +28,15 @@ FRigUnit_GetControlFloat_Execute()
 	const URigHierarchy* Hierarchy = ExecuteContext.Hierarchy;
 	if (Hierarchy)
 	{
-		switch (ExecuteContext.UnitContext.State)
+		if (!CachedControlIndex.UpdateCache(FRigElementKey(Control, ERigElementType::Control), Hierarchy))
 		{
-			case EControlRigState::Init:
-			{
-				CachedControlIndex.Reset();
-			}
-			case EControlRigState::Update:
-			{
-				if (!CachedControlIndex.UpdateCache(FRigElementKey(Control, ERigElementType::Control), Hierarchy))
-				{
-					UE_CONTROLRIG_RIGUNIT_REPORT_WARNING(TEXT("Control '%s' is not valid."), *Control.ToString());
-				}
-				else
-				{
-					FloatValue = Hierarchy->GetControlValue(CachedControlIndex).Get<float>();
-					Minimum = Hierarchy->GetControlValue(CachedControlIndex, ERigControlValueType::Minimum).Get<float>();
-					Maximum = Hierarchy->GetControlValue(CachedControlIndex, ERigControlValueType::Maximum).Get<float>();
-				}
-			}
-			default:
-			{
-				break;
-			}
+			UE_CONTROLRIG_RIGUNIT_REPORT_WARNING(TEXT("Control '%s' is not valid."), *Control.ToString());
+		}
+		else
+		{
+			FloatValue = Hierarchy->GetControlValue(CachedControlIndex).Get<float>();
+			Minimum = Hierarchy->GetControlValue(CachedControlIndex, ERigControlValueType::Minimum).Get<float>();
+			Maximum = Hierarchy->GetControlValue(CachedControlIndex, ERigControlValueType::Maximum).Get<float>();
 		}
 	}
 }
@@ -76,29 +48,15 @@ FRigUnit_GetControlInteger_Execute()
 	const URigHierarchy* Hierarchy = ExecuteContext.Hierarchy;
 	if (Hierarchy)
 	{
-		switch (ExecuteContext.UnitContext.State)
+		if (!CachedControlIndex.UpdateCache(FRigElementKey(Control, ERigElementType::Control), Hierarchy))
 		{
-			case EControlRigState::Init:
-			{
-				CachedControlIndex.Reset();
-			}
-			case EControlRigState::Update:
-			{
-				if (!CachedControlIndex.UpdateCache(FRigElementKey(Control, ERigElementType::Control), Hierarchy))
-				{
-					UE_CONTROLRIG_RIGUNIT_REPORT_WARNING(TEXT("Control '%s' is not valid."), *Control.ToString());
-				}
-				else
-				{
-					IntegerValue = Hierarchy->GetControlValue(CachedControlIndex).Get<int32>();
-					Minimum = Hierarchy->GetControlValue(CachedControlIndex, ERigControlValueType::Minimum).Get<int32>();
-					Maximum = Hierarchy->GetControlValue(CachedControlIndex, ERigControlValueType::Maximum).Get<int32>();
-				}
-			}
-			default:
-			{
-				break;
-			}
+			UE_CONTROLRIG_RIGUNIT_REPORT_WARNING(TEXT("Control '%s' is not valid."), *Control.ToString());
+		}
+		else
+		{
+			IntegerValue = Hierarchy->GetControlValue(CachedControlIndex).Get<int32>();
+			Minimum = Hierarchy->GetControlValue(CachedControlIndex, ERigControlValueType::Minimum).Get<int32>();
+			Maximum = Hierarchy->GetControlValue(CachedControlIndex, ERigControlValueType::Maximum).Get<int32>();
 		}
 	}
 }
@@ -109,32 +67,18 @@ FRigUnit_GetControlVector2D_Execute()
 	const URigHierarchy* Hierarchy = ExecuteContext.Hierarchy;
 	if (Hierarchy)
 	{
-		switch (ExecuteContext.UnitContext.State)
+		if (!CachedControlIndex.UpdateCache(FRigElementKey(Control, ERigElementType::Control), Hierarchy))
 		{
-			case EControlRigState::Init:
-			{
-				CachedControlIndex.Reset();
-			}
-			case EControlRigState::Update:
-			{
-				if (!CachedControlIndex.UpdateCache(FRigElementKey(Control, ERigElementType::Control), Hierarchy))
-				{
-					UE_CONTROLRIG_RIGUNIT_REPORT_WARNING(TEXT("Control '%s' is not valid."), *Control.ToString());
-				}
-				else
-				{
-					const FVector3f TempVector = Hierarchy->GetControlValue(CachedControlIndex).Get<FVector3f>();
-					const FVector3f TempMinimum = Hierarchy->GetControlValue(CachedControlIndex, ERigControlValueType::Minimum).Get<FVector3f>();
-					const FVector3f TempMaximum = Hierarchy->GetControlValue(CachedControlIndex, ERigControlValueType::Maximum).Get<FVector3f>();
-					Vector = FVector2D(TempVector.X, TempVector.Y);
-					Minimum = FVector2D(TempMinimum.X, TempMinimum.Y);
-					Maximum = FVector2D(TempMaximum.X, TempMaximum.Y);
-				}
-			}
-			default:
-			{
-				break;
-			}
+			UE_CONTROLRIG_RIGUNIT_REPORT_WARNING(TEXT("Control '%s' is not valid."), *Control.ToString());
+		}
+		else
+		{
+			const FVector3f TempVector = Hierarchy->GetControlValue(CachedControlIndex).Get<FVector3f>();
+			const FVector3f TempMinimum = Hierarchy->GetControlValue(CachedControlIndex, ERigControlValueType::Minimum).Get<FVector3f>();
+			const FVector3f TempMaximum = Hierarchy->GetControlValue(CachedControlIndex, ERigControlValueType::Maximum).Get<FVector3f>();
+			Vector = FVector2D(TempVector.X, TempVector.Y);
+			Minimum = FVector2D(TempMinimum.X, TempMinimum.Y);
+			Maximum = FVector2D(TempMaximum.X, TempMaximum.Y);
 		}
 	}
 }
@@ -145,58 +89,44 @@ FRigUnit_GetControlVector_Execute()
 	const URigHierarchy* Hierarchy = ExecuteContext.Hierarchy;
 	if (Hierarchy)
 	{
-		switch (ExecuteContext.UnitContext.State)
+		if (!CachedControlIndex.UpdateCache(FRigElementKey(Control, ERigElementType::Control), Hierarchy))
 		{
-			case EControlRigState::Init:
-			{
-				CachedControlIndex.Reset();
-			}
-			case EControlRigState::Update:
+			UE_CONTROLRIG_RIGUNIT_REPORT_WARNING(TEXT("Control '%s' is not valid."), *Control.ToString());
+		}
+		else
+		{
+			FTransform Transform = FTransform::Identity;
+			switch (Space)
 			{
-				if (!CachedControlIndex.UpdateCache(FRigElementKey(Control, ERigElementType::Control), Hierarchy))
+				case EBoneGetterSetterMode::GlobalSpace:
 				{
-					UE_CONTROLRIG_RIGUNIT_REPORT_WARNING(TEXT("Control '%s' is not valid."), *Control.ToString());
+					Transform = Hierarchy->GetGlobalTransform(CachedControlIndex);
+					break;
 				}
-				else
+				case EBoneGetterSetterMode::LocalSpace:
 				{
-					FTransform Transform = FTransform::Identity;
-					switch (Space)
-					{
-						case EBoneGetterSetterMode::GlobalSpace:
-						{
-							Transform = Hierarchy->GetGlobalTransform(CachedControlIndex);
-							break;
-						}
-						case EBoneGetterSetterMode::LocalSpace:
-						{
-							Transform = Hierarchy->GetLocalTransform(CachedControlIndex);
-							break;
-						}
-						default:
-						{
-							break;
-						}
-					}
-
-					const ERigControlType ControlType = Hierarchy->GetChecked<FRigControlElement>(CachedControlIndex)->Settings.ControlType;
-					
-					if(ControlType == ERigControlType::Position)
-					{
-						Vector = Transform.GetLocation();
-					}
-					else if(ControlType == ERigControlType::Scale)
-					{
-						Vector = Transform.GetScale3D();
-					}
-
-					Minimum = (FVector)Hierarchy->GetControlValue(CachedControlIndex, ERigControlValueType::Minimum).Get<FVector3f>();
-					Maximum = (FVector)Hierarchy->GetControlValue(CachedControlIndex, ERigControlValueType::Maximum).Get<FVector3f>();
+					Transform = Hierarchy->GetLocalTransform(CachedControlIndex);
+					break;
+				}
+				default:
+				{
+					break;
 				}
 			}
-			default:
+
+			const ERigControlType ControlType = Hierarchy->GetChecked<FRigControlElement>(CachedControlIndex)->Settings.ControlType;
+			
+			if(ControlType == ERigControlType::Position)
 			{
-				break;
+				Vector = Transform.GetLocation();
 			}
+			else if(ControlType == ERigControlType::Scale)
+			{
+				Vector = Transform.GetScale3D();
+			}
+
+			Minimum = (FVector)Hierarchy->GetControlValue(CachedControlIndex, ERigControlValueType::Minimum).Get<FVector3f>();
+			Maximum = (FVector)Hierarchy->GetControlValue(CachedControlIndex, ERigControlValueType::Maximum).Get<FVector3f>();
 		}
 	}
 }
@@ -207,51 +137,37 @@ FRigUnit_GetControlRotator_Execute()
 	const URigHierarchy* Hierarchy = ExecuteContext.Hierarchy;
 	if (Hierarchy)
 	{
-		switch (ExecuteContext.UnitContext.State)
+		if (!CachedControlIndex.UpdateCache(FRigElementKey(Control, ERigElementType::Control), Hierarchy))
 		{
-			case EControlRigState::Init:
-			{
-				CachedControlIndex.Reset();
-			}
-			case EControlRigState::Update:
+			UE_CONTROLRIG_RIGUNIT_REPORT_WARNING(TEXT("Control '%s' is not valid."), *Control.ToString());
+		}
+		else
+		{
+			FTransform Transform = FTransform::Identity;
+			switch (Space)
 			{
-				if (!CachedControlIndex.UpdateCache(FRigElementKey(Control, ERigElementType::Control), Hierarchy))
+				case EBoneGetterSetterMode::GlobalSpace:
 				{
-					UE_CONTROLRIG_RIGUNIT_REPORT_WARNING(TEXT("Control '%s' is not valid."), *Control.ToString());
+					Transform = Hierarchy->GetGlobalTransform(CachedControlIndex);
+					break;
 				}
-				else
+				case EBoneGetterSetterMode::LocalSpace:
 				{
-					FTransform Transform = FTransform::Identity;
-					switch (Space)
-					{
-						case EBoneGetterSetterMode::GlobalSpace:
-						{
-							Transform = Hierarchy->GetGlobalTransform(CachedControlIndex);
-							break;
-						}
-						case EBoneGetterSetterMode::LocalSpace:
-						{
-							Transform = Hierarchy->GetLocalTransform(CachedControlIndex);
-							break;
-						}
-						default:
-						{
-							break;
-						}
-					}
-
-					Rotator = Transform.GetRotation().Rotator();
-
-					const FVector MinimumVector = (FVector)Hierarchy->GetControlValue(CachedControlIndex, ERigControlValueType::Minimum).Get<FVector3f>();
-					const FVector MaximumVector = (FVector)Hierarchy->GetControlValue(CachedControlIndex, ERigControlValueType::Maximum).Get<FVector3f>();
-					Minimum = FRotator::MakeFromEuler(MinimumVector); 
-					Maximum = FRotator::MakeFromEuler(MaximumVector); 
+					Transform = Hierarchy->GetLocalTransform(CachedControlIndex);
+					break;
+				}
+				default:
+				{
+					break;
 				}
 			}
-			default:
-			{
-				break;
-			}
+
+			Rotator = Transform.GetRotation().Rotator();
+
+			const FVector MinimumVector = (FVector)Hierarchy->GetControlValue(CachedControlIndex, ERigControlValueType::Minimum).Get<FVector3f>();
+			const FVector MaximumVector = (FVector)Hierarchy->GetControlValue(CachedControlIndex, ERigControlValueType::Maximum).Get<FVector3f>();
+			Minimum = FRotator::MakeFromEuler(MinimumVector); 
+			Maximum = FRotator::MakeFromEuler(MaximumVector); 
 		}
 	}
 }
@@ -262,46 +178,32 @@ FRigUnit_GetControlTransform_Execute()
 	const URigHierarchy* Hierarchy = ExecuteContext.Hierarchy;
 	if (Hierarchy)
 	{
-		switch (ExecuteContext.UnitContext.State)
+		if (!CachedControlIndex.UpdateCache(FRigElementKey(Control, ERigElementType::Control), Hierarchy))
 		{
-			case EControlRigState::Init:
-			{
-				CachedControlIndex.Reset();
-			}
-			case EControlRigState::Update:
+			UE_CONTROLRIG_RIGUNIT_REPORT_WARNING(TEXT("Control '%s' is not valid."), *Control.ToString());
+		}
+		else
+		{
+			switch (Space)
 			{
-				if (!CachedControlIndex.UpdateCache(FRigElementKey(Control, ERigElementType::Control), Hierarchy))
+				case EBoneGetterSetterMode::GlobalSpace:
 				{
-					UE_CONTROLRIG_RIGUNIT_REPORT_WARNING(TEXT("Control '%s' is not valid."), *Control.ToString());
+					Transform = Hierarchy->GetGlobalTransform(CachedControlIndex);
+					break;
 				}
-				else
+				case EBoneGetterSetterMode::LocalSpace:
 				{
-					switch (Space)
-					{
-						case EBoneGetterSetterMode::GlobalSpace:
-						{
-							Transform = Hierarchy->GetGlobalTransform(CachedControlIndex);
-							break;
-						}
-						case EBoneGetterSetterMode::LocalSpace:
-						{
-							Transform = Hierarchy->GetLocalTransform(CachedControlIndex);
-							break;
-						}
-						default:
-						{
-							break;
-						}
-					}
-					
-					Minimum = Hierarchy->GetControlValue(CachedControlIndex, ERigControlValueType::Minimum).Get<FRigControlValue::FTransform_Float>().ToTransform();
-					Maximum = Hierarchy->GetControlValue(CachedControlIndex, ERigControlValueType::Maximum).Get<FRigControlValue::FTransform_Float>().ToTransform();
+					Transform = Hierarchy->GetLocalTransform(CachedControlIndex);
+					break;
+				}
+				default:
+				{
+					break;
 				}
 			}
-			default:
-			{
-				break;
-			}
+			
+			Minimum = Hierarchy->GetControlValue(CachedControlIndex, ERigControlValueType::Minimum).Get<FRigControlValue::FTransform_Float>().ToTransform();
+			Maximum = Hierarchy->GetControlValue(CachedControlIndex, ERigControlValueType::Maximum).Get<FRigControlValue::FTransform_Float>().ToTransform();
 		}
 	}
 }
@@ -332,26 +234,26 @@ IMPLEMENT_RIGUNIT_AUTOMATION_TEST(FRigUnit_GetControlTransform)
 
 	Unit.Control = TEXT("Unknown");
 	Unit.Space = EBoneGetterSetterMode::GlobalSpace;
-	InitAndExecute();
+	Execute();
 	AddErrorIfFalse(Unit.Transform.GetTranslation().Equals(FVector(0.f, 0.f, 0.f)), TEXT("unexpected global transform (0)"));
 	Unit.Space = EBoneGetterSetterMode::LocalSpace;
-	InitAndExecute();
+	Execute();
 	AddErrorIfFalse(Unit.Transform.GetTranslation().Equals(FVector(0.f, 0.f, 0.f)), TEXT("unexpected local transform (0)"));
 
 	Unit.Control = TEXT("Root");
 	Unit.Space = EBoneGetterSetterMode::GlobalSpace;
-	InitAndExecute();
+	Execute();
 	AddErrorIfFalse(Unit.Transform.GetTranslation().Equals(FVector(1.f, 0.f, 0.f)), TEXT("unexpected global transform (1)"));
 	Unit.Space = EBoneGetterSetterMode::LocalSpace;
-	InitAndExecute();
+	Execute();
 	AddErrorIfFalse(Unit.Transform.GetTranslation().Equals(FVector(1.f, 0.f, 0.f)), TEXT("unexpected local transform (1)"));
 
 	Unit.Control = TEXT("ControlA");
 	Unit.Space = EBoneGetterSetterMode::GlobalSpace;
-	InitAndExecute();
+	Execute();
 	AddErrorIfFalse(Unit.Transform.GetTranslation().Equals(FVector(2.f, 2.f, 3.f)), TEXT("unexpected global transform (2)"));
 	Unit.Space = EBoneGetterSetterMode::LocalSpace;
-	InitAndExecute();
+	Execute();
 	AddErrorIfFalse(Unit.Transform.GetTranslation().Equals(FVector(1.f, 2.f, 3.f)), TEXT("unexpected local transform (2)"));
 
 	return true;
diff --git a/Engine/Plugins/Animation/ControlRig/Source/ControlRig/Private/Units/Hierarchy/RigUnit_GetCurveValue.cpp b/Engine/Plugins/Animation/ControlRig/Source/ControlRig/Private/Units/Hierarchy/RigUnit_GetCurveValue.cpp
--- a/Engine/Plugins/Animation/ControlRig/Source/ControlRig/Private/Units/Hierarchy/RigUnit_GetCurveValue.cpp
+++ b/Engine/Plugins/Animation/ControlRig/Source/ControlRig/Private/Units/Hierarchy/RigUnit_GetCurveValue.cpp
@@ -10,24 +10,10 @@ FRigUnit_GetCurveValue_Execute()
     DECLARE_SCOPE_HIERARCHICAL_COUNTER_RIGUNIT()
     if (const URigHierarchy* Hierarchy = ExecuteContext.Hierarchy)
 	{
-		switch (ExecuteContext.UnitContext.State)
+		if(CachedCurveIndex.UpdateCache(FRigElementKey(Curve, ERigElementType::Curve), Hierarchy))
 		{
-			case EControlRigState::Init:
-			{
-				CachedCurveIndex.Reset();
-			}
-			case EControlRigState::Update:
-			{
-				if(CachedCurveIndex.UpdateCache(FRigElementKey(Curve, ERigElementType::Curve), Hierarchy))
-				{
-					Valid = Hierarchy->IsCurveValueSetByIndex(CachedCurveIndex);
-					Value = Hierarchy->GetCurveValueByIndex(CachedCurveIndex);
-				}
-			}
-			default:
-			{
-				break;
-			}
+			Valid = Hierarchy->IsCurveValueSetByIndex(CachedCurveIndex);
+			Value = Hierarchy->GetCurveValueByIndex(CachedCurveIndex);
 		}
 	}
 }
@@ -43,26 +29,26 @@ IMPLEMENT_RIGUNIT_AUTOMATION_TEST(FRigUnit_GetCurveValue)
 
 	Unit.Curve = TEXT("Unknown");
 	Unit.Space = ECurveGetterSetterMode::GlobalSpace;
-	InitAndExecute();
+	Execute();
 	AddErrorIfFalse(Unit.Transform.GetTranslation().Equals(FVector(0.f, 0.f, 0.f)), TEXT("unexpected global transform"));
 	Unit.Space = ECurveGetterSetterMode::LocalSpace;
-	InitAndExecute();
+	Execute();
 	AddErrorIfFalse(Unit.Transform.GetTranslation().Equals(FVector(0.f, 0.f, 0.f)), TEXT("unexpected local transform"));
 
 	Unit.Curve = TEXT("Root");
 	Unit.Space = ECurveGetterSetterMode::GlobalSpace;
-	InitAndExecute();
+	Execute();
 	AddErrorIfFalse(Unit.Transform.GetTranslation().Equals(FVector(1.f, 0.f, 0.f)), TEXT("unexpected global transform"));
 	Unit.Space = ECurveGetterSetterMode::LocalSpace;
-	InitAndExecute();
+	Execute();
 	AddErrorIfFalse(Unit.Transform.GetTranslation().Equals(FVector(1.f, 0.f, 0.f)), TEXT("unexpected local transform"));
 
 	Unit.Curve = TEXT("CurveA");
 	Unit.Space = ECurveGetterSetterMode::GlobalSpace;
-	InitAndExecute();
+	Execute();
 	AddErrorIfFalse(Unit.Transform.GetTranslation().Equals(FVector(1.f, 2.f, 3.f)), TEXT("unexpected global transform"));
 	Unit.Space = ECurveGetterSetterMode::LocalSpace;
-	InitAndExecute();
+	Execute();
 	AddErrorIfFalse(Unit.Transform.GetTranslation().Equals(FVector(0.f, 2.f, 3.f)), TEXT("unexpected local transform"));
 	*/
 	return true;
diff --git a/Engine/Plugins/Animation/ControlRig/Source/ControlRig/Private/Units/Hierarchy/RigUnit_GetInitialBoneTransform.cpp b/Engine/Plugins/Animation/ControlRig/Source/ControlRig/Private/Units/Hierarchy/RigUnit_GetInitialBoneTransform.cpp
--- a/Engine/Plugins/Animation/ControlRig/Source/ControlRig/Private/Units/Hierarchy/RigUnit_GetInitialBoneTransform.cpp
+++ b/Engine/Plugins/Animation/ControlRig/Source/ControlRig/Private/Units/Hierarchy/RigUnit_GetInitialBoneTransform.cpp
@@ -12,42 +12,28 @@ FRigUnit_GetInitialBoneTransform_Execute()
 	const URigHierarchy* Hierarchy = ExecuteContext.Hierarchy;
 	if (Hierarchy)
 	{
-		switch (ExecuteContext.UnitContext.State)
+		if (!CachedBone.UpdateCache(FRigElementKey(Bone, ERigElementType::Bone), Hierarchy))
 		{
-			case EControlRigState::Init:
-			{
-				CachedBone.Reset();
-			}
-			case EControlRigState::Update:
+			UE_CONTROLRIG_RIGUNIT_REPORT_WARNING(TEXT("Bone '%s' is not valid."), *Bone.ToString());
+		}
+		else
+		{
+			switch (Space)
 			{
-				if (!CachedBone.UpdateCache(FRigElementKey(Bone, ERigElementType::Bone), Hierarchy))
+				case EBoneGetterSetterMode::GlobalSpace:
 				{
-					UE_CONTROLRIG_RIGUNIT_REPORT_WARNING(TEXT("Bone '%s' is not valid."), *Bone.ToString());
+					Transform = Hierarchy->GetInitialGlobalTransform(CachedBone);
+					break;
 				}
-				else
+				case EBoneGetterSetterMode::LocalSpace:
 				{
-					switch (Space)
-					{
-						case EBoneGetterSetterMode::GlobalSpace:
-						{
-							Transform = Hierarchy->GetInitialGlobalTransform(CachedBone);
-							break;
-						}
-						case EBoneGetterSetterMode::LocalSpace:
-						{
-							Transform = Hierarchy->GetInitialLocalTransform(CachedBone);
-							break;
-						}
-						default:
-						{
-							break;
-						}
-					}
+					Transform = Hierarchy->GetInitialLocalTransform(CachedBone);
+					break;
+				}
+				default:
+				{
+					break;
 				}
-			}
-			default:
-			{
-				break;
 			}
 		}
 	}
@@ -75,18 +61,18 @@ IMPLEMENT_RIGUNIT_AUTOMATION_TEST(FRigUnit_GetInitialBoneTransform)
 
 	Unit.Bone = TEXT("Root");
 	Unit.Space = EBoneGetterSetterMode::GlobalSpace;
-	InitAndExecute();
+	Execute();
 	AddErrorIfFalse(Unit.Transform.GetTranslation().Equals(FVector(1.f, 0.f, 0.f)), TEXT("unexpected global transform"));
 	Unit.Space = EBoneGetterSetterMode::LocalSpace;
-	InitAndExecute();
+	Execute();
 	AddErrorIfFalse(Unit.Transform.GetTranslation().Equals(FVector(1.f, 0.f, 0.f)), TEXT("unexpected local transform"));
 
 	Unit.Bone = TEXT("BoneA");
 	Unit.Space = EBoneGetterSetterMode::GlobalSpace;
-	InitAndExecute();
+	Execute();
 	AddErrorIfFalse(Unit.Transform.GetTranslation().Equals(FVector(1.f, 2.f, 3.f)), TEXT("unexpected global transform"));
 	Unit.Space = EBoneGetterSetterMode::LocalSpace;
-	InitAndExecute();
+	Execute();
 	AddErrorIfFalse(Unit.Transform.GetTranslation().Equals(FVector(0.f, 2.f, 3.f)), TEXT("unexpected local transform"));
 
 	return true;
