ControlRig: Remove EControlRigState from units

diff --git a/Engine/Plugins/Animation/ControlRig/Source/ControlRig/Private/Units/Hierarchy/RigUnit_GetRelativeBoneTransform.cpp b/Engine/Plugins/Animation/ControlRig/Source/ControlRig/Private/Units/Hierarchy/RigUnit_GetRelativeBoneTransform.cpp
--- a/Engine/Plugins/Animation/ControlRig/Source/ControlRig/Private/Units/Hierarchy/RigUnit_GetRelativeBoneTransform.cpp
+++ b/Engine/Plugins/Animation/ControlRig/Source/ControlRig/Private/Units/Hierarchy/RigUnit_GetRelativeBoneTransform.cpp
@@ -12,34 +12,19 @@ FRigUnit_GetRelativeBoneTransform_Execute()
 	const URigHierarchy* Hierarchy = ExecuteContext.Hierarchy;
 	if (Hierarchy)
 	{
-		switch (ExecuteContext.UnitContext.State)
+		if (!CachedBone.UpdateCache(FRigElementKey(Bone, ERigElementType::Bone), Hierarchy))
 		{
-			case EControlRigState::Init:
-			{
-				CachedBone.Reset();
-				CachedSpace.Reset();
-			}
-			case EControlRigState::Update:
-			{
-				if (!CachedBone.UpdateCache(FRigElementKey(Bone, ERigElementType::Bone), Hierarchy))
-				{
-					UE_CONTROLRIG_RIGUNIT_REPORT_WARNING(TEXT("Bone '%s' is not valid."), *Bone.ToString());
-				}
-				else if (!CachedSpace.UpdateCache(FRigElementKey(Space, ERigElementType::Bone), Hierarchy))
-				{
-					UE_CONTROLRIG_RIGUNIT_REPORT_WARNING(TEXT("Space '%s' is not valid."), *Space.ToString());
-				}
-				else
-				{
-					const FTransform SpaceTransform = Hierarchy->GetGlobalTransform(CachedSpace);
-					const FTransform BoneTransform = Hierarchy->GetGlobalTransform(CachedBone);
-					Transform = BoneTransform.GetRelativeTransform(SpaceTransform);
-				}
-			}
-			default:
-			{
-				break;
-			}
+			UE_CONTROLRIG_RIGUNIT_REPORT_WARNING(TEXT("Bone '%s' is not valid."), *Bone.ToString());
+		}
+		else if (!CachedSpace.UpdateCache(FRigElementKey(Space, ERigElementType::Bone), Hierarchy))
+		{
+			UE_CONTROLRIG_RIGUNIT_REPORT_WARNING(TEXT("Space '%s' is not valid."), *Space.ToString());
+		}
+		else
+		{
+			const FTransform SpaceTransform = Hierarchy->GetGlobalTransform(CachedSpace);
+			const FTransform BoneTransform = Hierarchy->GetGlobalTransform(CachedBone);
+			Transform = BoneTransform.GetRelativeTransform(SpaceTransform);
 		}
 	}
 }
@@ -68,15 +53,15 @@ IMPLEMENT_RIGUNIT_AUTOMATION_TEST(FRigUnit_GetRelativeBoneTransform)
 
 	Unit.Bone = TEXT("Unknown");
 	Unit.Space = TEXT("Root");
-	InitAndExecute();
+	Execute();
 	AddErrorIfFalse(Unit.Transform.GetTranslation().Equals(FVector(0.f, 0.f, 0.f)), TEXT("unexpected transform"));
 
 	Unit.Bone = TEXT("BoneA");
-	InitAndExecute();
+	Execute();
 	AddErrorIfFalse(Unit.Transform.GetTranslation().Equals(FVector(0.f, 2.f, 3.f)), TEXT("unexpected transform"));
 
 	Unit.Space = TEXT("BoneB");
-	InitAndExecute();
+	Execute();
 	AddErrorIfFalse(Unit.Transform.GetTranslation().Equals(FVector(5.f, 2.f, 3.f)), TEXT("unexpected transform"));
 
 	return true;
diff --git a/Engine/Plugins/Animation/ControlRig/Source/ControlRig/Private/Units/Hierarchy/RigUnit_GetSpaceTransform.cpp b/Engine/Plugins/Animation/ControlRig/Source/ControlRig/Private/Units/Hierarchy/RigUnit_GetSpaceTransform.cpp
--- a/Engine/Plugins/Animation/ControlRig/Source/ControlRig/Private/Units/Hierarchy/RigUnit_GetSpaceTransform.cpp
+++ b/Engine/Plugins/Animation/ControlRig/Source/ControlRig/Private/Units/Hierarchy/RigUnit_GetSpaceTransform.cpp
@@ -12,38 +12,24 @@ FRigUnit_GetSpaceTransform_Execute()
 	const URigHierarchy* Hierarchy = ExecuteContext.Hierarchy;
 	if (Hierarchy)
 	{
-		switch (ExecuteContext.UnitContext.State)
+		if (CachedSpaceIndex.UpdateCache(FRigElementKey(Space, ERigElementType::Null), Hierarchy))
 		{
-			case EControlRigState::Init:
+			switch (SpaceType)
 			{
-				CachedSpaceIndex.Reset();
-			}
-			case EControlRigState::Update:
-			{
-				if (CachedSpaceIndex.UpdateCache(FRigElementKey(Space, ERigElementType::Null), Hierarchy))
+				case EBoneGetterSetterMode::GlobalSpace:
 				{
-					switch (SpaceType)
-					{
-						case EBoneGetterSetterMode::GlobalSpace:
-						{
-							Transform = Hierarchy->GetGlobalTransform(CachedSpaceIndex);
-							break;
-						}
-						case EBoneGetterSetterMode::LocalSpace:
-						{
-							Transform = Hierarchy->GetLocalTransform(CachedSpaceIndex);
-							break;
-						}
-						default:
-						{
-							break;
-						}
-					}
+					Transform = Hierarchy->GetGlobalTransform(CachedSpaceIndex);
+					break;
+				}
+				case EBoneGetterSetterMode::LocalSpace:
+				{
+					Transform = Hierarchy->GetLocalTransform(CachedSpaceIndex);
+					break;
+				}
+				default:
+				{
+					break;
 				}
-			}
-			default:
-			{
-				break;
 			}
 		}
 	}
diff --git a/Engine/Plugins/Animation/ControlRig/Source/ControlRig/Private/Units/Hierarchy/RigUnit_GetTransform.cpp b/Engine/Plugins/Animation/ControlRig/Source/ControlRig/Private/Units/Hierarchy/RigUnit_GetTransform.cpp
--- a/Engine/Plugins/Animation/ControlRig/Source/ControlRig/Private/Units/Hierarchy/RigUnit_GetTransform.cpp
+++ b/Engine/Plugins/Animation/ControlRig/Source/ControlRig/Private/Units/Hierarchy/RigUnit_GetTransform.cpp
@@ -18,68 +18,51 @@ FRigUnit_GetTransform_Execute()
 
 	if (URigHierarchy* Hierarchy = ExecuteContext.Hierarchy)
 	{
-		switch (ExecuteContext.UnitContext.State)
+		if (!CachedIndex.UpdateCache(Item, Hierarchy))
 		{
-			case EControlRigState::Init:
-			{
-				CachedIndex.Reset();
-				// there is no "break;" here because an old node, Transform Constrtaint,
-				// still caches data during init stage. Thus, if a it takes a GetTransform
-				// as input, that GetTransform needs to execute and output a valid transform					
-			}
-			case EControlRigState::Update:
+			UE_CONTROLRIG_RIGUNIT_REPORT_WARNING(TEXT("Item '%s' is not valid."), *Item.ToString());
+		}
+		else
+		{
+			if(bInitial)
 			{
-				if (!CachedIndex.UpdateCache(Item, Hierarchy))
+				switch (Space)
 				{
-					UE_CONTROLRIG_RIGUNIT_REPORT_WARNING(TEXT("Item '%s' is not valid."), *Item.ToString());
-				}
-				else
-				{
-					if(bInitial || ExecuteContext.UnitContext.State == EControlRigState::Init)
+					case EBoneGetterSetterMode::GlobalSpace:
+					{
+						Transform = Hierarchy->GetInitialGlobalTransform(CachedIndex);
+						break;
+					}
+					case EBoneGetterSetterMode::LocalSpace:
 					{
-						switch (Space)
-						{
-							case EBoneGetterSetterMode::GlobalSpace:
-							{
-								Transform = Hierarchy->GetInitialGlobalTransform(CachedIndex);
-								break;
-							}
-							case EBoneGetterSetterMode::LocalSpace:
-							{
-								Transform = Hierarchy->GetInitialLocalTransform(CachedIndex);
-								break;
-							}
-							default:
-							{
-								break;
-							}
-						}
+						Transform = Hierarchy->GetInitialLocalTransform(CachedIndex);
+						break;
 					}
-					else
+					default:
 					{
-						switch (Space)
-						{
-							case EBoneGetterSetterMode::GlobalSpace:
-							{
-								Transform = Hierarchy->GetGlobalTransform(CachedIndex);
-								break;
-							}
-							case EBoneGetterSetterMode::LocalSpace:
-							{
-								Transform = Hierarchy->GetLocalTransform(CachedIndex);
-								break;
-							}
-							default:
-							{
-								break;
-							}
-						}
+						break;
 					}
 				}
 			}
-			default:
+			else
 			{
-				break;
+				switch (Space)
+				{
+					case EBoneGetterSetterMode::GlobalSpace:
+					{
+						Transform = Hierarchy->GetGlobalTransform(CachedIndex);
+						break;
+					}
+					case EBoneGetterSetterMode::LocalSpace:
+					{
+						Transform = Hierarchy->GetLocalTransform(CachedIndex);
+						break;
+					}
+					default:
+					{
+						break;
+					}
+				}
 			}
 		}
 	}
diff --git a/Engine/Plugins/Animation/ControlRig/Source/ControlRig/Private/Units/Hierarchy/RigUnit_ParentSwitchConstraint.cpp b/Engine/Plugins/Animation/ControlRig/Source/ControlRig/Private/Units/Hierarchy/RigUnit_ParentSwitchConstraint.cpp
--- a/Engine/Plugins/Animation/ControlRig/Source/ControlRig/Private/Units/Hierarchy/RigUnit_ParentSwitchConstraint.cpp
+++ b/Engine/Plugins/Animation/ControlRig/Source/ControlRig/Private/Units/Hierarchy/RigUnit_ParentSwitchConstraint.cpp
@@ -31,73 +31,57 @@ FRigUnit_ParentSwitchConstraintArray_Execute()
 
 	if (URigHierarchy* Hierarchy = ExecuteContext.Hierarchy)
 	{
-		switch (ExecuteContext.UnitContext.State)
+		if (!CachedSubject.UpdateCache(Subject, Hierarchy))
 		{
-			case EControlRigState::Init:
+			UE_CONTROLRIG_RIGUNIT_REPORT_WARNING(TEXT("Subject '%s' is not valid."), *Subject.ToString());
+			return;
+		}
+		if (ParentIndex < 0 || ParentIndex >= Parents.Num())
+		{
+			UE_CONTROLRIG_RIGUNIT_REPORT_WARNING(TEXT("Parent Index is out of bounds."));
+			return;
+		}
+
+		if (!CachedParent.IsValid())
+		{
+			if (!CachedParent.UpdateCache(Parents[ParentIndex], Hierarchy))
+			{
+				UE_CONTROLRIG_RIGUNIT_REPORT_WARNING(TEXT("Parent '%s' is not valid."), *Parents[ParentIndex].ToString());
+				return;
+			}
+
+			FTransform ParentTransform = Hierarchy->GetGlobalTransform(CachedParent);
+			RelativeOffset = InitialGlobalTransform.GetRelativeTransform(ParentTransform);
+		}
+
+		FTransform ParentTransform = Hierarchy->GetGlobalTransform(CachedParent);
+		Transform = RelativeOffset * ParentTransform;
+
+		if (Weight > SMALL_NUMBER)
+		{
+			if (Weight < 1.0f - SMALL_NUMBER)
 			{
-				CachedSubject.Reset();
-				CachedParent.Reset();
-				break;
+				FTransform WeightedTransform = Hierarchy->GetGlobalTransform(CachedSubject);
+				WeightedTransform = FControlRigMathLibrary::LerpTransform(WeightedTransform, Transform, Weight);
+				Hierarchy->SetGlobalTransform(CachedSubject, WeightedTransform);
 			}
-			case EControlRigState::Update:
+			else
 			{
-				if (!CachedSubject.UpdateCache(Subject, Hierarchy))
-				{
-					UE_CONTROLRIG_RIGUNIT_REPORT_WARNING(TEXT("Subject '%s' is not valid."), *Subject.ToString());
-					return;
-				}
-				if (ParentIndex < 0 || ParentIndex >= Parents.Num())
-				{
-					UE_CONTROLRIG_RIGUNIT_REPORT_WARNING(TEXT("Parent Index is out of bounds."));
-					return;
-				}
-
-				if (!CachedParent.IsValid())
-				{
-					if (!CachedParent.UpdateCache(Parents[ParentIndex], Hierarchy))
-					{
-						UE_CONTROLRIG_RIGUNIT_REPORT_WARNING(TEXT("Parent '%s' is not valid."), *Parents[ParentIndex].ToString());
-						return;
-					}
-
-					FTransform ParentTransform = Hierarchy->GetGlobalTransform(CachedParent);
-					RelativeOffset = InitialGlobalTransform.GetRelativeTransform(ParentTransform);
-				}
-
-				FTransform ParentTransform = Hierarchy->GetGlobalTransform(CachedParent);
-				Transform = RelativeOffset * ParentTransform;
-
-				if (Weight > SMALL_NUMBER)
-				{
-					if (Weight < 1.0f - SMALL_NUMBER)
-					{
-						FTransform WeightedTransform = Hierarchy->GetGlobalTransform(CachedSubject);
-						WeightedTransform = FControlRigMathLibrary::LerpTransform(WeightedTransform, Transform, Weight);
-						Hierarchy->SetGlobalTransform(CachedSubject, WeightedTransform);
-					}
-					else
-					{
-						Hierarchy->SetGlobalTransform(CachedSubject, Transform);
-					}
-				}
-
-				if (CachedParent != Parents[ParentIndex])
-				{
-					if (!CachedParent.UpdateCache(Parents[ParentIndex], Hierarchy))
-					{
-						UE_CONTROLRIG_RIGUNIT_REPORT_WARNING(TEXT("Parent '%s' is not valid."), *Parents[ParentIndex].ToString());
-						return;
-					}
-
-					ParentTransform = Hierarchy->GetGlobalTransform(CachedParent);
-					RelativeOffset = Transform.GetRelativeTransform(ParentTransform);
-					Switched = true;
-				}
+				Hierarchy->SetGlobalTransform(CachedSubject, Transform);
 			}
-			default:
+		}
+
+		if (CachedParent != Parents[ParentIndex])
+		{
+			if (!CachedParent.UpdateCache(Parents[ParentIndex], Hierarchy))
 			{
-				break;
+				UE_CONTROLRIG_RIGUNIT_REPORT_WARNING(TEXT("Parent '%s' is not valid."), *Parents[ParentIndex].ToString());
+				return;
 			}
+
+			ParentTransform = Hierarchy->GetGlobalTransform(CachedParent);
+			RelativeOffset = Transform.GetRelativeTransform(ParentTransform);
+			Switched = true;
 		}
 	}
 }
diff --git a/Engine/Plugins/Animation/ControlRig/Source/ControlRig/Private/Units/Hierarchy/RigUnit_PropagateTransform.cpp b/Engine/Plugins/Animation/ControlRig/Source/ControlRig/Private/Units/Hierarchy/RigUnit_PropagateTransform.cpp
--- a/Engine/Plugins/Animation/ControlRig/Source/ControlRig/Private/Units/Hierarchy/RigUnit_PropagateTransform.cpp
+++ b/Engine/Plugins/Animation/ControlRig/Source/ControlRig/Private/Units/Hierarchy/RigUnit_PropagateTransform.cpp
@@ -22,47 +22,33 @@ FRigUnit_PropagateTransform_Execute()
     	/*
     	 * This node doesn't do anything anymore now that the hierarchy is lazy
     	 *
-		switch (ExecuteContext.UnitContext.State)
+		if (!CachedIndex.UpdateCache(Item, Hierarchy))
 		{
-			case EControlRigState::Init:
+			UE_CONTROLRIG_RIGUNIT_REPORT_WARNING(TEXT("Item '%s' is not valid."), *Item.ToString());
+		}
+		else
+		{
+			FRigBoneHierarchy& Bones = Hierarchy->BoneHierarchy;
+			int32 BoneIndex = CachedIndex.GetIndex();
+
+			if (bRecomputeGlobal)
 			{
-				CachedIndex.Reset();
+				Bones.RecalculateGlobalTransform(BoneIndex);
 			}
-			case EControlRigState::Update:
+			if (bApplyToChildren)
 			{
-				if (!CachedIndex.UpdateCache(Item, Hierarchy))
+				if (bRecursive)
 				{
-					UE_CONTROLRIG_RIGUNIT_REPORT_WARNING(TEXT("Item '%s' is not valid."), *Item.ToString());
+					Bones.PropagateTransform(BoneIndex);
 				}
 				else
 				{
-					FRigBoneHierarchy& Bones = Hierarchy->BoneHierarchy;
-					int32 BoneIndex = CachedIndex.GetIndex();
-
-					if (bRecomputeGlobal)
+					for (int32 Dependent : Bones[BoneIndex].Dependents)
 					{
-						Bones.RecalculateGlobalTransform(BoneIndex);
-					}
-					if (bApplyToChildren)
-					{
-						if (bRecursive)
-						{
-							Bones.PropagateTransform(BoneIndex);
-						}
-						else
-						{
-							for (int32 Dependent : Bones[BoneIndex].Dependents)
-							{
-								Bones.RecalculateGlobalTransform(Dependent);
-							}
-						}
+						Bones.RecalculateGlobalTransform(Dependent);
 					}
 				}
 			}
-			default:
-			{
-				break;
-			}
 		}
 		*/
 	}
diff --git a/Engine/Plugins/Animation/ControlRig/Source/ControlRig/Private/Units/Hierarchy/RigUnit_SendEvent.cpp b/Engine/Plugins/Animation/ControlRig/Source/ControlRig/Private/Units/Hierarchy/RigUnit_SendEvent.cpp
--- a/Engine/Plugins/Animation/ControlRig/Source/ControlRig/Private/Units/Hierarchy/RigUnit_SendEvent.cpp
+++ b/Engine/Plugins/Animation/ControlRig/Source/ControlRig/Private/Units/Hierarchy/RigUnit_SendEvent.cpp
@@ -22,27 +22,12 @@ FRigUnit_SendEvent_Execute()
 	URigHierarchy* Hierarchy = ExecuteContext.Hierarchy;
 	if (Hierarchy)
 	{
-		switch (ExecuteContext.UnitContext.State)
-		{
-			case EControlRigState::Init:
-			{
-				break;
-			}
-			case EControlRigState::Update:
-			{
-				FRigEventContext EventContext;
-				EventContext.Key = Item;
-				EventContext.Event = Event;
-				EventContext.SourceEventName = ExecuteContext.GetEventName();
-				EventContext.LocalTime = ExecuteContext.UnitContext.AbsoluteTime + OffsetInSeconds;
-				Hierarchy->SendEvent(EventContext, false /* async */); //needs to be false for sequencer keying to work
-				break;
-			}
-			default:
-			{
-				break;
-			}
-		}
+		FRigEventContext EventContext;
+		EventContext.Key = Item;
+		EventContext.Event = Event;
+		EventContext.SourceEventName = ExecuteContext.GetEventName();
+		EventContext.LocalTime = ExecuteContext.GetAbsoluteTime() + OffsetInSeconds;
+		Hierarchy->SendEvent(EventContext, false /* async */); //needs to be false for sequencer keying to work
 	}
 }
 
diff --git a/Engine/Plugins/Animation/ControlRig/Source/ControlRig/Private/Units/Hierarchy/RigUnit_SetBoneInitialTransform.cpp b/Engine/Plugins/Animation/ControlRig/Source/ControlRig/Private/Units/Hierarchy/RigUnit_SetBoneInitialTransform.cpp
--- a/Engine/Plugins/Animation/ControlRig/Source/ControlRig/Private/Units/Hierarchy/RigUnit_SetBoneInitialTransform.cpp
+++ b/Engine/Plugins/Animation/ControlRig/Source/ControlRig/Private/Units/Hierarchy/RigUnit_SetBoneInitialTransform.cpp
@@ -14,35 +14,20 @@ FRigUnit_SetBoneInitialTransform_Execute()
 	URigHierarchy* Hierarchy = ExecuteContext.Hierarchy;
 	if (Hierarchy)
 	{
-		switch (ExecuteContext.UnitContext.State)
+		const FRigElementKey Key(Bone, ERigElementType::Bone);
+		if (!CachedBone.UpdateCache(Key, Hierarchy))
 		{
-			case EControlRigState::Init:
-			{
-				CachedBone.Reset();
-				break;
-			}
-			case EControlRigState::Update:
-			{
-				const FRigElementKey Key(Bone, ERigElementType::Bone);
-				if (!CachedBone.UpdateCache(Key, Hierarchy))
-				{
-					UE_CONTROLRIG_RIGUNIT_REPORT_WARNING(TEXT("Bone '%s' is not valid."), *Bone.ToString());
-					return;
-				}
-
-				if (Space == EBoneGetterSetterMode::LocalSpace)
-				{
-					Hierarchy->SetInitialLocalTransform(CachedBone, Transform);
-				}
-				else
-				{
-					Hierarchy->SetInitialGlobalTransform(CachedBone, Transform);
-				}
-			}
-			default:
-			{
-				break;
-			}
+			UE_CONTROLRIG_RIGUNIT_REPORT_WARNING(TEXT("Bone '%s' is not valid."), *Bone.ToString());
+			return;
+		}
+
+		if (Space == EBoneGetterSetterMode::LocalSpace)
+		{
+			Hierarchy->SetInitialLocalTransform(CachedBone, Transform);
+		}
+		else
+		{
+			Hierarchy->SetInitialGlobalTransform(CachedBone, Transform);
 		}
 	}
 }
diff --git a/Engine/Plugins/Animation/ControlRig/Source/ControlRig/Private/Units/Hierarchy/RigUnit_SetBoneRotation.cpp b/Engine/Plugins/Animation/ControlRig/Source/ControlRig/Private/Units/Hierarchy/RigUnit_SetBoneRotation.cpp
--- a/Engine/Plugins/Animation/ControlRig/Source/ControlRig/Private/Units/Hierarchy/RigUnit_SetBoneRotation.cpp
+++ b/Engine/Plugins/Animation/ControlRig/Source/ControlRig/Private/Units/Hierarchy/RigUnit_SetBoneRotation.cpp
@@ -12,68 +12,53 @@ FRigUnit_SetBoneRotation_Execute()
 	URigHierarchy* Hierarchy = ExecuteContext.Hierarchy;
 	if (Hierarchy)
 	{
-		switch (ExecuteContext.UnitContext.State)
+		const FRigElementKey Key(Bone, ERigElementType::Bone);
+		if (!CachedBone.UpdateCache(Key, Hierarchy))
 		{
-			case EControlRigState::Init:
-			{
-				CachedBone.Reset();
-				// fall through to update
-			}
-			case EControlRigState::Update:
+			UE_CONTROLRIG_RIGUNIT_REPORT_WARNING(TEXT("Bone '%s' is not valid."), *Bone.ToString());
+		}
+		else
+		{
+			switch (Space)
 			{
-				const FRigElementKey Key(Bone, ERigElementType::Bone);
-				if (!CachedBone.UpdateCache(Key, Hierarchy))
+				case EBoneGetterSetterMode::GlobalSpace:
 				{
-					UE_CONTROLRIG_RIGUNIT_REPORT_WARNING(TEXT("Bone '%s' is not valid."), *Bone.ToString());
+					FTransform Transform = Hierarchy->GetGlobalTransform(CachedBone);
+
+					if (FMath::IsNearlyEqual(Weight, 1.f))
+					{
+						Transform.SetRotation(Rotation);
+					}
+					else
+					{
+						float T = FMath::Clamp<float>(Weight, 0.f, 1.f);
+						Transform.SetRotation(FQuat::Slerp(Transform.GetRotation(), Rotation, T));
+					}
+
+					Hierarchy->SetGlobalTransform(CachedBone, Transform, bPropagateToChildren);
+					break;
 				}
-				else
+				case EBoneGetterSetterMode::LocalSpace:
 				{
-					switch (Space)
+					FTransform Transform = Hierarchy->GetLocalTransform(CachedBone);
+
+					if (FMath::IsNearlyEqual(Weight, 1.f))
 					{
-						case EBoneGetterSetterMode::GlobalSpace:
-						{
-							FTransform Transform = Hierarchy->GetGlobalTransform(CachedBone);
-
-							if (FMath::IsNearlyEqual(Weight, 1.f))
-							{
-								Transform.SetRotation(Rotation);
-							}
-							else
-							{
-								float T = FMath::Clamp<float>(Weight, 0.f, 1.f);
-								Transform.SetRotation(FQuat::Slerp(Transform.GetRotation(), Rotation, T));
-							}
-
-							Hierarchy->SetGlobalTransform(CachedBone, Transform, bPropagateToChildren);
-							break;
-						}
-						case EBoneGetterSetterMode::LocalSpace:
-						{
-							FTransform Transform = Hierarchy->GetLocalTransform(CachedBone);
-
-							if (FMath::IsNearlyEqual(Weight, 1.f))
-							{
-								Transform.SetRotation(Rotation);
-							}
-							else
-							{
-								float T = FMath::Clamp<float>(Weight, 0.f, 1.f);
-								Transform.SetRotation(FQuat::Slerp(Transform.GetRotation(), Rotation, T));
-							}
-
-							Hierarchy->SetLocalTransform(CachedBone, Transform, bPropagateToChildren);
-							break;
-						}
-						default:
-						{
-							break;
-						}
+						Transform.SetRotation(Rotation);
 					}
+					else
+					{
+						float T = FMath::Clamp<float>(Weight, 0.f, 1.f);
+						Transform.SetRotation(FQuat::Slerp(Transform.GetRotation(), Rotation, T));
+					}
+
+					Hierarchy->SetLocalTransform(CachedBone, Transform, bPropagateToChildren);
+					break;
+				}
+				default:
+				{
+					break;
 				}
-			}
-			default:
-			{
-				break;
 			}
 		}
 	}
@@ -109,21 +94,21 @@ IMPLEMENT_RIGUNIT_AUTOMATION_TEST(FRigUnit_SetBoneRotation)
 	Unit.Space = EBoneGetterSetterMode::GlobalSpace;
 	Unit.Rotation = FQuat(FVector(-1.f, 0.f, 0.f), 0.25f);
 	Unit.bPropagateToChildren = false;
-	InitAndExecute();
+	Execute();
 	AddErrorIfFalse(FMath::IsNearlyEqual(Hierarchy->GetGlobalTransform(0).GetRotation().GetAngle(), 0.25f, 0.001f), FString::Printf(TEXT("unexpected angle %.04f"), Hierarchy->GetGlobalTransform(0).GetRotation().GetAngle()));
 	AddErrorIfFalse(FMath::IsNearlyEqual(Hierarchy->GetGlobalTransform(1).GetRotation().GetAngle(), 0.5f, 0.001f), FString::Printf(TEXT("unexpected angle %.04f"), Hierarchy->GetGlobalTransform(1).GetRotation().GetAngle()));
 	AddErrorIfFalse(FMath::IsNearlyEqual(Hierarchy->GetGlobalTransform(2).GetRotation().GetAngle(), 0.7f, 0.001f), FString::Printf(TEXT("unexpected angle %.04f"), Hierarchy->GetGlobalTransform(2).GetRotation().GetAngle()));
 
 	Hierarchy->ResetPoseToInitial(ERigElementType::Bone);
 	Unit.Space = EBoneGetterSetterMode::LocalSpace;
-	InitAndExecute();
+	Execute();
 	AddErrorIfFalse(FMath::IsNearlyEqual(Hierarchy->GetGlobalTransform(0).GetRotation().GetAngle(), 0.25f, 0.001f), FString::Printf(TEXT("unexpected angle %.04f"), Hierarchy->GetGlobalTransform(0).GetRotation().GetAngle()));
 	AddErrorIfFalse(FMath::IsNearlyEqual(Hierarchy->GetGlobalTransform(1).GetRotation().GetAngle(), 0.5f, 0.001f), FString::Printf(TEXT("unexpected angle %.04f"), Hierarchy->GetGlobalTransform(1).GetRotation().GetAngle()));
 	AddErrorIfFalse(FMath::IsNearlyEqual(Hierarchy->GetGlobalTransform(2).GetRotation().GetAngle(), 0.7f, 0.001f), FString::Printf(TEXT("unexpected angle %.04f"), Hierarchy->GetGlobalTransform(2).GetRotation().GetAngle()));
 
 	Hierarchy->ResetPoseToInitial(ERigElementType::Bone);
 	Unit.bPropagateToChildren = true;
-	InitAndExecute();
+	Execute();
 	AddErrorIfFalse(FMath::IsNearlyEqual(Hierarchy->GetGlobalTransform(0).GetRotation().GetAngle(), 0.25f, 0.001f), FString::Printf(TEXT("unexpected angle %.04f"), Hierarchy->GetGlobalTransform(0).GetRotation().GetAngle()));
 	AddErrorIfFalse(FMath::IsNearlyEqual(Hierarchy->GetGlobalTransform(1).GetRotation().GetAngle(), 0.65f, 0.001f), FString::Printf(TEXT("unexpected angle %.04f"), Hierarchy->GetGlobalTransform(1).GetRotation().GetAngle()));
 	AddErrorIfFalse(FMath::IsNearlyEqual(Hierarchy->GetGlobalTransform(2).GetRotation().GetAngle(), 0.85f, 0.001f), FString::Printf(TEXT("unexpected angle %.04f"), Hierarchy->GetGlobalTransform(2).GetRotation().GetAngle()));
@@ -132,21 +117,21 @@ IMPLEMENT_RIGUNIT_AUTOMATION_TEST(FRigUnit_SetBoneRotation)
 	Unit.Bone = TEXT("BoneA");
 	Unit.Space = EBoneGetterSetterMode::GlobalSpace;
 	Unit.bPropagateToChildren = false;
-	InitAndExecute();
+	Execute();
 	AddErrorIfFalse(FMath::IsNearlyEqual(Hierarchy->GetGlobalTransform(0).GetRotation().GetAngle(), 0.1f, 0.001f), FString::Printf(TEXT("unexpected angle %.04f"), Hierarchy->GetGlobalTransform(0).GetRotation().GetAngle()));
 	AddErrorIfFalse(FMath::IsNearlyEqual(Hierarchy->GetGlobalTransform(1).GetRotation().GetAngle(), 0.25f, 0.001f), FString::Printf(TEXT("unexpected angle %.04f"), Hierarchy->GetGlobalTransform(1).GetRotation().GetAngle()));
 	AddErrorIfFalse(FMath::IsNearlyEqual(Hierarchy->GetGlobalTransform(2).GetRotation().GetAngle(), 0.7f, 0.001f), FString::Printf(TEXT("unexpected angle %.04f"), Hierarchy->GetGlobalTransform(2).GetRotation().GetAngle()));
 
 	Hierarchy->ResetPoseToInitial(ERigElementType::Bone);
 	Unit.Space = EBoneGetterSetterMode::LocalSpace;
-	InitAndExecute();
+	Execute();
 	AddErrorIfFalse(FMath::IsNearlyEqual(Hierarchy->GetGlobalTransform(0).GetRotation().GetAngle(), 0.1f, 0.001f), FString::Printf(TEXT("unexpected angle %.04f"), Hierarchy->GetGlobalTransform(0).GetRotation().GetAngle()));
 	AddErrorIfFalse(FMath::IsNearlyEqual(Hierarchy->GetGlobalTransform(1).GetRotation().GetAngle(), 0.35f, 0.001f), FString::Printf(TEXT("unexpected angle %.04f"), Hierarchy->GetGlobalTransform(1).GetRotation().GetAngle()));
 	AddErrorIfFalse(FMath::IsNearlyEqual(Hierarchy->GetGlobalTransform(2).GetRotation().GetAngle(), 0.7f, 0.001f), FString::Printf(TEXT("unexpected angle %.04f"), Hierarchy->GetGlobalTransform(2).GetRotation().GetAngle()));
 
 	Hierarchy->ResetPoseToInitial(ERigElementType::Bone);
 	Unit.bPropagateToChildren = true;
-	InitAndExecute();
+	Execute();
 	AddErrorIfFalse(FMath::IsNearlyEqual(Hierarchy->GetGlobalTransform(0).GetRotation().GetAngle(), 0.1f, 0.001f), FString::Printf(TEXT("unexpected angle %.04f"), Hierarchy->GetGlobalTransform(0).GetRotation().GetAngle()));
 	AddErrorIfFalse(FMath::IsNearlyEqual(Hierarchy->GetGlobalTransform(1).GetRotation().GetAngle(), 0.35f, 0.001f), FString::Printf(TEXT("unexpected angle %.04f"), Hierarchy->GetGlobalTransform(1).GetRotation().GetAngle()));
 	AddErrorIfFalse(FMath::IsNearlyEqual(Hierarchy->GetGlobalTransform(2).GetRotation().GetAngle(), 0.55f, 0.001f), FString::Printf(TEXT("unexpected angle %.04f"), Hierarchy->GetGlobalTransform(2).GetRotation().GetAngle()));
diff --git a/Engine/Plugins/Animation/ControlRig/Source/ControlRig/Private/Units/Hierarchy/RigUnit_SetBoneTransform.cpp b/Engine/Plugins/Animation/ControlRig/Source/ControlRig/Private/Units/Hierarchy/RigUnit_SetBoneTransform.cpp
--- a/Engine/Plugins/Animation/ControlRig/Source/ControlRig/Private/Units/Hierarchy/RigUnit_SetBoneTransform.cpp
+++ b/Engine/Plugins/Animation/ControlRig/Source/ControlRig/Private/Units/Hierarchy/RigUnit_SetBoneTransform.cpp
@@ -13,64 +13,49 @@ FRigUnit_SetBoneTransform_Execute()
 	URigHierarchy* Hierarchy = ExecuteContext.Hierarchy;
 	if (Hierarchy)
 	{
-		switch (ExecuteContext.UnitContext.State)
+		const FRigElementKey Key(Bone, ERigElementType::Bone);
+		if (!CachedBone.UpdateCache(Key, Hierarchy))
 		{
-			case EControlRigState::Init:
-			{
-				CachedBone.Reset();
-				// fall through to update
-			}
-			case EControlRigState::Update:
+			UE_CONTROLRIG_RIGUNIT_REPORT_WARNING(TEXT("Bone '%s' is not valid."), *Bone.ToString());
+		}
+		else
+		{
+			switch (Space)
 			{
-				const FRigElementKey Key(Bone, ERigElementType::Bone);
-				if (!CachedBone.UpdateCache(Key, Hierarchy))
+				case EBoneGetterSetterMode::GlobalSpace:
 				{
-					UE_CONTROLRIG_RIGUNIT_REPORT_WARNING(TEXT("Bone '%s' is not valid."), *Bone.ToString());
+					if (FMath::IsNearlyEqual(Weight, 1.f))
+					{
+						Result = Transform;
+					}
+					else
+					{
+						float T = FMath::Clamp<float>(Weight, 0.f, 1.f);
+						const FTransform PreviousTransform = Hierarchy->GetGlobalTransform(CachedBone);
+						Result = FControlRigMathLibrary::LerpTransform(PreviousTransform, Transform, T);
+					}
+					Hierarchy->SetGlobalTransform(CachedBone, Result, bPropagateToChildren);
+					break;
 				}
-				else
+				case EBoneGetterSetterMode::LocalSpace:
 				{
-					switch (Space)
+					if (FMath::IsNearlyEqual(Weight, 1.f))
+					{
+						Result = Transform;
+					}
+					else
 					{
-						case EBoneGetterSetterMode::GlobalSpace:
-						{
-							if (FMath::IsNearlyEqual(Weight, 1.f))
-							{
-								Result = Transform;
-							}
-							else
-							{
-								float T = FMath::Clamp<float>(Weight, 0.f, 1.f);
-								const FTransform PreviousTransform = Hierarchy->GetGlobalTransform(CachedBone);
-								Result = FControlRigMathLibrary::LerpTransform(PreviousTransform, Transform, T);
-							}
-							Hierarchy->SetGlobalTransform(CachedBone, Result, bPropagateToChildren);
-							break;
-						}
-						case EBoneGetterSetterMode::LocalSpace:
-						{
-							if (FMath::IsNearlyEqual(Weight, 1.f))
-							{
-								Result = Transform;
-							}
-							else
-							{
-								float T = FMath::Clamp<float>(Weight, 0.f, 1.f);
-								const FTransform PreviousTransform = Hierarchy->GetLocalTransform(CachedBone);
-								Result = FControlRigMathLibrary::LerpTransform(PreviousTransform, Transform, T);
-							}
-							Hierarchy->SetLocalTransform(CachedBone, Result, bPropagateToChildren);
-							break;
-						}
-						default:
-						{
-							break;
-						}
+						float T = FMath::Clamp<float>(Weight, 0.f, 1.f);
+						const FTransform PreviousTransform = Hierarchy->GetLocalTransform(CachedBone);
+						Result = FControlRigMathLibrary::LerpTransform(PreviousTransform, Transform, T);
 					}
+					Hierarchy->SetLocalTransform(CachedBone, Result, bPropagateToChildren);
+					break;
+				}
+				default:
+				{
+					break;
 				}
-			}
-			default:
-			{
-				break;
 			}
 		}
 	}
@@ -105,21 +90,21 @@ IMPLEMENT_RIGUNIT_AUTOMATION_TEST(FRigUnit_SetBoneTransform)
 	Unit.Space = EBoneGetterSetterMode::GlobalSpace;
 	Unit.Transform = FTransform(FVector(0.f, 0.f, 7.f));
 	Unit.bPropagateToChildren = false;
-	InitAndExecute();
+	Execute();
 	AddErrorIfFalse(Hierarchy->GetGlobalTransform(0).GetTranslation().Equals(FVector(0.f, 0.f, 7.f)), TEXT("unexpected transform"));
 	AddErrorIfFalse(Hierarchy->GetGlobalTransform(1).GetTranslation().Equals(FVector(1.f, 2.f, 3.f)), TEXT("unexpected transform"));
 	AddErrorIfFalse(Hierarchy->GetGlobalTransform(2).GetTranslation().Equals(FVector(1.f, 5.f, 3.f)), TEXT("unexpected transform"));
 
 	Hierarchy->ResetPoseToInitial(ERigElementType::Bone);
 	Unit.Space = EBoneGetterSetterMode::LocalSpace;
-	InitAndExecute();
+	Execute();
 	AddErrorIfFalse(Hierarchy->GetGlobalTransform(0).GetTranslation().Equals(FVector(0.f, 0.f, 7.f)), TEXT("unexpected transform"));
 	AddErrorIfFalse(Hierarchy->GetGlobalTransform(1).GetTranslation().Equals(FVector(1.f, 2.f, 3.f)), TEXT("unexpected transform"));
 	AddErrorIfFalse(Hierarchy->GetGlobalTransform(2).GetTranslation().Equals(FVector(1.f, 5.f, 3.f)), TEXT("unexpected transform"));
 
 	Hierarchy->ResetPoseToInitial(ERigElementType::Bone);
 	Unit.bPropagateToChildren = true;
-	InitAndExecute();
+	Execute();
 	AddErrorIfFalse(Hierarchy->GetGlobalTransform(0).GetTranslation().Equals(FVector(0.f, 0.f, 7.f)), TEXT("unexpected transform"));
 	AddErrorIfFalse(Hierarchy->GetGlobalTransform(1).GetTranslation().Equals(FVector(0.f, 2.f, 10.f)), TEXT("unexpected transform"));
 	AddErrorIfFalse(Hierarchy->GetGlobalTransform(2).GetTranslation().Equals(FVector(0.f, 5.f, 10.f)), TEXT("unexpected transform"));
@@ -128,21 +113,21 @@ IMPLEMENT_RIGUNIT_AUTOMATION_TEST(FRigUnit_SetBoneTransform)
 	Unit.Bone = TEXT("BoneA");
 	Unit.Space = EBoneGetterSetterMode::GlobalSpace;
 	Unit.bPropagateToChildren = false;
-	InitAndExecute();
+	Execute();
 	AddErrorIfFalse(Hierarchy->GetGlobalTransform(0).GetTranslation().Equals(FVector(1.f, 0.f, 0.f)), TEXT("unexpected transform"));
 	AddErrorIfFalse(Hierarchy->GetGlobalTransform(1).GetTranslation().Equals(FVector(0.f, 0.f, 7.f)), TEXT("unexpected transform"));
 	AddErrorIfFalse(Hierarchy->GetGlobalTransform(2).GetTranslation().Equals(FVector(1.f, 5.f, 3.f)), TEXT("unexpected transform"));
 
 	Hierarchy->ResetPoseToInitial(ERigElementType::Bone);
 	Unit.Space = EBoneGetterSetterMode::LocalSpace;
-	InitAndExecute();
+	Execute();
 	AddErrorIfFalse(Hierarchy->GetGlobalTransform(0).GetTranslation().Equals(FVector(1.f, 0.f, 0.f)), TEXT("unexpected transform"));
 	AddErrorIfFalse(Hierarchy->GetGlobalTransform(1).GetTranslation().Equals(FVector(1.f, 0.f, 7.f)), TEXT("unexpected transform"));
 	AddErrorIfFalse(Hierarchy->GetGlobalTransform(2).GetTranslation().Equals(FVector(1.f, 5.f, 3.f)), TEXT("unexpected transform"));
 
 	Hierarchy->ResetPoseToInitial(ERigElementType::Bone);
 	Unit.bPropagateToChildren = true;
-	InitAndExecute();
+	Execute();
 	AddErrorIfFalse(Hierarchy->GetGlobalTransform(0).GetTranslation().Equals(FVector(1.f, 0.f, 0.f)), TEXT("unexpected transform"));
 	AddErrorIfFalse(Hierarchy->GetGlobalTransform(1).GetTranslation().Equals(FVector(1.f, 0.f, 7.f)), TEXT("unexpected transform"));
 	AddErrorIfFalse(Hierarchy->GetGlobalTransform(2).GetTranslation().Equals(FVector(1.f, 3.f, 7.f)), TEXT("unexpected transform"));
diff --git a/Engine/Plugins/Animation/ControlRig/Source/ControlRig/Private/Units/Hierarchy/RigUnit_SetBoneTranslation.cpp b/Engine/Plugins/Animation/ControlRig/Source/ControlRig/Private/Units/Hierarchy/RigUnit_SetBoneTranslation.cpp
--- a/Engine/Plugins/Animation/ControlRig/Source/ControlRig/Private/Units/Hierarchy/RigUnit_SetBoneTranslation.cpp
+++ b/Engine/Plugins/Animation/ControlRig/Source/ControlRig/Private/Units/Hierarchy/RigUnit_SetBoneTranslation.cpp
@@ -12,68 +12,53 @@ FRigUnit_SetBoneTranslation_Execute()
 	URigHierarchy* Hierarchy = ExecuteContext.Hierarchy;
 	if (Hierarchy)
 	{
-		switch (ExecuteContext.UnitContext.State)
+		const FRigElementKey Key(Bone, ERigElementType::Bone); 
+		if (!CachedBone.UpdateCache(Key, Hierarchy))
 		{
-			case EControlRigState::Init:
-			{
-				CachedBone.Reset();
-				// fall through to update
-			}
-			case EControlRigState::Update:
+			UE_CONTROLRIG_RIGUNIT_REPORT_WARNING(TEXT("Bone '%s' is not valid."), *Bone.ToString());
+		}
+		else
+		{
+			switch (Space)
 			{
-				const FRigElementKey Key(Bone, ERigElementType::Bone); 
-				if (!CachedBone.UpdateCache(Key, Hierarchy))
-				{
-					UE_CONTROLRIG_RIGUNIT_REPORT_WARNING(TEXT("Bone '%s' is not valid."), *Bone.ToString());
-				}
-				else
+				case EBoneGetterSetterMode::GlobalSpace:
 				{
-					switch (Space)
-					{
-						case EBoneGetterSetterMode::GlobalSpace:
-						{
-							FTransform Transform = Hierarchy->GetGlobalTransform(CachedBone);
-
-							if (FMath::IsNearlyEqual(Weight, 1.f))
-							{
-								Transform.SetTranslation(Translation);
-							}
-							else
-							{
-								float T = FMath::Clamp<float>(Weight, 0.f, 1.f);
-								Transform.SetTranslation(FMath::Lerp<FVector>(Transform.GetTranslation(), Translation, T));
-							}
+					FTransform Transform = Hierarchy->GetGlobalTransform(CachedBone);
 
-							Hierarchy->SetGlobalTransform(CachedBone, Transform, bPropagateToChildren);
-							break;
-						}
-						case EBoneGetterSetterMode::LocalSpace:
-						{
-							FTransform Transform = Hierarchy->GetLocalTransform(CachedBone);
+					if (FMath::IsNearlyEqual(Weight, 1.f))
+					{
+						Transform.SetTranslation(Translation);
+					}
+					else
+					{
+						float T = FMath::Clamp<float>(Weight, 0.f, 1.f);
+						Transform.SetTranslation(FMath::Lerp<FVector>(Transform.GetTranslation(), Translation, T));
+					}
 
-							if (FMath::IsNearlyEqual(Weight, 1.f))
-							{
-								Transform.SetTranslation(Translation);
-							}
-							else
-							{
-								float T = FMath::Clamp<float>(Weight, 0.f, 1.f);
-								Transform.SetTranslation(FMath::Lerp<FVector>(Transform.GetTranslation(), Translation, T));
-							}
+					Hierarchy->SetGlobalTransform(CachedBone, Transform, bPropagateToChildren);
+					break;
+				}
+				case EBoneGetterSetterMode::LocalSpace:
+				{
+					FTransform Transform = Hierarchy->GetLocalTransform(CachedBone);
 
-							Hierarchy->SetLocalTransform(CachedBone, Transform, bPropagateToChildren);
-							break;
-						}
-						default:
-						{
-							break;
-						}
+					if (FMath::IsNearlyEqual(Weight, 1.f))
+					{
+						Transform.SetTranslation(Translation);
+					}
+					else
+					{
+						float T = FMath::Clamp<float>(Weight, 0.f, 1.f);
+						Transform.SetTranslation(FMath::Lerp<FVector>(Transform.GetTranslation(), Translation, T));
 					}
+
+					Hierarchy->SetLocalTransform(CachedBone, Transform, bPropagateToChildren);
+					break;
+				}
+				default:
+				{
+					break;
 				}
-			}
-			default:
-			{
-				break;
 			}
 		}
 	}
diff --git a/Engine/Plugins/Animation/ControlRig/Source/ControlRig/Private/Units/Hierarchy/RigUnit_SetControlColor.cpp b/Engine/Plugins/Animation/ControlRig/Source/ControlRig/Private/Units/Hierarchy/RigUnit_SetControlColor.cpp
--- a/Engine/Plugins/Animation/ControlRig/Source/ControlRig/Private/Units/Hierarchy/RigUnit_SetControlColor.cpp
+++ b/Engine/Plugins/Animation/ControlRig/Source/ControlRig/Private/Units/Hierarchy/RigUnit_SetControlColor.cpp
@@ -14,26 +14,10 @@ FRigUnit_GetControlColor_Execute()
 	const URigHierarchy* Hierarchy = ExecuteContext.Hierarchy;
 	if (Hierarchy)
 	{
-		switch (ExecuteContext.UnitContext.State)
+		if (CachedControlIndex.UpdateCache(FRigElementKey(Control, ERigElementType::Control), Hierarchy))
 		{
-			case EControlRigState::Init:
-			{
-				CachedControlIndex.Reset();
-				break;
-			}
-			case EControlRigState::Update:
-			{
-				if (CachedControlIndex.UpdateCache(FRigElementKey(Control, ERigElementType::Control), Hierarchy))
-				{
-					const FRigControlElement* ControlElement = Hierarchy->GetChecked<FRigControlElement>(CachedControlIndex);
-					Color = ControlElement->Settings.ShapeColor;
-				}
-				break;
-			}
-			default:
-			{
-				break;
-			}
+			const FRigControlElement* ControlElement = Hierarchy->GetChecked<FRigControlElement>(CachedControlIndex);
+			Color = ControlElement->Settings.ShapeColor;
 		}
 	}
 }
@@ -44,27 +28,11 @@ FRigUnit_SetControlColor_Execute()
 	URigHierarchy* Hierarchy = ExecuteContext.Hierarchy;
 	if (Hierarchy)
 	{
-		switch (ExecuteContext.UnitContext.State)
+		if (CachedControlIndex.UpdateCache(FRigElementKey(Control, ERigElementType::Control), Hierarchy))
 		{
-			case EControlRigState::Init:
-			{
-				CachedControlIndex.Reset();
-				break;
-			}
-			case EControlRigState::Update:
-			{
-				if (CachedControlIndex.UpdateCache(FRigElementKey(Control, ERigElementType::Control), Hierarchy))
-				{
-					FRigControlElement* ControlElement = Hierarchy->GetChecked<FRigControlElement>(CachedControlIndex);
-					ControlElement->Settings.ShapeColor = Color;
-					Hierarchy->Notify(ERigHierarchyNotification::ControlSettingChanged, ControlElement);
-				}
-				break;
-			}
-			default:
-			{
-				break;
-			}
+			FRigControlElement* ControlElement = Hierarchy->GetChecked<FRigControlElement>(CachedControlIndex);
+			ControlElement->Settings.ShapeColor = Color;
+			Hierarchy->Notify(ERigHierarchyNotification::ControlSettingChanged, ControlElement);
 		}
 	}
 }
diff --git a/Engine/Plugins/Animation/ControlRig/Source/ControlRig/Private/Units/Hierarchy/RigUnit_SetControlDrivenList.cpp b/Engine/Plugins/Animation/ControlRig/Source/ControlRig/Private/Units/Hierarchy/RigUnit_SetControlDrivenList.cpp
--- a/Engine/Plugins/Animation/ControlRig/Source/ControlRig/Private/Units/Hierarchy/RigUnit_SetControlDrivenList.cpp
+++ b/Engine/Plugins/Animation/ControlRig/Source/ControlRig/Private/Units/Hierarchy/RigUnit_SetControlDrivenList.cpp
@@ -16,26 +16,10 @@ FRigUnit_GetControlDrivenList_Execute()
 	const URigHierarchy* Hierarchy = ExecuteContext.Hierarchy;
 	if (Hierarchy)
 	{
-		switch (ExecuteContext.UnitContext.State)
+		if (CachedControlIndex.UpdateCache(FRigElementKey(Control, ERigElementType::Control), Hierarchy))
 		{
-			case EControlRigState::Init:
-			{
-				CachedControlIndex.Reset();
-				break;
-			}
-			case EControlRigState::Update:
-			{
-				if (CachedControlIndex.UpdateCache(FRigElementKey(Control, ERigElementType::Control), Hierarchy))
-				{
-					const FRigControlElement* ControlElement = Hierarchy->GetChecked<FRigControlElement>(CachedControlIndex);
-					Driven = ControlElement->Settings.DrivenControls;
-				}
-				break;
-			}
-			default:
-			{
-				break;
-			}
+			const FRigControlElement* ControlElement = Hierarchy->GetChecked<FRigControlElement>(CachedControlIndex);
+			Driven = ControlElement->Settings.DrivenControls;
 		}
 	}
 }
@@ -46,30 +30,14 @@ FRigUnit_SetControlDrivenList_Execute()
 	URigHierarchy* Hierarchy = ExecuteContext.Hierarchy;
 	if (Hierarchy)
 	{
-		switch (ExecuteContext.UnitContext.State)
+		if (CachedControlIndex.UpdateCache(FRigElementKey(Control, ERigElementType::Control), Hierarchy))
 		{
-			case EControlRigState::Init:
-			{
-				CachedControlIndex.Reset();
-				break;
-			}
-			case EControlRigState::Update:
-			{
-				if (CachedControlIndex.UpdateCache(FRigElementKey(Control, ERigElementType::Control), Hierarchy))
-				{
-					FRigControlElement* ControlElement = Hierarchy->GetChecked<FRigControlElement>(CachedControlIndex);
-					if(ControlElement->Settings.DrivenControls != Driven)
-					{
-						Swap(ControlElement->Settings.DrivenControls, ControlElement->Settings.PreviouslyDrivenControls);
-						ControlElement->Settings.DrivenControls = Driven;
-						Hierarchy->Notify(ERigHierarchyNotification::ControlDrivenListChanged, ControlElement);
-					}
-				}
-				break;
-			}
-			default:
+			FRigControlElement* ControlElement = Hierarchy->GetChecked<FRigControlElement>(CachedControlIndex);
+			if(ControlElement->Settings.DrivenControls != Driven)
 			{
-				break;
+				Swap(ControlElement->Settings.DrivenControls, ControlElement->Settings.PreviouslyDrivenControls);
+				ControlElement->Settings.DrivenControls = Driven;
+				Hierarchy->Notify(ERigHierarchyNotification::ControlDrivenListChanged, ControlElement);
 			}
 		}
 	}
diff --git a/Engine/Plugins/Animation/ControlRig/Source/ControlRig/Private/Units/Hierarchy/RigUnit_SetControlOffset.cpp b/Engine/Plugins/Animation/ControlRig/Source/ControlRig/Private/Units/Hierarchy/RigUnit_SetControlOffset.cpp
--- a/Engine/Plugins/Animation/ControlRig/Source/ControlRig/Private/Units/Hierarchy/RigUnit_SetControlOffset.cpp
+++ b/Engine/Plugins/Animation/ControlRig/Source/ControlRig/Private/Units/Hierarchy/RigUnit_SetControlOffset.cpp
@@ -16,37 +16,22 @@ FRigUnit_SetControlOffset_Execute()
 	URigHierarchy* Hierarchy = ExecuteContext.Hierarchy;
 	if (Hierarchy)
 	{
-		switch (ExecuteContext.UnitContext.State)
+		if (!CachedControlIndex.UpdateCache(FRigElementKey(Control, ERigElementType::Control), Hierarchy))
 		{
-			case EControlRigState::Init:
-			{
-				CachedControlIndex.Reset();
-				break;
-			}
-			case EControlRigState::Update:
-			{
-				if (!CachedControlIndex.UpdateCache(FRigElementKey(Control, ERigElementType::Control), Hierarchy))
-				{
-					UE_CONTROLRIG_RIGUNIT_REPORT_WARNING(TEXT("Control '%s' is not valid."), *Control.ToString());
-					return;
-				}
+			UE_CONTROLRIG_RIGUNIT_REPORT_WARNING(TEXT("Control '%s' is not valid."), *Control.ToString());
+			return;
+		}
 
-				FRigControlElement* ControlElement = Hierarchy->Get<FRigControlElement>(CachedControlIndex);
-				if (Space == EBoneGetterSetterMode::GlobalSpace)
-				{
-					Hierarchy->SetControlOffsetTransform(ControlElement, Offset, ERigTransformType::CurrentGlobal, true, false);
-					Hierarchy->SetControlOffsetTransform(ControlElement, Offset, ERigTransformType::InitialGlobal, true, false);
-				}
-				else
-				{
-					Hierarchy->SetControlOffsetTransform(ControlElement, Offset, ERigTransformType::CurrentLocal, true, false);
-					Hierarchy->SetControlOffsetTransform(ControlElement, Offset, ERigTransformType::InitialLocal, true, false);
-				}
-			}
-			default:
-			{
-				break;
-			}
+		FRigControlElement* ControlElement = Hierarchy->Get<FRigControlElement>(CachedControlIndex);
+		if (Space == EBoneGetterSetterMode::GlobalSpace)
+		{
+			Hierarchy->SetControlOffsetTransform(ControlElement, Offset, ERigTransformType::CurrentGlobal, true, false);
+			Hierarchy->SetControlOffsetTransform(ControlElement, Offset, ERigTransformType::InitialGlobal, true, false);
+		}
+		else
+		{
+			Hierarchy->SetControlOffsetTransform(ControlElement, Offset, ERigTransformType::CurrentLocal, true, false);
+			Hierarchy->SetControlOffsetTransform(ControlElement, Offset, ERigTransformType::InitialLocal, true, false);
 		}
 	}
 }
@@ -58,30 +43,14 @@ FRigUnit_GetShapeTransform_Execute()
 	const URigHierarchy* Hierarchy = ExecuteContext.Hierarchy;
 	if (Hierarchy)
 	{
-		switch (ExecuteContext.UnitContext.State)
+		if (!CachedControlIndex.UpdateCache(FRigElementKey(Control, ERigElementType::Control), Hierarchy))
 		{
-			case EControlRigState::Init:
-			{
-				CachedControlIndex.Reset();
-				break;
-			}
-			case EControlRigState::Update:
-			{
-				if (!CachedControlIndex.UpdateCache(FRigElementKey(Control, ERigElementType::Control), Hierarchy))
-				{
-					UE_CONTROLRIG_RIGUNIT_REPORT_WARNING(TEXT("Control '%s' is not valid."), *Control.ToString());
-					return;
-				}
-
-				const FRigControlElement* ControlElement = Hierarchy->Get<FRigControlElement>(CachedControlIndex);
-				Transform = Hierarchy->GetControlShapeTransform((FRigControlElement*)ControlElement, ERigTransformType::CurrentLocal);
-				break;
-			}
-			default:
-			{
-				break;
-			}
+			UE_CONTROLRIG_RIGUNIT_REPORT_WARNING(TEXT("Control '%s' is not valid."), *Control.ToString());
+			return;
 		}
+
+		const FRigControlElement* ControlElement = Hierarchy->Get<FRigControlElement>(CachedControlIndex);
+		Transform = Hierarchy->GetControlShapeTransform((FRigControlElement*)ControlElement, ERigTransformType::CurrentLocal);
 	}
 }
 
@@ -92,29 +61,13 @@ FRigUnit_SetShapeTransform_Execute()
 	URigHierarchy* Hierarchy = ExecuteContext.Hierarchy;
 	if (Hierarchy)
 	{
-		switch (ExecuteContext.UnitContext.State)
+		if (!CachedControlIndex.UpdateCache(FRigElementKey(Control, ERigElementType::Control), Hierarchy))
 		{
-			case EControlRigState::Init:
-			{
-				CachedControlIndex.Reset();
-			break;
-			}
-			case EControlRigState::Update:
-			{
-				if (!CachedControlIndex.UpdateCache(FRigElementKey(Control, ERigElementType::Control), Hierarchy))
-				{
-					UE_CONTROLRIG_RIGUNIT_REPORT_WARNING(TEXT("Control '%s' is not valid."), *Control.ToString());
-					return;
-				}
-
-				FRigControlElement* ControlElement = Hierarchy->Get<FRigControlElement>(CachedControlIndex);
-				Hierarchy->SetControlShapeTransform(ControlElement, Transform, ERigTransformType::InitialLocal);
-				break;
-			}
-			default:
-			{
-				break;
-			}
+			UE_CONTROLRIG_RIGUNIT_REPORT_WARNING(TEXT("Control '%s' is not valid."), *Control.ToString());
+			return;
 		}
+
+		FRigControlElement* ControlElement = Hierarchy->Get<FRigControlElement>(CachedControlIndex);
+		Hierarchy->SetControlShapeTransform(ControlElement, Transform, ERigTransformType::InitialLocal);
 	}
 }
diff --git a/Engine/Plugins/Animation/ControlRig/Source/ControlRig/Private/Units/Hierarchy/RigUnit_SetControlTransform.cpp b/Engine/Plugins/Animation/ControlRig/Source/ControlRig/Private/Units/Hierarchy/RigUnit_SetControlTransform.cpp
--- a/Engine/Plugins/Animation/ControlRig/Source/ControlRig/Private/Units/Hierarchy/RigUnit_SetControlTransform.cpp
+++ b/Engine/Plugins/Animation/ControlRig/Source/ControlRig/Private/Units/Hierarchy/RigUnit_SetControlTransform.cpp
@@ -13,26 +13,10 @@ FRigUnit_SetControlBool_Execute()
 	URigHierarchy* Hierarchy = ExecuteContext.Hierarchy;
 	if (Hierarchy)
 	{
-		switch (ExecuteContext.UnitContext.State)
+		const FRigElementKey Key(Control, ERigElementType::Control);
+		if (CachedControlIndex.UpdateCache(Key, Hierarchy))
 		{
-			case EControlRigState::Init:
-			{
-				CachedControlIndex.Reset();
-				break;
-			}
-			case EControlRigState::Update:
-			{
-				const FRigElementKey Key(Control, ERigElementType::Control);
-				if (CachedControlIndex.UpdateCache(Key, Hierarchy))
-				{
-					Hierarchy->SetControlValue(CachedControlIndex, FRigControlValue::Make<bool>(BoolValue));
-				}
-				break;
-			}
-			default:
-			{
-				break;
-			}
+			Hierarchy->SetControlValue(CachedControlIndex, FRigControlValue::Make<bool>(BoolValue));
 		}
 	}
 }
@@ -45,22 +29,11 @@ FRigUnit_SetMultiControlBool_Execute()
 	{
 		CachedControlIndices.SetNum(Entries.Num());
 
-		switch (ExecuteContext.UnitContext.State)
+		for (int32 EntryIndex = 0; EntryIndex < Entries.Num(); EntryIndex++)
 		{
-			case EControlRigState::Update:
-			{
-				for (int32 EntryIndex = 0; EntryIndex < Entries.Num(); EntryIndex++)
-				{
-					bool BoolValue = Entries[EntryIndex].BoolValue;
+			bool BoolValue = Entries[EntryIndex].BoolValue;
 
-					FRigUnit_SetControlBool::StaticExecute(ExecuteContext, Entries[EntryIndex].Control, BoolValue, CachedControlIndices[EntryIndex]);
-				}
-				break;
-			}
-			default:
-			{
-				break;
-			}
+			FRigUnit_SetControlBool::StaticExecute(ExecuteContext, Entries[EntryIndex].Control, BoolValue, CachedControlIndices[EntryIndex]);
 		}
 	}
 }
@@ -71,33 +44,17 @@ FRigUnit_SetControlFloat_Execute()
 	URigHierarchy* Hierarchy = ExecuteContext.Hierarchy;
 	if (Hierarchy)
 	{
-		switch (ExecuteContext.UnitContext.State)
+		const FRigElementKey Key(Control, ERigElementType::Control);
+		if (CachedControlIndex.UpdateCache(Key, Hierarchy))
 		{
-			case EControlRigState::Init:
-			{
-				CachedControlIndex.Reset();
-				break;
-			}
-			case EControlRigState::Update:
+			if(FMath::IsNearlyEqual((float)Weight, 1.f))
 			{
-				const FRigElementKey Key(Control, ERigElementType::Control);
-				if (CachedControlIndex.UpdateCache(Key, Hierarchy))
-				{
-					if(FMath::IsNearlyEqual((float)Weight, 1.f))
-					{
-						Hierarchy->SetControlValue(CachedControlIndex, FRigControlValue::Make<float>(FloatValue));
-					}
-					else
-					{
-						float PreviousValue = Hierarchy->GetControlValue(CachedControlIndex).Get<float>();
-						Hierarchy->SetControlValue(CachedControlIndex, FRigControlValue::Make<float>(FMath::Lerp<float>(PreviousValue, FloatValue, FMath::Clamp<float>(Weight, 0.f, 1.f))));
-					}
-				}
-				break;
+				Hierarchy->SetControlValue(CachedControlIndex, FRigControlValue::Make<float>(FloatValue));
 			}
-			default:
+			else
 			{
-				break;
+				float PreviousValue = Hierarchy->GetControlValue(CachedControlIndex).Get<float>();
+				Hierarchy->SetControlValue(CachedControlIndex, FRigControlValue::Make<float>(FMath::Lerp<float>(PreviousValue, FloatValue, FMath::Clamp<float>(Weight, 0.f, 1.f))));
 			}
 		}
 	}
@@ -111,22 +68,11 @@ FRigUnit_SetMultiControlFloat_Execute()
 	{ 
 		CachedControlIndices.SetNum(Entries.Num());
 	 
-		switch (ExecuteContext.UnitContext.State)
+		for (int32 EntryIndex = 0; EntryIndex < Entries.Num(); EntryIndex++)
 		{
-			case EControlRigState::Update:
-			{ 
-				for (int32 EntryIndex = 0; EntryIndex < Entries.Num(); EntryIndex++)
-				{
-					float FloatValue = Entries[EntryIndex].FloatValue;
+			float FloatValue = Entries[EntryIndex].FloatValue;
 
-					FRigUnit_SetControlFloat::StaticExecute(ExecuteContext, Entries[EntryIndex].Control, Weight, FloatValue, CachedControlIndices[EntryIndex]); 
-				}
-				break;
-			}
-			default:
-			{
-				break;
-			}
+			FRigUnit_SetControlFloat::StaticExecute(ExecuteContext, Entries[EntryIndex].Control, Weight, FloatValue, CachedControlIndices[EntryIndex]); 
 		}
 	}
 }
@@ -137,33 +83,17 @@ FRigUnit_SetControlInteger_Execute()
 	URigHierarchy* Hierarchy = ExecuteContext.Hierarchy;
 	if (Hierarchy)
 	{
-		switch (ExecuteContext.UnitContext.State)
+		const FRigElementKey Key(Control, ERigElementType::Control);
+		if (CachedControlIndex.UpdateCache(Key, Hierarchy))
 		{
-			case EControlRigState::Init:
+			if(FMath::IsNearlyEqual((float)Weight, 1.f))
 			{
-				CachedControlIndex.Reset();
-				break;
+				Hierarchy->SetControlValue(CachedControlIndex, FRigControlValue::Make<int32>(IntegerValue));
 			}
-			case EControlRigState::Update:
+			else
 			{
-				const FRigElementKey Key(Control, ERigElementType::Control);
-				if (CachedControlIndex.UpdateCache(Key, Hierarchy))
-				{
-					if(FMath::IsNearlyEqual((float)Weight, 1.f))
-					{
-						Hierarchy->SetControlValue(CachedControlIndex, FRigControlValue::Make<int32>(IntegerValue));
-					}
-					else
-					{
-						int32 PreviousValue = Hierarchy->GetControlValue(CachedControlIndex).Get<int32>();
-						Hierarchy->SetControlValue(CachedControlIndex, FRigControlValue::Make<int32>((int32)FMath::Lerp<float>((float)PreviousValue, (float)IntegerValue, FMath::Clamp<float>(Weight, 0.f, 1.f))));
-					}
-				}
-				break;
-			}
-			default:
-			{
-				break;
+				int32 PreviousValue = Hierarchy->GetControlValue(CachedControlIndex).Get<int32>();
+				Hierarchy->SetControlValue(CachedControlIndex, FRigControlValue::Make<int32>((int32)FMath::Lerp<float>((float)PreviousValue, (float)IntegerValue, FMath::Clamp<float>(Weight, 0.f, 1.f))));
 			}
 		}
 	}
@@ -177,22 +107,11 @@ FRigUnit_SetMultiControlInteger_Execute()
 	{
 		CachedControlIndices.SetNum(Entries.Num());
 
-		switch (ExecuteContext.UnitContext.State)
+		for (int32 EntryIndex = 0; EntryIndex < Entries.Num(); EntryIndex++)
 		{
-			case EControlRigState::Update:
-			{
-				for (int32 EntryIndex = 0; EntryIndex < Entries.Num(); EntryIndex++)
-				{
-					int32 IntegerValue = Entries[EntryIndex].IntegerValue;
+			int32 IntegerValue = Entries[EntryIndex].IntegerValue;
 
-					FRigUnit_SetControlInteger::StaticExecute(ExecuteContext, Entries[EntryIndex].Control, Weight, IntegerValue, CachedControlIndices[EntryIndex]);
-				}
-				break;
-			}
-			default:
-			{
-				break;
-			}
+			FRigUnit_SetControlInteger::StaticExecute(ExecuteContext, Entries[EntryIndex].Control, Weight, IntegerValue, CachedControlIndices[EntryIndex]);
 		}
 	}
 }
@@ -203,34 +122,18 @@ FRigUnit_SetControlVector2D_Execute()
 	URigHierarchy* Hierarchy = ExecuteContext.Hierarchy;
 	if (Hierarchy)
 	{
-		switch (ExecuteContext.UnitContext.State)
+		const FRigElementKey Key(Control, ERigElementType::Control);
+		if (CachedControlIndex.UpdateCache(Key, Hierarchy))
 		{
-			case EControlRigState::Init:
+			const FVector3f CurrentValue(Vector.X, Vector.Y, 0.f);
+			if(FMath::IsNearlyEqual((float)Weight, 1.f))
 			{
-				CachedControlIndex.Reset();
-				break;
+				Hierarchy->SetControlValue(CachedControlIndex, FRigControlValue::Make<FVector3f>(CurrentValue));
 			}
-			case EControlRigState::Update:
+			else
 			{
-				const FRigElementKey Key(Control, ERigElementType::Control);
-				if (CachedControlIndex.UpdateCache(Key, Hierarchy))
-				{
-					const FVector3f CurrentValue(Vector.X, Vector.Y, 0.f);
-					if(FMath::IsNearlyEqual((float)Weight, 1.f))
-					{
-						Hierarchy->SetControlValue(CachedControlIndex, FRigControlValue::Make<FVector3f>(CurrentValue));
-					}
-					else
-					{
-						const FVector3f PreviousValue = Hierarchy->GetControlValue(CachedControlIndex).Get<FVector3f>();
-						Hierarchy->SetControlValue(CachedControlIndex, FRigControlValue::Make<FVector3f>(FMath::Lerp<FVector3f>(PreviousValue, CurrentValue, FMath::Clamp<float>(Weight, 0.f, 1.f))));
-					}
-				}
-				break;
-			}
-			default:
-			{
-				break;
+				const FVector3f PreviousValue = Hierarchy->GetControlValue(CachedControlIndex).Get<FVector3f>();
+				Hierarchy->SetControlValue(CachedControlIndex, FRigControlValue::Make<FVector3f>(FMath::Lerp<FVector3f>(PreviousValue, CurrentValue, FMath::Clamp<float>(Weight, 0.f, 1.f))));
 			}
 		}
 	}
@@ -244,22 +147,11 @@ FRigUnit_SetMultiControlVector2D_Execute()
 	{
 		CachedControlIndices.SetNum(Entries.Num());
 
-		switch (ExecuteContext.UnitContext.State)
+		for (int32 EntryIndex = 0; EntryIndex < Entries.Num(); EntryIndex++)
 		{
-			case EControlRigState::Update:
-			{
-				for (int32 EntryIndex = 0; EntryIndex < Entries.Num(); EntryIndex++)
-				{
-					FVector2D Vector = Entries[EntryIndex].Vector;
+			FVector2D Vector = Entries[EntryIndex].Vector;
 
-					FRigUnit_SetControlVector2D::StaticExecute(ExecuteContext, Entries[EntryIndex].Control, Weight, Vector, CachedControlIndices[EntryIndex]);
-				}
-				break;
-			}
-			default:
-			{
-				break;
-			}
+			FRigUnit_SetControlVector2D::StaticExecute(ExecuteContext, Entries[EntryIndex].Control, Weight, Vector, CachedControlIndices[EntryIndex]);
 		}
 	}
 }
@@ -270,73 +162,58 @@ FRigUnit_SetControlVector_Execute()
 	URigHierarchy* Hierarchy = ExecuteContext.Hierarchy;
 	if (Hierarchy)
 	{
-		switch (ExecuteContext.UnitContext.State)
+		const FRigElementKey Key(Control, ERigElementType::Control);
+		if (CachedControlIndex.UpdateCache(Key, Hierarchy))
 		{
-			case EControlRigState::Init:
+			const ERigControlType ControlType = Hierarchy->GetChecked<FRigControlElement>(CachedControlIndex)->Settings.ControlType;
+			
+			FTransform Transform = FTransform::Identity;
+			if (Space == EBoneGetterSetterMode::GlobalSpace)
 			{
-				CachedControlIndex.Reset();
-				break;
+				Transform = Hierarchy->GetGlobalTransform(CachedControlIndex);
 			}
-			case EControlRigState::Update:
+
+			if (ControlType == ERigControlType::Position)
 			{
-				const FRigElementKey Key(Control, ERigElementType::Control);
-				if (CachedControlIndex.UpdateCache(Key, Hierarchy))
+				if(FMath::IsNearlyEqual((float)Weight, 1.f))
 				{
-					const ERigControlType ControlType = Hierarchy->GetChecked<FRigControlElement>(CachedControlIndex)->Settings.ControlType;
-					
-					FTransform Transform = FTransform::Identity;
-					if (Space == EBoneGetterSetterMode::GlobalSpace)
-					{
-						Transform = Hierarchy->GetGlobalTransform(CachedControlIndex);
-					}
-
-					if (ControlType == ERigControlType::Position)
-					{
-						if(FMath::IsNearlyEqual((float)Weight, 1.f))
-						{
-							Transform.SetLocation(Vector);
-						}
-						else
-						{
-							FVector PreviousValue = Transform.GetLocation();
-							Transform.SetLocation(FMath::Lerp<FVector>(PreviousValue, Vector, FMath::Clamp<float>(Weight, 0.f, 1.f)));
-						}
-					}
-					else if (ControlType == ERigControlType::Scale)
-					{
-						if(FMath::IsNearlyEqual((float)Weight, 1.f))
-						{
-							Transform.SetScale3D(Vector);
-						}
-						else
-						{
-							FVector PreviousValue = Transform.GetScale3D();
-							Transform.SetScale3D(FMath::Lerp<FVector>(PreviousValue, Vector, FMath::Clamp<float>(Weight, 0.f, 1.f)));
-						}
-					}
-
-					switch (Space)
-					{
-						case EBoneGetterSetterMode::GlobalSpace:
-						{
-							Hierarchy->SetGlobalTransform(CachedControlIndex, Transform);
-							break;
-						}
-						case EBoneGetterSetterMode::LocalSpace:
-						{
-							Hierarchy->SetLocalTransform(CachedControlIndex, Transform);
-							break;
-						}
-						default:
-						{
-							break;
-						}
-					}
+					Transform.SetLocation(Vector);
+				}
+				else
+				{
+					FVector PreviousValue = Transform.GetLocation();
+					Transform.SetLocation(FMath::Lerp<FVector>(PreviousValue, Vector, FMath::Clamp<float>(Weight, 0.f, 1.f)));
 				}
 			}
-			default:
+			else if (ControlType == ERigControlType::Scale)
 			{
-				break;
+				if(FMath::IsNearlyEqual((float)Weight, 1.f))
+				{
+					Transform.SetScale3D(Vector);
+				}
+				else
+				{
+					FVector PreviousValue = Transform.GetScale3D();
+					Transform.SetScale3D(FMath::Lerp<FVector>(PreviousValue, Vector, FMath::Clamp<float>(Weight, 0.f, 1.f)));
+				}
+			}
+
+			switch (Space)
+			{
+				case EBoneGetterSetterMode::GlobalSpace:
+				{
+					Hierarchy->SetGlobalTransform(CachedControlIndex, Transform);
+					break;
+				}
+				case EBoneGetterSetterMode::LocalSpace:
+				{
+					Hierarchy->SetLocalTransform(CachedControlIndex, Transform);
+					break;
+				}
+				default:
+				{
+					break;
+				}
 			}
 		}
 	}
@@ -348,58 +225,43 @@ FRigUnit_SetControlRotator_Execute()
 	URigHierarchy* Hierarchy = ExecuteContext.Hierarchy;
 	if (Hierarchy)
 	{
-		switch (ExecuteContext.UnitContext.State)
+		const FRigElementKey Key(Control, ERigElementType::Control);
+		if (CachedControlIndex.UpdateCache(Key, Hierarchy))
 		{
-			case EControlRigState::Init:
+			FTransform Transform = FTransform::Identity;
+			if (Space == EBoneGetterSetterMode::GlobalSpace)
 			{
-				CachedControlIndex.Reset();
-				break;
+				Transform = Hierarchy->GetGlobalTransform(CachedControlIndex);
 			}
-			case EControlRigState::Update:
-			{
-				const FRigElementKey Key(Control, ERigElementType::Control);
-				if (CachedControlIndex.UpdateCache(Key, Hierarchy))
-				{
-					FTransform Transform = FTransform::Identity;
-					if (Space == EBoneGetterSetterMode::GlobalSpace)
-					{
-						Transform = Hierarchy->GetGlobalTransform(CachedControlIndex);
-					}
 
-					FQuat Quat = FQuat(Rotator);
-					if (FMath::IsNearlyEqual((float)Weight, 1.f))
-					{
-						Transform.SetRotation(Quat);
-					}
-					else
-					{
-						FQuat PreviousValue = Transform.GetRotation();
-						Transform.SetRotation(FQuat::Slerp(PreviousValue, Quat, FMath::Clamp<float>(Weight, 0.f, 1.f)));
-					}
-					Transform.NormalizeRotation();
-
-					switch (Space)
-					{
-						case EBoneGetterSetterMode::GlobalSpace:
-						{
-							Hierarchy->SetGlobalTransform(CachedControlIndex, Transform);
-							break;
-						}
-						case EBoneGetterSetterMode::LocalSpace:
-						{
-							Hierarchy->SetLocalTransform(CachedControlIndex, Transform);
-							break;
-						}
-						default:
-						{
-							break;
-						}
-					}
-				}
+			FQuat Quat = FQuat(Rotator);
+			if (FMath::IsNearlyEqual((float)Weight, 1.f))
+			{
+				Transform.SetRotation(Quat);
 			}
-			default:
+			else
 			{
-				break;
+				FQuat PreviousValue = Transform.GetRotation();
+				Transform.SetRotation(FQuat::Slerp(PreviousValue, Quat, FMath::Clamp<float>(Weight, 0.f, 1.f)));
+			}
+			Transform.NormalizeRotation();
+
+			switch (Space)
+			{
+				case EBoneGetterSetterMode::GlobalSpace:
+				{
+					Hierarchy->SetGlobalTransform(CachedControlIndex, Transform);
+					break;
+				}
+				case EBoneGetterSetterMode::LocalSpace:
+				{
+					Hierarchy->SetLocalTransform(CachedControlIndex, Transform);
+					break;
+				}
+				default:
+				{
+					break;
+				}
 			}
 		}
 	}
@@ -413,22 +275,11 @@ FRigUnit_SetMultiControlRotator_Execute()
 	{
 		CachedControlIndices.SetNum(Entries.Num());
 
-		switch (ExecuteContext.UnitContext.State)
+		for (int32 EntryIndex = 0; EntryIndex < Entries.Num(); EntryIndex++)
 		{
-			case EControlRigState::Update:
-			{
-				for (int32 EntryIndex = 0; EntryIndex < Entries.Num(); EntryIndex++)
-				{
-					FRotator Rotator = Entries[EntryIndex].Rotator;
+			FRotator Rotator = Entries[EntryIndex].Rotator;
 
-					FRigUnit_SetControlRotator::StaticExecute(ExecuteContext, Entries[EntryIndex].Control, Weight, Rotator, Entries[EntryIndex].Space , CachedControlIndices[EntryIndex]);
-				}
-				break;
-			}
-			default:
-			{
-				break;
-			}
+			FRigUnit_SetControlRotator::StaticExecute(ExecuteContext, Entries[EntryIndex].Control, Weight, Rotator, Entries[EntryIndex].Space , CachedControlIndices[EntryIndex]);
 		}
 	}
 }
@@ -439,56 +290,41 @@ FRigUnit_SetControlTransform_Execute()
 	URigHierarchy* Hierarchy = ExecuteContext.Hierarchy;
 	if (Hierarchy)
 	{
-		switch (ExecuteContext.UnitContext.State)
+		const FRigElementKey Key(Control, ERigElementType::Control);
+		if (CachedControlIndex.UpdateCache(Key, Hierarchy))
 		{
-			case EControlRigState::Init:
-			{
-				CachedControlIndex.Reset();
-				break;
-			}
-			case EControlRigState::Update:
+			switch (Space)
 			{
-				const FRigElementKey Key(Control, ERigElementType::Control);
-				if (CachedControlIndex.UpdateCache(Key, Hierarchy))
+				case EBoneGetterSetterMode::GlobalSpace:
 				{
-					switch (Space)
+					if(FMath::IsNearlyEqual((float)Weight, 1.f))
+					{
+						Hierarchy->SetGlobalTransform(CachedControlIndex, Transform);
+					}
+					else
 					{
-						case EBoneGetterSetterMode::GlobalSpace:
-						{
-							if(FMath::IsNearlyEqual((float)Weight, 1.f))
-							{
-								Hierarchy->SetGlobalTransform(CachedControlIndex, Transform);
-							}
-							else
-							{
-								FTransform PreviousTransform = Hierarchy->GetGlobalTransform(CachedControlIndex);
-								Hierarchy->SetGlobalTransform(CachedControlIndex, FControlRigMathLibrary::LerpTransform(PreviousTransform, Transform, FMath::Clamp<float>(Weight, 0.f, 1.f)));
-							}
-							break;
-						}
-						case EBoneGetterSetterMode::LocalSpace:
-						{
-							if(FMath::IsNearlyEqual((float)Weight, 1.f))
-							{
-								Hierarchy->SetLocalTransform(CachedControlIndex, Transform);
-							}
-							else
-							{
-								FTransform PreviousTransform = Hierarchy->GetLocalTransform(CachedControlIndex);
-								Hierarchy->SetLocalTransform(CachedControlIndex, FControlRigMathLibrary::LerpTransform(PreviousTransform, Transform, FMath::Clamp<float>(Weight, 0.f, 1.f)));
-							}
-							break;
-						}
-						default:
-						{
-							break;
-						}
+						FTransform PreviousTransform = Hierarchy->GetGlobalTransform(CachedControlIndex);
+						Hierarchy->SetGlobalTransform(CachedControlIndex, FControlRigMathLibrary::LerpTransform(PreviousTransform, Transform, FMath::Clamp<float>(Weight, 0.f, 1.f)));
 					}
+					break;
+				}
+				case EBoneGetterSetterMode::LocalSpace:
+				{
+					if(FMath::IsNearlyEqual((float)Weight, 1.f))
+					{
+						Hierarchy->SetLocalTransform(CachedControlIndex, Transform);
+					}
+					else
+					{
+						FTransform PreviousTransform = Hierarchy->GetLocalTransform(CachedControlIndex);
+						Hierarchy->SetLocalTransform(CachedControlIndex, FControlRigMathLibrary::LerpTransform(PreviousTransform, Transform, FMath::Clamp<float>(Weight, 0.f, 1.f)));
+					}
+					break;
+				}
+				default:
+				{
+					break;
 				}
-			}
-			default:
-			{
-				break;
 			}
 		}
 	}
@@ -536,7 +372,6 @@ IMPLEMENT_RIGUNIT_AUTOMATION_TEST(FRigUnit_SetMultiControlBool)
 	Unit.Entries.Add(Entry);
 	AddErrorIfFalse(Unit.Entries.Num() == 2, TEXT("unexpected number of entries"));
 
-	Init();
 	Execute();
 
 	AddErrorIfFalse(Hierarchy->GetControlValue(FRigElementKey(TEXT("Control1"), ERigElementType::Control), ERigControlValueType::Current).Get<bool>() == true, TEXT("unexpected control value"));
@@ -569,7 +404,6 @@ IMPLEMENT_RIGUNIT_AUTOMATION_TEST(FRigUnit_SetMultiControlFloat)
 	Unit.Entries.Add(Entry);
 	AddErrorIfFalse(Unit.Entries.Num() == 2, TEXT("unexpected number of entries"));
 
-	Init();
 	Execute();
 
 	AddErrorIfFalse(Hierarchy->GetControlValue(FRigElementKey(TEXT("Control1"), ERigElementType::Control), ERigControlValueType::Current).Get<float>() == 10.0f, TEXT("unexpected control value"));
@@ -630,7 +464,6 @@ IMPLEMENT_RIGUNIT_AUTOMATION_TEST(FRigUnit_SetMultiControlInteger)
 	Unit.Entries.Add(Entry);
 	AddErrorIfFalse(Unit.Entries.Num() == 2, TEXT("unexpected number of entries"));
 
-	Init();
 	Execute();
 
 	AddErrorIfFalse(Hierarchy->GetControlValue(FRigElementKey(TEXT("Control1"), ERigElementType::Control), ERigControlValueType::Current).Get<int32>() == 10, TEXT("unexpected control value"));
@@ -663,7 +496,6 @@ IMPLEMENT_RIGUNIT_AUTOMATION_TEST(FRigUnit_SetMultiControlVector2D)
 	Unit.Entries.Add(Entry);
 	AddErrorIfFalse(Unit.Entries.Num() == 2, TEXT("unexpected number of entries"));
 
-	Init();
 	Execute();
 
 	const FVector3f TempValue1 = Hierarchy->GetControlValue(FRigElementKey(TEXT("Control1"), ERigElementType::Control), ERigControlValueType::Current).Get<FVector3f>();
@@ -700,7 +532,6 @@ IMPLEMENT_RIGUNIT_AUTOMATION_TEST(FRigUnit_SetMultiControlRotator)
 	Unit.Entries.Add(Entry);
 	AddErrorIfFalse(Unit.Entries.Num() == 2, TEXT("unexpected number of entries"));
 
-	Init();
 	Execute(); 
 
 	const FVector TempValue1 = (FVector)Hierarchy->GetControlValue(FRigElementKey(TEXT("Control1"), ERigElementType::Control), ERigControlValueType::Current).Get<FVector3f>();
diff --git a/Engine/Plugins/Animation/ControlRig/Source/ControlRig/Private/Units/Hierarchy/RigUnit_SetControlVisibility.cpp b/Engine/Plugins/Animation/ControlRig/Source/ControlRig/Private/Units/Hierarchy/RigUnit_SetControlVisibility.cpp
--- a/Engine/Plugins/Animation/ControlRig/Source/ControlRig/Private/Units/Hierarchy/RigUnit_SetControlVisibility.cpp
+++ b/Engine/Plugins/Animation/ControlRig/Source/ControlRig/Private/Units/Hierarchy/RigUnit_SetControlVisibility.cpp
@@ -14,26 +14,11 @@ FRigUnit_GetControlVisibility_Execute()
 	const URigHierarchy* Hierarchy = ExecuteContext.Hierarchy;
 	if (Hierarchy)
 	{
-		switch (ExecuteContext.UnitContext.State)
+		if (CachedControlIndex.UpdateCache(Item, Hierarchy))
 		{
-			case EControlRigState::Init:
+			if(const FRigControlElement* ControlElement = Hierarchy->Get<FRigControlElement>(CachedControlIndex))
 			{
-				CachedControlIndex.Reset();
-			}
-			case EControlRigState::Update:
-			{
-				if (CachedControlIndex.UpdateCache(Item, Hierarchy))
-				{
-					if(const FRigControlElement* ControlElement = Hierarchy->Get<FRigControlElement>(CachedControlIndex))
-					{
-						bVisible = ControlElement->Settings.bShapeVisible;
-					}
-				}
-				break;
-			}
-			default:
-			{
-				break;
+				bVisible = ControlElement->Settings.bShapeVisible;
 			}
 		}
 	}
@@ -45,60 +30,45 @@ FRigUnit_SetControlVisibility_Execute()
 	URigHierarchy* Hierarchy = ExecuteContext.Hierarchy;
 	if (Hierarchy)
 	{
-		switch (ExecuteContext.UnitContext.State)
+		TArray<FRigElementKey> Keys;
+
+		if (Item.IsValid())
 		{
-			case EControlRigState::Init:
+			if (Item.Type != ERigElementType::Control)
 			{
-				CachedControlIndices.Reset();
+				return;
 			}
-			case EControlRigState::Update:
-			{
-				TArray<FRigElementKey> Keys;
-
-				if (Item.IsValid())
-				{
-					if (Item.Type != ERigElementType::Control)
-					{
-						return;
-					}
 
-					Keys.Add(Item);
-				}
-				else if (!Pattern.IsEmpty())
+			Keys.Add(Item);
+		}
+		else if (!Pattern.IsEmpty())
+		{
+			Hierarchy->ForEach<FRigControlElement>([&Keys, Pattern](FRigControlElement* ControlElement) -> bool
+			{
+				if (ControlElement->GetName().ToString().Contains(Pattern, ESearchCase::CaseSensitive))
 				{
-					Hierarchy->ForEach<FRigControlElement>([&Keys, Pattern](FRigControlElement* ControlElement) -> bool
-					{
-						if (ControlElement->GetName().ToString().Contains(Pattern, ESearchCase::CaseSensitive))
-						{
-							Keys.Add(ControlElement->GetKey());
-						}
-						return true;
-					});
+					Keys.Add(ControlElement->GetKey());
 				}
+				return true;
+			});
+		}
 
-				if (CachedControlIndices.Num() != Keys.Num())
-				{
-					CachedControlIndices.Reset();
-					CachedControlIndices.SetNumZeroed(Keys.Num());
-				}
+		if (CachedControlIndices.Num() != Keys.Num())
+		{
+			CachedControlIndices.Reset();
+			CachedControlIndices.SetNumZeroed(Keys.Num());
+		}
 
-				for (int32 Index = 0; Index < Keys.Num(); Index++)
-				{
-					CachedControlIndices[Index].UpdateCache(Keys[Index], Hierarchy);
-				}
+		for (int32 Index = 0; Index < Keys.Num(); Index++)
+		{
+			CachedControlIndices[Index].UpdateCache(Keys[Index], Hierarchy);
+		}
 
-				for (const FCachedRigElement& CachedControlIndex : CachedControlIndices)
-				{
-					if (CachedControlIndex.IsValid())
-					{
-						Hierarchy->SetControlVisibility(CachedControlIndex, bVisible);
-					}
-				}
-				break;
-			}
-			default:
+		for (const FCachedRigElement& CachedControlIndex : CachedControlIndices)
+		{
+			if (CachedControlIndex.IsValid())
 			{
-				break;
+				Hierarchy->SetControlVisibility(CachedControlIndex, bVisible);
 			}
 		}
 	}
diff --git a/Engine/Plugins/Animation/ControlRig/Source/ControlRig/Private/Units/Hierarchy/RigUnit_SetCurveValue.cpp b/Engine/Plugins/Animation/ControlRig/Source/ControlRig/Private/Units/Hierarchy/RigUnit_SetCurveValue.cpp
--- a/Engine/Plugins/Animation/ControlRig/Source/ControlRig/Private/Units/Hierarchy/RigUnit_SetCurveValue.cpp
+++ b/Engine/Plugins/Animation/ControlRig/Source/ControlRig/Private/Units/Hierarchy/RigUnit_SetCurveValue.cpp
@@ -11,25 +11,10 @@ FRigUnit_SetCurveValue_Execute()
 	URigHierarchy* Hierarchy = ExecuteContext.Hierarchy;
 	if (Hierarchy)
 	{
-		switch (ExecuteContext.UnitContext.State)
+		const FRigElementKey Key(Curve, ERigElementType::Curve);
+		if (CachedCurveIndex.UpdateCache(Key, Hierarchy))
 		{
-			case EControlRigState::Init:
-			{
-				CachedCurveIndex.Reset();
-				// fall through to update
-			}
-			case EControlRigState::Update:
-			{
-				const FRigElementKey Key(Curve, ERigElementType::Curve);
-				if (CachedCurveIndex.UpdateCache(Key, Hierarchy))
-				{
-					Hierarchy->SetCurveValueByIndex(CachedCurveIndex, Value);
-				}
-			}
-			default:
-			{
-				break;
-			}
+			Hierarchy->SetCurveValueByIndex(CachedCurveIndex, Value);
 		}
 	}
 }
@@ -46,14 +31,14 @@ IMPLEMENT_RIGUNIT_AUTOMATION_TEST(FRigUnit_SetCurveValue)
 	Hierarchy->ResetCurveValues();
 	Unit.Curve = TEXT("CurveA");
 	Unit.Value = 3.0f;
-	InitAndExecute();
+	Execute();
 
 	AddErrorIfFalse(Hierarchy->GetCurveValue(CurveA) == 3.f, TEXT("unexpected value"));
 
 	Hierarchy->ResetCurveValues();
 	Unit.Curve = TEXT("CurveB");
 	Unit.Value = 13.0f;
-	InitAndExecute();
+	Execute();
 
 	AddErrorIfFalse(Hierarchy->GetCurveValue(CurveB) == 13.f, TEXT("unexpected value"));
 
diff --git a/Engine/Plugins/Animation/ControlRig/Source/ControlRig/Private/Units/Hierarchy/RigUnit_SetRelativeBoneTransform.cpp b/Engine/Plugins/Animation/ControlRig/Source/ControlRig/Private/Units/Hierarchy/RigUnit_SetRelativeBoneTransform.cpp
--- a/Engine/Plugins/Animation/ControlRig/Source/ControlRig/Private/Units/Hierarchy/RigUnit_SetRelativeBoneTransform.cpp
+++ b/Engine/Plugins/Animation/ControlRig/Source/ControlRig/Private/Units/Hierarchy/RigUnit_SetRelativeBoneTransform.cpp
@@ -12,45 +12,30 @@ FRigUnit_SetRelativeBoneTransform_Execute()
 	URigHierarchy* Hierarchy = ExecuteContext.Hierarchy;
 	if (Hierarchy)
 	{
-		switch (ExecuteContext.UnitContext.State)
+		const FRigElementKey BoneKey(Bone, ERigElementType::Bone);
+		const FRigElementKey SpaceKey(Space, ERigElementType::Bone);
+		if (!CachedBone.UpdateCache(BoneKey, Hierarchy))
 		{
-			case EControlRigState::Init:
-			{
-				CachedBone.Reset();
-				CachedSpaceIndex.Reset();
-			}
-			case EControlRigState::Update:
-			{
-				const FRigElementKey BoneKey(Bone, ERigElementType::Bone);
-				const FRigElementKey SpaceKey(Space, ERigElementType::Bone);
-				if (!CachedBone.UpdateCache(BoneKey, Hierarchy))
-				{
-					UE_CONTROLRIG_RIGUNIT_REPORT_WARNING(TEXT("Bone '%s' is not valid."), *Bone.ToString());
-				}
-				else if (!CachedSpaceIndex.UpdateCache(SpaceKey, Hierarchy))
-				{
-					UE_CONTROLRIG_RIGUNIT_REPORT_WARNING(TEXT("Bone '%s' is not valid."), *Bone.ToString());
-					UE_CONTROLRIG_RIGUNIT_REPORT_WARNING(TEXT("Space '%s' is not valid."), *Space.ToString());
-				}
-				else
-				{
-					const FTransform SpaceTransform = Hierarchy->GetGlobalTransform(CachedSpaceIndex);
-					FTransform TargetTransform = Transform * SpaceTransform;
-
-					if (!FMath::IsNearlyEqual(Weight, 1.f))
-					{
-						float T = FMath::Clamp<float>(Weight, 0.f, 1.f);
-						const FTransform PreviousTransform = Hierarchy->GetGlobalTransform(CachedBone);
-						TargetTransform = FControlRigMathLibrary::LerpTransform(PreviousTransform, TargetTransform, T);
-					}
+			UE_CONTROLRIG_RIGUNIT_REPORT_WARNING(TEXT("Bone '%s' is not valid."), *Bone.ToString());
+		}
+		else if (!CachedSpaceIndex.UpdateCache(SpaceKey, Hierarchy))
+		{
+			UE_CONTROLRIG_RIGUNIT_REPORT_WARNING(TEXT("Bone '%s' is not valid."), *Bone.ToString());
+			UE_CONTROLRIG_RIGUNIT_REPORT_WARNING(TEXT("Space '%s' is not valid."), *Space.ToString());
+		}
+		else
+		{
+			const FTransform SpaceTransform = Hierarchy->GetGlobalTransform(CachedSpaceIndex);
+			FTransform TargetTransform = Transform * SpaceTransform;
 
-					Hierarchy->SetGlobalTransform(CachedBone, TargetTransform, bPropagateToChildren);
-				}
-			}
-			default:
+			if (!FMath::IsNearlyEqual(Weight, 1.f))
 			{
-				break;
+				float T = FMath::Clamp<float>(Weight, 0.f, 1.f);
+				const FTransform PreviousTransform = Hierarchy->GetGlobalTransform(CachedBone);
+				TargetTransform = FControlRigMathLibrary::LerpTransform(PreviousTransform, TargetTransform, T);
 			}
+
+			Hierarchy->SetGlobalTransform(CachedBone, TargetTransform, bPropagateToChildren);
 		}
 	}
 }
@@ -86,7 +71,7 @@ IMPLEMENT_RIGUNIT_AUTOMATION_TEST(FRigUnit_SetRelativeBoneTransform)
 	Unit.Space = TEXT("Root");
 	Unit.Transform = FTransform(FVector(0.f, 0.f, 7.f));
 	Unit.bPropagateToChildren = false;
-	InitAndExecute();
+	Execute();
 	AddErrorIfFalse(Hierarchy->GetGlobalTransform(0).GetTranslation().Equals(FVector(1.f, 0.f, 0.f)), TEXT("unexpected transform"));
 	AddErrorIfFalse(Hierarchy->GetGlobalTransform(1).GetTranslation().Equals(FVector(1.f, 0.f, 7.f)), TEXT("unexpected transform"));
 	AddErrorIfFalse(Hierarchy->GetGlobalTransform(2).GetTranslation().Equals(FVector(1.f, 5.f, 3.f)), TEXT("unexpected transform"));
@@ -94,7 +79,7 @@ IMPLEMENT_RIGUNIT_AUTOMATION_TEST(FRigUnit_SetRelativeBoneTransform)
 
 	Hierarchy->ResetPoseToInitial(ERigElementType::Bone);
 	Unit.bPropagateToChildren = true;
-	InitAndExecute();
+	Execute();
 	AddErrorIfFalse(Hierarchy->GetGlobalTransform(0).GetTranslation().Equals(FVector(1.f, 0.f, 0.f)), TEXT("unexpected transform"));
 	AddErrorIfFalse(Hierarchy->GetGlobalTransform(1).GetTranslation().Equals(FVector(1.f, 0.f, 7.f)), TEXT("unexpected transform"));
 	AddErrorIfFalse(Hierarchy->GetGlobalTransform(2).GetTranslation().Equals(FVector(1.f, 3.f, 7.f)), TEXT("unexpected transform"));
@@ -103,7 +88,7 @@ IMPLEMENT_RIGUNIT_AUTOMATION_TEST(FRigUnit_SetRelativeBoneTransform)
 	Hierarchy->ResetPoseToInitial(ERigElementType::Bone);
 	Unit.Space = TEXT("BoneC");
 	Unit.bPropagateToChildren = false;
-	InitAndExecute();
+	Execute();
 	AddErrorIfFalse(Hierarchy->GetGlobalTransform(0).GetTranslation().Equals(FVector(1.f, 0.f, 0.f)), TEXT("unexpected transform"));
 	AddErrorIfFalse(Hierarchy->GetGlobalTransform(1).GetTranslation().Equals(FVector(-4.f, 0.f, 7.f)), TEXT("unexpected transform"));
 	AddErrorIfFalse(Hierarchy->GetGlobalTransform(2).GetTranslation().Equals(FVector(1.f, 5.f, 3.f)), TEXT("unexpected transform"));
@@ -111,7 +96,7 @@ IMPLEMENT_RIGUNIT_AUTOMATION_TEST(FRigUnit_SetRelativeBoneTransform)
 
 	Hierarchy->ResetPoseToInitial(ERigElementType::Bone);
 	Unit.bPropagateToChildren = true;
-	InitAndExecute();
+	Execute();
 	AddErrorIfFalse(Hierarchy->GetGlobalTransform(0).GetTranslation().Equals(FVector(1.f, 0.f, 0.f)), TEXT("unexpected transform"));
 	AddErrorIfFalse(Hierarchy->GetGlobalTransform(1).GetTranslation().Equals(FVector(-4.f, 0.f, 7.f)), TEXT("unexpected transform"));
 	AddErrorIfFalse(Hierarchy->GetGlobalTransform(2).GetTranslation().Equals(FVector(-4.f, 3.f, 7.f)), TEXT("unexpected transform"));
diff --git a/Engine/Plugins/Animation/ControlRig/Source/ControlRig/Private/Units/Hierarchy/RigUnit_SetSpaceInitialTransform.cpp b/Engine/Plugins/Animation/ControlRig/Source/ControlRig/Private/Units/Hierarchy/RigUnit_SetSpaceInitialTransform.cpp
--- a/Engine/Plugins/Animation/ControlRig/Source/ControlRig/Private/Units/Hierarchy/RigUnit_SetSpaceInitialTransform.cpp
+++ b/Engine/Plugins/Animation/ControlRig/Source/ControlRig/Private/Units/Hierarchy/RigUnit_SetSpaceInitialTransform.cpp
@@ -14,36 +14,21 @@ FRigUnit_SetSpaceInitialTransform_Execute()
 	URigHierarchy* Hierarchy = ExecuteContext.Hierarchy;
 	if (Hierarchy)
 	{
-		switch (ExecuteContext.UnitContext.State)
+		const FRigElementKey SpaceKey(SpaceName, ERigElementType::Null);
+		if (!CachedSpaceIndex.UpdateCache(SpaceKey, Hierarchy))
 		{
-			case EControlRigState::Init:
-			{
-				CachedSpaceIndex.Reset();
-				break;
-			}
-			case EControlRigState::Update:
-			{
-				const FRigElementKey SpaceKey(SpaceName, ERigElementType::Null);
-				if (!CachedSpaceIndex.UpdateCache(SpaceKey, Hierarchy))
-				{
-					UE_CONTROLRIG_RIGUNIT_REPORT_WARNING(TEXT("Space '%s' is not valid."), *SpaceName.ToString());
-					return;
-				}
-
-				FTransform InitialTransform = Transform;
-				if (Space == EBoneGetterSetterMode::GlobalSpace)
-				{
-					const FTransform ParentTransform = Hierarchy->GetParentTransformByIndex(CachedSpaceIndex, true);
-					InitialTransform = InitialTransform.GetRelativeTransform(ParentTransform);
-				}
-
-				Hierarchy->SetInitialLocalTransform(CachedSpaceIndex, InitialTransform);
-			}
-			default:
-			{
-				break;
-			}
+			UE_CONTROLRIG_RIGUNIT_REPORT_WARNING(TEXT("Space '%s' is not valid."), *SpaceName.ToString());
+			return;
 		}
+
+		FTransform InitialTransform = Transform;
+		if (Space == EBoneGetterSetterMode::GlobalSpace)
+		{
+			const FTransform ParentTransform = Hierarchy->GetParentTransformByIndex(CachedSpaceIndex, true);
+			InitialTransform = InitialTransform.GetRelativeTransform(ParentTransform);
+		}
+
+		Hierarchy->SetInitialLocalTransform(CachedSpaceIndex, InitialTransform);
 	}
 }
 
diff --git a/Engine/Plugins/Animation/ControlRig/Source/ControlRig/Private/Units/Hierarchy/RigUnit_SetSpaceTransform.cpp b/Engine/Plugins/Animation/ControlRig/Source/ControlRig/Private/Units/Hierarchy/RigUnit_SetSpaceTransform.cpp
--- a/Engine/Plugins/Animation/ControlRig/Source/ControlRig/Private/Units/Hierarchy/RigUnit_SetSpaceTransform.cpp
+++ b/Engine/Plugins/Animation/ControlRig/Source/ControlRig/Private/Units/Hierarchy/RigUnit_SetSpaceTransform.cpp
@@ -13,57 +13,42 @@ FRigUnit_SetSpaceTransform_Execute()
 	URigHierarchy* Hierarchy = ExecuteContext.Hierarchy;
 	if (Hierarchy)
 	{
-		switch (ExecuteContext.UnitContext.State)
+		const FRigElementKey SpaceKey(Space, ERigElementType::Null);
+
+		if (CachedSpaceIndex.UpdateCache(SpaceKey, Hierarchy))
 		{
-			case EControlRigState::Init:
-			{
-				CachedSpaceIndex.Reset();
-				break;
-			}
-			case EControlRigState::Update:
+			switch (SpaceType)
 			{
-				const FRigElementKey SpaceKey(Space, ERigElementType::Null);
-
-				if (CachedSpaceIndex.UpdateCache(SpaceKey, Hierarchy))
+				case EBoneGetterSetterMode::GlobalSpace:
 				{
-					switch (SpaceType)
+					if(FMath::IsNearlyEqual(Weight, 1.f))
 					{
-						case EBoneGetterSetterMode::GlobalSpace:
-						{
-							if(FMath::IsNearlyEqual(Weight, 1.f))
-							{
-								Hierarchy->SetGlobalTransform(CachedSpaceIndex, Transform, true);
-							}
-							else
-							{
-								const FTransform PreviousTransform = Hierarchy->GetGlobalTransform(CachedSpaceIndex);
-								Hierarchy->SetGlobalTransform(CachedSpaceIndex, FControlRigMathLibrary::LerpTransform(PreviousTransform, Transform, FMath::Clamp<float>(Weight, 0.f, 1.f)), true);
-							}
-							break;
-						}
-						case EBoneGetterSetterMode::LocalSpace:
-						{
-							if(FMath::IsNearlyEqual(Weight, 1.f))
-							{
-								Hierarchy->SetLocalTransform(CachedSpaceIndex, Transform, true);
-							}
-							else
-							{
-								const FTransform PreviousTransform = Hierarchy->GetLocalTransform(CachedSpaceIndex);
-								Hierarchy->SetLocalTransform(CachedSpaceIndex, FControlRigMathLibrary::LerpTransform(PreviousTransform, Transform, FMath::Clamp<float>(Weight, 0.f, 1.f)), true);
-							}
-							break;
-						}
-						default:
-						{
-							break;
-						}
+						Hierarchy->SetGlobalTransform(CachedSpaceIndex, Transform, true);
 					}
+					else
+					{
+						const FTransform PreviousTransform = Hierarchy->GetGlobalTransform(CachedSpaceIndex);
+						Hierarchy->SetGlobalTransform(CachedSpaceIndex, FControlRigMathLibrary::LerpTransform(PreviousTransform, Transform, FMath::Clamp<float>(Weight, 0.f, 1.f)), true);
+					}
+					break;
+				}
+				case EBoneGetterSetterMode::LocalSpace:
+				{
+					if(FMath::IsNearlyEqual(Weight, 1.f))
+					{
+						Hierarchy->SetLocalTransform(CachedSpaceIndex, Transform, true);
+					}
+					else
+					{
+						const FTransform PreviousTransform = Hierarchy->GetLocalTransform(CachedSpaceIndex);
+						Hierarchy->SetLocalTransform(CachedSpaceIndex, FControlRigMathLibrary::LerpTransform(PreviousTransform, Transform, FMath::Clamp<float>(Weight, 0.f, 1.f)), true);
+					}
+					break;
+				}
+				default:
+				{
+					break;
 				}
-			}
-			default:
-			{
-				break;
 			}
 		}
 	}
diff --git a/Engine/Plugins/Animation/ControlRig/Source/ControlRig/Private/Units/Hierarchy/RigUnit_SetTransform.cpp b/Engine/Plugins/Animation/ControlRig/Source/ControlRig/Private/Units/Hierarchy/RigUnit_SetTransform.cpp
--- a/Engine/Plugins/Animation/ControlRig/Source/ControlRig/Private/Units/Hierarchy/RigUnit_SetTransform.cpp
+++ b/Engine/Plugins/Animation/ControlRig/Source/ControlRig/Private/Units/Hierarchy/RigUnit_SetTransform.cpp
@@ -27,103 +27,87 @@ FRigUnit_SetTransform_Execute()
 
 	if (URigHierarchy* Hierarchy = ExecuteContext.Hierarchy)
 	{
-		switch (ExecuteContext.UnitContext.State)
+		if (!CachedIndex.UpdateCache(Item, Hierarchy))
 		{
-			case EControlRigState::Init:
+			UE_CONTROLRIG_RIGUNIT_REPORT_WARNING(TEXT("Item '%s' is not valid."), *Item.ToString());
+		}
+		else
+		{
+#if WITH_EDITOR
+			if(bInitial)
 			{
-				CachedIndex.Reset();
-
-				return;
+				// provide some user feedback when changing initial transforms during forward solve
+				if(ExecuteContext.GetEventName() == FRigUnit_BeginExecution::EventName)
+				{
+					UE_CONTROLRIG_RIGUNIT_LOG_MESSAGE(TEXT("Changing initial transforms during %s is not recommended."), *ExecuteContext.GetEventName().ToString());
+				}
 			}
-			case EControlRigState::Update:
+#endif
+
+			// for controls - set the control offset transform instead
+			if(bInitial && (CachedIndex.GetKey().Type == ERigElementType::Control))
 			{
-				if (!CachedIndex.UpdateCache(Item, Hierarchy))
+				FTransform TransformMutable = Value;
+				FRigUnit_SetControlOffset::StaticExecute(ExecuteContext, CachedIndex.GetKey().Name, TransformMutable, Space, CachedIndex);
+				
+				if (ExecuteContext.GetEventName() == FRigUnit_PrepareForExecution::EventName)
 				{
-					UE_CONTROLRIG_RIGUNIT_REPORT_WARNING(TEXT("Item '%s' is not valid."), *Item.ToString());
+					Hierarchy->SetLocalTransformByIndex(CachedIndex, FTransform::Identity, true, bPropagateToChildren);
+					Hierarchy->SetLocalTransformByIndex(CachedIndex, FTransform::Identity, false, bPropagateToChildren);
 				}
-				else
+				return;
+			}
+			
+			FTransform WeightedTransform = Value;
+			if (Weight < 1.f - SMALL_NUMBER)
+			{
+				FTransform PreviousTransform = WeightedTransform;
+				switch (Space)
 				{
-#if WITH_EDITOR
-					if(bInitial)
+					case EBoneGetterSetterMode::GlobalSpace:
 					{
-						// provide some user feedback when changing initial transforms during forward solve
-						if(ExecuteContext.GetEventName() == FRigUnit_BeginExecution::EventName)
-						{
-							UE_CONTROLRIG_RIGUNIT_LOG_MESSAGE(TEXT("Changing initial transforms during %s is not recommended."), *ExecuteContext.GetEventName().ToString());
-						}
+						PreviousTransform = Hierarchy->GetGlobalTransformByIndex(CachedIndex, bInitial);
+						break;
 					}
-#endif
-
-					// for controls - set the control offset transform instead
-					if(bInitial && (CachedIndex.GetKey().Type == ERigElementType::Control))
+					case EBoneGetterSetterMode::LocalSpace:
 					{
-						FTransform TransformMutable = Value;
-						FRigUnit_SetControlOffset::StaticExecute(ExecuteContext, CachedIndex.GetKey().Name, TransformMutable, Space, CachedIndex);
-						
-						if (ExecuteContext.GetEventName() == FRigUnit_PrepareForExecution::EventName)
-						{
-							Hierarchy->SetLocalTransformByIndex(CachedIndex, FTransform::Identity, true, bPropagateToChildren);
-							Hierarchy->SetLocalTransformByIndex(CachedIndex, FTransform::Identity, false, bPropagateToChildren);
-						}
-						return;
+						PreviousTransform = Hierarchy->GetLocalTransformByIndex(CachedIndex, bInitial);
+						break;
 					}
-					
-					FTransform WeightedTransform = Value;
-					if (Weight < 1.f - SMALL_NUMBER)
+					default:
 					{
-						FTransform PreviousTransform = WeightedTransform;
-						switch (Space)
-						{
-							case EBoneGetterSetterMode::GlobalSpace:
-							{
-								PreviousTransform = Hierarchy->GetGlobalTransformByIndex(CachedIndex, bInitial);
-								break;
-							}
-							case EBoneGetterSetterMode::LocalSpace:
-							{
-								PreviousTransform = Hierarchy->GetLocalTransformByIndex(CachedIndex, bInitial);
-								break;
-							}
-							default:
-							{
-								break;
-							}
-						}
-						WeightedTransform = FControlRigMathLibrary::LerpTransform(PreviousTransform, WeightedTransform, Weight);
+						break;
 					}
+				}
+				WeightedTransform = FControlRigMathLibrary::LerpTransform(PreviousTransform, WeightedTransform, Weight);
+			}
 
-					switch (Space)
+			switch (Space)
+			{
+				case EBoneGetterSetterMode::GlobalSpace:
+				{
+					Hierarchy->SetGlobalTransformByIndex(CachedIndex, WeightedTransform, bInitial, bPropagateToChildren);
+
+					if (bInitial && ExecuteContext.GetEventName() == FRigUnit_PrepareForExecution::EventName)
 					{
-						case EBoneGetterSetterMode::GlobalSpace:
-						{
-							Hierarchy->SetGlobalTransformByIndex(CachedIndex, WeightedTransform, bInitial, bPropagateToChildren);
-
-							if (bInitial && ExecuteContext.GetEventName() == FRigUnit_PrepareForExecution::EventName)
-							{
-								Hierarchy->SetGlobalTransformByIndex(CachedIndex, WeightedTransform, false, bPropagateToChildren);
-							}
-							break;
-						}
-						case EBoneGetterSetterMode::LocalSpace:
-						{
-							Hierarchy->SetLocalTransformByIndex(CachedIndex, WeightedTransform, bInitial, bPropagateToChildren);
-
-							if (bInitial && ExecuteContext.GetEventName() == FRigUnit_PrepareForExecution::EventName)
-							{
-								Hierarchy->SetLocalTransformByIndex(CachedIndex, WeightedTransform, false, bPropagateToChildren);
-							}
-							break;
-						}
-						default:
-						{
-							break;
-						}
+						Hierarchy->SetGlobalTransformByIndex(CachedIndex, WeightedTransform, false, bPropagateToChildren);
 					}
+					break;
+				}
+				case EBoneGetterSetterMode::LocalSpace:
+				{
+					Hierarchy->SetLocalTransformByIndex(CachedIndex, WeightedTransform, bInitial, bPropagateToChildren);
+
+					if (bInitial && ExecuteContext.GetEventName() == FRigUnit_PrepareForExecution::EventName)
+					{
+						Hierarchy->SetLocalTransformByIndex(CachedIndex, WeightedTransform, false, bPropagateToChildren);
+					}
+					break;
+				}
+				default:
+				{
+					break;
 				}
-			}
-			default:
-			{
-				break;
 			}
 		}
 	}
diff --git a/Engine/Plugins/Animation/ControlRig/Source/ControlRig/Private/Units/Hierarchy/RigUnit_UnsetCurveValue.cpp b/Engine/Plugins/Animation/ControlRig/Source/ControlRig/Private/Units/Hierarchy/RigUnit_UnsetCurveValue.cpp
--- a/Engine/Plugins/Animation/ControlRig/Source/ControlRig/Private/Units/Hierarchy/RigUnit_UnsetCurveValue.cpp
+++ b/Engine/Plugins/Animation/ControlRig/Source/ControlRig/Private/Units/Hierarchy/RigUnit_UnsetCurveValue.cpp
@@ -11,25 +11,10 @@ FRigUnit_UnsetCurveValue_Execute()
 	URigHierarchy* Hierarchy = ExecuteContext.Hierarchy;
 	if (Hierarchy)
 	{
-		switch (ExecuteContext.UnitContext.State)
+		const FRigElementKey Key(Curve, ERigElementType::Curve);
+		if (CachedCurveIndex.UpdateCache(Key, Hierarchy))
 		{
-			case EControlRigState::Init:
-			{
-				CachedCurveIndex.Reset();
-				// fall through to update
-			}
-			case EControlRigState::Update:
-			{
-				const FRigElementKey Key(Curve, ERigElementType::Curve);
-				if (CachedCurveIndex.UpdateCache(Key, Hierarchy))
-				{
-					Hierarchy->UnsetCurveValueByIndex(CachedCurveIndex);
-				}
-			}
-			default:
-			{
-				break;
-			}
+			Hierarchy->UnsetCurveValueByIndex(CachedCurveIndex);
 		}
 	}
 }
@@ -44,7 +29,7 @@ IMPLEMENT_RIGUNIT_AUTOMATION_TEST(FRigUnit_UnsetCurveValue)
 	
 	Hierarchy->ResetCurveValues();
 	Unit.Curve = TEXT("Curve");
-	InitAndExecute();
+	Execute();
 
 	AddErrorIfFalse(!Hierarchy->IsCurveValueSet(Curve), TEXT("unexpected value"));
 
diff --git a/Engine/Plugins/Animation/ControlRig/Source/ControlRig/Private/Units/Highlevel/Harmonics/RigUnit_BoneHarmonics.cpp b/Engine/Plugins/Animation/ControlRig/Source/ControlRig/Private/Units/Highlevel/Harmonics/RigUnit_BoneHarmonics.cpp
--- a/Engine/Plugins/Animation/ControlRig/Source/ControlRig/Private/Units/Highlevel/Harmonics/RigUnit_BoneHarmonics.cpp
+++ b/Engine/Plugins/Animation/ControlRig/Source/ControlRig/Private/Units/Highlevel/Harmonics/RigUnit_BoneHarmonics.cpp
@@ -77,12 +77,10 @@ FRigUnit_ItemHarmonics_Execute()
 	TArray<FCachedRigElement>& CachedItems = WorkData.CachedItems;
 	FVector& WaveTime = WorkData.WaveTime;
 
-	if (ExecuteContext.UnitContext.State == EControlRigState::Init ||
-		CachedItems.Num() != Targets.Num())
+	if (CachedItems.Num() != Targets.Num())
 	{
 		CachedItems.Reset();
 		WaveTime = FVector::ZeroVector;
-		return;
 	}
 
 	for (int32 ItemIndex = 0; ItemIndex < Targets.Num(); ItemIndex++)
@@ -126,6 +124,6 @@ FRigUnit_ItemHarmonics_Execute()
 		}
 	}
 
-	WaveTime += WaveSpeed * ExecuteContext.UnitContext.DeltaTime;
+	WaveTime += WaveSpeed * ExecuteContext.GetDeltaTime();
 }
 
diff --git a/Engine/Plugins/Animation/ControlRig/Source/ControlRig/Private/Units/Highlevel/Harmonics/RigUnit_ChainHarmonics.cpp b/Engine/Plugins/Animation/ControlRig/Source/ControlRig/Private/Units/Highlevel/Harmonics/RigUnit_ChainHarmonics.cpp
--- a/Engine/Plugins/Animation/ControlRig/Source/ControlRig/Private/Units/Highlevel/Harmonics/RigUnit_ChainHarmonics.cpp
+++ b/Engine/Plugins/Animation/ControlRig/Source/ControlRig/Private/Units/Highlevel/Harmonics/RigUnit_ChainHarmonics.cpp
@@ -57,13 +57,6 @@ FRigUnit_ChainHarmonicsPerItem_Execute()
 	TArray<FVector>& HierarchyLine = WorkData.HierarchyLine;
 	TArray<FVector>& VelocityLines = WorkData.VelocityLines;
 	
-	if (ExecuteContext.UnitContext.State == EControlRigState::Init)
-	{
-		Time = FVector::ZeroVector;
-		Items.Reset();
-		return;
-	}
-
 	if(Items.Num() == 0)
 	{
 		if (!ChainRoot.IsValid())
@@ -188,15 +181,15 @@ FRigUnit_ChainHarmonicsPerItem_Execute()
 			FVector Velocity = Pendulum.PendulumGravity;
 			Velocity += Stiffness * Pendulum.PendulumStiffness;
 
-			if (ExecuteContext.UnitContext.DeltaTime > 0.f)
+			if (ExecuteContext.GetDeltaTime() > 0.f)
 			{
 				PendulumVelocity[Index] = FMath::Lerp<FVector>(PendulumVelocity[Index], Velocity, FMath::Clamp<float>(Pendulum.PendulumBlend, 0.f, 0.999f));
 				PendulumVelocity[Index] = PendulumVelocity[Index] * Pendulum.PendulumDrag;
 
 				FVector PrevPosition = PendulumPosition[Index];
-				PendulumPosition[Index] = PendulumPosition[Index] + PendulumVelocity[Index] * ExecuteContext.UnitContext.DeltaTime;
+				PendulumPosition[Index] = PendulumPosition[Index] + PendulumVelocity[Index] * ExecuteContext.GetDeltaTime();
 				PendulumPosition[Index] = GlobalTransform.GetLocation() + (PendulumPosition[Index] - GlobalTransform.GetLocation()).GetSafeNormal() * Length;
-				PendulumVelocity[Index] = (PendulumPosition[Index] - PrevPosition) / ExecuteContext.UnitContext.DeltaTime;
+				PendulumVelocity[Index] = (PendulumPosition[Index] - PrevPosition) / ExecuteContext.GetDeltaTime();
 			}
 
 			VelocityLines[Index * 2 + 0] = PendulumPosition[Index];
@@ -220,7 +213,7 @@ FRigUnit_ChainHarmonicsPerItem_Execute()
 		ParentTransform = GlobalTransform;
 	}
 
-	Time = Time + Speed * ExecuteContext.UnitContext.DeltaTime;
+	Time = Time + Speed * ExecuteContext.GetDeltaTime();
 
 	if (ExecuteContext.UnitContext.DrawInterface != nullptr && bDrawDebug)
 	{
diff --git a/Engine/Plugins/Animation/ControlRig/Source/ControlRig/Private/Units/Highlevel/Hierarchy/RigUnit_AimBone.cpp b/Engine/Plugins/Animation/ControlRig/Source/ControlRig/Private/Units/Highlevel/Hierarchy/RigUnit_AimBone.cpp
--- a/Engine/Plugins/Animation/ControlRig/Source/ControlRig/Private/Units/Highlevel/Hierarchy/RigUnit_AimBone.cpp
+++ b/Engine/Plugins/Animation/ControlRig/Source/ControlRig/Private/Units/Highlevel/Hierarchy/RigUnit_AimBone.cpp
@@ -19,11 +19,11 @@ FRigUnit_AimBoneMath_Execute()
 		return;
 	}
 
-	if (ExecuteContext.UnitContext.State == EControlRigState::Init)
+	if (!bIsInitialized)
 	{
 		PrimaryCachedSpace.Reset();
 		SecondaryCachedSpace.Reset();
-		return;
+		bIsInitialized = true;
 	}
 
 	if ((Weight <= SMALL_NUMBER) || (Primary.Weight <= SMALL_NUMBER && Secondary.Weight <= SMALL_NUMBER))
@@ -183,7 +183,8 @@ FRigUnit_AimBone_Execute()
 		DebugSettings,
 		CachedBoneIndex,
 		PrimaryCachedSpace,
-		SecondaryCachedSpace);
+		SecondaryCachedSpace,
+		bIsInitialized);
 }
 
 FRigVMStructUpgradeInfo FRigUnit_AimBone::GetUpgradeInfo() const
@@ -219,12 +220,12 @@ FRigUnit_AimItem_Execute()
 		return;
 	}
 
-	if (ExecuteContext.UnitContext.State == EControlRigState::Init)
+	if (!bIsInitialized)
 	{
 		CachedItem.Reset();
 		PrimaryCachedSpace.Reset();
 		SecondaryCachedSpace.Reset();
-		return;
+		bIsInitialized = true;
 	}
 
 	if (!CachedItem.UpdateCache(Item, Hierarchy))
@@ -249,7 +250,8 @@ FRigUnit_AimItem_Execute()
 		Transform,
 		DebugSettings,
 		PrimaryCachedSpace,
-		SecondaryCachedSpace);
+		SecondaryCachedSpace,
+		bIsInitialized);
 
 	Hierarchy->SetGlobalTransform(CachedItem, Transform);
 }
@@ -264,11 +266,12 @@ FRigUnit_AimConstraintLocalSpaceOffset_Execute()
 		return;
 	}
 	
-	if(ExecuteContext.UnitContext.State == EControlRigState::Init)
+	if(!bIsInitialized)
 	{
 		WorldUpSpaceCache.Reset();
 		ChildCache.Reset();
 		ParentCaches.Reset();
+		bIsInitialized = true;
 	}
 	
 	URigHierarchy* Hierarchy = ExecuteContext.Hierarchy;
@@ -390,7 +393,8 @@ FRigUnit_AimConstraintLocalSpaceOffset_Execute()
 					InitialAimResult,
 					DummyDebugSettings,
 					PrimaryCachedSpace,
-					SecondaryCachedSpace);
+					SecondaryCachedSpace,
+					bIsInitialized);
 
 				FTransform ChildParentInitialGlobalTransform = Hierarchy->GetParentTransformByIndex(ChildCache, true);
 				FQuat MixedInitialLocalRotation = ChildParentInitialGlobalTransform.GetRotation().Inverse() * InitialAimResult.GetRotation();
@@ -458,7 +462,8 @@ FRigUnit_AimConstraintLocalSpaceOffset_Execute()
 				AimResult,
 				AdvancedSettings.DebugSettings,
 				PrimaryCachedSpace,
-				WorldUpSpaceCache);	
+				WorldUpSpaceCache,
+				bIsInitialized);	
 
 			// handle filtering, performed in local space
 			FTransform ChildParentGlobalTransform = Hierarchy->GetParentTransformByIndex(ChildCache, false);
