ControlRig: Remove EControlRigState from units

diff --git a/Engine/Plugins/Animation/ControlRig/Source/ControlRig/Private/Units/Highlevel/Hierarchy/RigUnit_AimBone.h b/Engine/Plugins/Animation/ControlRig/Source/ControlRig/Private/Units/Highlevel/Hierarchy/RigUnit_AimBone.h
--- a/Engine/Plugins/Animation/ControlRig/Source/ControlRig/Private/Units/Highlevel/Hierarchy/RigUnit_AimBone.h
+++ b/Engine/Plugins/Animation/ControlRig/Source/ControlRig/Private/Units/Highlevel/Hierarchy/RigUnit_AimBone.h
@@ -149,6 +149,7 @@ struct CONTROLRIG_API FRigUnit_AimBoneMath : public FRigUnit_HighlevelBase
 		DebugSettings = FRigUnit_AimBone_DebugSettings();
 		PrimaryCachedSpace = FCachedRigElement();
 		SecondaryCachedSpace = FCachedRigElement();
+		bIsInitialized = false;
 	}
 
 	RIGVM_METHOD()
@@ -193,6 +194,9 @@ struct CONTROLRIG_API FRigUnit_AimBoneMath : public FRigUnit_HighlevelBase
 
 	UPROPERTY()
 	FCachedRigElement SecondaryCachedSpace;
+
+	UPROPERTY()
+	bool bIsInitialized;
 };
 
 /**
@@ -217,6 +221,7 @@ struct CONTROLRIG_API FRigUnit_AimBone : public FRigUnit_HighlevelBaseMutable
 		CachedBoneIndex = FCachedRigElement();
 		PrimaryCachedSpace = FCachedRigElement();
 		SecondaryCachedSpace = FCachedRigElement();
+		bIsInitialized = false;
 	}
 
 	virtual FRigElementKey DetermineSpaceForPin(const FString& InPinPath, void* InUserContext) const override
@@ -280,6 +285,9 @@ struct CONTROLRIG_API FRigUnit_AimBone : public FRigUnit_HighlevelBaseMutable
 	UPROPERTY()
 	FCachedRigElement SecondaryCachedSpace;
 
+	UPROPERTY()
+	bool bIsInitialized;
+
 	RIGVM_METHOD()
 	virtual FRigVMStructUpgradeInfo GetUpgradeInfo() const override;
 };
@@ -305,6 +313,7 @@ struct CONTROLRIG_API FRigUnit_AimItem: public FRigUnit_HighlevelBaseMutable
 		CachedItem = FCachedRigElement();
 		PrimaryCachedSpace = FCachedRigElement();
 		SecondaryCachedSpace = FCachedRigElement();
+		bIsInitialized = false;
 	}
 
 	virtual FRigElementKey DetermineSpaceForPin(const FString& InPinPath, void* InUserContext) const override
@@ -359,6 +368,9 @@ struct CONTROLRIG_API FRigUnit_AimItem: public FRigUnit_HighlevelBaseMutable
 
 	UPROPERTY()
 	FCachedRigElement SecondaryCachedSpace;
+
+	UPROPERTY()
+	bool bIsInitialized;
 };
 
 
@@ -434,6 +446,7 @@ struct CONTROLRIG_API FRigUnit_AimConstraintLocalSpaceOffset: public FRigUnit_Hi
 		, WorldUp(FRigUnit_AimConstraint_WorldUp())
 		, AdvancedSettings(FRigUnit_AimConstraint_AdvancedSettings())
 		, Weight(1.0f)
+		, bIsInitialized(false)
 	{
 		Parents.Add(FConstraintParent());
 	};
@@ -496,4 +509,6 @@ struct CONTROLRIG_API FRigUnit_AimConstraintLocalSpaceOffset: public FRigUnit_Hi
 	UPROPERTY()
 	TArray<FCachedRigElement> ParentCaches;
 	
+	UPROPERTY()
+	bool bIsInitialized;
 };
diff --git a/Engine/Plugins/Animation/ControlRig/Source/ControlRig/Private/Units/Highlevel/Hierarchy/RigUnit_CCDIK.cpp b/Engine/Plugins/Animation/ControlRig/Source/ControlRig/Private/Units/Highlevel/Hierarchy/RigUnit_CCDIK.cpp
--- a/Engine/Plugins/Animation/ControlRig/Source/ControlRig/Private/Units/Highlevel/Hierarchy/RigUnit_CCDIK.cpp
+++ b/Engine/Plugins/Animation/ControlRig/Source/ControlRig/Private/Units/Highlevel/Hierarchy/RigUnit_CCDIK.cpp
@@ -7,12 +7,6 @@
 
 FRigUnit_CCDIK_Execute()
 {
-	if (ExecuteContext.UnitContext.State == EControlRigState::Init)
-	{
-		WorkData.CachedItems.Reset();
-		return;
-	}
-
 	FRigElementKeyCollection Items;
 	if(WorkData.CachedItems.Num() == 0)
 	{
@@ -90,111 +84,104 @@ FRigUnit_CCDIKItemArray_Execute()
 	TArray<float>& RotationLimitsPerItem = WorkData.RotationLimitsPerItem;
 	FCachedRigElement& CachedEffector = WorkData.CachedEffector;
 
-	if ((ExecuteContext.UnitContext.State == EControlRigState::Init) ||
-		(
-			(RotationLimits.Num() != RotationLimitIndex.Num()) &&
-			(RotationLimitIndex.Num() > 0))
-		)
+	if ((RotationLimits.Num() != RotationLimitIndex.Num()) &&
+		(RotationLimitIndex.Num() > 0))
 	{
 		CachedItems.Reset();
 		RotationLimitIndex.Reset();
 		RotationLimitsPerItem.Reset();
 		CachedEffector.Reset();
-		return;
 	}
 	
-	if (ExecuteContext.UnitContext.State == EControlRigState::Update)
+	if (CachedItems.Num() == 0 && Items.Num() > 0)
 	{
-		if (CachedItems.Num() == 0 && Items.Num() > 0)
-		{
-			CachedItems.Add(FCachedRigElement(Hierarchy->GetFirstParent(Items[0]), Hierarchy));
+		CachedItems.Add(FCachedRigElement(Hierarchy->GetFirstParent(Items[0]), Hierarchy));
 
-			for (FRigElementKey Item : Items)
-			{
-				CachedItems.Add(FCachedRigElement(Item, Hierarchy));
-			}
+		for (FRigElementKey Item : Items)
+		{
+			CachedItems.Add(FCachedRigElement(Item, Hierarchy));
+		}
 
-			CachedEffector = CachedItems.Last();
-			Chain.Reserve(CachedItems.Num());
+		CachedEffector = CachedItems.Last();
+		Chain.Reserve(CachedItems.Num());
 
-			RotationLimitsPerItem.SetNumUninitialized(CachedItems.Num());
-			for (const FRigUnit_CCDIK_RotationLimitPerItem& RotationLimit : RotationLimits)
-			{
-				FCachedRigElement BoneIndex(RotationLimit.Item, Hierarchy);
-				int32 BoneIndexInLookup = CachedItems.Find(BoneIndex);
-				RotationLimitIndex.Add(BoneIndexInLookup);
-			}
+		RotationLimitsPerItem.SetNumUninitialized(CachedItems.Num());
+		for (const FRigUnit_CCDIK_RotationLimitPerItem& RotationLimit : RotationLimits)
+		{
+			FCachedRigElement BoneIndex(RotationLimit.Item, Hierarchy);
+			int32 BoneIndexInLookup = CachedItems.Find(BoneIndex);
+			RotationLimitIndex.Add(BoneIndexInLookup);
+		}
 
-			if (CachedItems.Num() < 2)
-			{
-				UE_CONTROLRIG_RIGUNIT_REPORT_WARNING(TEXT("No bones found."));
-			}
+		if (CachedItems.Num() < 2)
+		{
+			UE_CONTROLRIG_RIGUNIT_REPORT_WARNING(TEXT("No bones found."));
 		}
+	}
 
-		if (CachedItems.Num() > 1)
+	if (CachedItems.Num() > 1)
+	{
+		// Gather chain links. These are non zero length bones.
+		Chain.Reset();
+		
+		for (int32 ChainIndex = 0; ChainIndex < CachedItems.Num(); ChainIndex++)
 		{
-			// Gather chain links. These are non zero length bones.
-			Chain.Reset();
-			
-			for (int32 ChainIndex = 0; ChainIndex < CachedItems.Num(); ChainIndex++)
-			{
-				const FTransform& GlobalTransform = Hierarchy->GetGlobalTransform(CachedItems[ChainIndex]);
-				const FTransform& LocalTransform = Hierarchy->GetLocalTransform(CachedItems[ChainIndex]);
-				Chain.Add(FCCDIKChainLink(GlobalTransform, LocalTransform, ChainIndex));
-			}
+			const FTransform& GlobalTransform = Hierarchy->GetGlobalTransform(CachedItems[ChainIndex]);
+			const FTransform& LocalTransform = Hierarchy->GetLocalTransform(CachedItems[ChainIndex]);
+			Chain.Add(FCCDIKChainLink(GlobalTransform, LocalTransform, ChainIndex));
+		}
 
-			for (float& Limit : RotationLimitsPerItem)
-			{
-				Limit = BaseRotationLimit;
-			}
-			
-			for (int32 LimitIndex = 0; LimitIndex < RotationLimitIndex.Num(); LimitIndex++)
+		for (float& Limit : RotationLimitsPerItem)
+		{
+			Limit = BaseRotationLimit;
+		}
+		
+		for (int32 LimitIndex = 0; LimitIndex < RotationLimitIndex.Num(); LimitIndex++)
+		{
+			if (RotationLimitIndex[LimitIndex] != INDEX_NONE)
 			{
-				if (RotationLimitIndex[LimitIndex] != INDEX_NONE)
-				{
-					RotationLimitsPerItem[RotationLimitIndex[LimitIndex]] = RotationLimits[LimitIndex].Limit;
-				}
+				RotationLimitsPerItem[RotationLimitIndex[LimitIndex]] = RotationLimits[LimitIndex].Limit;
 			}
+		}
 
-			bool bBoneLocationUpdated = AnimationCore::SolveCCDIK(Chain, EffectorTransform.GetLocation(), Precision, MaxIterations, bStartFromTail, RotationLimits.Num() > 0, RotationLimitsPerItem);
+		bool bBoneLocationUpdated = AnimationCore::SolveCCDIK(Chain, EffectorTransform.GetLocation(), Precision, MaxIterations, bStartFromTail, RotationLimits.Num() > 0, RotationLimitsPerItem);
 
-			// If we moved some bones, update bone transforms.
-			if (bBoneLocationUpdated)
+		// If we moved some bones, update bone transforms.
+		if (bBoneLocationUpdated)
+		{
+			if (FMath::IsNearlyEqual(Weight, 1.f))
 			{
-				if (FMath::IsNearlyEqual(Weight, 1.f))
+				for (int32 LinkIndex = 0; LinkIndex < CachedItems.Num(); LinkIndex++)
 				{
-					for (int32 LinkIndex = 0; LinkIndex < CachedItems.Num(); LinkIndex++)
-					{
-						const FCCDIKChainLink& CurrentLink = Chain[LinkIndex];
-						Hierarchy->SetGlobalTransform(CachedItems[LinkIndex], CurrentLink.Transform, bPropagateToChildren);
-					}
-				}
-				else
-				{
-					float T = FMath::Clamp<float>(Weight, 0.f, 1.f);
-
-					for (int32 LinkIndex = 0; LinkIndex < CachedItems.Num(); LinkIndex++)
-					{
-						const FCCDIKChainLink& CurrentLink = Chain[LinkIndex];
-						FTransform PreviousXfo = Hierarchy->GetGlobalTransform(CachedItems[LinkIndex]);
-						FTransform Xfo = FControlRigMathLibrary::LerpTransform(PreviousXfo, CurrentLink.Transform, T);
-						Hierarchy->SetGlobalTransform(CachedItems[LinkIndex], Xfo, bPropagateToChildren);
-					}
+					const FCCDIKChainLink& CurrentLink = Chain[LinkIndex];
+					Hierarchy->SetGlobalTransform(CachedItems[LinkIndex], CurrentLink.Transform, bPropagateToChildren);
 				}
 			}
-
-			if (FMath::IsNearlyEqual(Weight, 1.f))
-			{
-				Hierarchy->SetGlobalTransform(CachedEffector, EffectorTransform, bPropagateToChildren);
-			}
 			else
 			{
 				float T = FMath::Clamp<float>(Weight, 0.f, 1.f);
-				FTransform PreviousXfo = Hierarchy->GetGlobalTransform(CachedEffector);
-				FTransform Xfo = FControlRigMathLibrary::LerpTransform(PreviousXfo, EffectorTransform, T);
-				Hierarchy->SetGlobalTransform(CachedEffector, Xfo, bPropagateToChildren);
+
+				for (int32 LinkIndex = 0; LinkIndex < CachedItems.Num(); LinkIndex++)
+				{
+					const FCCDIKChainLink& CurrentLink = Chain[LinkIndex];
+					FTransform PreviousXfo = Hierarchy->GetGlobalTransform(CachedItems[LinkIndex]);
+					FTransform Xfo = FControlRigMathLibrary::LerpTransform(PreviousXfo, CurrentLink.Transform, T);
+					Hierarchy->SetGlobalTransform(CachedItems[LinkIndex], Xfo, bPropagateToChildren);
+				}
 			}
 		}
+
+		if (FMath::IsNearlyEqual(Weight, 1.f))
+		{
+			Hierarchy->SetGlobalTransform(CachedEffector, EffectorTransform, bPropagateToChildren);
+		}
+		else
+		{
+			float T = FMath::Clamp<float>(Weight, 0.f, 1.f);
+			FTransform PreviousXfo = Hierarchy->GetGlobalTransform(CachedEffector);
+			FTransform Xfo = FControlRigMathLibrary::LerpTransform(PreviousXfo, EffectorTransform, T);
+			Hierarchy->SetGlobalTransform(CachedEffector, Xfo, bPropagateToChildren);
+		}
 	}
 }
 
diff --git a/Engine/Plugins/Animation/ControlRig/Source/ControlRig/Private/Units/Highlevel/Hierarchy/RigUnit_DistributeRotation.cpp b/Engine/Plugins/Animation/ControlRig/Source/ControlRig/Private/Units/Highlevel/Hierarchy/RigUnit_DistributeRotation.cpp
--- a/Engine/Plugins/Animation/ControlRig/Source/ControlRig/Private/Units/Highlevel/Hierarchy/RigUnit_DistributeRotation.cpp
+++ b/Engine/Plugins/Animation/ControlRig/Source/ControlRig/Private/Units/Highlevel/Hierarchy/RigUnit_DistributeRotation.cpp
@@ -106,14 +106,13 @@ FRigUnit_DistributeRotationForItemArray_Execute()
 		}
 	}
 
-	if (ExecuteContext.UnitContext.State == EControlRigState::Init || (CachedItems.Num() > 0 && CachedItems.Num() != Items.Num()))
+	if (CachedItems.Num() > 0 && CachedItems.Num() != Items.Num())
 	{
 		CachedItems.Reset();
 		ItemRotationA.Reset();
 		ItemRotationB.Reset();
 		ItemRotationT.Reset();
 		ItemLocalTransforms.Reset();
-		return;
 	}
 
 	if (CachedItems.Num() == 0)
@@ -295,7 +294,6 @@ IMPLEMENT_RIGUNIT_AUTOMATION_TEST(FRigUnit_DistributeRotation)
 	Rotation.Ratio = 0.5f;
 	Unit.Rotations.Add(Rotation);
 
-	Init();
 	Execute();
 
 	AddErrorIfFalse(Unit.WorkData.ItemRotationA[0] == 0, TEXT("unexpected bone a"));
diff --git a/Engine/Plugins/Animation/ControlRig/Source/ControlRig/Private/Units/Highlevel/Hierarchy/RigUnit_FABRIK.cpp b/Engine/Plugins/Animation/ControlRig/Source/ControlRig/Private/Units/Highlevel/Hierarchy/RigUnit_FABRIK.cpp
--- a/Engine/Plugins/Animation/ControlRig/Source/ControlRig/Private/Units/Highlevel/Hierarchy/RigUnit_FABRIK.cpp
+++ b/Engine/Plugins/Animation/ControlRig/Source/ControlRig/Private/Units/Highlevel/Hierarchy/RigUnit_FABRIK.cpp
@@ -7,12 +7,6 @@
 
 FRigUnit_FABRIK_Execute()
 {
-	if (ExecuteContext.UnitContext.State == EControlRigState::Init)
-	{
-		WorkData.CachedItems.Reset();
-		return;
-	}
-	
 	FRigElementKeyCollection Items;
 	if(WorkData.CachedItems.Num() == 0)
 	{
@@ -73,133 +67,123 @@ FRigUnit_FABRIKItemArray_Execute()
 	TArray<FCachedRigElement>& CachedItems = WorkData.CachedItems;
 	FCachedRigElement& CachedEffector = WorkData.CachedEffector;
 
-	if (ExecuteContext.UnitContext.State == EControlRigState::Init)
+	if (CachedItems.Num() == 0 && Items.Num() > 0)
 	{
-		CachedItems.Reset();
-		CachedEffector.Reset();
-		return;
+		for (FRigElementKey Item : Items)
+		{
+			CachedItems.Add(FCachedRigElement(Item, Hierarchy));
+		}
+
+		CachedEffector = CachedItems.Last();
 	}
 
-	if (ExecuteContext.UnitContext.State == EControlRigState::Update)
+	if (CachedItems.Num() > 0)
 	{
-		if (CachedItems.Num() == 0 && Items.Num() > 0)
+		// Gather chain links. These are non zero length bones.
+		Chain.Reset();
+		
+		TArray<FTransform> Transforms;
+		Transforms.AddDefaulted(CachedItems.Num());
+
+		float MaximumReach = 0.f;
+		int32 const NumChainLinks = CachedItems.Num();
+		const FCachedRigElement& RootIndex = CachedItems[0];
+		Chain.Add(FFABRIKChainLink(Hierarchy->GetGlobalTransform(RootIndex).GetLocation(), 0.f, 0, 0));
+		Transforms[0] = Hierarchy->GetGlobalTransform(RootIndex);
+
+		// start from child to up
+		for (int32 ChainIndex = 1; ChainIndex < CachedItems.Num() ; ChainIndex++)
 		{
-			for (FRigElementKey Item : Items)
-			{
-				CachedItems.Add(FCachedRigElement(Item, Hierarchy));
-			}
+			const FTransform& BoneTransform = Hierarchy->GetGlobalTransform(CachedItems[ChainIndex]);
+			const FTransform& ParentTransform = Hierarchy->GetGlobalTransform(CachedItems[ChainIndex - 1]);
 
-			CachedEffector = CachedItems.Last();
+			// Calculate the combined length of this segment of skeleton
+			float const BoneLength = FVector::Dist(BoneTransform.GetLocation(), ParentTransform.GetLocation());
+
+			const int32 TransformIndex = Chain.Num();
+			Chain.Add(FFABRIKChainLink(BoneTransform.GetLocation(), BoneLength, ChainIndex, TransformIndex));
+			MaximumReach += BoneLength;
+
+			Transforms[TransformIndex] = BoneTransform;
 		}
 
-		if (CachedItems.Num() > 0)
+
+		bool bBoneLocationUpdated = AnimationCore::SolveFabrik(Chain, EffectorTransform.GetLocation(), MaximumReach, Precision, MaxIterations);
+		// If we moved some bones, update bone transforms.
+		if (bBoneLocationUpdated)
 		{
-			// Gather chain links. These are non zero length bones.
-			Chain.Reset();
-			
-			TArray<FTransform> Transforms;
-			Transforms.AddDefaulted(CachedItems.Num());
-
-			float MaximumReach = 0.f;
-			int32 const NumChainLinks = CachedItems.Num();
-			const FCachedRigElement& RootIndex = CachedItems[0];
-			Chain.Add(FFABRIKChainLink(Hierarchy->GetGlobalTransform(RootIndex).GetLocation(), 0.f, 0, 0));
-			Transforms[0] = Hierarchy->GetGlobalTransform(RootIndex);
-
-			// start from child to up
-			for (int32 ChainIndex = 1; ChainIndex < CachedItems.Num() ; ChainIndex++)
+			// FABRIK algorithm - re-orientation of bone local axes after translation calculation
+			for (int32 LinkIndex = 0; LinkIndex < NumChainLinks - 1; LinkIndex++)
 			{
-				const FTransform& BoneTransform = Hierarchy->GetGlobalTransform(CachedItems[ChainIndex]);
-				const FTransform& ParentTransform = Hierarchy->GetGlobalTransform(CachedItems[ChainIndex - 1]);
+				const FFABRIKChainLink& CurrentLink = Chain[LinkIndex];
+				const FFABRIKChainLink& ChildLink = Chain[LinkIndex + 1];
+				const FCachedRigElement& CurrentItem = CachedItems[CurrentLink.BoneIndex];
+				const FCachedRigElement& ChildItem = CachedItems[ChildLink.BoneIndex];
 
-				// Calculate the combined length of this segment of skeleton
-				float const BoneLength = FVector::Dist(BoneTransform.GetLocation(), ParentTransform.GetLocation());
+				// Calculate pre-translation vector between this bone and child
+				FVector const OldDir = (Hierarchy->GetGlobalTransform(ChildItem).GetLocation() - Hierarchy->GetGlobalTransform(CurrentItem).GetLocation()).GetUnsafeNormal();
 
-				const int32 TransformIndex = Chain.Num();
-				Chain.Add(FFABRIKChainLink(BoneTransform.GetLocation(), BoneLength, ChainIndex, TransformIndex));
-				MaximumReach += BoneLength;
+				// Get vector from the post-translation bone to it's child
+				FVector const NewDir = (ChildLink.Position - CurrentLink.Position).GetUnsafeNormal();
 
-				Transforms[TransformIndex] = BoneTransform;
+				// Calculate axis of rotation from pre-translation vector to post-translation vector
+				FVector const RotationAxis = FVector::CrossProduct(OldDir, NewDir).GetSafeNormal();
+				float const RotationAngle = FMath::Acos(FVector::DotProduct(OldDir, NewDir));
+				FQuat const DeltaRotation = FQuat(RotationAxis, RotationAngle);
+				// We're going to multiply it, in order to not have to re-normalize the final quaternion, it has to be a unit quaternion.
+				checkSlow(DeltaRotation.IsNormalized());
+
+				// Calculate absolute rotation and set it
+				FTransform& CurrentBoneTransform = Transforms[CurrentLink.TransformIndex];
+				CurrentBoneTransform.SetRotation(DeltaRotation * CurrentBoneTransform.GetRotation());
+				CurrentBoneTransform.NormalizeRotation();
+				CurrentBoneTransform.SetTranslation(CurrentLink.Position);
 			}
 
+			// fill up the last data transform
+			const FFABRIKChainLink & CurrentLink = Chain[NumChainLinks - 1];
+			const FCachedRigElement& CurrentItem = CachedItems[CurrentLink.BoneIndex];
+
+			FTransform& CurrentBoneTransform = Transforms[CurrentLink.TransformIndex];
+			CurrentBoneTransform.SetTranslation(CurrentLink.Position);
+			CurrentBoneTransform.SetRotation(Hierarchy->GetGlobalTransform(CurrentItem).GetRotation());
 
-			bool bBoneLocationUpdated = AnimationCore::SolveFabrik(Chain, EffectorTransform.GetLocation(), MaximumReach, Precision, MaxIterations);
-			// If we moved some bones, update bone transforms.
-			if (bBoneLocationUpdated)
+			if (FMath::IsNearlyEqual(Weight, 1.f))
 			{
-				// FABRIK algorithm - re-orientation of bone local axes after translation calculation
-				for (int32 LinkIndex = 0; LinkIndex < NumChainLinks - 1; LinkIndex++)
+				for (int32 LinkIndex = 0; LinkIndex < NumChainLinks; LinkIndex++)
 				{
-					const FFABRIKChainLink& CurrentLink = Chain[LinkIndex];
-					const FFABRIKChainLink& ChildLink = Chain[LinkIndex + 1];
-					const FCachedRigElement& CurrentItem = CachedItems[CurrentLink.BoneIndex];
-					const FCachedRigElement& ChildItem = CachedItems[ChildLink.BoneIndex];
-
-					// Calculate pre-translation vector between this bone and child
-					FVector const OldDir = (Hierarchy->GetGlobalTransform(ChildItem).GetLocation() - Hierarchy->GetGlobalTransform(CurrentItem).GetLocation()).GetUnsafeNormal();
-
-					// Get vector from the post-translation bone to it's child
-					FVector const NewDir = (ChildLink.Position - CurrentLink.Position).GetUnsafeNormal();
-
-					// Calculate axis of rotation from pre-translation vector to post-translation vector
-					FVector const RotationAxis = FVector::CrossProduct(OldDir, NewDir).GetSafeNormal();
-					float const RotationAngle = FMath::Acos(FVector::DotProduct(OldDir, NewDir));
-					FQuat const DeltaRotation = FQuat(RotationAxis, RotationAngle);
-					// We're going to multiply it, in order to not have to re-normalize the final quaternion, it has to be a unit quaternion.
-					checkSlow(DeltaRotation.IsNormalized());
-
-					// Calculate absolute rotation and set it
-					FTransform& CurrentBoneTransform = Transforms[CurrentLink.TransformIndex];
-					CurrentBoneTransform.SetRotation(DeltaRotation * CurrentBoneTransform.GetRotation());
-					CurrentBoneTransform.NormalizeRotation();
-					CurrentBoneTransform.SetTranslation(CurrentLink.Position);
+					FFABRIKChainLink const & LocalLink = Chain[LinkIndex];
+					const FCachedRigElement& LocalItem = CachedItems[LocalLink.BoneIndex];
+					Hierarchy->SetGlobalTransform(LocalItem, Transforms[LocalLink.TransformIndex], bPropagateToChildren);
 				}
+			}
+			else
+			{
+				float T = FMath::Clamp<float>(Weight, 0.f, 1.f);
 
-				// fill up the last data transform
-				const FFABRIKChainLink & CurrentLink = Chain[NumChainLinks - 1];
-				const FCachedRigElement& CurrentItem = CachedItems[CurrentLink.BoneIndex];
-
-				FTransform& CurrentBoneTransform = Transforms[CurrentLink.TransformIndex];
-				CurrentBoneTransform.SetTranslation(CurrentLink.Position);
-				CurrentBoneTransform.SetRotation(Hierarchy->GetGlobalTransform(CurrentItem).GetRotation());
-
-				if (FMath::IsNearlyEqual(Weight, 1.f))
-				{
-					for (int32 LinkIndex = 0; LinkIndex < NumChainLinks; LinkIndex++)
-					{
-						FFABRIKChainLink const & LocalLink = Chain[LinkIndex];
-						const FCachedRigElement& LocalItem = CachedItems[LocalLink.BoneIndex];
-						Hierarchy->SetGlobalTransform(LocalItem, Transforms[LocalLink.TransformIndex], bPropagateToChildren);
-					}
-				}
-				else
+				for (int32 LinkIndex = 0; LinkIndex < NumChainLinks; LinkIndex++)
 				{
-					float T = FMath::Clamp<float>(Weight, 0.f, 1.f);
-
-					for (int32 LinkIndex = 0; LinkIndex < NumChainLinks; LinkIndex++)
-					{
-						FFABRIKChainLink const & LocalLink = Chain[LinkIndex];
-						const FCachedRigElement& LocalItem = CachedItems[LocalLink.BoneIndex];
-						FTransform PreviousXfo = Hierarchy->GetGlobalTransform(LocalItem);
-						FTransform Xfo = FControlRigMathLibrary::LerpTransform(PreviousXfo, Transforms[LocalLink.TransformIndex], T);
-						Hierarchy->SetGlobalTransform(LocalItem, Xfo, bPropagateToChildren);
-					}
+					FFABRIKChainLink const & LocalLink = Chain[LinkIndex];
+					const FCachedRigElement& LocalItem = CachedItems[LocalLink.BoneIndex];
+					FTransform PreviousXfo = Hierarchy->GetGlobalTransform(LocalItem);
+					FTransform Xfo = FControlRigMathLibrary::LerpTransform(PreviousXfo, Transforms[LocalLink.TransformIndex], T);
+					Hierarchy->SetGlobalTransform(LocalItem, Xfo, bPropagateToChildren);
 				}
 			}
+		}
 
-			if (bSetEffectorTransform)
+		if (bSetEffectorTransform)
+		{
+			if (FMath::IsNearlyEqual(Weight, 1.f))
 			{
-				if (FMath::IsNearlyEqual(Weight, 1.f))
-				{
-					Hierarchy->SetGlobalTransform(CachedEffector, EffectorTransform, bPropagateToChildren);
-				}
-				else
-				{
-					float T = FMath::Clamp<float>(Weight, 0.f, 1.f);
-					FTransform PreviousXfo = Hierarchy->GetGlobalTransform(CachedEffector);
-					FTransform Xfo = FControlRigMathLibrary::LerpTransform(PreviousXfo, EffectorTransform, T);
-					Hierarchy->SetGlobalTransform(CachedEffector, Xfo, bPropagateToChildren);
-				}
+				Hierarchy->SetGlobalTransform(CachedEffector, EffectorTransform, bPropagateToChildren);
+			}
+			else
+			{
+				float T = FMath::Clamp<float>(Weight, 0.f, 1.f);
+				FTransform PreviousXfo = Hierarchy->GetGlobalTransform(CachedEffector);
+				FTransform Xfo = FControlRigMathLibrary::LerpTransform(PreviousXfo, EffectorTransform, T);
+				Hierarchy->SetGlobalTransform(CachedEffector, Xfo, bPropagateToChildren);
 			}
 		}
 	}
diff --git a/Engine/Plugins/Animation/ControlRig/Source/ControlRig/Private/Units/Highlevel/Hierarchy/RigUnit_FitChainToCurve.cpp b/Engine/Plugins/Animation/ControlRig/Source/ControlRig/Private/Units/Highlevel/Hierarchy/RigUnit_FitChainToCurve.cpp
--- a/Engine/Plugins/Animation/ControlRig/Source/ControlRig/Private/Units/Highlevel/Hierarchy/RigUnit_FitChainToCurve.cpp
+++ b/Engine/Plugins/Animation/ControlRig/Source/ControlRig/Private/Units/Highlevel/Hierarchy/RigUnit_FitChainToCurve.cpp
@@ -7,12 +7,6 @@
 
 FRigUnit_FitChainToCurve_Execute()
 {
-	if (ExecuteContext.UnitContext.State == EControlRigState::Init)
-	{
-		WorkData.CachedItems.Reset();
-		return;
-	}
-
 	FRigElementKeyCollection Items;
 	if(WorkData.CachedItems.Num() == 0)
 	{
@@ -97,12 +91,6 @@ FRigUnit_FitChainToCurveItemArray_Execute()
 	TArray<float>& ItemRotationT = WorkData.ItemRotationT;
 	TArray<FTransform>& ItemLocalTransforms = WorkData.ItemLocalTransforms;
 
-	if (ExecuteContext.UnitContext.State == EControlRigState::Init)
-	{
-		CachedItems.Reset();
-		return;
-	}
-
 	if(CachedItems.Num() == 0 && Items.Num() > 1)
 	{
 		CurvePositions.Reset();
diff --git a/Engine/Plugins/Animation/ControlRig/Source/ControlRig/Private/Units/Highlevel/Hierarchy/RigUnit_ModifyBoneTransforms.cpp b/Engine/Plugins/Animation/ControlRig/Source/ControlRig/Private/Units/Highlevel/Hierarchy/RigUnit_ModifyBoneTransforms.cpp
--- a/Engine/Plugins/Animation/ControlRig/Source/ControlRig/Private/Units/Highlevel/Hierarchy/RigUnit_ModifyBoneTransforms.cpp
+++ b/Engine/Plugins/Animation/ControlRig/Source/ControlRig/Private/Units/Highlevel/Hierarchy/RigUnit_ModifyBoneTransforms.cpp
@@ -49,28 +49,28 @@ IMPLEMENT_RIGUNIT_AUTOMATION_TEST(FRigUnit_ModifyBoneTransforms)
 
 	Hierarchy->ResetPoseToInitial(ERigElementType::Bone);
 	Unit.Mode = EControlRigModifyBoneMode::AdditiveLocal;
-	InitAndExecute();
+	Execute();
 	AddErrorIfFalse((Hierarchy->GetGlobalTransform(0).GetTranslation() - FVector(1.f, 0.f, 0.f)).IsNearlyZero(), TEXT("unexpected transform"));
 	AddErrorIfFalse((Hierarchy->GetGlobalTransform(1).GetTranslation() - FVector(11.f, 13.f, 15.f)).IsNearlyZero(), TEXT("unexpected transform"));
 	AddErrorIfFalse((Hierarchy->GetGlobalTransform(2).GetTranslation() - FVector(15.f, 17.f, 19.f)).IsNearlyZero(), TEXT("unexpected transform"));
 
 	Hierarchy->ResetPoseToInitial(ERigElementType::Bone);
 	Unit.Mode = EControlRigModifyBoneMode::AdditiveGlobal;
-	InitAndExecute();
+	Execute();
 	AddErrorIfFalse((Hierarchy->GetGlobalTransform(0).GetTranslation() - FVector(1.f, 0.f, 0.f)).IsNearlyZero(), TEXT("unexpected transform"));
 	AddErrorIfFalse((Hierarchy->GetGlobalTransform(1).GetTranslation() - FVector(11.f, 13.f, 15.f)).IsNearlyZero(), TEXT("unexpected transform"));
 	AddErrorIfFalse((Hierarchy->GetGlobalTransform(2).GetTranslation() - FVector(15.f, 17.f, 19.f)).IsNearlyZero(), TEXT("unexpected transform"));
 
 	Hierarchy->ResetPoseToInitial(ERigElementType::Bone);
 	Unit.Mode = EControlRigModifyBoneMode::OverrideLocal;
-	InitAndExecute();
+	Execute();
 	AddErrorIfFalse((Hierarchy->GetGlobalTransform(0).GetTranslation() - FVector(1.f, 0.f, 0.f)).IsNearlyZero(), TEXT("unexpected transform"));
 	AddErrorIfFalse((Hierarchy->GetGlobalTransform(1).GetTranslation() - FVector(11.f, 11.f, 12.f)).IsNearlyZero(), TEXT("unexpected transform"));
 	AddErrorIfFalse((Hierarchy->GetGlobalTransform(2).GetTranslation() - FVector(11.f, 11.f, 12.f)).IsNearlyZero(), TEXT("unexpected transform"));
 
 	Hierarchy->ResetPoseToInitial(ERigElementType::Bone);
 	Unit.Mode = EControlRigModifyBoneMode::OverrideGlobal;
-	InitAndExecute();
+	Execute();
 	AddErrorIfFalse((Hierarchy->GetGlobalTransform(0).GetTranslation() - FVector(1.f, 0.f, 0.f)).IsNearlyZero(), TEXT("unexpected transform"));
 	AddErrorIfFalse((Hierarchy->GetGlobalTransform(1).GetTranslation() - FVector(10.f, 11.f, 12.f)).IsNearlyZero(), TEXT("unexpected transform"));
 	AddErrorIfFalse((Hierarchy->GetGlobalTransform(2).GetTranslation() - FVector(10.f, 11.f, 12.f)).IsNearlyZero(), TEXT("unexpected transform"));
@@ -78,7 +78,7 @@ IMPLEMENT_RIGUNIT_AUTOMATION_TEST(FRigUnit_ModifyBoneTransforms)
 	Hierarchy->ResetPoseToInitial(ERigElementType::Bone);
 	Unit.Mode = EControlRigModifyBoneMode::AdditiveLocal;
 	Unit.Weight = 0.5f;
-	InitAndExecute();
+	Execute();
 	AddErrorIfFalse((Hierarchy->GetGlobalTransform(0).GetTranslation() - FVector(1.f, 0.f, 0.f)).IsNearlyZero(), TEXT("unexpected transform"));
 	AddErrorIfFalse((Hierarchy->GetGlobalTransform(1).GetTranslation() - FVector(6.f, 7.5f, 9.f)).IsNearlyZero(), TEXT("unexpected transform"));
 	AddErrorIfFalse((Hierarchy->GetGlobalTransform(2).GetTranslation() - FVector(10.f, 11.5f, 13.f)).IsNearlyZero(), TEXT("unexpected transform"));
diff --git a/Engine/Plugins/Animation/ControlRig/Source/ControlRig/Private/Units/Highlevel/Hierarchy/RigUnit_ModifyTransforms.cpp b/Engine/Plugins/Animation/ControlRig/Source/ControlRig/Private/Units/Highlevel/Hierarchy/RigUnit_ModifyTransforms.cpp
--- a/Engine/Plugins/Animation/ControlRig/Source/ControlRig/Private/Units/Highlevel/Hierarchy/RigUnit_ModifyTransforms.cpp
+++ b/Engine/Plugins/Animation/ControlRig/Source/ControlRig/Private/Units/Highlevel/Hierarchy/RigUnit_ModifyTransforms.cpp
@@ -14,96 +14,81 @@ FRigUnit_ModifyTransforms_Execute()
 	URigHierarchy* Hierarchy = ExecuteContext.Hierarchy;
 	if (Hierarchy)
 	{
-		switch (ExecuteContext.UnitContext.State)
+		if (CachedItems.Num() == 0)
 		{
-			case EControlRigState::Init:
-			{
-				CachedItems.Reset();
-				return;
-			}
-			case EControlRigState::Update:
-			{
-				if (CachedItems.Num() == 0)
-				{
-					CachedItems.SetNum(ItemToModify.Num());
-				}
+			CachedItems.SetNum(ItemToModify.Num());
+		}
 
-				float Minimum = FMath::Min<float>(WeightMinimum, WeightMaximum);
-				float Maximum = FMath::Max<float>(WeightMinimum, WeightMaximum);
+		float Minimum = FMath::Min<float>(WeightMinimum, WeightMaximum);
+		float Maximum = FMath::Max<float>(WeightMinimum, WeightMaximum);
 
-				if (Weight <= Minimum + SMALL_NUMBER || FMath::IsNearlyEqual(Minimum, Maximum))
+		if (Weight <= Minimum + SMALL_NUMBER || FMath::IsNearlyEqual(Minimum, Maximum))
+		{
+			return;
+		}
+
+		if (CachedItems.Num() == ItemToModify.Num())
+		{
+			float T = FMath::Clamp<float>((Weight - Minimum) / (Maximum - Minimum), 0.f, 1.f);
+			bool bNeedsBlend = T < 1.f - SMALL_NUMBER;
+
+			int32 EntryIndex = 0;
+			for (const FRigUnit_ModifyTransforms_PerItem& Entry : ItemToModify)
+			{
+				FCachedRigElement& CachedItem = CachedItems[EntryIndex];
+				if (!CachedItem.UpdateCache(Entry.Item, Hierarchy))
 				{
-					return;
+					continue;
 				}
 
-				if (CachedItems.Num() == ItemToModify.Num())
-				{
-					float T = FMath::Clamp<float>((Weight - Minimum) / (Maximum - Minimum), 0.f, 1.f);
-					bool bNeedsBlend = T < 1.f - SMALL_NUMBER;
+				FTransform Transform = Entry.Transform;
 
-					int32 EntryIndex = 0;
-					for (const FRigUnit_ModifyTransforms_PerItem& Entry : ItemToModify)
+				switch (Mode)
+				{
+					case EControlRigModifyBoneMode::OverrideLocal:
 					{
-						FCachedRigElement& CachedItem = CachedItems[EntryIndex];
-						if (!CachedItem.UpdateCache(Entry.Item, Hierarchy))
+						if (bNeedsBlend)
 						{
-							continue;
+							Transform = FControlRigMathLibrary::LerpTransform(Hierarchy->GetLocalTransform(CachedItem), Transform, T);
 						}
-
-						FTransform Transform = Entry.Transform;
-
-						switch (Mode)
+						Hierarchy->SetLocalTransform(CachedItem, Transform, true);
+						break;
+					}
+					case EControlRigModifyBoneMode::OverrideGlobal:
+					{
+						if (bNeedsBlend)
+						{
+							Transform = FControlRigMathLibrary::LerpTransform(Hierarchy->GetGlobalTransform(CachedItem), Transform, T);
+						}
+						Hierarchy->SetGlobalTransform(CachedItem, Transform, true);
+						break;
+					}
+					case EControlRigModifyBoneMode::AdditiveLocal:
+					{
+						if (bNeedsBlend)
 						{
-							case EControlRigModifyBoneMode::OverrideLocal:
-							{
-								if (bNeedsBlend)
-								{
-									Transform = FControlRigMathLibrary::LerpTransform(Hierarchy->GetLocalTransform(CachedItem), Transform, T);
-								}
-								Hierarchy->SetLocalTransform(CachedItem, Transform, true);
-								break;
-							}
-							case EControlRigModifyBoneMode::OverrideGlobal:
-							{
-								if (bNeedsBlend)
-								{
-									Transform = FControlRigMathLibrary::LerpTransform(Hierarchy->GetGlobalTransform(CachedItem), Transform, T);
-								}
-								Hierarchy->SetGlobalTransform(CachedItem, Transform, true);
-								break;
-							}
-							case EControlRigModifyBoneMode::AdditiveLocal:
-							{
-								if (bNeedsBlend)
-								{
-									Transform = FControlRigMathLibrary::LerpTransform(FTransform::Identity, Transform, T);
-								}
-								Transform = Transform * Hierarchy->GetLocalTransform(CachedItem);
-								Hierarchy->SetLocalTransform(CachedItem, Transform, true);
-								break;
-							}
-							case EControlRigModifyBoneMode::AdditiveGlobal:
-							{
-								if (bNeedsBlend)
-								{
-									Transform = FControlRigMathLibrary::LerpTransform(FTransform::Identity, Transform, T);
-								}
-								Transform = Hierarchy->GetGlobalTransform(CachedItem) * Transform;
-								Hierarchy->SetGlobalTransform(CachedItem, Transform, true);
-								break;
-							}
-							default:
-							{
-								break;
-							}
+							Transform = FControlRigMathLibrary::LerpTransform(FTransform::Identity, Transform, T);
 						}
-						EntryIndex++;
+						Transform = Transform * Hierarchy->GetLocalTransform(CachedItem);
+						Hierarchy->SetLocalTransform(CachedItem, Transform, true);
+						break;
+					}
+					case EControlRigModifyBoneMode::AdditiveGlobal:
+					{
+						if (bNeedsBlend)
+						{
+							Transform = FControlRigMathLibrary::LerpTransform(FTransform::Identity, Transform, T);
+						}
+						Transform = Hierarchy->GetGlobalTransform(CachedItem) * Transform;
+						Hierarchy->SetGlobalTransform(CachedItem, Transform, true);
+						break;
+					}
+					default:
+					{
+						break;
 					}
 				}
-			}
-			default:
-			{
-				break;
+				EntryIndex++;
 			}
 		}
 	}
diff --git a/Engine/Plugins/Animation/ControlRig/Source/ControlRig/Private/Units/Highlevel/Hierarchy/RigUnit_MultiFABRIK.cpp b/Engine/Plugins/Animation/ControlRig/Source/ControlRig/Private/Units/Highlevel/Hierarchy/RigUnit_MultiFABRIK.cpp
--- a/Engine/Plugins/Animation/ControlRig/Source/ControlRig/Private/Units/Highlevel/Hierarchy/RigUnit_MultiFABRIK.cpp
+++ b/Engine/Plugins/Animation/ControlRig/Source/ControlRig/Private/Units/Highlevel/Hierarchy/RigUnit_MultiFABRIK.cpp
@@ -227,247 +227,244 @@ FRigUnit_MultiFABRIK_Execute()
 	FRigUnit_MultiFABRIK_ChainGroup&				ChainGroup = WorkData.ChainGroup;
 	TArray<FRigUnit_MultiFABRIK_BoneWorkingData>&	BoneTree = WorkData.BoneTree;
 	TArray<FCachedRigElement>&				        EffectorIndices = WorkData.EffectorBoneIndices;
-	
-	if (ExecuteContext.UnitContext.State == EControlRigState::Init)
+
+	if (!bIsInitialized)
 	{
 		EffectorIndices.Reset();
 		BoneTree.Reset();
 		ChainGroup.Reset();
-		return;
+		bIsInitialized = false;
 	}
-
-	if (ExecuteContext.UnitContext.State == EControlRigState::Update)
+	
+	if(BoneTree.Num() == 0)
 	{
-		if(BoneTree.Num() == 0)
+		// verify the chain
+		const FRigElementKey RootKey(RootBone, ERigElementType::Bone);
+		const int32 RootBoneIndex = Hierarchy->GetIndex(RootKey);
+		if (RootBoneIndex != INDEX_NONE)
 		{
-			// verify the chain
-			const FRigElementKey RootKey(RootBone, ERigElementType::Bone);
-			const int32 RootBoneIndex = Hierarchy->GetIndex(RootKey);
-			if (RootBoneIndex != INDEX_NONE)
+			// fill up all effector indices
+			for (int32 Index = 0; Index < Effectors.Num(); ++Index)
 			{
-				// fill up all effector indices
-				for (int32 Index = 0; Index < Effectors.Num(); ++Index)
-				{
-					const FRigElementKey EffectorKey(Effectors[Index].Bone, ERigElementType::Bone);
-					EffectorIndices.Add(FCachedRigElement(EffectorKey, Hierarchy));
-				}
+				const FRigElementKey EffectorKey(Effectors[Index].Bone, ERigElementType::Bone);
+				EffectorIndices.Add(FCachedRigElement(EffectorKey, Hierarchy));
+			}
 
-				// go up to until you meet root
-				struct FLocalBoneData
-				{
-					FName BoneName;
-					int32 BoneIndex;
-					FName ParentBoneName;
-					int32 ParentIndex;
-
-					FLocalBoneData(const FName& InBoneName, int32 InBoneIndex, const FName& InParentBoneName, int32 InParentIndex)
-						: BoneName(InBoneName)
-						, BoneIndex(InBoneIndex)
-						, ParentBoneName(InParentBoneName)
-						, ParentIndex(InParentIndex)
-					{}
-
-					FLocalBoneData(URigHierarchy* InHierarchy, int32 InBoneIndex)
-					: BoneName(NAME_None)
+			// go up to until you meet root
+			struct FLocalBoneData
+			{
+				FName BoneName;
+				int32 BoneIndex;
+				FName ParentBoneName;
+				int32 ParentIndex;
+
+				FLocalBoneData(const FName& InBoneName, int32 InBoneIndex, const FName& InParentBoneName, int32 InParentIndex)
+					: BoneName(InBoneName)
 					, BoneIndex(InBoneIndex)
-					, ParentBoneName(NAME_None)
-				    , ParentIndex(INDEX_NONE)
-					{
-						FRigBaseElement* BoneElement = InHierarchy->Get(InBoneIndex);
-						BoneName = BoneElement->GetName();
-						BoneIndex = BoneElement->GetIndex();
+					, ParentBoneName(InParentBoneName)
+					, ParentIndex(InParentIndex)
+				{}
+
+				FLocalBoneData(URigHierarchy* InHierarchy, int32 InBoneIndex)
+				: BoneName(NAME_None)
+				, BoneIndex(InBoneIndex)
+				, ParentBoneName(NAME_None)
+			    , ParentIndex(INDEX_NONE)
+				{
+					FRigBaseElement* BoneElement = InHierarchy->Get(InBoneIndex);
+					BoneName = BoneElement->GetName();
+					BoneIndex = BoneElement->GetIndex();
 
-						if(FRigBaseElement* ParentBoneElement = InHierarchy->GetFirstParent(BoneElement))
-						{
-							ParentBoneName = ParentBoneElement->GetName();
-							ParentIndex = ParentBoneElement->GetIndex();
-						}
+					if(FRigBaseElement* ParentBoneElement = InHierarchy->GetFirstParent(BoneElement))
+					{
+						ParentBoneName = ParentBoneElement->GetName();
+						ParentIndex = ParentBoneElement->GetIndex();
 					}
+				}
 
-				};
+			};
 
-				TMap<FName, int32> ExistingBoneList;
-				// now fill up data on bonedata_workdata
-				// get all list of joints to fill up BoneData array
-				for (int32 Index = 0; Index < EffectorIndices.Num(); ++Index)
+			TMap<FName, int32> ExistingBoneList;
+			// now fill up data on bonedata_workdata
+			// get all list of joints to fill up BoneData array
+			for (int32 Index = 0; Index < EffectorIndices.Num(); ++Index)
+			{
+				const int32 EffectorIndex = EffectorIndices[Index];
+				if (EffectorIndex != INDEX_NONE)
 				{
-					const int32 EffectorIndex = EffectorIndices[Index];
-					if (EffectorIndex != INDEX_NONE)
+					TArray<FLocalBoneData> BoneList;
+
+					int32 CurrentIndex = EffectorIndex;
+					do 
 					{
-						TArray<FLocalBoneData> BoneList;
+						BoneList.Insert(FLocalBoneData(Hierarchy, CurrentIndex), 0);
+						CurrentIndex = Hierarchy->GetIndex(Hierarchy->GetFirstParent(Hierarchy->GetKey(CurrentIndex)));
+					} while (CurrentIndex != RootBoneIndex && CurrentIndex != INDEX_NONE);
 
-						int32 CurrentIndex = EffectorIndex;
-						do 
-						{
-							BoneList.Insert(FLocalBoneData(Hierarchy, CurrentIndex), 0);
-							CurrentIndex = Hierarchy->GetIndex(Hierarchy->GetFirstParent(Hierarchy->GetKey(CurrentIndex)));
-						} while (CurrentIndex != RootBoneIndex && CurrentIndex != INDEX_NONE);
+					// if we haven't got to root, this is not valid chain
+					if (CurrentIndex == RootBoneIndex)
+					{
+						// add root
+						BoneList.Insert(FLocalBoneData(Hierarchy, RootBoneIndex), 0);
 
-						// if we haven't got to root, this is not valid chain
-						if (CurrentIndex == RootBoneIndex)
+						// go through all bone list collected
+						for (const FLocalBoneData& Bone : BoneList)
 						{
-							// add root
-							BoneList.Insert(FLocalBoneData(Hierarchy, RootBoneIndex), 0);
-
-							// go through all bone list collected
-							for (const FLocalBoneData& Bone : BoneList)
+							// if it's not collected yet, add to the list
+							int32* FoundIndex = ExistingBoneList.Find(Bone.BoneName);
+							int32 TreeIndex = INDEX_NONE;
+							FRigUnit_MultiFABRIK_BoneWorkingData* Data = nullptr;
+							if (FoundIndex == nullptr)
 							{
-								// if it's not collected yet, add to the list
-								int32* FoundIndex = ExistingBoneList.Find(Bone.BoneName);
-								int32 TreeIndex = INDEX_NONE;
-								FRigUnit_MultiFABRIK_BoneWorkingData* Data = nullptr;
-								if (FoundIndex == nullptr)
+								TreeIndex = BoneTree.AddDefaulted();
+								Data = &BoneTree[TreeIndex];
+
+								Data->BoneName = Bone.BoneName;
+								Data->BoneIndex = FCachedRigElement(FRigElementKey(Bone.BoneName, ERigElementType::Bone), Hierarchy);
+								Data->ParentIndex = FCachedRigElement(FRigElementKey(Bone.ParentBoneName, ERigElementType::Bone), Hierarchy);
+
+								// set bone length 
+								if (Data->ParentIndex.IsValid())
 								{
-									TreeIndex = BoneTree.AddDefaulted();
-									Data = &BoneTree[TreeIndex];
-
-									Data->BoneName = Bone.BoneName;
-									Data->BoneIndex = FCachedRigElement(FRigElementKey(Bone.BoneName, ERigElementType::Bone), Hierarchy);
-									Data->ParentIndex = FCachedRigElement(FRigElementKey(Bone.ParentBoneName, ERigElementType::Bone), Hierarchy);
-
-									// set bone length 
-									if (Data->ParentIndex.IsValid())
-									{
-										// save size - this is initial. we could do this in every frame instead
-										Data->BoneLength = Hierarchy->GetLocalTransform(Data->BoneIndex).GetTranslation().Size();
-									}
-									else
-									{
-										Data->BoneLength = 0.f;
-									}
-
-									// find tree index?
-									int32* FoundParent = ExistingBoneList.Find(Bone.ParentBoneName);
-									if (FoundParent)
-									{
-										Data->ParentTreeIndex = *FoundParent;
-									}
-
-									ExistingBoneList.Add(Data->BoneName, TreeIndex);
+									// save size - this is initial. we could do this in every frame instead
+									Data->BoneLength = Hierarchy->GetLocalTransform(Data->BoneIndex).GetTranslation().Size();
 								}
 								else
 								{
-									TreeIndex = *FoundIndex;
-									Data = &BoneTree[TreeIndex];
+									Data->BoneLength = 0.f;
 								}
 
-								if (Data->ParentTreeIndex != INDEX_NONE)
+								// find tree index?
+								int32* FoundParent = ExistingBoneList.Find(Bone.ParentBoneName);
+								if (FoundParent)
 								{
-									// add my index to my parent
-									BoneTree[Data->ParentTreeIndex].ChildrenTreeIndices.AddUnique(TreeIndex);
+									Data->ParentTreeIndex = *FoundParent;
 								}
+
+								ExistingBoneList.Add(Data->BoneName, TreeIndex);
+							}
+							else
+							{
+								TreeIndex = *FoundIndex;
+								Data = &BoneTree[TreeIndex];
+							}
+
+							if (Data->ParentTreeIndex != INDEX_NONE)
+							{
+								// add my index to my parent
+								BoneTree[Data->ParentTreeIndex].ChildrenTreeIndices.AddUnique(TreeIndex);
 							}
 						}
 					}
 				}
+			}
 
-				// create tree from BoneTree, this way it's easier for us to create chain per group
-				TArray<MultiFABRIK::FChainNode> ChainNodes;
+			// create tree from BoneTree, this way it's easier for us to create chain per group
+			TArray<MultiFABRIK::FChainNode> ChainNodes;
 
-				// first create all nodes
-				ChainNodes.AddZeroed(BoneTree.Num());
-				for (int32 Index = 0; Index < BoneTree.Num(); ++Index)
-				{
-					const int32 ParentTreeIndex = BoneTree[Index].ParentTreeIndex;
-					ChainNodes[Index].Parent = (ParentTreeIndex != INDEX_NONE)? &ChainNodes[ParentTreeIndex] : nullptr;
-					ChainNodes[Index].BoneTreeIndex = Index;
-				}
+			// first create all nodes
+			ChainNodes.AddZeroed(BoneTree.Num());
+			for (int32 Index = 0; Index < BoneTree.Num(); ++Index)
+			{
+				const int32 ParentTreeIndex = BoneTree[Index].ParentTreeIndex;
+				ChainNodes[Index].Parent = (ParentTreeIndex != INDEX_NONE)? &ChainNodes[ParentTreeIndex] : nullptr;
+				ChainNodes[Index].BoneTreeIndex = Index;
+			}
 
-				// then now add children
-				for (int32 Index = 0; Index < BoneTree.Num(); ++Index)
+			// then now add children
+			for (int32 Index = 0; Index < BoneTree.Num(); ++Index)
+			{
+				if (BoneTree[Index].ChildrenTreeIndices.Num() > 0)
 				{
-					if (BoneTree[Index].ChildrenTreeIndices.Num() > 0)
-					{
-						ChainNodes[Index].Children.AddZeroed(BoneTree[Index].ChildrenTreeIndices.Num());
+					ChainNodes[Index].Children.AddZeroed(BoneTree[Index].ChildrenTreeIndices.Num());
 
-						for (int32 ChildIndex = 0; ChildIndex < BoneTree[Index].ChildrenTreeIndices.Num(); ++ChildIndex)
-						{
-							const int32 ChildTreeIndex = BoneTree[Index].ChildrenTreeIndices[ChildIndex];
-							ChainNodes[Index].Children[ChildIndex] = (ChildTreeIndex != INDEX_NONE) ? &ChainNodes[ChildTreeIndex] : nullptr;
-						}
+					for (int32 ChildIndex = 0; ChildIndex < BoneTree[Index].ChildrenTreeIndices.Num(); ++ChildIndex)
+					{
+						const int32 ChildTreeIndex = BoneTree[Index].ChildrenTreeIndices[ChildIndex];
+						ChainNodes[Index].Children[ChildIndex] = (ChildTreeIndex != INDEX_NONE) ? &ChainNodes[ChildTreeIndex] : nullptr;
 					}
 				}
+			}
 
-				// now we have tree 
-				// add new link recursively
-				// now go from root to child, and create group if you have more children, you can always go back 
-				if (ChainNodes.Num() > 0)
+			// now we have tree 
+			// add new link recursively
+			// now go from root to child, and create group if you have more children, you can always go back 
+			if (ChainNodes.Num() > 0)
+			{
+				ChainGroup.RootBoneTreeIndex = 0;
+
+				// add new links recursively
+				// we group per chain, and sub groups goes next
+				// when you have multiple children, we add new groups, and each group gets chain
+				TArray<int32> ResolvedEffectorIndices;
+				ResolvedEffectorIndices.Reserve(EffectorIndices.Num());
+				for (FCachedRigElement& EffectorIndex : EffectorIndices)
 				{
-					ChainGroup.RootBoneTreeIndex = 0;
-
-					// add new links recursively
-					// we group per chain, and sub groups goes next
-					// when you have multiple children, we add new groups, and each group gets chain
-					TArray<int32> ResolvedEffectorIndices;
-					ResolvedEffectorIndices.Reserve(EffectorIndices.Num());
-					for (FCachedRigElement& EffectorIndex : EffectorIndices)
-					{
-						ResolvedEffectorIndices.Add(EffectorIndex);
-					}
-					MultiFABRIK::AddNewLinkRecursive(&ChainNodes[0], ChainGroup, BoneTree, ResolvedEffectorIndices);
+					ResolvedEffectorIndices.Add(EffectorIndex);
 				}
+				MultiFABRIK::AddNewLinkRecursive(&ChainNodes[0], ChainGroup, BoneTree, ResolvedEffectorIndices);
 			}
 		}
+	}
 
-		if (BoneTree.Num() > 0)
+	if (BoneTree.Num() > 0)
+	{
+		// update bone tree data
+		for (int32 BoneIndex = 0; BoneIndex < BoneTree.Num(); ++BoneIndex)
 		{
-			// update bone tree data
-			for (int32 BoneIndex = 0; BoneIndex < BoneTree.Num(); ++BoneIndex)
-			{
-				// update bone transform
-				BoneTree[BoneIndex].Location = Hierarchy->GetGlobalTransform(BoneTree[BoneIndex].BoneIndex).GetLocation();
-			}
+			// update bone transform
+			BoneTree[BoneIndex].Location = Hierarchy->GetGlobalTransform(BoneTree[BoneIndex].BoneIndex).GetLocation();
+		}
 
-			FVector RootLocation = BoneTree[0].Location;
-			// now iterate 
-			// we start from end to root group prioritizing getting to effector
+		FVector RootLocation = BoneTree[0].Location;
+		// now iterate 
+		// we start from end to root group prioritizing getting to effector
 
-			const TArrayView<const FRigUnit_MultiFABRIK_EndEffector> EffectorsView(Effectors.GetData(), Effectors.Num());
-			for (int32 IterIndex = 0; IterIndex < MaxIterations; ++IterIndex)
-			{
-				// update transform first
-				MultiFABRIK::UpdateTransfomRecusively(ChainGroup, BoneTree);
-				MultiFABRIK::SolveForwardRecursively(ChainGroup, BoneTree, EffectorsView, Precision);
+		const TArrayView<const FRigUnit_MultiFABRIK_EndEffector> EffectorsView(Effectors.GetData(), Effectors.Num());
+		for (int32 IterIndex = 0; IterIndex < MaxIterations; ++IterIndex)
+		{
+			// update transform first
+			MultiFABRIK::UpdateTransfomRecusively(ChainGroup, BoneTree);
+			MultiFABRIK::SolveForwardRecursively(ChainGroup, BoneTree, EffectorsView, Precision);
 
-				// root position can change while forward reaching
-				BoneTree[0].Location = RootLocation;
+			// root position can change while forward reaching
+			BoneTree[0].Location = RootLocation;
 
-				// solve backward
-				MultiFABRIK::SolveBackwardRecursively(ChainGroup, BoneTree, Precision);
-			}
+			// solve backward
+			MultiFABRIK::SolveBackwardRecursively(ChainGroup, BoneTree, Precision);
+		}
 
-			TArray<FTransform> NewTransforms;
+		TArray<FTransform> NewTransforms;
 
-			// now we should copy bone tree to hierarchy
-			const int32 NumBones = BoneTree.Num();
-			NewTransforms.AddUninitialized(NumBones);
-			NewTransforms[0] = Hierarchy->GetGlobalTransform(BoneTree[0].BoneIndex);
+		// now we should copy bone tree to hierarchy
+		const int32 NumBones = BoneTree.Num();
+		NewTransforms.AddUninitialized(NumBones);
+		NewTransforms[0] = Hierarchy->GetGlobalTransform(BoneTree[0].BoneIndex);
 
-			// we don't modify root bone, and all of these should have parent
-			for (int32 BoneIndex = 1; BoneIndex < NumBones; ++BoneIndex)
-			{
-				// rotation will be mess
-				NewTransforms[BoneIndex] = Hierarchy->GetGlobalTransform(BoneTree[BoneIndex].BoneIndex);
-				FTransform ParentTransform = Hierarchy->GetGlobalTransform(BoneTree[BoneIndex].ParentIndex);
-				FVector NewParentLocation = BoneTree[BoneTree[BoneIndex].ParentTreeIndex].Location;
+		// we don't modify root bone, and all of these should have parent
+		for (int32 BoneIndex = 1; BoneIndex < NumBones; ++BoneIndex)
+		{
+			// rotation will be mess
+			NewTransforms[BoneIndex] = Hierarchy->GetGlobalTransform(BoneTree[BoneIndex].BoneIndex);
+			FTransform ParentTransform = Hierarchy->GetGlobalTransform(BoneTree[BoneIndex].ParentIndex);
+			FVector NewParentLocation = BoneTree[BoneTree[BoneIndex].ParentTreeIndex].Location;
 
-				// Calculate pre-translation vector between this bone and child
-				const FVector OldDir = (NewTransforms[BoneIndex].GetLocation()-ParentTransform.GetLocation()).GetUnsafeNormal();
+			// Calculate pre-translation vector between this bone and child
+			const FVector OldDir = (NewTransforms[BoneIndex].GetLocation()-ParentTransform.GetLocation()).GetUnsafeNormal();
 
-				// Get vector from the post-translation bone to it's child
-				const FVector  NewDir = (BoneTree[BoneIndex].Location - NewParentLocation).GetUnsafeNormal();
+			// Get vector from the post-translation bone to it's child
+			const FVector  NewDir = (BoneTree[BoneIndex].Location - NewParentLocation).GetUnsafeNormal();
 
-				NewTransforms[BoneIndex].SetLocation(BoneTree[BoneIndex].Location);
-				MultiFABRIK::RecalculateRotationBasedOnDirection(OldDir, NewDir, NewTransforms[BoneTree[BoneIndex].ParentTreeIndex]);
-			}
+			NewTransforms[BoneIndex].SetLocation(BoneTree[BoneIndex].Location);
+			MultiFABRIK::RecalculateRotationBasedOnDirection(OldDir, NewDir, NewTransforms[BoneTree[BoneIndex].ParentTreeIndex]);
+		}
 
-			// we don't modify root bone, and all of these should have parent
-			for (int32 BoneIndex = 1; BoneIndex < NumBones; ++BoneIndex)
-			{
-				// only propagate, if you are leaf joints here
-				// this means, only the last joint in the test
-				Hierarchy->SetGlobalTransform(BoneTree[BoneIndex].BoneIndex, NewTransforms[BoneIndex], bPropagateToChildren);
-			}
+		// we don't modify root bone, and all of these should have parent
+		for (int32 BoneIndex = 1; BoneIndex < NumBones; ++BoneIndex)
+		{
+			// only propagate, if you are leaf joints here
+			// this means, only the last joint in the test
+			Hierarchy->SetGlobalTransform(BoneTree[BoneIndex].BoneIndex, NewTransforms[BoneIndex], bPropagateToChildren);
 		}
 	}
 }
@@ -486,7 +483,7 @@ IMPLEMENT_RIGUNIT_AUTOMATION_TEST(FRigUnit_MultiFABRIK)
 
 	// first validation test
 	// make sure this doesn't crash
-	InitAndExecute();
+	Execute();
 
 	Unit.RootBone = TEXT("Root");
 	Unit.Effectors.AddDefaulted(2);
diff --git a/Engine/Plugins/Animation/ControlRig/Source/ControlRig/Private/Units/Highlevel/Hierarchy/RigUnit_MultiFABRIK.h b/Engine/Plugins/Animation/ControlRig/Source/ControlRig/Private/Units/Highlevel/Hierarchy/RigUnit_MultiFABRIK.h
--- a/Engine/Plugins/Animation/ControlRig/Source/ControlRig/Private/Units/Highlevel/Hierarchy/RigUnit_MultiFABRIK.h
+++ b/Engine/Plugins/Animation/ControlRig/Source/ControlRig/Private/Units/Highlevel/Hierarchy/RigUnit_MultiFABRIK.h
@@ -131,6 +131,12 @@ struct CONTROLRIG_API FRigUnit_MultiFABRIK : public FRigUnit_HighlevelBaseMutabl
 		Precision = 1.f;
 		MaxIterations = 4;
 		bPropagateToChildren = true;
+		bIsInitialized = false;
+	}
+
+	virtual void Initialize() override
+	{
+		bIsInitialized = false;
 	}
 
 	/**
@@ -164,4 +170,7 @@ struct CONTROLRIG_API FRigUnit_MultiFABRIK : public FRigUnit_HighlevelBaseMutabl
 
 	UPROPERTY(transient)
 	FRigUnit_MultiFABRIK_WorkData WorkData;
+
+	UPROPERTY(transient)
+	bool bIsInitialized;
 };
diff --git a/Engine/Plugins/Animation/ControlRig/Source/ControlRig/Private/Units/Highlevel/Hierarchy/RigUnit_SlideChain.cpp b/Engine/Plugins/Animation/ControlRig/Source/ControlRig/Private/Units/Highlevel/Hierarchy/RigUnit_SlideChain.cpp
--- a/Engine/Plugins/Animation/ControlRig/Source/ControlRig/Private/Units/Highlevel/Hierarchy/RigUnit_SlideChain.cpp
+++ b/Engine/Plugins/Animation/ControlRig/Source/ControlRig/Private/Units/Highlevel/Hierarchy/RigUnit_SlideChain.cpp
@@ -8,12 +8,6 @@
 
 FRigUnit_SlideChain_Execute()
 {
-	if (ExecuteContext.UnitContext.State == EControlRigState::Init)
-	{
-		WorkData.CachedItems.Reset();
-		return;
-	}
-
 	FRigElementKeyCollection Items;
 	if(WorkData.CachedItems.Num() == 0)
 	{
@@ -68,12 +62,6 @@ FRigUnit_SlideChainItemArray_Execute()
 	TArray<FTransform>& Transforms = WorkData.Transforms;
 	TArray<FTransform>& BlendedTransforms = WorkData.BlendedTransforms;
 
-	if (ExecuteContext.UnitContext.State == EControlRigState::Init)
-	{
-		CachedItems.Reset();
-		return;
-	}
-
 	if(CachedItems.Num() == 0 && Items.Num() > 1)
 	{
 		ItemSegments.Reset();
diff --git a/Engine/Plugins/Animation/ControlRig/Source/ControlRig/Private/Units/Highlevel/Hierarchy/RigUnit_SphericalPoseReader.cpp b/Engine/Plugins/Animation/ControlRig/Source/ControlRig/Private/Units/Highlevel/Hierarchy/RigUnit_SphericalPoseReader.cpp
--- a/Engine/Plugins/Animation/ControlRig/Source/ControlRig/Private/Units/Highlevel/Hierarchy/RigUnit_SphericalPoseReader.cpp
+++ b/Engine/Plugins/Animation/ControlRig/Source/ControlRig/Private/Units/Highlevel/Hierarchy/RigUnit_SphericalPoseReader.cpp
@@ -16,14 +16,6 @@ FRigUnit_SphericalPoseReader_Execute()
 		return;
 	}
 
-	if (ExecuteContext.UnitContext.State == EControlRigState::Init)
-	{
-		DriverCache.Reset();
-		OptionalParentCache.Reset();
-		bCachedInitTransforms = false;
-		return;
-	}
-
 	if (!DriverCache.UpdateCache(DriverItem, Hierarchy))
 	{
 		return;
diff --git a/Engine/Plugins/Animation/ControlRig/Source/ControlRig/Private/Units/Highlevel/Hierarchy/RigUnit_SpringIK.cpp b/Engine/Plugins/Animation/ControlRig/Source/ControlRig/Private/Units/Highlevel/Hierarchy/RigUnit_SpringIK.cpp
--- a/Engine/Plugins/Animation/ControlRig/Source/ControlRig/Private/Units/Highlevel/Hierarchy/RigUnit_SpringIK.cpp
+++ b/Engine/Plugins/Animation/ControlRig/Source/ControlRig/Private/Units/Highlevel/Hierarchy/RigUnit_SpringIK.cpp
@@ -19,13 +19,6 @@ FRigUnit_SpringIK_Execute()
 	TArray<FTransform>& Transforms = WorkData.Transforms;
 	FCRSimPointContainer& Simulation = WorkData.Simulation;
 
-	if (ExecuteContext.UnitContext.State == EControlRigState::Init)
-	{
-		CachedBones.Reset();
-		CachedPoleVector.Reset();
-		return;
-	}
-
 	if(CachedBones.Num() == 0)
 	{
 		Simulation.Reset();
@@ -215,7 +208,7 @@ FRigUnit_SpringIK_Execute()
 		}
 	}
 
-	Simulation.StepSemiExplicitEuler((bLiveSimulation ? ExecuteContext.UnitContext.DeltaTime : Simulation.TimeStep) * (float)FMath::Clamp<int32>(Iterations, 1, 32));
+	Simulation.StepSemiExplicitEuler((bLiveSimulation ? ExecuteContext.GetDeltaTime() : Simulation.TimeStep) * (float)FMath::Clamp<int32>(Iterations, 1, 32));
 
 	FVector AccumulatedTarget = FVector::ZeroVector;
 	FVector LastPrimaryTarget = FVector::ZeroVector;
diff --git a/Engine/Plugins/Animation/ControlRig/Source/ControlRig/Private/Units/Highlevel/Hierarchy/RigUnit_TransformConstraint.cpp b/Engine/Plugins/Animation/ControlRig/Source/ControlRig/Private/Units/Highlevel/Hierarchy/RigUnit_TransformConstraint.cpp
--- a/Engine/Plugins/Animation/ControlRig/Source/ControlRig/Private/Units/Highlevel/Hierarchy/RigUnit_TransformConstraint.cpp
+++ b/Engine/Plugins/Animation/ControlRig/Source/ControlRig/Private/Units/Highlevel/Hierarchy/RigUnit_TransformConstraint.cpp
@@ -110,54 +110,40 @@ FRigUnit_TransformConstraintPerItem_Execute()
 		}
 	};
 
-	if (ExecuteContext.UnitContext.State == EControlRigState::Init)
+	URigHierarchy* Hierarchy = ExecuteContext.Hierarchy;
+	if (Hierarchy)
 	{
-		ConstraintData.Reset();
-		ConstraintDataToTargets.Reset();
-		URigHierarchy* Hierarchy = ExecuteContext.Hierarchy;
-		
-		if (Hierarchy)
+		if ((ConstraintData.Num() != Targets.Num()))
 		{
 			SetupConstraintData();
 		}
-	}
-	else if (ExecuteContext.UnitContext.State == EControlRigState::Update)
-	{
-		URigHierarchy* Hierarchy = ExecuteContext.Hierarchy;
-		if (Hierarchy)
-		{
-			if ((ConstraintData.Num() != Targets.Num()))
-			{
-				SetupConstraintData();
-			}
 
-			if (Item.IsValid())
+		if (Item.IsValid())
+		{
+			const int32 TargetNum = Targets.Num();
+			if (TargetNum > 0 && ConstraintData.Num() > 0)
 			{
-				const int32 TargetNum = Targets.Num();
-				if (TargetNum > 0 && ConstraintData.Num() > 0)
+				for (int32 ConstraintIndex= 0; ConstraintIndex< ConstraintData.Num(); ++ConstraintIndex)
 				{
-					for (int32 ConstraintIndex= 0; ConstraintIndex< ConstraintData.Num(); ++ConstraintIndex)
+					// for now just try translate
+					const int32* TargetIndexPtr = ConstraintDataToTargets.Find(ConstraintIndex);
+					if (TargetIndexPtr)
 					{
-						// for now just try translate
-						const int32* TargetIndexPtr = ConstraintDataToTargets.Find(ConstraintIndex);
-						if (TargetIndexPtr)
-						{
-							const int32 TargetIndex = *TargetIndexPtr;
-							ConstraintData[ConstraintIndex].CurrentTransform = Targets[TargetIndex].Transform;
-							ConstraintData[ConstraintIndex].Weight = Targets[TargetIndex].Weight;
-						}
+						const int32 TargetIndex = *TargetIndexPtr;
+						ConstraintData[ConstraintIndex].CurrentTransform = Targets[TargetIndex].Transform;
+						ConstraintData[ConstraintIndex].Weight = Targets[TargetIndex].Weight;
 					}
+				}
 
-					FTransform InputBaseTransform = UtilityHelpers::GetBaseTransformByMode(BaseTransformSpace, [Hierarchy](const FRigElementKey& Item) { return Hierarchy->GetGlobalTransform(Item); },
-							Hierarchy->GetFirstParent(Item), BaseItem, BaseTransform);
+				FTransform InputBaseTransform = UtilityHelpers::GetBaseTransformByMode(BaseTransformSpace, [Hierarchy](const FRigElementKey& Item) { return Hierarchy->GetGlobalTransform(Item); },
+						Hierarchy->GetFirstParent(Item), BaseItem, BaseTransform);
 
-					FTransform SourceTransform = Hierarchy->GetGlobalTransform(Item);
+				FTransform SourceTransform = Hierarchy->GetGlobalTransform(Item);
 
-					// @todo: ignore maintain offset for now
-					FTransform ConstrainedTransform = AnimationCore::SolveConstraints(SourceTransform, InputBaseTransform, ConstraintData);
+				// @todo: ignore maintain offset for now
+				FTransform ConstrainedTransform = AnimationCore::SolveConstraints(SourceTransform, InputBaseTransform, ConstraintData);
 
-					Hierarchy->SetGlobalTransform(Item, ConstrainedTransform);
-				}
+				Hierarchy->SetGlobalTransform(Item, ConstrainedTransform);
 			}
 		}
 	}
@@ -195,12 +181,6 @@ FRigUnit_ParentConstraint_Execute()
 		return;
 	}
 
-	if(ExecuteContext.UnitContext.State == EControlRigState::Init)
-	{
-		ChildCache.Reset();
-		ParentCaches.Reset();
-	}
-	
  	URigHierarchy* Hierarchy = ExecuteContext.Hierarchy;	
 	if (Hierarchy)
 	{
@@ -636,12 +616,6 @@ FRigUnit_PositionConstraintLocalSpaceOffset_Execute()
 		return;
 	}
 
-	if(ExecuteContext.UnitContext.State == EControlRigState::Init)
-	{
-		ChildCache.Reset();
-		ParentCaches.Reset();
-	}
-
 	URigHierarchy* Hierarchy = ExecuteContext.Hierarchy;	
 	if (Hierarchy)
 	{
@@ -1118,12 +1092,6 @@ FRigUnit_RotationConstraintLocalSpaceOffset_Execute()
 		return;
 	}
 	
-	if(ExecuteContext.UnitContext.State == EControlRigState::Init)
-	{
-		ChildCache.Reset();
-		ParentCaches.Reset();
-	}
-
 	URigHierarchy* Hierarchy = ExecuteContext.Hierarchy;	
 	if (Hierarchy)
 	{
@@ -1731,12 +1699,6 @@ FRigUnit_ScaleConstraintLocalSpaceOffset_Execute()
 		return;
 	}
 	
-	if(ExecuteContext.UnitContext.State == EControlRigState::Init)
-	{
-		ChildCache.Reset();
-		ParentCaches.Reset();
-	}
-
 	URigHierarchy* Hierarchy = ExecuteContext.Hierarchy;	
 	if (Hierarchy)
 	{
diff --git a/Engine/Plugins/Animation/ControlRig/Source/ControlRig/Private/Units/Highlevel/Hierarchy/RigUnit_TwoBoneIKSimple.cpp b/Engine/Plugins/Animation/ControlRig/Source/ControlRig/Private/Units/Highlevel/Hierarchy/RigUnit_TwoBoneIKSimple.cpp
--- a/Engine/Plugins/Animation/ControlRig/Source/ControlRig/Private/Units/Highlevel/Hierarchy/RigUnit_TwoBoneIKSimple.cpp
+++ b/Engine/Plugins/Animation/ControlRig/Source/ControlRig/Private/Units/Highlevel/Hierarchy/RigUnit_TwoBoneIKSimple.cpp
@@ -75,15 +75,6 @@ FRigUnit_TwoBoneIKSimplePerItem_Execute()
 		return;
 	}
 
-	if (ExecuteContext.UnitContext.State == EControlRigState::Init)
-	{
-		CachedItemAIndex.Reset();
-		CachedItemBIndex.Reset();
-		CachedEffectorItemIndex.Reset();
-		CachedPoleVectorSpaceIndex.Reset();
-		return;
-	}
-
 	if (!CachedItemAIndex.UpdateCache(ItemA, Hierarchy) ||
 		!CachedItemBIndex.UpdateCache(ItemB, Hierarchy))
 	{
diff --git a/Engine/Plugins/Animation/ControlRig/Source/ControlRig/Private/Units/Math/RigUnit_MathBool.cpp b/Engine/Plugins/Animation/ControlRig/Source/ControlRig/Private/Units/Math/RigUnit_MathBool.cpp
--- a/Engine/Plugins/Animation/ControlRig/Source/ControlRig/Private/Units/Math/RigUnit_MathBool.cpp
+++ b/Engine/Plugins/Animation/ControlRig/Source/ControlRig/Private/Units/Math/RigUnit_MathBool.cpp
@@ -85,18 +85,11 @@ FRigUnit_MathBoolToggled_Execute()
 {
 	DECLARE_SCOPE_HIERARCHICAL_COUNTER_RIGUNIT()
 
-	if(ExecuteContext.UnitContext.State == EControlRigState::Init)
-	{
-		Initialized = false;
-		LastValue = Value;
-		Toggled = false;
-		return;
-	}
-
 	if(!Initialized)
 	{
 		Initialized = true;
 		Toggled = false;
+		LastValue = Value;
 	}
 	else
 	{
@@ -110,14 +103,7 @@ FRigUnit_MathBoolFlipFlop_Execute()
 {
 	DECLARE_SCOPE_HIERARCHICAL_COUNTER_RIGUNIT()
 
-	if(ExecuteContext.UnitContext.State == EControlRigState::Init)
-	{
-		Result = LastValue = !StartValue;
-		TimeLeft = 0.f;
-		return;
-	}
-
-	TimeLeft = TimeLeft - ExecuteContext.UnitContext.DeltaTime;
+	TimeLeft = TimeLeft - ExecuteContext.GetDeltaTime();
 	if (TimeLeft <= 0.f)
 	{
 		LastValue = !LastValue;
@@ -131,14 +117,6 @@ FRigUnit_MathBoolOnce_Execute()
 {
 	DECLARE_SCOPE_HIERARCHICAL_COUNTER_RIGUNIT()
 
-	if(ExecuteContext.UnitContext.State == EControlRigState::Init)
-	{
-		Result = false;
-		LastValue = true;
-		TimeLeft = 0.f;
-		return;
-	}
-
 	if(!Result && LastValue)
 	{
 		TimeLeft = Duration;
@@ -146,7 +124,7 @@ FRigUnit_MathBoolOnce_Execute()
 
 	Result = LastValue;
 
-	TimeLeft -= ExecuteContext.UnitContext.DeltaTime;
+	TimeLeft -= ExecuteContext.GetDeltaTime();
 	if(TimeLeft <= 0.f)
 	{
 		LastValue = false;
diff --git a/Engine/Plugins/Animation/ControlRig/Source/ControlRig/Private/Units/Math/RigUnit_MathBool.h b/Engine/Plugins/Animation/ControlRig/Source/ControlRig/Private/Units/Math/RigUnit_MathBool.h
--- a/Engine/Plugins/Animation/ControlRig/Source/ControlRig/Private/Units/Math/RigUnit_MathBool.h
+++ b/Engine/Plugins/Animation/ControlRig/Source/ControlRig/Private/Units/Math/RigUnit_MathBool.h
@@ -298,7 +298,8 @@ struct CONTROLRIG_API FRigUnit_MathBoolFlipFlop : public FRigUnit_MathBoolBase
 	FRigUnit_MathBoolFlipFlop()
 	{
 		Duration = TimeLeft = 0.f;
-		Result = LastValue = StartValue = false;
+		StartValue = false;
+		Result = LastValue = !StartValue;
 	}
 
 	// The initial value to use for the flag
@@ -335,7 +336,8 @@ struct CONTROLRIG_API FRigUnit_MathBoolOnce : public FRigUnit_MathBoolBase
 	FRigUnit_MathBoolOnce()
 	{
 		Duration = TimeLeft = 0.f;
-		Result = LastValue = false;
+		Result = false;
+		LastValue = true;
 	}
 
 	/**
diff --git a/Engine/Plugins/Animation/ControlRig/Source/ControlRig/Private/Units/Math/RigUnit_MathQuaternionTests.cpp b/Engine/Plugins/Animation/ControlRig/Source/ControlRig/Private/Units/Math/RigUnit_MathQuaternionTests.cpp
--- a/Engine/Plugins/Animation/ControlRig/Source/ControlRig/Private/Units/Math/RigUnit_MathQuaternionTests.cpp
+++ b/Engine/Plugins/Animation/ControlRig/Source/ControlRig/Private/Units/Math/RigUnit_MathQuaternionTests.cpp
@@ -46,7 +46,7 @@ IMPLEMENT_RIGUNIT_AUTOMATION_TEST(FRigUnit_MathQuaternionFromAxisAndAngle)
 {
 	Unit.Axis = FVector(3.f, 0.f, 0.f);
 	Unit.Angle = HALF_PI;
-	InitAndExecute();
+	Execute();
 	AddErrorIfFalse(FRigUnit_MathQuatTest_Utils::IsNearlyEqual(Unit.Result, FVector(-90.f, 0.f, 0.f)), TEXT("unexpected result"));
 	return true;
 }
@@ -54,11 +54,11 @@ IMPLEMENT_RIGUNIT_AUTOMATION_TEST(FRigUnit_MathQuaternionFromAxisAndAngle)
 IMPLEMENT_RIGUNIT_AUTOMATION_TEST(FRigUnit_MathQuaternionFromEuler)
 {
 	Unit.Euler = FVector(30.f, 20.f, 10.f);
-	InitAndExecute();
+	Execute();
 	FVector Euler = AnimationCore::EulerFromQuat(Unit.Result, Unit.RotationOrder);
 	AddErrorIfFalse(FRigUnit_MathQuatTest_Utils::IsNearlyEqual(Euler, FVector(30.f, 20.f, 10.f)), TEXT("unexpected result"));
 	Unit.RotationOrder = EEulerRotationOrder::YZX;
-	InitAndExecute();
+	Execute();
 	Euler = AnimationCore::EulerFromQuat(Unit.Result, Unit.RotationOrder);
 	AddErrorIfFalse(FRigUnit_MathQuatTest_Utils::IsNearlyEqual(Euler, FVector(30.f, 20.f, 10.f)), TEXT("unexpected result"));
 	return true;
@@ -67,7 +67,7 @@ IMPLEMENT_RIGUNIT_AUTOMATION_TEST(FRigUnit_MathQuaternionFromEuler)
 IMPLEMENT_RIGUNIT_AUTOMATION_TEST(FRigUnit_MathQuaternionFromRotator)
 {
 	Unit.Rotator = FRotator(30.f, 20.f, 10.f);
-	InitAndExecute();
+	Execute();
 	AddErrorIfFalse(FRigUnit_MathQuatTest_Utils::IsNearlyEqual(Unit.Result, FVector(10.f, 30.f, 20.f)), TEXT("unexpected result"));
 	return true;
 }
@@ -76,7 +76,7 @@ IMPLEMENT_RIGUNIT_AUTOMATION_TEST(FRigUnit_MathQuaternionFromTwoVectors)
 {
 	Unit.A = FVector(0.f, 0.f, 2.f);
 	Unit.B = FVector(0.f, 4.f, 0.f);
-	InitAndExecute();
+	Execute();
 	AddErrorIfFalse(FRigUnit_MathQuatTest_Utils::IsNearlyEqual(Unit.Result, FVector(90, 0.f, 0.f)), TEXT("unexpected result"));
 	return true;
 }
@@ -84,7 +84,7 @@ IMPLEMENT_RIGUNIT_AUTOMATION_TEST(FRigUnit_MathQuaternionFromTwoVectors)
 IMPLEMENT_RIGUNIT_AUTOMATION_TEST(FRigUnit_MathQuaternionToAxisAndAngle)
 {
 	Unit.Value = FQuat::MakeFromEuler(FVector(-90.f, 0.f, 0.f));
-	InitAndExecute();
+	Execute();
 	AddErrorIfFalse(FRigUnit_MathQuatTest_Utils::IsNearlyEqual(Unit.Axis, FVector(1.f, 0.f, 0.f)), TEXT("unexpected result"));
 	AddErrorIfFalse(FRigUnit_MathQuatTest_Utils::IsNearlyEqual(Unit.Angle, HALF_PI), TEXT("unexpected result"));
 	return true;
@@ -93,11 +93,11 @@ IMPLEMENT_RIGUNIT_AUTOMATION_TEST(FRigUnit_MathQuaternionToAxisAndAngle)
 IMPLEMENT_RIGUNIT_AUTOMATION_TEST(FRigUnit_MathQuaternionToEuler)
 {
 	Unit.Value = AnimationCore::QuatFromEuler(FVector(10.f, 20.f, 30.f), Unit.RotationOrder);
-	InitAndExecute();
+	Execute();
 	AddErrorIfFalse(FRigUnit_MathQuatTest_Utils::IsNearlyEqual(Unit.Result, FVector(10.f, 20.f, 30.f)), TEXT("unexpected result"));
 	Unit.RotationOrder = EEulerRotationOrder::YZX;
 	Unit.Value = AnimationCore::QuatFromEuler(FVector(10.f, 20.f, 30.f), Unit.RotationOrder);
-	InitAndExecute();
+	Execute();
 	AddErrorIfFalse(FRigUnit_MathQuatTest_Utils::IsNearlyEqual(Unit.Result, FVector(10.f, 20.f, 30.f)), TEXT("unexpected result"));
 	return true;
 }
@@ -105,7 +105,7 @@ IMPLEMENT_RIGUNIT_AUTOMATION_TEST(FRigUnit_MathQuaternionToEuler)
 IMPLEMENT_RIGUNIT_AUTOMATION_TEST(FRigUnit_MathQuaternionToRotator)
 {
 	Unit.Value = FQuat::MakeFromEuler(FVector(10.f, 30.f, 20.f));
-	InitAndExecute();
+	Execute();
 	AddErrorIfFalse(FRigUnit_MathQuatTest_Utils::IsNearlyEqual(Unit.Result, FRotator(30.f, 20.f, 10.f)), TEXT("unexpected result"));
 	return true;
 }
@@ -114,7 +114,7 @@ IMPLEMENT_RIGUNIT_AUTOMATION_TEST(FRigUnit_MathQuaternionMul)
 {
 	Unit.A = FQuat::MakeFromEuler(FVector(90.f, 0.f, 0.f));
 	Unit.B = FQuat::MakeFromEuler(FVector(0.f, 90.f, 0.f));
-	InitAndExecute();
+	Execute();
 	AddErrorIfFalse(FRigUnit_MathQuatTest_Utils::IsNearlyEqual(Unit.Result, FVector(90.f, 0.f, 90.f)), TEXT("unexpected result"));
 	return true;
 }
@@ -122,7 +122,7 @@ IMPLEMENT_RIGUNIT_AUTOMATION_TEST(FRigUnit_MathQuaternionMul)
 IMPLEMENT_RIGUNIT_AUTOMATION_TEST(FRigUnit_MathQuaternionInverse)
 {
 	Unit.Value = FQuat::MakeFromEuler(FVector(90.f, 0.f, 0.f));
-	InitAndExecute();
+	Execute();
 	AddErrorIfFalse(FRigUnit_MathQuatTest_Utils::IsNearlyEqual(Unit.Result, FVector(-90.f, 0.f, 0.f)), TEXT("unexpected result"));
 	return true;
 }
@@ -132,7 +132,7 @@ IMPLEMENT_RIGUNIT_AUTOMATION_TEST(FRigUnit_MathQuaternionSlerp)
 	Unit.A = FQuat::MakeFromEuler(FVector(60.f, 0.f, 0.f));
 	Unit.B = FQuat::MakeFromEuler(FVector(20.f, 0.f, 0.f));
 	Unit.T = 0.5f;
-	InitAndExecute();
+	Execute();
 	AddErrorIfFalse(FRigUnit_MathQuatTest_Utils::IsNearlyEqual(Unit.Result, FVector(40.f, 0.f, 0.f)), TEXT("unexpected result"));
 	return true;
 }
@@ -141,10 +141,10 @@ IMPLEMENT_RIGUNIT_AUTOMATION_TEST(FRigUnit_MathQuaternionEquals)
 {
 	Unit.A = FQuat(1.f, 0.f, 0.f, 0.f);
 	Unit.B = FQuat(1.f, 0.f, 0.f, 0.f);
-	InitAndExecute();
+	Execute();
 	AddErrorIfFalse(Unit.Result == true, TEXT("unexpected result"));
 	Unit.B = FQuat(1.f, 1.f, 0.f, 0.f);
-	InitAndExecute();
+	Execute();
 	AddErrorIfFalse(Unit.Result == false, TEXT("unexpected result"));
 	return true;
 }
@@ -153,10 +153,10 @@ IMPLEMENT_RIGUNIT_AUTOMATION_TEST(FRigUnit_MathQuaternionNotEquals)
 {
 	Unit.A = FQuat(1.f, 0.f, 0.f, 0.f);
 	Unit.B = FQuat(1.f, 0.f, 0.f, 0.f);
-	InitAndExecute();
+	Execute();
 	AddErrorIfFalse(Unit.Result == false, TEXT("unexpected result"));
 	Unit.B = FQuat(1.f, 1.f, 0.f, 0.f);
-	InitAndExecute();
+	Execute();
 	AddErrorIfFalse(Unit.Result == true, TEXT("unexpected result"));
 	return true;
 }
@@ -166,10 +166,10 @@ IMPLEMENT_RIGUNIT_AUTOMATION_TEST(FRigUnit_MathQuaternionSelectBool)
 	Unit.Condition = false;
 	Unit.IfTrue = FQuat(1.f, 0.f, 0.f, 0.f);
 	Unit.IfFalse = FQuat(1.f, 1.f, 0.f, 0.f);
-	InitAndExecute();
+	Execute();
 	AddErrorIfFalse(FRigUnit_MathQuatTest_Utils::IsNearlyEqual(Unit.Result, FQuat(1.f, 1.f, 0.f, 0.f)), TEXT("unexpected result"));
 	Unit.Condition = true;
-	InitAndExecute();
+	Execute();
 	AddErrorIfFalse(FRigUnit_MathQuatTest_Utils::IsNearlyEqual(Unit.Result, FQuat(1.f, 0.f, 0.f, 0.f)), TEXT("unexpected result"));
 	return true;
 }
@@ -178,7 +178,7 @@ IMPLEMENT_RIGUNIT_AUTOMATION_TEST(FRigUnit_MathQuaternionDot)
 {
 	Unit.A = FQuat(1.f, 2.f, 3.f, 4.f);
 	Unit.B = FQuat(5.f, 6.f, 7.f, 8.f);
-	InitAndExecute();
+	Execute();
 	AddErrorIfFalse(FRigUnit_MathQuatTest_Utils::IsNearlyEqual(Unit.Result, Unit.A | Unit.B), TEXT("unexpected result"));
 	return true;
 }
@@ -186,7 +186,7 @@ IMPLEMENT_RIGUNIT_AUTOMATION_TEST(FRigUnit_MathQuaternionDot)
 IMPLEMENT_RIGUNIT_AUTOMATION_TEST(FRigUnit_MathQuaternionUnit)
 {
 	Unit.Value = FQuat(FVector(3.f, 0.f, 0.f), HALF_PI);
-	InitAndExecute();
+	Execute();
 	AddErrorIfFalse(FRigUnit_MathQuatTest_Utils::IsNearlyEqual(Unit.Result, FVector(-143.130096, 0.f, 0.f)), TEXT("unexpected result"));
 	return true;
 }
@@ -195,7 +195,7 @@ IMPLEMENT_RIGUNIT_AUTOMATION_TEST(FRigUnit_MathQuaternionRotateVector)
 {
 	Unit.Transform = FQuat(FVector(1.f, 0.f, 0.f), -HALF_PI);
 	Unit.Vector = FVector(0.f, 0.f, 1.f);
-	InitAndExecute();
+	Execute();
 	AddErrorIfFalse(FRigUnit_MathQuatTest_Utils::IsNearlyEqual(Unit.Result, FVector(0.f, 1.f, 0.f)), TEXT("unexpected result"));
 	return true;
 }
@@ -204,13 +204,13 @@ IMPLEMENT_RIGUNIT_AUTOMATION_TEST(FRigUnit_MathQuaternionGetAxis)
 {
 	Unit.Quaternion = FQuat::Identity;
 	Unit.Axis = EAxis::X;
-	InitAndExecute();
+	Execute();
 	AddErrorIfFalse(FRigUnit_MathQuatTest_Utils::IsNearlyEqual(Unit.Result, FVector(1.f, 0.f, 0.f)), TEXT("unexpected result"));
 	Unit.Axis = EAxis::Y;
-	InitAndExecute();
+	Execute();
 	AddErrorIfFalse(FRigUnit_MathQuatTest_Utils::IsNearlyEqual(Unit.Result, FVector(0.f, 1.f, 0.f)), TEXT("unexpected result"));
 	Unit.Axis = EAxis::Z;
-	InitAndExecute();
+	Execute();
 	AddErrorIfFalse(FRigUnit_MathQuatTest_Utils::IsNearlyEqual(Unit.Result, FVector(0.f, 0.f, 1.f)), TEXT("unexpected result"));
 	return true;
 }
@@ -224,11 +224,6 @@ IMPLEMENT_RIGUNIT_AUTOMATION_TEST(FRigUnit_MathQuaternionMakeRelative)
 	TransformUnit.Global.SetRotation(Unit.Global);
 	TransformUnit.Parent.SetRotation(Unit.Parent);
 
-	ExecuteContext.UnitContext.State = EControlRigState::Init;
-	Unit.Execute();
-	TransformUnit.Execute();
-
-	ExecuteContext.UnitContext.State = EControlRigState::Update;
 	Unit.Execute();
 	TransformUnit.Execute();
 
@@ -245,11 +240,6 @@ IMPLEMENT_RIGUNIT_AUTOMATION_TEST(FRigUnit_MathQuaternionMakeAbsolute)
 	TransformUnit.Local.SetRotation(Unit.Local);
 	TransformUnit.Parent.SetRotation(Unit.Parent);
 
-	ExecuteContext.UnitContext.State = EControlRigState::Init;
-	Unit.Execute();
-	TransformUnit.Execute();
-
-	ExecuteContext.UnitContext.State = EControlRigState::Update;
 	Unit.Execute();
 	TransformUnit.Execute();
 
diff --git a/Engine/Plugins/Animation/ControlRig/Source/ControlRig/Private/Units/Math/RigUnit_MathRBFInterpolate.cpp b/Engine/Plugins/Animation/ControlRig/Source/ControlRig/Private/Units/Math/RigUnit_MathRBFInterpolate.cpp
--- a/Engine/Plugins/Animation/ControlRig/Source/ControlRig/Private/Units/Math/RigUnit_MathRBFInterpolate.cpp
+++ b/Engine/Plugins/Animation/ControlRig/Source/ControlRig/Private/Units/Math/RigUnit_MathRBFInterpolate.cpp
@@ -170,7 +170,6 @@ static uint64 HashValues(std::initializer_list<uint64> Values, uint64 Seed)
 
 template<typename T>
 void FRigUnit_MathRBFInterpolateQuatBase::GetInterpolatedWeights(
-	EControlRigState State,
 	FRigUnit_MathRBFInterpolateQuatWorkData& WorkData,
 	const TArrayView<const T>& Targets,
 	const FQuat &Input,
@@ -193,7 +192,7 @@ void FRigUnit_MathRBFInterpolateQuatBase::GetInterpolatedWeights(
 	FRigUnit_MathRBFQuatWeightFunctor WeightFunc(
 		DistanceFunction, SmoothingFunction, FMath::DegreesToRadians(SmoothingAngle), TwistAxis);
 
-	if (State == EControlRigState::Init)
+	if (WorkData.Targets.Num() != Quats.Num())
 	{
 		WorkData.Targets = Quats;
 		WorkData.Hash = HashTargets(TArrayView<const FQuat>(Quats.GetData(), Quats.Num()));
@@ -234,7 +233,6 @@ uint64 FRigUnit_MathRBFInterpolateQuatBase::HashTargets(const TArrayView<const F
 
 template<typename T>
 void FRigUnit_MathRBFInterpolateVectorBase::GetInterpolatedWeights(
-	EControlRigState State,
 	FRigUnit_MathRBFInterpolateVectorWorkData& WorkData,
 	const TArrayView<const T>& Targets,
 	const FVector& Input,
@@ -256,7 +254,7 @@ void FRigUnit_MathRBFInterpolateVectorBase::GetInterpolatedWeights(
 	FRigUnit_MathRBFVectorWeightFunctor WeightFunc(
 		DistanceFunction, SmoothingFunction, SmoothingRadius);
 
-	if (State == EControlRigState::Init)
+	if (WorkData.Targets.Num() != Vectors.Num())
 	{
 		WorkData.Targets = Vectors;
 		WorkData.Hash = HashTargets(TArrayView<const FVector>(Vectors.GetData(), Vectors.Num()));
@@ -301,7 +299,7 @@ FRigUnit_MathRBFInterpolateQuatFloat_Execute()
 	const TArrayView<const FMathRBFInterpolateQuatFloat_Target> TargetsView(Targets.GetData(), Targets.Num());
 	TArray<float> Weights;
 	GetInterpolatedWeights<FMathRBFInterpolateQuatFloat_Target>(
-		ExecuteContext.UnitContext.State, WorkData, TargetsView, Input,  
+		WorkData, TargetsView, Input,  
 		DistanceFunction, SmoothingFunction, SmoothingAngle, 
 		bNormalizeOutput, TwistAxis, Weights);
 
@@ -320,7 +318,7 @@ FRigUnit_MathRBFInterpolateQuatVector_Execute()
 	const TArrayView<const FMathRBFInterpolateQuatVector_Target> TargetsView(Targets.GetData(), Targets.Num());
 	TArray<float> Weights;
 	GetInterpolatedWeights<FMathRBFInterpolateQuatVector_Target>(
-		ExecuteContext.UnitContext.State, WorkData, TargetsView, Input,
+		WorkData, TargetsView, Input,
 		DistanceFunction, SmoothingFunction, SmoothingAngle, 
 		bNormalizeOutput, TwistAxis, Weights);
 
@@ -339,7 +337,7 @@ FRigUnit_MathRBFInterpolateQuatColor_Execute()
 	const TArrayView<const FMathRBFInterpolateQuatColor_Target> TargetsView(Targets.GetData(), Targets.Num());
 	TArray<float> Weights;
 	GetInterpolatedWeights<FMathRBFInterpolateQuatColor_Target>(
-		ExecuteContext.UnitContext.State, WorkData, TargetsView, Input,
+		WorkData, TargetsView, Input,
 		DistanceFunction, SmoothingFunction, SmoothingAngle, 
 		bNormalizeOutput, TwistAxis, Weights);
 
@@ -357,7 +355,7 @@ FRigUnit_MathRBFInterpolateQuatQuat_Execute()
 	const TArrayView<const FMathRBFInterpolateQuatQuat_Target> TargetsView(Targets.GetData(), Targets.Num());
 	TArray<float> Weights;
 	GetInterpolatedWeights<FMathRBFInterpolateQuatQuat_Target>(
-		ExecuteContext.UnitContext.State, WorkData, TargetsView, Input,
+		WorkData, TargetsView, Input,
 		DistanceFunction, SmoothingFunction, SmoothingAngle,
 		bNormalizeOutput, TwistAxis, Weights);
 	
@@ -377,7 +375,7 @@ FRigUnit_MathRBFInterpolateQuatXform_Execute()
 	const TArrayView<const FMathRBFInterpolateQuatXform_Target> TargetsView(Targets.GetData(), Targets.Num());
 	TArray<float> Weights;
 	GetInterpolatedWeights<FMathRBFInterpolateQuatXform_Target>(
-		ExecuteContext.UnitContext.State, WorkData, TargetsView, Input,
+		WorkData, TargetsView, Input,
 		DistanceFunction, SmoothingFunction, SmoothingAngle,
 		bNormalizeOutput, TwistAxis, Weights);
 
@@ -407,7 +405,7 @@ FRigUnit_MathRBFInterpolateVectorFloat_Execute()
 	const TArrayView<const FMathRBFInterpolateVectorFloat_Target> TargetsView(Targets.GetData(), Targets.Num());
 	TArray<float> Weights;
 	GetInterpolatedWeights<FMathRBFInterpolateVectorFloat_Target>(
-		ExecuteContext.UnitContext.State, WorkData, TargetsView, Input,
+		WorkData, TargetsView, Input,
 		DistanceFunction, SmoothingFunction, SmoothingRadius, 
 		bNormalizeOutput, Weights);
 
@@ -426,7 +424,7 @@ FRigUnit_MathRBFInterpolateVectorVector_Execute()
 	const TArrayView<const FMathRBFInterpolateVectorVector_Target> TargetsView(Targets.GetData(), Targets.Num());
 	TArray<float> Weights;
 	GetInterpolatedWeights<FMathRBFInterpolateVectorVector_Target>(
-		ExecuteContext.UnitContext.State, WorkData, TargetsView, Input,
+		WorkData, TargetsView, Input,
 		DistanceFunction, SmoothingFunction, SmoothingRadius, 
 		bNormalizeOutput, Weights);
 
@@ -445,7 +443,7 @@ FRigUnit_MathRBFInterpolateVectorColor_Execute()
 	const TArrayView<const FMathRBFInterpolateVectorColor_Target> TargetsView(Targets.GetData(), Targets.Num());
 	TArray<float> Weights;
 	GetInterpolatedWeights<FMathRBFInterpolateVectorColor_Target>(
-		ExecuteContext.UnitContext.State, WorkData, TargetsView, Input,
+		WorkData, TargetsView, Input,
 		DistanceFunction, SmoothingFunction, SmoothingRadius, 
 		bNormalizeOutput, Weights);
 
@@ -464,7 +462,7 @@ FRigUnit_MathRBFInterpolateVectorQuat_Execute()
 	const TArrayView<const FMathRBFInterpolateVectorQuat_Target> TargetsView(Targets.GetData(), Targets.Num());
 	TArray<float> Weights;
 	GetInterpolatedWeights<FMathRBFInterpolateVectorQuat_Target>(
-		ExecuteContext.UnitContext.State, WorkData, TargetsView, Input,
+		WorkData, TargetsView, Input,
 		DistanceFunction, SmoothingFunction, SmoothingRadius, 
 		bNormalizeOutput, Weights);
 
@@ -484,7 +482,7 @@ FRigUnit_MathRBFInterpolateVectorXform_Execute()
 	const TArrayView<const FMathRBFInterpolateVectorXform_Target> TargetsView(Targets.GetData(), Targets.Num());
 	TArray<float> Weights;
 	GetInterpolatedWeights<FMathRBFInterpolateVectorXform_Target>(
-		ExecuteContext.UnitContext.State, WorkData, TargetsView, Input,
+		WorkData, TargetsView, Input,
 		DistanceFunction, SmoothingFunction, SmoothingRadius, 
 		bNormalizeOutput, Weights);
 
diff --git a/Engine/Plugins/Animation/ControlRig/Source/ControlRig/Private/Units/Math/RigUnit_MathRBFInterpolate.h b/Engine/Plugins/Animation/ControlRig/Source/ControlRig/Private/Units/Math/RigUnit_MathRBFInterpolate.h
--- a/Engine/Plugins/Animation/ControlRig/Source/ControlRig/Private/Units/Math/RigUnit_MathRBFInterpolate.h
+++ b/Engine/Plugins/Animation/ControlRig/Source/ControlRig/Private/Units/Math/RigUnit_MathRBFInterpolate.h
@@ -7,7 +7,6 @@
 
 #include "RigUnit_MathRBFInterpolate.generated.h"
 
-enum class EControlRigState : uint8;
 template<typename T> class TRBFInterpolator;
 struct CONTROLRIG_API FRigUnit_MathRBFQuatWeightFunctor;
 struct CONTROLRIG_API FRigUnit_MathRBFVectorWeightFunctor;
@@ -110,7 +109,6 @@ protected:
 	
 	template<typename T>
 	static void GetInterpolatedWeights(
-		EControlRigState State,
 		FRigUnit_MathRBFInterpolateQuatWorkData& WorkData,
 		const TArrayView<const T>& Targets,
 		const FQuat& Input,
@@ -152,7 +150,6 @@ struct CONTROLRIG_API FRigUnit_MathRBFInterpolateVectorBase :
 protected:
 	template<typename T>
 	static void GetInterpolatedWeights(
-		EControlRigState State,
 		FRigUnit_MathRBFInterpolateVectorWorkData& WorkData,
 		const TArrayView<const T>& Targets,
 		const FVector& Input,
diff --git a/Engine/Plugins/Animation/ControlRig/Source/ControlRig/Private/Units/Math/RigUnit_MathTransformTests.cpp b/Engine/Plugins/Animation/ControlRig/Source/ControlRig/Private/Units/Math/RigUnit_MathTransformTests.cpp
--- a/Engine/Plugins/Animation/ControlRig/Source/ControlRig/Private/Units/Math/RigUnit_MathTransformTests.cpp
+++ b/Engine/Plugins/Animation/ControlRig/Source/ControlRig/Private/Units/Math/RigUnit_MathTransformTests.cpp
@@ -43,7 +43,7 @@ namespace FRigUnit_MathTransformTest_Utils
 IMPLEMENT_RIGUNIT_AUTOMATION_TEST(FRigUnit_MathTransformFromEulerTransform)
 {
 	Unit.EulerTransform = FEulerTransform(FVector(1.f, 2.f, 3.f), FRotator(10.f, 20.f, 30.f), FVector(4.f, 5.f, 6.f));
-	InitAndExecute();
+	Execute();
 	AddErrorIfFalse(FRigUnit_MathTransformTest_Utils::IsNearlyEqual(Unit.Result.GetLocation(), FVector(1.f, 2.f, 3.f)), TEXT("unexpected result"));
 	AddErrorIfFalse(FRigUnit_MathTransformTest_Utils::IsNearlyEqual(Unit.Result.GetRotation(), FVector(30.f, 10.f, 20.f)), TEXT("unexpected result"));
 	AddErrorIfFalse(FRigUnit_MathTransformTest_Utils::IsNearlyEqual(Unit.Result.GetScale3D(), FVector(4.f, 5.f, 6.f)), TEXT("unexpected result"));
@@ -53,7 +53,7 @@ IMPLEMENT_RIGUNIT_AUTOMATION_TEST(FRigUnit_MathTransformFromEulerTransform)
 IMPLEMENT_RIGUNIT_AUTOMATION_TEST(FRigUnit_MathTransformToEulerTransform)
 {
 	Unit.Value = FTransform(FQuat::MakeFromEuler(FVector(30.f, 10.f, 20.f)), FVector(1.f, 2.f, 3.f), FVector(4.f, 5.f, 6.f));
-	InitAndExecute();
+	Execute();
 	AddErrorIfFalse(FRigUnit_MathTransformTest_Utils::IsNearlyEqual(Unit.Result.Location, FVector(1.f, 2.f, 3.f)), TEXT("unexpected result"));
 	AddErrorIfFalse(FRigUnit_MathTransformTest_Utils::IsNearlyEqual(Unit.Result.Rotation, FRotator(10.f, 20.f, 30.f)), TEXT("unexpected result"));
 	AddErrorIfFalse(FRigUnit_MathTransformTest_Utils::IsNearlyEqual(Unit.Result.Scale, FVector(4.f, 5.f, 6.f)), TEXT("unexpected result"));
@@ -64,7 +64,7 @@ IMPLEMENT_RIGUNIT_AUTOMATION_TEST(FRigUnit_MathTransformMul)
 {
 	Unit.A = FTransform(FQuat::Identity, FVector(1.f, 2.f, 3.f));
 	Unit.B = FTransform(FQuat::Identity, FVector(4.f, 5.f, 6.f), FVector(3.f, 3.f, 3.f));
-	InitAndExecute();
+	Execute();
 	AddErrorIfFalse(FRigUnit_MathTransformTest_Utils::IsNearlyEqual(Unit.Result.GetLocation(), FVector(7.f, 11.f, 15.f)), TEXT("unexpected result"));
 	return true;
 }
@@ -73,7 +73,7 @@ IMPLEMENT_RIGUNIT_AUTOMATION_TEST(FRigUnit_MathTransformMakeRelative)
 {
 	Unit.Global = FTransform(FQuat::Identity, FVector(7.f, 11.f, 15.f));
 	Unit.Parent = FTransform(FQuat::Identity, FVector(4.f, 5.f, 6.f), FVector(3.f, 3.f, 3.f));
-	InitAndExecute();
+	Execute();
 	AddErrorIfFalse(FRigUnit_MathTransformTest_Utils::IsNearlyEqual(Unit.Local.GetLocation(), FVector(1.f, 2.f, 3.f)), TEXT("unexpected result"));
 	return true;
 }
@@ -81,7 +81,7 @@ IMPLEMENT_RIGUNIT_AUTOMATION_TEST(FRigUnit_MathTransformMakeRelative)
 IMPLEMENT_RIGUNIT_AUTOMATION_TEST(FRigUnit_MathTransformInverse)
 {
 	Unit.Value = FTransform(FQuat::Identity, FVector(1.f, 2.f, 3.f));
-	InitAndExecute();
+	Execute();
 	AddErrorIfFalse(FRigUnit_MathTransformTest_Utils::IsNearlyEqual(Unit.Result.GetLocation(), FVector(-1.f, -2.f, -3.f)), TEXT("unexpected result"));
 	return true;
 }
@@ -91,7 +91,7 @@ IMPLEMENT_RIGUNIT_AUTOMATION_TEST(FRigUnit_MathTransformLerp)
 	Unit.A = FTransform(FQuat::Identity, FVector(1.f, 2.f, 3.f));
 	Unit.B = FTransform(FQuat::Identity, FVector(4.f, 5.f, 6.f));
 	Unit.T = 0.5f;
-	InitAndExecute();
+	Execute();
 	AddErrorIfFalse(FRigUnit_MathTransformTest_Utils::IsNearlyEqual(Unit.Result.GetLocation(), FVector(2.5f, 3.5f, 4.5f)), TEXT("unexpected result"));
 	return true;
 }
@@ -101,10 +101,10 @@ IMPLEMENT_RIGUNIT_AUTOMATION_TEST(FRigUnit_MathTransformSelectBool)
 	Unit.Condition = false;
 	Unit.IfTrue = FTransform(FQuat::Identity, FVector(1.f, 2.f, 3.f));
 	Unit.IfFalse = FTransform(FQuat::Identity, FVector(4.f, 5.f, 6.f));
-	InitAndExecute();
+	Execute();
 	AddErrorIfFalse(FRigUnit_MathTransformTest_Utils::IsNearlyEqual(Unit.Result.GetLocation(), FVector(4.f, 5.f, 6.f)), TEXT("unexpected result"));
 	Unit.Condition = true;
-	InitAndExecute();
+	Execute();
 	AddErrorIfFalse(FRigUnit_MathTransformTest_Utils::IsNearlyEqual(Unit.Result.GetLocation(), FVector(1.f, 2.f, 3.f)), TEXT("unexpected result"));
 	return true;
 }
@@ -113,7 +113,7 @@ IMPLEMENT_RIGUNIT_AUTOMATION_TEST(FRigUnit_MathTransformRotateVector)
 {
 	Unit.Transform = FTransform(FQuat(FVector(1.f, 0.f, 0.f), -HALF_PI), FVector(1.f, 2.f, 3.f));
 	Unit.Vector = FVector(0.f, 0.f, 4.f);
-	InitAndExecute();
+	Execute();
 	AddErrorIfFalse(FRigUnit_MathTransformTest_Utils::IsNearlyEqual(Unit.Result, FVector(0.f, 4.f, 0.f)), TEXT("unexpected result"));
 	return true;
 }
@@ -122,7 +122,7 @@ IMPLEMENT_RIGUNIT_AUTOMATION_TEST(FRigUnit_MathTransformTransformVector)
 {
 	Unit.Transform = FTransform(FQuat(FVector(1.f, 0.f, 0.f), -HALF_PI), FVector(1.f, 2.f, 3.f));
 	Unit.Location = FVector(0.f, 0.f, 4.f);
-	InitAndExecute();
+	Execute();
 	AddErrorIfFalse(FRigUnit_MathTransformTest_Utils::IsNearlyEqual(Unit.Result, FVector(1.f, 6.f, 3.f)), TEXT("unexpected result"));
 	return true;
 }
@@ -136,7 +136,7 @@ IMPLEMENT_RIGUNIT_AUTOMATION_TEST(FRigUnit_MathTransformArrayToSRT)
 		Unit.Transforms[Index] = FTransform(FQuat(FVector(1.f, 0.f, 0.f), HALF_PI * float(Index)), FVector(float(Index), 0.f, 0.f));
 	}
 
-	InitAndExecute();
+	Execute();
 
 	for (int32 Index = 0; Index < Unit.Transforms.Num(); Index++)
 	{
