ControlRig: Remove EControlRigState from units

diff --git a/Engine/Plugins/Animation/ControlRig/Source/ControlRig/Private/Units/Math/RigUnit_MathVectorTests.cpp b/Engine/Plugins/Animation/ControlRig/Source/ControlRig/Private/Units/Math/RigUnit_MathVectorTests.cpp
--- a/Engine/Plugins/Animation/ControlRig/Source/ControlRig/Private/Units/Math/RigUnit_MathVectorTests.cpp
+++ b/Engine/Plugins/Animation/ControlRig/Source/ControlRig/Private/Units/Math/RigUnit_MathVectorTests.cpp
@@ -18,7 +18,7 @@ IMPLEMENT_RIGUNIT_AUTOMATION_TEST(FRigUnit_MathVectorAdd)
 {
 	Unit.A = FVector(4.f, 6.f, 8.f);
 	Unit.B = FVector(1.f, 2.f, 3.f);
-	InitAndExecute();
+	Execute();
 	AddErrorIfFalse(FRigUnit_MathVectorTest_Utils::IsNearlyEqual(Unit.Result, FVector(5.f, 8.f, 11.f)), TEXT("unexpected result"));
 	return true;
 }
@@ -26,7 +26,7 @@ IMPLEMENT_RIGUNIT_AUTOMATION_TEST(FRigUnit_MathVectorAdd)
 IMPLEMENT_RIGUNIT_AUTOMATION_TEST(FRigUnit_MathVectorFromFloat)
 {
 	Unit.Value = 7.f;
-	InitAndExecute();
+	Execute();
 	AddErrorIfFalse(FRigUnit_MathVectorTest_Utils::IsNearlyEqual(Unit.Result, FVector(7.f, 7.f, 7.f)), TEXT("unexpected result"));
 	return true;
 }
@@ -35,7 +35,7 @@ IMPLEMENT_RIGUNIT_AUTOMATION_TEST(FRigUnit_MathVectorSub)
 {
 	Unit.A = FVector(4.f, 6.f, 8.f);
 	Unit.B = FVector(1.f, 2.f, 3.f);
-	InitAndExecute();
+	Execute();
 	AddErrorIfFalse(FRigUnit_MathVectorTest_Utils::IsNearlyEqual(Unit.Result, FVector(3.f, 4.f, 5.f)), TEXT("unexpected result"));
 	return true;
 }
@@ -44,7 +44,7 @@ IMPLEMENT_RIGUNIT_AUTOMATION_TEST(FRigUnit_MathVectorMul)
 {
 	Unit.A = FVector(4.f, 6.f, 8.f);
 	Unit.B = FVector(1.f, 2.f, 3.f);
-	InitAndExecute();
+	Execute();
 	AddErrorIfFalse(FRigUnit_MathVectorTest_Utils::IsNearlyEqual(Unit.Result, FVector(4.f, 12.f, 24.f)), TEXT("unexpected result"));
 	return true;
 }
@@ -53,7 +53,7 @@ IMPLEMENT_RIGUNIT_AUTOMATION_TEST(FRigUnit_MathVectorScale)
 {
 	Unit.Value = FVector(4.f, 6.f, 8.f);
 	Unit.Factor = 2.f;
-	InitAndExecute();
+	Execute();
 	AddErrorIfFalse(FRigUnit_MathVectorTest_Utils::IsNearlyEqual(Unit.Result, FVector(8.f, 12.f, 16.f)), TEXT("unexpected result"));
 	return true;
 }
@@ -62,7 +62,7 @@ IMPLEMENT_RIGUNIT_AUTOMATION_TEST(FRigUnit_MathVectorDiv)
 {
 	Unit.A = FVector(4.f, 6.f, 8.f);
 	Unit.B = FVector(1.f, 2.f, 3.f);
-	InitAndExecute();
+	Execute();
 	AddErrorIfFalse(FRigUnit_MathVectorTest_Utils::IsNearlyEqual(Unit.Result, FVector(4.f, 3.f, 8.f / 3.f), KINDA_SMALL_NUMBER), TEXT("unexpected result"));
 	return true;
 }
@@ -71,7 +71,7 @@ IMPLEMENT_RIGUNIT_AUTOMATION_TEST(FRigUnit_MathVectorMod)
 {
 	Unit.A = FVector(4.f, 6.f, 9.f);
 	Unit.B = FVector(5.f, 4.f, 6.f);
-	InitAndExecute();
+	Execute();
 	AddErrorIfFalse(FRigUnit_MathVectorTest_Utils::IsNearlyEqual(Unit.Result, FVector(4.f, 2.f, 3.f)), TEXT("unexpected result"));
 	return true;
 }
@@ -80,7 +80,7 @@ IMPLEMENT_RIGUNIT_AUTOMATION_TEST(FRigUnit_MathVectorMin)
 {
 	Unit.A = FVector(4.f, 2.f, 8.f);
 	Unit.B = FVector(1.f, 6.f, 3.f);
-	InitAndExecute();
+	Execute();
 	AddErrorIfFalse(FRigUnit_MathVectorTest_Utils::IsNearlyEqual(Unit.Result, FVector(1.f, 2.f, 3.f)), TEXT("unexpected result"));
 	return true;
 }
@@ -89,7 +89,7 @@ IMPLEMENT_RIGUNIT_AUTOMATION_TEST(FRigUnit_MathVectorMax)
 {
 	Unit.A = FVector(4.f, 2.f, 8.f);
 	Unit.B = FVector(1.f, 6.f, 3.f);
-	InitAndExecute();
+	Execute();
 	AddErrorIfFalse(FRigUnit_MathVectorTest_Utils::IsNearlyEqual(Unit.Result, FVector(4.f, 6.f, 8.f)), TEXT("unexpected result"));
 	return true;
 }
@@ -97,7 +97,7 @@ IMPLEMENT_RIGUNIT_AUTOMATION_TEST(FRigUnit_MathVectorMax)
 IMPLEMENT_RIGUNIT_AUTOMATION_TEST(FRigUnit_MathVectorNegate)
 {
 	Unit.Value = FVector(4.f, -2.f, 8.f);
-	InitAndExecute();
+	Execute();
 	AddErrorIfFalse(FRigUnit_MathVectorTest_Utils::IsNearlyEqual(Unit.Result, FVector(-4.f, 2.f, -8.f)), TEXT("unexpected result"));
 	return true;
 }
@@ -105,7 +105,7 @@ IMPLEMENT_RIGUNIT_AUTOMATION_TEST(FRigUnit_MathVectorNegate)
 IMPLEMENT_RIGUNIT_AUTOMATION_TEST(FRigUnit_MathVectorAbs)
 {
 	Unit.Value = FVector(-4.f, 2.f, -8.f);
-	InitAndExecute();
+	Execute();
 	AddErrorIfFalse(FRigUnit_MathVectorTest_Utils::IsNearlyEqual(Unit.Result, FVector(4.f, 2.f, 8.f)), TEXT("unexpected result"));
 	return true;
 }
@@ -113,7 +113,7 @@ IMPLEMENT_RIGUNIT_AUTOMATION_TEST(FRigUnit_MathVectorAbs)
 IMPLEMENT_RIGUNIT_AUTOMATION_TEST(FRigUnit_MathVectorFloor)
 {
 	Unit.Value = FVector(4.1f, 2.6f, 8.f);
-	InitAndExecute();
+	Execute();
 	AddErrorIfFalse(FRigUnit_MathVectorTest_Utils::IsNearlyEqual(Unit.Result, FVector(4.f, 2.f, 8.f)), TEXT("unexpected result"));
 	return true;
 }
@@ -121,7 +121,7 @@ IMPLEMENT_RIGUNIT_AUTOMATION_TEST(FRigUnit_MathVectorFloor)
 IMPLEMENT_RIGUNIT_AUTOMATION_TEST(FRigUnit_MathVectorCeil)
 {
 	Unit.Value = FVector(4.1f, 2.6f, 8.f);
-	InitAndExecute();
+	Execute();
 	AddErrorIfFalse(FRigUnit_MathVectorTest_Utils::IsNearlyEqual(Unit.Result, FVector(5.f, 3.f, 8.f)), TEXT("unexpected result"));
 	return true;
 }
@@ -129,7 +129,7 @@ IMPLEMENT_RIGUNIT_AUTOMATION_TEST(FRigUnit_MathVectorCeil)
 IMPLEMENT_RIGUNIT_AUTOMATION_TEST(FRigUnit_MathVectorRound)
 {
 	Unit.Value = FVector(4.1f, 2.6f, 8.f);
-	InitAndExecute();
+	Execute();
 	AddErrorIfFalse(FRigUnit_MathVectorTest_Utils::IsNearlyEqual(Unit.Result, FVector(4.f, 3.f, 8.f)), TEXT("unexpected result"));
 	return true;
 }
@@ -137,7 +137,7 @@ IMPLEMENT_RIGUNIT_AUTOMATION_TEST(FRigUnit_MathVectorRound)
 IMPLEMENT_RIGUNIT_AUTOMATION_TEST(FRigUnit_MathVectorSign)
 {
 	Unit.Value = FVector(4.1f, -2.6f, 0.f);
-	InitAndExecute();
+	Execute();
 	AddErrorIfFalse(FRigUnit_MathVectorTest_Utils::IsNearlyEqual(Unit.Result, FVector(1.f, -1.f, 1.f)), TEXT("unexpected result"));
 	return true;
 }
@@ -147,7 +147,7 @@ IMPLEMENT_RIGUNIT_AUTOMATION_TEST(FRigUnit_MathVectorClamp)
 	Unit.Value = FVector(-3.f, 2.f, 8.f);
 	Unit.Minimum = FVector(0.f, 1.f, 2.f);
 	Unit.Maximum = FVector(3.f, 4.f, 5.f);
-	InitAndExecute();
+	Execute();
 	AddErrorIfFalse(FRigUnit_MathVectorTest_Utils::IsNearlyEqual(Unit.Result, FVector(0.f, 2.f, 5.f)), TEXT("unexpected result"));
 	return true;
 }
@@ -157,7 +157,7 @@ IMPLEMENT_RIGUNIT_AUTOMATION_TEST(FRigUnit_MathVectorLerp)
 	Unit.A = FVector(2.f, 4.f, 8.f);
 	Unit.B = FVector(12.f, 24.f, 38.f);
 	Unit.T = 0.5f;
-	InitAndExecute();
+	Execute();
 	AddErrorIfFalse(FRigUnit_MathVectorTest_Utils::IsNearlyEqual(Unit.Result, FVector(7.f, 14.f, 23.f)), TEXT("unexpected result"));
 	return true;
 }
@@ -169,7 +169,7 @@ IMPLEMENT_RIGUNIT_AUTOMATION_TEST(FRigUnit_MathVectorRemap)
 	Unit.SourceMaximum = FVector(4.f, 14.f, 24.f);
 	Unit.TargetMinimum = FVector(36.f, 46.f, 56.f);
 	Unit.TargetMaximum = FVector(38.f, 48.f, 58.f);
-	InitAndExecute();
+	Execute();
 	AddErrorIfFalse(FRigUnit_MathVectorTest_Utils::IsNearlyEqual(Unit.Result, FVector(37.f, 47.f, 57.f)), TEXT("unexpected result"));
 	return true;
 }
@@ -178,13 +178,13 @@ IMPLEMENT_RIGUNIT_AUTOMATION_TEST(FRigUnit_MathVectorEquals)
 {
 	Unit.A = FVector(1.f, 2.f, 3.f);
 	Unit.B = FVector(1.f, 2.f, 3.f);
-	InitAndExecute();
+	Execute();
 	AddErrorIfFalse(Unit.Result == true, TEXT("unexpected result"));
 	Unit.B = FVector(1.f, 2.f, 4.f);
-	InitAndExecute();
+	Execute();
 	AddErrorIfFalse(Unit.Result == false, TEXT("unexpected result"));
 	Unit.B = FVector(2.f, 3.f, 4.f);
-	InitAndExecute();
+	Execute();
 	AddErrorIfFalse(Unit.Result == false, TEXT("unexpected result"));
 	return true;
 }
@@ -193,13 +193,13 @@ IMPLEMENT_RIGUNIT_AUTOMATION_TEST(FRigUnit_MathVectorNotEquals)
 {
 	Unit.A = FVector(1.f, 2.f, 3.f);
 	Unit.B = FVector(1.f, 2.f, 3.f);
-	InitAndExecute();
+	Execute();
 	AddErrorIfFalse(Unit.Result == false, TEXT("unexpected result"));
 	Unit.B = FVector(1.f, 2.f, 4.f);
-	InitAndExecute();
+	Execute();
 	AddErrorIfFalse(Unit.Result == true, TEXT("unexpected result"));
 	Unit.B = FVector(2.f, 3.f, 4.f);
-	InitAndExecute();
+	Execute();
 	AddErrorIfFalse(Unit.Result == true, TEXT("unexpected result"));
 	return true;
 }
@@ -208,13 +208,13 @@ IMPLEMENT_RIGUNIT_AUTOMATION_TEST(FRigUnit_MathVectorIsNearlyZero)
 {
 	Unit.Value = FVector(0.f, 0.f, 0.f);
 	Unit.Tolerance = 0.f;
-	InitAndExecute();
+	Execute();
 	AddErrorIfFalse(Unit.Result == true, TEXT("unexpected result"));
 	Unit.Value = FVector(0.001f, 0.003f, 0.005f);
-	InitAndExecute();
+	Execute();
 	AddErrorIfFalse(Unit.Result == false, TEXT("unexpected result"));
 	Unit.Tolerance = 0.01f;
-	InitAndExecute();
+	Execute();
 	AddErrorIfFalse(Unit.Result == true, TEXT("unexpected result"));
 	return true;
 }
@@ -224,13 +224,13 @@ IMPLEMENT_RIGUNIT_AUTOMATION_TEST(FRigUnit_MathVectorIsNearlyEqual)
 	Unit.A = FVector(1.f, 2.f, 3.f);
 	Unit.B = FVector(1.f, 2.f, 3.f);
 	Unit.Tolerance = 0.f;
-	InitAndExecute();
+	Execute();
 	AddErrorIfFalse(Unit.Result == true, TEXT("unexpected result"));
 	Unit.B= FVector(1.001f, 2.003f, 3.005f);
-	InitAndExecute();
+	Execute();
 	AddErrorIfFalse(Unit.Result == false, TEXT("unexpected result"));
 	Unit.Tolerance = 0.01f;
-	InitAndExecute();
+	Execute();
 	AddErrorIfFalse(Unit.Result == true, TEXT("unexpected result"));
 	return true;
 }
@@ -240,10 +240,10 @@ IMPLEMENT_RIGUNIT_AUTOMATION_TEST(FRigUnit_MathVectorSelectBool)
 	Unit.Condition = false;
 	Unit.IfTrue = FVector(1.f, 2.f, 3.f);
 	Unit.IfFalse = FVector(4.f, 5.f, 7.f);
-	InitAndExecute();
+	Execute();
 	AddErrorIfFalse(FRigUnit_MathVectorTest_Utils::IsNearlyEqual(Unit.Result, FVector(4.f, 5.f, 7.f)), TEXT("unexpected result"));
 	Unit.Condition = true;
-	InitAndExecute();
+	Execute();
 	AddErrorIfFalse(FRigUnit_MathVectorTest_Utils::IsNearlyEqual(Unit.Result, FVector(1.f, 2.f, 3.f)), TEXT("unexpected result"));
 	return true;
 }
@@ -253,7 +253,7 @@ IMPLEMENT_RIGUNIT_AUTOMATION_TEST(FRigUnit_MathVectorSelectBool)
 IMPLEMENT_RIGUNIT_AUTOMATION_TEST(FRigUnit_MathVectorDeg)
 {
 	Unit.Value = FVector(-PI, PI * 2.f, PI * 0.5f);
-	InitAndExecute();
+	Execute();
 	AddErrorIfFalse(FRigUnit_MathVectorTest_Utils::IsNearlyEqual(Unit.Result, FVector(-180.f, 360.f, 90.f)), TEXT("unexpected result"));
 	return true;
 }
@@ -261,7 +261,7 @@ IMPLEMENT_RIGUNIT_AUTOMATION_TEST(FRigUnit_MathVectorDeg)
 IMPLEMENT_RIGUNIT_AUTOMATION_TEST(FRigUnit_MathVectorRad)
 {
 	Unit.Value = FVector(-180.f, 360.f, 90.f);
-	InitAndExecute();
+	Execute();
 	AddErrorIfFalse(FRigUnit_MathVectorTest_Utils::IsNearlyEqual(Unit.Result, FVector(-PI, PI * 2.f, PI * 0.5f)), TEXT("unexpected result"));
 	return true;
 }
@@ -271,7 +271,7 @@ IMPLEMENT_RIGUNIT_AUTOMATION_TEST(FRigUnit_MathVectorRad)
 IMPLEMENT_RIGUNIT_AUTOMATION_TEST(FRigUnit_MathVectorLengthSquared)
 {
 	Unit.Value = FVector(2.f, 3.f, 4.f);
-	InitAndExecute();
+	Execute();
 	AddErrorIfFalse(FMath::IsNearlyEqual(Unit.Result, 29.f), TEXT("unexpected result"));
 	return true;
 }
@@ -279,7 +279,7 @@ IMPLEMENT_RIGUNIT_AUTOMATION_TEST(FRigUnit_MathVectorLengthSquared)
 IMPLEMENT_RIGUNIT_AUTOMATION_TEST(FRigUnit_MathVectorLength)
 {
 	Unit.Value = FVector(2.f, 3.f, 4.f);
-	InitAndExecute();
+	Execute();
 	AddErrorIfFalse(FMath::IsNearlyEqual(Unit.Result, FMath::Sqrt(29.f)), TEXT("unexpected result"));
 	return true;
 }
@@ -288,7 +288,7 @@ IMPLEMENT_RIGUNIT_AUTOMATION_TEST(FRigUnit_MathVectorDistance)
 {
 	Unit.A = FVector(3.f, 4.f, 5.f);
 	Unit.B = FVector(5.f, 7.f, 9.f);
-	InitAndExecute();
+	Execute();
 	AddErrorIfFalse(FMath::IsNearlyEqual(Unit.Result, FMath::Sqrt(29.f)), TEXT("unexpected result"));
 	return true;
 }
@@ -297,7 +297,7 @@ IMPLEMENT_RIGUNIT_AUTOMATION_TEST(FRigUnit_MathVectorCross)
 {
 	Unit.A = FVector(1.f, 0.f, 0.f);
 	Unit.B = FVector(0.f, 2.f, 0.f);
-	InitAndExecute();
+	Execute();
 	AddErrorIfFalse(FRigUnit_MathVectorTest_Utils::IsNearlyEqual(Unit.Result, FVector(0.f, 0.f, 2.f)), TEXT("unexpected result"));
 	return true;
 }
@@ -306,7 +306,7 @@ IMPLEMENT_RIGUNIT_AUTOMATION_TEST(FRigUnit_MathVectorDot)
 {
 	Unit.A = FVector(1.f, 2.f, 3.f);
 	Unit.B = FVector(4.f, 5.f, 6.f);
-	InitAndExecute();
+	Execute();
 	AddErrorIfFalse(FMath::IsNearlyEqual(Unit.Result, 32.f), TEXT("unexpected result"));
 	return true;
 }
@@ -314,13 +314,13 @@ IMPLEMENT_RIGUNIT_AUTOMATION_TEST(FRigUnit_MathVectorDot)
 IMPLEMENT_RIGUNIT_AUTOMATION_TEST(FRigUnit_MathVectorUnit)
 {
 	Unit.Value = FVector(4.f, 0.f, 0.f);
-	InitAndExecute();
+	Execute();
 	AddErrorIfFalse(FRigUnit_MathVectorTest_Utils::IsNearlyEqual(Unit.Result, FVector(1.f, 0.f, 0.f), 0.001f), TEXT("unexpected result (0)"));
 	Unit.Value = FVector(0.f, 5.f, 0.f);
-	InitAndExecute();
+	Execute();
 	AddErrorIfFalse(FRigUnit_MathVectorTest_Utils::IsNearlyEqual(Unit.Result, FVector(0.f, 1.f, 0.f), 0.001f), TEXT("unexpected result (1)"));
 	Unit.Value = FVector(0.f, 0.f, 6.f);
-	InitAndExecute();
+	Execute();
 	AddErrorIfFalse(FRigUnit_MathVectorTest_Utils::IsNearlyEqual(Unit.Result, FVector(0.f, 0.f, 1.f), 0.001f), TEXT("unexpected result (2)"));
 	return true;
 }
@@ -329,11 +329,11 @@ IMPLEMENT_RIGUNIT_AUTOMATION_TEST(FRigUnit_MathVectorSetLength)
 {
 	Unit.Value = FVector(4.f, 0.f, 0.f);
 	Unit.Length = 2.f;
-	InitAndExecute();
+	Execute();
 	AddErrorIfFalse(FRigUnit_MathVectorTest_Utils::IsNearlyEqual(Unit.Result, FVector(2.f, 0.f, 0.f), 0.001f), TEXT("unexpected result (0)"));
 	Unit.Value = FVector(0.f, 5.f, 0.f);
 	Unit.Length = 3.f;
-	InitAndExecute();
+	Execute();
 	AddErrorIfFalse(FRigUnit_MathVectorTest_Utils::IsNearlyEqual(Unit.Result, FVector(0.f, 3.f, 0.f), 0.001f), TEXT("unexpected result (1)"));
 	return true;
 }
@@ -343,13 +343,13 @@ IMPLEMENT_RIGUNIT_AUTOMATION_TEST(FRigUnit_MathVectorClampLength)
 	Unit.Value = FVector(4.f, 0.f, 0.f);
 	Unit.MinimumLength = 2.f;
 	Unit.MaximumLength = 6.f;
-	InitAndExecute();
+	Execute();
 	AddErrorIfFalse(FRigUnit_MathVectorTest_Utils::IsNearlyEqual(Unit.Result, FVector(4.f, 0.f, 0.f), 0.001f), TEXT("unexpected result (0)"));
 	Unit.Value = FVector(1.f, 0.f, 0.f);
-	InitAndExecute();
+	Execute();
 	AddErrorIfFalse(FRigUnit_MathVectorTest_Utils::IsNearlyEqual(Unit.Result, FVector(2.f, 0.f, 0.f), 0.001f), TEXT("unexpected result (1)"));
 	Unit.Value = FVector(8.f, 0.f, 0.f);
-	InitAndExecute();
+	Execute();
 	AddErrorIfFalse(FRigUnit_MathVectorTest_Utils::IsNearlyEqual(Unit.Result, FVector(6.f, 0.f, 0.f), 0.001f), TEXT("unexpected result (1)"));
 	return true;
 }
@@ -358,10 +358,10 @@ IMPLEMENT_RIGUNIT_AUTOMATION_TEST(FRigUnit_MathVectorMirror)
 {
 	Unit.Value = FVector(1.f, 2.f, 3.f);
 	Unit.Normal = FVector(1.f, 0.f, 0.f);
-	InitAndExecute();
+	Execute();
 	AddErrorIfFalse(FRigUnit_MathVectorTest_Utils::IsNearlyEqual(Unit.Result, FVector(-1.f, 2.f, 3.f)), TEXT("unexpected result"));
 	Unit.Normal = FVector(0.f, 3.f, 0.f);
-	InitAndExecute();
+	Execute();
 	AddErrorIfFalse(FRigUnit_MathVectorTest_Utils::IsNearlyEqual(Unit.Result, FVector(1.f, -2.f, 3.f), 1e-6), TEXT("unexpected result"));
 	return true;
 }
@@ -370,7 +370,7 @@ IMPLEMENT_RIGUNIT_AUTOMATION_TEST(FRigUnit_MathVectorAngle)
 {
 	Unit.A = FVector(-4.f, 0.f, 0.f);
 	Unit.B = FVector(0.f, 3.f, 0.f);
-	InitAndExecute();
+	Execute();
 	AddErrorIfFalse(FMath::IsNearlyEqual(Unit.Result, HALF_PI), TEXT("unexpected result"));
 	return true;
 }
@@ -379,10 +379,10 @@ IMPLEMENT_RIGUNIT_AUTOMATION_TEST(FRigUnit_MathVectorParallel)
 {
 	Unit.A = FVector(-4.f, 0.f, 0.f);
 	Unit.B = FVector(1.f, 0.f, 0.f);
-	InitAndExecute();
+	Execute();
 	AddErrorIfFalse(Unit.Result == true, TEXT("unexpected result"));
 	Unit.B = FVector(0.f, 1.f, 0.f);
-	InitAndExecute();
+	Execute();
 	AddErrorIfFalse(Unit.Result == false, TEXT("unexpected result"));
 	return true;
 }
@@ -391,10 +391,10 @@ IMPLEMENT_RIGUNIT_AUTOMATION_TEST(FRigUnit_MathVectorOrthogonal)
 {
 	Unit.A = FVector(-4.f, 0.f, 0.f);
 	Unit.B = FVector(1.f, 0.f, 0.f);
-	InitAndExecute();
+	Execute();
 	AddErrorIfFalse(Unit.Result == false, TEXT("unexpected result"));
 	Unit.B = FVector(0.f, 1.f, 0.f);
-	InitAndExecute();
+	Execute();
 	AddErrorIfFalse(Unit.Result == true, TEXT("unexpected result"));
 	return true;
 }
@@ -404,14 +404,14 @@ IMPLEMENT_RIGUNIT_AUTOMATION_TEST(FRigUnit_MathDistanceToPlane)
 	Unit.Point = FVector(5, -6, 3);
 	Unit.PlanePoint = FVector(4, 5, 0);
 	Unit.PlaneNormal = FVector(3, -2, 1);
-	InitAndExecute();
+	Execute();
 	AddErrorIfFalse(FRigUnit_MathVectorTest_Utils::IsNearlyEqual(Unit.ClosestPointOnPlane, FVector(-1,-2,1), KINDA_SMALL_NUMBER), TEXT("unexpected result"));
 	AddErrorIfFalse(FMath::IsNearlyEqual(Unit.SignedDistance, 7.48331477f, KINDA_SMALL_NUMBER), TEXT("unexpected result"));
 	
 	Unit.Point = FVector(5, -6, 3);
 	Unit.PlanePoint = FVector(2, 2, 2);
 	Unit.PlaneNormal = FVector(0,0,0);
-	InitAndExecute();
+	Execute();
 	AddErrorIfFalse(Unit.ClosestPointOnPlane.IsZero(), TEXT("unexpected result"));
 	AddErrorIfFalse(Unit.SignedDistance == 0, TEXT("unexpected result"));
 	return true;
diff --git a/Engine/Plugins/Animation/ControlRig/Source/ControlRig/Private/Units/Math/RigUnit_Noise.cpp b/Engine/Plugins/Animation/ControlRig/Source/ControlRig/Private/Units/Math/RigUnit_Noise.cpp
--- a/Engine/Plugins/Animation/ControlRig/Source/ControlRig/Private/Units/Math/RigUnit_Noise.cpp
+++ b/Engine/Plugins/Animation/ControlRig/Source/ControlRig/Private/Units/Math/RigUnit_Noise.cpp
@@ -8,29 +8,19 @@
 FRigUnit_NoiseFloat_Execute()
 {
     DECLARE_SCOPE_HIERARCHICAL_COUNTER_RIGUNIT()
-	if (ExecuteContext.UnitContext.State == EControlRigState::Init)
-	{
-		Time = 0.f;
-		return;
-	}
 
 	const float Noise = FMath::PerlinNoise1D(Value * Frequency + Time) + 0.5f;
 	Result = FMath::Lerp<float>(Minimum, Maximum, Noise);
-	Time = Time + Speed * ExecuteContext.UnitContext.DeltaTime;
+	Time = Time + Speed * ExecuteContext.GetDeltaTime();
 }
 
 FRigUnit_NoiseDouble_Execute()
 {
 	DECLARE_SCOPE_HIERARCHICAL_COUNTER_RIGUNIT()
-	if (ExecuteContext.UnitContext.State == EControlRigState::Init)
-	{
-		Time = 0.f;
-		return;
-	}
 
 	const double Noise = double(FMath::PerlinNoise1D(Value * Frequency + Time)) + 0.5;
 	Result = FMath::Lerp<double>(Minimum, Maximum, Noise);
-	Time = Time + Speed * ExecuteContext.UnitContext.DeltaTime;
+	Time = Time + Speed * ExecuteContext.GetDeltaTime();
 }
 
 FRigUnit_NoiseVector_Execute()
@@ -55,11 +45,6 @@ FRigVMStructUpgradeInfo FRigUnit_NoiseVector::GetUpgradeInfo() const
 FRigUnit_NoiseVector2_Execute()
 {
 	DECLARE_SCOPE_HIERARCHICAL_COUNTER_RIGUNIT()
-	if (ExecuteContext.UnitContext.State == EControlRigState::Init)
-	{
-		Time = FVector::ZeroVector;
-		return;
-	}
 
 	const double NoiseX = double(FMath::PerlinNoise1D(Value.X * Frequency.X + Time.X)) + 0.5;
 	const double NoiseY = double(FMath::PerlinNoise1D(Value.Y * Frequency.Y + Time.Y)) + 0.5;
@@ -67,6 +52,6 @@ FRigUnit_NoiseVector2_Execute()
 	Result.X = FMath::Lerp<double>(Minimum, Maximum, NoiseX);
 	Result.Y = FMath::Lerp<double>(Minimum, Maximum, NoiseY);
 	Result.Z = FMath::Lerp<double>(Minimum, Maximum, NoiseZ);
-	Time = Time + Speed * ExecuteContext.UnitContext.DeltaTime;
+	Time = Time + Speed * ExecuteContext.GetDeltaTime();
 }
 
diff --git a/Engine/Plugins/Animation/ControlRig/Source/ControlRig/Private/Units/Math/RigUnit_Random.cpp b/Engine/Plugins/Animation/ControlRig/Source/ControlRig/Private/Units/Math/RigUnit_Random.cpp
--- a/Engine/Plugins/Animation/ControlRig/Source/ControlRig/Private/Units/Math/RigUnit_Random.cpp
+++ b/Engine/Plugins/Animation/ControlRig/Source/ControlRig/Private/Units/Math/RigUnit_Random.cpp
@@ -20,19 +20,13 @@ float FRigUnit_Random_Helper(int32& Seed)
 FRigUnit_RandomFloat_Execute()
 {
     DECLARE_SCOPE_HIERARCHICAL_COUNTER_RIGUNIT()
-	if (ExecuteContext.UnitContext.State == EControlRigState::Init)
-	{
-		LastSeed = BaseSeed = Seed;
-		TimeLeft = 0.f;
-		return;
-	}
 
 	if(Seed != BaseSeed)
 	{
 		LastSeed = BaseSeed = Seed;
 	}
 
-	TimeLeft = TimeLeft - ExecuteContext.UnitContext.DeltaTime;
+	TimeLeft = TimeLeft - ExecuteContext.GetDeltaTime();
 	if (TimeLeft > 0.f)
 	{
 		Result = LastResult;
@@ -48,19 +42,13 @@ FRigUnit_RandomFloat_Execute()
 FRigUnit_RandomVector_Execute()
 {
     DECLARE_SCOPE_HIERARCHICAL_COUNTER_RIGUNIT()
-	if (ExecuteContext.UnitContext.State == EControlRigState::Init)
-	{
-		LastSeed = BaseSeed = Seed;
-		TimeLeft = 0.f;
-		return;
-	}
 
 	if(Seed != BaseSeed)
 	{
 		LastSeed = BaseSeed = Seed;
 	}
 
-	TimeLeft = TimeLeft - ExecuteContext.UnitContext.DeltaTime;
+	TimeLeft = TimeLeft - ExecuteContext.GetDeltaTime();
 	if (TimeLeft > 0.f)
 	{
 		Result = LastResult;
diff --git a/Engine/Plugins/Animation/ControlRig/Source/ControlRig/Private/Units/Simulation/RigUnit_Accumulate.cpp b/Engine/Plugins/Animation/ControlRig/Source/ControlRig/Private/Units/Simulation/RigUnit_Accumulate.cpp
--- a/Engine/Plugins/Animation/ControlRig/Source/ControlRig/Private/Units/Simulation/RigUnit_Accumulate.cpp
+++ b/Engine/Plugins/Animation/ControlRig/Source/ControlRig/Private/Units/Simulation/RigUnit_Accumulate.cpp
@@ -8,13 +8,15 @@
 FRigUnit_AccumulateFloatAdd_Execute()
 {
     DECLARE_SCOPE_HIERARCHICAL_COUNTER_RIGUNIT()
-	if (ExecuteContext.UnitContext.State == EControlRigState::Init)
+	if (!bIsInitialized)
 	{
 		Result = AccumulatedValue = InitialValue;
+		bIsInitialized = true;
 	}
-	else if(bIntegrateDeltaTime)
+	
+	if(bIntegrateDeltaTime)
 	{
-		Result = AccumulatedValue = AccumulatedValue + Increment * ExecuteContext.UnitContext.DeltaTime;
+		Result = AccumulatedValue = AccumulatedValue + Increment * ExecuteContext.GetDeltaTime();
 	}
 	else
 	{
@@ -25,13 +27,15 @@ FRigUnit_AccumulateFloatAdd_Execute()
 FRigUnit_AccumulateVectorAdd_Execute()
 {
     DECLARE_SCOPE_HIERARCHICAL_COUNTER_RIGUNIT()
-	if (ExecuteContext.UnitContext.State == EControlRigState::Init)
+	if (!bIsInitialized)
 	{
 		Result = AccumulatedValue = InitialValue;
+		bIsInitialized = true;
 	}
-	else if (bIntegrateDeltaTime)
+	
+	if (bIntegrateDeltaTime)
 	{
-		Result = AccumulatedValue = AccumulatedValue + Increment * ExecuteContext.UnitContext.DeltaTime;
+		Result = AccumulatedValue = AccumulatedValue + Increment * ExecuteContext.GetDeltaTime();
 	}
 	else
 	{
@@ -42,177 +46,167 @@ FRigUnit_AccumulateVectorAdd_Execute()
 FRigUnit_AccumulateFloatMul_Execute()
 {
     DECLARE_SCOPE_HIERARCHICAL_COUNTER_RIGUNIT()
-	if (ExecuteContext.UnitContext.State == EControlRigState::Init)
+	if (!bIsInitialized)
 	{
 		Result = AccumulatedValue = InitialValue;
+		bIsInitialized = true;
 	}
-	else
+
+	float Factor = Multiplier;
+	if (bIntegrateDeltaTime)
 	{
-		float Factor = Multiplier;
-		if (bIntegrateDeltaTime)
-		{
-			Factor = FMath::Lerp<float>(1.f, Factor, ExecuteContext.UnitContext.DeltaTime);
-		}
-		Result = AccumulatedValue = AccumulatedValue * Factor;
+		Factor = FMath::Lerp<float>(1.f, Factor, ExecuteContext.GetDeltaTime());
 	}
+	Result = AccumulatedValue = AccumulatedValue * Factor;
 }
 
 FRigUnit_AccumulateVectorMul_Execute()
 {
     DECLARE_SCOPE_HIERARCHICAL_COUNTER_RIGUNIT()
-	if (ExecuteContext.UnitContext.State == EControlRigState::Init)
+	if (!bIsInitialized)
 	{
 		Result = AccumulatedValue = InitialValue;
+		bIsInitialized = true;
 	}
-	else
+
+	FVector Factor = Multiplier;
+	if (bIntegrateDeltaTime)
 	{
-		FVector Factor = Multiplier;
-		if (bIntegrateDeltaTime)
-		{
-			Factor = FMath::Lerp<FVector>(FVector::OneVector, Factor, ExecuteContext.UnitContext.DeltaTime);
-		}
-		Result = AccumulatedValue = AccumulatedValue * Factor;
+		Factor = FMath::Lerp<FVector>(FVector::OneVector, Factor, ExecuteContext.GetDeltaTime());
 	}
+	Result = AccumulatedValue = AccumulatedValue * Factor;
 }
 
 FRigUnit_AccumulateQuatMul_Execute()
 {
     DECLARE_SCOPE_HIERARCHICAL_COUNTER_RIGUNIT()
-	if (ExecuteContext.UnitContext.State == EControlRigState::Init)
+	if (!bIsInitialized)
 	{
 		Result = AccumulatedValue = InitialValue;
+		bIsInitialized = true;
+	}
+
+	FQuat Factor = Multiplier;
+	if (bIntegrateDeltaTime)
+	{
+		Factor = FQuat::Slerp(FQuat::Identity, Factor, ExecuteContext.GetDeltaTime());
+	}
+	if (bFlipOrder)
+	{
+		Result = AccumulatedValue = Factor * AccumulatedValue;
 	}
 	else
 	{
-		FQuat Factor = Multiplier;
-		if (bIntegrateDeltaTime)
-		{
-			Factor = FQuat::Slerp(FQuat::Identity, Factor, ExecuteContext.UnitContext.DeltaTime);
-		}
-		if (bFlipOrder)
-		{
-			Result = AccumulatedValue = Factor * AccumulatedValue;
-		}
-		else
-		{
-			Result = AccumulatedValue = AccumulatedValue * Factor;
-		}
+		Result = AccumulatedValue = AccumulatedValue * Factor;
 	}
 }
 
 FRigUnit_AccumulateTransformMul_Execute()
 {
     DECLARE_SCOPE_HIERARCHICAL_COUNTER_RIGUNIT()
-	if (ExecuteContext.UnitContext.State == EControlRigState::Init)
+	if (!bIsInitialized)
 	{
 		Result = AccumulatedValue = InitialValue;
+		bIsInitialized = true;
+	}
+
+	FTransform Factor = Multiplier;
+	if (bIntegrateDeltaTime)
+	{
+		Factor.SetTranslation(FMath::Lerp<FVector>(FVector::OneVector, Factor.GetTranslation(), ExecuteContext.GetDeltaTime()));
+		Factor.SetRotation(FQuat::Slerp(FQuat::Identity, Factor.GetRotation(), ExecuteContext.GetDeltaTime()));
+		Factor.SetScale3D(FMath::Lerp<FVector>(FVector::OneVector, Factor.GetScale3D(), ExecuteContext.GetDeltaTime()));
+	}
+	if (bFlipOrder)
+	{
+		Result = AccumulatedValue = Factor * AccumulatedValue;
 	}
 	else
 	{
-		FTransform Factor = Multiplier;
-		if (bIntegrateDeltaTime)
-		{
-			Factor.SetTranslation(FMath::Lerp<FVector>(FVector::OneVector, Factor.GetTranslation(), ExecuteContext.UnitContext.DeltaTime));
-			Factor.SetRotation(FQuat::Slerp(FQuat::Identity, Factor.GetRotation(), ExecuteContext.UnitContext.DeltaTime));
-			Factor.SetScale3D(FMath::Lerp<FVector>(FVector::OneVector, Factor.GetScale3D(), ExecuteContext.UnitContext.DeltaTime));
-		}
-		if (bFlipOrder)
-		{
-			Result = AccumulatedValue = Factor * AccumulatedValue;
-		}
-		else
-		{
-			Result = AccumulatedValue = AccumulatedValue * Factor;
-		}
+		Result = AccumulatedValue = AccumulatedValue * Factor;
 	}
 }
 
 FRigUnit_AccumulateFloatLerp_Execute()
 {
     DECLARE_SCOPE_HIERARCHICAL_COUNTER_RIGUNIT()
-	if (ExecuteContext.UnitContext.State == EControlRigState::Init)
+	if (!bIsInitialized)
 	{
 		Result = AccumulatedValue = InitialValue;
+		bIsInitialized = true;
 	}
-	else
-	{
-		Result = AccumulatedValue = FMath::Lerp<float>(AccumulatedValue, TargetValue, FMath::Clamp<float>(bIntegrateDeltaTime ? ExecuteContext.UnitContext.DeltaTime * Blend : Blend, 0.f, 1.f));
-	}
+
+	Result = AccumulatedValue = FMath::Lerp<float>(AccumulatedValue, TargetValue, FMath::Clamp<float>(bIntegrateDeltaTime ? ExecuteContext.GetDeltaTime() * Blend : Blend, 0.f, 1.f));
 }
 
 FRigUnit_AccumulateVectorLerp_Execute()
 {
     DECLARE_SCOPE_HIERARCHICAL_COUNTER_RIGUNIT()
-	if (ExecuteContext.UnitContext.State == EControlRigState::Init)
+	if (!bIsInitialized)
 	{
 		Result = AccumulatedValue = InitialValue;
+		bIsInitialized = true;
 	}
-	else
-	{
-		Result = AccumulatedValue = FMath::Lerp<FVector>(AccumulatedValue, TargetValue, FMath::Clamp<float>(bIntegrateDeltaTime ? ExecuteContext.UnitContext.DeltaTime * Blend : Blend, 0.f, 1.f));
-	}
+
+	Result = AccumulatedValue = FMath::Lerp<FVector>(AccumulatedValue, TargetValue, FMath::Clamp<float>(bIntegrateDeltaTime ? ExecuteContext.GetDeltaTime() * Blend : Blend, 0.f, 1.f));
 }
 
 FRigUnit_AccumulateQuatLerp_Execute()
 {
     DECLARE_SCOPE_HIERARCHICAL_COUNTER_RIGUNIT()
-	if (ExecuteContext.UnitContext.State == EControlRigState::Init)
+	if (!bIsInitialized)
 	{
 		Result = AccumulatedValue = InitialValue;
+		bIsInitialized = true;
 	}
-	else
-	{
-		Result = AccumulatedValue = FQuat::Slerp(AccumulatedValue, TargetValue, FMath::Clamp<float>(bIntegrateDeltaTime ? ExecuteContext.UnitContext.DeltaTime * Blend : Blend, 0.f, 1.f));
-	}
+
+	Result = AccumulatedValue = FQuat::Slerp(AccumulatedValue, TargetValue, FMath::Clamp<float>(bIntegrateDeltaTime ? ExecuteContext.GetDeltaTime() * Blend : Blend, 0.f, 1.f));
 }
 
 FRigUnit_AccumulateTransformLerp_Execute()
 {
     DECLARE_SCOPE_HIERARCHICAL_COUNTER_RIGUNIT()
-	if (ExecuteContext.UnitContext.State == EControlRigState::Init)
+	if (!bIsInitialized)
 	{
 		Result = AccumulatedValue = InitialValue;
+		bIsInitialized = true;
 	}
-	else
-	{
-		float B = FMath::Clamp<float>(bIntegrateDeltaTime ? ExecuteContext.UnitContext.DeltaTime * Blend : Blend, 0.f, 1.f);
-		AccumulatedValue.SetTranslation(FMath::Lerp<FVector>(AccumulatedValue.GetTranslation(), TargetValue.GetTranslation(), B));
-		AccumulatedValue.SetRotation(FQuat::Slerp(AccumulatedValue.GetRotation(), TargetValue.GetRotation(), B));
-		AccumulatedValue.SetScale3D(FMath::Lerp<FVector>(AccumulatedValue.GetScale3D(), TargetValue.GetScale3D(), B));
-		Result = AccumulatedValue;
-	}
+
+	float B = FMath::Clamp<float>(bIntegrateDeltaTime ? ExecuteContext.GetDeltaTime() * Blend : Blend, 0.f, 1.f);
+	AccumulatedValue.SetTranslation(FMath::Lerp<FVector>(AccumulatedValue.GetTranslation(), TargetValue.GetTranslation(), B));
+	AccumulatedValue.SetRotation(FQuat::Slerp(AccumulatedValue.GetRotation(), TargetValue.GetRotation(), B));
+	AccumulatedValue.SetScale3D(FMath::Lerp<FVector>(AccumulatedValue.GetScale3D(), TargetValue.GetScale3D(), B));
+	Result = AccumulatedValue;
 }
 
 FRigUnit_AccumulateFloatRange_Execute()
 {
     DECLARE_SCOPE_HIERARCHICAL_COUNTER_RIGUNIT()
-	if (ExecuteContext.UnitContext.State == EControlRigState::Init)
+	if (!bIsInitialized)
 	{
 		Minimum = AccumulatedMinimum = Maximum = AccumulatedMaximum = Value;
+		bIsInitialized = true;
 	}
-	else
-	{
-		Minimum = AccumulatedMinimum = FMath::Min(AccumulatedMinimum, Value);
-		Maximum = AccumulatedMaximum = FMath::Max(AccumulatedMaximum, Value);
-	}
+
+	Minimum = AccumulatedMinimum = FMath::Min(AccumulatedMinimum, Value);
+	Maximum = AccumulatedMaximum = FMath::Max(AccumulatedMaximum, Value);
 }
 
 FRigUnit_AccumulateVectorRange_Execute()
 {
     DECLARE_SCOPE_HIERARCHICAL_COUNTER_RIGUNIT()
-	if (ExecuteContext.UnitContext.State == EControlRigState::Init)
+	if (!bIsInitialized)
 	{
 		Minimum = AccumulatedMinimum = Maximum = AccumulatedMaximum = Value;
+		bIsInitialized = true;
 	}
-	else
-	{
-		Minimum.X = AccumulatedMinimum.X = FMath::Min(AccumulatedMinimum.X, Value.X);
-		Minimum.Y = AccumulatedMinimum.Y = FMath::Min(AccumulatedMinimum.Y, Value.Y);
-		Minimum.Z = AccumulatedMinimum.Z = FMath::Min(AccumulatedMinimum.Z, Value.Z);
-		Maximum.X = AccumulatedMaximum.X = FMath::Max(AccumulatedMaximum.X, Value.X);
-		Maximum.Y = AccumulatedMaximum.Y = FMath::Max(AccumulatedMaximum.Y, Value.Y);
-		Maximum.Z = AccumulatedMaximum.Z = FMath::Max(AccumulatedMaximum.Z, Value.Z);
-	}
+
+	Minimum.X = AccumulatedMinimum.X = FMath::Min(AccumulatedMinimum.X, Value.X);
+	Minimum.Y = AccumulatedMinimum.Y = FMath::Min(AccumulatedMinimum.Y, Value.Y);
+	Minimum.Z = AccumulatedMinimum.Z = FMath::Min(AccumulatedMinimum.Z, Value.Z);
+	Maximum.X = AccumulatedMaximum.X = FMath::Max(AccumulatedMaximum.X, Value.X);
+	Maximum.Y = AccumulatedMaximum.Y = FMath::Max(AccumulatedMaximum.Y, Value.Y);
+	Maximum.Z = AccumulatedMaximum.Z = FMath::Max(AccumulatedMaximum.Z, Value.Z);
 }
 
 #if WITH_DEV_AUTOMATION_TESTS
@@ -220,7 +214,7 @@ FRigUnit_AccumulateVectorRange_Execute()
 
 IMPLEMENT_RIGUNIT_AUTOMATION_TEST(FRigUnit_AccumulateFloatAdd)
 {
-	ExecuteContext.UnitContext.DeltaTime = 0.5f;
+	ExecuteContext.SetDeltaTime(0.5f);
 
 	Unit.Increment = 1.f;
 	InitAndExecute();
@@ -240,7 +234,7 @@ IMPLEMENT_RIGUNIT_AUTOMATION_TEST(FRigUnit_AccumulateFloatAdd)
 IMPLEMENT_RIGUNIT_AUTOMATION_TEST(FRigUnit_AccumulateVectorAdd)
 {
 	using TFloat = decltype(Unit.Result.X);
-	ExecuteContext.UnitContext.DeltaTime = 0.5f;
+	ExecuteContext.SetDeltaTime(0.5f);
 
 	Unit.Increment = FVector(1.0f, 0.f, 0.f);
 	InitAndExecute();
diff --git a/Engine/Plugins/Animation/ControlRig/Source/ControlRig/Private/Units/Simulation/RigUnit_Accumulate.h b/Engine/Plugins/Animation/ControlRig/Source/ControlRig/Private/Units/Simulation/RigUnit_Accumulate.h
--- a/Engine/Plugins/Animation/ControlRig/Source/ControlRig/Private/Units/Simulation/RigUnit_Accumulate.h
+++ b/Engine/Plugins/Animation/ControlRig/Source/ControlRig/Private/Units/Simulation/RigUnit_Accumulate.h
@@ -23,8 +23,11 @@ struct CONTROLRIG_API FRigUnit_AccumulateFloatAdd : public FRigUnit_AccumulateBa
 	{
 		InitialValue = Increment = Result = AccumulatedValue = 0.f;
 		bIntegrateDeltaTime = false;
+		bIsInitialized = false;
 	}
 
+	virtual void Initialize() override { bIsInitialized = false; }
+
 	RIGVM_METHOD()
 	virtual void Execute() override;
 
@@ -42,6 +45,9 @@ struct CONTROLRIG_API FRigUnit_AccumulateFloatAdd : public FRigUnit_AccumulateBa
 
 	UPROPERTY()
 	float AccumulatedValue;
+
+	UPROPERTY()
+	bool bIsInitialized;
 };
 
 /**
@@ -56,8 +62,11 @@ struct CONTROLRIG_API FRigUnit_AccumulateVectorAdd : public FRigUnit_AccumulateB
 	{
 		InitialValue = Increment = Result = AccumulatedValue = FVector::ZeroVector;
 		bIntegrateDeltaTime = false;
+		bIsInitialized = false;
 	}
 
+	virtual void Initialize() override { bIsInitialized = false; }
+
 	RIGVM_METHOD()
 	virtual void Execute() override;
 
@@ -75,6 +84,9 @@ struct CONTROLRIG_API FRigUnit_AccumulateVectorAdd : public FRigUnit_AccumulateB
 
 	UPROPERTY()
 	FVector AccumulatedValue;
+
+	UPROPERTY()
+	bool bIsInitialized;
 };
 
 /**
@@ -89,8 +101,11 @@ struct CONTROLRIG_API FRigUnit_AccumulateFloatMul : public FRigUnit_AccumulateBa
 	{
 		InitialValue = Multiplier = Result = AccumulatedValue = 1.f;
 		bIntegrateDeltaTime = false;
+		bIsInitialized = false;
 	}
 
+	virtual void Initialize() override { bIsInitialized = false; }
+
 	RIGVM_METHOD()
 	virtual void Execute() override;
 
@@ -108,6 +123,9 @@ struct CONTROLRIG_API FRigUnit_AccumulateFloatMul : public FRigUnit_AccumulateBa
 
 	UPROPERTY()
 	float AccumulatedValue;
+
+	UPROPERTY()
+	bool bIsInitialized;
 };
 
 /**
@@ -122,8 +140,11 @@ struct CONTROLRIG_API FRigUnit_AccumulateVectorMul : public FRigUnit_AccumulateB
 	{
 		InitialValue = Multiplier = Result = AccumulatedValue = FVector::OneVector;
 		bIntegrateDeltaTime = false;
+		bIsInitialized = false;
 	}
 
+	virtual void Initialize() override { bIsInitialized = false; }
+
 	RIGVM_METHOD()
 	virtual void Execute() override;
 
@@ -141,6 +162,9 @@ struct CONTROLRIG_API FRigUnit_AccumulateVectorMul : public FRigUnit_AccumulateB
 
 	UPROPERTY()
 	FVector AccumulatedValue;
+
+	UPROPERTY()
+	bool bIsInitialized;
 };
 
 /**
@@ -155,8 +179,11 @@ struct CONTROLRIG_API FRigUnit_AccumulateQuatMul : public FRigUnit_AccumulateBas
 	{
 		InitialValue = Multiplier = Result = AccumulatedValue = FQuat::Identity;
 		bIntegrateDeltaTime = bFlipOrder = false;
+		bIsInitialized = false;
 	}
 
+	virtual void Initialize() override { bIsInitialized = false; }
+
 	RIGVM_METHOD()
 	virtual void Execute() override;
 
@@ -177,6 +204,9 @@ struct CONTROLRIG_API FRigUnit_AccumulateQuatMul : public FRigUnit_AccumulateBas
 
 	UPROPERTY()
 	FQuat AccumulatedValue;
+
+	UPROPERTY()
+	bool bIsInitialized;
 };
 
 /**
@@ -191,8 +221,11 @@ struct CONTROLRIG_API FRigUnit_AccumulateTransformMul : public FRigUnit_Accumula
 	{
 		InitialValue = Multiplier = Result = AccumulatedValue = FTransform::Identity;
 		bIntegrateDeltaTime = bFlipOrder = false;
+		bIsInitialized = false;
 	}
 
+	virtual void Initialize() override { bIsInitialized = false; }
+
 	RIGVM_METHOD()
 	virtual void Execute() override;
 
@@ -213,6 +246,9 @@ struct CONTROLRIG_API FRigUnit_AccumulateTransformMul : public FRigUnit_Accumula
 
 	UPROPERTY()
 	FTransform AccumulatedValue;
+
+	UPROPERTY()
+	bool bIsInitialized;
 };
 
 /**
@@ -227,8 +263,11 @@ struct CONTROLRIG_API FRigUnit_AccumulateFloatLerp : public FRigUnit_AccumulateB
 	{
 		InitialValue = TargetValue = Blend = Result = AccumulatedValue = 0.f;
 		bIntegrateDeltaTime = false;
+		bIsInitialized = false;
 	}
 
+	virtual void Initialize() override { bIsInitialized = false; }
+
 	RIGVM_METHOD()
 	virtual void Execute() override;
 
@@ -249,6 +288,9 @@ struct CONTROLRIG_API FRigUnit_AccumulateFloatLerp : public FRigUnit_AccumulateB
 
 	UPROPERTY()
 	float AccumulatedValue;
+
+	UPROPERTY()
+	bool bIsInitialized;
 };
 
 /**
@@ -264,8 +306,11 @@ struct CONTROLRIG_API FRigUnit_AccumulateVectorLerp : public FRigUnit_Accumulate
 		InitialValue = TargetValue = Result = AccumulatedValue = FVector::ZeroVector;
 		Blend = 0.f;
 		bIntegrateDeltaTime = false;
+		bIsInitialized = false;
 	}
 
+	virtual void Initialize() override { bIsInitialized = false; }
+
 	RIGVM_METHOD()
 	virtual void Execute() override;
 
@@ -286,6 +331,9 @@ struct CONTROLRIG_API FRigUnit_AccumulateVectorLerp : public FRigUnit_Accumulate
 
 	UPROPERTY()
 	FVector AccumulatedValue;
+
+	UPROPERTY()
+	bool bIsInitialized;
 };
 
 /**
@@ -301,8 +349,11 @@ struct CONTROLRIG_API FRigUnit_AccumulateQuatLerp : public FRigUnit_AccumulateBa
 		InitialValue = TargetValue = Result = AccumulatedValue = FQuat::Identity;
 		Blend = 0.f;
 		bIntegrateDeltaTime = false;
+		bIsInitialized = false;
 	}
 
+	virtual void Initialize() override { bIsInitialized = false; }
+
 	RIGVM_METHOD()
 	virtual void Execute() override;
 
@@ -323,6 +374,9 @@ struct CONTROLRIG_API FRigUnit_AccumulateQuatLerp : public FRigUnit_AccumulateBa
 
 	UPROPERTY()
 	FQuat AccumulatedValue;
+
+	UPROPERTY()
+	bool bIsInitialized;
 };
 
 /**
@@ -338,8 +392,11 @@ struct CONTROLRIG_API FRigUnit_AccumulateTransformLerp : public FRigUnit_Accumul
 		InitialValue = TargetValue = Result = AccumulatedValue = FTransform::Identity;
 		Blend = 0.f;
 		bIntegrateDeltaTime = false;
+		bIsInitialized = false;
 	}
 
+	virtual void Initialize() override { bIsInitialized = false; }
+
 	RIGVM_METHOD()
 	virtual void Execute() override;
 
@@ -360,6 +417,9 @@ struct CONTROLRIG_API FRigUnit_AccumulateTransformLerp : public FRigUnit_Accumul
 
 	UPROPERTY()
 	FTransform AccumulatedValue;
+
+	UPROPERTY()
+	bool bIsInitialized;
 };
 
 /**
@@ -373,8 +433,11 @@ struct CONTROLRIG_API FRigUnit_AccumulateFloatRange : public FRigUnit_Accumulate
 	FRigUnit_AccumulateFloatRange()
 	{
 		Value = Minimum = Maximum = AccumulatedMinimum = AccumulatedMaximum = 0.f;
+		bIsInitialized = false;
 	}
 
+	virtual void Initialize() override { bIsInitialized = false; }
+
 	RIGVM_METHOD()
 	virtual void Execute() override;
 
@@ -392,6 +455,9 @@ struct CONTROLRIG_API FRigUnit_AccumulateFloatRange : public FRigUnit_Accumulate
 
 	UPROPERTY()
 	float AccumulatedMaximum;
+
+	UPROPERTY()
+	bool bIsInitialized;
 };
 
 /**
@@ -402,11 +468,14 @@ struct CONTROLRIG_API FRigUnit_AccumulateVectorRange : public FRigUnit_Accumulat
 {
 	GENERATED_BODY()
 
-		FRigUnit_AccumulateVectorRange()
+	FRigUnit_AccumulateVectorRange()
 	{
 		Value = Minimum = Maximum = AccumulatedMinimum = AccumulatedMaximum = FVector::ZeroVector;
+		bIsInitialized = false;
 	}
 
+	virtual void Initialize() override { bIsInitialized = false; }
+
 	RIGVM_METHOD()
 	virtual void Execute() override;
 
@@ -424,4 +493,7 @@ struct CONTROLRIG_API FRigUnit_AccumulateVectorRange : public FRigUnit_Accumulat
 
 	UPROPERTY()
 	FVector AccumulatedMaximum;
+
+	UPROPERTY()
+	bool bIsInitialized;
 };
diff --git a/Engine/Plugins/Animation/ControlRig/Source/ControlRig/Private/Units/Simulation/RigUnit_AlphaInterp.cpp b/Engine/Plugins/Animation/ControlRig/Source/ControlRig/Private/Units/Simulation/RigUnit_AlphaInterp.cpp
--- a/Engine/Plugins/Animation/ControlRig/Source/ControlRig/Private/Units/Simulation/RigUnit_AlphaInterp.cpp
+++ b/Engine/Plugins/Animation/ControlRig/Source/ControlRig/Private/Units/Simulation/RigUnit_AlphaInterp.cpp
@@ -13,23 +13,21 @@ FRigUnit_AlphaInterp_Execute()
 	ScaleBiasClamp.bClampResult = bClampResult;
 	ScaleBiasClamp.bInterpResult = bInterpResult;
 
-	if (ExecuteContext.UnitContext.State == EControlRigState::Init)
+	if (ScaleBiasClamp.bInitialized)
 	{
 		ScaleBiasClamp.Reinitialize();
 	}
-	else
-	{
-		ScaleBiasClamp.InRange = InRange;
-		ScaleBiasClamp.OutRange = OutRange;
-		ScaleBiasClamp.ClampMin = ClampMin;
-		ScaleBiasClamp.ClampMax = ClampMax;
-		ScaleBiasClamp.Scale = Scale;
-		ScaleBiasClamp.Bias = Bias;
-		ScaleBiasClamp.InterpSpeedIncreasing = InterpSpeedIncreasing;
-		ScaleBiasClamp.InterpSpeedDecreasing = InterpSpeedDecreasing;
-
-		Result = ScaleBiasClamp.ApplyTo(Value, ExecuteContext.UnitContext.DeltaTime);
-	}
+
+	ScaleBiasClamp.InRange = InRange;
+	ScaleBiasClamp.OutRange = OutRange;
+	ScaleBiasClamp.ClampMin = ClampMin;
+	ScaleBiasClamp.ClampMax = ClampMax;
+	ScaleBiasClamp.Scale = Scale;
+	ScaleBiasClamp.Bias = Bias;
+	ScaleBiasClamp.InterpSpeedIncreasing = InterpSpeedIncreasing;
+	ScaleBiasClamp.InterpSpeedDecreasing = InterpSpeedDecreasing;
+
+	Result = ScaleBiasClamp.ApplyTo(Value, ExecuteContext.GetDeltaTime());
 }
 
 FString FRigUnit_AlphaInterp::ProcessPinLabelForInjection(const FString& InLabel) const
@@ -63,25 +61,23 @@ FRigUnit_AlphaInterpVector_Execute()
 	ScaleBiasClamp.bClampResult = bClampResult;
 	ScaleBiasClamp.bInterpResult = bInterpResult;
 
-	if (ExecuteContext.UnitContext.State == EControlRigState::Init)
+	if (ScaleBiasClamp.bInitialized)
 	{
 		ScaleBiasClamp.Reinitialize();
 	}
-	else
-	{
-		ScaleBiasClamp.InRange = InRange;
-		ScaleBiasClamp.OutRange = OutRange;
-		ScaleBiasClamp.ClampMin = ClampMin;
-		ScaleBiasClamp.ClampMax = ClampMax;
-		ScaleBiasClamp.Scale = Scale;
-		ScaleBiasClamp.Bias = Bias;
-		ScaleBiasClamp.InterpSpeedIncreasing = InterpSpeedIncreasing;
-		ScaleBiasClamp.InterpSpeedDecreasing = InterpSpeedDecreasing;
-
-		Result.X = ScaleBiasClamp.ApplyTo(Value.X, ExecuteContext.UnitContext.DeltaTime);
-		Result.Y = ScaleBiasClamp.ApplyTo(Value.Y, ExecuteContext.UnitContext.DeltaTime);
-		Result.Z = ScaleBiasClamp.ApplyTo(Value.Z, ExecuteContext.UnitContext.DeltaTime);
-	}
+
+	ScaleBiasClamp.InRange = InRange;
+	ScaleBiasClamp.OutRange = OutRange;
+	ScaleBiasClamp.ClampMin = ClampMin;
+	ScaleBiasClamp.ClampMax = ClampMax;
+	ScaleBiasClamp.Scale = Scale;
+	ScaleBiasClamp.Bias = Bias;
+	ScaleBiasClamp.InterpSpeedIncreasing = InterpSpeedIncreasing;
+	ScaleBiasClamp.InterpSpeedDecreasing = InterpSpeedDecreasing;
+
+	Result.X = ScaleBiasClamp.ApplyTo(Value.X, ExecuteContext.GetDeltaTime());
+	Result.Y = ScaleBiasClamp.ApplyTo(Value.Y, ExecuteContext.GetDeltaTime());
+	Result.Z = ScaleBiasClamp.ApplyTo(Value.Z, ExecuteContext.GetDeltaTime());
 }
 
 FString FRigUnit_AlphaInterpVector::ProcessPinLabelForInjection(const FString& InLabel) const
@@ -115,24 +111,22 @@ FRigUnit_AlphaInterpQuat_Execute()
 	ScaleBiasClamp.bClampResult = bClampResult;
 	ScaleBiasClamp.bInterpResult = bInterpResult;
 
-	if (ExecuteContext.UnitContext.State == EControlRigState::Init)
+	if (ScaleBiasClamp.bInitialized)
 	{
 		ScaleBiasClamp.Reinitialize();
 	}
-	else
-	{
-		ScaleBiasClamp.InRange = InRange;
-		ScaleBiasClamp.OutRange = OutRange;
-		ScaleBiasClamp.ClampMin = ClampMin;
-		ScaleBiasClamp.ClampMax = ClampMax;
-		ScaleBiasClamp.Scale = Scale;
-		ScaleBiasClamp.Bias = Bias;
-		ScaleBiasClamp.InterpSpeedIncreasing = InterpSpeedIncreasing;
-		ScaleBiasClamp.InterpSpeedDecreasing = InterpSpeedDecreasing;
-
-		const float T = ScaleBiasClamp.ApplyTo(1.f, ExecuteContext.UnitContext.DeltaTime);
-		Result = FQuat::Slerp(FQuat::Identity, Value, T);
-	}
+
+	ScaleBiasClamp.InRange = InRange;
+	ScaleBiasClamp.OutRange = OutRange;
+	ScaleBiasClamp.ClampMin = ClampMin;
+	ScaleBiasClamp.ClampMax = ClampMax;
+	ScaleBiasClamp.Scale = Scale;
+	ScaleBiasClamp.Bias = Bias;
+	ScaleBiasClamp.InterpSpeedIncreasing = InterpSpeedIncreasing;
+	ScaleBiasClamp.InterpSpeedDecreasing = InterpSpeedDecreasing;
+
+	const float T = ScaleBiasClamp.ApplyTo(1.f, ExecuteContext.GetDeltaTime());
+	Result = FQuat::Slerp(FQuat::Identity, Value, T);
 }
 
 FString FRigUnit_AlphaInterpQuat::ProcessPinLabelForInjection(const FString& InLabel) const
diff --git a/Engine/Plugins/Animation/ControlRig/Source/ControlRig/Private/Units/Simulation/RigUnit_DeltaFromPrevious.cpp b/Engine/Plugins/Animation/ControlRig/Source/ControlRig/Private/Units/Simulation/RigUnit_DeltaFromPrevious.cpp
--- a/Engine/Plugins/Animation/ControlRig/Source/ControlRig/Private/Units/Simulation/RigUnit_DeltaFromPrevious.cpp
+++ b/Engine/Plugins/Animation/ControlRig/Source/ControlRig/Private/Units/Simulation/RigUnit_DeltaFromPrevious.cpp
@@ -8,17 +8,17 @@
 FRigUnit_DeltaFromPreviousFloat_Execute()
 {
     DECLARE_SCOPE_HIERARCHICAL_COUNTER_RIGUNIT()
-	if (ExecuteContext.UnitContext.State == EControlRigState::Init)
+	if (!bIsInitialized)
 	{
 		Delta = 0.f;
 		PreviousValue = Cache = Value;
-		return;
+		bIsInitialized = true;
 	}
 
 	PreviousValue = Cache;
 	Delta = Cache - Value;
 
-	if (FMath::Abs(ExecuteContext.UnitContext.DeltaTime) > SMALL_NUMBER)
+	if (FMath::Abs(ExecuteContext.GetDeltaTime()) > SMALL_NUMBER)
 	{
 		Cache = Value;
 	}
@@ -27,17 +27,17 @@ FRigUnit_DeltaFromPreviousFloat_Execute()
 FRigUnit_DeltaFromPreviousVector_Execute()
 {
     DECLARE_SCOPE_HIERARCHICAL_COUNTER_RIGUNIT()
-	if (ExecuteContext.UnitContext.State == EControlRigState::Init)
+	if (!bIsInitialized)
 	{
 		Delta = FVector::ZeroVector;
 		PreviousValue = Cache = Value;
-		return;
+		bIsInitialized = true;
 	}
 
 	PreviousValue = Cache;
 	Delta = Cache - Value;
 
-	if (FMath::Abs(ExecuteContext.UnitContext.DeltaTime) > SMALL_NUMBER)
+	if (FMath::Abs(ExecuteContext.GetDeltaTime()) > SMALL_NUMBER)
 	{
 		Cache = Value;
 	}
@@ -46,17 +46,17 @@ FRigUnit_DeltaFromPreviousVector_Execute()
 FRigUnit_DeltaFromPreviousQuat_Execute()
 {
     DECLARE_SCOPE_HIERARCHICAL_COUNTER_RIGUNIT()
-	if (ExecuteContext.UnitContext.State == EControlRigState::Init)
+	if (!bIsInitialized)
 	{
 		Delta = FQuat::Identity;
 		PreviousValue = Cache = Value;
-		return;
+		bIsInitialized = true;
 	}
 
 	PreviousValue = Cache;
 	Delta = Cache.Inverse() * Value;
 
-	if (FMath::Abs(ExecuteContext.UnitContext.DeltaTime) > SMALL_NUMBER)
+	if (FMath::Abs(ExecuteContext.GetDeltaTime()) > SMALL_NUMBER)
 	{
 		Cache = Value;
 	}
@@ -65,17 +65,17 @@ FRigUnit_DeltaFromPreviousQuat_Execute()
 FRigUnit_DeltaFromPreviousTransform_Execute()
 {
     DECLARE_SCOPE_HIERARCHICAL_COUNTER_RIGUNIT()
-	if (ExecuteContext.UnitContext.State == EControlRigState::Init)
+	if (!bIsInitialized)
 	{
 		Delta = FTransform::Identity;
 		PreviousValue = Cache = Value;
-		return;
+		bIsInitialized = true;
 	}
 
 	PreviousValue = Cache;
 	Delta = Value.GetRelativeTransform(Cache);
 
-	if (FMath::Abs(ExecuteContext.UnitContext.DeltaTime) > SMALL_NUMBER)
+	if (FMath::Abs(ExecuteContext.GetDeltaTime()) > SMALL_NUMBER)
 	{
 		Cache = Value;
 	}
@@ -86,12 +86,12 @@ FRigUnit_DeltaFromPreviousTransform_Execute()
 
 IMPLEMENT_RIGUNIT_AUTOMATION_TEST(FRigUnit_DeltaFromPreviousFloat)
 {
-	ExecuteContext.UnitContext.DeltaTime = 0.1f;
+	ExecuteContext.SetDeltaTime(0.1f);
 	Unit.Value = 1.f;
-	InitAndExecute();
+	Execute();
 	AddErrorIfFalse(FMath::IsNearlyEqual(Unit.Delta, 0.f), TEXT("unexpected average result"));
 	AddErrorIfFalse(FMath::IsNearlyEqual(Unit.PreviousValue, 1.f), TEXT("unexpected average result"));
-	InitAndExecute();
+	Execute();
 	AddErrorIfFalse(FMath::IsNearlyEqual(Unit.Delta, 0.f), TEXT("unexpected average result"));
 	AddErrorIfFalse(FMath::IsNearlyEqual(Unit.PreviousValue, 1.f), TEXT("unexpected average result"));
 	Execute();
diff --git a/Engine/Plugins/Animation/ControlRig/Source/ControlRig/Private/Units/Simulation/RigUnit_DeltaFromPrevious.h b/Engine/Plugins/Animation/ControlRig/Source/ControlRig/Private/Units/Simulation/RigUnit_DeltaFromPrevious.h
--- a/Engine/Plugins/Animation/ControlRig/Source/ControlRig/Private/Units/Simulation/RigUnit_DeltaFromPrevious.h
+++ b/Engine/Plugins/Animation/ControlRig/Source/ControlRig/Private/Units/Simulation/RigUnit_DeltaFromPrevious.h
@@ -16,6 +16,7 @@ struct CONTROLRIG_API FRigUnit_DeltaFromPreviousFloat : public FRigUnit_SimBase
 	FRigUnit_DeltaFromPreviousFloat()
 	{
 		Value = Delta = PreviousValue = Cache = 0.f;
+		bIsInitialized = false;
 	}
 
 	RIGVM_METHOD()
@@ -32,6 +33,9 @@ struct CONTROLRIG_API FRigUnit_DeltaFromPreviousFloat : public FRigUnit_SimBase
 
 	UPROPERTY()
 	float Cache;
+
+	UPROPERTY()
+	bool bIsInitialized;
 };
 
 /**
@@ -45,6 +49,7 @@ struct CONTROLRIG_API FRigUnit_DeltaFromPreviousVector : public FRigUnit_SimBase
 	FRigUnit_DeltaFromPreviousVector()
 	{
 		Value = Delta = PreviousValue = Cache = FVector::ZeroVector;
+		bIsInitialized = false;
 	}
 
 	RIGVM_METHOD()
@@ -61,6 +66,9 @@ struct CONTROLRIG_API FRigUnit_DeltaFromPreviousVector : public FRigUnit_SimBase
 
 	UPROPERTY()
 	FVector Cache;
+
+	UPROPERTY()
+	bool bIsInitialized;
 };
 
 /**
@@ -74,6 +82,7 @@ struct CONTROLRIG_API FRigUnit_DeltaFromPreviousQuat : public FRigUnit_SimBase
 	FRigUnit_DeltaFromPreviousQuat()
 	{
 		Value = Delta = PreviousValue = Cache = FQuat::Identity;
+		bIsInitialized = false;
 	}
 
 	RIGVM_METHOD()
@@ -90,6 +99,9 @@ struct CONTROLRIG_API FRigUnit_DeltaFromPreviousQuat : public FRigUnit_SimBase
 
 	UPROPERTY()
 	FQuat Cache;
+
+	UPROPERTY()
+	bool bIsInitialized;
 };
 
 /**
@@ -103,6 +115,7 @@ struct CONTROLRIG_API FRigUnit_DeltaFromPreviousTransform : public FRigUnit_SimB
 	FRigUnit_DeltaFromPreviousTransform()
 	{
 		Value = Delta = PreviousValue = Cache = FTransform::Identity;
+		bIsInitialized = false;
 	}
 
 	RIGVM_METHOD()
@@ -119,4 +132,7 @@ struct CONTROLRIG_API FRigUnit_DeltaFromPreviousTransform : public FRigUnit_SimB
 
 	UPROPERTY()
 	FTransform Cache;
+
+	UPROPERTY()
+	bool bIsInitialized;
 };
diff --git a/Engine/Plugins/Animation/ControlRig/Source/ControlRig/Private/Units/Simulation/RigUnit_Kalman.cpp b/Engine/Plugins/Animation/ControlRig/Source/ControlRig/Private/Units/Simulation/RigUnit_Kalman.cpp
--- a/Engine/Plugins/Animation/ControlRig/Source/ControlRig/Private/Units/Simulation/RigUnit_Kalman.cpp
+++ b/Engine/Plugins/Animation/ControlRig/Source/ControlRig/Private/Units/Simulation/RigUnit_Kalman.cpp
@@ -10,7 +10,7 @@ FRigUnit_KalmanFloat_Execute()
     DECLARE_SCOPE_HIERARCHICAL_COUNTER_RIGUNIT()
 	if(BufferSize <= 0)
 	{
-		if (ExecuteContext.UnitContext.State == EControlRigState::Init)
+		if (Buffer.IsEmpty()) // only report this error once
 		{
 			UE_CONTROLRIG_RIGUNIT_REPORT_WARNING(TEXT("BufferSize is too small."));
 		}
@@ -19,40 +19,38 @@ FRigUnit_KalmanFloat_Execute()
 		return;
 	}
 
-	Buffer.SetNum(BufferSize);
-
-	if (ExecuteContext.UnitContext.State == EControlRigState::Init)
+	if(Buffer.Num() != BufferSize)
 	{
+		Buffer.SetNum(BufferSize);
+
 		for (float& Element : Buffer)
 		{
 			Element = FLT_MAX;
 			LastInsertIndex = -1;
 		}
 	}
-	else
+
+	if(LastInsertIndex == Buffer.Num() - 1)
 	{
-		if(LastInsertIndex == Buffer.Num() - 1)
-		{
-			LastInsertIndex = -1;
-		}
-		Buffer[++LastInsertIndex] = Value;
+		LastInsertIndex = -1;
+	}
+	Buffer[++LastInsertIndex] = Value;
 
-		Result = 0.f;
-		int32 NumberValidEntries = 0;
-		for(const float F : Buffer)
+	Result = 0.f;
+	int32 NumberValidEntries = 0;
+	for(const float F : Buffer)
+	{
+		if (F == FLT_MAX)
 		{
-			if (F == FLT_MAX)
-			{
-				break;
-			}
-			Result += F;
-			NumberValidEntries++;
+			break;
 		}
+		Result += F;
+		NumberValidEntries++;
+	}
 
-		if (NumberValidEntries > 0)
-		{
-			Result = Result / float(NumberValidEntries);
-		}
+	if (NumberValidEntries > 0)
+	{
+		Result = Result / float(NumberValidEntries);
 	}
 }
 
@@ -61,7 +59,7 @@ FRigUnit_KalmanVector_Execute()
     DECLARE_SCOPE_HIERARCHICAL_COUNTER_RIGUNIT()
 	if(BufferSize <= 0)
 	{
-		if (ExecuteContext.UnitContext.State == EControlRigState::Init)
+		if (Buffer.IsEmpty()) // only report this error once
 		{
 			UE_CONTROLRIG_RIGUNIT_REPORT_WARNING(TEXT("BufferSize is too small."));
 		}
@@ -70,40 +68,38 @@ FRigUnit_KalmanVector_Execute()
 		return;
 	}
 
-	Buffer.SetNum(BufferSize);
-
-	if (ExecuteContext.UnitContext.State == EControlRigState::Init)
+	if(Buffer.Num() != BufferSize)
 	{
+		Buffer.SetNum(BufferSize);
+
 		for (FVector& Element : Buffer)
 		{
 			Element.X = FLT_MAX;
 			LastInsertIndex = -1;
 		}
 	}
-	else
+
+	if(LastInsertIndex == Buffer.Num() - 1)
 	{
-		if(LastInsertIndex == Buffer.Num() - 1)
-		{
-			LastInsertIndex = -1;
-		}
-		Buffer[++LastInsertIndex] = Value;
+		LastInsertIndex = -1;
+	}
+	Buffer[++LastInsertIndex] = Value;
 
-		Result = FVector::ZeroVector;
-		int32 NumberValidEntries = 0;
-		for (const FVector& F : Buffer)
+	Result = FVector::ZeroVector;
+	int32 NumberValidEntries = 0;
+	for (const FVector& F : Buffer)
+	{
+		if (F.X == FLT_MAX)
 		{
-			if (F.X == FLT_MAX)
-			{
-				break;
-			}
-			Result += F;
-			NumberValidEntries++;
+			break;
 		}
+		Result += F;
+		NumberValidEntries++;
+	}
 
-		if (NumberValidEntries > 0)
-		{
-			Result = Result / float(NumberValidEntries);
-		}
+	if (NumberValidEntries > 0)
+	{
+		Result = Result / float(NumberValidEntries);
 	}
 }
 
@@ -112,7 +108,7 @@ FRigUnit_KalmanTransform_Execute()
     DECLARE_SCOPE_HIERARCHICAL_COUNTER_RIGUNIT()
 	if(BufferSize <= 0)
 	{
-		if (ExecuteContext.UnitContext.State == EControlRigState::Init)
+		if (Buffer.IsEmpty()) // only report this error once
 		{
 			UE_CONTROLRIG_RIGUNIT_REPORT_WARNING(TEXT("BufferSize is too small."));
 		}
@@ -121,56 +117,54 @@ FRigUnit_KalmanTransform_Execute()
 		return;
 	}
 
-	Buffer.SetNum(BufferSize);
-
-	if (ExecuteContext.UnitContext.State == EControlRigState::Init)
+	if(Buffer.Num() != BufferSize)
 	{
+		Buffer.SetNum(BufferSize);
+
 		for (FTransform& Element : Buffer)
 		{
 			Element.SetTranslation(FVector(FLT_MAX));
 			LastInsertIndex = -1;
 		}
 	}
-	else
-	{
-		if(LastInsertIndex == Buffer.Num() - 1)
-		{
-			LastInsertIndex = -1;
-		}
-		Buffer[++LastInsertIndex] = Value;
 
-		FVector Location = FVector::ZeroVector;
-		FVector AxisX = FVector::ZeroVector;
-		FVector AxisY = FVector::ZeroVector;
-		FVector Scale = FVector::ZeroVector;
-		
-		int32 NumberValidEntries = 0;
+	if(LastInsertIndex == Buffer.Num() - 1)
+	{
+		LastInsertIndex = -1;
+	}
+	Buffer[++LastInsertIndex] = Value;
 
-		for(const FTransform& F : Buffer)
-		{
-			if (F.GetLocation().X == FLT_MAX)
-			{
-				break;
-			}
-			Location += F.GetLocation();
-			AxisX += F.TransformVectorNoScale(FVector(1.f, 0.f, 0.f));
-			AxisY += F.TransformVectorNoScale(FVector(0.f, 1.f, 0.f));
-			Scale += F.GetScale3D();
-			NumberValidEntries++;
-		}
+	FVector Location = FVector::ZeroVector;
+	FVector AxisX = FVector::ZeroVector;
+	FVector AxisY = FVector::ZeroVector;
+	FVector Scale = FVector::ZeroVector;
+	
+	int32 NumberValidEntries = 0;
 
-		if (NumberValidEntries > 0)
+	for(const FTransform& F : Buffer)
+	{
+		if (F.GetLocation().X == FLT_MAX)
 		{
-			Location = Location / float(NumberValidEntries);
-			AxisX = (AxisX / float(NumberValidEntries)).GetSafeNormal();
-			AxisY = (AxisY / float(NumberValidEntries)).GetSafeNormal();
-			Scale = Scale / float(NumberValidEntries);
+			break;
 		}
+		Location += F.GetLocation();
+		AxisX += F.TransformVectorNoScale(FVector(1.f, 0.f, 0.f));
+		AxisY += F.TransformVectorNoScale(FVector(0.f, 1.f, 0.f));
+		Scale += F.GetScale3D();
+		NumberValidEntries++;
+	}
 
-		Result.SetLocation(Location);
-		Result.SetRotation(FRotationMatrix::MakeFromXY(AxisX, AxisY).ToQuat());
-		Result.SetScale3D(Scale);
+	if (NumberValidEntries > 0)
+	{
+		Location = Location / float(NumberValidEntries);
+		AxisX = (AxisX / float(NumberValidEntries)).GetSafeNormal();
+		AxisY = (AxisY / float(NumberValidEntries)).GetSafeNormal();
+		Scale = Scale / float(NumberValidEntries);
 	}
+
+	Result.SetLocation(Location);
+	Result.SetRotation(FRotationMatrix::MakeFromXY(AxisX, AxisY).ToQuat());
+	Result.SetScale3D(Scale);
 }
 
 #if WITH_DEV_AUTOMATION_TESTS
diff --git a/Engine/Plugins/Animation/ControlRig/Source/ControlRig/Private/Units/Simulation/RigUnit_Kalman.h b/Engine/Plugins/Animation/ControlRig/Source/ControlRig/Private/Units/Simulation/RigUnit_Kalman.h
--- a/Engine/Plugins/Animation/ControlRig/Source/ControlRig/Private/Units/Simulation/RigUnit_Kalman.h
+++ b/Engine/Plugins/Animation/ControlRig/Source/ControlRig/Private/Units/Simulation/RigUnit_Kalman.h
@@ -21,6 +21,8 @@ struct CONTROLRIG_API FRigUnit_KalmanFloat : public FRigUnit_SimBase
 		LastInsertIndex = 0;
 	}
 
+	virtual void Initialize() override { Buffer.Reset(); }
+
 	RIGVM_METHOD()
 	virtual void Execute() override;
 
@@ -55,6 +57,8 @@ struct CONTROLRIG_API FRigUnit_KalmanVector : public FRigUnit_SimBase
 		BufferSize = 16;
 		LastInsertIndex = 0;
 	}
+	
+	virtual void Initialize() override { Buffer.Reset(); }
 
 	RIGVM_METHOD()
 	virtual void Execute() override;
@@ -91,6 +95,8 @@ struct CONTROLRIG_API FRigUnit_KalmanTransform : public FRigUnit_SimBase
 		LastInsertIndex = 0;
 	}
 
+	virtual void Initialize() override { Buffer.Reset(); }
+
 	RIGVM_METHOD()
 	virtual void Execute() override;
 
diff --git a/Engine/Plugins/Animation/ControlRig/Source/ControlRig/Private/Units/Simulation/RigUnit_PointSimulation.cpp b/Engine/Plugins/Animation/ControlRig/Source/ControlRig/Private/Units/Simulation/RigUnit_PointSimulation.cpp
--- a/Engine/Plugins/Animation/ControlRig/Source/ControlRig/Private/Units/Simulation/RigUnit_PointSimulation.cpp
+++ b/Engine/Plugins/Animation/ControlRig/Source/ControlRig/Private/Units/Simulation/RigUnit_PointSimulation.cpp
@@ -10,13 +10,12 @@ FRigUnit_PointSimulation_Execute()
     DECLARE_SCOPE_HIERARCHICAL_COUNTER_RIGUNIT()
 	URigHierarchy* Hierarchy = ExecuteContext.Hierarchy;
 
-	float DeltaTime = ExecuteContext.UnitContext.DeltaTime;
+	float DeltaTime = ExecuteContext.GetDeltaTime();
 
 	FCRSimPointContainer& Simulation = WorkData.Simulation;
 	TArray<FCachedRigElement>& BoneIndices = WorkData.BoneIndices;
 
-	if (ExecuteContext.UnitContext.State == EControlRigState::Init ||
-		Simulation.Points.Num() != Points.Num() ||
+	if (Simulation.Points.Num() != Points.Num() ||
 		Simulation.Springs.Num() != Links.Num() ||
 		Simulation.Forces.Num() != Forces.Num() ||
 		Simulation.CollisionVolumes.Num() != CollisionVolumes.Num() ||
diff --git a/Engine/Plugins/Animation/ControlRig/Source/ControlRig/Private/Units/Simulation/RigUnit_SpringInterp.cpp b/Engine/Plugins/Animation/ControlRig/Source/ControlRig/Private/Units/Simulation/RigUnit_SpringInterp.cpp
--- a/Engine/Plugins/Animation/ControlRig/Source/ControlRig/Private/Units/Simulation/RigUnit_SpringInterp.cpp
+++ b/Engine/Plugins/Animation/ControlRig/Source/ControlRig/Private/Units/Simulation/RigUnit_SpringInterp.cpp
@@ -16,24 +16,17 @@ FRigUnit_SpringInterp_Execute()
 {
 	DECLARE_SCOPE_HIERARCHICAL_COUNTER_RIGUNIT()
  
-	if (ExecuteContext.UnitContext.State == EControlRigState::Init)
-	{
-		SpringState.Reset();
-	}
-	else
+	// Clamp to avoid large time deltas.
+	float RemainingTime = FMath::Min(ExecuteContext.GetDeltaTime(), RigUnitSpringInterpConstants::MaxTimeStep);
+
+	Result = Current;
+	while (RemainingTime >= RigUnitSpringInterpConstants::FixedTimeStep)
 	{
-		// Clamp to avoid large time deltas.
-		float RemainingTime = FMath::Min(ExecuteContext.UnitContext.DeltaTime, RigUnitSpringInterpConstants::MaxTimeStep);
- 
-		Result = Current;
-		while (RemainingTime >= RigUnitSpringInterpConstants::FixedTimeStep)
-		{
-			Result = UKismetMathLibrary::FloatSpringInterp(Result, Target, SpringState, Stiffness, CriticalDamping, RigUnitSpringInterpConstants::FixedTimeStep, Mass);
-			RemainingTime -= RigUnitSpringInterpConstants::FixedTimeStep;
-		}
- 
-		Result = UKismetMathLibrary::FloatSpringInterp(Result, Target, SpringState, Stiffness, CriticalDamping, RemainingTime, Mass);
+		Result = UKismetMathLibrary::FloatSpringInterp(Result, Target, SpringState, Stiffness, CriticalDamping, RigUnitSpringInterpConstants::FixedTimeStep, Mass);
+		RemainingTime -= RigUnitSpringInterpConstants::FixedTimeStep;
 	}
+
+	Result = UKismetMathLibrary::FloatSpringInterp(Result, Target, SpringState, Stiffness, CriticalDamping, RemainingTime, Mass);
 }
 
 FRigVMStructUpgradeInfo FRigUnit_SpringInterp::GetUpgradeInfo() const
@@ -46,24 +39,17 @@ FRigUnit_SpringInterpVector_Execute()
 {
 	DECLARE_SCOPE_HIERARCHICAL_COUNTER_RIGUNIT()
  
-	if (ExecuteContext.UnitContext.State == EControlRigState::Init)
-	{
-		SpringState.Reset();
-	}
-	else
+	// Clamp to avoid large time deltas.
+	float RemainingTime = FMath::Min(ExecuteContext.GetDeltaTime(), RigUnitSpringInterpConstants::MaxTimeStep);
+
+	Result = Current;
+	while (RemainingTime >= RigUnitSpringInterpConstants::FixedTimeStep)
 	{
-		// Clamp to avoid large time deltas.
-		float RemainingTime = FMath::Min(ExecuteContext.UnitContext.DeltaTime, RigUnitSpringInterpConstants::MaxTimeStep);
- 
-		Result = Current;
-		while (RemainingTime >= RigUnitSpringInterpConstants::FixedTimeStep)
-		{
-			Result = UKismetMathLibrary::VectorSpringInterp(Result, Target, SpringState, Stiffness, CriticalDamping, RigUnitSpringInterpConstants::FixedTimeStep, Mass);
-			RemainingTime -= RigUnitSpringInterpConstants::FixedTimeStep;
-		}
- 
-		Result = UKismetMathLibrary::VectorSpringInterp(Result, Target, SpringState, Stiffness, CriticalDamping, RemainingTime, Mass);
+		Result = UKismetMathLibrary::VectorSpringInterp(Result, Target, SpringState, Stiffness, CriticalDamping, RigUnitSpringInterpConstants::FixedTimeStep, Mass);
+		RemainingTime -= RigUnitSpringInterpConstants::FixedTimeStep;
 	}
+
+	Result = UKismetMathLibrary::VectorSpringInterp(Result, Target, SpringState, Stiffness, CriticalDamping, RemainingTime, Mass);
 }
 
 FRigVMStructUpgradeInfo FRigUnit_SpringInterpVector::GetUpgradeInfo() const
@@ -76,89 +62,64 @@ FRigUnit_SpringInterpV2_Execute()
 {
 	DECLARE_SCOPE_HIERARCHICAL_COUNTER_RIGUNIT()
 
-	if (ExecuteContext.UnitContext.State == EControlRigState::Init)
+	// Treat the input as a frequency in Hz
+	float AngularFrequency = Strength * 2.0f * PI;
+	float Stiffness = AngularFrequency * AngularFrequency;
+	float AdjustedTarget = Target;
+	if (!FMath::IsNearlyZero(Stiffness))
 	{
-		SpringState.Reset();
-		Result = Target;
+		AdjustedTarget += Force / (Stiffness * RigUnitSpringInterpConstants::Mass);
 	}
 	else
 	{
-		// Treat the input as a frequency in Hz
-		float AngularFrequency = Strength * 2.0f * PI;
-		float Stiffness = AngularFrequency * AngularFrequency;
-		float AdjustedTarget = Target;
-		if (!FMath::IsNearlyZero(Stiffness))
-		{
-			AdjustedTarget += Force / (Stiffness * RigUnitSpringInterpConstants::Mass);
-		}
-		else
-		{
-			SpringState.Velocity += Force * (ExecuteContext.UnitContext.DeltaTime / RigUnitSpringInterpConstants::Mass);
-		}
-		SimulatedResult = UKismetMathLibrary::FloatSpringInterp(
-			bUseCurrentInput ? Current : SimulatedResult, AdjustedTarget, SpringState, Stiffness, CriticalDamping,
-			ExecuteContext.UnitContext.DeltaTime, RigUnitSpringInterpConstants::Mass, TargetVelocityAmount, 
-			false, 0.0f, 0.0f, !bUseCurrentInput || bInitializeFromTarget);
-
-		Result = SimulatedResult;
-		Velocity = SpringState.Velocity;
+		SpringState.Velocity += Force * (ExecuteContext.GetDeltaTime() / RigUnitSpringInterpConstants::Mass);
 	}
+	SimulatedResult = UKismetMathLibrary::FloatSpringInterp(
+		bUseCurrentInput ? Current : SimulatedResult, AdjustedTarget, SpringState, Stiffness, CriticalDamping,
+		ExecuteContext.GetDeltaTime(), RigUnitSpringInterpConstants::Mass, TargetVelocityAmount, 
+		false, 0.0f, 0.0f, !bUseCurrentInput || bInitializeFromTarget);
+
+	Result = SimulatedResult;
+	Velocity = SpringState.Velocity;
 }
 
 FRigUnit_SpringInterpVectorV2_Execute()
 {
     DECLARE_SCOPE_HIERARCHICAL_COUNTER_RIGUNIT()
 
-	if (ExecuteContext.UnitContext.State == EControlRigState::Init)
+	// Treat the input as a frequency in Hz
+	float AngularFrequency = Strength * 2.0f * PI;
+	float Stiffness = AngularFrequency * AngularFrequency;
+	FVector AdjustedTarget = Target;
+	if (!FMath::IsNearlyZero(Stiffness))
 	{
-		SpringState.Reset();
-		Result = Target;
+		AdjustedTarget += Force / (Stiffness * RigUnitSpringInterpConstants::Mass);
 	}
 	else
 	{
-		// Treat the input as a frequency in Hz
-		float AngularFrequency = Strength * 2.0f * PI;
-		float Stiffness = AngularFrequency * AngularFrequency;
-		FVector AdjustedTarget = Target;
-		if (!FMath::IsNearlyZero(Stiffness))
-		{
-			AdjustedTarget += Force / (Stiffness * RigUnitSpringInterpConstants::Mass);
-		}
-		else
-		{
-			SpringState.Velocity += Force * (ExecuteContext.UnitContext.DeltaTime / RigUnitSpringInterpConstants::Mass);
-		}
-		SimulatedResult = UKismetMathLibrary::VectorSpringInterp(
-			bUseCurrentInput ? Current : SimulatedResult, AdjustedTarget, SpringState, Stiffness, CriticalDamping,
-			ExecuteContext.UnitContext.DeltaTime, RigUnitSpringInterpConstants::Mass, TargetVelocityAmount,
-			false, FVector(), FVector(), !bUseCurrentInput || bInitializeFromTarget);
-		Result = SimulatedResult;
+		SpringState.Velocity += Force * (ExecuteContext.GetDeltaTime() / RigUnitSpringInterpConstants::Mass);
 	}
-	Velocity = SpringState.Velocity;
+	SimulatedResult = UKismetMathLibrary::VectorSpringInterp(
+		bUseCurrentInput ? Current : SimulatedResult, AdjustedTarget, SpringState, Stiffness, CriticalDamping,
+		ExecuteContext.GetDeltaTime(), RigUnitSpringInterpConstants::Mass, TargetVelocityAmount,
+		false, FVector(), FVector(), !bUseCurrentInput || bInitializeFromTarget);
+	Result = SimulatedResult;
 }
 
 FRigUnit_SpringInterpQuaternionV2_Execute()
 {
     DECLARE_SCOPE_HIERARCHICAL_COUNTER_RIGUNIT()
 
-	if (ExecuteContext.UnitContext.State == EControlRigState::Init)
-	{
-		SpringState.Reset();
-		Result = Target;
-	}
-	else
-	{
-		// Treat the input as a frequency in Hz
-		float AngularFrequency = Strength * 2.0f * PI;
-		float Stiffness = AngularFrequency * AngularFrequency;
-		SpringState.AngularVelocity += Torque * (ExecuteContext.UnitContext.DeltaTime / RigUnitSpringInterpConstants::Mass);
-		SimulatedResult = UKismetMathLibrary::QuaternionSpringInterp(
-			bUseCurrentInput ? Current : SimulatedResult, Target, SpringState, Stiffness, CriticalDamping,
-			ExecuteContext.UnitContext.DeltaTime, RigUnitSpringInterpConstants::Mass, TargetVelocityAmount, 
-			!bUseCurrentInput || bInitializeFromTarget);
-		Result = SimulatedResult;
-		AngularVelocity = SpringState.AngularVelocity;
-	}
+	// Treat the input as a frequency in Hz
+	float AngularFrequency = Strength * 2.0f * PI;
+	float Stiffness = AngularFrequency * AngularFrequency;
+	SpringState.AngularVelocity += Torque * (ExecuteContext.GetDeltaTime() / RigUnitSpringInterpConstants::Mass);
+	SimulatedResult = UKismetMathLibrary::QuaternionSpringInterp(
+		bUseCurrentInput ? Current : SimulatedResult, Target, SpringState, Stiffness, CriticalDamping,
+		ExecuteContext.GetDeltaTime(), RigUnitSpringInterpConstants::Mass, TargetVelocityAmount, 
+		!bUseCurrentInput || bInitializeFromTarget);
+	Result = SimulatedResult;
+	AngularVelocity = SpringState.AngularVelocity;
 }
 
 
diff --git a/Engine/Plugins/Animation/ControlRig/Source/ControlRig/Private/Units/Simulation/RigUnit_TimeOffset.cpp b/Engine/Plugins/Animation/ControlRig/Source/ControlRig/Private/Units/Simulation/RigUnit_TimeOffset.cpp
--- a/Engine/Plugins/Animation/ControlRig/Source/ControlRig/Private/Units/Simulation/RigUnit_TimeOffset.cpp
+++ b/Engine/Plugins/Animation/ControlRig/Source/ControlRig/Private/Units/Simulation/RigUnit_TimeOffset.cpp
@@ -10,9 +10,11 @@ FRigUnit_TimeOffsetFloat_Execute()
     DECLARE_SCOPE_HIERARCHICAL_COUNTER_RIGUNIT()
 	Result = Value;
 
+	const bool bIsInitialized = !Buffer.IsEmpty();
+
 	if(BufferSize <= 0)
 	{
-		if (ExecuteContext.UnitContext.State == EControlRigState::Init)
+		if (!bIsInitialized)
 		{
 			UE_CONTROLRIG_RIGUNIT_REPORT_WARNING(TEXT("BufferSize is too small."));
 		}
@@ -22,7 +24,7 @@ FRigUnit_TimeOffsetFloat_Execute()
 
 	if (TimeRange <= 0)
 	{
-		if (ExecuteContext.UnitContext.State == EControlRigState::Init)
+		if (!bIsInitialized)
 		{
 			UE_CONTROLRIG_RIGUNIT_REPORT_WARNING(TEXT("TimeRange is too small."));
 		}
@@ -30,18 +32,16 @@ FRigUnit_TimeOffsetFloat_Execute()
 		return;
 	}
 
+	if(!bIsInitialized)
+	{
+		UpperBound = 0;
+	}
+
 	int32 MaxSize = FMath::Clamp<int32>(BufferSize, 2, 512);
 	Buffer.SetNum(MaxSize);
 	DeltaTimes.SetNum(MaxSize);
 	MaxSize = FMath::Min<int32>(MaxSize, Buffer.Num());
 
-	if (ExecuteContext.UnitContext.State == EControlRigState::Init)
-	{
-		UpperBound = 0;
-		Result = Value;
-		return;
-	}
-
 	if (UpperBound == 0)
 	{
 		Result = Value;
@@ -80,35 +80,35 @@ FRigUnit_TimeOffsetFloat_Execute()
 		Result = FMath::Lerp<float>(A, B, T);
 	}
 
-	if (ExecuteContext.UnitContext.DeltaTime > SMALL_NUMBER)
+	if (ExecuteContext.GetDeltaTime() > SMALL_NUMBER)
 	{
 		if (UpperBound == 0)
 		{
 			LastInsertIndex = UpperBound;
 			Buffer[UpperBound] = Value;
-			DeltaTimes[UpperBound++] = ExecuteContext.UnitContext.DeltaTime;
+			DeltaTimes[UpperBound++] = ExecuteContext.GetDeltaTime();
 		}
 		else
 		{
 			float SecondsPerEntry = TimeRange / float(MaxSize - 1);
-			if (DeltaTimes[LastInsertIndex] > SecondsPerEntry - ExecuteContext.UnitContext.DeltaTime * 0.5f)
+			if (DeltaTimes[LastInsertIndex] > SecondsPerEntry - ExecuteContext.GetDeltaTime() * 0.5f)
 			{
 				if (UpperBound < MaxSize)
 				{
 					LastInsertIndex = UpperBound;
 					Buffer[UpperBound] = Value;
-					DeltaTimes[UpperBound++] = ExecuteContext.UnitContext.DeltaTime;
+					DeltaTimes[UpperBound++] = ExecuteContext.GetDeltaTime();
 				}
 				else
 				{
 					LastInsertIndex = (LastInsertIndex + 1) % UpperBound;
 					Buffer[LastInsertIndex] = Value;
-					DeltaTimes[LastInsertIndex] = ExecuteContext.UnitContext.DeltaTime;
+					DeltaTimes[LastInsertIndex] = ExecuteContext.GetDeltaTime();
 				}
 			}
 			else
 			{
-				DeltaTimes[LastInsertIndex] += ExecuteContext.UnitContext.DeltaTime;
+				DeltaTimes[LastInsertIndex] += ExecuteContext.GetDeltaTime();
 			}
 		}
 	}
@@ -119,9 +119,11 @@ FRigUnit_TimeOffsetVector_Execute()
     DECLARE_SCOPE_HIERARCHICAL_COUNTER_RIGUNIT()
 	Result = Value;
 
+	const bool bIsInitialized = !Buffer.IsEmpty();
+
 	if(BufferSize <= 0)
 	{
-		if (ExecuteContext.UnitContext.State == EControlRigState::Init)
+		if (!bIsInitialized)
 		{
 			UE_CONTROLRIG_RIGUNIT_REPORT_WARNING(TEXT("BufferSize is too small."));
 		}
@@ -131,7 +133,7 @@ FRigUnit_TimeOffsetVector_Execute()
 
 	if (TimeRange <= 0)
 	{
-		if (ExecuteContext.UnitContext.State == EControlRigState::Init)
+		if (!bIsInitialized)
 		{
 			UE_CONTROLRIG_RIGUNIT_REPORT_WARNING(TEXT("TimeRange is too small."));
 		}
@@ -139,18 +141,16 @@ FRigUnit_TimeOffsetVector_Execute()
 		return;
 	}
 
+	if(!bIsInitialized)
+	{
+		UpperBound = 0;
+	}
+
 	int32 MaxSize = FMath::Clamp<int32>(BufferSize, 2, 512);
 	Buffer.SetNum(MaxSize);
 	DeltaTimes.SetNum(MaxSize);
 	MaxSize = FMath::Min<int32>(MaxSize, Buffer.Num());
 
-	if (ExecuteContext.UnitContext.State == EControlRigState::Init)
-	{
-		UpperBound = 0;
-		Result = Value;
-		return;
-	}
-
 	if (UpperBound == 0)
 	{
 		Result = Value;
@@ -189,35 +189,35 @@ FRigUnit_TimeOffsetVector_Execute()
 		Result = FMath::Lerp<FVector>(A, B, T);
 	}
 
-	if (ExecuteContext.UnitContext.DeltaTime > SMALL_NUMBER)
+	if (ExecuteContext.GetDeltaTime() > SMALL_NUMBER)
 	{
 		if (UpperBound == 0)
 		{
 			LastInsertIndex = UpperBound;
 			Buffer[UpperBound] = Value;
-			DeltaTimes[UpperBound++] = ExecuteContext.UnitContext.DeltaTime;
+			DeltaTimes[UpperBound++] = ExecuteContext.GetDeltaTime();
 		}
 		else
 		{
 			float SecondsPerEntry = TimeRange / float(MaxSize - 1);
-			if (DeltaTimes[LastInsertIndex] > SecondsPerEntry - ExecuteContext.UnitContext.DeltaTime * 0.5f)
+			if (DeltaTimes[LastInsertIndex] > SecondsPerEntry - ExecuteContext.GetDeltaTime() * 0.5f)
 			{
 				if (UpperBound < MaxSize)
 				{
 					LastInsertIndex = UpperBound;
 					Buffer[UpperBound] = Value;
-					DeltaTimes[UpperBound++] = ExecuteContext.UnitContext.DeltaTime;
+					DeltaTimes[UpperBound++] = ExecuteContext.GetDeltaTime();
 				}
 				else
 				{
 					LastInsertIndex = (LastInsertIndex + 1) % UpperBound;
 					Buffer[LastInsertIndex] = Value;
-					DeltaTimes[LastInsertIndex] = ExecuteContext.UnitContext.DeltaTime;
+					DeltaTimes[LastInsertIndex] = ExecuteContext.GetDeltaTime();
 				}
 			}
 			else
 			{
-				DeltaTimes[LastInsertIndex] += ExecuteContext.UnitContext.DeltaTime;
+				DeltaTimes[LastInsertIndex] += ExecuteContext.GetDeltaTime();
 			}
 		}
 	}
@@ -228,9 +228,11 @@ FRigUnit_TimeOffsetTransform_Execute()
     DECLARE_SCOPE_HIERARCHICAL_COUNTER_RIGUNIT()
 	Result = Value;
 
+	const bool bIsInitialized = !Buffer.IsEmpty();
+
 	if(BufferSize <= 0)
 	{
-		if (ExecuteContext.UnitContext.State == EControlRigState::Init)
+		if (!bIsInitialized)
 		{
 			UE_CONTROLRIG_RIGUNIT_REPORT_WARNING(TEXT("BufferSize is too small."));
 		}
@@ -240,7 +242,7 @@ FRigUnit_TimeOffsetTransform_Execute()
 
 	if (TimeRange <= 0)
 	{
-		if (ExecuteContext.UnitContext.State == EControlRigState::Init)
+		if (!bIsInitialized)
 		{
 			UE_CONTROLRIG_RIGUNIT_REPORT_WARNING(TEXT("TimeRange is too small."));
 		}
@@ -248,18 +250,16 @@ FRigUnit_TimeOffsetTransform_Execute()
 		return;
 	}
 
+	if(!bIsInitialized)
+	{
+		UpperBound = 0;
+	}
+
 	int32 MaxSize = FMath::Clamp<int32>(BufferSize, 2, 512);
 	Buffer.SetNum(MaxSize);
 	DeltaTimes.SetNum(MaxSize);
 	MaxSize = FMath::Min<int32>(MaxSize, Buffer.Num());
 
-	if (ExecuteContext.UnitContext.State == EControlRigState::Init)
-	{
-		UpperBound = 0;
-		Result = Value;
-		return;
-	}
-
 	if (UpperBound == 0)
 	{
 		Result = Value;
@@ -300,35 +300,35 @@ FRigUnit_TimeOffsetTransform_Execute()
 		Result.SetScale3D(FMath::Lerp<FVector>(A.GetScale3D(), B.GetScale3D(), T));
 	}
 
-	if (ExecuteContext.UnitContext.DeltaTime > SMALL_NUMBER)
+	if (ExecuteContext.GetDeltaTime() > SMALL_NUMBER)
 	{
 		if (UpperBound == 0)
 		{
 			LastInsertIndex = UpperBound;
 			Buffer[UpperBound] = Value;
-			DeltaTimes[UpperBound++] = ExecuteContext.UnitContext.DeltaTime;
+			DeltaTimes[UpperBound++] = ExecuteContext.GetDeltaTime();
 		}
 		else
 		{
 			float SecondsPerEntry = TimeRange / float(MaxSize - 1);
-			if (DeltaTimes[LastInsertIndex] > SecondsPerEntry - ExecuteContext.UnitContext.DeltaTime * 0.5f)
+			if (DeltaTimes[LastInsertIndex] > SecondsPerEntry - ExecuteContext.GetDeltaTime() * 0.5f)
 			{
 				if (UpperBound < MaxSize)
 				{
 					LastInsertIndex = UpperBound;
 					Buffer[UpperBound] = Value;
-					DeltaTimes[UpperBound++] = ExecuteContext.UnitContext.DeltaTime;
+					DeltaTimes[UpperBound++] = ExecuteContext.GetDeltaTime();
 				}
 				else
 				{
 					LastInsertIndex = (LastInsertIndex + 1) % UpperBound;
 					Buffer[LastInsertIndex] = Value;
-					DeltaTimes[LastInsertIndex] = ExecuteContext.UnitContext.DeltaTime;
+					DeltaTimes[LastInsertIndex] = ExecuteContext.GetDeltaTime();
 				}
 			}
 			else
 			{
-				DeltaTimes[LastInsertIndex] += ExecuteContext.UnitContext.DeltaTime;
+				DeltaTimes[LastInsertIndex] += ExecuteContext.GetDeltaTime();
 			}
 		}
 	}
@@ -339,7 +339,7 @@ FRigUnit_TimeOffsetTransform_Execute()
 
 IMPLEMENT_RIGUNIT_AUTOMATION_TEST(FRigUnit_TimeOffsetFloat)
 {
-	ExecuteContext.UnitContext.DeltaTime = 1.f;
+	ExecuteContext.SetDeltaTime(1.f);
 	Unit.SecondsAgo = 0.5;
 	Unit.TimeRange = 5.f;
 	Unit.BufferSize = 16;
@@ -365,7 +365,7 @@ IMPLEMENT_RIGUNIT_AUTOMATION_TEST(FRigUnit_TimeOffsetFloat)
 	Execute();
 	AddErrorIfFalse(FMath::IsNearlyEqual(Unit.Result, 11.f), TEXT("unexpected previous result"));
 
-	ExecuteContext.UnitContext.DeltaTime = 0.f;
+	ExecuteContext.SetDeltaTime(0.f);
 	Unit.SecondsAgo = 1.5;
 	Execute();
 	AddErrorIfFalse(FMath::IsNearlyEqual(Unit.Result, 9.f), TEXT("unexpected previous result"));
@@ -379,7 +379,7 @@ IMPLEMENT_RIGUNIT_AUTOMATION_TEST(FRigUnit_TimeOffsetFloat)
 	Execute();
 	AddErrorIfFalse(FMath::IsNearlyEqual(Unit.Result, 3.f), TEXT("unexpected previous result"));
 
-	ExecuteContext.UnitContext.DeltaTime = 1.f;
+	ExecuteContext.SetDeltaTime(1.f);
 	Unit.SecondsAgo = 0.5;
 	Unit.Value = 14.f;
 	Execute();
@@ -418,7 +418,7 @@ IMPLEMENT_RIGUNIT_AUTOMATION_TEST(FRigUnit_TimeOffsetFloat)
 	Execute();
 	AddErrorIfFalse(FMath::IsNearlyEqual(Unit.Result, 35.f), TEXT("unexpected previous result"));
 
-	ExecuteContext.UnitContext.DeltaTime = 0.f;
+	ExecuteContext.SetDeltaTime(0.f);
 	Unit.SecondsAgo = 1.5;
 	Execute();
 	AddErrorIfFalse(FMath::IsNearlyEqual(Unit.Result, 33.f), TEXT("unexpected previous result"));
@@ -433,7 +433,7 @@ IMPLEMENT_RIGUNIT_AUTOMATION_TEST(FRigUnit_TimeOffsetFloat)
 	AddErrorIfFalse(FMath::IsNearlyEqual(Unit.Result, 27.f), TEXT("unexpected previous result"));
 
 
-	ExecuteContext.UnitContext.DeltaTime = 1.f;
+	ExecuteContext.SetDeltaTime(1.f);
 	Unit.TimeRange = 8.f;
 	Unit.BufferSize = 3;
 	Unit.SecondsAgo = 0.5f;
@@ -448,7 +448,7 @@ IMPLEMENT_RIGUNIT_AUTOMATION_TEST(FRigUnit_TimeOffsetFloat)
 	Execute();
 	AddErrorIfFalse(FMath::IsNearlyEqual(Unit.Result, 5.f), TEXT("unexpected previous result"));
 
-	ExecuteContext.UnitContext.DeltaTime = 0.f;
+	ExecuteContext.SetDeltaTime(0.f);
 	Unit.SecondsAgo = 1.5f;
 	Execute();
 	AddErrorIfFalse(FMath::IsNearlyEqual(Unit.Result, 4.f), TEXT("unexpected previous result"));
diff --git a/Engine/Plugins/Animation/ControlRig/Source/ControlRig/Private/Units/Simulation/RigUnit_TimeOffset.h b/Engine/Plugins/Animation/ControlRig/Source/ControlRig/Private/Units/Simulation/RigUnit_TimeOffset.h
--- a/Engine/Plugins/Animation/ControlRig/Source/ControlRig/Private/Units/Simulation/RigUnit_TimeOffset.h
+++ b/Engine/Plugins/Animation/ControlRig/Source/ControlRig/Private/Units/Simulation/RigUnit_TimeOffset.h
@@ -23,6 +23,8 @@ struct CONTROLRIG_API FRigUnit_TimeOffsetFloat : public FRigUnit_SimBase
 		UpperBound = 0;
 	}
 
+	virtual void Initialize() override { Buffer.Reset(); }
+
 	RIGVM_METHOD()
 	virtual void Execute() override;
 	
@@ -76,6 +78,8 @@ struct CONTROLRIG_API FRigUnit_TimeOffsetVector : public FRigUnit_SimBase
 		UpperBound = 0;
 	}
 
+	virtual void Initialize() override { Buffer.Reset(); }
+
 	RIGVM_METHOD()
 	virtual void Execute() override;
 	
@@ -129,6 +133,8 @@ struct CONTROLRIG_API FRigUnit_TimeOffsetTransform : public FRigUnit_SimBase
 		UpperBound = 0;
 	}
 
+	virtual void Initialize() override { Buffer.Reset(); }
+
 	RIGVM_METHOD()
 	virtual void Execute() override;
 	
diff --git a/Engine/Plugins/Animation/ControlRig/Source/ControlRig/Private/Units/Simulation/RigUnit_Timeline.cpp b/Engine/Plugins/Animation/ControlRig/Source/ControlRig/Private/Units/Simulation/RigUnit_Timeline.cpp
--- a/Engine/Plugins/Animation/ControlRig/Source/ControlRig/Private/Units/Simulation/RigUnit_Timeline.cpp
+++ b/Engine/Plugins/Animation/ControlRig/Source/ControlRig/Private/Units/Simulation/RigUnit_Timeline.cpp
@@ -8,13 +8,13 @@
 FRigUnit_Timeline_Execute()
 {
     DECLARE_SCOPE_HIERARCHICAL_COUNTER_RIGUNIT()
-	if (ExecuteContext.UnitContext.State == EControlRigState::Init)
+	if (!bIsInitialized)
 	{
 		Time = AccumulatedValue = 0.f;
-		return;
+		bIsInitialized = true;
 	}
 
-	Time = AccumulatedValue = AccumulatedValue + ExecuteContext.UnitContext.DeltaTime * Speed;
+	Time = AccumulatedValue = AccumulatedValue + ExecuteContext.GetDeltaTime() * Speed;
 }
 
 #if WITH_DEV_AUTOMATION_TESTS
@@ -22,7 +22,7 @@ FRigUnit_Timeline_Execute()
 
 IMPLEMENT_RIGUNIT_AUTOMATION_TEST(FRigUnit_Timeline)
 {
-	ExecuteContext.UnitContext.DeltaTime = 1.f;
+	ExecuteContext.SetDeltaTime(1.f);
 	InitAndExecute();
 	AddErrorIfFalse(FMath::IsNearlyEqual(Unit.Time, 1.f), TEXT("unexpected time"));
 	InitAndExecute();
@@ -42,16 +42,16 @@ IMPLEMENT_RIGUNIT_AUTOMATION_TEST(FRigUnit_Timeline)
 FRigUnit_TimeLoop_Execute()
 {
 	DECLARE_SCOPE_HIERARCHICAL_COUNTER_RIGUNIT()
-	if (ExecuteContext.UnitContext.State == EControlRigState::Init)
+	if (!bIsInitialized)
 	{
 		Absolute = Relative = FlipFlop = AccumulatedAbsolute = AccumulatedRelative = 0.f;
 		NumIterations = 0;
 		Even = false;
-		return;
+		bIsInitialized = true;
 	}
 
 	const float DurationClamped = FMath::Max(Duration, 0.0001f);
-	const float Increment = ExecuteContext.UnitContext.DeltaTime * Speed;
+	const float Increment = ExecuteContext.GetDeltaTime() * Speed;
 	Absolute = AccumulatedAbsolute = AccumulatedAbsolute + Increment;
 
 	AccumulatedRelative = AccumulatedRelative + Increment;
diff --git a/Engine/Plugins/Animation/ControlRig/Source/ControlRig/Private/Units/Simulation/RigUnit_Timeline.h b/Engine/Plugins/Animation/ControlRig/Source/ControlRig/Private/Units/Simulation/RigUnit_Timeline.h
--- a/Engine/Plugins/Animation/ControlRig/Source/ControlRig/Private/Units/Simulation/RigUnit_Timeline.h
+++ b/Engine/Plugins/Animation/ControlRig/Source/ControlRig/Private/Units/Simulation/RigUnit_Timeline.h
@@ -17,8 +17,11 @@ struct CONTROLRIG_API FRigUnit_Timeline : public FRigUnit_SimBase
 	{
 		Speed = 1.f;
 		Time = AccumulatedValue = 0.f;
+		bIsInitialized = false;
 	}
 
+	virtual void Initialize() override { bIsInitialized = false; }
+
 	RIGVM_METHOD()
 	virtual void Execute() override;
 
@@ -30,6 +33,9 @@ struct CONTROLRIG_API FRigUnit_Timeline : public FRigUnit_SimBase
 
 	UPROPERTY()
 	float AccumulatedValue;
+
+	UPROPERTY()
+	bool bIsInitialized;
 };
 
 /**
@@ -48,8 +54,11 @@ struct CONTROLRIG_API FRigUnit_TimeLoop : public FRigUnit_SimBase
 		Absolute = Relative = FlipFlop = AccumulatedAbsolute = AccumulatedRelative = 0.f;
 		NumIterations = 0;
 		Even = false;
+		bIsInitialized = false;
 	}
 
+	virtual void Initialize() override { bIsInitialized = false; }
+
 	RIGVM_METHOD()
 	virtual void Execute() override;
 
@@ -91,4 +100,7 @@ struct CONTROLRIG_API FRigUnit_TimeLoop : public FRigUnit_SimBase
 
 	UPROPERTY()
 	int32 NumIterations;
+
+	UPROPERTY()
+	bool bIsInitialized;
 };
\ No newline at end of file
diff --git a/Engine/Plugins/Animation/ControlRig/Source/ControlRig/Private/Units/Simulation/RigUnit_Verlet.cpp b/Engine/Plugins/Animation/ControlRig/Source/ControlRig/Private/Units/Simulation/RigUnit_Verlet.cpp
--- a/Engine/Plugins/Animation/ControlRig/Source/ControlRig/Private/Units/Simulation/RigUnit_Verlet.cpp
+++ b/Engine/Plugins/Animation/ControlRig/Source/ControlRig/Private/Units/Simulation/RigUnit_Verlet.cpp
@@ -8,11 +8,6 @@
 FRigUnit_VerletIntegrateVector_Execute()
 {
     DECLARE_SCOPE_HIERARCHICAL_COUNTER_RIGUNIT()
-	if (ExecuteContext.UnitContext.State == EControlRigState::Init)
-	{
-		bInitialized = false;
-		return;
-	}
 
 	if (!bInitialized)
 	{
@@ -20,17 +15,16 @@ FRigUnit_VerletIntegrateVector_Execute()
 		Position = Point.Position = Target;
 		Velocity = Acceleration = Point.LinearVelocity = FVector::ZeroVector;
 		bInitialized = true;
-		return;
 	}
 
 	Point.LinearDamping = Damp;
-	if (ExecuteContext.UnitContext.DeltaTime > SMALL_NUMBER)
+	if (ExecuteContext.GetDeltaTime() > SMALL_NUMBER)
 	{
-		float U = FMath::Clamp<float>(Blend * ExecuteContext.UnitContext.DeltaTime, 0.f, 1.f);
-		const FVector CombinedForce = (Target - Point.Position) * FMath::Max(Strength, 0.0001f) + Force * ExecuteContext.UnitContext.DeltaTime * 60.f;
+		float U = FMath::Clamp<float>(Blend * ExecuteContext.GetDeltaTime(), 0.f, 1.f);
+		const FVector CombinedForce = (Target - Point.Position) * FMath::Max(Strength, 0.0001f) + Force * ExecuteContext.GetDeltaTime() * 60.f;
 		const FVector PreviousVelocity = Point.LinearVelocity;
-		Point = Point.IntegrateVerlet(CombinedForce, Blend, ExecuteContext.UnitContext.DeltaTime);
-		Acceleration = (Point.LinearVelocity - PreviousVelocity) / ExecuteContext.UnitContext.DeltaTime;
+		Point = Point.IntegrateVerlet(CombinedForce, Blend, ExecuteContext.GetDeltaTime());
+		Acceleration = (Point.LinearVelocity - PreviousVelocity) / ExecuteContext.GetDeltaTime();
 		Position = Point.Position;
 		Velocity = Point.LinearVelocity;
 	}
diff --git a/Engine/Plugins/Animation/ControlRig/Source/ControlRig/Public/AnimNode_ControlRig.h b/Engine/Plugins/Animation/ControlRig/Source/ControlRig/Public/AnimNode_ControlRig.h
--- a/Engine/Plugins/Animation/ControlRig/Source/ControlRig/Public/AnimNode_ControlRig.h
+++ b/Engine/Plugins/Animation/ControlRig/Source/ControlRig/Public/AnimNode_ControlRig.h
@@ -38,7 +38,7 @@ struct CONTROLRIG_API FAnimNode_ControlRig : public FAnimNode_ControlRigBase
 	virtual void PropagateInputProperties(const UObject* InSourceInstance) override;
 	
 private:
-	void HandleOnInitialized_AnyThread(UControlRig*, const EControlRigState, const FName&);
+	void HandleOnInitialized_AnyThread(UControlRig*, const FName&);
 #if WITH_EDITOR
 	virtual void HandleObjectsReinstanced_Impl(UObject* InSourceObject, UObject* InTargetObject, const TMap<UObject*, UObject*>& OldToNewInstanceMap) override;
 #endif
diff --git a/Engine/Plugins/Animation/ControlRig/Source/ControlRig/Public/ControlRig.h b/Engine/Plugins/Animation/ControlRig/Source/ControlRig/Public/ControlRig.h
--- a/Engine/Plugins/Animation/ControlRig/Source/ControlRig/Public/ControlRig.h
+++ b/Engine/Plugins/Animation/ControlRig/Source/ControlRig/Public/ControlRig.h
@@ -109,10 +109,6 @@ private:
 	/** true if the rig itself should increase the AbsoluteTime */
 	bool bAccumulateTime;
 
-	/** Latest state being processed */
-	EControlRigState LatestExecutedState;
-
-	
 public:
 	virtual void Serialize(FArchive& Ar) override;
 	virtual void PostLoad() override;
@@ -182,6 +178,9 @@ public:
 	/** Initialize things for the ControlRig */
 	virtual void Initialize(bool bInitRigUnits = true);
 
+	/** Initialize the VM */
+	virtual bool InitializeVM(const FName& InEventName);
+
 	/** Evaluate at Any Thread */
 	virtual void Evaluate_AnyThread();
 
@@ -303,7 +302,7 @@ public:
 
 	/** Execute */
 	UFUNCTION(BlueprintCallable, Category = "Control Rig")
-	bool Execute(const EControlRigState State, const FName& InEventName);
+	bool Execute(const FName& InEventName);
 
 	/** ExecuteUnits */
 	virtual bool ExecuteUnits(FRigUnitContext& InOutContext, const FName& InEventName);
@@ -448,7 +447,7 @@ public:
 
 	bool IsCurveControl(const FRigControlElement* InControlElement) const;
 
-	DECLARE_EVENT_ThreeParams(UControlRig, FControlRigExecuteEvent, class UControlRig*, const EControlRigState, const FName&);
+	DECLARE_EVENT_TwoParams(UControlRig, FControlRigExecuteEvent, class UControlRig*, const FName&);
 	FControlRigExecuteEvent& OnInitialized_AnyThread() { return InitializedEvent; }
 #if WITH_EDITOR
 	FControlRigExecuteEvent& OnPreConstructionForUI_AnyThread() { return PreConstructionForUIEvent; }
@@ -512,15 +511,6 @@ protected:
 	UPROPERTY()
 	TObjectPtr<URigVM> VM;
 
-	// this is only used by the CDO
-	// and stores an initialized VM per hash.
-	UPROPERTY(transient)
-	TMap<uint32, TObjectPtr<URigVM>> InitializedVMSnapshots;
-	uint32 CachedMemoryHash = 0;
-
-	// computes the hash used to store / find VM snapshots
-	uint32 GetHashForInitializeVMSnapShot();
-	
 	UPROPERTY()
 	TObjectPtr<URigHierarchy> DynamicHierarchy;
 
@@ -736,8 +726,8 @@ private:
 
 	void CopyPoseFromOtherRig(UControlRig* Subject);
 	void HandleInteractionRigControlModified(UControlRig* Subject, FRigControlElement* Control, const FRigControlModifiedContext& Context);
-	void HandleInteractionRigInitialized(UControlRig* Subject, EControlRigState State, const FName& EventName);
-	void HandleInteractionRigExecuted(UControlRig* Subject, EControlRigState State, const FName& EventName);
+	void HandleInteractionRigInitialized(UControlRig* Subject, const FName& EventName);
+	void HandleInteractionRigExecuted(UControlRig* Subject, const FName& EventName);
 	void HandleInteractionRigControlSelected(UControlRig* Subject, FRigControlElement* InControl, bool bSelected, bool bInverted);
 
 
diff --git a/Engine/Plugins/Animation/ControlRig/Source/ControlRig/Public/ControlRigComponent.h b/Engine/Plugins/Animation/ControlRig/Source/ControlRig/Public/ControlRigComponent.h
--- a/Engine/Plugins/Animation/ControlRig/Source/ControlRig/Public/ControlRigComponent.h
+++ b/Engine/Plugins/Animation/ControlRig/Source/ControlRig/Public/ControlRigComponent.h
@@ -645,12 +645,12 @@ private:
 	void TransferOutputs();
 	static FName GetComponentNameWithinActor(UActorComponent* InComponent);
 
-	void HandleControlRigInitializedEvent(UControlRig* InControlRig, const EControlRigState InState, const FName& InEventName);
-	void HandleControlRigPreConstructionEvent(UControlRig* InControlRig, const EControlRigState InState, const FName& InEventName);
-	void HandleControlRigPostConstructionEvent(UControlRig* InControlRig, const EControlRigState InState, const FName& InEventName);
-	void HandleControlRigPreForwardsSolveEvent(UControlRig* InControlRig, const EControlRigState InState, const FName& InEventName);
-	void HandleControlRigPostForwardsSolveEvent(UControlRig* InControlRig, const EControlRigState InState, const FName& InEventName);
-	void HandleControlRigExecutedEvent(UControlRig* InControlRig, const EControlRigState InState, const FName& InEventName);
+	void HandleControlRigInitializedEvent(UControlRig* InControlRig, const FName& InEventName);
+	void HandleControlRigPreConstructionEvent(UControlRig* InControlRig, const FName& InEventName);
+	void HandleControlRigPostConstructionEvent(UControlRig* InControlRig, const FName& InEventName);
+	void HandleControlRigPreForwardsSolveEvent(UControlRig* InControlRig, const FName& InEventName);
+	void HandleControlRigPostForwardsSolveEvent(UControlRig* InControlRig, const FName& InEventName);
+	void HandleControlRigExecutedEvent(UControlRig* InControlRig, const FName& InEventName);
 
 	void ConvertTransformToRigSpace(FTransform& InOutTransform, EControlRigComponentSpace FromSpace);
 	void ConvertTransformFromRigSpace(FTransform& InOutTransform, EControlRigComponentSpace ToSpace);
diff --git a/Engine/Plugins/Animation/ControlRig/Source/ControlRig/Public/ControlRigValidationPass.h b/Engine/Plugins/Animation/ControlRig/Source/ControlRig/Public/ControlRigValidationPass.h
--- a/Engine/Plugins/Animation/ControlRig/Source/ControlRig/Public/ControlRigValidationPass.h
+++ b/Engine/Plugins/Animation/ControlRig/Source/ControlRig/Public/ControlRigValidationPass.h
@@ -65,8 +65,8 @@ private:
 	UPROPERTY()
 	TArray<TObjectPtr<UControlRigValidationPass>> Passes;
 
-	void OnControlRigInitialized(UControlRig* Subject, EControlRigState State, const FName& EventName);
-	void OnControlRigExecuted(UControlRig* Subject, EControlRigState State, const FName& EventName);
+	void OnControlRigInitialized(UControlRig* Subject, const FName& EventName);
+	void OnControlRigExecuted(UControlRig* Subject, const FName& EventName);
 
 	FControlRigValidationContext ValidationContext;
 	TWeakObjectPtr<UControlRig> WeakControlRig;
diff --git a/Engine/Plugins/Animation/ControlRig/Source/ControlRig/Public/Units/Control/RigUnit_Control.h b/Engine/Plugins/Animation/ControlRig/Source/ControlRig/Public/Units/Control/RigUnit_Control.h
--- a/Engine/Plugins/Animation/ControlRig/Source/ControlRig/Public/Units/Control/RigUnit_Control.h
+++ b/Engine/Plugins/Animation/ControlRig/Source/ControlRig/Public/Units/Control/RigUnit_Control.h
@@ -17,6 +17,7 @@ struct CONTROLRIG_API FRigUnit_Control : public FRigUnit
 		: Transform(FEulerTransform::Identity)
 		, Base(FTransform::Identity)
 		, Result(FTransform::Identity)
+		, bIsInitialized(false)
 	{
 	}
 
@@ -65,4 +66,7 @@ struct CONTROLRIG_API FRigUnit_Control : public FRigUnit
 
 	RIGVM_METHOD()
 	virtual FRigVMStructUpgradeInfo GetUpgradeInfo() const override;
+
+	UPROPERTY()
+ 	bool bIsInitialized;
 };
\ No newline at end of file
diff --git a/Engine/Plugins/Animation/ControlRig/Source/ControlRig/Public/Units/Deprecated/RigUnit_TwoBoneIKFK.h b/Engine/Plugins/Animation/ControlRig/Source/ControlRig/Public/Units/Deprecated/RigUnit_TwoBoneIKFK.h
--- a/Engine/Plugins/Animation/ControlRig/Source/ControlRig/Public/Units/Deprecated/RigUnit_TwoBoneIKFK.h
+++ b/Engine/Plugins/Animation/ControlRig/Source/ControlRig/Public/Units/Deprecated/RigUnit_TwoBoneIKFK.h
@@ -26,6 +26,7 @@ struct CONTROLRIG_API FRigUnit_TwoBoneIKFK : public FRigUnitMutable
 		, EndJointIndex(INDEX_NONE)
 		, UpperLimbLength(0.f)
 		, LowerLimbLength(0.f)
+		, bIsInitialized(false)
 	{}
 
 	UPROPERTY(EditAnywhere, Category = "TwoBoneIKFK", meta = (Input))
@@ -91,6 +92,9 @@ private:
 	UPROPERTY(transient)
 	float LowerLimbLength;
 
+	UPROPERTY()
+	bool bIsInitialized;
+
 	RIGVM_METHOD()
 	virtual FRigVMStructUpgradeInfo GetUpgradeInfo() const override;
 };
diff --git a/Engine/Plugins/Animation/ControlRig/Source/ControlRig/Public/Units/RigUnit.h b/Engine/Plugins/Animation/ControlRig/Source/ControlRig/Public/Units/RigUnit.h
--- a/Engine/Plugins/Animation/ControlRig/Source/ControlRig/Public/Units/RigUnit.h
+++ b/Engine/Plugins/Animation/ControlRig/Source/ControlRig/Public/Units/RigUnit.h
@@ -25,6 +25,9 @@ struct CONTROLRIG_API FRigUnit : public FRigVMStruct
 	/** Returns the label of this unit */
 	virtual FString GetUnitLabel() const { return FString(); }
 
+	/** initialize logic for this rig unit */
+	virtual void Initialize() {}
+
 	/** Execute logic for this rig unit */
 	virtual void Execute() {}
 
diff --git a/Engine/Plugins/Animation/ControlRig/Source/ControlRig/Public/Units/RigUnitContext.h b/Engine/Plugins/Animation/ControlRig/Source/ControlRig/Public/Units/RigUnitContext.h
--- a/Engine/Plugins/Animation/ControlRig/Source/ControlRig/Public/Units/RigUnitContext.h
+++ b/Engine/Plugins/Animation/ControlRig/Source/ControlRig/Public/Units/RigUnitContext.h
@@ -13,17 +13,6 @@
 #include "Components/SceneComponent.h"
 #include "RigUnitContext.generated.h"
 
-/** Current state of rig
-*	What  state Control Rig currently is
-*/
-UENUM()
-enum class EControlRigState : uint8
-{
-	Init,
-	Update,
-	Invalid,
-};
-
 /**
  * The type of interaction happening on a rig
  */
@@ -61,9 +50,6 @@ struct FRigUnitContext
 		, DrawInterface(nullptr)
 		, DrawContainer(nullptr)
 		, DataSourceRegistry(nullptr)
-		, DeltaTime(0.f)
-		, AbsoluteTime(0.f)
-		, State(EControlRigState::Invalid)
 		, InteractionType((uint8)EControlRigInteractionType::None)
 		, ElementsBeingInteracted()
 		, ToWorldSpaceTransform(FTransform::Identity)
@@ -88,18 +74,6 @@ struct FRigUnitContext
 	/** The registry to access data source */
 	const UAnimationDataSourceRegistry* DataSourceRegistry;
 
-	/** The current delta time */
-	float DeltaTime;
-
-	/** The current delta time */
-	float AbsoluteTime;
-
-	/** The current frames per second */
-	float FramesPerSecond;
-
-	/** Current execution context */
-	EControlRigState State;
-
 	/** The current hierarchy settings */
 	FRigHierarchySettings HierarchySettings;
 
diff --git a/Engine/Plugins/Animation/ControlRig/Source/ControlRig/Public/Units/RigUnitTest.h b/Engine/Plugins/Animation/ControlRig/Source/ControlRig/Public/Units/RigUnitTest.h
--- a/Engine/Plugins/Animation/ControlRig/Source/ControlRig/Public/Units/RigUnitTest.h
+++ b/Engine/Plugins/Animation/ControlRig/Source/ControlRig/Public/Units/RigUnitTest.h
@@ -76,17 +76,15 @@ public:
 		} \
 		virtual bool RunControlRigUnitTest(const FString& Parameters); \
 		virtual FString GetBeautifiedTestName() const override { return TEXT(CONTROLRIG_RIGUNIT_STRINGIFY(ControlRig.Units.TUnitStruct)); } \
-		void Init() \
+		void InitAndExecute() \
 		{ \
-			ExecuteContext.UnitContext.State = EControlRigState::Init; \
+			Unit.Initialize(); \
 			Unit.Execute(ExecuteContext); \
 		} \
 		void Execute() \
 		{ \
-			ExecuteContext.UnitContext.State = EControlRigState::Update; \
 			Unit.Execute(ExecuteContext); \
 		} \
-		void InitAndExecute() { Init(); Execute(); } \
 	}; \
 	namespace\
 	{\
diff --git a/Engine/Plugins/Animation/ControlRig/Source/ControlRigEditor/Private/ControlRigEditorModule.cpp b/Engine/Plugins/Animation/ControlRig/Source/ControlRigEditor/Private/ControlRigEditorModule.cpp
--- a/Engine/Plugins/Animation/ControlRig/Source/ControlRigEditor/Private/ControlRigEditorModule.cpp
+++ b/Engine/Plugins/Animation/ControlRig/Source/ControlRigEditor/Private/ControlRigEditorModule.cpp
@@ -2191,7 +2191,7 @@ void FControlRigEditorModule::GetContextMenuActions(const UControlRigGraphSchema
 
 										TSharedPtr<FDelegateHandle> RunOnceHandle = MakeShareable(new FDelegateHandle);
 										*(RunOnceHandle.Get()) = ControlRig->OnExecuted_AnyThread().AddLambda(
-											[RunOnceHandle, EventName, PreviousEventQueue](UControlRig* InRig, const EControlRigState InState, const FName& InEventName)
+											[RunOnceHandle, EventName, PreviousEventQueue](UControlRig* InRig, const FName& InEventName)
 											{
 												if(InEventName == EventName)
 												{
diff --git a/Engine/Plugins/Animation/ControlRig/Source/ControlRigEditor/Private/EditMode/ControlRigEditMode.cpp b/Engine/Plugins/Animation/ControlRig/Source/ControlRigEditor/Private/EditMode/ControlRigEditMode.cpp
--- a/Engine/Plugins/Animation/ControlRig/Source/ControlRigEditor/Private/EditMode/ControlRigEditMode.cpp
+++ b/Engine/Plugins/Animation/ControlRig/Source/ControlRigEditor/Private/EditMode/ControlRigEditMode.cpp
@@ -3712,14 +3712,12 @@ void FControlRigEditMode::OnControlModified(UControlRig* Subject, FRigControlEle
 	*/
 }
 
-void FControlRigEditMode::OnPreConstruction_AnyThread(UControlRig* InRig, const EControlRigState InState,
-	const FName& InEventName)
+void FControlRigEditMode::OnPreConstruction_AnyThread(UControlRig* InRig, const FName& InEventName)
 {
 	bIsConstructionEventRunning = true;
 }
 
-void FControlRigEditMode::OnPostConstruction_AnyThread(UControlRig* InRig, const EControlRigState InState,
-	const FName& InEventName)
+void FControlRigEditMode::OnPostConstruction_AnyThread(UControlRig* InRig, const FName& InEventName)
 {
 	bIsConstructionEventRunning = false;
 
diff --git a/Engine/Plugins/Animation/ControlRig/Source/ControlRigEditor/Private/EditMode/ControlRigEditMode.h b/Engine/Plugins/Animation/ControlRig/Source/ControlRigEditor/Private/EditMode/ControlRigEditMode.h
--- a/Engine/Plugins/Animation/ControlRig/Source/ControlRigEditor/Private/EditMode/ControlRigEditMode.h
+++ b/Engine/Plugins/Animation/ControlRig/Source/ControlRigEditor/Private/EditMode/ControlRigEditMode.h
@@ -164,8 +164,8 @@ public:
 	void OnHierarchyModified(ERigHierarchyNotification InNotif, URigHierarchy* InHierarchy, const FRigBaseElement* InElement);
 	void OnHierarchyModified_AnyThread(ERigHierarchyNotification InNotif, URigHierarchy* InHierarchy, const FRigBaseElement* InElement);
 	void OnControlModified(UControlRig* Subject, FRigControlElement* InControlElement, const FRigControlModifiedContext& Context);
-	void OnPreConstruction_AnyThread(UControlRig* InRig, const EControlRigState InState, const FName& InEventName);
-	void OnPostConstruction_AnyThread(UControlRig* InRig, const EControlRigState InState, const FName& InEventName);
+	void OnPreConstruction_AnyThread(UControlRig* InRig, const FName& InEventName);
+	void OnPostConstruction_AnyThread(UControlRig* InRig, const FName& InEventName);
 
 	/** return true if it can be removed from preview scene 
 	- this is to ensure preview scene doesn't remove shape actors */
diff --git a/Engine/Plugins/Animation/ControlRig/Source/ControlRigEditor/Private/Editor/ControlRigEditor.cpp b/Engine/Plugins/Animation/ControlRig/Source/ControlRigEditor/Private/Editor/ControlRigEditor.cpp
--- a/Engine/Plugins/Animation/ControlRig/Source/ControlRigEditor/Private/Editor/ControlRigEditor.cpp
+++ b/Engine/Plugins/Animation/ControlRig/Source/ControlRigEditor/Private/Editor/ControlRigEditor.cpp
@@ -3234,7 +3234,7 @@ void FControlRigEditor::HandleVMCompiledEvent(UObject* InCompiledObject, URigVM*
 	UpdateGraphCompilerErrors();
 }
 
-void FControlRigEditor::HandleControlRigExecutedEvent(UControlRig* InControlRig, const EControlRigState InState, const FName& InEventName)
+void FControlRigEditor::HandleControlRigExecutedEvent(UControlRig* InControlRig, const FName& InEventName)
 {
 	if (UControlRigBlueprint* ControlRigBP = GetControlRigBlueprint())
 	{
@@ -7196,14 +7196,12 @@ FString* FControlRigEditor::GetSnippetStorage(int32 InSnippetIndex)
 	return nullptr;
 }
 
-void FControlRigEditor::OnPreConstruction_AnyThread(UControlRig* InRig, const EControlRigState InState,
-	const FName& InEventName)
+void FControlRigEditor::OnPreConstruction_AnyThread(UControlRig* InRig, const FName& InEventName)
 {
 	bIsConstructionEventRunning = true;
 }
 
-void FControlRigEditor::OnPostConstruction_AnyThread(UControlRig* InRig, const EControlRigState InState,
-	const FName& InEventName)
+void FControlRigEditor::OnPostConstruction_AnyThread(UControlRig* InRig, const FName& InEventName)
 {
 	bIsConstructionEventRunning = false;
 	
