Control Rig: Follow up fix for 23407131

diff --git a/Engine/Plugins/Animation/ControlRig/Source/ControlRig/Private/Units/Highlevel/Hierarchy/RigUnit_MultiFABRIK.cpp b/Engine/Plugins/Animation/ControlRig/Source/ControlRig/Private/Units/Highlevel/Hierarchy/RigUnit_MultiFABRIK.cpp
--- a/Engine/Plugins/Animation/ControlRig/Source/ControlRig/Private/Units/Highlevel/Hierarchy/RigUnit_MultiFABRIK.cpp
+++ b/Engine/Plugins/Animation/ControlRig/Source/ControlRig/Private/Units/Highlevel/Hierarchy/RigUnit_MultiFABRIK.cpp
@@ -233,7 +233,7 @@ FRigUnit_MultiFABRIK_Execute()
 		EffectorIndices.Reset();
 		BoneTree.Reset();
 		ChainGroup.Reset();
-		bIsInitialized = false;
+		bIsInitialized = true;
 	}
 	
 	if(BoneTree.Num() == 0)
