Control Rig: Fix missing includes

diff --git a/Engine/Plugins/Animation/ControlRig/Source/ControlRig/Private/Units/Animation/RigUnit_AnimEasing.cpp b/Engine/Plugins/Animation/ControlRig/Source/ControlRig/Private/Units/Animation/RigUnit_AnimEasing.cpp
--- a/Engine/Plugins/Animation/ControlRig/Source/ControlRig/Private/Units/Animation/RigUnit_AnimEasing.cpp
+++ b/Engine/Plugins/Animation/ControlRig/Source/ControlRig/Private/Units/Animation/RigUnit_AnimEasing.cpp
@@ -2,6 +2,7 @@
 
 #include "Units/Animation/RigUnit_AnimEasing.h"
 #include "Units/RigUnitContext.h"
+#include "Math/ControlRigMathLibrary.h"
 
 #include UE_INLINE_GENERATED_CPP_BY_NAME(RigUnit_AnimEasing)
 
diff --git a/Engine/Plugins/Animation/ControlRig/Source/ControlRig/Private/Units/Debug/RigUnit_DebugBezier.cpp b/Engine/Plugins/Animation/ControlRig/Source/ControlRig/Private/Units/Debug/RigUnit_DebugBezier.cpp
--- a/Engine/Plugins/Animation/ControlRig/Source/ControlRig/Private/Units/Debug/RigUnit_DebugBezier.cpp
+++ b/Engine/Plugins/Animation/ControlRig/Source/ControlRig/Private/Units/Debug/RigUnit_DebugBezier.cpp
@@ -2,6 +2,7 @@
 
 #include "Units/Debug/RigUnit_DebugBezier.h"
 #include "Units/RigUnitContext.h"
+#include "Math/ControlRigMathLibrary.h"
 
 #include UE_INLINE_GENERATED_CPP_BY_NAME(RigUnit_DebugBezier)
 
diff --git a/Engine/Plugins/Animation/ControlRig/Source/ControlRig/Private/Units/Deprecated/Math/RigUnit_Vector.cpp b/Engine/Plugins/Animation/ControlRig/Source/ControlRig/Private/Units/Deprecated/Math/RigUnit_Vector.cpp
--- a/Engine/Plugins/Animation/ControlRig/Source/ControlRig/Private/Units/Deprecated/Math/RigUnit_Vector.cpp
+++ b/Engine/Plugins/Animation/ControlRig/Source/ControlRig/Private/Units/Deprecated/Math/RigUnit_Vector.cpp
@@ -2,6 +2,7 @@
 
 #include "RigUnit_Vector.h"
 #include "RigVMFunctions/Math/RigVMFunction_MathVector.h"
+#include "Math/ControlRigMathLibrary.h"
 
 #include UE_INLINE_GENERATED_CPP_BY_NAME(RigUnit_Vector)
 
diff --git a/Engine/Plugins/Animation/ControlRig/Source/ControlRig/Private/Units/Hierarchy/RigUnit_AddBoneTransform.cpp b/Engine/Plugins/Animation/ControlRig/Source/ControlRig/Private/Units/Hierarchy/RigUnit_AddBoneTransform.cpp
--- a/Engine/Plugins/Animation/ControlRig/Source/ControlRig/Private/Units/Hierarchy/RigUnit_AddBoneTransform.cpp
+++ b/Engine/Plugins/Animation/ControlRig/Source/ControlRig/Private/Units/Hierarchy/RigUnit_AddBoneTransform.cpp
@@ -3,6 +3,7 @@
 #include "Units/Hierarchy/RigUnit_AddBoneTransform.h"
 #include "Units/Hierarchy/RigUnit_OffsetTransform.h"
 #include "Units/RigUnitContext.h"
+#include "Math/ControlRigMathLibrary.h"
 
 #include UE_INLINE_GENERATED_CPP_BY_NAME(RigUnit_AddBoneTransform)
 
diff --git a/Engine/Plugins/Animation/ControlRig/Source/ControlRig/Private/Units/Hierarchy/RigUnit_ParentSwitchConstraint.cpp b/Engine/Plugins/Animation/ControlRig/Source/ControlRig/Private/Units/Hierarchy/RigUnit_ParentSwitchConstraint.cpp
--- a/Engine/Plugins/Animation/ControlRig/Source/ControlRig/Private/Units/Hierarchy/RigUnit_ParentSwitchConstraint.cpp
+++ b/Engine/Plugins/Animation/ControlRig/Source/ControlRig/Private/Units/Hierarchy/RigUnit_ParentSwitchConstraint.cpp
@@ -2,6 +2,7 @@
 
 #include "RigUnit_ParentSwitchConstraint.h"
 #include "Units/RigUnitContext.h"
+#include "Math/ControlRigMathLibrary.h"
 
 #include UE_INLINE_GENERATED_CPP_BY_NAME(RigUnit_ParentSwitchConstraint)
 
diff --git a/Engine/Plugins/Animation/ControlRig/Source/ControlRig/Private/Units/Highlevel/Harmonics/RigUnit_BoneHarmonics.cpp b/Engine/Plugins/Animation/ControlRig/Source/ControlRig/Private/Units/Highlevel/Harmonics/RigUnit_BoneHarmonics.cpp
--- a/Engine/Plugins/Animation/ControlRig/Source/ControlRig/Private/Units/Highlevel/Harmonics/RigUnit_BoneHarmonics.cpp
+++ b/Engine/Plugins/Animation/ControlRig/Source/ControlRig/Private/Units/Highlevel/Harmonics/RigUnit_BoneHarmonics.cpp
@@ -3,6 +3,7 @@
 #include "Units/Highlevel/Harmonics/RigUnit_BoneHarmonics.h"
 #include "Units/RigUnitContext.h"
 #include "AnimationCoreLibrary.h"
+#include "Math/ControlRigMathLibrary.h"
 
 #include UE_INLINE_GENERATED_CPP_BY_NAME(RigUnit_BoneHarmonics)
 
diff --git a/Engine/Plugins/Animation/ControlRig/Source/ControlRig/Private/Units/Highlevel/Hierarchy/RigUnit_CCDIK.cpp b/Engine/Plugins/Animation/ControlRig/Source/ControlRig/Private/Units/Highlevel/Hierarchy/RigUnit_CCDIK.cpp
--- a/Engine/Plugins/Animation/ControlRig/Source/ControlRig/Private/Units/Highlevel/Hierarchy/RigUnit_CCDIK.cpp
+++ b/Engine/Plugins/Animation/ControlRig/Source/ControlRig/Private/Units/Highlevel/Hierarchy/RigUnit_CCDIK.cpp
@@ -2,6 +2,7 @@
 
 #include "RigUnit_CCDIK.h"
 #include "Units/RigUnitContext.h"
+#include "Math/ControlRigMathLibrary.h"
 
 #include UE_INLINE_GENERATED_CPP_BY_NAME(RigUnit_CCDIK)
 
diff --git a/Engine/Plugins/Animation/ControlRig/Source/ControlRig/Private/Units/Highlevel/Hierarchy/RigUnit_DistributeRotation.cpp b/Engine/Plugins/Animation/ControlRig/Source/ControlRig/Private/Units/Highlevel/Hierarchy/RigUnit_DistributeRotation.cpp
--- a/Engine/Plugins/Animation/ControlRig/Source/ControlRig/Private/Units/Highlevel/Hierarchy/RigUnit_DistributeRotation.cpp
+++ b/Engine/Plugins/Animation/ControlRig/Source/ControlRig/Private/Units/Highlevel/Hierarchy/RigUnit_DistributeRotation.cpp
@@ -3,6 +3,7 @@
 #include "Units/Highlevel/Hierarchy/RigUnit_DistributeRotation.h"
 #include "Units/RigUnitContext.h"
 #include "AnimationCoreLibrary.h"
+#include "Math/ControlRigMathLibrary.h"
 
 #include UE_INLINE_GENERATED_CPP_BY_NAME(RigUnit_DistributeRotation)
 
diff --git a/Engine/Plugins/Animation/ControlRig/Source/ControlRig/Private/Units/Highlevel/Hierarchy/RigUnit_FABRIK.cpp b/Engine/Plugins/Animation/ControlRig/Source/ControlRig/Private/Units/Highlevel/Hierarchy/RigUnit_FABRIK.cpp
--- a/Engine/Plugins/Animation/ControlRig/Source/ControlRig/Private/Units/Highlevel/Hierarchy/RigUnit_FABRIK.cpp
+++ b/Engine/Plugins/Animation/ControlRig/Source/ControlRig/Private/Units/Highlevel/Hierarchy/RigUnit_FABRIK.cpp
@@ -2,6 +2,7 @@
 
 #include "RigUnit_FABRIK.h"
 #include "Units/RigUnitContext.h"
+#include "Math/ControlRigMathLibrary.h"
 
 #include UE_INLINE_GENERATED_CPP_BY_NAME(RigUnit_FABRIK)
 
diff --git a/Engine/Plugins/Animation/ControlRig/Source/ControlRig/Private/Units/Highlevel/Hierarchy/RigUnit_FitChainToCurve.cpp b/Engine/Plugins/Animation/ControlRig/Source/ControlRig/Private/Units/Highlevel/Hierarchy/RigUnit_FitChainToCurve.cpp
--- a/Engine/Plugins/Animation/ControlRig/Source/ControlRig/Private/Units/Highlevel/Hierarchy/RigUnit_FitChainToCurve.cpp
+++ b/Engine/Plugins/Animation/ControlRig/Source/ControlRig/Private/Units/Highlevel/Hierarchy/RigUnit_FitChainToCurve.cpp
@@ -3,6 +3,7 @@
 #include "Units/Highlevel/Hierarchy/RigUnit_FitChainToCurve.h"
 #include "Units/Debug/RigUnit_DebugBezier.h"
 #include "Units/RigUnitContext.h"
+#include "Math/ControlRigMathLibrary.h"
 
 #include UE_INLINE_GENERATED_CPP_BY_NAME(RigUnit_FitChainToCurve)
 
diff --git a/Engine/Plugins/Animation/ControlRig/Source/ControlRig/Private/Units/Highlevel/Hierarchy/RigUnit_ModifyTransforms.cpp b/Engine/Plugins/Animation/ControlRig/Source/ControlRig/Private/Units/Highlevel/Hierarchy/RigUnit_ModifyTransforms.cpp
--- a/Engine/Plugins/Animation/ControlRig/Source/ControlRig/Private/Units/Highlevel/Hierarchy/RigUnit_ModifyTransforms.cpp
+++ b/Engine/Plugins/Animation/ControlRig/Source/ControlRig/Private/Units/Highlevel/Hierarchy/RigUnit_ModifyTransforms.cpp
@@ -2,6 +2,7 @@
 
 #include "RigUnit_ModifyTransforms.h"
 #include "Units/RigUnitContext.h"
+#include "Math/ControlRigMathLibrary.h"
 
 #include UE_INLINE_GENERATED_CPP_BY_NAME(RigUnit_ModifyTransforms)
 
diff --git a/Engine/Plugins/Animation/ControlRig/Source/ControlRig/Private/Units/Highlevel/Hierarchy/RigUnit_TransformConstraint.cpp b/Engine/Plugins/Animation/ControlRig/Source/ControlRig/Private/Units/Highlevel/Hierarchy/RigUnit_TransformConstraint.cpp
--- a/Engine/Plugins/Animation/ControlRig/Source/ControlRig/Private/Units/Highlevel/Hierarchy/RigUnit_TransformConstraint.cpp
+++ b/Engine/Plugins/Animation/ControlRig/Source/ControlRig/Private/Units/Highlevel/Hierarchy/RigUnit_TransformConstraint.cpp
@@ -4,6 +4,7 @@
 #include "Units/RigUnitContext.h"
 #include "HelperUtil.h"
 #include "AnimationCoreLibrary.h"
+#include "Math/ControlRigMathLibrary.h"
 
 #include UE_INLINE_GENERATED_CPP_BY_NAME(RigUnit_TransformConstraint)
 
diff --git a/Engine/Plugins/Animation/ControlRig/Source/ControlRig/Private/Units/Highlevel/Hierarchy/RigUnit_TwistBones.cpp b/Engine/Plugins/Animation/ControlRig/Source/ControlRig/Private/Units/Highlevel/Hierarchy/RigUnit_TwistBones.cpp
--- a/Engine/Plugins/Animation/ControlRig/Source/ControlRig/Private/Units/Highlevel/Hierarchy/RigUnit_TwistBones.cpp
+++ b/Engine/Plugins/Animation/ControlRig/Source/ControlRig/Private/Units/Highlevel/Hierarchy/RigUnit_TwistBones.cpp
@@ -2,6 +2,7 @@
 
 #include "Units/Highlevel/Hierarchy/RigUnit_TwistBones.h"
 #include "Units/RigUnitContext.h"
+#include "Math/ControlRigMathLibrary.h"
 
 #include UE_INLINE_GENERATED_CPP_BY_NAME(RigUnit_TwistBones)
 
