Control Rig: Fix redirect for EBoneGetterSetterMode EControlRigAnimEasingType and EControlRigClampSpatialMode

diff --git a/Engine/Plugins/Animation/ControlRig/Config/BaseControlRig.ini b/Engine/Plugins/Animation/ControlRig/Config/BaseControlRig.ini
--- a/Engine/Plugins/Animation/ControlRig/Config/BaseControlRig.ini
+++ b/Engine/Plugins/Animation/ControlRig/Config/BaseControlRig.ini
@@ -87,9 +87,9 @@
 +StructRedirects=(OldName="/Script/ControlRig.ControlRigDrawInterface",NewName="/Script/RigVM.RigVMDrawInterface")
 
 ; moving math units to RigVM runtime
-+EnumRedirects=(OldName="EBoneGetterSetterMode",NewName="/Script.RigVM.ERigVMTransformSpace")
-+EnumRedirects=(OldName="EControlRigAnimEasingType",NewName="/Script.RigVM.ERigVMAnimEasingType")
-+EnumRedirects=(OldName="EControlRigClampSpatialMode",NewName="/Script.RigVM.ERigVMClampSpatialMode")
++EnumRedirects=(OldName="EBoneGetterSetterMode",NewName="/Script/RigVM.ERigVMTransformSpace")
++EnumRedirects=(OldName="EControlRigAnimEasingType",NewName="/Script/RigVM.ERigVMAnimEasingType")
++EnumRedirects=(OldName="EControlRigClampSpatialMode",NewName="/Script/RigVM.ERigVMClampSpatialMode")
 +StructRedirects=(OldName="/Script/ControlRig.CRFourPointBezier",NewName="/Script/RigVM.RigVMFourPointBezier")
 +StructRedirects=(OldName="/Script/ControlRig.RigMirrorSettings",NewName="/Script/RigVM.RigVMMirrorSettings")
 +StructRedirects=(OldName="/Script/ControlRig.RigUnit_MathBase",NewName="/Script/RigVM.RigVMFunction_MathBase")
