Control Rig: Fix python command adding tags to elements not working

diff --git a/Engine/Plugins/Animation/ControlRig/Source/ControlRig/Public/Rigs/RigHierarchy.h b/Engine/Plugins/Animation/ControlRig/Source/ControlRig/Public/Rigs/RigHierarchy.h
--- a/Engine/Plugins/Animation/ControlRig/Source/ControlRig/Public/Rigs/RigHierarchy.h
+++ b/Engine/Plugins/Animation/ControlRig/Source/ControlRig/Public/Rigs/RigHierarchy.h
@@ -1109,7 +1109,7 @@ public:
 	UFUNCTION(BlueprintCallable, Category = URigHierarchy)
 	bool SetNameArrayMetadata(FRigElementKey InItem, FName InMetadataName, TArray<FName> InValue)
 	{
-		return SetArrayMetadata<FName>(InItem, ERigMetadataType::FloatArray, InMetadataName, InValue);
+		return SetArrayMetadata<FName>(InItem, ERigMetadataType::NameArray, InMetadataName, InValue);
 	}
 
 	/**
diff --git a/Engine/Plugins/Animation/ControlRig/Source/ControlRigEditor/Private/Graph/NodeSpawners/ControlRigUnitNodeSpawner.cpp b/Engine/Plugins/Animation/ControlRig/Source/ControlRigEditor/Private/Graph/NodeSpawners/ControlRigUnitNodeSpawner.cpp
--- a/Engine/Plugins/Animation/ControlRig/Source/ControlRigEditor/Private/Graph/NodeSpawners/ControlRigUnitNodeSpawner.cpp
+++ b/Engine/Plugins/Animation/ControlRig/Source/ControlRigEditor/Private/Graph/NodeSpawners/ControlRigUnitNodeSpawner.cpp
@@ -294,7 +294,7 @@ UControlRigGraphNode* UControlRigUnitNodeSpawner::SpawnNode(UEdGraph* ParentGrap
 						for (int64 EnumValue = 0; EnumValue < MaxEnumValue; EnumValue++)
 						{
 							FString EnumText = RigElementTypeEnum->GetDisplayNameTextByValue(EnumValue).ToString().ToLower();
-							if (UsedFilterString.Contains(EnumText))
+							if (!EnumText.IsEmpty() && UsedFilterString.Contains(EnumText))
 							{
 								UsedElementType = (ERigElementType)EnumValue;
 								break;
