Control Rig: Fix crash when context menu is constructed from unit nodes

diff --git a/Engine/Plugins/Animation/ControlRig/Source/ControlRig/Private/Units/Hierarchy/RigUnit_Metadata.cpp b/Engine/Plugins/Animation/ControlRig/Source/ControlRig/Private/Units/Hierarchy/RigUnit_Metadata.cpp
--- a/Engine/Plugins/Animation/ControlRig/Source/ControlRig/Private/Units/Hierarchy/RigUnit_Metadata.cpp
+++ b/Engine/Plugins/Animation/ControlRig/Source/ControlRig/Private/Units/Hierarchy/RigUnit_Metadata.cpp
@@ -395,6 +395,10 @@ FRigUnit_RemoveMetadata_Execute()
 	URigHierarchy* Hierarchy = ExecuteContext.Hierarchy;
 
 	Removed = false;
+	if (!Hierarchy)
+	{
+		return;
+	}
 
 	if(CachedIndex.UpdateCache(Item, Hierarchy))
 	{
@@ -410,6 +414,10 @@ FRigUnit_RemoveAllMetadata_Execute()
 	URigHierarchy* Hierarchy = ExecuteContext.Hierarchy;
 
 	Removed = false;
+	if (!Hierarchy)
+	{
+		return;
+	}
 
 	if(CachedIndex.UpdateCache(Item, Hierarchy))
 	{
@@ -425,6 +433,10 @@ FRigUnit_HasMetadata_Execute()
 	const URigHierarchy* Hierarchy = ExecuteContext.Hierarchy;
 
 	Found = false;
+	if (!Hierarchy)
+	{
+		return;
+	}
 
 	if(CachedIndex.UpdateCache(Item, Hierarchy))
 	{
@@ -440,6 +452,10 @@ FRigUnit_FindItemsWithMetadata_Execute()
 	const URigHierarchy* Hierarchy = ExecuteContext.Hierarchy;
 
 	Items.Reset();
+	if (!Hierarchy)
+	{
+		return;
+	}
 
 	Hierarchy->Traverse([&Items, Name, Type](const FRigBaseElement* Element, bool& bContinue)
 	{
@@ -456,6 +472,11 @@ FRigUnit_GetMetadataTags_Execute()
 	Tags.Reset();
 	
 	const URigHierarchy* Hierarchy = ExecuteContext.Hierarchy;
+	if (!Hierarchy)
+	{
+		return;
+	}
+	
 	if(CachedIndex.UpdateCache(Item, Hierarchy))
 	{
 		if(const FRigBaseElement* Element = Hierarchy->Get(CachedIndex))
@@ -470,9 +491,15 @@ FRigUnit_GetMetadataTags_Execute()
 
 FRigUnit_SetMetadataTag_Execute()
 {
-	if(CachedIndex.UpdateCache(Item, ExecuteContext.Hierarchy))
+	URigHierarchy* Hierarchy = ExecuteContext.Hierarchy;
+	if (!Hierarchy)
+	{
+		return;
+	}
+	
+	if(CachedIndex.UpdateCache(Item, Hierarchy))
 	{
-		if(FRigBaseElement* Element = ExecuteContext.Hierarchy->Get(CachedIndex))
+		if(FRigBaseElement* Element = Hierarchy->Get(CachedIndex))
 		{
 			if(FRigNameArrayMetadata* Md = Cast<FRigNameArrayMetadata>(Element->SetupValidMetadata(URigHierarchy::TagMetadataName, ERigMetadataType::NameArray)))
 			{
@@ -488,9 +515,15 @@ FRigUnit_SetMetadataTag_Execute()
 
 FRigUnit_SetMetadataTagArray_Execute()
 {
-	if(CachedIndex.UpdateCache(Item, ExecuteContext.Hierarchy))
+	URigHierarchy* Hierarchy = ExecuteContext.Hierarchy;
+	if (!Hierarchy)
+	{
+		return;
+	}
+	
+	if(CachedIndex.UpdateCache(Item, Hierarchy))
 	{
-		if(FRigBaseElement* Element = ExecuteContext.Hierarchy->Get(CachedIndex))
+		if(FRigBaseElement* Element = Hierarchy->Get(CachedIndex))
 		{
 			if(FRigNameArrayMetadata* Md = Cast<FRigNameArrayMetadata>(Element->SetupValidMetadata(URigHierarchy::TagMetadataName, ERigMetadataType::NameArray)))
 			{
@@ -510,9 +543,15 @@ FRigUnit_SetMetadataTagArray_Execute()
 FRigUnit_RemoveMetadataTag_Execute()
 {
 	Removed = false;
-	if(CachedIndex.UpdateCache(Item, ExecuteContext.Hierarchy))
+	URigHierarchy* Hierarchy = ExecuteContext.Hierarchy;
+	if (!Hierarchy)
+	{
+		return;
+	}
+	
+	if(CachedIndex.UpdateCache(Item, Hierarchy))
 	{
-		if(FRigBaseElement* Element = ExecuteContext.Hierarchy->Get(CachedIndex))
+		if(FRigBaseElement* Element = Hierarchy->Get(CachedIndex))
 		{
 			if(FRigNameArrayMetadata* Md = Cast<FRigNameArrayMetadata>(Element->GetMetadata(URigHierarchy::TagMetadataName, ERigMetadataType::NameArray)))
 			{
@@ -531,6 +570,11 @@ FRigUnit_HasMetadataTag_Execute()
 	Found = false;
 	
 	const URigHierarchy* Hierarchy = ExecuteContext.Hierarchy;
+	if (!Hierarchy)
+	{
+		return;
+	}
+	
 	if(CachedIndex.UpdateCache(Item, Hierarchy))
 	{
 		if(const FRigBaseElement* Element = Hierarchy->Get(CachedIndex))
@@ -548,6 +592,11 @@ FRigUnit_HasMetadataTagArray_Execute()
 	Found = false;
 	
 	const URigHierarchy* Hierarchy = ExecuteContext.Hierarchy;
+	if (!Hierarchy)
+	{
+		return;
+	}
+	
 	if(CachedIndex.UpdateCache(Item, Hierarchy))
 	{
 		if(const FRigBaseElement* Element = Hierarchy->Get(CachedIndex))
@@ -573,6 +622,10 @@ FRigUnit_FindItemsWithMetadataTag_Execute()
 	const URigHierarchy* Hierarchy = ExecuteContext.Hierarchy;
 
 	Items.Reset();
+	if (!Hierarchy)
+	{
+		return;
+	}
 
 	Hierarchy->Traverse([&Items, Tag](const FRigBaseElement* Element, bool& bContinue)
 	{
@@ -592,6 +645,10 @@ FRigUnit_FindItemsWithMetadataTagArray_Execute()
 	const URigHierarchy* Hierarchy = ExecuteContext.Hierarchy;
 
 	Items.Reset();
+	if (!Hierarchy)
+	{
+		return;
+	}
 
 	Hierarchy->Traverse([&Items, Tags](const FRigBaseElement* Element, bool& bContinue)
 	{
@@ -621,7 +678,11 @@ FRigUnit_FilterItemsByMetadataTags_Execute()
 	const URigHierarchy* Hierarchy = ExecuteContext.Hierarchy;
 
 	Result.Reset();
-
+	if (!Hierarchy)
+	{
+		return;
+	}
+	
 	if(CachedIndices.Num() != Items.Num())
 	{
 		CachedIndices.Reset();
