Control Rig: Add pin to aggregate node not working

diff --git a/Engine/Plugins/Animation/ControlRig/Source/ControlRigEditor/Private/Graph/SControlRigGraphNode.cpp b/Engine/Plugins/Animation/ControlRig/Source/ControlRigEditor/Private/Graph/SControlRigGraphNode.cpp
--- a/Engine/Plugins/Animation/ControlRig/Source/ControlRigEditor/Private/Graph/SControlRigGraphNode.cpp
+++ b/Engine/Plugins/Animation/ControlRig/Source/ControlRigEditor/Private/Graph/SControlRigGraphNode.cpp
@@ -294,7 +294,7 @@ FReply SControlRigGraphNode::OnAddPin()
 	{
 		if (UControlRigGraphNode* ControlRigGraphNode = Cast<UControlRigGraphNode>(GraphNode))
 		{
-			if (ModelNode->IsA<URigVMAggregateNode>())
+			if (ModelNode->IsA<URigVMAggregateNode>() || ModelNode->IsAggregate())
 			{
 				ControlRigGraphNode->HandleAddAggregateElement(ModelNode->GetNodePath());
 			}
