Sequencer: When adding constraint channel make sure to recreate curve and resync selection
#trivial

diff --git a/Engine/Plugins/Animation/ControlRig/Source/ControlRigEditor/Private/Sequencer/ConstraintChannelHelper.cpp b/Engine/Plugins/Animation/ControlRig/Source/ControlRigEditor/Private/Sequencer/ConstraintChannelHelper.cpp
--- a/Engine/Plugins/Animation/ControlRig/Source/ControlRigEditor/Private/Sequencer/ConstraintChannelHelper.cpp
+++ b/Engine/Plugins/Animation/ControlRig/Source/ControlRigEditor/Private/Sequencer/ConstraintChannelHelper.cpp
@@ -579,7 +579,14 @@ bool FConstraintChannelHelper::SmartComponentConstraintKey(
 	InConstraint->bDynamicOffset = true;
 	
 	// add the channel
-	ConstraintSection->AddConstraintChannel(InConstraint);
+	if (ConstraintSection->HasConstraintChannel(InConstraint->GetFName()) == false)
+	{
+		ConstraintSection->AddConstraintChannel(InConstraint);
+		if (InSequencer.IsValid())
+		{
+			InSequencer->RecreateCurveEditor();
+		}
+	}
 
 	// add key if needed
 	if (FConstraintAndActiveChannel* Channel = ConstraintSection->GetConstraintChannel(InConstraint->GetFName()))
diff --git a/Engine/Source/Editor/Sequencer/Private/Sequencer.cpp b/Engine/Source/Editor/Sequencer/Private/Sequencer.cpp
--- a/Engine/Source/Editor/Sequencer/Private/Sequencer.cpp
+++ b/Engine/Source/Editor/Sequencer/Private/Sequencer.cpp
@@ -2352,6 +2352,14 @@ void FSequencer::RecreateCurveEditor()
 	if (ensure(CurveEditorIntegration))
 	{
 		CurveEditorIntegration->ResetCurveEditor();
+		if (GetSequencerSettings()->ShouldSyncCurveEditorSelection())
+		{
+			FCurveEditorExtension* CurveEditorExtension = ViewModel->CastDynamic<FCurveEditorExtension>();
+			if (ensure(CurveEditorExtension))
+			{
+				return CurveEditorExtension->RequestSyncSelection();
+			}
+		}
 	}
 }
 
