Only compensate the current modified child to avoid useless computation

diff --git a/Engine/Plugins/Animation/ControlRig/Source/ControlRigEditor/Private/Sequencer/ConstraintChannelHelper.cpp b/Engine/Plugins/Animation/ControlRig/Source/ControlRigEditor/Private/Sequencer/ConstraintChannelHelper.cpp
--- a/Engine/Plugins/Animation/ControlRig/Source/ControlRigEditor/Private/Sequencer/ConstraintChannelHelper.cpp
+++ b/Engine/Plugins/Animation/ControlRig/Source/ControlRigEditor/Private/Sequencer/ConstraintChannelHelper.cpp
@@ -747,14 +747,15 @@ void FConstraintChannelHelper::Compensate(UTickableTransformConstraint* InConstr
 		return;
 	}
 	
-	CompensateIfNeeded(World, Sequencer, Section, OptTime);
+	CompensateIfNeeded(World, Sequencer, Section, OptTime, InConstraint->GetTargetHash());
 }
 
 void FConstraintChannelHelper::CompensateIfNeeded(
 	UWorld* InWorld,
 	const TSharedPtr<ISequencer>& InSequencer,
 	IMovieSceneConstrainedSection* ConstraintSection,
-	const TOptional<FFrameNumber>& OptionalTime)
+	const TOptional<FFrameNumber>& OptionalTime,
+	const int32 InChildHash)
 {
 	if (FMovieSceneConstraintChannelHelper::bDoNotCompensate)
 	{
@@ -770,7 +771,7 @@ void FConstraintChannelHelper::CompensateIfNeeded(
 		OptionalTimeArray.Add(OptionalTime.GetValue());
 	}
 
-	auto GetSpaceTimesToCompensate = [&OptionalTimeArray](const FConstraintAndActiveChannel& Channel)->TArrayView<const FFrameNumber>
+	auto GetConstraintTimesToCompensate = [&OptionalTimeArray](const FConstraintAndActiveChannel& Channel)->TArrayView<const FFrameNumber>
 	{
 		if (OptionalTimeArray.IsEmpty())
 		{
@@ -779,28 +780,36 @@ void FConstraintChannelHelper::CompensateIfNeeded(
 		return OptionalTimeArray;
 	};
 
-	bool bNeedsEvaluation = false;
-
-	// gather all transform constraints
-	TArray<FConstraintAndActiveChannel>& ConstraintChannels = ConstraintSection->GetConstraintsChannels();
+	// gather all transform constraints' channels for
+	const TArray<FConstraintAndActiveChannel>& ConstraintChannels = ConstraintSection->GetConstraintsChannels();
 	TArray<FConstraintAndActiveChannel> TransformConstraintsChannels;
 	Algo::CopyIf(ConstraintChannels, TransformConstraintsChannels,
-		[](const FConstraintAndActiveChannel& InChannel)
+		[InChildHash](const FConstraintAndActiveChannel& InChannel)
 		{
 			if (!InChannel.Constraint.IsValid())
 			{
 				return false;
 			}
 
+			if ((InChildHash != INDEX_NONE) && (InChannel.Constraint->GetTargetHash() != InChildHash))
+			{
+				return false;
+			}
+
 			const UTickableTransformConstraint* Constraint = Cast<UTickableTransformConstraint>(InChannel.Constraint.Get());
-			return Constraint && Constraint->NeedsCompensation();
+			return Constraint && (Constraint->GetTargetHash() == InChildHash) && Constraint->NeedsCompensation();
 		}
 	);
 
-	// compensate constraints
+	// we only need to treat one single constraint per child as FCompensationEvaluator::ComputeCompensation will
+	// compensate within the last active constraint's space
+	using CompensationData = TPair< UTickableTransformConstraint*, TArray<FFrameNumber> >;
+	TArray< CompensationData > ToCompensate;
+
+	// store constraints and times where compensation is needed 
 	for (const FConstraintAndActiveChannel& Channel : TransformConstraintsChannels)
 	{
-		const TArrayView<const FFrameNumber> FramesToCompensate = GetSpaceTimesToCompensate(Channel);
+		const TArrayView<const FFrameNumber> FramesToCompensate = GetConstraintTimesToCompensate(Channel);
 		for (const FFrameNumber& Time : FramesToCompensate)
 		{
 			const FFrameNumber TimeMinusOne(Time - 1);
@@ -813,21 +822,49 @@ void FConstraintChannelHelper::CompensateIfNeeded(
 			{
 				UTickableTransformConstraint* Constraint = Cast<UTickableTransformConstraint>(Channel.Constraint.Get());
 
-				// compute transform to set
-				// if switching from active to inactive then we must add a key at T-1 in the constraint space
-				// if switching from inactive to active then we must add a key at T-1 in the previous constraint or parent space
-				FCompensationEvaluator Evaluator(Constraint);
-				Evaluator.ComputeCompensation(InWorld, InSequencer, Time);
-				const TArray<FTransform>& LocalTransforms = Evaluator.ChildLocals;
-
-				const EMovieSceneTransformChannel ChannelsToKey = Constraint->GetChannelsToKey();
-				FConstraintBaker::AddTransformKeys(
-					InSequencer, Constraint->ChildTRSHandle, { TimeMinusOne }, LocalTransforms, ChannelsToKey);
-				bNeedsEvaluation = true;
+				// is the child already in that array?
+				int32 DataIndex = ToCompensate.IndexOfByPredicate([Constraint](const CompensationData& InData)
+				{
+					return InData.Key->GetTargetHash() == Constraint->GetTargetHash();
+				});
+
+				// if not, add the constraint
+				if (DataIndex == INDEX_NONE)
+				{
+					DataIndex = ToCompensate.Emplace(Constraint, TArray<FFrameNumber>() );
+				}
+
+				// store the time it needs to be compensated at
+				TArray<FFrameNumber>& Times = ToCompensate[DataIndex].Value;
+				Times.AddUnique(Time);
 			}
 		}
 	}
 
+	// compensate
+	bool bNeedsEvaluation = false;
+	for (const CompensationData& Data: ToCompensate)
+	{
+		UTickableTransformConstraint* Constraint = Data.Key;
+		FCompensationEvaluator Evaluator(Constraint);
+		const EMovieSceneTransformChannel ChannelsToKey = Constraint->GetChannelsToKey();
+		
+		for (const FFrameNumber& Time : Data.Value)
+		{
+			// compute transform to set
+			// if switching from active to inactive then we must add a key at T-1 in the constraint space
+			// if switching from inactive to active then we must add a key at T-1 in the previous constraint or parent space
+			Evaluator.ComputeCompensation(InWorld, InSequencer, Time);
+			const TArray<FTransform>& LocalTransforms = Evaluator.ChildLocals;
+
+			const FFrameNumber TimeMinusOne(Time - 1);
+			FConstraintBaker::AddTransformKeys(
+				InSequencer, Constraint->ChildTRSHandle, { TimeMinusOne }, LocalTransforms, ChannelsToKey);
+			
+			bNeedsEvaluation = true;
+		}
+	}
+
 	if (bNeedsEvaluation)
 	{
 		InSequencer->ForceEvaluate();
diff --git a/Engine/Plugins/Animation/ControlRig/Source/ControlRigEditor/Private/Sequencer/ControlRigParameterTrackEditor.cpp b/Engine/Plugins/Animation/ControlRig/Source/ControlRigEditor/Private/Sequencer/ControlRigParameterTrackEditor.cpp
--- a/Engine/Plugins/Animation/ControlRig/Source/ControlRigEditor/Private/Sequencer/ControlRigParameterTrackEditor.cpp
+++ b/Engine/Plugins/Animation/ControlRig/Source/ControlRigEditor/Private/Sequencer/ControlRigParameterTrackEditor.cpp
@@ -3472,8 +3472,9 @@ FKeyPropertyResult FControlRigParameterTrackEditor::AddKeysToControlRigHandle(US
 					RigControlName, OptionalKeyTime);
 
 				// compensate constraints
+				const uint32 ControlHash = UTransformableControlHandle::ComputeHash(SectionControlRig, RigControlName);
 				FConstraintChannelHelper::CompensateIfNeeded(
-					SectionControlRig->GetWorld(), GetSequencer(), ParamSection, OptionalKeyTime);
+					SectionControlRig->GetWorld(), GetSequencer(), ParamSection, OptionalKeyTime, ControlHash);
 			}
 		}
 	}
diff --git a/Engine/Plugins/Animation/ControlRig/Source/ControlRigEditor/Public/ConstraintChannelHelper.h b/Engine/Plugins/Animation/ControlRig/Source/ControlRigEditor/Public/ConstraintChannelHelper.h
--- a/Engine/Plugins/Animation/ControlRig/Source/ControlRigEditor/Public/ConstraintChannelHelper.h
+++ b/Engine/Plugins/Animation/ControlRig/Source/ControlRigEditor/Public/ConstraintChannelHelper.h
@@ -46,7 +46,8 @@ public:
 		UWorld* InWorld,
 		const TSharedPtr<ISequencer>& InSequencer,
 		IMovieSceneConstrainedSection* Section,
-		const TOptional<FFrameNumber>& OptionalTime);
+		const TOptional<FFrameNumber>& OptionalTime,
+		const int32 InChildHash = INDEX_NONE);
 
 	/** @todo documentation. */
 	static UMovieSceneControlRigParameterSection* GetControlSection(
