Support selecting sub-components when dealing with several skeletal meshes under the same actor
   - this also simplifies the handle factory interface so that it now supports any registered uobject

diff --git a/Engine/Plugins/Animation/ControlRig/Source/ControlRig/Private/ControlRigModule.cpp b/Engine/Plugins/Animation/ControlRig/Source/ControlRig/Private/ControlRigModule.cpp
--- a/Engine/Plugins/Animation/ControlRig/Source/ControlRig/Private/ControlRigModule.cpp
+++ b/Engine/Plugins/Animation/ControlRig/Source/ControlRig/Private/ControlRigModule.cpp
@@ -38,34 +38,52 @@ void FControlRigModule::ShutdownModule()
 
 void FControlRigModule::RegisterTransformableCustomization() const
 {
-	// Register UTransformableControlHandle and has function to handle control thru the transform constraints system
-	// as AControlRigShapeActor is only available if the ControlRig plugin is loaded.
-	auto CreateControlHandleFunc = [](const UObject* InObject, UObject* Outer)->UTransformableHandle*
+	// load the module to trigger default objects registration
+	static const FName ConstraintsModuleName(TEXT("Constraints"));
+	FModuleManager::Get().LoadModule(ConstraintsModuleName);
+	
+	// register UControlRig and AControlRigShapeActor
+	auto CreateControlHandle = [](UObject* Outer, UObject* InObject, const FName& InControlName)->UTransformableHandle*
+	{
+		if (const UControlRig* ControlRig = Cast<UControlRig>(InObject))
+		{
+			return ControlRig->CreateTransformableControlHandle(Outer, InControlName);
+		}
+		return nullptr;
+	};
+
+	auto GetControlHash = [](const UObject* InObject, const FName& InControlName)->uint32
+	{
+		if (const UControlRig* ControlRig = Cast<UControlRig>(InObject))
+		{
+			return UTransformableControlHandle::ComputeHash(ControlRig, InControlName);
+		}
+		return 0;
+	};
+	
+	auto CreateControlHandleFromActor = [CreateControlHandle](UObject* Outer, UObject* InObject, const FName&)->UTransformableHandle*
 	{
 		if (const AControlRigShapeActor* ControlActor = Cast<AControlRigShapeActor>(InObject))
 		{
-			if (UControlRig* ControlRig = ControlActor->ControlRig.Get())
-			{
-				UTransformableControlHandle* CtrlHandle = ControlRig->CreateTransformableControlHandle(Outer, ControlActor->ControlName);
-				return CtrlHandle;
-			}
+			return CreateControlHandle(Outer, ControlActor->ControlRig.Get(), ControlActor->ControlName);
 		}
 		return nullptr;
 	};
 
-	auto GetControlHashFunc = [](const UObject* InObject)->uint32
+	auto GetControlHashFromActor = [GetControlHash](const UObject* InObject, const FName&)->uint32
 	{
 		if (const AControlRigShapeActor* ControlActor = Cast<AControlRigShapeActor>(InObject))
 		{
-			const uint32 ControlHash = UTransformableControlHandle::ComputeHash(
-				ControlActor->ControlRig.Get(), ControlActor->ControlName);
-			return ControlHash;
+			return GetControlHash(ControlActor->ControlRig.Get(), ControlActor->ControlName);
 		}
 		return 0;
 	};
 
+	// Register UTransformableControlHandle and has function to handle control thru the transform constraints system
+	// as AControlRigShapeActor is only available if the ControlRig plugin is loaded.
 	FTransformableRegistry& Registry = FTransformableRegistry::Get();
-	Registry.Register( AControlRigShapeActor::StaticClass(), CreateControlHandleFunc, GetControlHashFunc);
+	Registry.Register(AControlRigShapeActor::StaticClass(), CreateControlHandleFromActor, GetControlHashFromActor);
+	Registry.Register(UControlRig::StaticClass(), CreateControlHandle, GetControlHash);
 }
 
 IMPLEMENT_MODULE(FControlRigModule, ControlRig)
diff --git a/Engine/Plugins/Animation/ControlRig/Source/ControlRigEditor/Private/Editor/Constraints/SConstraintsWidget.cpp b/Engine/Plugins/Animation/ControlRig/Source/ControlRigEditor/Private/Editor/Constraints/SConstraintsWidget.cpp
--- a/Engine/Plugins/Animation/ControlRig/Source/ControlRigEditor/Private/Editor/Constraints/SConstraintsWidget.cpp
+++ b/Engine/Plugins/Animation/ControlRig/Source/ControlRigEditor/Private/Editor/Constraints/SConstraintsWidget.cpp
@@ -25,6 +25,7 @@
 #include "ISequencer.h"
 #include "Tools/BakingHelper.h"
 #include "MovieSceneToolHelpers.h"
+#include "Styling/SlateIconFinder.h"
 
 #define LOCTEXT_NAMESPACE "SConstraintsWidget"
 
@@ -293,19 +294,16 @@ void SDroppableConstraintItem::CreateConstraint(
 		return;	
 	}
 
-	// has socket?
-	USceneComponent* ComponentWithSockets = nullptr;
-	if(USceneComponent* ParentComponent = InParent->GetRootComponent())
+	// gather sub components with sockets
+	const TInlineComponentArray<USceneComponent*> Components(InParent);
+	const TArray<USceneComponent*> ComponentsWithSockets = Components.FilterByPredicate([](const USceneComponent* Component)
 	{
-		if (ParentComponent->HasAnySockets())
-		{
-			ComponentWithSockets = ParentComponent;
-		}
-	}
+		return Component->HasAnySockets();
+	});
 
 	// create constraints
 	auto CreateConstraint = [InCreationDelegate, InConstraintType](
-		const TArray<AActor*>& Selection, AActor* InParent, const FName& InSocketName)
+		const TArray<AActor*>& Selection, UObject* InParent, const FName& InSocketName)
 	{
 		UWorld* World = GetCurrentWorld();
 		if (!IsValid(World))
@@ -320,7 +318,7 @@ void SDroppableConstraintItem::CreateConstraint(
 			{
 				FScopedTransaction Transaction(LOCTEXT("CreateConstraintKey", "Create Constraint Key"));
 				UTickableTransformConstraint* Constraint =
-					FTransformConstraintUtils::CreateAndAddFromActors(World, InParent, InSocketName, Child, InConstraintType);
+					FTransformConstraintUtils::CreateAndAddFromObjects(World, InParent, InSocketName, Child, NAME_None, InConstraintType);
 				if (Constraint)
 				{
 					FConstraintChannelHelper::SmartConstraintKey(Constraint, TOptional<bool>(), TOptional<FFrameNumber>());
@@ -335,34 +333,69 @@ void SDroppableConstraintItem::CreateConstraint(
 			InCreationDelegate.Execute();
 		}
 	};
-	
-	// Show socket chooser if we have sockets to select
-	if (ComponentWithSockets != nullptr)
-	{		
+
+	const int32 NumComponentsWithSockets = ComponentsWithSockets.Num();
+
+	// if no component socket available then constrain the whole actor 
+	if (NumComponentsWithSockets == 0)
+	{
+		return CreateConstraint(Selection, InParent, NAME_None);
+	}
+
+	// creates a menu encapsulating InContent  
+	auto CreateMenu = [](const TSharedRef<SWidget>& InContent, const FVector2D& InLocation)
+	{
 		const FLevelEditorModule& LevelEditorModule = FModuleManager::GetModuleChecked<FLevelEditorModule>("LevelEditor");
 		const TSharedPtr< ILevelEditor > LevelEditor = LevelEditorModule.GetFirstLevelEditor();
-
-		FVector2D SummonLocation = FSlateApplication::Get().GetCursorPos();
-		SummonLocation.Y += 4 * FSlateApplication::Get().GetCursorSize().Y;
-		
-		// Create as context menu
 		FSlateApplication::Get().PushMenu(
 			LevelEditor.ToSharedRef(),
 			FWidgetPath(),
-			SNew(SSocketChooserPopup)
-			.SceneComponent( ComponentWithSockets )
-			.OnSocketChosen_Lambda([CreateConstraint, Selection, InParent](FName InSocketName)
-			{
-				CreateConstraint(Selection, InParent, InSocketName);
-			}),
-			SummonLocation,
-			FPopupTransitionEffect( FPopupTransitionEffect::ContextMenu )
-			);
+			InContent,
+		InLocation,
+		FPopupTransitionEffect( FPopupTransitionEffect::ContextMenu ) );
+	};
+
+	// creates a new socket chooser popup widget
+	auto GetSocketChooserWidget = [CreateConstraint, Selection](USceneComponent* Component)
+	{
+		return SNew(SSocketChooserPopup)
+		.SceneComponent( Component )
+		.OnSocketChosen_Lambda([CreateConstraint, Selection, Component](FName InSocketName)
+		{
+			CreateConstraint(Selection, Component, InSocketName);
+		});
+	};
+
+	// decal the menu 
+	FVector2D SummonLocation = FSlateApplication::Get().GetCursorPos();
+	SummonLocation.Y += 4 * FSlateApplication::Get().GetCursorSize().Y;
+	
+	// if one component with sockets then constrain the component selecting the socket
+	if (NumComponentsWithSockets == 1)
+	{
+		return CreateMenu( GetSocketChooserWidget(ComponentsWithSockets[0]), SummonLocation);
 	}
-	else
+
+	// if there are several of them, then build a component chooser first then the socket chooser
+	static constexpr bool CloseAfterSelection = true;
+	FMenuBuilder MenuBuilder(CloseAfterSelection, nullptr);
+	MenuBuilder.BeginSection("ChooseComp", LOCTEXT("ChooseComponentSection", "Choose Component"));
 	{
-		CreateConstraint(Selection, InParent, NAME_None);
+		for (USceneComponent* Component: ComponentsWithSockets)
+		{
+			MenuBuilder.AddMenuEntry(FText::FromName(Component->GetFName()),
+			FText(),
+				FSlateIconFinder::FindIconForClass(Component->GetClass(), TEXT("SCS.Component")),
+				FUIAction(FExecuteAction::CreateLambda([CreateMenu, GetSocketChooserWidget, Component]()
+				{
+					CreateMenu(GetSocketChooserWidget(Component), FSlateApplication::Get().GetCursorPos());
+				})),
+				NAME_None,
+			   EUserInterfaceActionType::Button);
+		}
 	}
+	MenuBuilder.EndSection();
+	CreateMenu(MenuBuilder.MakeWidget(), SummonLocation);
 }
 
 /**
diff --git a/Engine/Source/Runtime/Experimental/Animation/Constraints/Private/ConstraintsManager.cpp b/Engine/Source/Runtime/Experimental/Animation/Constraints/Private/ConstraintsManager.cpp
--- a/Engine/Source/Runtime/Experimental/Animation/Constraints/Private/ConstraintsManager.cpp
+++ b/Engine/Source/Runtime/Experimental/Animation/Constraints/Private/ConstraintsManager.cpp
@@ -213,6 +213,10 @@ void UConstraintsManager::Init(UWorld* World)
 
 UConstraintsManager* UConstraintsManager::Get(UWorld* InWorld)
 {
+	if (!IsValid(InWorld))
+	{
+		return nullptr;
+	}
 	// look for ConstraintsActor and return its manager
 	if (UConstraintsManager* Manager = Find(InWorld))
 	{
@@ -230,6 +234,11 @@ UConstraintsManager* UConstraintsManager::Get(UWorld* InWorld)
 
 UConstraintsManager* UConstraintsManager::Find(const UWorld* InWorld)
 {
+	if (!IsValid(InWorld))
+	{
+		return nullptr;
+	}
+	
 	// should we work with the persistent level?
 	const ULevel* Level = InWorld->GetCurrentLevel();
 
diff --git a/Engine/Source/Runtime/Experimental/Animation/Constraints/Private/ConstraintsModule.cpp b/Engine/Source/Runtime/Experimental/Animation/Constraints/Private/ConstraintsModule.cpp
new file mode 100644
--- /dev/null
+++ b/Engine/Source/Runtime/Experimental/Animation/Constraints/Private/ConstraintsModule.cpp
@@ -0,0 +1,28 @@
+﻿// Copyright Epic Games, Inc. All Rights Reserved.
+
+#include "Modules/ModuleManager.h"
+#include "Modules/ModuleInterface.h"
+
+#include "TransformableRegistry.h"
+
+DECLARE_LOG_CATEGORY_EXTERN(LogConstraints, Log, All);
+DEFINE_LOG_CATEGORY(LogConstraints);
+
+class FConstraintsModule : public IModuleInterface
+{
+public:
+
+	virtual void StartupModule() override
+	{
+		FTransformableRegistry::RegisterBaseObjects();
+	}
+
+	virtual void ShutdownModule() override
+	{
+		FTransformableRegistry::UnregisterAllObjects();
+	}
+};
+
+IMPLEMENT_MODULE(FConstraintsModule, Constraints);
+
+
diff --git a/Engine/Source/Runtime/Experimental/Animation/Constraints/Private/ConstraintsScripting.cpp b/Engine/Source/Runtime/Experimental/Animation/Constraints/Private/ConstraintsScripting.cpp
--- a/Engine/Source/Runtime/Experimental/Animation/Constraints/Private/ConstraintsScripting.cpp
+++ b/Engine/Source/Runtime/Experimental/Animation/Constraints/Private/ConstraintsScripting.cpp
@@ -5,6 +5,7 @@
 #include "TransformConstraint.h"
 #include "TransformableHandle.h"
 #include "ConstraintsManager.h"
+#include "TransformableRegistry.h"
 
 #include UE_INLINE_GENERATED_CPP_BY_NAME(ConstraintsScripting)
 
@@ -28,6 +29,25 @@ UTransformableComponentHandle* UConstraintsScriptingLibrary::CreateTransformable
 	return nullptr;
 }
 
+UTransformableHandle* UConstraintsScriptingLibrary::CreateTransformableHandle(UWorld* InWorld, UObject* InObject, const FName& InAttachmentName)
+{
+	UConstraintsManager* ConstraintsManager = UConstraintsManager::Get(InWorld);
+	if (!ConstraintsManager)
+	{
+		return nullptr;
+	}
+	
+	// look for customized transform handle
+	const FTransformableRegistry& Registry = FTransformableRegistry::Get();
+	if (const FTransformableRegistry::CreateHandleFuncT CreateFunction = Registry.GetCreateFunction(InObject->GetClass()))
+	{
+		return CreateFunction(ConstraintsManager, InObject, InAttachmentName);
+	}
+	
+	return nullptr;
+}
+
+
 UTickableTransformConstraint* UConstraintsScriptingLibrary::CreateFromType(
 	UWorld* InWorld,
 	const ETransformConstraintType InType)
diff --git a/Engine/Source/Runtime/Experimental/Animation/Constraints/Private/TransformConstraint.cpp b/Engine/Source/Runtime/Experimental/Animation/Constraints/Private/TransformConstraint.cpp
--- a/Engine/Source/Runtime/Experimental/Animation/Constraints/Private/TransformConstraint.cpp
+++ b/Engine/Source/Runtime/Experimental/Animation/Constraints/Private/TransformConstraint.cpp
@@ -1275,36 +1275,28 @@ void UTickableLookAtConstraint::PostEditChangeProperty(FPropertyChangedEvent& Pr
 namespace
 {
 
-UTransformableHandle* GetHandle(AActor* InActor, const FName& InSocketName, UObject* Outer)
+UTransformableHandle* GetHandle(UObject* InObject, const FName& InSocketName, UObject* Outer)
 {
 	// look for customized transform handle
 	const FTransformableRegistry& Registry = FTransformableRegistry::Get();
-	if (const FTransformableRegistry::CreateHandleFuncT CreateFunction = Registry.GetCreateFunction(InActor->GetClass()))
+	if (const FTransformableRegistry::CreateHandleFuncT CreateFunction = Registry.GetCreateFunction(InObject->GetClass()))
 	{
-		return CreateFunction(InActor, Outer);
+		return CreateFunction(Outer, InObject, InSocketName);
 	}
 
-	// need to make sure it's moveable
-	if (InActor->GetRootComponent())
-	{
-		return FTransformConstraintUtils::CreateHandleForSceneComponent(InActor->GetRootComponent(), InSocketName, Outer);
-
-	}
 	return nullptr;
 }
 	
-uint32 GetConstrainableHash(const AActor* InActor)
+uint32 GetConstrainableHash(const UObject* InObject)
 {
 	// look for customized hash function
 	const FTransformableRegistry& Registry = FTransformableRegistry::Get();
-	if (const FTransformableRegistry::GetHashFuncT HashFunction = Registry.GetHashFunction(InActor->GetClass()))
+	if (const FTransformableRegistry::GetHashFuncT HashFunction = Registry.GetHashFunction(InObject->GetClass()))
 	{
-		return HashFunction(InActor);
+		return HashFunction(InObject, NAME_None);
 	}
 
-	// scene component hash
-	const uint32 ComponentHash = GetTypeHash(InActor->GetRootComponent());
-	return ComponentHash;
+	return 0;
 }
 	
 }
@@ -1481,11 +1473,10 @@ bool AreHandlesConstrainable( UWorld* InWorld, UTransformableHandle* InParentHan
 
 }
 
-UTickableTransformConstraint* FTransformConstraintUtils::CreateAndAddFromActors(
+UTickableTransformConstraint* FTransformConstraintUtils::CreateAndAddFromObjects(
 	UWorld* InWorld,
-	AActor* InParent,
-	const FName& InSocketName,
-	AActor* InChild,
+	UObject* InParent, const FName& InParentSocketName,
+	UObject* InChild, const FName& InChildSocketName,
 	const ETransformConstraintType InType,
 	const bool bMaintainOffset)
 {
@@ -1505,13 +1496,13 @@ UTickableTransformConstraint* FTransformConstraintUtils::CreateAndAddFromActors(
 		return nullptr;
 	}
 
-	UTransformableHandle* ParentHandle = GetHandle(InParent, InSocketName, ConstraintsManager);
+	UTransformableHandle* ParentHandle = GetHandle(InParent, InParentSocketName, ConstraintsManager);
 	if (!ParentHandle)
 	{
 		return nullptr;
 	}
 	
-	UTransformableHandle* ChildHandle = GetHandle(InChild, NAME_None, ConstraintsManager);
+	UTransformableHandle* ChildHandle = GetHandle(InChild, InChildSocketName, ConstraintsManager);
 	if (!ChildHandle)
 	{
 		return nullptr;
diff --git a/Engine/Source/Runtime/Experimental/Animation/Constraints/Private/TransformableHandle.cpp b/Engine/Source/Runtime/Experimental/Animation/Constraints/Private/TransformableHandle.cpp
--- a/Engine/Source/Runtime/Experimental/Animation/Constraints/Private/TransformableHandle.cpp
+++ b/Engine/Source/Runtime/Experimental/Animation/Constraints/Private/TransformableHandle.cpp
@@ -268,6 +268,8 @@ void UTransformableComponentHandle::UnregisterDelegates() const
 	{
 		GEngine->OnActorMoving().RemoveAll(this);
 	}
+
+	FCoreUObjectDelegates::OnObjectsReplaced.RemoveAll(this);
 #endif
 }
 
@@ -276,11 +278,16 @@ void UTransformableComponentHandle::RegisterDelegates()
 	UnregisterDelegates();
 
 #if WITH_EDITOR
-	GEngine->OnActorMoving().AddUObject(this, &UTransformableComponentHandle::OnActorMoving);
+	if (GEngine)
+	{
+		GEngine->OnActorMoving().AddUObject(this, &UTransformableComponentHandle::OnActorMoving);
+	}
 	FCoreUObjectDelegates::OnObjectPropertyChanged.AddUObject(this, &UTransformableComponentHandle::OnPostPropertyChanged);
+	FCoreUObjectDelegates::OnObjectsReplaced.AddUObject(this, &UTransformableComponentHandle::OnObjectsReplaced);
 #endif
 }
 
+#if WITH_EDITOR
 void UTransformableComponentHandle::OnActorMoving(AActor* InActor)
 {
 	if (!Component.IsValid())
@@ -288,8 +295,8 @@ void UTransformableComponentHandle::OnActorMoving(AActor* InActor)
 		return;
 	}
 
-	const USceneComponent* SceneComponent = InActor ? InActor->GetRootComponent() : nullptr;
-	if (SceneComponent != Component)
+	const TInlineComponentArray<USceneComponent*> Components(InActor);
+	if (!Components.Contains(Component))
 	{
 		return;
 	}
@@ -314,11 +321,16 @@ void UTransformableComponentHandle::OnPostPropertyChanged(
 	{
 		if (const AActor* Actor = Cast<AActor>(InObject))
 		{
-			SceneComponent = Actor->GetRootComponent();
+			const TInlineComponentArray<USceneComponent*> Components(Actor);
+			const int32 Index = Components.IndexOfByKey(Component.Get());
+			if (Index != INDEX_NONE)
+			{
+				SceneComponent = Components[Index];
+			}
 		}
 	}
 	
-	if (SceneComponent!= Component)
+	if (SceneComponent != Component)
 	{
 		return;
 	}
@@ -345,6 +357,22 @@ void UTransformableComponentHandle::OnPostPropertyChanged(
 	}
 }
 
+void UTransformableComponentHandle::OnObjectsReplaced(const TMap<UObject*, UObject*>& InOldToNewInstances)
+{
+	// in the context of blueprints being recompiled (cf. AActor::RerunConstructionScripts()), the component has to
+	// be updated. as this is called after it has been destroyed, we get the component even if it's pending kill
+	// otherwise Get() will return a nullptr.
+	static constexpr bool bEvenIfPendingKill = true;
+	if (UObject* NewObject = InOldToNewInstances.FindRef(Component.Get(bEvenIfPendingKill)))
+	{
+		if (USceneComponent* NewSceneComponent= Cast<USceneComponent>(NewObject))
+		{
+			Component = NewSceneComponent;
+		}
+	}
+}
+#endif
+
 TArrayView<FMovieSceneFloatChannel*>  UTransformableComponentHandle::GetFloatChannels(const UMovieSceneSection* InSection) const
 {
 	// no floats for transform sections
@@ -498,7 +526,12 @@ void UTransformableComponentHandle::ResolveBoundObjects(FMovieSceneSequenceID Lo
 	{
 		if (AActor* Actor = Cast<AActor>(ParentObject.Get()))
 		{
-			Component = Actor->GetRootComponent();
+			const TInlineComponentArray<USceneComponent*> Components(Actor);
+			const int32 Index = Components.IndexOfByPredicate([this](const USceneComponent* SubComponent)
+			{
+				return Component.IsValid() && (SubComponent->GetFName() == Component->GetFName()); 
+			});
+			Component = Index != INDEX_NONE ? Components[Index] : Actor->GetRootComponent();
 		}
 		else if (USceneComponent* Comp = Cast<USceneComponent>(ParentObject.Get()))
 		{
diff --git a/Engine/Source/Runtime/Experimental/Animation/Constraints/Private/TransformableRegistry.cpp b/Engine/Source/Runtime/Experimental/Animation/Constraints/Private/TransformableRegistry.cpp
--- a/Engine/Source/Runtime/Experimental/Animation/Constraints/Private/TransformableRegistry.cpp
+++ b/Engine/Source/Runtime/Experimental/Animation/Constraints/Private/TransformableRegistry.cpp
@@ -2,6 +2,10 @@
 
 #include "TransformableRegistry.h"
 
+#include "TransformConstraint.h"
+#include "TransformableHandle.h"
+#include "GameFramework/Actor.h"
+#include "Components/SceneComponent.h"
 #include "UObject/Class.h"
 
 FTransformableRegistry::~FTransformableRegistry() = default;
@@ -47,10 +51,67 @@ FTransformableRegistry::CreateHandleFuncT FTransformableRegistry::GetCreateFunct
 
 const FTransformableRegistry::FTransformableInfo* FTransformableRegistry::FindInfo(const UClass* InClass) const
 {
+	if (!InClass)
+	{
+		return nullptr;
+	}
+	
+	// look for registered class
 	if (const FTransformableInfo* Info = Transformables.Find(InClass))
 	{
 		return Info;
 	}
 
-	return nullptr;
+	// if not directly registered, look for super class registration
+	return FindInfo(InClass->GetSuperClass());
+}
+
+void FTransformableRegistry::RegisterBaseObjects()
+{
+    // register USceneComponent and AActor
+    auto CreateComponentHandle = [](UObject* Outer, UObject* InObject, const FName& InSocketName)->UTransformableHandle*
+    {
+    	if (USceneComponent* Component = Cast<USceneComponent>(InObject))
+    	{
+    		return FTransformConstraintUtils::CreateHandleForSceneComponent(Component, InSocketName, Outer);
+    	}
+    	return nullptr;
+    };
+
+    auto GetComponentHash = [](const UObject* InObject, const FName&)->uint32
+    {
+    	if (const USceneComponent* Component = Cast<USceneComponent>(InObject))
+    	{
+    		return GetTypeHash(Component);
+    	}
+    	return 0;
+    };
+    
+    auto CreateComponentHandleFromActor = [CreateComponentHandle](UObject* Outer, UObject* InObject, const FName& InSocketName)->UTransformableHandle*
+    {
+    	if (const AActor* Actor = Cast<AActor>(InObject))
+    	{
+    		return CreateComponentHandle(Outer, Actor->GetRootComponent(), InSocketName);
+    	}
+    	return nullptr;
+    };
+
+    auto GetComponentHashFromActor = [GetComponentHash](const UObject* InObject, const FName& InSocketName)->uint32
+    {
+    	if (const AActor* Actor = Cast<AActor>(InObject))
+    	{
+    		return GetComponentHash(Actor->GetRootComponent(), InSocketName);
+    	}
+    	return 0;
+    };
+
+	FTransformableRegistry& Registry = Get();
+    Registry.Register(USceneComponent::StaticClass(), CreateComponentHandle, GetComponentHash);
+    Registry.Register(AActor::StaticClass(), CreateComponentHandleFromActor, GetComponentHashFromActor);
+}
+
+void FTransformableRegistry::UnregisterAllObjects()
+{
+	FTransformableRegistry& Registry = Get();
+	Registry.Transformables.Reset();
 }
\ No newline at end of file
diff --git a/Engine/Source/Runtime/Experimental/Animation/Constraints/Public/ConstraintsScripting.h b/Engine/Source/Runtime/Experimental/Animation/Constraints/Public/ConstraintsScripting.h
--- a/Engine/Source/Runtime/Experimental/Animation/Constraints/Public/ConstraintsScripting.h
+++ b/Engine/Source/Runtime/Experimental/Animation/Constraints/Public/ConstraintsScripting.h
@@ -9,6 +9,7 @@
 #include "ConstraintsScripting.generated.h"
 
 class UConstraintsManager;
+class UTransformableHandle;
 class UTransformableComponentHandle;
 class UTickableConstraint;
 class UTickableTransformConstraint;
@@ -45,6 +46,16 @@ public:
 	UFUNCTION(BlueprintCallable, Category = "Editor Scripting | Sequencer Tools | Control Rig | Constraints")
 	static UTransformableComponentHandle* CreateTransformableComponentHandle(UWorld* InWorld, USceneComponent* InSceneComponent, const FName& InSocketName);
 
+	/**
+	* Create the transformable handle that deals with getting and setting transforms on this object
+	* @param InWorld, the world you are in
+	* @param InObject World to create the constraint
+	* @param InAttachmentName Optional name of the attachment to get the transform. Not that this can represent a scene component's socket name or a control rig control for example. 
+	* @return returns the handle for this scene component
+	**/
+	UFUNCTION(BlueprintCallable, Category = "Editor Scripting | Sequencer Tools | Control Rig | Constraints")
+	static UTransformableHandle* CreateTransformableHandle(UWorld* InWorld, UObject* InObject, const FName& InAttachmentName = NAME_None);
+
 	/**
 	* Create Constraint based on the specified type.
 	* @param InWorld World to create the constraint
diff --git a/Engine/Source/Runtime/Experimental/Animation/Constraints/Public/TransformConstraint.h b/Engine/Source/Runtime/Experimental/Animation/Constraints/Public/TransformConstraint.h
--- a/Engine/Source/Runtime/Experimental/Animation/Constraints/Public/TransformConstraint.h
+++ b/Engine/Source/Runtime/Experimental/Animation/Constraints/Public/TransformConstraint.h
@@ -408,11 +408,10 @@ struct CONSTRAINTS_API FTransformConstraintUtils
 		const ETransformConstraintType InType);
 
 	/** Creates respective handles and creates a new InType transform constraint. */	
-	static UTickableTransformConstraint* CreateAndAddFromActors(
+	static UTickableTransformConstraint* CreateAndAddFromObjects(
 		UWorld* InWorld,
-		AActor* InParent,
-		const FName& InSocketName,
-		AActor* InChild,
+		UObject* InParent, const FName& InParentSocketName,
+		UObject* InChild, const FName& InChildSocketName,
 		const ETransformConstraintType InType,
 		const bool bMaintainOffset = true);
 
diff --git a/Engine/Source/Runtime/Experimental/Animation/Constraints/Public/TransformableHandle.h b/Engine/Source/Runtime/Experimental/Animation/Constraints/Public/TransformableHandle.h
--- a/Engine/Source/Runtime/Experimental/Animation/Constraints/Public/TransformableHandle.h
+++ b/Engine/Source/Runtime/Experimental/Animation/Constraints/Public/TransformableHandle.h
@@ -205,7 +205,9 @@ public:
 	void UnregisterDelegates() const;
 	void RegisterDelegates();
 	
-	/** @todo document */
+#if WITH_EDITOR
 	void OnActorMoving(AActor* InActor);
 	void OnPostPropertyChanged(UObject* InObject, FPropertyChangedEvent& InPropertyChangedEvent);
+	void OnObjectsReplaced(const TMap<UObject*, UObject*>& InOldToNewInstances);
+#endif
 };
\ No newline at end of file
diff --git a/Engine/Source/Runtime/Experimental/Animation/Constraints/Public/TransformableRegistry.h b/Engine/Source/Runtime/Experimental/Animation/Constraints/Public/TransformableRegistry.h
--- a/Engine/Source/Runtime/Experimental/Animation/Constraints/Public/TransformableRegistry.h
+++ b/Engine/Source/Runtime/Experimental/Animation/Constraints/Public/TransformableRegistry.h
@@ -18,36 +18,46 @@ class UTransformableHandle;
 class CONSTRAINTS_API FTransformableRegistry
 {
 public:
-	/** @todo document */
-	using CreateHandleFuncT = TFunction<UTransformableHandle*(const UObject*, UObject*)>;
-	using GetHashFuncT = TFunction<uint32(const UObject*)>;
+
+	/** Defines the function that will allocate a new transformable handle to wrap the object. */
+	using CreateHandleFuncT = TFunction<UTransformableHandle*(UObject*, UObject*, const FName&)>;
+	/** Defines the function that return a hash value from the object and the attachment name. */
+	using GetHashFuncT = TFunction<uint32(const UObject*, const FName&)>;
 	
 	~FTransformableRegistry();
 
-	/** @todo document */
+	/** Get registry singleton */
 	static FTransformableRegistry& Get();
 
-	/** @todo document */
+	/** Registers a InClass as transformable object. */
 	void Register(UClass* InClass, CreateHandleFuncT InHandleFunc, GetHashFuncT InHashFunc);
 
-	/** @todo document */
+	/** Get the associated hash function for that specific class. */
 	GetHashFuncT GetHashFunction(const UClass* InClass) const;
-	/** @todo document */
+	/** Get the associated handle creation function for that specific class. */
 	CreateHandleFuncT GetCreateFunction(const UClass* InClass) const;
 
+protected:
+	friend class FConstraintsModule;
+
+	/** Registers the basic transformable objects. (called when starting the constraints module). */
+	static void RegisterBaseObjects();
+	/** Registers the basic transformable objects. (called when shutting down the constraints module). */
+	static void UnregisterAllObjects();
+	
 private:
 	FTransformableRegistry() = default;
 
-	/** @todo document */
+	/** Per class information that need to be register. */
 	struct FTransformableInfo
 	{
 		CreateHandleFuncT CreateHandleFunc;
 		GetHashFuncT GetHashFunc;
 	};
 
-	/** @todo document */
+	/** Find the associated info for that specific class. */
 	const FTransformableInfo* FindInfo(const UClass* InClass) const;
 
-	/** @todo document */
+	/** List of all registered transformable objects. */
 	TMap<UClass*, FTransformableInfo> Transformables;
 };
