Control Rig: Promote pin to interface

diff --git a/Engine/Plugins/Animation/ControlRig/Source/ControlRigDeveloper/Private/Graph/ControlRigGraphSchema.cpp b/Engine/Plugins/Animation/ControlRig/Source/ControlRigDeveloper/Private/Graph/ControlRigGraphSchema.cpp
--- a/Engine/Plugins/Animation/ControlRig/Source/ControlRigDeveloper/Private/Graph/ControlRigGraphSchema.cpp
+++ b/Engine/Plugins/Animation/ControlRig/Source/ControlRigDeveloper/Private/Graph/ControlRigGraphSchema.cpp
@@ -390,6 +390,86 @@ UEdGraphNode* FControlRigGraphSchemaAction_PromoteToVariable::PerformAction(UEdG
 	return nullptr;
 }
 
+FControlRigGraphSchemaAction_PromoteToExposedPin::FControlRigGraphSchemaAction_PromoteToExposedPin(UEdGraphPin* InEdGraphPin)
+: FEdGraphSchemaAction(	FText(), 
+						LOCTEXT("PromoteToExposedPin", "Promote to exposed pin"),
+						LOCTEXT("PromoteToExposedPin", "Promote to exposed pin"),
+						1)
+, EdGraphPin(InEdGraphPin)
+{
+}
+
+UEdGraphNode* FControlRigGraphSchemaAction_PromoteToExposedPin::PerformAction(UEdGraph* ParentGraph, UEdGraphPin* FromPin, const FVector2D Location, bool bSelectNewNode)
+{
+	UControlRigGraph* RigGraph = Cast<UControlRigGraph>(ParentGraph);
+	if(RigGraph == nullptr)
+	{
+		return nullptr;
+	}
+
+	URigVMGraph* Model = RigGraph->GetModel();
+	URigVMController* Controller = RigGraph->GetController();
+	if(Model == nullptr || Controller == nullptr)
+	{
+		return nullptr;
+	}
+	
+	URigVMPin* ModelPin = Model->FindPin(FromPin->GetName());
+	if (ModelPin == nullptr)
+	{
+		return nullptr;
+	}
+
+	const FScopedTransaction Transaction(
+		LOCTEXT("GraphEd_PromoteToExposedPin", "Promote To Exposed Pin"));
+	
+#if WITH_EDITOR
+	if (GEditor)
+	{
+		GEditor->CancelTransaction(0);
+	}	
+#endif
+
+	Controller->OpenUndoBracket(TEXT("Promote to Exposed Pin"));
+
+	const UObject* CPPTypeObject = ModelPin->GetCPPTypeObject();
+	const FString PinName = Controller->AddExposedPin(
+		*ModelPin->GetName(),
+		ModelPin->GetDirection(),
+		ModelPin->GetCPPType(),
+		CPPTypeObject ? (FName)*CPPTypeObject->GetPathName() : NAME_None,
+		ModelPin->GetDefaultValue(),
+		true,
+		true
+	).ToString();
+
+	UEdGraphNode* Result = nullptr;
+	if(!PinName.IsEmpty())
+	{
+		if (ModelPin->GetDirection() == ERigVMPinDirection::Input)
+		{
+			Controller->AddLink(Model->GetEntryNode()->FindPin(PinName), ModelPin, true);
+			Result = RigGraph->FindNodeForModelNodeName(Model->GetEntryNode()->GetFName());
+		}
+		else if(ModelPin->GetDirection() == ERigVMPinDirection::Output)
+		{
+			Controller->AddLink(ModelPin, Model->GetReturnNode()->FindPin(PinName), true);
+			Result = RigGraph->FindNodeForModelNodeName(Model->GetReturnNode()->GetFName());
+		}
+		else if (ModelPin->GetDirection() == ERigVMPinDirection::IO)
+		{
+			URigVMFunctionEntryNode* EntryNode = Model->GetEntryNode();
+			Controller->AddLink(EntryNode->FindPin(PinName), ModelPin, true);
+			Controller->AddLink(ModelPin, Model->GetReturnNode()->FindPin(PinName), true);
+			Result = RigGraph->FindNodeForModelNodeName(EntryNode->GetFName());
+		}		
+	}
+
+	Controller->CloseUndoBracket();
+	
+	return Result;
+}
+
 FControlRigGraphSchemaAction_Event::FControlRigGraphSchemaAction_Event(const FName& InEventName, const FString& InNodePath, const FText& InNodeCategory)
 : FEdGraphSchemaAction(	InNodeCategory, 
 						FText::FromName(InEventName),
@@ -862,6 +942,7 @@ void UControlRigGraphSchema::InsertAdditionalActions(TArray<UBlueprint*> InBluep
 		{
 			if(URigVMPin* ModelPin = RigNode->GetModelPinFromPinPath(EdGraphPins[0]->GetName()))
 			{
+				const bool bIsRootGraph = ModelPin->GetGraph()->IsRootGraph();
 				if(!ModelPin->IsExecuteContext() && !ModelPin->IsWildCard())
 				{
 					if(!ModelPin->GetNode()->IsA<URigVMVariableNode>())
@@ -870,7 +951,7 @@ void UControlRigGraphSchema::InsertAdditionalActions(TArray<UBlueprint*> InBluep
 							new FControlRigGraphSchemaAction_PromoteToVariable(EdGraphPins[0], false)
 						));
 
-						if(!ModelPin->GetGraph()->IsRootGraph())
+						if(!bIsRootGraph && !ModelPin->IsWildCard())
 						{
 							OutAllActions.AddAction(TSharedPtr<FControlRigGraphSchemaAction_PromoteToVariable>(
 								new FControlRigGraphSchemaAction_PromoteToVariable(EdGraphPins[0], true)
@@ -878,6 +959,16 @@ void UControlRigGraphSchema::InsertAdditionalActions(TArray<UBlueprint*> InBluep
 						}
 					}
 				}
+
+				if (!bIsRootGraph)
+				{
+					if (!ModelPin->GetGraph()->GetRootGraph()->IsA<URigVMFunctionLibrary>() || !ModelPin->IsWildCard())
+					{
+						OutAllActions.AddAction(TSharedPtr<FControlRigGraphSchemaAction_PromoteToExposedPin>(
+								   new FControlRigGraphSchemaAction_PromoteToExposedPin(EdGraphPins[0])
+							   ));
+					}
+				}
 			}
 		}
 	}
diff --git a/Engine/Plugins/Animation/ControlRig/Source/ControlRigDeveloper/Public/Graph/ControlRigGraphSchema.h b/Engine/Plugins/Animation/ControlRig/Source/ControlRigDeveloper/Public/Graph/ControlRigGraphSchema.h
--- a/Engine/Plugins/Animation/ControlRig/Source/ControlRigDeveloper/Public/Graph/ControlRigGraphSchema.h
+++ b/Engine/Plugins/Animation/ControlRig/Source/ControlRigDeveloper/Public/Graph/ControlRigGraphSchema.h
@@ -141,6 +141,35 @@ private:
 	bool bLocalVariable;
 };
 
+USTRUCT()
+struct CONTROLRIGDEVELOPER_API FControlRigGraphSchemaAction_PromoteToExposedPin : public FEdGraphSchemaAction
+{
+	GENERATED_BODY()
+
+public:
+
+	// Simple type info
+	static FName StaticGetTypeId() {static FName Type("FControlRigGraphSchemaAction_PromoteToExposedPin"); return Type;}
+	virtual FName GetTypeId() const override { return StaticGetTypeId(); } 
+
+	FControlRigGraphSchemaAction_PromoteToExposedPin()
+		: FEdGraphSchemaAction()
+	{}
+
+	FControlRigGraphSchemaAction_PromoteToExposedPin(UEdGraphPin* InEdGraphPin);
+
+	virtual bool IsA(const FName& InType) const override
+	{
+		return InType == GetTypeId();
+	}
+
+	virtual UEdGraphNode* PerformAction(UEdGraph* ParentGraph, UEdGraphPin* FromPin, const FVector2D Location, bool bSelectNewNode) override;
+
+private:
+
+	UEdGraphPin* EdGraphPin;
+};
+
 USTRUCT()
 struct CONTROLRIGDEVELOPER_API FControlRigGraphSchemaAction_Event : public FEdGraphSchemaAction
 {
