Sequencer: Control Rig: Use Mask when filtering controls.
#trivial

diff --git a/Engine/Plugins/Animation/ControlRig/Source/ControlRigEditor/Private/Sequencer/ControlRigParameterTrackEditor.cpp b/Engine/Plugins/Animation/ControlRig/Source/ControlRigEditor/Private/Sequencer/ControlRigParameterTrackEditor.cpp
--- a/Engine/Plugins/Animation/ControlRig/Source/ControlRigEditor/Private/Sequencer/ControlRigParameterTrackEditor.cpp
+++ b/Engine/Plugins/Animation/ControlRig/Source/ControlRigEditor/Private/Sequencer/ControlRigParameterTrackEditor.cpp
@@ -4620,15 +4620,17 @@ void FControlRigParameterSection::ShowSelectedControlsChannels()
 
 		TArray<FRigControlElement*> Controls;
 		ControlRig->GetControlsInOrder(Controls);
-		int32 Index = 0;
 		for (const FRigControlElement* RigControl : Controls)
 		{
 			const FName RigName = RigControl->GetName();
 			if (ControlRig->IsControlSelected(RigName))
 			{
-				ParameterSection->SetControlsMask(Index, true);
+				FChannelMapInfo* pChannelIndex = ParameterSection->ControlChannelMap.Find(RigName);
+				if (pChannelIndex)
+				{
+					ParameterSection->SetControlsMask(pChannelIndex->MaskIndex, true);
+				}
 			}
-			++Index;
 		}
 		SequencerPtr->NotifyMovieSceneDataChanged(EMovieSceneDataChangeType::MovieSceneStructureItemsChanged);
 	}
