Control Rig: ComponentIsTouchingSelectionFrustum doesn't work if there are 'zero' planes for the width,height so change our malform test to look for that.
#trivial

diff --git a/Engine/Plugins/Animation/ControlRig/Source/ControlRigEditor/Private/EditMode/ControlRigEditMode.cpp b/Engine/Plugins/Animation/ControlRig/Source/ControlRigEditor/Private/EditMode/ControlRigEditMode.cpp
--- a/Engine/Plugins/Animation/ControlRig/Source/ControlRigEditor/Private/EditMode/ControlRigEditMode.cpp
+++ b/Engine/Plugins/Animation/ControlRig/Source/ControlRigEditor/Private/EditMode/ControlRigEditMode.cpp
@@ -1676,8 +1676,8 @@ bool FControlRigEditMode::FrustumSelect(const FConvexVolume& InFrustum, FEditorV
 {
 	const UControlRigEditModeSettings* Settings = GetDefault<UControlRigEditModeSettings>();
 	//need to check for a zero frustum since ComponentIsTouchingSelectionFrustum will return true, selecting everything, when this is the case
-	const bool bMalformedFrustum = (InFrustum.Planes[0].IsNearlyZero() && InFrustum.Planes[1].IsNearlyZero() && InFrustum.Planes[2].IsNearlyZero() &&
-		InFrustum.Planes[3].IsNearlyZero());
+	const bool bMalformedFrustum = (InFrustum.Planes[0].IsNearlyZero() && InFrustum.Planes[2].IsNearlyZero()) || (InFrustum.Planes[3].IsNearlyZero() &&
+		InFrustum.Planes[4].IsNearlyZero());
 	if (bMalformedFrustum || InViewportClient->IsInGameView() == true || Settings->bHideControlShapes)
 	{
 		if (Settings->bOnlySelectRigControls)
