Control Rig: Catch nullptr in RigHierarchy Widget

diff --git a/Engine/Plugins/Animation/ControlRig/Source/ControlRigEditor/Private/Editor/SRigHierarchyTreeView.cpp b/Engine/Plugins/Animation/ControlRig/Source/ControlRigEditor/Private/Editor/SRigHierarchyTreeView.cpp
--- a/Engine/Plugins/Animation/ControlRig/Source/ControlRigEditor/Private/Editor/SRigHierarchyTreeView.cpp
+++ b/Engine/Plugins/Animation/ControlRig/Source/ControlRigEditor/Private/Editor/SRigHierarchyTreeView.cpp
@@ -571,8 +571,10 @@ bool SRigHierarchyTreeView::ReparentElement(FRigElementKey InKey, FRigElementKey
 		ParentMap.Add(InKey, InParentKey);
 
 		TSharedPtr<FRigTreeElement>* FoundParent = ElementMap.Find(InParentKey);
-		check(FoundParent);
-		FoundParent->Get()->Children.Add(*FoundItem);
+		if(FoundParent)
+		{
+			FoundParent->Get()->Children.Add(*FoundItem);
+		}
 	}
 	else
 	{
