Control Rig: Fix comment bubble not visible when opening editor

diff --git a/Engine/Plugins/Animation/ControlRig/Source/ControlRigDeveloper/Private/Graph/ControlRigGraph.cpp b/Engine/Plugins/Animation/ControlRig/Source/ControlRigDeveloper/Private/Graph/ControlRigGraph.cpp
--- a/Engine/Plugins/Animation/ControlRig/Source/ControlRigDeveloper/Private/Graph/ControlRigGraph.cpp
+++ b/Engine/Plugins/Animation/ControlRig/Source/ControlRigDeveloper/Private/Graph/ControlRigGraph.cpp
@@ -484,6 +484,11 @@ void UControlRigGraph::HandleModifiedEvent(ERigVMGraphNotifType InNotifType, URi
 					NewNode->NodeHeight = ModelNode->GetSize().Y;
 					NewNode->CommentColor = ModelNode->GetNodeColor();
 					NewNode->NodeComment = CommentModelNode->GetCommentText();
+					NewNode->FontSize = CommentModelNode->GetCommentFontSize();
+					NewNode->bCommentBubbleVisible = CommentModelNode->GetCommentBubbleVisible();
+					NewNode->bCommentBubbleVisible_InDetailsPanel = CommentModelNode->GetCommentBubbleVisible();
+					NewNode->bCommentBubblePinned = CommentModelNode->GetCommentBubbleVisible();
+					NewNode->bColorCommentBubble = CommentModelNode->GetCommentColorBubble();
 					NewNode->SetFlags(RF_Transactional);
 					NewNode->GetNodesUnderComment();
 
