Control Rig: Fix cannot move bones in construction mode

diff --git a/Engine/Plugins/Animation/ControlRig/Source/ControlRigEditor/Private/Editor/ControlRigEditor.cpp b/Engine/Plugins/Animation/ControlRig/Source/ControlRigEditor/Private/Editor/ControlRigEditor.cpp
--- a/Engine/Plugins/Animation/ControlRig/Source/ControlRigEditor/Private/Editor/ControlRigEditor.cpp
+++ b/Engine/Plugins/Animation/ControlRig/Source/ControlRigEditor/Private/Editor/ControlRigEditor.cpp
@@ -4266,11 +4266,7 @@ void FControlRigEditor::OnPinControlNameListComboBox(const TArray<TSharedPtr<FSt
 
 bool FControlRigEditor::IsConstructionModeEnabled() const
 {
-	if(ControlRig)
-	{
-		return ControlRig->IsConstructionModeEnabled();
-	}
-	return false;
+	return GetEventQueue() == ConstructionEventQueue;
 }
 
 void FControlRigEditor::HandlePreviewSceneCreated(const TSharedRef<IPersonaPreviewScene>& InPersonaPreviewScene)
