Sequencer: When baking to control rig, if a spawnable could not be made, continue with the possessable actor

diff --git a/Engine/Plugins/Animation/ControlRig/Source/ControlRigEditor/Private/ControlRigEditorModule.cpp b/Engine/Plugins/Animation/ControlRig/Source/ControlRigEditor/Private/ControlRigEditorModule.cpp
--- a/Engine/Plugins/Animation/ControlRig/Source/ControlRigEditor/Private/ControlRigEditorModule.cpp
+++ b/Engine/Plugins/Animation/ControlRig/Source/ControlRigEditor/Private/ControlRigEditorModule.cpp
@@ -690,19 +690,26 @@ void FControlRigEditorModule::BakeToControlRig(UClass* ControlRigClass, UAnimSeq
 			TArray<FGuid> ActorTracks = WeakSequencer.Pin()->AddActors(ActorsToAdd, false);
 			FGuid ActorTrackGuid = ActorTracks[0];
 
+			// By default, convert this to a spawnable and delete the existing actor. If for some reason, 
+			// the spawnable couldn't be generated, use the existing actor as a possessable (this could 
+			// eventually be an option)
 			TArray<FGuid> SpawnableGuids = WeakSequencer.Pin()->ConvertToSpawnable(ActorTrackGuid);
-			ActorTrackGuid = SpawnableGuids[0];
-			UObject* SpawnedMesh = WeakSequencer.Pin()->FindSpawnedObjectOrTemplate(ActorTrackGuid);
+			if (SpawnableGuids.Num())
+			{	
+				ActorTrackGuid = SpawnableGuids[0];
 
-			if (SpawnedMesh)
-			{
-				GCurrentLevelEditingViewportClient->GetWorld()->EditorDestroyActor(MeshActor, true);
-				MeshActor = Cast<ASkeletalMeshActor>(SpawnedMesh);
-				if (SkelMesh)
+				UObject* SpawnedMesh = WeakSequencer.Pin()->FindSpawnedObjectOrTemplate(ActorTrackGuid);
+
+				if (SpawnedMesh)
 				{
-					MeshActor->GetSkeletalMeshComponent()->SetSkeletalMesh(SkelMesh);
+					GCurrentLevelEditingViewportClient->GetWorld()->EditorDestroyActor(MeshActor, true);
+					MeshActor = Cast<ASkeletalMeshActor>(SpawnedMesh);
+					if (SkelMesh)
+					{
+						MeshActor->GetSkeletalMeshComponent()->SetSkeletalMesh(SkelMesh);
+					}
+					MeshActor->RegisterAllComponents();
 				}
-				MeshActor->RegisterAllComponents();
 			}
 
 			//Delete binding from default animating rig
