Control Rig: Move world and other members to FRigVMExecuteContext

diff --git a/Engine/Plugins/Animation/ControlRig/Source/ControlRig/Private/ControlRig.cpp b/Engine/Plugins/Animation/ControlRig/Source/ControlRig/Private/ControlRig.cpp
--- a/Engine/Plugins/Animation/ControlRig/Source/ControlRig/Private/ControlRig.cpp
+++ b/Engine/Plugins/Animation/ControlRig/Source/ControlRig/Private/ControlRig.cpp
@@ -391,67 +391,34 @@ bool UControlRig::Execute(const FName& InEventName)
 	// disable any controller access outside of the construction event
 	FRigHierarchyEnableControllerBracket DisableHierarchyController(PublicContext.Hierarchy, bIsConstructionEvent);
 
-	// setup the context with further fields
-	Context.ToWorldSpaceTransform = FTransform::Identity;
-	Context.OwningComponent = nullptr;
-	Context.OwningActor = nullptr;
-	Context.World = nullptr;
-	
-	if (!OuterSceneComponent.IsValid())
-	{
-		USceneComponent* SceneComponentFromRegistry = Context.DataSourceRegistry->RequestSource<USceneComponent>(UControlRig::OwnerComponent);
-		if (SceneComponentFromRegistry)
-		{
-			OuterSceneComponent = SceneComponentFromRegistry;
-		}
-		else
-		{
-			UObject* Parent = this;
-			while (Parent)
-			{
-				Parent = Parent->GetOuter();
-				if (Parent)
-				{
-					if (USceneComponent* SceneComponent = Cast<USceneComponent>(Parent))
-					{
-						OuterSceneComponent = SceneComponent;
-						break;
-					}
-				}
-			}
-		}
-	}
-
 	// given the outer scene component configure
 	// the transform lookups to map transforms from rig space to world space
-	if (OuterSceneComponent.IsValid())
+	if (USceneComponent* SceneComponent = GetOwningSceneComponent())
 	{
-		Context.ToWorldSpaceTransform = OuterSceneComponent->GetComponentToWorld();
-		Context.OwningComponent = OuterSceneComponent.Get();
-		Context.OwningActor = Context.OwningComponent->GetOwner();
-		Context.World = Context.OwningComponent->GetWorld();
+		PublicContext.SetOwningComponent(SceneComponent);
 	}
 	else
 	{
+		PublicContext.SetOwningComponent(nullptr);
+		
 		if (ObjectBinding.IsValid())
 		{
 			AActor* HostingActor = ObjectBinding->GetHostingActor();
 			if (HostingActor)
 			{
-				Context.OwningActor = HostingActor;
-				Context.World = HostingActor->GetWorld();
+				PublicContext.SetOwningActor(HostingActor);
 			}
 			else if (UObject* Owner = ObjectBinding->GetBoundObject())
 			{
-				Context.World = Owner->GetWorld();
+				PublicContext.SetWorld(Owner->GetWorld());
 			}
 		}
 
-		if (Context.World == nullptr)
+		if (PublicContext.GetWorld() == nullptr)
 		{
 			if (UObject* Outer = GetOuter())
 			{
-				Context.World = Outer->GetWorld();
+				PublicContext.SetWorld(Outer->GetWorld());
 			}
 		}
 	}
@@ -459,7 +426,7 @@ bool UControlRig::Execute(const FName& InEventName)
 	// if we have any referenced elements dirty them
 	if(GetHierarchy())
 	{
-		GetHierarchy()->UpdateReferences(&Context);
+		GetHierarchy()->UpdateReferences(&PublicContext);
 	}
 
 	// guard against recursion
@@ -2536,8 +2503,31 @@ void UControlRig::DeclareConstructClasses(TArray<FTopLevelAssetPath>& OutConstru
 	OutConstructClasses.Add(FTopLevelAssetPath(URigHierarchy::StaticClass()));
 	OutConstructClasses.Add(FTopLevelAssetPath(UAnimationDataSourceRegistry::StaticClass()));
 }
+
 #endif
 
+USceneComponent* UControlRig::GetOwningSceneComponent()
+{
+	if(OuterSceneComponent == nullptr)
+	{
+		FRigVMExtendedExecuteContext& ExtendedExecuteContext = GetVM()->GetContext();
+		const FControlRigExecuteContext& PublicContext = ExtendedExecuteContext.GetPublicDataSafe<FControlRigExecuteContext>();
+		const FRigUnitContext& Context = PublicContext.UnitContext;
+
+		USceneComponent* SceneComponentFromRegistry = Context.DataSourceRegistry->RequestSource<USceneComponent>(UControlRig::OwnerComponent);
+		if (SceneComponentFromRegistry)
+		{
+			OuterSceneComponent = SceneComponentFromRegistry;
+		}
+
+		if(OuterSceneComponent == nullptr)
+		{
+			OuterSceneComponent = Super::GetOwningSceneComponent();
+		}
+	}
+	return OuterSceneComponent.Get();
+}
+
 void UControlRig::PostInitInstance(URigVMHost* InCDO)
 {
 	const EObjectFlags SubObjectFlags =
diff --git a/Engine/Plugins/Animation/ControlRig/Source/ControlRig/Private/Rigs/RigHierarchy.cpp b/Engine/Plugins/Animation/ControlRig/Source/ControlRig/Private/Rigs/RigHierarchy.cpp
--- a/Engine/Plugins/Animation/ControlRig/Source/ControlRig/Private/Rigs/RigHierarchy.cpp
+++ b/Engine/Plugins/Animation/ControlRig/Source/ControlRig/Private/Rigs/RigHierarchy.cpp
@@ -615,7 +615,7 @@ void URigHierarchy::CopyPose(URigHierarchy* InHierarchy, bool bCurrent, bool bIn
 	EnsureCacheValidity();
 }
 
-void URigHierarchy::UpdateReferences(const FRigUnitContext* InContext)
+void URigHierarchy::UpdateReferences(const FRigVMExecuteContext* InContext)
 {
 	check(InContext);
 	LLM_SCOPE_BYNAME(TEXT("Animation/ControlRig"));
@@ -627,8 +627,8 @@ void URigHierarchy::UpdateReferences(const FRigUnitContext* InContext)
 			const FTransform InitialWorldTransform = Reference->GetReferenceWorldTransform(InContext, true);
 			const FTransform CurrentWorldTransform = Reference->GetReferenceWorldTransform(InContext, false);
 
-			const FTransform InitialGlobalTransform = InitialWorldTransform.GetRelativeTransform(InContext->ToWorldSpaceTransform);
-			const FTransform CurrentGlobalTransform = CurrentWorldTransform.GetRelativeTransform(InContext->ToWorldSpaceTransform);
+			const FTransform InitialGlobalTransform = InitialWorldTransform.GetRelativeTransform(InContext->GetToWorldSpaceTransform());
+			const FTransform CurrentGlobalTransform = CurrentWorldTransform.GetRelativeTransform(InContext->GetToWorldSpaceTransform());
 
 			const FTransform InitialParentTransform = GetParentTransform(Reference, ERigTransformType::InitialGlobal); 
 			const FTransform CurrentParentTransform = GetParentTransform(Reference, ERigTransformType::CurrentGlobal);
@@ -5246,7 +5246,7 @@ bool URigHierarchy::ShouldBeGrouped(const FRigControlElement* InControlElement)
 	return false;
 }
 
-FTransform URigHierarchy::GetWorldTransformForReference(const FRigUnitContext* InContext, const FRigElementKey& InKey, bool bInitial)
+FTransform URigHierarchy::GetWorldTransformForReference(const FRigVMExecuteContext* InContext, const FRigElementKey& InKey, bool bInitial)
 {
 	if(const USceneComponent* OuterSceneComponent = GetTypedOuter<USceneComponent>())
 	{
diff --git a/Engine/Plugins/Animation/ControlRig/Source/ControlRig/Private/Rigs/RigHierarchyElements.cpp b/Engine/Plugins/Animation/ControlRig/Source/ControlRig/Private/Rigs/RigHierarchyElements.cpp
--- a/Engine/Plugins/Animation/ControlRig/Source/ControlRig/Private/Rigs/RigHierarchyElements.cpp
+++ b/Engine/Plugins/Animation/ControlRig/Source/ControlRig/Private/Rigs/RigHierarchyElements.cpp
@@ -3,7 +3,7 @@
 #include "Rigs/RigHierarchyElements.h"
 #include "Rigs/RigHierarchy.h"
 #include "Rigs/RigControlHierarchy.h"
-#include "Units/RigUnitContext.h"
+#include "RigVMCore/RigVMExecuteContext.h"
 #include "ControlRigObjectVersion.h"
 #include "ControlRigGizmoLibrary.h"
 #include "AnimationCoreLibrary.h"
@@ -1331,7 +1331,7 @@ void FRigReferenceElement::CopyFrom(URigHierarchy* InHierarchy, FRigBaseElement*
 	GetWorldTransformDelegate = Source->GetWorldTransformDelegate;
 }
 
-FTransform FRigReferenceElement::GetReferenceWorldTransform(const FRigUnitContext* InContext, bool bInitial) const
+FTransform FRigReferenceElement::GetReferenceWorldTransform(const FRigVMExecuteContext* InContext, bool bInitial) const
 {
 	if(GetWorldTransformDelegate.IsBound())
 	{
diff --git a/Engine/Plugins/Animation/ControlRig/Source/ControlRig/Private/Units/Animation/RigUnit_AnimAttribute.cpp b/Engine/Plugins/Animation/ControlRig/Source/ControlRig/Private/Units/Animation/RigUnit_AnimAttribute.cpp
--- a/Engine/Plugins/Animation/ControlRig/Source/ControlRig/Private/Units/Animation/RigUnit_AnimAttribute.cpp
+++ b/Engine/Plugins/Animation/ControlRig/Source/ControlRig/Private/Units/Animation/RigUnit_AnimAttribute.cpp
@@ -91,7 +91,7 @@ const TArray<FRigVMTemplateArgument::ETypeCategory>& FRigDispatch_AnimAttributeB
 static uint8* GetAnimAttributeValue(
 	bool bAddIfNotFound,
 	UScriptStruct* InAttributeScriptStruct,
-	const FRigUnitContext& Context,
+	const FControlRigExecuteContext& Context,
 	const FName& Name,
 	const FName& BoneName,
 	FName& CachedBoneName,
@@ -107,12 +107,12 @@ static uint8* GetAnimAttributeValue(
 		return nullptr;
 	}
 
-	if (!Context.AnimAttributeContainer)
+	if (!Context.UnitContext.AnimAttributeContainer)
 	{
 		return nullptr;
 	}
 	
-	const USkeletalMeshComponent* OwningComponent = Cast<USkeletalMeshComponent>(Context.OwningComponent);
+	const USkeletalMeshComponent* OwningComponent = Cast<USkeletalMeshComponent>(Context.GetOwningComponent());
 
 	if (!OwningComponent ||
 		!OwningComponent->GetSkeletalMeshAsset())
@@ -140,8 +140,8 @@ static uint8* GetAnimAttributeValue(
 	{
 		const UE::Anim::FAttributeId Id = {Name, FCompactPoseBoneIndex(CachedBoneIndex)} ;
 		uint8* Attribute = bAddIfNotFound ?
-			Context.AnimAttributeContainer->FindOrAdd(InAttributeScriptStruct, Id) :
-			Context.AnimAttributeContainer->Find(InAttributeScriptStruct, Id);
+			Context.UnitContext.AnimAttributeContainer->FindOrAdd(InAttributeScriptStruct, Id) :
+			Context.UnitContext.AnimAttributeContainer->Find(InAttributeScriptStruct, Id);
 		
 		if (Attribute)
 		{
@@ -180,7 +180,7 @@ void FRigDispatch_GetAnimAttribute::GetAnimAttributeDispatch(FRigVMExtendedExecu
 	{
 		UScriptStruct* ScriptStruct = StructProperty->Struct;
 		bool Registered = true;
-		const FRigUnitContext& Context = GetRigUnitContext(InContext);
+		const FControlRigExecuteContext& Context = InContext.GetPublicDataSafe<FControlRigExecuteContext>();
 
 #if WITH_EDITOR
 		{
@@ -242,7 +242,7 @@ void FRigDispatch_SetAnimAttribute::SetAnimAttributeDispatch(FRigVMExtendedExecu
 	{
 		UScriptStruct* ScriptStruct = StructProperty->Struct;
 		bool Registered = true;
-		const FRigUnitContext& Context = GetRigUnitContext(InContext);
+		const FControlRigExecuteContext& Context = InContext.GetPublicDataSafe<FControlRigExecuteContext>();
 
 #if WITH_EDITOR
 		{
diff --git a/Engine/Plugins/Animation/ControlRig/Source/ControlRig/Private/Units/Animation/RigUnit_AnimAttribute.h b/Engine/Plugins/Animation/ControlRig/Source/ControlRig/Private/Units/Animation/RigUnit_AnimAttribute.h
--- a/Engine/Plugins/Animation/ControlRig/Source/ControlRig/Private/Units/Animation/RigUnit_AnimAttribute.h
+++ b/Engine/Plugins/Animation/ControlRig/Source/ControlRig/Private/Units/Animation/RigUnit_AnimAttribute.h
@@ -59,7 +59,7 @@ namespace RigUnit_AnimAttribute
 	template<typename T>
 	T* GetAnimAttributeValue(
 		bool bAddIfNotFound,
-		const FRigUnitContext& Context,
+		const FControlRigExecuteContext& Context,
 		const FName& Name,
 		const FName& BoneName,
 		FName& CachedBoneName,
@@ -70,12 +70,12 @@ namespace RigUnit_AnimAttribute
 			return nullptr;
 		}
 
-		if (!Context.AnimAttributeContainer)
+		if (!Context.UnitContext.AnimAttributeContainer)
 		{
 			return nullptr;
 		}
 	
-		const USkeletalMeshComponent* OwningComponent = Cast<USkeletalMeshComponent>(Context.OwningComponent);
+		const USkeletalMeshComponent* OwningComponent = Cast<USkeletalMeshComponent>(Context.GetOwningComponent());
 
 		if (!OwningComponent ||
 			!OwningComponent->GetSkeletalMeshAsset())
@@ -103,8 +103,8 @@ namespace RigUnit_AnimAttribute
 		{
 			const UE::Anim::FAttributeId Id = {Name, FCompactPoseBoneIndex(CachedBoneIndex)} ;
 			typename TAnimAttributeType<T>::Type* Attribute = bAddIfNotFound ?
-				Context.AnimAttributeContainer->FindOrAdd<typename TAnimAttributeType<T>::Type>(Id) :
-				Context.AnimAttributeContainer->Find<typename TAnimAttributeType<T>::Type>(Id);
+				Context.UnitContext.AnimAttributeContainer->FindOrAdd<typename TAnimAttributeType<T>::Type>(Id) :
+				Context.UnitContext.AnimAttributeContainer->Find<typename TAnimAttributeType<T>::Type>(Id);
 			if (Attribute)
 			{
 				return &Attribute->Value;
@@ -237,7 +237,7 @@ protected:
 		int32& CachedBoneIndex = *(int32*)Handles[Factory->CachedBoneIndexArgIndex].GetData(false, InContext.GetSlice().GetIndex());
 
 		// extract the animation attribute
-		const FRigUnitContext& Context = GetRigUnitContext(InContext);
+		const FControlRigExecuteContext& Context = InContext.GetPublicDataSafe<FControlRigExecuteContext>();
 		const ValueType* ValuePtr = RigUnit_AnimAttribute::GetAnimAttributeValue<ValueType>(false, Context, Name, BoneName, CachedBoneName, CachedBoneIndex);
 		Found = ValuePtr ? true : false;
 		Value = ValuePtr ? *ValuePtr : Default;
@@ -310,7 +310,7 @@ protected:
 		int32& CachedBoneIndex = *(int32*)Handles[Factory->CachedBoneIndexArgIndex].GetData(false, InContext.GetSlice().GetIndex());
 
 		// extract the animation attribute
-		const FRigUnitContext& Context = GetRigUnitContext(InContext);
+		const FControlRigExecuteContext& Context = InContext.GetPublicDataSafe<FControlRigExecuteContext>();
 		ValueType* ValuePtr = RigUnit_AnimAttribute::GetAnimAttributeValue<ValueType>(true, Context, Name, BoneName, CachedBoneName, CachedBoneIndex);
 		if (ValuePtr)
 		{
diff --git a/Engine/Plugins/Animation/ControlRig/Source/ControlRig/Private/Units/Collision/RigUnit_WorldCollision.cpp b/Engine/Plugins/Animation/ControlRig/Source/ControlRig/Private/Units/Collision/RigUnit_WorldCollision.cpp
--- a/Engine/Plugins/Animation/ControlRig/Source/ControlRig/Private/Units/Collision/RigUnit_WorldCollision.cpp
+++ b/Engine/Plugins/Animation/ControlRig/Source/ControlRig/Private/Units/Collision/RigUnit_WorldCollision.cpp
@@ -16,7 +16,7 @@ FRigUnit_SphereTraceWorld_Execute()
 	HitLocation = FVector::ZeroVector;
 	HitNormal = FVector(0.f, 0.f, 1.f);
 
-	if(ExecuteContext.UnitContext.World == nullptr)
+	if(ExecuteContext.GetWorld() == nullptr)
 	{
 		return;
 	}
@@ -24,11 +24,11 @@ FRigUnit_SphereTraceWorld_Execute()
 	FCollisionQueryParams QueryParams;
 	QueryParams.bTraceComplex = true;
 
-	if (ExecuteContext.UnitContext.OwningActor)
+	if (ExecuteContext.GetOwningActor())
 	{
-		QueryParams.AddIgnoredActor(ExecuteContext.UnitContext.OwningActor);
+		QueryParams.AddIgnoredActor(ExecuteContext.GetOwningActor());
 	}
-	else if (UPrimitiveComponent* PrimitiveComponent = Cast<UPrimitiveComponent>(ExecuteContext.UnitContext.OwningComponent))
+	else if (const UPrimitiveComponent* PrimitiveComponent = Cast<UPrimitiveComponent>(ExecuteContext.GetOwningComponent()))
 	{
 		QueryParams.AddIgnoredComponent(PrimitiveComponent);
 	}
@@ -39,13 +39,13 @@ FRigUnit_SphereTraceWorld_Execute()
 	const FCollisionShape CollisionShape = FCollisionShape::MakeSphere(Radius);
 
 	FHitResult HitResult;
-	bHit = ExecuteContext.UnitContext.World->SweepSingleByChannel(HitResult, ExecuteContext.UnitContext.ToWorldSpace(Start), ExecuteContext.UnitContext.ToWorldSpace(End), 
+	bHit = ExecuteContext.GetWorld()->SweepSingleByChannel(HitResult, ExecuteContext.ToWorldSpace(Start), ExecuteContext.ToWorldSpace(End), 
 			FQuat::Identity, Channel, CollisionShape, QueryParams, ResponseParams);
 
 	if (bHit)
 	{
-		HitLocation = ExecuteContext.UnitContext.ToRigSpace(HitResult.ImpactPoint);
-		HitNormal = ExecuteContext.UnitContext.ToWorldSpaceTransform.InverseTransformVector(HitResult.ImpactNormal);
+		HitLocation = ExecuteContext.ToVMSpace(HitResult.ImpactPoint);
+		HitNormal = ExecuteContext.GetToWorldSpaceTransform().InverseTransformVector(HitResult.ImpactNormal);
 	}
 }
 
@@ -67,7 +67,7 @@ FRigUnit_SphereTraceByTraceChannel_Execute()
 	HitLocation = FVector::ZeroVector;
 	HitNormal = FVector(0.f, 0.f, 1.f);
 
-	if (ExecuteContext.UnitContext.World == nullptr)
+	if (ExecuteContext.GetWorld() == nullptr)
 	{
 		return;
 	}
@@ -78,22 +78,22 @@ FRigUnit_SphereTraceByTraceChannel_Execute()
 	FCollisionQueryParams QueryParams;
 	QueryParams.bTraceComplex = true;
 
-	if (ExecuteContext.UnitContext.OwningActor)
+	if (ExecuteContext.GetOwningActor())
 	{
-		QueryParams.AddIgnoredActor(ExecuteContext.UnitContext.OwningActor);
+		QueryParams.AddIgnoredActor(ExecuteContext.GetOwningActor());
 	}
-	else if (UPrimitiveComponent* PrimitiveComponent = Cast<UPrimitiveComponent>(ExecuteContext.UnitContext.OwningComponent))
+	else if (const UPrimitiveComponent* PrimitiveComponent = Cast<UPrimitiveComponent>(ExecuteContext.GetOwningComponent()))
 	{
 		QueryParams.AddIgnoredComponent(PrimitiveComponent);
 	}
 
 	FHitResult HitResult;
-	bHit = ExecuteContext.UnitContext.World->SweepSingleByChannel(HitResult, ExecuteContext.UnitContext.ToWorldSpace(Start), ExecuteContext.UnitContext.ToWorldSpace(End), FQuat::Identity, CollisionChannel, CollisionShape, QueryParams);
+	bHit = ExecuteContext.GetWorld()->SweepSingleByChannel(HitResult, ExecuteContext.ToWorldSpace(Start), ExecuteContext.ToWorldSpace(End), FQuat::Identity, CollisionChannel, CollisionShape, QueryParams);
 	
 	if (bHit)
 	{
-		HitLocation = ExecuteContext.UnitContext.ToRigSpace(HitResult.ImpactPoint);
-		HitNormal = ExecuteContext.UnitContext.ToWorldSpaceTransform.InverseTransformVector(HitResult.ImpactNormal);
+		HitLocation = ExecuteContext.ToVMSpace(HitResult.ImpactPoint);
+		HitNormal = ExecuteContext.GetToWorldSpaceTransform().InverseTransformVector(HitResult.ImpactNormal);
 	}
 }
 
@@ -105,7 +105,7 @@ FRigUnit_SphereTraceByObjectTypes_Execute()
 	HitLocation = FVector::ZeroVector;
 	HitNormal = FVector(0.f, 0.f, 1.f);
 
-	if (ExecuteContext.UnitContext.World == nullptr)
+	if (ExecuteContext.GetWorld() == nullptr)
 	{
 		return;
 	}
@@ -133,22 +133,22 @@ FRigUnit_SphereTraceByObjectTypes_Execute()
 	FCollisionQueryParams QueryParams;
 	QueryParams.bTraceComplex = true;
 
-	if (ExecuteContext.UnitContext.OwningActor)
+	if (ExecuteContext.GetOwningActor())
 	{
-		QueryParams.AddIgnoredActor(ExecuteContext.UnitContext.OwningActor);
+		QueryParams.AddIgnoredActor(ExecuteContext.GetOwningActor());
 	}
-	else if (UPrimitiveComponent* PrimitiveComponent = Cast<UPrimitiveComponent>(ExecuteContext.UnitContext.OwningComponent))
+	else if (const UPrimitiveComponent* PrimitiveComponent = Cast<UPrimitiveComponent>(ExecuteContext.GetOwningComponent()))
 	{
 		QueryParams.AddIgnoredComponent(PrimitiveComponent);
 	} 
 
 	FHitResult HitResult;
-	bHit = ExecuteContext.UnitContext.World->SweepSingleByObjectType(HitResult, ExecuteContext.UnitContext.ToWorldSpace(Start), ExecuteContext.UnitContext.ToWorldSpace(End), FQuat::Identity, ObjectParams, CollisionShape, QueryParams);
+	bHit = ExecuteContext.GetWorld()->SweepSingleByObjectType(HitResult, ExecuteContext.ToWorldSpace(Start), ExecuteContext.ToWorldSpace(End), FQuat::Identity, ObjectParams, CollisionShape, QueryParams);
 
 	if (bHit)
 	{
-		HitLocation = ExecuteContext.UnitContext.ToRigSpace(HitResult.ImpactPoint);
-		HitNormal = ExecuteContext.UnitContext.ToWorldSpaceTransform.InverseTransformVector(HitResult.ImpactNormal);
+		HitLocation = ExecuteContext.ToVMSpace(HitResult.ImpactPoint);
+		HitNormal = ExecuteContext.GetToWorldSpaceTransform().InverseTransformVector(HitResult.ImpactNormal);
 	}
 }
 
diff --git a/Engine/Plugins/Animation/ControlRig/Source/ControlRig/Private/Units/Hierarchy/RigUnit_WorldSpace.cpp b/Engine/Plugins/Animation/ControlRig/Source/ControlRig/Private/Units/Hierarchy/RigUnit_WorldSpace.cpp
--- a/Engine/Plugins/Animation/ControlRig/Source/ControlRig/Private/Units/Hierarchy/RigUnit_WorldSpace.cpp
+++ b/Engine/Plugins/Animation/ControlRig/Source/ControlRig/Private/Units/Hierarchy/RigUnit_WorldSpace.cpp
@@ -8,35 +8,35 @@
 FRigUnit_ToWorldSpace_Transform_Execute()
 {
     DECLARE_SCOPE_HIERARCHICAL_COUNTER_RIGUNIT()
-    World = ExecuteContext.UnitContext.ToWorldSpace(Value);
+    World = ExecuteContext.ToWorldSpace(Value);
 }
 
 FRigUnit_ToRigSpace_Transform_Execute()
 {
     DECLARE_SCOPE_HIERARCHICAL_COUNTER_RIGUNIT()
-    Global = ExecuteContext.UnitContext.ToRigSpace(Value);
+    Global = ExecuteContext.ToVMSpace(Value);
 }
 
 FRigUnit_ToWorldSpace_Location_Execute()
 {
     DECLARE_SCOPE_HIERARCHICAL_COUNTER_RIGUNIT()
-    World = ExecuteContext.UnitContext.ToWorldSpace(Value);
+    World = ExecuteContext.ToWorldSpace(Value);
 }
 
 FRigUnit_ToRigSpace_Location_Execute()
 {
     DECLARE_SCOPE_HIERARCHICAL_COUNTER_RIGUNIT()
-    Global = ExecuteContext.UnitContext.ToRigSpace(Value);
+    Global = ExecuteContext.ToVMSpace(Value);
 }
 
 FRigUnit_ToWorldSpace_Rotation_Execute()
 {
     DECLARE_SCOPE_HIERARCHICAL_COUNTER_RIGUNIT()
-    World = ExecuteContext.UnitContext.ToWorldSpace(Value);
+    World = ExecuteContext.ToWorldSpace(Value);
 }
 
 FRigUnit_ToRigSpace_Rotation_Execute()
 {
     DECLARE_SCOPE_HIERARCHICAL_COUNTER_RIGUNIT()
-    Global = ExecuteContext.UnitContext.ToRigSpace(Value);
+    Global = ExecuteContext.ToVMSpace(Value);
 }
diff --git a/Engine/Plugins/Animation/ControlRig/Source/ControlRig/Public/ControlRig.h b/Engine/Plugins/Animation/ControlRig/Source/ControlRig/Public/ControlRig.h
--- a/Engine/Plugins/Animation/ControlRig/Source/ControlRig/Public/ControlRig.h
+++ b/Engine/Plugins/Animation/ControlRig/Source/ControlRig/Public/ControlRig.h
@@ -330,6 +330,8 @@ public:
 	static void DeclareConstructClasses(TArray<FTopLevelAssetPath>& OutConstructClasses, const UClass* SpecificSubclass);
 #endif
 
+	virtual USceneComponent* GetOwningSceneComponent() override;;
+
 protected:
 
 	virtual void PostInitInstance(URigVMHost* InCDO) override;
diff --git a/Engine/Plugins/Animation/ControlRig/Source/ControlRig/Public/Rigs/RigHierarchy.h b/Engine/Plugins/Animation/ControlRig/Source/ControlRig/Public/Rigs/RigHierarchy.h
--- a/Engine/Plugins/Animation/ControlRig/Source/ControlRig/Public/Rigs/RigHierarchy.h
+++ b/Engine/Plugins/Animation/ControlRig/Source/ControlRig/Public/Rigs/RigHierarchy.h
@@ -174,7 +174,7 @@ public:
 	/**
 	 * Update all elements that depend on external references
 	 */
-	void UpdateReferences(const FRigUnitContext* InContext);
+	void UpdateReferences(const FRigVMExecuteContext* InContext);
 
 	/**
 	 * Resets the current pose of a filtered list of elements to the initial / ref pose.
@@ -3868,7 +3868,7 @@ private:
 	
 	mutable TMap<uint32, FRigElementKeyCollection> KeyCollectionCache;
 
-	FTransform GetWorldTransformForReference(const FRigUnitContext* InContext, const FRigElementKey& InKey, bool bInitial);
+	FTransform GetWorldTransformForReference(const FRigVMExecuteContext* InContext, const FRigElementKey& InKey, bool bInitial);
 	
 	static float GetWeightForLerp(const float WeightA, const float WeightB)
 	{
diff --git a/Engine/Plugins/Animation/ControlRig/Source/ControlRig/Public/Rigs/RigHierarchyElements.h b/Engine/Plugins/Animation/ControlRig/Source/ControlRig/Public/Rigs/RigHierarchyElements.h
--- a/Engine/Plugins/Animation/ControlRig/Source/ControlRig/Public/Rigs/RigHierarchyElements.h
+++ b/Engine/Plugins/Animation/ControlRig/Source/ControlRig/Public/Rigs/RigHierarchyElements.h
@@ -7,15 +7,15 @@
 #include "RigHierarchyMetadata.h"
 #include "RigHierarchyElements.generated.h"
 
-struct FRigUnitContext;
+struct FRigVMExecuteContext;
 struct FRigBaseElement;
 struct FRigControlElement;
 class URigHierarchy;
 
-DECLARE_DELEGATE_RetVal_ThreeParams(FTransform, FRigReferenceGetWorldTransformDelegate, const FRigUnitContext*, const FRigElementKey& /* Key */, bool /* bInitial */);
+DECLARE_DELEGATE_RetVal_ThreeParams(FTransform, FRigReferenceGetWorldTransformDelegate, const FRigVMExecuteContext*, const FRigElementKey& /* Key */, bool /* bInitial */);
 DECLARE_DELEGATE_TwoParams(FRigElementMetadataChangedDelegate, const FRigElementKey& /* Key */, const FName& /* Name */);
 DECLARE_DELEGATE_ThreeParams(FRigElementMetadataTagChangedDelegate, const FRigElementKey& /* Key */, const FName& /* Tag */, bool /* AddedOrRemoved */);
-DECLARE_DELEGATE_RetVal_ThreeParams(FTransform, FRigReferenceGetWorldTransformDelegate, const FRigUnitContext*, const FRigElementKey& /* Key */, bool /* bInitial */);
+DECLARE_DELEGATE_RetVal_ThreeParams(FTransform, FRigReferenceGetWorldTransformDelegate, const FRigVMExecuteContext*, const FRigElementKey& /* Key */, bool /* bInitial */);
 
 #define DECLARE_RIG_ELEMENT_METHODS(ElementType) \
 template<typename T> \
@@ -1446,7 +1446,7 @@ public:
 	virtual void Save(FArchive& A, URigHierarchy* Hierarchy, ESerializationPhase SerializationPhase) override;
 	virtual void Load(FArchive& Ar, URigHierarchy* Hierarchy, ESerializationPhase SerializationPhase) override;
 
-	FTransform GetReferenceWorldTransform(const FRigUnitContext* InContext, bool bInitial) const;
+	FTransform GetReferenceWorldTransform(const FRigVMExecuteContext* InContext, bool bInitial) const;
 
 	virtual void CopyPose(FRigBaseElement* InOther, bool bCurrent, bool bInitial, bool bWeights) override;
 
diff --git a/Engine/Plugins/Animation/ControlRig/Source/ControlRig/Public/Units/RigUnitContext.h b/Engine/Plugins/Animation/ControlRig/Source/ControlRig/Public/Units/RigUnitContext.h
--- a/Engine/Plugins/Animation/ControlRig/Source/ControlRig/Public/Units/RigUnitContext.h
+++ b/Engine/Plugins/Animation/ControlRig/Source/ControlRig/Public/Units/RigUnitContext.h
@@ -7,8 +7,6 @@
 #include "Rigs/RigCurveContainer.h"
 #include "AnimationDataSource.h"
 #include "Animation/AttributesRuntime.h"
-#include "GameFramework/Actor.h"
-#include "Components/SceneComponent.h"
 #include "RigVMCore/RigVMExecuteContext.h"
 #include "RigUnitContext.generated.h"
 
@@ -49,10 +47,6 @@ struct FRigUnitContext
 		, DataSourceRegistry(nullptr)
 		, InteractionType((uint8)EControlRigInteractionType::None)
 		, ElementsBeingInteracted()
-		, ToWorldSpaceTransform(FTransform::Identity)
-		, OwningComponent(nullptr)
-		, OwningActor(nullptr)
-		, World(nullptr)
 	{
 	}
 
@@ -71,18 +65,6 @@ struct FRigUnitContext
 	/** The elements being interacted with. */
 	TArray<FRigElementKey> ElementsBeingInteracted;
 
-	/** The current transform going from rig (global) space to world space */
-	FTransform ToWorldSpaceTransform;
-
-	/** The current component this rig is owned by */
-	USceneComponent* OwningComponent;
-
-	/** The current actor this rig is owned by */
-	const AActor* OwningActor;
-
-	/** The world this rig is running in */
-	const UWorld* World;
-
 	/**
 	 * Returns a given data source and cast it to the expected class.
 	 *
@@ -99,54 +81,6 @@ struct FRigUnitContext
 		return DataSourceRegistry->RequestSource<T>(InName);
 	}
 
-	/**
-	 * Converts a transform from rig (global) space to world space
-	 */
-	FTransform ToWorldSpace(const FTransform& InTransform) const
-	{
-		return InTransform * ToWorldSpaceTransform;
-	}
-
-	/**
-	 * Converts a transform from world space to rig (global) space
-	 */
-	FTransform ToRigSpace(const FTransform& InTransform) const
-	{
-		return InTransform.GetRelativeTransform(ToWorldSpaceTransform);
-	}
-
-	/**
-	 * Converts a location from rig (global) space to world space
-	 */
-	FVector ToWorldSpace(const FVector& InLocation) const
-	{
-		return ToWorldSpaceTransform.TransformPosition(InLocation);
-	}
-
-	/**
-	 * Converts a location from world space to rig (global) space
-	 */
-	FVector ToRigSpace(const FVector& InLocation) const
-	{
-		return ToWorldSpaceTransform.InverseTransformPosition(InLocation);
-	}
-
-	/**
-	 * Converts a rotation from rig (global) space to world space
-	 */
-	FQuat ToWorldSpace(const FQuat& InRotation) const
-	{
-		return ToWorldSpaceTransform.TransformRotation(InRotation);
-	}
-
-	/**
-	 * Converts a rotation from world space to rig (global) space
-	 */
-	FQuat ToRigSpace(const FQuat& InRotation) const
-	{
-		return ToWorldSpaceTransform.InverseTransformRotation(InRotation);
-	}
-
 	/**
 	 * Returns true if this context is currently being interacted on
 	 */
diff --git a/Engine/Source/Runtime/RigVM/Private/RigVMCore/RigVMExecuteContext.cpp b/Engine/Source/Runtime/RigVM/Private/RigVMCore/RigVMExecuteContext.cpp
--- a/Engine/Source/Runtime/RigVM/Private/RigVMCore/RigVMExecuteContext.cpp
+++ b/Engine/Source/Runtime/RigVM/Private/RigVMCore/RigVMExecuteContext.cpp
@@ -4,6 +4,35 @@
 
 #include UE_INLINE_GENERATED_CPP_BY_NAME(RigVMExecuteContext)
 
+void FRigVMExecuteContext::SetOwningComponent(const USceneComponent* InOwningComponent)
+{
+	OwningComponent = InOwningComponent;
+	OwningActor = nullptr;
+	World = nullptr;
+	ToWorldSpaceTransform = FTransform::Identity;
+	
+	if(OwningComponent)
+	{
+		ToWorldSpaceTransform = OwningComponent->GetComponentToWorld();
+		SetOwningActor(OwningComponent->GetOwner());
+	}
+}
+
+void FRigVMExecuteContext::SetOwningActor(const AActor* InActor)
+{
+	OwningActor = InActor;
+	World = nullptr;
+	if(OwningActor)
+	{
+		World = OwningActor->GetWorld();
+	}
+}
+
+void FRigVMExecuteContext::SetWorld(const UWorld* InWorld)
+{
+	World = InWorld;
+}
+
 bool FRigVMExecuteContext::SerializeFromMismatchedTag(const FPropertyTag& Tag, FStructuredArchive::FSlot Slot)
 {
 	static const FName ControlRigExecuteContextName("ControlRigExecuteContext");
diff --git a/Engine/Source/Runtime/RigVM/Private/RigVMHost.cpp b/Engine/Source/Runtime/RigVM/Private/RigVMHost.cpp
--- a/Engine/Source/Runtime/RigVM/Private/RigVMHost.cpp
+++ b/Engine/Source/Runtime/RigVM/Private/RigVMHost.cpp
@@ -419,6 +419,7 @@ bool URigVMHost::Execute(const FName& InEventName)
 	PublicContext.SetDeltaTime(DeltaTime);
 	PublicContext.SetAbsoluteTime(AbsoluteTime);
 	PublicContext.SetFramesPerSecond(GetCurrentFramesPerSecond());
+	PublicContext.SetOwningComponent(GetOwningSceneComponent());
 
 	if (VM)
 	{
@@ -713,6 +714,11 @@ URigVM* URigVMHost::GetVM()
 	return VM;
 }
 
+USceneComponent* URigVMHost::GetOwningSceneComponent()
+{
+	return GetTypedOuter<USceneComponent>();
+}
+
 void URigVMHost::SwapVMToNativizedIfRequired(UClass* InNativizedClass)
 {
 		if (HasAnyFlags(RF_NeedPostLoad))
diff --git a/Engine/Source/Runtime/RigVM/Public/RigVMCore/RigVMExecuteContext.h b/Engine/Source/Runtime/RigVM/Public/RigVMCore/RigVMExecuteContext.h
--- a/Engine/Source/Runtime/RigVM/Public/RigVMCore/RigVMExecuteContext.h
+++ b/Engine/Source/Runtime/RigVM/Public/RigVMCore/RigVMExecuteContext.h
@@ -25,6 +25,8 @@
 #include "RigVMLog.h"
 #include "RigVMDrawInterface.h"
 #include "RigVMDrawContainer.h"
+#include "GameFramework/Actor.h"
+#include "Components/SceneComponent.h"
 
 #include "RigVMExecuteContext.generated.h"
 
@@ -237,6 +239,10 @@ struct RIGVM_API FRigVMExecuteContext
 #endif
 		, DrawInterfacePtr(nullptr)
 		, DrawContainerPtr(nullptr)
+		, ToWorldSpaceTransform(FTransform::Identity)
+		, OwningComponent(nullptr)
+		, OwningActor(nullptr)
+		, World(nullptr)
 	{
 	}
 
@@ -288,6 +294,74 @@ struct RIGVM_API FRigVMExecuteContext
 
 	double GetFramesPerSecond() const { return FramesPerSecond; } 
 	void SetFramesPerSecond(double InFramesPerSecond) { FramesPerSecond = InFramesPerSecond; }
+
+	/** The current transform going from rig (global) space to world space */
+	const FTransform& GetToWorldSpaceTransform() const { return ToWorldSpaceTransform; };
+
+	/** The current component this VM is owned by */
+	const USceneComponent* GetOwningComponent() const { return OwningComponent; }
+
+	/** The current actor this VM is owned by */
+	const AActor* GetOwningActor() const { return OwningActor; }
+
+	/** The world this VM is running in */
+	const UWorld* GetWorld() const { return World; }
+
+	void SetOwningComponent(const USceneComponent* InOwningComponent);
+
+	void SetOwningActor(const AActor* InActor);
+
+	void SetWorld(const UWorld* InWorld);
+
+	void SetToWorldSpaceTransform(const FTransform& InToWorldSpaceTransform) { ToWorldSpaceTransform = InToWorldSpaceTransform; }
+
+	/**
+	 * Converts a transform from VM (global) space to world space
+	 */
+	FTransform ToWorldSpace(const FTransform& InTransform) const
+	{
+		return InTransform * ToWorldSpaceTransform;
+	}
+
+	/**
+	 * Converts a transform from world space to VM (global) space
+	 */
+	FTransform ToVMSpace(const FTransform& InTransform) const
+	{
+		return InTransform.GetRelativeTransform(ToWorldSpaceTransform);
+	}
+
+	/**
+	 * Converts a location from VM (global) space to world space
+	 */
+	FVector ToWorldSpace(const FVector& InLocation) const
+	{
+		return ToWorldSpaceTransform.TransformPosition(InLocation);
+	}
+
+	/**
+	 * Converts a location from world space to VM (global) space
+	 */
+	FVector ToVMSpace(const FVector& InLocation) const
+	{
+		return ToWorldSpaceTransform.InverseTransformPosition(InLocation);
+	}
+
+	/**
+	 * Converts a rotation from VM (global) space to world space
+	 */
+	FQuat ToWorldSpace(const FQuat& InRotation) const
+	{
+		return ToWorldSpaceTransform.TransformRotation(InRotation);
+	}
+
+	/**
+	 * Converts a rotation from world space to VM (global) space
+	 */
+	FQuat ToVMSpace(const FQuat& InRotation) const
+	{
+		return ToWorldSpaceTransform.InverseTransformRotation(InRotation);
+	}
 	
 	FRigVMNameCache* GetNameCache() const { return NameCache; }
 
@@ -355,6 +429,19 @@ protected:
 	FRigVMDrawInterface* DrawInterfacePtr;
 	FRigVMDrawContainer* DrawContainerPtr;
 
+	/** The current transform going from rig (global) space to world space */
+	FTransform ToWorldSpaceTransform;
+
+	/** The current component this VM is owned by */
+	const USceneComponent* OwningComponent;
+
+	/** The current actor this VM is owned by */
+	const AActor* OwningActor;
+
+	/** The world this VM is running in */
+	const UWorld* World;
+
+
 #if UE_RIGVM_DEBUG_EXECUTION
 	FString DebugMemoryString;
 
diff --git a/Engine/Source/Runtime/RigVM/Public/RigVMHost.h b/Engine/Source/Runtime/RigVM/Public/RigVMHost.h
--- a/Engine/Source/Runtime/RigVM/Public/RigVMHost.h
+++ b/Engine/Source/Runtime/RigVM/Public/RigVMHost.h
@@ -201,6 +201,8 @@ public:
 	const FRigVMDrawContainer& GetDrawContainer() const { return DrawContainer; };
 	FRigVMDrawContainer& GetDrawContainer() { return DrawContainer; };
 
+	virtual USceneComponent* GetOwningSceneComponent(); 
+
 	virtual void PostInitInstanceIfRequired() {};
 	void SwapVMToNativizedIfRequired(UClass* InNativizedClass = nullptr);
 	static bool AreNativizedVMsDisabled() ;
