Control Rig: Fix crash when refreshing hierarchy

diff --git a/Engine/Plugins/Animation/ControlRig/Source/ControlRigEditor/Private/Editor/SRigHierarchy.cpp b/Engine/Plugins/Animation/ControlRig/Source/ControlRigEditor/Private/Editor/SRigHierarchy.cpp
--- a/Engine/Plugins/Animation/ControlRig/Source/ControlRigEditor/Private/Editor/SRigHierarchy.cpp
+++ b/Engine/Plugins/Animation/ControlRig/Source/ControlRigEditor/Private/Editor/SRigHierarchy.cpp
@@ -46,6 +46,8 @@
 #include "Framework/Notifications/NotificationManager.h"
 #include "Widgets/Notifications/SNotificationList.h"
 #include "Styling/AppStyle.h"
+#include "ControlRigSkeletalMeshComponent.h"
+#include "Sequencer/ControlRigLayerInstance.h"
 #include "Algo/MinElement.h"
 #include "Algo/MaxElement.h"
 #include "RigVMFunctions/Math/RigVMMathLibrary.h"
@@ -1345,7 +1347,8 @@ void SRigHierarchy::RefreshHierarchy(const FAssetData& InAssetData)
 		return;
 	}
 
-	ControlRigEditor.Pin()->ClearDetailObject();
+	FControlRigEditor* StrongEditor = ControlRigEditor.Pin().Get();
+	StrongEditor->ClearDetailObject();
 
 	URigHierarchy* Hierarchy = GetDefaultHierarchy();
 	USkeletalMesh* Mesh = Cast<USkeletalMesh>(InAssetData.GetAsset());
@@ -1359,7 +1362,7 @@ void SRigHierarchy::RefreshHierarchy(const FAssetData& InAssetData)
 		// we do this to avoid the editmode / viewport shapes to refresh recursively,
 		// which can add an extreme slowdown depending on the number of bones (n^(n-1))
 		bool bSelectBones = true;
-		if (UControlRig* CurrentRig = ControlRigEditor.Pin()->ControlRig)
+		if (UControlRig* CurrentRig = StrongEditor->ControlRig)
 		{
 			bSelectBones = !CurrentRig->IsConstructionModeEnabled();
 		}
@@ -1374,20 +1377,28 @@ void SRigHierarchy::RefreshHierarchy(const FAssetData& InAssetData)
 	}
 
 	ControlRigBlueprint->PropagateHierarchyFromBPToInstances();
-	ControlRigEditor.Pin()->OnHierarchyChanged();
+	StrongEditor->OnHierarchyChanged();
 	ControlRigBlueprint->BroadcastRefreshEditor();
 	RefreshTreeView();
 	FSlateApplication::Get().DismissAllMenus();
 
-	if (ControlRigEditor.IsValid() && 
-		Mesh != nullptr)
+	if (Mesh != nullptr)
 	{
-		ControlRigEditor.Pin()->GetPersonaToolkit()->SetPreviewMesh(Mesh, true);
+		StrongEditor->GetPersonaToolkit()->SetPreviewMesh(Mesh, true);
+		const UControlRigSkeletalMeshComponent* EditorSkelComp = Cast<UControlRigSkeletalMeshComponent>(StrongEditor->GetPersonaToolkit()->GetPreviewMeshComponent());
+		if (UControlRigLayerInstance* ControlRigLayerInstance = Cast<UControlRigLayerInstance>(EditorSkelComp->GetAnimInstance()))
+		{
+			StrongEditor->PreviewInstance = Cast<UAnimPreviewInstance>(ControlRigLayerInstance->GetSourceAnimInstance());
+		}
+		else
+		{
+			StrongEditor->PreviewInstance = Cast<UAnimPreviewInstance>(EditorSkelComp->GetAnimInstance());
+		}
 	}
 
 	if (ControlRigEditor.IsValid())
 	{
-		ControlRigEditor.Pin()->Compile();
+		StrongEditor->Compile();
 	}
 }
 
