Control Rig: Propagate user bone transformation to the preview instance when switching from construction to forward solve

diff --git a/Engine/Plugins/Animation/ControlRig/Source/ControlRigEditor/Private/Editor/ControlRigEditor.cpp b/Engine/Plugins/Animation/ControlRig/Source/ControlRigEditor/Private/Editor/ControlRigEditor.cpp
--- a/Engine/Plugins/Animation/ControlRig/Source/ControlRigEditor/Private/Editor/ControlRigEditor.cpp
+++ b/Engine/Plugins/Animation/ControlRig/Source/ControlRigEditor/Private/Editor/ControlRigEditor.cpp
@@ -1310,6 +1310,12 @@ void FControlRigEditor::SetEventQueue(TArray<FName> InEventQueue, bool bCompile)
 			if (UControlRigBlueprint* RigBlueprint = Cast<UControlRigBlueprint>(GetBlueprintObj()))
 			{
 				RigBlueprint->Validator->SetControlRig(ControlRig);
+
+				if (LastEventQueue == ConstructionEventQueue)
+				{
+					// This will propagate any user bone transformation done during construction to the preview instance
+					ResetAllBoneModification();
+				}
 			}
 		}
 
