Control Rig: Remove use of BlueprintInternalUseOnlyHierarchical

diff --git a/Engine/Plugins/Animation/ControlRig/Source/ControlRig/Public/Units/RigDispatchFactory.h b/Engine/Plugins/Animation/ControlRig/Source/ControlRig/Public/Units/RigDispatchFactory.h
--- a/Engine/Plugins/Animation/ControlRig/Source/ControlRig/Public/Units/RigDispatchFactory.h
+++ b/Engine/Plugins/Animation/ControlRig/Source/ControlRig/Public/Units/RigDispatchFactory.h
@@ -9,7 +9,7 @@
 #include "RigDispatchFactory.generated.h"
 
 /** Base class for all rig dispatch factories */
-USTRUCT(BlueprintType, meta=(Abstract, BlueprintInternalUseOnlyHierarchical, ExecuteContextType=FControlRigExecuteContext))
+USTRUCT(BlueprintType, meta=(Abstract, ExecuteContextType=FControlRigExecuteContext))
 struct CONTROLRIG_API FRigDispatchFactory : public FRigVMDispatchFactory
 {
 	GENERATED_BODY()
diff --git a/Engine/Plugins/Animation/ControlRig/Source/ControlRig/Public/Units/RigUnit.h b/Engine/Plugins/Animation/ControlRig/Source/ControlRig/Public/Units/RigUnit.h
--- a/Engine/Plugins/Animation/ControlRig/Source/ControlRig/Public/Units/RigUnit.h
+++ b/Engine/Plugins/Animation/ControlRig/Source/ControlRig/Public/Units/RigUnit.h
@@ -12,7 +12,7 @@
 struct FRigUnitContext;
 
 /** Base class for all rig units */
-USTRUCT(BlueprintType, meta=(Abstract, NodeColor = "0.1 0.1 0.1", ExecuteContext="FControlRigExecuteContext"), BlueprintInternalUseOnlyHierarchical)
+USTRUCT(BlueprintType, meta=(Abstract, NodeColor = "0.1 0.1 0.1", ExecuteContext="FControlRigExecuteContext"))
 struct CONTROLRIG_API FRigUnit : public FRigVMStruct
 {
 	GENERATED_BODY()
