Control Rig: Fix crash when manipulating multiple characters in level

diff --git a/Engine/Plugins/Animation/ControlRig/Source/ControlRigEditor/Private/Sequencer/ControlRigParameterTrackEditor.cpp b/Engine/Plugins/Animation/ControlRig/Source/ControlRigEditor/Private/Sequencer/ControlRigParameterTrackEditor.cpp
--- a/Engine/Plugins/Animation/ControlRig/Source/ControlRigEditor/Private/Sequencer/ControlRigParameterTrackEditor.cpp
+++ b/Engine/Plugins/Animation/ControlRig/Source/ControlRigEditor/Private/Sequencer/ControlRigParameterTrackEditor.cpp
@@ -199,6 +199,7 @@ static USkeleton* AcquireSkeletonFromObjectGuid(const FGuid& Guid, UObject** Obj
 }
 
 bool FControlRigParameterTrackEditor::bAutoGenerateControlRigTrack = true;
+FCriticalSection FControlRigParameterTrackEditor::ControlUndoTransactionMutex;
 
 FControlRigParameterTrackEditor::FControlRigParameterTrackEditor(TSharedRef<ISequencer> InSequencer)
 	: FKeyframeTrackEditor<UMovieSceneControlRigParameterTrack>(InSequencer)
@@ -3070,6 +3071,7 @@ void FControlRigParameterTrackEditor::HandleControlUndoBracket(UControlRig* Subj
 {
 	if(bOpenUndoBracket && ControlUndoBracket == 0)
 	{
+		FScopeLock ScopeLock(&ControlUndoTransactionMutex);
 		ControlUndoTransaction = MakeShareable(new FScopedTransaction(LOCTEXT("KeyMultipleControls", "Auto-Key multiple controls")));
 	}
 
@@ -3077,6 +3079,7 @@ void FControlRigParameterTrackEditor::HandleControlUndoBracket(UControlRig* Subj
 	
 	if(!bOpenUndoBracket && ControlUndoBracket == 0)
 	{
+		FScopeLock ScopeLock(&ControlUndoTransactionMutex);
 		ControlUndoTransaction.Reset();
 	}
 }
diff --git a/Engine/Plugins/Animation/ControlRig/Source/ControlRigEditor/Private/Sequencer/ControlRigParameterTrackEditor.h b/Engine/Plugins/Animation/ControlRig/Source/ControlRigEditor/Private/Sequencer/ControlRigParameterTrackEditor.h
--- a/Engine/Plugins/Animation/ControlRig/Source/ControlRigEditor/Private/Sequencer/ControlRigParameterTrackEditor.h
+++ b/Engine/Plugins/Animation/ControlRig/Source/ControlRigEditor/Private/Sequencer/ControlRigParameterTrackEditor.h
@@ -272,6 +272,9 @@ private:
 	/** An index counter for the opened undo brackets */
 	int32 ControlUndoBracket;
 
+	/** Lock to avoid registering multiple transactions from different tracks at the same time */
+	static FCriticalSection ControlUndoTransactionMutex;
+
 	/** A transaction used to group multiple key events */
 	TSharedPtr<FScopedTransaction> ControlUndoTransaction;
 
