Control Rig: Improve performance when selecting multiple nodes in graph

diff --git a/Engine/Plugins/Animation/ControlRig/Source/ControlRigEditor/Private/Editor/ControlRigEditor.cpp b/Engine/Plugins/Animation/ControlRig/Source/ControlRigEditor/Private/Editor/ControlRigEditor.cpp
--- a/Engine/Plugins/Animation/ControlRig/Source/ControlRigEditor/Private/Editor/ControlRigEditor.cpp
+++ b/Engine/Plugins/Animation/ControlRig/Source/ControlRigEditor/Private/Editor/ControlRigEditor.cpp
@@ -3078,6 +3078,7 @@ void FControlRigEditor::HandleModifiedEvent(ERigVMGraphNotifType InNotifType, UR
 
 	switch (InNotifType)
 	{
+		case ERigVMGraphNotifType::NodeSelectionChanged:
 		case ERigVMGraphNotifType::NodeSelected:
 		case ERigVMGraphNotifType::NodeDeselected:
 		{
@@ -3085,6 +3086,14 @@ void FControlRigEditor::HandleModifiedEvent(ERigVMGraphNotifType InNotifType, UR
 			{
 				TSharedPtr<SGraphEditor> GraphEd = GetGraphEditor(RigGraph);
 				URigVMNode* Node = Cast<URigVMNode>(InSubject);
+				if (InNotifType == ERigVMGraphNotifType::NodeSelectionChanged)
+				{
+					const TArray<FName> SelectedNodes = InGraph->GetSelectNodes();
+					if (!SelectedNodes.IsEmpty())
+					{
+						Node = Cast<URigVMNode>(InGraph->FindNodeByName(SelectedNodes.Last()));	
+					}
+				}
 
 				if (GraphEd.IsValid() && Node != nullptr)
 				{
@@ -3093,12 +3102,9 @@ void FControlRigEditor::HandleModifiedEvent(ERigVMGraphNotifType InNotifType, UR
 					if (!RigGraph->bIsSelecting)
 					{
 						TGuardValue<bool> SelectingGuard(RigGraph->bIsSelecting, true);
-						if (URigVMNode* ModelNode = Cast<URigVMNode>(InSubject))
+						if (UEdGraphNode* EdNode = RigGraph->FindNodeForModelNodeName(Node->GetFName()))
 						{
-							if (UEdGraphNode* EdNode = RigGraph->FindNodeForModelNodeName(ModelNode->GetFName()))
-							{
-								GraphEd->SetNodeSelection(EdNode, InNotifType == ERigVMGraphNotifType::NodeSelected);
-							}
+							GraphEd->SetNodeSelection(EdNode, InNotifType == ERigVMGraphNotifType::NodeSelected);
 						}
 					}
 				}
@@ -3195,7 +3201,6 @@ void FControlRigEditor::HandleModifiedEvent(ERigVMGraphNotifType InNotifType, UR
 			}
 			break;
 		}
-		case ERigVMGraphNotifType::NodeSelectionChanged:
 		default:
 		{
 			break;
diff --git a/Engine/Source/Developer/RigVMDeveloper/Private/RigVMModel/RigVMController.cpp b/Engine/Source/Developer/RigVMDeveloper/Private/RigVMModel/RigVMController.cpp
--- a/Engine/Source/Developer/RigVMDeveloper/Private/RigVMModel/RigVMController.cpp
+++ b/Engine/Source/Developer/RigVMDeveloper/Private/RigVMModel/RigVMController.cpp
@@ -6379,8 +6379,8 @@ bool URigVMController::SetNodeSelection(const TArray<FName>& InNodeNames, bool b
 		{
 			if(Graph->SelectedNodes.Remove(PreviouslySelectedNode) > 0)
 			{
-				Notify(ERigVMGraphNotifType::NodeDeselected, Graph->FindNodeByName(PreviouslySelectedNode));
 				bSelectedSomething = true;
+				break;
 			}
 		}
 	}
@@ -6393,7 +6393,6 @@ bool URigVMController::SetNodeSelection(const TArray<FName>& InNodeNames, bool b
 			Graph->SelectedNodes.AddUnique(InNodeName);
 			if (PreviousNum != Graph->SelectedNodes.Num())
 			{
-				Notify(ERigVMGraphNotifType::NodeSelected, NodeToSelect);
 				bSelectedSomething = true;
 			}
 		}
