Control Rig: Fix crash after updating preview mesh

diff --git a/Engine/Plugins/Animation/ControlRig/Source/ControlRigEditor/Private/Editor/ControlRigEditor.cpp b/Engine/Plugins/Animation/ControlRig/Source/ControlRigEditor/Private/Editor/ControlRigEditor.cpp
--- a/Engine/Plugins/Animation/ControlRig/Source/ControlRigEditor/Private/Editor/ControlRigEditor.cpp
+++ b/Engine/Plugins/Animation/ControlRig/Source/ControlRigEditor/Private/Editor/ControlRigEditor.cpp
@@ -4384,6 +4384,7 @@ void FControlRigEditor::UpdateControlRig()
 				ControlRig->Initialize(true);
  			}
 
+ 			PreviewInstance = Cast<UAnimPreviewInstance>(AnimInstance->GetSourceAnimInstance());
 			ControlRig->PreviewInstance = PreviewInstance;
 
 #if WITH_EDITOR
