Control Rig: Fix ensure for empty default value for ERigElementType pin type

diff --git a/Engine/Plugins/Animation/ControlRig/Source/ControlRigEditor/Private/Graph/NodeSpawners/ControlRigUnitNodeSpawner.cpp b/Engine/Plugins/Animation/ControlRig/Source/ControlRigEditor/Private/Graph/NodeSpawners/ControlRigUnitNodeSpawner.cpp
--- a/Engine/Plugins/Animation/ControlRig/Source/ControlRigEditor/Private/Graph/NodeSpawners/ControlRigUnitNodeSpawner.cpp
+++ b/Engine/Plugins/Animation/ControlRig/Source/ControlRigEditor/Private/Graph/NodeSpawners/ControlRigUnitNodeSpawner.cpp
@@ -289,14 +289,14 @@ UControlRigGraphNode* UControlRigUnitNodeSpawner::SpawnNode(UEdGraph* ParentGrap
 					if (UEnum* RigElementTypeEnum = StaticEnum<ERigElementType>())
 					{
 						ERigElementType UsedElementType = ERigElementType::None;
-						int64 MaxEnumValue = RigElementTypeEnum->GetMaxEnumValue();
+						const int32 NumEnumValues = RigElementTypeEnum->NumEnums();
 
-						for (int64 EnumValue = 0; EnumValue < MaxEnumValue; EnumValue++)
+						for (int32 EnumIndex = 0; EnumIndex < NumEnumValues; EnumIndex++)
 						{
-							FString EnumText = RigElementTypeEnum->GetDisplayNameTextByValue(EnumValue).ToString().ToLower();
+							const FString EnumText = RigElementTypeEnum->GetDisplayNameTextByIndex(EnumIndex).ToString().ToLower();
 							if (!EnumText.IsEmpty() && UsedFilterString.Contains(EnumText))
 							{
-								UsedElementType = (ERigElementType)EnumValue;
+								UsedElementType = (ERigElementType)RigElementTypeEnum->GetValueByIndex(EnumIndex);
 								break;
 							}
 						}
