Control Rig: Fix cannot add space in the space picker widget

diff --git a/Engine/Plugins/Animation/ControlRig/Source/ControlRigEditor/Private/Editor/SRigSpacePickerWidget.cpp b/Engine/Plugins/Animation/ControlRig/Source/ControlRigEditor/Private/Editor/SRigSpacePickerWidget.cpp
--- a/Engine/Plugins/Animation/ControlRig/Source/ControlRigEditor/Private/Editor/SRigSpacePickerWidget.cpp
+++ b/Engine/Plugins/Animation/ControlRig/Source/ControlRigEditor/Private/Editor/SRigSpacePickerWidget.cpp
@@ -286,7 +286,11 @@ FReply SRigSpacePickerWidget::OpenDialog(bool bModal)
 	{
 		Window->GetOnWindowDeactivatedEvent().AddLambda([this]()
 		{
-			SetControls(nullptr, {});
+			// Do not reset if we lost focus because of opening the context menu
+			if (!ContextMenu.IsValid())
+			{
+				SetControls(nullptr, {});
+			}
 		});
 	}
 	
