Control Rig: Hide parent class related user interface

diff --git a/Engine/Plugins/Animation/ControlRig/Source/ControlRigEditor/Private/ControlRigBlueprintFactory.cpp b/Engine/Plugins/Animation/ControlRig/Source/ControlRigEditor/Private/ControlRigBlueprintFactory.cpp
--- a/Engine/Plugins/Animation/ControlRig/Source/ControlRigEditor/Private/ControlRigBlueprintFactory.cpp
+++ b/Engine/Plugins/Animation/ControlRig/Source/ControlRigEditor/Private/ControlRigBlueprintFactory.cpp
@@ -270,8 +270,9 @@ UControlRigBlueprintFactory::UControlRigBlueprintFactory()
 
 bool UControlRigBlueprintFactory::ConfigureProperties()
 {
-	TSharedRef<SControlRigBlueprintCreateDialog> Dialog = SNew(SControlRigBlueprintCreateDialog);
-	return Dialog->ConfigureProperties(this);
+	// we don't need to do anything,
+	// let's return true to indicate that all properties have been configured to produce a control rig
+	return true;
 };
 
 UObject* UControlRigBlueprintFactory::FactoryCreateNew(UClass* Class, UObject* InParent, FName Name, EObjectFlags Flags, UObject* Context, FFeedbackContext* Warn, FName CallingContext)
diff --git a/Engine/Plugins/Animation/ControlRig/Source/ControlRigEditor/Private/Editor/ControlRigEditor.h b/Engine/Plugins/Animation/ControlRig/Source/ControlRigEditor/Private/Editor/ControlRigEditor.h
--- a/Engine/Plugins/Animation/ControlRig/Source/ControlRigEditor/Private/Editor/ControlRigEditor.h
+++ b/Engine/Plugins/Animation/ControlRig/Source/ControlRigEditor/Private/Editor/ControlRigEditor.h
@@ -115,6 +115,9 @@ public:
 	virtual void PasteNodes() override;
 	virtual bool CanPasteNodes() const override;
 
+	virtual bool IsNativeParentClassCodeLinkEnabled() const override { return false; }
+	virtual bool ReparentBlueprint_IsVisible() const override { return false; }
+
 	virtual FReply OnSpawnGraphNodeByShortcut(FInputChord InChord, const FVector2D& InPosition, UEdGraph* InGraph) override;
 
 	// Control Rig BP does not use regular BP function libraries, no need to load them.
diff --git a/Engine/Source/Editor/Kismet/Public/BlueprintEditor.h b/Engine/Source/Editor/Kismet/Public/BlueprintEditor.h
--- a/Engine/Source/Editor/Kismet/Public/BlueprintEditor.h
+++ b/Engine/Source/Editor/Kismet/Public/BlueprintEditor.h
@@ -448,7 +448,7 @@ public:
 	bool IsParentClassNative() const;
 
 	/** Returns true if the parent class is native and the link to it's header can be shown*/
-	bool IsNativeParentClassCodeLinkEnabled() const;
+	virtual bool IsNativeParentClassCodeLinkEnabled() const;
 
 	/** Handles opening the header file of native parent class */
 	void OnEditParentClassNativeCodeClicked();
@@ -593,7 +593,7 @@ public:
 
 	/** Reparent the current blueprint */
 	void ReparentBlueprint_Clicked();
-	bool ReparentBlueprint_IsVisible() const;
+	virtual bool ReparentBlueprint_IsVisible() const;
 	void ReparentBlueprint_NewParentChosen(UClass* ChosenClass);
 
 	/** Utility function to handle all steps required to rename a newly added action */
