Control Rig: Clear target instance on stale animnodes

diff --git a/Engine/Plugins/Animation/ControlRig/Source/ControlRig/Private/AnimNode_ControlRigBase.cpp b/Engine/Plugins/Animation/ControlRig/Source/ControlRig/Private/AnimNode_ControlRigBase.cpp
--- a/Engine/Plugins/Animation/ControlRig/Source/ControlRig/Private/AnimNode_ControlRigBase.cpp
+++ b/Engine/Plugins/Animation/ControlRig/Source/ControlRig/Private/AnimNode_ControlRigBase.cpp
@@ -77,6 +77,10 @@ void FAnimNode_ControlRigBase::Initialize_AnyThread(const FAnimationInitializeCo
 		bControlRigRequiresInitialization = true;
 		LastBonesSerialNumberForCacheBones = 0;
 	}
+	else
+	{
+		SetTargetInstance(nullptr);
+	}
 }
 
 void FAnimNode_ControlRigBase::GatherDebugData(FNodeDebugData& DebugData)
