Control Rig: Fix binding variable to subpin removes valid pins and ensures
#swarm https://p4-swarm.epicgames.net/reviews/23870506

diff --git a/Engine/Plugins/Animation/ControlRig/Source/ControlRigDeveloper/Private/Graph/ControlRigGraphNode.cpp b/Engine/Plugins/Animation/ControlRig/Source/ControlRigDeveloper/Private/Graph/ControlRigGraphNode.cpp
--- a/Engine/Plugins/Animation/ControlRig/Source/ControlRigDeveloper/Private/Graph/ControlRigGraphNode.cpp
+++ b/Engine/Plugins/Animation/ControlRig/Source/ControlRigDeveloper/Private/Graph/ControlRigGraphNode.cpp
@@ -802,7 +802,8 @@ bool UControlRigGraphNode::ModelPinsChanged(bool bForce)
 		for(const auto& Pair : CachedPins)
 		{
 			URigVMPin* ModelPin = Pair.Key;
-			if(!PinListForPin(ModelPin).Contains(ModelPin))
+			const URigVMPin* RootPin = ModelPin->GetRootPin();
+			if(!PinListForPin(RootPin).Contains(RootPin))
 			{
 				PinsToRemove.Add(ModelPin);
 			}
diff --git a/Engine/Plugins/Animation/ControlRig/Source/ControlRigEditor/Private/Graph/SControlRigGraphNode.cpp b/Engine/Plugins/Animation/ControlRig/Source/ControlRigEditor/Private/Graph/SControlRigGraphNode.cpp
--- a/Engine/Plugins/Animation/ControlRig/Source/ControlRigEditor/Private/Graph/SControlRigGraphNode.cpp
+++ b/Engine/Plugins/Animation/ControlRig/Source/ControlRigEditor/Private/Graph/SControlRigGraphNode.cpp
@@ -1650,7 +1650,8 @@ void SControlRigGraphNode::UpdatePinTreeView()
 				}
 			}
 
-			if(!PinInfo.bHasChildren && ExpanderSlotIndex != INDEX_NONE)
+			// The expander needs to be recreated to adjust for PinInfo.Index changes
+			if (FullPinHorizontalRowWidget->IsValidSlotIndex(ExpanderSlotIndex))
 			{
 				SHorizontalBox::FSlot& Slot = FullPinHorizontalRowWidget->GetSlot(ExpanderSlotIndex);
 				const TSharedRef<SWidget> Widget = Slot.GetWidget();
diff --git a/Engine/Source/Developer/RigVMDeveloper/Private/RigVMCompiler/RigVMASTProxy.cpp b/Engine/Source/Developer/RigVMDeveloper/Private/RigVMCompiler/RigVMASTProxy.cpp
--- a/Engine/Source/Developer/RigVMDeveloper/Private/RigVMCompiler/RigVMASTProxy.cpp
+++ b/Engine/Source/Developer/RigVMDeveloper/Private/RigVMCompiler/RigVMASTProxy.cpp
@@ -26,8 +26,11 @@ FString FRigVMCallstack::GetCallPath(bool bIncludeLast) const
 		}
 		else if (URigVMPin* Pin = Cast<URigVMPin>(Entry))
 		{
-			const bool bUseNodePath = Pin->GetGraph()->IsRootGraph();
-			Segments.Add(Pin->GetPinPath(bUseNodePath));
+			if (URigVMGraph* Graph = Pin->GetGraph())
+			{
+				const bool bUseNodePath = Graph->IsRootGraph();
+				Segments.Add(Pin->GetPinPath(bUseNodePath));
+			}
 		}
 	}
 
diff --git a/Engine/Source/Developer/RigVMDeveloper/Private/RigVMModel/RigVMPin.cpp b/Engine/Source/Developer/RigVMDeveloper/Private/RigVMModel/RigVMPin.cpp
--- a/Engine/Source/Developer/RigVMDeveloper/Private/RigVMModel/RigVMPin.cpp
+++ b/Engine/Source/Developer/RigVMDeveloper/Private/RigVMModel/RigVMPin.cpp
@@ -1605,7 +1605,7 @@ URigVMNode* URigVMPin::GetNode() const
 URigVMGraph* URigVMPin::GetGraph() const
 {
 	URigVMNode* Node = GetNode();
-	if(Node)
+	if(IsValid(Node))
 	{
 		return Node->GetGraph();
 	}
