Sequencer & CR Component: make sure to pick the correct CR instance for the  Components to bind, otherwise CR component would get mistakenly re-initialized all the time. In case of PIE + Spawanable Actor + CR component, sequencer should grab CR component's CR instance for evaluation, ie CR component should bind to the instance that itself owns.

merging to 5.2 at the request of sequencer team

#swarm https://p4-swarm.epicgames.net/reviews/23865352

diff --git a/Engine/Plugins/Animation/ControlRig/Source/ControlRig/Private/Sequencer/MovieSceneControlRigParameterTemplate.cpp b/Engine/Plugins/Animation/ControlRig/Source/ControlRig/Private/Sequencer/MovieSceneControlRigParameterTemplate.cpp
--- a/Engine/Plugins/Animation/ControlRig/Source/ControlRig/Private/Sequencer/MovieSceneControlRigParameterTemplate.cpp
+++ b/Engine/Plugins/Animation/ControlRig/Source/ControlRig/Private/Sequencer/MovieSceneControlRigParameterTemplate.cpp
@@ -1155,6 +1155,13 @@ struct FControlRigParameterExecutionToken : IMovieSceneExecutionToken
 						SelectControls(ControlRig, SelectedControls);
 					}
 				}
+
+				// make sure to pick the correct CR instance for the  Components to bind.
+				// In case of PIE + Spawnable Actor + CR component, sequencer should grab
+				// CR component's CR instance for evaluation, see comment in BindToSequencerInstance
+				// i.e. CR component should bind to the instance that it owns itself.
+				ControlRig = GetControlRig(Section, BoundObjects[0].Get());
+				
 				// ensure that pre animated state is saved, must be done before bind
 				Player.SavePreAnimatedState(*ControlRig, FMovieSceneControlRigParameterTemplate::GetAnimTypeID(), FControlRigParameterPreAnimatedTokenProducer(Operand.SequenceID));
 
