CR component:
delay Initialize() to tickcomponent to avoid:
1. Calling Initialize() before OnPreInitialize() BP events are bound, which was happening when we are only calling Init in OnRegister for cooked build
2. Calling Initialize() multiple times unnecessarily across OnRegister() + InitializeComponent(), see code comments.

#swarm https://p4-swarm.epicgames.net/reviews/23898340

diff --git a/Engine/Plugins/Animation/ControlRig/Source/ControlRig/Private/ControlRigComponent.cpp b/Engine/Plugins/Animation/ControlRig/Source/ControlRig/Private/ControlRigComponent.cpp
--- a/Engine/Plugins/Animation/ControlRig/Source/ControlRig/Private/ControlRigComponent.cpp
+++ b/Engine/Plugins/Animation/ControlRig/Source/ControlRig/Private/ControlRigComponent.cpp
@@ -71,6 +71,8 @@ UControlRigComponent::UControlRigComponent(const FObjectInitializer& ObjectIniti
 	LazyEvaluationRotationThreshold = 0.5f;
 	LazyEvaluationScaleThreshold = 0.01f;
 	bNeedsEvaluation = true;
+
+	bNeedToInitialize = false;
 }
 
 #if WITH_EDITOR
@@ -114,14 +116,6 @@ void UControlRigComponent::Serialize(FArchive& Ar)
 	}
 }
 
-#if WITH_EDITOR
-void UControlRigComponent::InitializeComponent()
-{
-	Super::InitializeComponent();
-	Initialize();
-}
-#endif
-
 void UControlRigComponent::OnRegister()
 {
 	Super::OnRegister();
@@ -133,7 +127,19 @@ void UControlRigComponent::OnRegister()
 		gPendingSkeletalMeshes.FindOrAdd(this);
 	}
 
-	Initialize();
+	// call to Initialize() should be delayed until TickComponent
+	
+	// Initialize() here directly is not enough because
+	// in case of PIE, OnRegister() is called before BP native events like
+	// OnPreInitialize are bound, which does not happen until
+	// 1. Re-Registration ( Sequencer / Level Editor / non- PIE use cases, and there is no way to detect it here)
+	// 2. InitializeComponent() (PIE + Cooked build only, not called for editor use cases)
+
+	// so to avoid calling Initialize() in both OnRegister() and InitializeComponent()
+	// we have to delay Initialize() until TickComponent, which is called all the time
+	// in both editor + cooked.
+
+	bNeedToInitialize = true;
 
 	if (AActor* Actor = GetOwner())
 	{
@@ -182,6 +188,12 @@ void UControlRigComponent::OnUnregister()
 
 void UControlRigComponent::TickComponent(float DeltaTime, enum ELevelTick TickType, FActorComponentTickFunction *ThisTickFunction)
 {
+	if (bNeedToInitialize)
+	{
+		Initialize();
+		bNeedToInitialize = false;
+	}
+	
 	if(!bUpdateRigOnTick)
 	{
 		return;
diff --git a/Engine/Plugins/Animation/ControlRig/Source/ControlRig/Public/ControlRigComponent.h b/Engine/Plugins/Animation/ControlRig/Source/ControlRig/Public/ControlRigComponent.h
--- a/Engine/Plugins/Animation/ControlRig/Source/ControlRig/Public/ControlRigComponent.h
+++ b/Engine/Plugins/Animation/ControlRig/Source/ControlRig/Public/ControlRigComponent.h
@@ -218,9 +218,6 @@ public:
 	//~ End UObject Interface
 	
 	//~ Begin UActorComponent Interface
-#if WITH_EDITOR
-	virtual void InitializeComponent() override;
-#endif
 	virtual void OnRegister() override;
 	virtual void OnUnregister() override;
 	virtual void TickComponent(float DeltaTime, enum ELevelTick TickType, FActorComponentTickFunction* ThisTickFunction) override;
@@ -677,6 +674,8 @@ private:
 
 	TSharedPtr<IControlRigObjectBinding> ObjectBinding;
 
+	bool bNeedToInitialize;
+
 	friend class FControlRigSceneProxy;
 };
 
