Control Rig: Pyton accessible modified event and unit test

diff --git a/Engine/Plugins/Animation/ControlRig/Source/ControlRig/Private/Rigs/RigHierarchy.cpp b/Engine/Plugins/Animation/ControlRig/Source/ControlRig/Private/Rigs/RigHierarchy.cpp
--- a/Engine/Plugins/Animation/ControlRig/Source/ControlRig/Private/Rigs/RigHierarchy.cpp
+++ b/Engine/Plugins/Animation/ControlRig/Source/ControlRig/Private/Rigs/RigHierarchy.cpp
@@ -2454,6 +2454,10 @@ void URigHierarchy::Notify(ERigHierarchyNotification InNotifType, const FRigBase
 	if(QueuedNotifications.IsEmpty())
 	{
 		ModifiedEvent.Broadcast(InNotifType, this, InElement);
+		if(ModifiedEventDynamic.IsBound())
+		{
+			ModifiedEventDynamic.Broadcast(InNotifType, this, InElement->GetKey());
+		}
 	}
 	else
 	{
@@ -2626,6 +2630,10 @@ void URigHierarchy::SendQueuedNotifications()
 	}
 
 	ModifiedEvent.Broadcast(ERigHierarchyNotification::InteractionBracketOpened, this, nullptr);
+	if(ModifiedEventDynamic.IsBound())
+	{
+		ModifiedEventDynamic.Broadcast(ERigHierarchyNotification::InteractionBracketOpened, this, FRigElementKey());
+	}
 
 	// finally send all of the notifications
 	// (they have been added in the reverse order to the array)
@@ -2636,10 +2644,18 @@ void URigHierarchy::SendQueuedNotifications()
 		if(Element)
 		{
 			ModifiedEvent.Broadcast(Entry.Type, this, Element);
+			if(ModifiedEventDynamic.IsBound())
+			{
+				ModifiedEventDynamic.Broadcast(Entry.Type, this, Element->GetKey());
+			}
 		}
 	}
 
 	ModifiedEvent.Broadcast(ERigHierarchyNotification::InteractionBracketClosed, this, nullptr);
+	if(ModifiedEventDynamic.IsBound())
+	{
+		ModifiedEventDynamic.Broadcast(ERigHierarchyNotification::InteractionBracketClosed, this, FRigElementKey());
+	}
 }
 
 FTransform URigHierarchy::GetTransform(FRigTransformElement* InTransformElement,
diff --git a/Engine/Plugins/Animation/ControlRig/Source/ControlRig/Public/Rigs/RigHierarchy.h b/Engine/Plugins/Animation/ControlRig/Source/ControlRig/Public/Rigs/RigHierarchy.h
--- a/Engine/Plugins/Animation/ControlRig/Source/ControlRig/Public/Rigs/RigHierarchy.h
+++ b/Engine/Plugins/Animation/ControlRig/Source/ControlRig/Public/Rigs/RigHierarchy.h
@@ -20,6 +20,7 @@ class URigHierarchy;
 class URigHierarchyController;
 
 DECLARE_MULTICAST_DELEGATE_ThreeParams(FRigHierarchyModifiedEvent, ERigHierarchyNotification /* type */, URigHierarchy* /* hierarchy */, const FRigBaseElement* /* element */);
+DECLARE_DYNAMIC_MULTICAST_DELEGATE_ThreeParams(FRigHierarchyModifiedDynamicEvent, ERigHierarchyNotification, NotifType, URigHierarchy*, Hierarchy, FRigElementKey, Subject);
 DECLARE_EVENT_FiveParams(URigHierarchy, FRigHierarchyUndoRedoTransformEvent, URigHierarchy*, const FRigElementKey&, ERigTransformType::Type, const FTransform&, bool /* bUndo */);
 DECLARE_MULTICAST_DELEGATE_TwoParams(FRigHierarchyMetadataChangedDelegate, const FRigElementKey& /* Key */, const FName& /* Name */);
 DECLARE_MULTICAST_DELEGATE_ThreeParams(FRigHierarchyMetadataTagChangedDelegate, const FRigElementKey& /* Key */, const FName& /* Tag */, bool /* AddedOrRemoved */);
@@ -3238,6 +3239,10 @@ public:
 private:
 
 	FRigHierarchyModifiedEvent ModifiedEvent;
+
+	UPROPERTY(BlueprintReadOnly, Category = RigHierarchy, meta = (ScriptName = "ModifiedEvent", AllowPrivateAccess = "true"))
+	FRigHierarchyModifiedDynamicEvent ModifiedEventDynamic;
+	
 	FRigHierarchyMetadataChangedDelegate MetadataChangedDelegate;
 	FRigHierarchyMetadataTagChangedDelegate MetadataTagChangedDelegate;
 	FRigEventDelegate EventDelegate;
