Control Rig: increase procedural element limit to 2000

diff --git a/Engine/Plugins/Animation/ControlRig/Source/ControlRig/Public/Units/RigUnitContext.h b/Engine/Plugins/Animation/ControlRig/Source/ControlRig/Public/Units/RigUnitContext.h
--- a/Engine/Plugins/Animation/ControlRig/Source/ControlRig/Public/Units/RigUnitContext.h
+++ b/Engine/Plugins/Animation/ControlRig/Source/ControlRig/Public/Units/RigUnitContext.h
@@ -29,7 +29,7 @@ struct CONTROLRIG_API FRigHierarchySettings
 	GENERATED_BODY();
 
 	FRigHierarchySettings()
-		: ProceduralElementLimit(128)
+		: ProceduralElementLimit(2000)
 	{
 	}
 
