Control Rig: Fix public functions not loading as public, and showing twice in context menu

diff --git a/Engine/Plugins/Animation/ControlRig/Source/ControlRigDeveloper/Private/ControlRigBlueprint.cpp b/Engine/Plugins/Animation/ControlRig/Source/ControlRigDeveloper/Private/ControlRigBlueprint.cpp
--- a/Engine/Plugins/Animation/ControlRig/Source/ControlRigDeveloper/Private/ControlRigBlueprint.cpp
+++ b/Engine/Plugins/Animation/ControlRig/Source/ControlRigDeveloper/Private/ControlRigBlueprint.cpp
@@ -4626,10 +4626,22 @@ void UControlRigBlueprint::PatchFunctionsOnLoad()
 			}
 		}
 	}
-	
+
+	// Addressing issue where PublicGraphFunctions is populated, but the model PublicFunctionNames is not
+	URigVMFunctionLibrary* FunctionLibrary = GetLocalFunctionLibrary();
+	if (FunctionLibrary)
+	{
+		if (PublicGraphFunctions.Num() > FunctionLibrary->PublicFunctionNames.Num())
+		{
+			for (const FRigVMGraphFunctionHeader& PublicHeader : PublicGraphFunctions)
+			{
+				BackwardsCompatiblePublicFunctions.Add(PublicHeader.Name);
+			}
+		}
+	}
 
 	// Lets rebuild the FunctionStore from the model
-	if (URigVMFunctionLibrary* FunctionLibrary = GetLocalFunctionLibrary())
+	if (FunctionLibrary)
 	{
 		Store.PublicFunctions.Reset();
 		Store.PrivateFunctions.Reset();
diff --git a/Engine/Plugins/Animation/ControlRig/Source/ControlRigEditor/Private/ControlRigEditorModule.cpp b/Engine/Plugins/Animation/ControlRig/Source/ControlRigEditor/Private/ControlRigEditorModule.cpp
--- a/Engine/Plugins/Animation/ControlRig/Source/ControlRigEditor/Private/ControlRigEditorModule.cpp
+++ b/Engine/Plugins/Animation/ControlRig/Source/ControlRigEditor/Private/ControlRigEditorModule.cpp
@@ -1161,16 +1161,18 @@ void FControlRigEditorModule::GetTypeActions(UControlRigBlueprint* CRB, FBluepri
 			}
 			PackagesProcessed.Add(ControlRigAssetData.PackageName);
 
-			FString PublicFunctionsString;
-			if (PublicFunctionsProperty)
-			{
-				PublicFunctionsString = ControlRigAssetData.GetTagValueRef<FString>(PublicFunctionsProperty->GetFName());
-			}
 			FString PublicGraphFunctionsString;
+			FString PublicFunctionsString;
 			if (PublicGraphFunctionsProperty)
 			{
 				PublicGraphFunctionsString = ControlRigAssetData.GetTagValueRef<FString>(PublicGraphFunctionsProperty->GetFName());
 			}
+			// Only look at the deprecated public functions if the PublicGraphFunctionsProperty is not found
+			else if (PublicFunctionsProperty)
+			{
+				PublicFunctionsString = ControlRigAssetData.GetTagValueRef<FString>(PublicFunctionsProperty->GetFName());
+			}
+			
 			if(PublicFunctionsString.IsEmpty() && PublicGraphFunctionsString.IsEmpty())
 			{
 				continue;
@@ -1184,7 +1186,7 @@ void FControlRigEditorModule::GetTypeActions(UControlRigBlueprint* CRB, FBluepri
 				{
 					UBlueprintNodeSpawner* NodeSpawner = UControlRigFunctionRefNodeSpawner::CreateFromAssetData(ControlRigAssetData, PublicFunction);
 					check(NodeSpawner != nullptr);
-					ActionRegistrar.AddBlueprintAction(ActionKey, NodeSpawner);						
+					ActionRegistrar.AddBlueprintAction(ActionKey, NodeSpawner);
 				}
 			}
 
diff --git a/Engine/Source/Developer/RigVMDeveloper/Private/RigVMModel/RigVMController.cpp b/Engine/Source/Developer/RigVMDeveloper/Private/RigVMModel/RigVMController.cpp
--- a/Engine/Source/Developer/RigVMDeveloper/Private/RigVMModel/RigVMController.cpp
+++ b/Engine/Source/Developer/RigVMDeveloper/Private/RigVMModel/RigVMController.cpp
@@ -13851,7 +13851,7 @@ bool URigVMController::CanAddFunctionRefForDefinition(const FRigVMGraphFunctionH
 				InFunctionDefinition.LibraryPointer.HostObject;
 			if (!bIsAvailable)
 			{
-				if (IRigVMGraphFunctionHost* Host = Cast<IRigVMGraphFunctionHost>(InFunctionDefinition.LibraryPointer.HostObject.ResolveObject()))
+				if (IRigVMGraphFunctionHost* Host = Cast<IRigVMGraphFunctionHost>(InFunctionDefinition.LibraryPointer.HostObject.TryLoad()))
 				{
 					bIsAvailable = Host->GetRigVMGraphFunctionStore()->IsFunctionPublic(InFunctionDefinition.LibraryPointer);		
 				}
diff --git a/Engine/Source/Developer/RigVMDeveloper/Public/RigVMModel/RigVMFunctionLibrary.h b/Engine/Source/Developer/RigVMDeveloper/Public/RigVMModel/RigVMFunctionLibrary.h
--- a/Engine/Source/Developer/RigVMDeveloper/Public/RigVMModel/RigVMFunctionLibrary.h
+++ b/Engine/Source/Developer/RigVMDeveloper/Public/RigVMModel/RigVMFunctionLibrary.h
@@ -76,6 +76,7 @@ public:
 
 private:
 
+	UPROPERTY()
 	TArray<FName> PublicFunctionNames;
 
 #if WITH_EDITORONLY_DATA
