CR anim instance should increment update counter such that it won't do an extra tickanimation() with delta time = 0, causing problems with RBAN

#swarm https://p4-swarm.epicgames.net/reviews/23953684

diff --git a/Engine/Plugins/Animation/ControlRig/Source/ControlRig/Private/ControlRigAnimInstance.cpp b/Engine/Plugins/Animation/ControlRig/Source/ControlRig/Private/ControlRigAnimInstance.cpp
--- a/Engine/Plugins/Animation/ControlRig/Source/ControlRig/Private/ControlRigAnimInstance.cpp
+++ b/Engine/Plugins/Animation/ControlRig/Source/ControlRig/Private/ControlRigAnimInstance.cpp
@@ -55,8 +55,8 @@ bool FControlRigAnimInstanceProxy::Evaluate(FPoseContext& Output)
 
 void FControlRigAnimInstanceProxy::UpdateAnimationNode(const FAnimationUpdateContext& InContext)
 {
-	//SnapshotNode.Update_AnyThread(InContext);
-	return;
+	// such that AnimScriptInstance->GetUpdateCounter().HasEverBeenUpdated() is accurate
+	UpdateCounter.Increment();
 }
 
 FControlRigAnimInstanceProxy::~FControlRigAnimInstanceProxy()
