Control Rig Editor: Adding some defensive code to guard against accessing proxy of a invalid preview anim instance

diff --git a/Engine/Plugins/Animation/ControlRig/Source/ControlRigEditor/Private/Editor/ControlRigEditor.cpp b/Engine/Plugins/Animation/ControlRig/Source/ControlRigEditor/Private/Editor/ControlRigEditor.cpp
--- a/Engine/Plugins/Animation/ControlRig/Source/ControlRigEditor/Private/Editor/ControlRigEditor.cpp
+++ b/Engine/Plugins/Animation/ControlRig/Source/ControlRigEditor/Private/Editor/ControlRigEditor.cpp
@@ -5689,7 +5689,7 @@ void FControlRigEditor::ResetAllBoneModification()
 {
 	if (ControlRig)
 	{
-		if (PreviewInstance)
+		if (IsValid(PreviewInstance))
 		{
 			PreviewInstance->ResetModifiedBone();
 		}
