Cherry pick to fix some control rig node fields losing focus on enter and tab.

diff --git a/Engine/Plugins/Animation/ControlRig/Source/ControlRigDeveloper/Private/Graph/ControlRigGraph.cpp b/Engine/Plugins/Animation/ControlRig/Source/ControlRigDeveloper/Private/Graph/ControlRigGraph.cpp
--- a/Engine/Plugins/Animation/ControlRig/Source/ControlRigDeveloper/Private/Graph/ControlRigGraph.cpp
+++ b/Engine/Plugins/Animation/ControlRig/Source/ControlRigDeveloper/Private/Graph/ControlRigGraph.cpp
@@ -866,6 +866,7 @@ void UControlRigGraph::HandleModifiedEvent(ERigVMGraphNotifType InNotifType, URi
 		{
 			if (URigVMNode* ModelNode = Cast<URigVMNode>(InSubject))
 			{
+				ModelNodePathToEdNode.Remove(ModelNode->GetPreviousFName());
 				UEdGraphNode* EdNode = FindNodeForModelNodeName(ModelNode->GetPreviousFName());
 				ModelNodePathToEdNode.Remove(ModelNode->GetPreviousFName());
 				if (UControlRigGraphNode* RigNode = Cast<UControlRigGraphNode>(EdNode))
