Control Rig: Fix ensure when compiling BP_ThirdPersonCharacter

diff --git a/Engine/Plugins/Animation/ControlRig/Source/ControlRig/Private/Rigs/RigHierarchyController.cpp b/Engine/Plugins/Animation/ControlRig/Source/ControlRig/Private/Rigs/RigHierarchyController.cpp
--- a/Engine/Plugins/Animation/ControlRig/Source/ControlRig/Private/Rigs/RigHierarchyController.cpp
+++ b/Engine/Plugins/Animation/ControlRig/Source/ControlRig/Private/Rigs/RigHierarchyController.cpp
@@ -24,6 +24,17 @@ URigHierarchyController::~URigHierarchyController()
 {
 }
 
+void URigHierarchyController::Serialize(FArchive& Ar)
+{
+	Super::Serialize(Ar);
+
+	if (Ar.IsLoading())
+    {
+		URigHierarchy* OuterHierarchy = Cast<URigHierarchy>(GetOuter());
+		SetHierarchy(OuterHierarchy);
+    }
+}
+
 void URigHierarchyController::SetHierarchy(URigHierarchy* InHierarchy)
 {
 	// since we changed the controller to be a property of the hierarchy,
diff --git a/Engine/Plugins/Animation/ControlRig/Source/ControlRig/Public/Rigs/RigHierarchyController.h b/Engine/Plugins/Animation/ControlRig/Source/ControlRig/Public/Rigs/RigHierarchyController.h
--- a/Engine/Plugins/Animation/ControlRig/Source/ControlRig/Public/Rigs/RigHierarchyController.h
+++ b/Engine/Plugins/Animation/ControlRig/Source/ControlRig/Public/Rigs/RigHierarchyController.h
@@ -27,6 +27,9 @@ public:
 
 	virtual ~URigHierarchyController();
 
+	// UObject interface
+	virtual void Serialize(FArchive& Ar) override;
+
 	// Returns the hierarchy currently linked to this controller
 	UFUNCTION(BlueprintCallable, Category = URigHierarchyController)
 	URigHierarchy* GetHierarchy() const
