Control Rig: Fix rigs with procedural hierarchies not showing up in BakeToControlRig lists

diff --git a/Engine/Plugins/Animation/ControlRig/Source/ControlRigEditor/Private/ControlRigEditorModule.cpp b/Engine/Plugins/Animation/ControlRig/Source/ControlRigEditor/Private/ControlRigEditorModule.cpp
--- a/Engine/Plugins/Animation/ControlRig/Source/ControlRigEditor/Private/ControlRigEditorModule.cpp
+++ b/Engine/Plugins/Animation/ControlRig/Source/ControlRigEditor/Private/ControlRigEditorModule.cpp
@@ -575,7 +575,7 @@ TSharedRef< SWidget > FControlRigEditorModule::GenerateAnimationMenu(TWeakPtr<IA
 								Options.bShowUnloadedBlueprints = true;
 								Options.NameTypeToDisplay = EClassViewerNameTypeToDisplay::DisplayName;
 
-								TSharedPtr<FControlRigClassFilter> ClassFilter = MakeShareable(new FControlRigClassFilter(bFilterAssetBySkeleton, true, true, Skeleton));
+								TSharedPtr<FControlRigClassFilter> ClassFilter = MakeShareable(new FControlRigClassFilter(bFilterAssetBySkeleton, false, true, Skeleton));
 								Options.ClassFilters.Add(ClassFilter.ToSharedRef());
 								Options.bShowNoneOption = false;
 
diff --git a/Engine/Plugins/Animation/ControlRig/Source/ControlRigEditor/Private/Sequencer/ControlRigParameterTrackEditor.cpp b/Engine/Plugins/Animation/ControlRig/Source/ControlRigEditor/Private/Sequencer/ControlRigParameterTrackEditor.cpp
--- a/Engine/Plugins/Animation/ControlRig/Source/ControlRigEditor/Private/Sequencer/ControlRigParameterTrackEditor.cpp
+++ b/Engine/Plugins/Animation/ControlRig/Source/ControlRigEditor/Private/Sequencer/ControlRigParameterTrackEditor.cpp
@@ -207,7 +207,7 @@ FControlRigParameterTrackEditor::FControlRigParameterTrackEditor(TSharedRef<ISeq
 	, bIsDoingSelection(false)
 	, bSkipNextSelectionFromTimer(false)
 	, bFilterAssetBySkeleton(true)
-	, bFilterAssetByAnimatableControls(true)
+	, bFilterAssetByAnimatableControls(false)
 	, ControlUndoBracket(0)
 {
 	FMovieSceneToolsModule::Get().RegisterAnimationBakeHelper(this);
@@ -819,7 +819,7 @@ void FControlRigParameterTrackEditor::BakeToControlRigSubMenu(FMenuBuilder& Menu
 		FClassViewerInitializationOptions Options;
 		Options.bShowUnloadedBlueprints = true;
 		Options.NameTypeToDisplay = EClassViewerNameTypeToDisplay::DisplayName;
-		TSharedPtr<FControlRigClassFilter> ClassFilter = MakeShareable(new FControlRigClassFilter(bFilterAssetBySkeleton, true, true, Skeleton));
+		TSharedPtr<FControlRigClassFilter> ClassFilter = MakeShareable(new FControlRigClassFilter(bFilterAssetBySkeleton, false, true, Skeleton));
 		Options.ClassFilters.Add(ClassFilter.ToSharedRef());
 		Options.bShowNoneOption = false;
 
