ControlRig Editor details customization: making sure the customization instance does not hold strong reference to the object being customized, otherwise, it will prevent the things like a deubgged PIE rig from getting gc when we end PIE, as the customization is not destructed( release the reference ) until 1 slate tick later.

diff --git a/Engine/Plugins/Animation/ControlRig/Source/ControlRigEditor/Private/ControlRigGraphDetails.cpp b/Engine/Plugins/Animation/ControlRig/Source/ControlRigEditor/Private/ControlRigGraphDetails.cpp
--- a/Engine/Plugins/Animation/ControlRig/Source/ControlRigEditor/Private/ControlRigGraphDetails.cpp
+++ b/Engine/Plugins/Animation/ControlRig/Source/ControlRigEditor/Private/ControlRigGraphDetails.cpp
@@ -1868,14 +1868,4 @@ void FControlRigGraphMathTypeDetails::CustomizeChildren(TSharedRef<IPropertyHand
 	}
 }
 
-void FControlRigGraphMathTypeDetails::AddReferencedObjects(FReferenceCollector& Collector)
-{
-	Collector.AddReferencedObjects(ObjectsBeingCustomized);
-}
-
-FString FControlRigGraphMathTypeDetails::GetReferencerName() const
-{
-	return TEXT("FControlRigGraphMathTypeDetails:") + ObjectsBeingCustomized[0]->GetPathName();
-}
-
 #undef LOCTEXT_NAMESPACE
diff --git a/Engine/Plugins/Animation/ControlRig/Source/ControlRigEditor/Private/ControlRigGraphDetails.h b/Engine/Plugins/Animation/ControlRig/Source/ControlRigEditor/Private/ControlRigGraphDetails.h
--- a/Engine/Plugins/Animation/ControlRig/Source/ControlRigEditor/Private/ControlRigGraphDetails.h
+++ b/Engine/Plugins/Animation/ControlRig/Source/ControlRigEditor/Private/ControlRigGraphDetails.h
@@ -239,7 +239,7 @@ public:
 };
 
 /** Customization for editing a Control Rig node */
-class FControlRigGraphMathTypeDetails : public IPropertyTypeCustomization, public FGCObject
+class FControlRigGraphMathTypeDetails : public IPropertyTypeCustomization
 {
 public:
 
@@ -254,10 +254,6 @@ public:
 	virtual void CustomizeHeader(TSharedRef<class IPropertyHandle> InPropertyHandle, class FDetailWidgetRow& HeaderRow, IPropertyTypeCustomizationUtils& StructCustomizationUtils) override;
 	virtual void CustomizeChildren(TSharedRef<class IPropertyHandle> InPropertyHandle, class IDetailChildrenBuilder& StructBuilder, IPropertyTypeCustomizationUtils& StructCustomizationUtils) override;
 
-	/** FGCObject interface */
-	virtual void AddReferencedObjects(FReferenceCollector& Collector) override;
-	virtual FString GetReferencerName() const override;
-
 protected:
 
 	bool GetPropertyChain(TSharedRef<class IPropertyHandle> InPropertyHandle, FEditPropertyChain& OutPropertyChain, TArray<int32> &OutPropertyArrayIndices, bool& bOutEnabled)
@@ -267,7 +263,7 @@ protected:
 		bOutEnabled = false;
 		if (!ObjectsBeingCustomized.IsEmpty())
 		{
-			if (ObjectsBeingCustomized[0])
+			if (ObjectsBeingCustomized[0].Get())
 			{
 				TSharedPtr<class IPropertyHandle> ChainHandle = InPropertyHandle;
 				while (ChainHandle.IsValid() && ChainHandle->GetProperty() != nullptr)
@@ -352,9 +348,9 @@ protected:
 			return Result;
 		}
 	
-		for(const TObjectPtr<UObject>& Object : ObjectsBeingCustomized)
+		for(const TWeakObjectPtr<UObject>& Object : ObjectsBeingCustomized)
 		{
-			if(Object && InPropertyHandle->IsValidHandle())
+			if(Object.Get() && InPropertyHandle->IsValidHandle())
 			{
 				static VectorType ZeroVector = VectorType();
 				const VectorType& Vector = ContainerUObjectToValueRef<VectorType>(Object.Get(), ZeroVector, PropertyChain, PropertyArrayIndices);
@@ -390,8 +386,8 @@ protected:
 		TArray<UObject*> ObjectsView;
 		for(int32 Index = 0; Index < ObjectsBeingCustomized.Num(); Index++)
 		{
-			const TObjectPtr<UObject>& Object = ObjectsBeingCustomized[Index];
-			if (Object)
+			const TWeakObjectPtr<UObject>& Object = ObjectsBeingCustomized[Index];
+			if (Object.Get())
 			{
 				ObjectsView.Add(Object.Get());
 			}
@@ -411,8 +407,8 @@ protected:
 
 		for(int32 Index = 0; Index < ObjectsBeingCustomized.Num(); Index++)
 		{
-			const TObjectPtr<UObject>& Object = ObjectsBeingCustomized[Index];
-			if(Object && InPropertyHandle->IsValidHandle())
+			const TWeakObjectPtr<UObject>& Object = ObjectsBeingCustomized[Index];
+			if(Object.Get() && InPropertyHandle->IsValidHandle())
 			{
 				static VectorType ZeroVector = VectorType();
 				VectorType& Vector = ContainerUObjectToValueRef<VectorType>(Object.Get(), ZeroVector, PropertyChain, PropertyArrayIndices);
@@ -512,9 +508,9 @@ protected:
 			return Result;
 		}
 		
-		for(const TObjectPtr<UObject>& Object : ObjectsBeingCustomized)
+		for(const TWeakObjectPtr<UObject>& Object : ObjectsBeingCustomized)
 		{
-			if(Object && InPropertyHandle->IsValidHandle())
+			if(Object.Get() && InPropertyHandle->IsValidHandle())
 			{
 				static RotationType ZeroRotation = RotationType();
 				const RotationType& Rotation = ContainerUObjectToValueRef<RotationType>(Object.Get(), ZeroRotation, PropertyChain, PropertyArrayIndices);
@@ -549,8 +545,8 @@ protected:
 		TArray<UObject*> ObjectsView;
 		for(int32 Index = 0; Index < ObjectsBeingCustomized.Num(); Index++)
 		{
-			const TObjectPtr<UObject>& Object = ObjectsBeingCustomized[Index];
-			if (Object)
+			const TWeakObjectPtr<UObject>& Object = ObjectsBeingCustomized[Index];
+			if (Object.Get())
 			{
 				ObjectsView.Add(Object.Get());
 			}
@@ -570,8 +566,8 @@ protected:
 
 		for(int32 Index = 0; Index < ObjectsBeingCustomized.Num(); Index++)
 		{
-			const TObjectPtr<UObject>& Object = ObjectsBeingCustomized[Index];
-			if(Object && InPropertyHandle->IsValidHandle())
+			const TWeakObjectPtr<UObject>& Object = ObjectsBeingCustomized[Index];
+			if(Object.Get() && InPropertyHandle->IsValidHandle())
 			{
 				static RotationType ZeroRotation = RotationType();
 				RotationType& Rotation = ContainerUObjectToValueRef<RotationType>(Object.Get(), ZeroRotation, PropertyChain, PropertyArrayIndices);
@@ -665,9 +661,9 @@ protected:
 				return false;
 			}
 			
-			for(const TObjectPtr<UObject>& Object : ObjectsBeingCustomized)
+			for(const TWeakObjectPtr<UObject>& Object : ObjectsBeingCustomized)
 			{
-				if(Object && InPropertyHandle->IsValidHandle())
+				if(Object.Get() && InPropertyHandle->IsValidHandle())
 				{
 					const TransformType& Transform = ContainerUObjectToValueRef<TransformType>(Object.Get(), Identity, PropertyChain, PropertyArrayIndices);
 
@@ -721,9 +717,9 @@ protected:
 				return Result;
 			}
 			
-			for(const TObjectPtr<UObject>& Object : ObjectsBeingCustomized)
+			for(const TWeakObjectPtr<UObject>& Object : ObjectsBeingCustomized)
 			{
-				if(Object && InPropertyHandle->IsValidHandle())
+				if(Object.Get() && InPropertyHandle->IsValidHandle())
 				{
 					const TransformType& Transform = ContainerUObjectToValueRef<TransformType>(Object.Get(), Identity, PropertyChain, PropertyArrayIndices);
 					
@@ -773,8 +769,8 @@ protected:
 			TArray<UObject*> ObjectsView;
 			for(int32 Index = 0; Index < ObjectsBeingCustomized.Num(); Index++)
 			{
-				const TObjectPtr<UObject>& Object = ObjectsBeingCustomized[Index];
-				if (Object)
+				const TWeakObjectPtr<UObject>& Object = ObjectsBeingCustomized[Index];
+				if (Object.Get())
 				{
 					ObjectsView.Add(Object.Get());
 				}
@@ -794,8 +790,8 @@ protected:
 
 			for(int32 Index = 0; Index < ObjectsBeingCustomized.Num(); Index++)
 			{
-				const TObjectPtr<UObject>& Object = ObjectsBeingCustomized[Index];
-				if(Object && InPropertyHandle->IsValidHandle())
+				const TWeakObjectPtr<UObject>& Object = ObjectsBeingCustomized[Index];
+				if(Object.Get() && InPropertyHandle->IsValidHandle())
 				{
 					TransformType& Transform = ContainerUObjectToValueRef<TransformType>(Object.Get(), Identity, PropertyChain, PropertyArrayIndices);
 					TransformType PreviousTransform = Transform;
@@ -863,8 +859,8 @@ protected:
 			TArray<UObject*> ObjectsView;
 			for(int32 Index = 0; Index < ObjectsBeingCustomized.Num(); Index++)
 			{
-				const TObjectPtr<UObject>& Object = ObjectsBeingCustomized[Index];
-				if (Object)
+				const TWeakObjectPtr<UObject>& Object = ObjectsBeingCustomized[Index];
+				if (Object.Get())
 				{
 					ObjectsView.Add(Object.Get());
 				}
@@ -874,8 +870,8 @@ protected:
 			
 			for(int32 Index = 0; Index < ObjectsBeingCustomized.Num(); Index++)
 			{
-				const TObjectPtr<UObject>& Object = ObjectsBeingCustomized[Index];
-				if(Object && InPropertyHandle->IsValidHandle())
+				const TWeakObjectPtr<UObject>& Object = ObjectsBeingCustomized[Index];
+				if(Object.Get() && InPropertyHandle->IsValidHandle())
 				{
 					TransformType& Transform = ContainerUObjectToValueRef<TransformType>(Object.Get(), Identity, PropertyChain, PropertyArrayIndices);
 					TransformType PreviousTransform = Transform;
@@ -930,7 +926,7 @@ protected:
 	UScriptStruct* ScriptStruct;
 	UControlRigBlueprint* BlueprintBeingCustomized;
 	URigVMGraph* GraphBeingCustomized;
-	TArray<TObjectPtr<UObject>> ObjectsBeingCustomized; 
+	TArray<TWeakObjectPtr<UObject>> ObjectsBeingCustomized; 
 };
 
 template<>
