Asset Definition: Convert ControlRigPose

diff --git a/Engine/Plugins/Animation/ControlRig/Source/ControlRigEditor/ControlRigEditor.Build.cs b/Engine/Plugins/Animation/ControlRig/Source/ControlRigEditor/ControlRigEditor.Build.cs
--- a/Engine/Plugins/Animation/ControlRig/Source/ControlRigEditor/ControlRigEditor.Build.cs
+++ b/Engine/Plugins/Animation/ControlRig/Source/ControlRigEditor/ControlRigEditor.Build.cs
@@ -31,6 +31,7 @@ namespace UnrealBuildTool.Rules
             PrivateDependencyModuleNames.AddRange(
                 new string[]
                 {
+					"AssetDefinition",
                     "Core",
                     "CoreUObject",
 					"CurveEditor",
diff --git a/Engine/Plugins/Animation/ControlRig/Source/ControlRigEditor/Private/ControlRigEditorModule.cpp b/Engine/Plugins/Animation/ControlRig/Source/ControlRigEditor/Private/ControlRigEditorModule.cpp
--- a/Engine/Plugins/Animation/ControlRig/Source/ControlRigEditor/Private/ControlRigEditorModule.cpp
+++ b/Engine/Plugins/Animation/ControlRig/Source/ControlRigEditor/Private/ControlRigEditorModule.cpp
@@ -109,7 +109,6 @@
 #include "Units/Execution/RigUnit_InverseExecution.h"
 #include "Graph/SControlRigGraphPinVariableBinding.h"
 #include "Graph/SControlRigGraphChangePinType.h"
-#include "Tools/AssetTypeActions_ControlRigPose.h"
 #include "ControlRigBlueprintFactory.h"
 #include "ControlRigPythonLogDetails.h"
 #include "Dialog/SCustomDialog.h"
@@ -206,7 +205,6 @@ void FControlRigEditorModule::StartupModule()
 
 	RegisterAssetTypeAction(MakeShareable(new FControlRigBlueprintActions()));
 	RegisterAssetTypeAction(MakeShareable(new FControlRigShapeLibraryActions()));
-	RegisterAssetTypeAction(MakeShareable(new FAssetTypeActions_ControlRigPose()));
 	
 	// Register sequencer track editor
 	ISequencerModule& SequencerModule = FModuleManager::Get().LoadModuleChecked<ISequencerModule>("Sequencer");
diff --git a/Engine/Plugins/Animation/ControlRig/Source/ControlRigEditor/Private/Tools/AssetDefinition_ControlRigPose.cpp b/Engine/Plugins/Animation/ControlRig/Source/ControlRigEditor/Private/Tools/AssetDefinition_ControlRigPose.cpp
new file mode 100644
--- /dev/null
+++ b/Engine/Plugins/Animation/ControlRig/Source/ControlRigEditor/Private/Tools/AssetDefinition_ControlRigPose.cpp
@@ -0,0 +1,154 @@
+// Copyright Epic Games, Inc. All Rights Reserved.
+
+#include "AssetDefinition_ControlRigPose.h"
+
+#include "ContentBrowserModule.h"
+#include "ContentBrowserMenuContexts.h"
+#include "ControlRig.h"
+#include "EditorModeManager.h"
+#include "EditMode/ControlRigEditMode.h"
+#include "EditMode/SControlRigRenamePoseControls.h"
+#include "ToolMenus.h"
+#include "ToolMenu.h"
+#include "ToolMenuSection.h"
+#include "Tools/ControlRigPose.h"
+#include "UnrealEdGlobals.h"
+
+#define LOCTEXT_NAMESPACE "AssetTypeActions"
+
+// Menu Extensions
+//--------------------------------------------------------------------
+
+namespace MenuExtension_AnimationAsset
+{
+	void ExecutePastePose(const FToolMenuContext& InContext)
+	{
+		const UContentBrowserAssetContextMenuContext* Context = UContentBrowserAssetContextMenuContext::FindContextWithAssets(InContext);
+		TArray<UControlRigPoseAsset*> ControlRigPoseAssets = Context->LoadSelectedObjects<UControlRigPoseAsset>();
+
+		FControlRigEditMode* ControlRigEditMode = static_cast<FControlRigEditMode*>(GLevelEditorModeTools().GetActiveMode(FControlRigEditMode::ModeName));
+		if (!ControlRigEditMode)
+		{
+			return;
+		}
+		
+		TMap<UControlRig*, TArray<FRigElementKey>> AllSelectedControls;
+		ControlRigEditMode->GetAllSelectedControls(AllSelectedControls);
+		TArray<UControlRig*> ControlRigs;
+		AllSelectedControls.GenerateKeyArray(ControlRigs);
+			
+		for (UControlRigPoseAsset* ControlRigPoseAsset : ControlRigPoseAssets)
+		{
+			for (UControlRig* ControlRig : ControlRigs)
+			{
+				ControlRigPoseAsset->PastePose(ControlRig, false, false);
+			}
+		}
+	}
+
+	void ExecuteSelectControls(const FToolMenuContext& InContext)
+	{
+		const UContentBrowserAssetContextMenuContext* Context = UContentBrowserAssetContextMenuContext::FindContextWithAssets(InContext);
+		TArray<UControlRigPoseAsset*> ControlRigPoseAssets = Context->LoadSelectedObjects<UControlRigPoseAsset>();
+
+		FControlRigEditMode* ControlRigEditMode = static_cast<FControlRigEditMode*>(GLevelEditorModeTools().GetActiveMode(FControlRigEditMode::ModeName));
+		if (!ControlRigEditMode)
+		{
+			return;
+		}
+
+		TMap<UControlRig*, TArray<FRigElementKey>> AllSelectedControls;
+		ControlRigEditMode->GetAllSelectedControls(AllSelectedControls);
+		TArray<UControlRig*> ControlRigs;
+		AllSelectedControls.GenerateKeyArray(ControlRigs);
+
+		for (UControlRigPoseAsset* ControlRigPoseAsset : ControlRigPoseAssets)
+		{
+			for (UControlRig* ControlRig : ControlRigs)
+			{
+				ControlRigPoseAsset->SelectControls(ControlRig);
+			}
+		}
+	}
+
+	void ExecuteUpdatePose(const FToolMenuContext& InContext)
+	{
+		const UContentBrowserAssetContextMenuContext* Context = UContentBrowserAssetContextMenuContext::FindContextWithAssets(InContext);
+		TArray<UControlRigPoseAsset*> ControlRigPoseAssets = Context->LoadSelectedObjects<UControlRigPoseAsset>();
+
+		FControlRigEditMode* ControlRigEditMode = static_cast<FControlRigEditMode*>(GLevelEditorModeTools().GetActiveMode(FControlRigEditMode::ModeName));
+		if (!ControlRigEditMode)
+		{
+			return;
+		}
+
+		TMap<UControlRig*, TArray<FRigElementKey>> AllSelectedControls;
+		ControlRigEditMode->GetAllSelectedControls(AllSelectedControls);
+		TArray<UControlRig*> ControlRigs;
+		AllSelectedControls.GenerateKeyArray(ControlRigs);
+
+		for (UControlRigPoseAsset* ControlRigPoseAsset : ControlRigPoseAssets)
+		{
+			for (UControlRig* ControlRig : ControlRigs)
+			{
+				ControlRigPoseAsset->SavePose(ControlRig, false);
+			}
+		}
+	}
+
+	void ExecuteRenameControls(const FToolMenuContext& InContext)
+	{
+		const UContentBrowserAssetContextMenuContext* Context = UContentBrowserAssetContextMenuContext::FindContextWithAssets(InContext);
+		TArray<UControlRigPoseAsset*> ControlRigPoseAssets = Context->LoadSelectedObjects<UControlRigPoseAsset>();
+
+		FControlRigRenameControlsDialog::RenameControls(ControlRigPoseAssets);
+	}
+
+	static FDelayedAutoRegisterHelper DelayedAutoRegister(EDelayedRegisterRunPhase::EndOfEngineInit, [] {
+		UToolMenus::RegisterStartupCallback(FSimpleMulticastDelegate::FDelegate::CreateLambda([]()
+	{
+		FToolMenuOwnerScoped OwnerScoped(UE_MODULE_NAME);
+		UToolMenu* Menu = UE::ContentBrowser::ExtendToolMenu_AssetContextMenu(UControlRigPoseAsset::StaticClass());
+
+		FToolMenuSection& Section = Menu->FindOrAddSection("GetAssetActions");
+		Section.AddDynamicEntry(NAME_None, FNewToolMenuSectionDelegate::CreateLambda([](FToolMenuSection& InSection)
+		{
+			{
+				const TAttribute<FText> Label = LOCTEXT("ControlRigPose_PastePose", "Paste Pose");
+				const TAttribute<FText> ToolTip = LOCTEXT("ControlRigPose_PastePoseTooltip", "Paste the selected pose");
+
+				FToolUIAction UIAction;
+				UIAction.ExecuteAction = FToolMenuExecuteAction::CreateStatic(&ExecutePastePose);
+				InSection.AddMenuEntry("ControlRigPose_PastePose", Label, ToolTip, FSlateIcon(), UIAction);
+			}
+			{
+				const TAttribute<FText> Label = LOCTEXT("ControlRigPose_SelectControls", "Select Controls");
+				const TAttribute<FText> ToolTip = LOCTEXT("ControlRigPose_SelectControlsTooltip", "Select controls in this pose on active control rig");
+
+				FToolUIAction UIAction;
+				UIAction.ExecuteAction = FToolMenuExecuteAction::CreateStatic(&ExecuteSelectControls);
+				InSection.AddMenuEntry("ControlRigPose_SelectControls", Label, ToolTip, FSlateIcon(), UIAction);
+			}
+			{
+				const TAttribute<FText> Label = LOCTEXT("ControlRigPose_UpdatePose", "Update Pose");
+				const TAttribute<FText> ToolTip = LOCTEXT("ControlRigPose_UpdatePoseTooltip", "Update the pose based upon current control rig pose and selected controls");
+
+				FToolUIAction UIAction;
+				UIAction.ExecuteAction = FToolMenuExecuteAction::CreateStatic(&ExecuteUpdatePose);
+				InSection.AddMenuEntry("ControlRigPose_UpdatePose", Label, ToolTip, FSlateIcon(), UIAction);
+			}
+			{
+				const TAttribute<FText> Label = LOCTEXT("ControlRigPose_RenameControls", "Rename Controls");
+				const TAttribute<FText> ToolTip = LOCTEXT("ControlRigPose_RenameControlsTooltip", "Rename controls on selected poses");
+
+				FToolUIAction UIAction;
+				UIAction.ExecuteAction = FToolMenuExecuteAction::CreateStatic(&ExecuteRenameControls);
+				InSection.AddMenuEntry("ControlRigPose_RenameControls", Label, ToolTip, FSlateIcon(), UIAction);
+			}
+
+		}));
+		}));
+	});
+}
+
+#undef LOCTEXT_NAMESPACE
diff --git a/Engine/Plugins/Animation/ControlRig/Source/ControlRigEditor/Private/Tools/AssetDefinition_ControlRigPose.h b/Engine/Plugins/Animation/ControlRig/Source/ControlRigEditor/Private/Tools/AssetDefinition_ControlRigPose.h
new file mode 100644
--- /dev/null
+++ b/Engine/Plugins/Animation/ControlRig/Source/ControlRigEditor/Private/Tools/AssetDefinition_ControlRigPose.h
@@ -0,0 +1,26 @@
+// Copyright Epic Games, Inc. All Rights Reserved.
+
+#pragma once
+
+#include "Tools/ControlRigPose.h"
+#include "AssetDefinitionDefault.h"
+
+#include "AssetDefinition_ControlRigPose.generated.h"
+
+UCLASS()
+class UAssetDefinition_ControlRigPose : public UAssetDefinitionDefault
+{
+	GENERATED_BODY()
+
+public:
+	// UAssetDefinition Begin
+	virtual FText GetAssetDisplayName() const override { return NSLOCTEXT("AssetTypeActions", "AssetTypeActions_ControlRigPose", "Control Rig Pose"); }
+	virtual FLinearColor GetAssetColor() const override { return FLinearColor(FColor(222, 128, 64)); }
+	virtual TSoftClassPtr<UObject> GetAssetClass() const override { return UControlRigPoseAsset::StaticClass(); }
+	virtual TConstArrayView<FAssetCategoryPath> GetAssetCategories() const override
+	{
+		static const auto Categories = { EAssetCategoryPaths::Animation };
+		return Categories;
+	}
+	// UAssetDefinition End
+};
diff --git a/Engine/Plugins/Animation/ControlRig/Source/ControlRigEditor/Private/Tools/AssetTypeActions_ControlRigPose.cpp b/Engine/Plugins/Animation/ControlRig/Source/ControlRigEditor/Private/Tools/AssetTypeActions_ControlRigPose.cpp
deleted file mode 100644
diff --git a/Engine/Plugins/Animation/ControlRig/Source/ControlRigEditor/Private/Tools/AssetTypeActions_ControlRigPose.h b/Engine/Plugins/Animation/ControlRig/Source/ControlRigEditor/Private/Tools/AssetTypeActions_ControlRigPose.h
deleted file mode 100644
