Control Rig: Avoid reconstructing EdGraph nodes during cook

diff --git a/Engine/Plugins/Animation/ControlRig/Source/ControlRigDeveloper/Private/ControlRigBlueprintUtils.cpp b/Engine/Plugins/Animation/ControlRig/Source/ControlRigDeveloper/Private/ControlRigBlueprintUtils.cpp
--- a/Engine/Plugins/Animation/ControlRig/Source/ControlRigDeveloper/Private/ControlRigBlueprintUtils.cpp
+++ b/Engine/Plugins/Animation/ControlRig/Source/ControlRigDeveloper/Private/ControlRigBlueprintUtils.cpp
@@ -98,7 +98,13 @@ void FControlRigBlueprintUtils::HandleReconstructAllNodes(UBlueprint* InBlueprin
 void FControlRigBlueprintUtils::HandleRefreshAllNodes(UBlueprint* InBlueprint)
 {
 	DECLARE_SCOPE_HIERARCHICAL_COUNTER_FUNC()
-
+#ifdef WITH_EDITORONLY_DATA
+	// Avoid refreshing EdGraph nodes during cook
+	if (GIsCookerLoadingPackage)
+	{
+		return;
+	}
+	
 	if(UControlRigBlueprint* RigBlueprint = Cast<UControlRigBlueprint>(InBlueprint))
 	{
 		if (RigBlueprint->GetDefaultModel() == nullptr)
@@ -124,6 +130,7 @@ void FControlRigBlueprintUtils::HandleRefreshAllNodes(UBlueprint* InBlueprint)
 			Node->ClearFlags(RF_Transient);
 		}
 	}
+#endif
 }
 
 void FControlRigBlueprintUtils::RemoveMemberVariableIfNotUsed(UBlueprint* Blueprint, const FName VarName, UControlRigGraphNode* ToBeDeleted)
