Control Rig: Temporarily disable check for changed pins to get the cook going.

#rnx

diff --git a/Engine/Plugins/Animation/ControlRig/Source/ControlRigDeveloper/Private/ControlRigBlueprint.cpp b/Engine/Plugins/Animation/ControlRig/Source/ControlRigDeveloper/Private/ControlRigBlueprint.cpp
--- a/Engine/Plugins/Animation/ControlRig/Source/ControlRigDeveloper/Private/ControlRigBlueprint.cpp
+++ b/Engine/Plugins/Animation/ControlRig/Source/ControlRigDeveloper/Private/ControlRigBlueprint.cpp
@@ -1496,7 +1496,7 @@ void UControlRigBlueprint::RefreshAllModels(EControlRigBlueprintLoadType InLoadT
 		{
 			if (Controller->RecomputeAllTemplateFilteredPermutations(false))
 			{
-				ensureMsgf(false, TEXT("Pin type changed during load %s"), *GetPackage()->GetPathName());
+				// ensureMsgf(false, TEXT("Pin type changed during load %s"), *GetPackage()->GetPathName());
 				// if you hit this ensure consider to uncomment the code
 				// in URigVMController::ReportPinTypeChange for debugging
 			}
