RigHierarchy.cpp: Fix static analysis

diff --git a/Engine/Plugins/Animation/ControlRig/Source/ControlRig/Private/Rigs/RigHierarchy.cpp b/Engine/Plugins/Animation/ControlRig/Source/ControlRig/Private/Rigs/RigHierarchy.cpp
--- a/Engine/Plugins/Animation/ControlRig/Source/ControlRig/Private/Rigs/RigHierarchy.cpp
+++ b/Engine/Plugins/Animation/ControlRig/Source/ControlRig/Private/Rigs/RigHierarchy.cpp
@@ -2456,7 +2456,12 @@ void URigHierarchy::Notify(ERigHierarchyNotification InNotifType, const FRigBase
 		ModifiedEvent.Broadcast(InNotifType, this, InElement);
 		if(ModifiedEventDynamic.IsBound())
 		{
-			ModifiedEventDynamic.Broadcast(InNotifType, this, InElement->GetKey());
+			FRigElementKey Key;
+			if(InElement)
+			{
+				Key = InElement->GetKey();
+			}
+			ModifiedEventDynamic.Broadcast(InNotifType, this, Key);
 		}
 	}
 	else
