Control Rig: Fix patching of if/select nodes unresolves pins

diff --git a/Engine/Plugins/Animation/ControlRig/Source/ControlRigDeveloper/Private/ControlRigBlueprint.cpp b/Engine/Plugins/Animation/ControlRig/Source/ControlRigDeveloper/Private/ControlRigBlueprint.cpp
--- a/Engine/Plugins/Animation/ControlRig/Source/ControlRigDeveloper/Private/ControlRigBlueprint.cpp
+++ b/Engine/Plugins/Animation/ControlRig/Source/ControlRigDeveloper/Private/ControlRigBlueprint.cpp
@@ -1496,7 +1496,7 @@ void UControlRigBlueprint::RefreshAllModels(EControlRigBlueprintLoadType InLoadT
 		{
 			if (Controller->RecomputeAllTemplateFilteredPermutations(false))
 			{
-				// ensureMsgf(false, TEXT("Pin type changed during load %s"), *GetPackage()->GetPathName());
+				ensureMsgf(false, TEXT("Pin type changed during load %s"), *GetPackage()->GetPathName());
 				// if you hit this ensure consider to uncomment the code
 				// in URigVMController::ReportPinTypeChange for debugging
 			}
diff --git a/Engine/Source/Developer/RigVMDeveloper/Private/RigVMModel/RigVMController.cpp b/Engine/Source/Developer/RigVMDeveloper/Private/RigVMModel/RigVMController.cpp
--- a/Engine/Source/Developer/RigVMDeveloper/Private/RigVMModel/RigVMController.cpp
+++ b/Engine/Source/Developer/RigVMDeveloper/Private/RigVMModel/RigVMController.cpp
@@ -19621,12 +19621,12 @@ FRigVMClientPatchResult URigVMController::PatchIfSelectNodesOnLoad()
 						{
 							if(!FRigVMRegistry::Get().IsWildCardType(*ResolvedTypeIndex))
 							{
-								ChangePinType(Pin, *ResolvedTypeIndex, false, false);
 								NewNode->UpdateFilteredPermutations(Pin, {*ResolvedTypeIndex});
 							}
 						}
 					}
 				}
+				UpdateTemplateNodePinTypes(NewNode, false, false);
 			}
 
 			ApplyPinStates(NewNode, PinStates, {}, false);
