Persona: Fix for Edit With Sequencer ending up with 2 Control Rig Tracks due to the Default Animation Rig one not getting deleted. This is a regression from something, maybe before we would add the skel mesh comp as a binding which deleting would handle this.

diff --git a/Engine/Plugins/Animation/ControlRig/Source/ControlRigEditor/Private/ControlRigEditorModule.cpp b/Engine/Plugins/Animation/ControlRig/Source/ControlRigEditor/Private/ControlRigEditorModule.cpp
--- a/Engine/Plugins/Animation/ControlRig/Source/ControlRigEditor/Private/ControlRigEditorModule.cpp
+++ b/Engine/Plugins/Animation/ControlRig/Source/ControlRigEditor/Private/ControlRigEditorModule.cpp
@@ -711,6 +711,7 @@ void FControlRigEditorModule::BakeToControlRig(UClass* ControlRigClass, UAnimSeq
 			}
 
 			//Delete binding from default animating rig
+			//if we have skel mesh component binding we can just delete that
 			FGuid CompGuid = WeakSequencer.Pin()->FindObjectId(*(MeshActor->GetSkeletalMeshComponent()), WeakSequencer.Pin()->GetFocusedTemplateID());
 			if (CompGuid.IsValid())
 			{
@@ -719,6 +720,13 @@ void FControlRigEditorModule::BakeToControlRig(UClass* ControlRigClass, UAnimSeq
 					MovieScene->RemoveSpawnable(CompGuid);
 				}
 			}
+			else //otherwise if not delete the track
+			{
+				if (UMovieSceneTrack* ExistingTrack = MovieScene->FindTrack<UMovieSceneControlRigParameterTrack>(ActorTrackGuid))
+				{
+					MovieScene->RemoveTrack(*ExistingTrack);
+				}
+			}
 
 			UMovieSceneControlRigParameterTrack* Track = MovieScene->AddTrack<UMovieSceneControlRigParameterTrack>(ActorTrackGuid);
 			if (Track)
