Sequencer: Constraints: Fix for constraints breaking after PIE do to recent fixes for duplication.

diff --git a/Engine/Plugins/Animation/ControlRig/Source/ControlRigEditor/Private/Editor/Constraints/SConstraintsWidget.cpp b/Engine/Plugins/Animation/ControlRig/Source/ControlRigEditor/Private/Editor/Constraints/SConstraintsWidget.cpp
--- a/Engine/Plugins/Animation/ControlRig/Source/ControlRigEditor/Private/Editor/Constraints/SConstraintsWidget.cpp
+++ b/Engine/Plugins/Animation/ControlRig/Source/ControlRigEditor/Private/Editor/Constraints/SConstraintsWidget.cpp
@@ -496,7 +496,7 @@ void SEditableConstraintItem::Construct(
 		const FConstraintsManagerController& Controller = FConstraintsManagerController::Get(World);
 		return Controller.GetConstraint(ConstraintItem->Name);
 	};
-	UTickableConstraint* Constraint = GetConstraint();
+	TWeakObjectPtr<UTickableConstraint> Constraint = GetConstraint();
 
 	// sequencer
 	TWeakPtr<ISequencer> WeakSequencer = GetSequencerChecked();
@@ -510,7 +510,7 @@ void SEditableConstraintItem::Construct(
 	}
 	
 	FString ParentFullLabel = ParentLabel, ChildFullLabel = ChildLabel;
-	if (IsValid(Constraint))
+	if (IsValid(Constraint.Get()))
 	{
 		Constraint->GetFullLabel().Split(TEXT("."), &ParentFullLabel, &ChildFullLabel);
 	}
@@ -532,7 +532,7 @@ void SEditableConstraintItem::Construct(
 			.BorderImage(RoundedBoxBrush)
 			.BorderBackgroundColor_Lambda([Constraint]()
 			{
-				if (!IsValid(Constraint))
+				if (!Constraint.IsValid() || !IsValid(Constraint.Get()))
 				{
 					return FStyleColors::Transparent;
 				}
@@ -568,7 +568,7 @@ void SEditableConstraintItem::Construct(
 					})
 					.Font_Lambda([Constraint]()
 					{
-						if (!IsValid(Constraint))
+						if (!Constraint.IsValid() || !IsValid(Constraint.Get()))
 						{
 							return IDetailLayoutBuilder::GetDetailFont();
 						}
@@ -576,7 +576,7 @@ void SEditableConstraintItem::Construct(
 					})
 					.ToolTipText_Lambda( [Constraint, ParentFullLabel, ChildFullLabel]()
 					{
-						if (!IsValid(Constraint))
+						if (!Constraint.IsValid() || !IsValid(Constraint.Get()))
 						{
 							return FText();
 						}
@@ -608,7 +608,7 @@ void SEditableConstraintItem::Construct(
 			.ContentPadding(0)
 			.OnClicked_Lambda(	[Constraint, WeakSequencer]()
 			{
-				if (!IsValid(Constraint))
+				if (!Constraint.IsValid() || !IsValid(Constraint.Get()))
 				{
 					return FReply::Handled();
 				}
diff --git a/Engine/Source/Runtime/Experimental/Animation/Constraints/Private/ConstraintsManager.cpp b/Engine/Source/Runtime/Experimental/Animation/Constraints/Private/ConstraintsManager.cpp
--- a/Engine/Source/Runtime/Experimental/Animation/Constraints/Private/ConstraintsManager.cpp
+++ b/Engine/Source/Runtime/Experimental/Animation/Constraints/Private/ConstraintsManager.cpp
@@ -8,6 +8,7 @@
 
 #include "Engine/World.h"
 #include "Engine/Level.h"
+#include "Engine/Engine.h"
 
 #include UE_INLINE_GENERATED_CPP_BY_NAME(ConstraintsManager)
 
@@ -344,6 +345,35 @@ FConstraintsManagerController& FConstraintsManagerController::Get(UWorld* InWorl
 	return Singleton;
 }
 
+bool FConstraintsManagerController::DoesExistInAnyWorld(UTickableConstraint* InConstraint) 
+{
+	bool bFound = false;
+	if (InConstraint)
+	{
+		UConstraintsManager* Outer = InConstraint->GetTypedOuter<UConstraintsManager>();
+		if (Outer)
+		{
+			UWorld* CurrentWorld = World; //save current World
+			for (const FWorldContext& Context : GEngine->GetWorldContexts())
+			{
+				World = Context.World();
+				UConstraintsManager* Manager = FindManager();
+				if (Manager)
+				{
+					if (Manager->Constraints.Find(InConstraint) != INDEX_NONE)
+					{
+						bFound = true;
+						break;
+					}
+				}
+			}
+			World = CurrentWorld; //restore current World
+		}
+	}
+	return bFound;
+}
+
+
 UConstraintsManager* FConstraintsManagerController::GetManager() const
 {
 	if (!World)
@@ -432,14 +462,7 @@ bool FConstraintsManagerController::AddConstraint(UTickableConstraint* InConstra
 	}
 
 	Manager->Modify();
-	//it's possible this constraint was actually in another ConstraintActor::ConstraintManager so we need to move it over via Rename.
-	//and clear out it ticks function since that may have been registered
-	UConstraintsManager* Outer = InConstraint->GetTypedOuter<UConstraintsManager>();
-	if (Outer && Outer != Manager)
-	{
-		InConstraint->ConstraintTick.UnRegisterTickFunction();
-		InConstraint->Rename(nullptr, Manager, REN_ForceNoResetLoaders | REN_DoNotDirty | REN_DontCreateRedirectors | REN_NonTransactional);
-	}
+
 	Manager->Constraints.Emplace(InConstraint);
 
 	InConstraint->ConstraintTick.RegisterFunction(InConstraint->GetFunction());
diff --git a/Engine/Source/Runtime/Experimental/Animation/Constraints/Public/ConstraintsManager.h b/Engine/Source/Runtime/Experimental/Animation/Constraints/Public/ConstraintsManager.h
--- a/Engine/Source/Runtime/Experimental/Animation/Constraints/Public/ConstraintsManager.h
+++ b/Engine/Source/Runtime/Experimental/Animation/Constraints/Public/ConstraintsManager.h
@@ -263,6 +263,8 @@ public:
 	/** Go through each constraint in order and evaluate and tick them*/
 	void EvaluateAllConstraints() const;
 
+	/** when PIEing/Simulating it's possible that the constraint isn't in the active manager but still lives*/
+	bool DoesExistInAnyWorld(UTickableConstraint* InConstraint);
 	
 private:
 	/** Delegeate that's fired when a scene component is constrained, this is needed to make sure things like gizmo's get updated after the constraint tick happens*/
diff --git a/Engine/Source/Runtime/MovieSceneTracks/Private/Systems/MovieSceneConstraintSystem.cpp b/Engine/Source/Runtime/MovieSceneTracks/Private/Systems/MovieSceneConstraintSystem.cpp
--- a/Engine/Source/Runtime/MovieSceneTracks/Private/Systems/MovieSceneConstraintSystem.cpp
+++ b/Engine/Source/Runtime/MovieSceneTracks/Private/Systems/MovieSceneConstraintSystem.cpp
@@ -23,7 +23,7 @@
 class UTickableConstraint;
 
 
-UTickableConstraint* CreateConstraintIfNeeded(FConstraintsManagerController& Controller, FConstraintAndActiveChannel* ConstraintAndActiveChannel,const  FConstraintComponentData& ConstraintChannel)
+static UTickableConstraint* CreateConstraintIfNeeded(FConstraintsManagerController& Controller, FConstraintAndActiveChannel* ConstraintAndActiveChannel,const  FConstraintComponentData& ConstraintChannel)
 {
 	if (ConstraintAndActiveChannel->Constraint.IsPending())
 	{
@@ -40,7 +40,14 @@ UTickableConstraint* CreateConstraintIfNeeded(FConstraintsManagerController& Con
 		const TArray< TObjectPtr<UTickableConstraint>>& ConstraintsArray = Controller.GetConstraintsArray();
 		if (ConstraintsArray.Find(Constraint) == INDEX_NONE)
 		{
-			Controller.AddConstraint(Constraint);
+			//it's possible it belongs to another Manager
+			if (Controller.DoesExistInAnyWorld(Constraint) == false) //it's no where make a copy and add it
+			{
+				Constraint = Controller.AddConstraintFromCopy(ConstraintAndActiveChannel->ConstraintCopyToSpawn);
+				Controller.AddConstraint(Constraint);
+				ConstraintChannel.Section->ReplaceConstraint(ConstraintChannel.ConstraintName, Constraint);
+			}
+
 		}
 	}
 	return Constraint;
