Control Rig: Fix crash when reopening old project

diff --git a/Engine/Plugins/Animation/ControlRig/Source/ControlRig/Private/AnimNode_ControlRigBase.cpp b/Engine/Plugins/Animation/ControlRig/Source/ControlRig/Private/AnimNode_ControlRigBase.cpp
--- a/Engine/Plugins/Animation/ControlRig/Source/ControlRig/Private/AnimNode_ControlRigBase.cpp
+++ b/Engine/Plugins/Animation/ControlRig/Source/ControlRig/Private/AnimNode_ControlRigBase.cpp
@@ -720,22 +720,24 @@ void FAnimNode_ControlRigBase::CacheBones_AnyThread(const FAnimationCacheBonesCo
 			
 			// we just support curves by name only
 			{
-				const FSmartNameMapping* CurveMapping = RequiredBones.GetSkeletonAsset()->GetSmartNameContainer(USkeleton::AnimCurveMappingName);
-				URigHierarchy* Hierarchy = ControlRig->GetHierarchy();
-				CurveMapping->Iterate([Hierarchy, this](const FSmartNameMappingIterator& Iterator)
+				const USkeleton* Skeleton = RequiredBones.GetSkeletonAsset();
+				if (IsValid(Skeleton))
 				{
-					// see if the curve name exists in the control rig
-					FName CurveName;
-					if (Iterator.GetName(CurveName))
+					const FSmartNameMapping* CurveMapping = Skeleton->GetSmartNameContainer(USkeleton::AnimCurveMappingName);
+					URigHierarchy* Hierarchy = ControlRig->GetHierarchy();
+					CurveMapping->Iterate([Hierarchy, this](const FSmartNameMappingIterator& Iterator)
 					{
-						if (Hierarchy->GetIndex(FRigElementKey(CurveName, ERigElementType::Curve)) != INDEX_NONE)
+						// see if the curve name exists in the control rig
+						FName CurveName;
+						if (Iterator.GetName(CurveName))
 						{
-							ControlRigCurveMappingByName.Add(CurveName, Iterator.GetIndex());
+							if (Hierarchy->GetIndex(FRigElementKey(CurveName, ERigElementType::Curve)) != INDEX_NONE)
+							{
+								ControlRigCurveMappingByName.Add(CurveName, Iterator.GetIndex());
+							}
 						}
-					}
-				});
-				
-					
+					});
+				}
 			}
 
 			// check if we can switch the bones to an index based mapping.
