Sequencer: Control Rig: Keying weighted layers was broken do to code being commented out.

diff --git a/Engine/Plugins/Animation/ControlRig/Source/ControlRig/Private/Sequencer/MovieSceneControlRigParameterTemplate.cpp b/Engine/Plugins/Animation/ControlRig/Source/ControlRig/Private/Sequencer/MovieSceneControlRigParameterTemplate.cpp
--- a/Engine/Plugins/Animation/ControlRig/Source/ControlRig/Private/Sequencer/MovieSceneControlRigParameterTemplate.cpp
+++ b/Engine/Plugins/Animation/ControlRig/Source/ControlRig/Private/Sequencer/MovieSceneControlRigParameterTemplate.cpp
@@ -2189,15 +2189,13 @@ void FMovieSceneControlRigParameterTemplate::Interrogate(const FMovieSceneContex
 
 		FControlRigAnimTypeIDsPtr TypeIDs = FControlRigAnimTypeIDs::Get(Section->GetControlRig());
 
-		float Weight = 1.f;
-
-		//float Weight = EvaluateEasing(Context.GetTime());
-	//if (EnumHasAllFlags(TemplateData.Mask.GetChannels(), EMovieSceneTransformChannel::Weight))
-	//{
-	//	float ManualWeight = 1.f;
-	//	TemplateData.ManualWeight.Evaluate(Context.GetTime(), ManualWeight);
-	//	Weight *= ManualWeight;
-	//}
+		float Weight = EvaluateEasing(Context.GetTime());
+		if (EnumHasAllFlags(Section->TransformMask.GetChannels(), EMovieSceneTransformChannel::Weight))
+		{
+			float ManualWeight = 1.f;
+			Section->Weight.Evaluate(Context.GetTime(), ManualWeight);
+			Weight *= ManualWeight;
+		}
 
 
 		for (const FScalarParameterStringAndValue& ScalarNameAndValue : Values.ScalarValues)
