Refresh the details panel when metadata are updated to avoid future crash

diff --git a/Engine/Plugins/Animation/ControlRig/Source/ControlRigEditor/Private/ControlRigElementDetails.cpp b/Engine/Plugins/Animation/ControlRig/Source/ControlRigEditor/Private/ControlRigElementDetails.cpp
--- a/Engine/Plugins/Animation/ControlRig/Source/ControlRigEditor/Private/ControlRigElementDetails.cpp
+++ b/Engine/Plugins/Animation/ControlRig/Source/ControlRigEditor/Private/ControlRigElementDetails.cpp
@@ -786,6 +786,26 @@ void FRigBaseElementDetails::CustomizeDetails(IDetailLayoutBuilder& DetailBuilde
 	}
 }
 
+void FRigBaseElementDetails::PendingDelete()
+{
+	if (MetadataHandle.IsValid())
+	{
+		for (const FPerElementInfo& Info : PerElementInfos)
+		{
+			if (URigHierarchy* Hierarchy = Info.IsValid() ? Info.GetHierarchy() : nullptr)
+			{
+				if (Hierarchy->OnMetadataChanged().Remove(MetadataHandle))
+				{
+					break;
+				}
+			}
+		}
+		MetadataHandle.Reset();
+	}
+	
+	IDetailCustomization::PendingDelete();
+}
+
 FRigElementKey FRigBaseElementDetails::GetElementKey() const
 {
 	check(PerElementInfos.Num() == 1);
@@ -1147,6 +1167,27 @@ void FRigBaseElementDetails::CustomizeMetadata(IDetailLayoutBuilder& DetailBuild
 	{
 		return;
 	}
+
+	if (!MetadataHandle.IsValid())
+	{
+		const FPerElementInfo& Info = PerElementInfos[0];
+		URigHierarchy* Hierarchy = Info.IsValid() ? Info.GetHierarchy() : nullptr;
+		if (!Hierarchy)
+		{
+			return;
+		}
+			
+		TSharedRef<IPropertyUtilities> PropertyUtilities = DetailBuilder.GetPropertyUtilities();
+		MetadataHandle = Hierarchy->OnMetadataChanged().AddLambda([this, PropertyUtilities](const FRigElementKey& InKey, const FName&)
+		{
+			const FRigBaseElement* Element = PerElementInfos.Num() == 1 ? PerElementInfos[0].GetElement() : nullptr;
+			if (Element && Element->GetKey() == InKey)
+			{
+				PropertyUtilities->ForceRefresh();
+			}
+		});
+	}
+	
 	const FRigBaseElement* Element = PerElementInfos[0].Element.Get();
 	if(Element->NumMetadata() == 0)
 	{
diff --git a/Engine/Plugins/Animation/ControlRig/Source/ControlRigEditor/Private/ControlRigElementDetails.h b/Engine/Plugins/Animation/ControlRig/Source/ControlRigEditor/Private/ControlRigElementDetails.h
--- a/Engine/Plugins/Animation/ControlRig/Source/ControlRigEditor/Private/ControlRigElementDetails.h
+++ b/Engine/Plugins/Animation/ControlRig/Source/ControlRigEditor/Private/ControlRigElementDetails.h
@@ -142,6 +142,7 @@ public:
 	/** IDetailCustomization interface */
 	virtual void CustomizeDetails(IDetailLayoutBuilder& DetailBuilder) override;
 	virtual void BeginDestroy() {};
+	virtual void PendingDelete() override;
 
 	FRigElementKey GetElementKey() const;
 	FText GetName() const;
@@ -248,7 +249,7 @@ protected:
 	void CustomizeMetadata(IDetailLayoutBuilder& DetailBuilder);
 
 	TArray<FPerElementInfo> PerElementInfos;
-	TArray<TStrongObjectPtr<UDetailsViewWrapperObject>> MetadataWrappers;
+	FDelegateHandle MetadataHandle;
 	
 	TSharedPtr<SButton> SelectParentElementButton;
 };
