Fix remaining UnrealEditor Win64 private V3 issues

#rnx

diff --git a/Engine/Plugins/Animation/AnimationData/Source/AnimationData/Private/AnimSequencerDataModel.cpp b/Engine/Plugins/Animation/AnimationData/Source/AnimationData/Private/AnimSequencerDataModel.cpp
--- a/Engine/Plugins/Animation/AnimationData/Source/AnimationData/Private/AnimSequencerDataModel.cpp
+++ b/Engine/Plugins/Animation/AnimationData/Source/AnimationData/Private/AnimSequencerDataModel.cpp
@@ -20,10 +20,11 @@
 
 #include "AnimSequencerHelpers.h"
 #include "Animation/AnimationSettings.h"
-#include "MovieScene/Private/Channels/MovieSceneCurveChannelImpl.h"
 #include "Compilation/MovieSceneCompiledDataManager.h"
 #include "UObject/UObjectThreadContext.h"
 
+#include "Runtime/MovieScene/Private/Channels/MovieSceneCurveChannelImpl.h"
+
 #define LOCTEXT_NAMESPACE "AnimSequencerDataModel"
 
 int32 UAnimationSequencerDataModel::RetainFloatCurves = 0;
diff --git a/Engine/Plugins/Animation/ControlRig/Source/ControlRig/Private/Units/Execution/RigUnit_InverseExecution.cpp b/Engine/Plugins/Animation/ControlRig/Source/ControlRig/Private/Units/Execution/RigUnit_InverseExecution.cpp
--- a/Engine/Plugins/Animation/ControlRig/Source/ControlRig/Private/Units/Execution/RigUnit_InverseExecution.cpp
+++ b/Engine/Plugins/Animation/ControlRig/Source/ControlRig/Private/Units/Execution/RigUnit_InverseExecution.cpp
@@ -1,6 +1,6 @@
 // Copyright Epic Games, Inc. All Rights Reserved.
 
-#include "RigUnit_InverseExecution.h"
+#include "Units/Execution/RigUnit_InverseExecution.h"
 #include "Units/RigUnitContext.h"
 
 #include UE_INLINE_GENERATED_CPP_BY_NAME(RigUnit_InverseExecution)
diff --git a/Engine/Plugins/Animation/ControlRig/Source/ControlRig/Private/Units/Execution/RigUnit_InverseExecution.h b/Engine/Plugins/Animation/ControlRig/Source/ControlRig/Private/Units/Execution/RigUnit_InverseExecution.h
deleted file mode 100644
diff --git a/Engine/Plugins/Animation/ControlRig/Source/ControlRig/Public/Units/Execution/RigUnit_InverseExecution.h b/Engine/Plugins/Animation/ControlRig/Source/ControlRig/Public/Units/Execution/RigUnit_InverseExecution.h
new file mode 100644
--- /dev/null
+++ b/Engine/Plugins/Animation/ControlRig/Source/ControlRig/Public/Units/Execution/RigUnit_InverseExecution.h
@@ -0,0 +1,28 @@
+// Copyright Epic Games, Inc. All Rights Reserved.
+
+#pragma once
+
+#include "Units/RigUnit.h"
+#include "ControlRigDefines.h"
+#include "RigUnit_InverseExecution.generated.h"
+
+/**
+ * Event for driving elements based off the skeleton hierarchy
+ */
+USTRUCT(meta=(DisplayName="Backwards Solve", Category="Events", NodeColor="1, 0, 0", Keywords="Inverse,Reverse,Backwards,Event"))
+struct CONTROLRIG_API FRigUnit_InverseExecution : public FRigUnit
+{
+	GENERATED_BODY()
+
+	RIGVM_METHOD()
+	virtual void Execute() override;
+
+	virtual FName GetEventName() const override { return EventName; }
+	virtual bool CanOnlyExistOnce() const override { return true; }
+
+	// The execution result
+	UPROPERTY(EditAnywhere, Transient, DisplayName = "Execute", Category = "InverseExecution", meta = (Output))
+	FControlRigExecuteContext ExecuteContext;
+
+	static FName EventName;
+};
diff --git a/Engine/Plugins/Animation/ControlRig/Source/ControlRigEditor/Private/EditMode/ControlRigEditMode.h b/Engine/Plugins/Animation/ControlRig/Source/ControlRigEditor/Private/EditMode/ControlRigEditMode.h
deleted file mode 100644
diff --git a/Engine/Plugins/Animation/ControlRig/Source/ControlRigEditor/Public/EditMode/ControlRigEditMode.h b/Engine/Plugins/Animation/ControlRig/Source/ControlRigEditor/Public/EditMode/ControlRigEditMode.h
new file mode 100644
--- /dev/null
+++ b/Engine/Plugins/Animation/ControlRig/Source/ControlRigEditor/Public/EditMode/ControlRigEditMode.h
@@ -0,0 +1,477 @@
+// Copyright Epic Games, Inc. All Rights Reserved.
+
+#pragma once
+
+#include "CoreMinimal.h"
+#include "InputCoreTypes.h"
+#include "IControlRigObjectBinding.h"
+#include "RigVMModel/RigVMGraph.h"
+#include "Rigs/RigHierarchyContainer.h"
+#include "Units/RigUnitContext.h"
+#include "Framework/MultiBox/MultiBoxBuilder.h"
+#include "UObject/StrongObjectPtr.h"
+#include "UnrealWidgetFwd.h"
+#include "IPersonaEditMode.h"
+#include "Misc/Guid.h"
+#include "ControlRigEditMode.generated.h"
+
+
+
+class FEditorViewportClient;
+class FViewport;
+class UActorFactory;
+struct FViewportClick;
+class UControlRig;
+class FControlRigInteractionScope;
+class ISequencer;
+class UControlManipulator;
+class FUICommandList;
+class FPrimitiveDrawInterface;
+class FToolBarBuilder;
+class FExtender;
+class IMovieScenePlayer;
+class AControlRigShapeActor;
+class UDefaultControlRigManipulationLayer;
+class UControlRigDetailPanelControlProxies;
+class UControlRigControlsProxy;
+struct FRigControl;
+class IControlRigManipulatable;
+class ISequencer;
+enum class EControlRigSetKey : uint8;
+class UToolMenu;
+
+DECLARE_DELEGATE_RetVal_ThreeParams(FTransform, FOnGetRigElementTransform, const FRigElementKey& /*RigElementKey*/, bool /*bLocal*/, bool /*bOnDebugInstance*/);
+DECLARE_DELEGATE_ThreeParams(FOnSetRigElementTransform, const FRigElementKey& /*RigElementKey*/, const FTransform& /*Transform*/, bool /*bLocal*/);
+DECLARE_DELEGATE_RetVal(TSharedPtr<FUICommandList>, FNewMenuCommandsDelegate);
+DECLARE_MULTICAST_DELEGATE_TwoParams(FControlRigAddedOrRemoved, UControlRig*, bool /*true if added, false if removed*/);
+DECLARE_MULTICAST_DELEGATE_ThreeParams(FControlRigSelected, UControlRig*, const FRigElementKey& /*RigElementKey*/,const bool /*bIsSelected*/);
+DECLARE_DELEGATE_RetVal(UToolMenu*, FOnGetContextMenu);
+
+class FControlRigEditMode;
+
+enum class ERecreateControlRigShape
+{
+	RecreateNone,
+	RecreateAll,
+	RecreateSpecified
+};
+
+
+UCLASS()
+class UControlRigEditModeDelegateHelper : public UObject
+{
+	GENERATED_BODY()
+
+public:
+
+	UFUNCTION()
+	void OnPoseInitialized();
+
+	UFUNCTION()
+	void PostPoseUpdate();
+
+	void AddDelegates(USkeletalMeshComponent* InSkeletalMeshComponent);
+	void RemoveDelegates();
+
+	TWeakObjectPtr<USkeletalMeshComponent> BoundComponent;
+	FControlRigEditMode* EditMode = nullptr;
+
+private:
+	FDelegateHandle OnBoneTransformsFinalizedHandle;
+};
+
+class CONTROLRIGEDITOR_API FControlRigEditMode : public IPersonaEditMode
+{
+public:
+	static FName ModeName;
+
+	FControlRigEditMode();
+	~FControlRigEditMode();
+
+	/** Set the Control Rig Object to be active in the edit mode. You set both the Control Rig and a possible binding together with an optional Sequencer
+	 This will remove all other control rigs present and should be called for stand alone editors, like the Control Rig Editor*/
+	void SetObjects(UControlRig* InControlRig, UObject* BindingObject, TWeakPtr<ISequencer> InSequencer);
+
+	/** Add a Control Rig object if it doesn't exist, will return true if it was added, false if it wasn't since it's already there. You can also set the Sequencer.*/
+	bool AddControlRigObject(UControlRig* InControlRig, TWeakPtr<ISequencer> InSequencer);
+
+	/* Remove control rig */
+	void RemoveControlRig(UControlRig* InControlRig);
+
+	/*Replace old Control Rig with the New Control Rig, perhaps from a recompile in the level editor*/
+	void ReplaceControlRig(UControlRig* OldControlRig, UControlRig* NewControlRig);
+
+	/** This edit mode is re-used between the level editor and the control rig editor. Calling this indicates which context we are in */
+	bool IsInLevelEditor() const;
+
+	// FEdMode interface
+	virtual bool UsesToolkits() const override;
+	virtual void Enter() override;
+	virtual void Exit() override;
+	virtual void Tick(FEditorViewportClient* ViewportClient, float DeltaTime) override;
+	virtual void Render(const FSceneView* View, FViewport* Viewport, FPrimitiveDrawInterface* PDI) override;
+	virtual bool InputKey(FEditorViewportClient* InViewportClient, FViewport* InViewport, FKey InKey, EInputEvent InEvent) override;
+	virtual bool EndTracking(FEditorViewportClient* InViewportClient, FViewport* InViewport) override;
+	virtual bool StartTracking(FEditorViewportClient* InViewportClient, FViewport* InViewport) override;
+	virtual bool ProcessCapturedMouseMoves(FEditorViewportClient* InViewportClient, FViewport* InViewport, const TArrayView<FIntPoint>& CapturedMouseMoves) override;
+	virtual bool HandleClick(FEditorViewportClient* InViewportClient, HHitProxy *HitProxy, const FViewportClick &Click) override;
+	virtual bool BoxSelect(FBox& InBox, bool InSelect = true) override;
+	virtual bool FrustumSelect(const FConvexVolume& InFrustum, FEditorViewportClient* InViewportClient, bool InSelect = true) override;
+	virtual void SelectNone() override;
+	virtual bool InputDelta(FEditorViewportClient* InViewportClient, FViewport* InViewport, FVector& InDrag, FRotator& InRot, FVector& InScale) override;
+	virtual bool UsesTransformWidget() const override;
+	virtual bool GetPivotForOrbit(FVector& OutPivot) const override;
+	virtual bool UsesTransformWidget(UE::Widget::EWidgetMode CheckMode) const;
+	virtual FVector GetWidgetLocation() const override;
+	virtual bool GetCustomDrawingCoordinateSystem(FMatrix& OutMatrix, void* InData) override;
+	virtual bool GetCustomInputCoordinateSystem(FMatrix& OutMatrix, void* InData) override;
+	virtual bool ShouldDrawWidget() const override;
+	virtual bool IsCompatibleWith(FEditorModeID OtherModeID) const override;
+	virtual bool MouseMove(FEditorViewportClient* ViewportClient, FViewport* Viewport, int32 x, int32 y) override;
+	virtual bool MouseLeave(FEditorViewportClient* ViewportClient, FViewport* Viewport);
+	virtual void PostUndo() override;
+
+	/* IPersonaEditMode interface */
+	virtual bool GetCameraTarget(FSphere& OutTarget) const override { return false; }
+	virtual class IPersonaPreviewScene& GetAnimPreviewScene() const override { check(false); return *(IPersonaPreviewScene*)this; }
+	virtual void GetOnScreenDebugInfo(TArray<FText>& OutDebugInfo) const override {}
+
+	/** FGCObject interface */
+	virtual void AddReferencedObjects( FReferenceCollector& Collector ) override;
+
+	/** Refresh our internal object list (they may have changed) */
+	void RefreshObjects();
+
+	/** Find the edit mode corresponding to the specified world context */
+	static FControlRigEditMode* GetEditModeFromWorldContext(UWorld* InWorldContext);
+
+	/** Bone Manipulation Delegates */
+	FOnGetRigElementTransform& OnGetRigElementTransform() { return OnGetRigElementTransformDelegate; }
+	FOnSetRigElementTransform& OnSetRigElementTransform() { return OnSetRigElementTransformDelegate; }
+
+	/** Context Menu Delegates */
+	FOnGetContextMenu& OnGetContextMenu() { return OnGetContextMenuDelegate; }
+	FNewMenuCommandsDelegate& OnContextMenuCommands() { return OnContextMenuCommandsDelegate; }
+	FSimpleMulticastDelegate& OnAnimSystemInitialized() { return OnAnimSystemInitializedDelegate; }
+
+	/* Control Rig Changed Delegate*/
+	FControlRigAddedOrRemoved& OnControlRigAddedOrRemoved() { return OnControlRigAddedOrRemovedDelegate; }
+
+	/* Control Rig Selected Delegate*/
+	FControlRigSelected& OnControlRigSelected() { return OnControlRigSelectedDelegate; }
+
+	// callback that gets called when rig element is selected in other view
+	void OnHierarchyModified(ERigHierarchyNotification InNotif, URigHierarchy* InHierarchy, const FRigBaseElement* InElement);
+	void OnHierarchyModified_AnyThread(ERigHierarchyNotification InNotif, URigHierarchy* InHierarchy, const FRigBaseElement* InElement);
+	void OnControlModified(UControlRig* Subject, FRigControlElement* InControlElement, const FRigControlModifiedContext& Context);
+	void OnPreConstruction_AnyThread(UControlRig* InRig, const FName& InEventName);
+	void OnPostConstruction_AnyThread(UControlRig* InRig, const FName& InEventName);
+
+	/** return true if it can be removed from preview scene 
+	- this is to ensure preview scene doesn't remove shape actors */
+	bool CanRemoveFromPreviewScene(const USceneComponent* InComponent);
+
+	FUICommandList* GetCommandBindings() const { return CommandBindings.Get(); }
+
+	/** Requests to recreate the shape actors in the next tick. Will recreate only the ones for the specified
+	Control Rig, otherwise will recreate all of them*/
+	void RequestToRecreateControlShapeActors(UControlRig* ControlRig = nullptr); 
+
+	static uint32 ValidControlTypeMask()
+	{
+		return FRigElementTypeHelper::ToMask(ERigElementType::Control);
+	}
+
+protected:
+
+	// shape related functions wrt enable/selection
+	/** Get the node name from the property path */
+	AControlRigShapeActor* GetControlShapeFromControlName(UControlRig* InControlRig,const FName& InControlName) const;
+
+protected:
+	/** Helper function: set ControlRigs array to the details panel */
+	void SetObjects_Internal();
+
+	/** Set up Details Panel based upon Selected Objects*/
+	void SetUpDetailPanel();
+
+	/** Updates cached pivot transform */
+	void RecalcPivotTransform();
+
+	/** Helper function for box/frustum intersection */
+	bool IntersectSelect(bool InSelect, const TFunctionRef<bool(const AControlRigShapeActor*, const FTransform&)>& Intersects);
+
+	/** Handle selection internally */
+	void HandleSelectionChanged();
+
+	/** Toggles visibility of acive control rig shapes in the viewport */
+	void ToggleManipulators();
+
+	/** Toggles visibility of all  control rig shapes in the viewport */
+	void ToggleAllManipulators();
+
+	/** If Anim Slider is open, got to the next tool*/
+	void ChangeAnimSliderTool();
+
+	/** If Anim Slider is open, then can drag*/
+	bool CanChangeAnimSliderTool() const;
+	/** Start Slider Tool*/
+	void StartAnimSliderTool(int32 InX);
+	/** If Anim Slider is open, drag the tool*/
+	void DragAnimSliderTool(double Val);
+	/** Reset and stop user the anim slider tool*/
+	void ResetAnimSlider();
+
+	/** If the Drag Anim Slider Tool is pressed*/
+	bool IsDragAnimSliderToolPressed(FViewport* InViewport);
+
+public:
+	
+	/** Clear Selection*/
+	void ClearSelection();
+
+	/** Frame to current Control Selection*/
+	void FrameSelection();
+
+	/** Frame a list of provided items*/
+   	void FrameItems(const TArray<FRigElementKey>& InItems);
+
+	/** Opens up the space picker widget */
+	void OpenSpacePickerWidget();
+
+private:
+	
+	/** Whether or not we should Frame Selection or not*/
+	bool CanFrameSelection();
+
+	/** Reset Transforms */
+	void ResetTransforms(bool bSelectionOnly);
+
+	/** Increase Shape Size */
+	void IncreaseShapeSize();
+
+	/** Decrease Shape Size */
+	void DecreaseShapeSize();
+
+	/** Reset Shape Size */
+	void ResetControlShapeSize();
+
+public:
+	
+	/** Toggle Shape Transform Edit*/
+	void ToggleControlShapeTransformEdit();
+
+private:
+	
+	/** The hotkey text is passed to a viewport notification to inform users how to toggle shape edit*/
+	FText GetToggleControlShapeTransformEditHotKey() const;
+
+	/** Bind our keyboard commands */
+	void BindCommands();
+
+	/** It creates if it doesn't have it */
+	void RecreateControlShapeActors(const TArray<FRigElementKey>& InSelectedElements = TArray<FRigElementKey>());
+
+	/** Let the preview scene know how we want to select components */
+	bool ShapeSelectionOverride(const UPrimitiveComponent* InComponent) const;
+
+	/** Enable editing of control's shape transform instead of control's transform*/
+	bool bIsChangingControlShapeTransform;
+
+protected:
+
+	TWeakPtr<ISequencer> WeakSequencer;
+	FGuid LastMovieSceneSig;
+
+	/** The scope for the interaction, one per manipulated Control rig */
+	TMap<UControlRig*,FControlRigInteractionScope*> InteractionScopes;
+
+	/** Whether a manipulator actually made a change when transacting */
+	bool bManipulatorMadeChange;
+
+	/** Guard value for selection */
+	bool bSelecting;
+
+	/** If selection was changed, we set up proxies on next tick */
+	bool bSelectionChanged;
+
+	/** Cached transform of pivot point for selected objects for each Control Rig */
+	TMap<UControlRig*,FTransform> PivotTransforms;
+
+	/** Previous cached transforms, need this to check on tick if any transform changed, gizmo may have changed*/
+	TMap<UControlRig*, FTransform> LastPivotTransforms;
+
+	/** Command bindings for keyboard shortcuts */
+	TSharedPtr<FUICommandList> CommandBindings;
+
+	/** Called from the editor when a blueprint object replacement has occurred */
+	void OnObjectsReplaced(const TMap<UObject*, UObject*>& OldToNewInstanceMap);
+
+	/** Return true if transform setter/getter delegates are available */
+	bool IsTransformDelegateAvailable() const;
+
+	FOnGetRigElementTransform OnGetRigElementTransformDelegate;
+	FOnSetRigElementTransform OnSetRigElementTransformDelegate;
+	FOnGetContextMenu OnGetContextMenuDelegate;
+	FNewMenuCommandsDelegate OnContextMenuCommandsDelegate;
+	FSimpleMulticastDelegate OnAnimSystemInitializedDelegate;
+	FControlRigAddedOrRemoved OnControlRigAddedOrRemovedDelegate;
+	FControlRigSelected OnControlRigSelectedDelegate;
+
+	/** GetSelectedRigElements, if InControlRig is nullptr get the first one */
+	TArray<FRigElementKey> GetSelectedRigElements(UControlRig* InControlRig = nullptr) const;
+
+	/* Flag to recreate shapes during tick */
+	ERecreateControlRigShape RecreateControlShapesRequired;
+	/* List of Control Rigs we should recreate*/
+	TArray<UControlRig*> ControlRigsToRecreate;
+
+	/* Flag to temporarily disable handling notifs from the hierarchy */
+	bool bSuspendHierarchyNotifs;
+
+	/** Shape actors */
+	TMap<UControlRig*,TArray<AControlRigShapeActor*>> ControlRigShapeActors;
+	UControlRigDetailPanelControlProxies* ControlProxy;
+
+	/** Utility functions for UI/Some other viewport manipulation*/
+	bool IsControlSelected() const;
+	bool AreRigElementSelectedAndMovable(UControlRig* InControlRig) const;
+	
+	/** Set initial transform handlers */
+	void OpenContextMenu(FEditorViewportClient* InViewportClient);
+
+	/* params to handle mouse move for changing anim tools sliders*/
+	bool bisTrackingAnimToolDrag = false;
+	/** Starting X Value*/
+	TOptional<int32> StartXValue;
+
+	
+public: 
+	/** Clear all selected RigElements */
+	void ClearRigElementSelection(uint32 InTypes);
+
+	/** Set a RigElement's selection state */
+	void SetRigElementSelection(UControlRig* ControlRig, ERigElementType Type, const FName& InRigElementName, bool bSelected);
+
+	/** Set multiple RigElement's selection states */
+	void SetRigElementSelection(UControlRig* ControlRig, ERigElementType Type, const TArray<FName>& InRigElementNames, bool bSelected);
+
+	/** Check if any RigElements are selected */
+	bool AreRigElementsSelected(uint32 InTypes, UControlRig* InControlRig) const;
+
+	/** Get the number of selected RigElements */
+	int32 GetNumSelectedRigElements(uint32 InTypes, UControlRig* InControlRig) const;
+
+	/** Get all of the selected Controls*/
+	void GetAllSelectedControls(TMap<UControlRig*, TArray<FRigElementKey>>& OutSelectedControls) const;
+
+	/** Get all of the ControlRigs, maybe not valid anymore */
+	TArrayView<const TWeakObjectPtr<UControlRig>> GetControlRigs() const;
+	TArrayView<TWeakObjectPtr<UControlRig>> GetControlRigs();
+	/* Get valid  Control Rigs possibly just visible*/
+	TArray<UControlRig*> GetControlRigsArray(bool bIsVisible);
+	TArray<const UControlRig*> GetControlRigsArray(bool bIsVisible) const;
+
+	/** Get the detail proxies control rig*/
+	UControlRigDetailPanelControlProxies* GetDetailProxies() { return ControlProxy; }
+
+	/** Get Sequencer Driving This*/
+	TWeakPtr<ISequencer> GetWeakSequencer() { return WeakSequencer; }
+
+	/** Suspend Rig Hierarchy Notifies*/
+	void SuspendHierarchyNotifs(bool bVal) { bSuspendHierarchyNotifs = bVal; }
+private:
+	/** Whether or not Pivot Transforms have changed, in which case we need to redraw viewport*/
+	bool HasPivotTransformsChanged() const;
+	/** Set a RigElement's selection state */
+	void SetRigElementSelectionInternal(UControlRig* ControlRig, ERigElementType Type, const FName& InRigElementName, bool bSelected);
+	
+	FEditorViewportClient* CurrentViewportClient;
+
+/* store coordinate system per widget mode*/
+private:
+	void OnWidgetModeChanged(UE::Widget::EWidgetMode InWidgetMode);
+	void OnCoordSystemChanged(ECoordSystem InCoordSystem);
+	TArray<ECoordSystem> CoordSystemPerWidgetMode;
+	bool bIsChangingCoordSystem;
+
+	bool CanChangeControlShapeTransform();
+public:
+	//Toolbar functions
+	void SetOnlySelectRigControls(bool val);
+	bool GetOnlySelectRigControls()const;
+
+private:
+	TSet<FName> GetActiveControlsFromSequencer(UControlRig* ControlRig);
+	bool SetSequencer(TWeakPtr<ISequencer> InSequencer);
+
+	/** Create/Delete for the specified ControlRig*/
+	void CreateShapeActors(UControlRig* InControlRig);
+	void DestroyShapesActors(UControlRig* InControlRig);
+
+	/*Internal function for adding ControlRig*/
+	void AddControlRigInternal(UControlRig* InControlRig);
+	void TickManipulatableObjects(float DeltaTime);
+
+	/* Check on tick to see if movie scene has changed, returns true if it has*/
+	bool CheckMovieSceneSig();
+	void SetControlShapeTransform( const AControlRigShapeActor* InShapeActor, const FTransform& InGlobalTransform,
+		const FTransform& InToWorldTransform, const FRigControlModifiedContext& InContext, const bool bPrintPython) const;
+	static FTransform GetControlShapeTransform(const AControlRigShapeActor* ShapeActor);
+	void MoveControlShape(AControlRigShapeActor* ShapeActor, const bool bTranslation, FVector& InDrag, 
+		const bool bRotation, FRotator& InRot, const bool bScale, FVector& InScale, const FTransform& ToWorldTransform,
+		bool bUseLocal, bool bCalcLocal, FTransform& InOutLocal);
+
+	void ChangeControlShapeTransform(AControlRigShapeActor* ShapeActor, const bool bTranslation, FVector& InDrag,
+		const bool bRotation, FRotator& InRot, const bool bScale, FVector& InScale, const FTransform& ToWorldTransform);
+
+	void TickControlShape(AControlRigShapeActor* ShapeActor, const FTransform& ComponentTransform);
+	bool ModeSupportedByShapeActor(const AControlRigShapeActor* ShapeActor, UE::Widget::EWidgetMode InMode) const;
+
+
+protected:
+	
+	/** Get bindings to a runtime object */
+	//If the passed in ControlRig is nullptr we use the first Control Rig(this can happen from the BP Editors).
+	USceneComponent* GetHostingSceneComponent(const UControlRig* ControlRig = nullptr) const;
+	FTransform	GetHostingSceneComponentTransform(const UControlRig* ControlRig =  nullptr) const;
+
+	//Get if the hosted component is visible
+	bool IsControlRigSkelMeshVisible(UControlRig* ControlRig) const;
+
+private:
+
+	// Post pose update handler
+	void OnPoseInitialized();
+	void PostPoseUpdate();
+	void NotifyDrivenControls(UControlRig* InControlRig, const FRigElementKey& InKey);
+	void UpdateSelectabilityOnSkeletalMeshes(UControlRig* InControlRig, bool bEnabled);
+
+	// world clean up handlers
+	FDelegateHandle OnWorldCleanupHandle;
+	void OnWorldCleanup(UWorld* World, bool bSessionEnded, bool bCleanupResources);
+	UWorld* WorldPtr = nullptr;
+
+	
+	TArray<TWeakObjectPtr<UControlRig>> RuntimeControlRigs;
+	TMap<UControlRig*,TStrongObjectPtr<UControlRigEditModeDelegateHelper>> DelegateHelpers;
+
+	//hack since we can't get the viewport client from the viewport, so in the tick we set the gameview bool and then in render/tickcontrolshapes we use it.
+	TMap<FViewport*, bool>  ViewportToGameView;
+
+	TArray<FRigElementKey> DeferredItemsToFrame;
+
+	/** Computes the current interaction types based on the widget mode */
+	static uint8 GetInteractionType(FEditorViewportClient* InViewportClient);
+	uint8 InteractionType;
+	bool bShowControlsAsOverlay;
+
+	bool bIsConstructionEventRunning;
+	TArray<uint32> LastHierarchyHash;
+
+	friend class FControlRigEditorModule;
+	friend class FControlRigEditor;
+	friend class FControlRigEditModeGenericDetails;
+	friend class UControlRigEditModeDelegateHelper;
+	friend class SControlRigEditModeTools;
+};
diff --git a/Engine/Plugins/Experimental/Animation/PoseCorrectives/Source/PoseCorrectivesEditor/Private/PoseCorrectivesEditorController.cpp b/Engine/Plugins/Experimental/Animation/PoseCorrectives/Source/PoseCorrectivesEditor/Private/PoseCorrectivesEditorController.cpp
--- a/Engine/Plugins/Experimental/Animation/PoseCorrectives/Source/PoseCorrectivesEditor/Private/PoseCorrectivesEditorController.cpp
+++ b/Engine/Plugins/Experimental/Animation/PoseCorrectives/Source/PoseCorrectivesEditor/Private/PoseCorrectivesEditorController.cpp
@@ -18,7 +18,7 @@
 #include "EditorModeManager.h"
 #include "ScopedTransaction.h"
 
-#include "ControlRig/Private/Units/Execution/RigUnit_InverseExecution.h"
+#include "Units/Execution/RigUnit_InverseExecution.h"
 
 
 #define LOCTEXT_NAMESPACE "PoseCorrectivesEditorController"
diff --git a/Engine/Plugins/Experimental/Animation/PoseCorrectives/Source/PoseCorrectivesEditor/Public/CorrectivesEditMode.h b/Engine/Plugins/Experimental/Animation/PoseCorrectives/Source/PoseCorrectivesEditor/Public/CorrectivesEditMode.h
--- a/Engine/Plugins/Experimental/Animation/PoseCorrectives/Source/PoseCorrectivesEditor/Public/CorrectivesEditMode.h
+++ b/Engine/Plugins/Experimental/Animation/PoseCorrectives/Source/PoseCorrectivesEditor/Public/CorrectivesEditMode.h
@@ -1,7 +1,7 @@
 // Copyright Epic Games, Inc. All Rights Reserved.
 #pragma once
 
-#include "ControlRigEditor/Private/EditMode/ControlRigEditMode.h"
+#include "EditMode/ControlRigEditMode.h"
 
 class FPoseCorrectivesEditorController;
 
diff --git a/Engine/Plugins/MovieScene/MovieRenderPipeline/Source/MovieRenderPipelineEditor/Private/Customizations/ConsoleVariableSettingCustomization.h b/Engine/Plugins/MovieScene/MovieRenderPipeline/Source/MovieRenderPipelineEditor/Private/Customizations/ConsoleVariableSettingCustomization.h
--- a/Engine/Plugins/MovieScene/MovieRenderPipeline/Source/MovieRenderPipelineEditor/Private/Customizations/ConsoleVariableSettingCustomization.h
+++ b/Engine/Plugins/MovieScene/MovieRenderPipeline/Source/MovieRenderPipelineEditor/Private/Customizations/ConsoleVariableSettingCustomization.h
@@ -9,10 +9,11 @@
 #include "IDetailGroup.h"
 #include "MoviePipelineConsoleVariableSetting.h"
 #include "PropertyCustomizationHelpers.h"
-#include "PropertyEditor/Private/PropertyNode.h"
 #include "PropertyHandle.h"
 #include "Sections/MovieSceneConsoleVariableTrackInterface.h"
 
+#include "Editor/PropertyEditor/Private/PropertyNode.h"
+
 #define LOCTEXT_NAMESPACE "SMoviePipelineEditor"
 
 /** Customize how properties in UMoviePipelineConsoleVariableSetting appear in the details panel. */
