Only cache items if needed

diff --git a/Engine/Plugins/Animation/ControlRig/Source/ControlRig/Private/Units/Highlevel/Harmonics/RigUnit_BoneHarmonics.cpp b/Engine/Plugins/Animation/ControlRig/Source/ControlRig/Private/Units/Highlevel/Harmonics/RigUnit_BoneHarmonics.cpp
--- a/Engine/Plugins/Animation/ControlRig/Source/ControlRig/Private/Units/Highlevel/Harmonics/RigUnit_BoneHarmonics.cpp
+++ b/Engine/Plugins/Animation/ControlRig/Source/ControlRig/Private/Units/Highlevel/Harmonics/RigUnit_BoneHarmonics.cpp
@@ -84,14 +84,17 @@ FRigUnit_ItemHarmonics_Execute()
 		WaveTime = FVector::ZeroVector;
 	}
 
-	for (int32 ItemIndex = 0; ItemIndex < Targets.Num(); ItemIndex++)
+	if (CachedItems.IsEmpty())
 	{
-		FCachedRigElement CachedItem(Targets[ItemIndex].Item, Hierarchy);
-		if (!CachedItem.IsValid())
+		for (int32 ItemIndex = 0; ItemIndex < Targets.Num(); ItemIndex++)
 		{
-			UE_CONTROLRIG_RIGUNIT_REPORT_WARNING(TEXT("Item not found."));
+			FCachedRigElement CachedItem(Targets[ItemIndex].Item, Hierarchy);
+			if (!CachedItem.IsValid())
+			{
+				UE_CONTROLRIG_RIGUNIT_REPORT_WARNING(TEXT("Item not found."));
+			}
+			CachedItems.Add(CachedItem);
 		}
-		CachedItems.Add(CachedItem);
 	}
 
 	for (int32 ItemIndex = 0; ItemIndex < CachedItems.Num(); ItemIndex++)
