Control Rig: Request pin type window opens on mouse position

diff --git a/Engine/Plugins/Animation/ControlRig/Source/ControlRigEditor/Private/Editor/ControlRigEditor.cpp b/Engine/Plugins/Animation/ControlRig/Source/ControlRigEditor/Private/Editor/ControlRigEditor.cpp
--- a/Engine/Plugins/Animation/ControlRig/Source/ControlRigEditor/Private/Editor/ControlRigEditor.cpp
+++ b/Engine/Plugins/Animation/ControlRig/Source/ControlRigEditor/Private/Editor/ControlRigEditor.cpp
@@ -5139,7 +5139,9 @@ TRigVMTypeIndex FControlRigEditor::OnRequestPinTypeSelectionDialog(const TArray<
 
 	TSharedPtr< SWindow > Window = SNew(SWindow)
 		.Title(LOCTEXT("SelectPinType", "Select Pin Type"))
+		.ScreenPosition(FSlateApplication::Get().GetCursorPos())
 		.SizingRule(ESizingRule::Autosized)
+		.AutoCenter(EAutoCenter::None)
 		.SupportsMaximize(false)
 		.SupportsMinimize(false)
 		[
