Sequencer: Follow up to 23916250 - track filter saved in preferences

No need to enable control rig controls filter when baking to control rig

diff --git a/Engine/Plugins/Animation/ControlRig/Source/ControlRigEditor/Private/ControlRigEditorModule.cpp b/Engine/Plugins/Animation/ControlRig/Source/ControlRigEditor/Private/ControlRigEditorModule.cpp
--- a/Engine/Plugins/Animation/ControlRig/Source/ControlRigEditor/Private/ControlRigEditorModule.cpp
+++ b/Engine/Plugins/Animation/ControlRig/Source/ControlRigEditor/Private/ControlRigEditorModule.cpp
@@ -766,8 +766,6 @@ void FControlRigEditorModule::BakeToControlRig(UClass* ControlRigClass, UAnimSeq
 					WeakSequencer.Pin()->SelectSection(ParamSection);
 					WeakSequencer.Pin()->ThrobSectionSelection();
 					WeakSequencer.Pin()->ObjectImplicitlyAdded(ControlRig);
-					FText Name = LOCTEXT("SequenceTrackFilter_ControlRigControls", "Control Rig Controls");
-					WeakSequencer.Pin()->SetTrackFilterEnabled(Name, true);
 					WeakSequencer.Pin()->NotifyMovieSceneDataChanged(EMovieSceneDataChangeType::MovieSceneStructureItemAdded);
 					FControlRigEditMode* ControlRigEditMode = static_cast<FControlRigEditMode*>(GLevelEditorModeTools().GetActiveMode(FControlRigEditMode::ModeName));
 					if (!ControlRigEditMode)
