Control Rig: Emit warning when setting visibility on control with shape visibility based on selection

diff --git a/Engine/Plugins/Animation/ControlRig/Source/ControlRig/Private/Units/Hierarchy/RigUnit_SetControlVisibility.cpp b/Engine/Plugins/Animation/ControlRig/Source/ControlRig/Private/Units/Hierarchy/RigUnit_SetControlVisibility.cpp
--- a/Engine/Plugins/Animation/ControlRig/Source/ControlRig/Private/Units/Hierarchy/RigUnit_SetControlVisibility.cpp
+++ b/Engine/Plugins/Animation/ControlRig/Source/ControlRig/Private/Units/Hierarchy/RigUnit_SetControlVisibility.cpp
@@ -68,6 +68,18 @@ FRigUnit_SetControlVisibility_Execute()
 		{
 			if (CachedControlIndex.IsValid())
 			{
+#if WITH_EDITOR
+				if (FRigControlElement* ControlElement = Cast<FRigControlElement>(Hierarchy->Find(CachedControlIndex.GetKey())))
+				{
+					if (ControlElement->Settings.ShapeVisibility == ERigControlVisibility::BasedOnSelection)
+					{
+						UE_CONTROLRIG_RIGUNIT_REPORT_WARNING(
+						TEXT("SetControlVisibility: visibility of '%s' is based on selection, and cannot be modified through SetControlVisibility. "
+						"In order to be able to modify the control's visibility, set the ShapeVisibility to UserDefined."),
+						*ControlElement->GetKey().ToString());
+					}
+				}
+#endif
 				Hierarchy->SetControlVisibility(CachedControlIndex, bVisible);
 			}
 		}
