Control Rig: Fix missing debug units that should have been deprecated

diff --git a/Engine/Plugins/Animation/ControlRig/Config/BaseControlRig.ini b/Engine/Plugins/Animation/ControlRig/Config/BaseControlRig.ini
--- a/Engine/Plugins/Animation/ControlRig/Config/BaseControlRig.ini
+++ b/Engine/Plugins/Animation/ControlRig/Config/BaseControlRig.ini
@@ -420,28 +420,8 @@
 +StructRedirects=(OldName="/Script/ControlRig.RigUnit_VerletIntegrateVector",NewName="/Script/RigVM.RigVMFunction_VerletIntegrateVector")
 +StructRedirects=(OldName="/Script/ControlRig.RigUnit_DebugBase",NewName="/Script/RigVM.RigVMFunction_DebugBase")
 +StructRedirects=(OldName="/Script/ControlRig.RigUnit_DebugBaseMutable",NewName="/Script/RigVM.RigVMFunction_DebugBaseMutable")
-+StructRedirects=(OldName="/Script/ControlRig.RigUnit_DebugLine",NewName="/Script/RigVM.RigVMFunction_DebugLine")
-+StructRedirects=(OldName="/Script/ControlRig.RigUnit_DebugLineItemSpace",NewName="/Script/RigVM.RigVMFunction_DebugLineNoSpace")
-+StructRedirects=(OldName="/Script/ControlRig.RigUnit_DebugLineStrip",NewName="/Script/RigVM.RigVMFunction_DebugLineStrip")
-+StructRedirects=(OldName="/Script/ControlRig.RigUnit_DebugLineStripItemSpace",NewName="/Script/RigVM.RigVMFunction_DebugLineStripNoSpace")
 +StructRedirects=(OldName="/Script/ControlRig.RigUnit_DebugPoint",NewName="/Script/RigVM.RigVMFunction_DebugPoint")
 +StructRedirects=(OldName="/Script/ControlRig.RigUnit_DebugPointMutable",NewName="/Script/RigVM.RigVMFunction_DebugPointMutable")
-+StructRedirects=(OldName="/Script/ControlRig.RigUnit_DebugRectangle",NewName="/Script/RigVM.RigVMFunction_DebugRectangle")
-+StructRedirects=(OldName="/Script/ControlRig.RigUnit_DebugRectangleItemSpace",NewName="/Script/RigVM.RigVMFunction_DebugRectangleNoSpace")
-+StructRedirects=(OldName="/Script/ControlRig.RigUnit_DebugArc",NewName="/Script/RigVM.RigVMFunction_DebugArc")
-+StructRedirects=(OldName="/Script/ControlRig.RigUnit_DebugArcItemSpace",NewName="/Script/RigVM.RigVMFunction_DebugArcNoSpace")
-+StructRedirects=(OldName="/Script/ControlRig.RigUnit_DebugTransform",NewName="/Script/RigVM.RigVMFunction_DebugTransform")
-+StructRedirects=(OldName="/Script/ControlRig.RigUnit_DebugTransformMutable",NewName="/Script/RigVM.RigVMFunction_DebugTransformMutable")
-+StructRedirects=(OldName="/Script/ControlRig.RigUnit_DebugTransformMutableItemSpace",NewName="/Script/RigVM.RigVMFunction_DebugTransformMutableNoSpace")
-+StructRedirects=(OldName="/Script/ControlRig.RigUnit_DebugTransformArrayMutable_WorkData",NewName="/Script/RigVM.RigVMFunction_DebugTransformArrayMutable_WorkData")
-+StructRedirects=(OldName="/Script/ControlRig.RigUnit_DebugTransformArrayMutable",NewName="/Script/RigVM.RigVMFunction_DebugTransformArrayMutable")
-+StructRedirects=(OldName="/Script/ControlRig.RigUnit_DebugTransformArrayMutableItemSpace",NewName="/Script/RigVM.RigVMFunction_DebugTransformArrayMutableNoSpace")
-+StructRedirects=(OldName="/Script/ControlRig.RigUnit_VisualDebugVector",NewName="/Script/RigVM.RigVMFunction_VisualDebugVector")
-+StructRedirects=(OldName="/Script/ControlRig.RigUnit_VisualDebugVectorItemSpace",NewName="/Script/RigVM.RigVMFunction_VisualDebugVectorNoSpace")
-+StructRedirects=(OldName="/Script/ControlRig.RigUnit_VisualDebugQuat",NewName="/Script/RigVM.RigVMFunction_VisualDebugQuat")
-+StructRedirects=(OldName="/Script/ControlRig.RigUnit_VisualDebugQuatItemSpace",NewName="/Script/RigVM.RigVMFunction_VisualDebugQuatNoSpace")
-+StructRedirects=(OldName="/Script/ControlRig.RigUnit_VisualDebugTransform",NewName="/Script/RigVM.RigVMFunction_VisualDebugTransform")
-+StructRedirects=(OldName="/Script/ControlRig.RigUnit_VisualDebugTransformItemSpace",NewName="/Script/RigVM.RigVMFunction_VisualDebugTransformNoSpace")
 +StructRedirects=(OldName="/Script/ControlRig.RigUnit_AnimBase",NewName="/Script/RigVM.RigVMFunction_AnimBase")
 +StructRedirects=(OldName="/Script/ControlRig.RigUnit_AnimEasingType",NewName="/Script/RigVM.RigVMFunction_AnimEasingType")
 +StructRedirects=(OldName="/Script/ControlRig.RigUnit_AnimEasing",NewName="/Script/RigVM.RigVMFunction_AnimEasing")
diff --git a/Engine/Plugins/Animation/ControlRig/Source/ControlRig/Private/Units/Debug/RigUnit_DebugBase.h b/Engine/Plugins/Animation/ControlRig/Source/ControlRig/Private/Units/Debug/RigUnit_DebugBase.h
new file mode 100644
--- /dev/null
+++ b/Engine/Plugins/Animation/ControlRig/Source/ControlRig/Private/Units/Debug/RigUnit_DebugBase.h
@@ -0,0 +1,18 @@
+// Copyright Epic Games, Inc. All Rights Reserved.
+
+#pragma once
+
+#include "Units/RigUnit.h"
+#include "RigUnit_DebugBase.generated.h"
+
+USTRUCT(meta=(Abstract, Category="Debug", NodeColor = "0.83077 0.846873 0.049707"))
+struct CONTROLRIG_API FRigUnit_DebugBase : public FRigUnit
+{
+	GENERATED_BODY()
+};
+
+USTRUCT(meta=(Abstract, Category="Debug", NodeColor = "0.83077 0.846873 0.049707"))
+struct CONTROLRIG_API FRigUnit_DebugBaseMutable : public FRigUnitMutable
+{
+	GENERATED_BODY()
+};
diff --git a/Engine/Plugins/Animation/ControlRig/Source/ControlRig/Private/Units/Debug/RigUnit_DebugLine.cpp b/Engine/Plugins/Animation/ControlRig/Source/ControlRig/Private/Units/Debug/RigUnit_DebugLine.cpp
new file mode 100644
--- /dev/null
+++ b/Engine/Plugins/Animation/ControlRig/Source/ControlRig/Private/Units/Debug/RigUnit_DebugLine.cpp
@@ -0,0 +1,67 @@
+// Copyright Epic Games, Inc. All Rights Reserved.
+
+#include "Units/Debug/RigUnit_DebugLine.h"
+#include "Units/RigUnitContext.h"
+#include "RigVMFunctions/Debug/RigVMFunction_DebugLine.h"
+
+#include UE_INLINE_GENERATED_CPP_BY_NAME(RigUnit_DebugLine)
+
+FRigUnit_DebugLine_Execute()
+{
+	FRigUnit_DebugLineItemSpace::StaticExecute(
+		ExecuteContext, 
+		A,
+		B,
+		Color,
+		Thickness,
+		FRigElementKey(Space, ERigElementType::Bone), 
+		WorldOffset, 
+		bEnabled);
+}
+
+FRigVMStructUpgradeInfo FRigUnit_DebugLine::GetUpgradeInfo() const
+{
+	FRigVMFunction_DebugLineNoSpace NewNode;
+	NewNode.A = A;
+	NewNode.B = B;
+	NewNode.Color = Color;
+	NewNode.Thickness = Thickness;
+	NewNode.WorldOffset = WorldOffset;
+	NewNode.bEnabled = bEnabled;
+
+	FRigVMStructUpgradeInfo Info(*this, NewNode);
+	return Info;
+}
+
+FRigUnit_DebugLineItemSpace_Execute()
+{
+    DECLARE_SCOPE_HIERARCHICAL_COUNTER_RIGUNIT()
+	if (ExecuteContext.GetDrawInterface() == nullptr || !bEnabled)
+	{
+		return;
+	}
+
+	FVector DrawA = A, DrawB = B;
+	if (Space.IsValid())
+	{
+		FTransform Transform = ExecuteContext.Hierarchy->GetGlobalTransform(Space);
+		DrawA = Transform.TransformPosition(DrawA);
+		DrawB = Transform.TransformPosition(DrawB);
+	}
+
+	ExecuteContext.GetDrawInterface()->DrawLine(WorldOffset, DrawA, DrawB, Color, Thickness);
+}
+
+FRigVMStructUpgradeInfo FRigUnit_DebugLineItemSpace::GetUpgradeInfo() const
+{
+	FRigVMFunction_DebugLineNoSpace NewNode;
+	NewNode.A = A;
+	NewNode.B = B;
+	NewNode.Color = Color;
+	NewNode.Thickness = Thickness;
+	NewNode.WorldOffset = WorldOffset;
+	NewNode.bEnabled = bEnabled;
+
+	FRigVMStructUpgradeInfo Info(*this, NewNode);
+	return Info;
+}
\ No newline at end of file
diff --git a/Engine/Plugins/Animation/ControlRig/Source/ControlRig/Private/Units/Debug/RigUnit_DebugLine.h b/Engine/Plugins/Animation/ControlRig/Source/ControlRig/Private/Units/Debug/RigUnit_DebugLine.h
new file mode 100644
--- /dev/null
+++ b/Engine/Plugins/Animation/ControlRig/Source/ControlRig/Private/Units/Debug/RigUnit_DebugLine.h
@@ -0,0 +1,94 @@
+// Copyright Epic Games, Inc. All Rights Reserved.
+
+#pragma once
+
+#include "RigUnit_DebugBase.h"
+#include "RigUnit_DebugLine.generated.h"
+
+USTRUCT(meta=(DisplayName="Draw Line", Deprecated = "4.25"))
+struct CONTROLRIG_API FRigUnit_DebugLine : public FRigUnit_DebugBaseMutable
+{
+	GENERATED_BODY()
+
+	FRigUnit_DebugLine()
+	{
+		A = B = FVector::ZeroVector;
+		Color = FLinearColor::Red;
+		Thickness = 0.f;
+		WorldOffset = FTransform::Identity;
+		bEnabled = true;
+	}
+
+	RIGVM_METHOD()
+	virtual void Execute() override;
+
+	UPROPERTY(meta = (Input))
+	FVector A;
+
+	UPROPERTY(meta = (Input))
+	FVector B;
+
+	UPROPERTY(meta = (Input))
+	FLinearColor Color;
+
+	UPROPERTY(meta = (Input))
+	float Thickness;
+
+	UPROPERTY(meta = (Input))
+	FName Space;
+
+	UPROPERTY(meta = (Input))
+	FTransform WorldOffset;
+
+	UPROPERTY(meta = (Input, Constant))
+	bool bEnabled;
+
+	RIGVM_METHOD()
+	virtual FRigVMStructUpgradeInfo GetUpgradeInfo() const override;
+};
+
+/**
+ * Draws a line in the viewport given a start and end vector
+ */
+USTRUCT(meta=(DisplayName="Draw Line", Deprecated = "5.2"))
+struct CONTROLRIG_API FRigUnit_DebugLineItemSpace : public FRigUnit_DebugBaseMutable
+{
+	GENERATED_BODY()
+
+	FRigUnit_DebugLineItemSpace()
+	{
+		A = B = FVector::ZeroVector;
+		Color = FLinearColor::Red;
+		Thickness = 0.f;
+		Space = FRigElementKey(NAME_None, ERigElementType::Bone);
+		WorldOffset = FTransform::Identity;
+		bEnabled = true;
+	}
+
+	RIGVM_METHOD()
+	virtual void Execute() override;
+
+	UPROPERTY(meta = (Input))
+	FVector A;
+
+	UPROPERTY(meta = (Input))
+	FVector B;
+
+	UPROPERTY(meta = (Input))
+	FLinearColor Color;
+
+	UPROPERTY(meta = (Input))
+	float Thickness;
+
+	UPROPERTY(meta = (Input))
+	FRigElementKey Space;
+
+	UPROPERTY(meta = (Input))
+	FTransform WorldOffset;
+
+	UPROPERTY(meta = (Input))
+	bool bEnabled;
+
+	RIGVM_METHOD()
+	virtual FRigVMStructUpgradeInfo GetUpgradeInfo() const override;
+};
\ No newline at end of file
diff --git a/Engine/Plugins/Animation/ControlRig/Source/ControlRig/Private/Units/Debug/RigUnit_DebugLineStrip.cpp b/Engine/Plugins/Animation/ControlRig/Source/ControlRig/Private/Units/Debug/RigUnit_DebugLineStrip.cpp
new file mode 100644
--- /dev/null
+++ b/Engine/Plugins/Animation/ControlRig/Source/ControlRig/Private/Units/Debug/RigUnit_DebugLineStrip.cpp
@@ -0,0 +1,71 @@
+// Copyright Epic Games, Inc. All Rights Reserved.
+
+#include "Units/Debug/RigUnit_DebugLineStrip.h"
+#include "Units/RigUnitContext.h"
+#include "RigVMFunctions/Debug/RigVMFunction_DebugLineStrip.h"
+
+#include UE_INLINE_GENERATED_CPP_BY_NAME(RigUnit_DebugLineStrip)
+
+FRigUnit_DebugLineStrip_Execute()
+{
+	FRigUnit_DebugLineStripItemSpace::StaticExecute(
+		ExecuteContext, 
+		Points,
+		Color,
+		Thickness,
+		FRigElementKey(Space, ERigElementType::Bone), 
+		WorldOffset, 
+		bEnabled);
+}
+
+FRigVMStructUpgradeInfo FRigUnit_DebugLineStrip::GetUpgradeInfo() const
+{
+	FRigVMFunction_DebugLineStripNoSpace NewNode;
+	NewNode.Points = Points;
+	NewNode.Color = Color;
+	NewNode.Thickness = Thickness;
+	NewNode.WorldOffset = WorldOffset;
+	NewNode.bEnabled = bEnabled;
+
+	FRigVMStructUpgradeInfo Info(*this, NewNode);
+	return Info;
+}
+
+FRigUnit_DebugLineStripItemSpace_Execute()
+{
+    DECLARE_SCOPE_HIERARCHICAL_COUNTER_RIGUNIT()
+	if (ExecuteContext.GetDrawInterface() == nullptr || !bEnabled)
+	{
+		return;
+	}
+
+	if (Space.IsValid())
+	{
+		FTransform Transform = ExecuteContext.Hierarchy->GetGlobalTransform(Space);
+		TArray<FVector> PointsTransformed;
+		PointsTransformed.Reserve(Points.Num());
+		for(const FVector& Point : Points)
+		{
+			PointsTransformed.Add(Transform.TransformPosition(Point));
+		}
+		ExecuteContext.GetDrawInterface()->DrawLineStrip(WorldOffset, PointsTransformed, Color, Thickness);
+	}
+	else
+	{
+		ExecuteContext.GetDrawInterface()->DrawLineStrip(WorldOffset, TArrayView<const FVector>(Points.GetData(), Points.Num()), Color, Thickness);
+	}
+}
+
+FRigVMStructUpgradeInfo FRigUnit_DebugLineStripItemSpace::GetUpgradeInfo() const
+{
+	FRigVMFunction_DebugLineStripNoSpace NewNode;
+	NewNode.Points = Points;
+	NewNode.Color = Color;
+	NewNode.Thickness = Thickness;
+	NewNode.WorldOffset = WorldOffset;
+	NewNode.bEnabled = bEnabled;
+
+	FRigVMStructUpgradeInfo Info(*this, NewNode);
+	return Info;
+}
+
diff --git a/Engine/Plugins/Animation/ControlRig/Source/ControlRig/Private/Units/Debug/RigUnit_DebugLineStrip.h b/Engine/Plugins/Animation/ControlRig/Source/ControlRig/Private/Units/Debug/RigUnit_DebugLineStrip.h
new file mode 100644
--- /dev/null
+++ b/Engine/Plugins/Animation/ControlRig/Source/ControlRig/Private/Units/Debug/RigUnit_DebugLineStrip.h
@@ -0,0 +1,86 @@
+// Copyright Epic Games, Inc. All Rights Reserved.
+
+#pragma once
+
+#include "RigUnit_DebugBase.h"
+#include "RigUnit_DebugLineStrip.generated.h"
+
+USTRUCT(meta=(DisplayName="Draw Line Strip", Deprecated = "4.25"))
+struct CONTROLRIG_API FRigUnit_DebugLineStrip : public FRigUnit_DebugBaseMutable
+{
+	GENERATED_BODY()
+
+	FRigUnit_DebugLineStrip()
+	{
+		Color = FLinearColor::Red;
+		Thickness = 0.f;
+		WorldOffset = FTransform::Identity;
+		bEnabled = true;
+	}
+
+	RIGVM_METHOD()
+	virtual void Execute() override;
+
+	UPROPERTY(meta = (Input))
+	TArray<FVector> Points;
+
+	UPROPERTY(meta = (Input))
+	FLinearColor Color;
+
+	UPROPERTY(meta = (Input))
+	float Thickness;
+
+	UPROPERTY(meta = (Input))
+	FName Space;
+
+	UPROPERTY(meta = (Input))
+	FTransform WorldOffset;
+	
+	UPROPERTY(meta = (Input, Constant))
+	bool bEnabled;
+
+	RIGVM_METHOD()
+	virtual FRigVMStructUpgradeInfo GetUpgradeInfo() const override;
+};
+
+/**
+ * Draws a line strip in the viewport given any number of points
+ */
+USTRUCT(meta=(DisplayName="Draw Line Strip", Deprecated = "5.2"))
+struct CONTROLRIG_API FRigUnit_DebugLineStripItemSpace : public FRigUnit_DebugBaseMutable
+{
+	GENERATED_BODY()
+
+	FRigUnit_DebugLineStripItemSpace()
+	{
+		Color = FLinearColor::Red;
+		Thickness = 0.f;
+		Space = FRigElementKey(NAME_None, ERigElementType::Bone);
+		WorldOffset = FTransform::Identity;
+		bEnabled = true;
+	}
+
+	RIGVM_METHOD()
+	virtual void Execute() override;
+
+	UPROPERTY(meta = (Input))
+	TArray<FVector> Points;
+
+	UPROPERTY(meta = (Input))
+	FLinearColor Color;
+
+	UPROPERTY(meta = (Input))
+	float Thickness;
+
+	UPROPERTY(meta = (Input))
+	FRigElementKey Space;
+
+	UPROPERTY(meta = (Input))
+	FTransform WorldOffset;
+	
+	UPROPERTY(meta = (Input))
+	bool bEnabled;
+	
+	RIGVM_METHOD()
+	virtual FRigVMStructUpgradeInfo GetUpgradeInfo() const override;
+};
diff --git a/Engine/Plugins/Animation/ControlRig/Source/ControlRig/Private/Units/Debug/RigUnit_DebugPrimitives.cpp b/Engine/Plugins/Animation/ControlRig/Source/ControlRig/Private/Units/Debug/RigUnit_DebugPrimitives.cpp
new file mode 100644
--- /dev/null
+++ b/Engine/Plugins/Animation/ControlRig/Source/ControlRig/Private/Units/Debug/RigUnit_DebugPrimitives.cpp
@@ -0,0 +1,130 @@
+// Copyright Epic Games, Inc. All Rights Reserved.
+
+#include "Units/Debug/RigUnit_DebugPrimitives.h"
+#include "Units/RigUnitContext.h"
+#include "RigVMFunctions/Debug/RigVMFunction_DebugPrimitives.h"
+
+#include UE_INLINE_GENERATED_CPP_BY_NAME(RigUnit_DebugPrimitives)
+
+FRigUnit_DebugRectangle_Execute()
+{
+	FRigUnit_DebugRectangleItemSpace::StaticExecute(
+		ExecuteContext, 
+		Transform,
+		Color,
+		Scale,
+		Thickness,
+		FRigElementKey(Space, ERigElementType::Bone), 
+		WorldOffset, 
+		bEnabled);
+}
+
+FRigVMStructUpgradeInfo FRigUnit_DebugRectangle::GetUpgradeInfo() const
+{
+	FRigVMFunction_DebugRectangleNoSpace NewNode;
+	NewNode.Transform = Transform;
+	NewNode.Color = Color;
+	NewNode.Scale = Scale;
+	NewNode.Thickness = Thickness;
+	NewNode.WorldOffset = WorldOffset;
+	NewNode.bEnabled = bEnabled;
+
+	FRigVMStructUpgradeInfo Info(*this, NewNode);
+	return Info;
+}
+
+FRigVMStructUpgradeInfo FRigUnit_DebugRectangleItemSpace::GetUpgradeInfo() const
+{
+	FRigVMFunction_DebugRectangleNoSpace NewNode;
+	NewNode.Transform = Transform;
+	NewNode.Color = Color;
+	NewNode.Scale = Scale;
+	NewNode.Thickness = Thickness;
+	NewNode.WorldOffset = WorldOffset;
+	NewNode.bEnabled = bEnabled;
+
+	return FRigVMStructUpgradeInfo(*this, NewNode);
+}
+
+FRigUnit_DebugRectangleItemSpace_Execute()
+{
+    DECLARE_SCOPE_HIERARCHICAL_COUNTER_RIGUNIT()
+	if (ExecuteContext.GetDrawInterface() == nullptr || !bEnabled)
+	{
+		return;
+	}
+
+	FTransform DrawTransform = Transform;
+	if (Space.IsValid())
+	{
+		DrawTransform = DrawTransform * ExecuteContext.Hierarchy->GetGlobalTransform(Space);
+	}
+
+	ExecuteContext.GetDrawInterface()->DrawRectangle(WorldOffset, DrawTransform, Scale, Color, Thickness);
+}
+
+FRigUnit_DebugArc_Execute()
+{
+	FRigUnit_DebugArcItemSpace::StaticExecute(
+		ExecuteContext, 
+		Transform,
+		Color,
+		Radius,
+		MinimumDegrees,
+		MaximumDegrees,
+		Thickness,
+		Detail,
+		FRigElementKey(Space, ERigElementType::Bone), 
+		WorldOffset, 
+		bEnabled);
+}
+
+FRigVMStructUpgradeInfo FRigUnit_DebugArc::GetUpgradeInfo() const
+{
+	FRigVMFunction_DebugArcNoSpace NewNode;
+	NewNode.Transform = Transform;
+	NewNode.Color = Color;
+	NewNode.Radius = Radius;
+	NewNode.MinimumDegrees = MinimumDegrees;
+	NewNode.MaximumDegrees = MaximumDegrees;
+	NewNode.Thickness = Thickness;
+	NewNode.Detail = Detail;
+	NewNode.WorldOffset = WorldOffset;
+	NewNode.bEnabled = bEnabled;
+
+	FRigVMStructUpgradeInfo Info(*this, NewNode);
+	return Info;
+}
+
+FRigVMStructUpgradeInfo FRigUnit_DebugArcItemSpace::GetUpgradeInfo() const
+{
+	FRigVMFunction_DebugArcNoSpace NewNode;
+	NewNode.Transform = Transform;
+	NewNode.Color = Color;
+	NewNode.Radius = Radius;
+	NewNode.MinimumDegrees = MinimumDegrees;
+	NewNode.MaximumDegrees = MaximumDegrees;
+	NewNode.Thickness = Thickness;
+	NewNode.Detail = Detail;
+	NewNode.WorldOffset = WorldOffset;
+	NewNode.bEnabled = bEnabled;
+
+	return FRigVMStructUpgradeInfo(*this, NewNode);
+}
+
+FRigUnit_DebugArcItemSpace_Execute()
+{
+    DECLARE_SCOPE_HIERARCHICAL_COUNTER_RIGUNIT()
+	if (ExecuteContext.GetDrawInterface() == nullptr || !bEnabled)
+	{
+		return;
+	}
+
+	FTransform DrawTransform = Transform;
+	if (Space.IsValid())
+	{
+		DrawTransform = DrawTransform * ExecuteContext.Hierarchy->GetGlobalTransform(Space);
+	}
+
+	ExecuteContext.GetDrawInterface()->DrawArc(WorldOffset, DrawTransform, Radius, FMath::DegreesToRadians(MinimumDegrees), FMath::DegreesToRadians(MaximumDegrees), Color, Thickness, Detail);
+}
diff --git a/Engine/Plugins/Animation/ControlRig/Source/ControlRig/Private/Units/Debug/RigUnit_DebugPrimitives.h b/Engine/Plugins/Animation/ControlRig/Source/ControlRig/Private/Units/Debug/RigUnit_DebugPrimitives.h
new file mode 100644
--- /dev/null
+++ b/Engine/Plugins/Animation/ControlRig/Source/ControlRig/Private/Units/Debug/RigUnit_DebugPrimitives.h
@@ -0,0 +1,209 @@
+// Copyright Epic Games, Inc. All Rights Reserved.
+
+#pragma once
+
+#include "RigUnit_DebugBase.h"
+#include "RigUnit_DebugPrimitives.generated.h"
+
+USTRUCT(meta=(DisplayName="Draw Rectangle", Keywords="Draw Square", Deprecated = "4.25"))
+struct CONTROLRIG_API FRigUnit_DebugRectangle : public FRigUnit_DebugBaseMutable
+{
+	GENERATED_BODY()
+
+	FRigUnit_DebugRectangle()
+	{
+		Transform = WorldOffset = FTransform::Identity;
+		Color = FLinearColor::Red;
+		Thickness = 0.f;
+		Scale = 10.f;
+		WorldOffset = FTransform::Identity;
+		bEnabled = true;
+	}
+
+	RIGVM_METHOD()
+	virtual void Execute() override;
+
+	UPROPERTY(meta = (Input))
+	FTransform Transform;
+
+	UPROPERTY(meta = (Input))
+	FLinearColor Color;
+
+	UPROPERTY(meta = (Input))
+	float Scale;
+
+	UPROPERTY(meta = (Input))
+	float Thickness;
+
+	UPROPERTY(meta = (Input))
+	FName Space;
+
+	UPROPERTY(meta = (Input))
+	FTransform WorldOffset;
+
+	UPROPERTY(meta = (Input, Constant))
+	bool bEnabled;
+
+	RIGVM_METHOD()
+	virtual FRigVMStructUpgradeInfo GetUpgradeInfo() const override;
+};
+
+/**
+ * Draws a rectangle in the viewport given a transform
+ */
+USTRUCT(meta=(DisplayName="Draw Rectangle", Keywords="Draw Square", Deprecated = "5.2"))
+struct CONTROLRIG_API FRigUnit_DebugRectangleItemSpace : public FRigUnit_DebugBaseMutable
+{
+	GENERATED_BODY()
+
+	FRigUnit_DebugRectangleItemSpace()
+	{
+		Transform = WorldOffset = FTransform::Identity;
+		Color = FLinearColor::Red;
+		Thickness = 0.f;
+		Scale = 10.f;
+		Space = FRigElementKey(NAME_None, ERigElementType::Bone);
+		WorldOffset = FTransform::Identity;
+		bEnabled = true;
+	}
+
+	RIGVM_METHOD()
+	virtual void Execute() override;
+
+	UPROPERTY(meta = (Input))
+	FTransform Transform;
+
+	UPROPERTY(meta = (Input))
+	FLinearColor Color;
+
+	UPROPERTY(meta = (Input))
+	float Scale;
+
+	UPROPERTY(meta = (Input))
+	float Thickness;
+
+	UPROPERTY(meta = (Input))
+	FRigElementKey Space;
+
+	UPROPERTY(meta = (Input))
+	FTransform WorldOffset;
+
+	UPROPERTY(meta = (Input))
+	bool bEnabled;
+
+	RIGVM_METHOD()
+	virtual FRigVMStructUpgradeInfo GetUpgradeInfo() const override;
+};
+
+
+USTRUCT(meta=(DisplayName="Draw Arc", Keywords="Draw Ellipse, Draw Circle", Deprecated = "4.25"))
+struct CONTROLRIG_API FRigUnit_DebugArc : public FRigUnit_DebugBaseMutable
+{
+	GENERATED_BODY()
+
+	FRigUnit_DebugArc()
+	{
+		Transform = WorldOffset = FTransform::Identity;
+		Color = FLinearColor::Red;
+		Thickness = MinimumDegrees = 0.f;
+		Radius = 10.f;
+		MaximumDegrees = 360.f;
+		Detail = 16;
+		WorldOffset = FTransform::Identity;
+		bEnabled = true;
+	}
+
+	RIGVM_METHOD()
+	virtual void Execute() override;
+
+	UPROPERTY(meta = (Input))
+	FTransform Transform;
+
+	UPROPERTY(meta = (Input))
+	FLinearColor Color;
+
+	UPROPERTY(meta = (Input))
+	float Radius;
+
+	UPROPERTY(meta = (Input))
+	float MinimumDegrees;
+
+	UPROPERTY(meta = (Input))
+	float MaximumDegrees;
+
+	UPROPERTY(meta = (Input))
+	float Thickness;
+
+	UPROPERTY(meta = (Input))
+	int32 Detail;
+
+	UPROPERTY(meta = (Input))
+	FName Space;
+
+	UPROPERTY(meta = (Input))
+	FTransform WorldOffset;
+
+	UPROPERTY(meta = (Input, Constant))
+	bool bEnabled;
+
+	RIGVM_METHOD()
+	virtual FRigVMStructUpgradeInfo GetUpgradeInfo() const override;
+};
+
+/**
+ * Draws an arc in the viewport, can take in different min and max degrees
+ */
+USTRUCT(meta=(DisplayName="Draw Arc", Keywords="Draw Ellipse, Draw Circle", Deprecated = "5.2"))
+struct CONTROLRIG_API FRigUnit_DebugArcItemSpace : public FRigUnit_DebugBaseMutable
+{
+	GENERATED_BODY()
+
+	FRigUnit_DebugArcItemSpace()
+	{
+		Transform = WorldOffset = FTransform::Identity;
+		Color = FLinearColor::Red;
+		Thickness = MinimumDegrees = 0.f;
+		Radius = 10.f;
+		MaximumDegrees = 360.f;
+		Detail = 16;
+		Space = FRigElementKey(NAME_None, ERigElementType::Bone);
+		WorldOffset = FTransform::Identity;
+		bEnabled = true;
+	}
+
+	RIGVM_METHOD()
+	virtual void Execute() override;
+
+	UPROPERTY(meta = (Input))
+	FTransform Transform;
+
+	UPROPERTY(meta = (Input))
+	FLinearColor Color;
+
+	UPROPERTY(meta = (Input))
+	float Radius;
+
+	UPROPERTY(meta = (Input))
+	float MinimumDegrees;
+
+	UPROPERTY(meta = (Input))
+	float MaximumDegrees;
+
+	UPROPERTY(meta = (Input))
+	float Thickness;
+
+	UPROPERTY(meta = (Input))
+	int32 Detail;
+
+	UPROPERTY(meta = (Input))
+	FRigElementKey Space;
+
+	UPROPERTY(meta = (Input))
+	FTransform WorldOffset;
+
+	UPROPERTY(meta = (Input))
+	bool bEnabled;
+
+	RIGVM_METHOD()
+	virtual FRigVMStructUpgradeInfo GetUpgradeInfo() const override;
+};
\ No newline at end of file
diff --git a/Engine/Plugins/Animation/ControlRig/Source/ControlRig/Private/Units/Debug/RigUnit_DebugTransform.cpp b/Engine/Plugins/Animation/ControlRig/Source/ControlRig/Private/Units/Debug/RigUnit_DebugTransform.cpp
new file mode 100644
--- /dev/null
+++ b/Engine/Plugins/Animation/ControlRig/Source/ControlRig/Private/Units/Debug/RigUnit_DebugTransform.cpp
@@ -0,0 +1,258 @@
+// Copyright Epic Games, Inc. All Rights Reserved.
+
+#include "Units/Debug/RigUnit_DebugTransform.h"
+#include "Units/Debug/RigUnit_VisualDebug.h"
+#include "Units/RigUnitContext.h"
+#include "RigVMFunctions/Debug/RigVMFunction_DebugTransform.h"
+#include "RigVMFunctions/Debug/RigVMFunction_VisualDebug.h"
+
+#include UE_INLINE_GENERATED_CPP_BY_NAME(RigUnit_DebugTransform)
+
+FRigUnit_DebugTransform_Execute()
+{
+    DECLARE_SCOPE_HIERARCHICAL_COUNTER_RIGUNIT()
+	if (ExecuteContext.GetDrawInterface() == nullptr || !bEnabled)
+	{
+		return;
+	}
+
+	FTransform DrawTransform = Transform;
+	if (Space != NAME_None && ExecuteContext.Hierarchy != nullptr)
+	{
+		DrawTransform = Transform * ExecuteContext.Hierarchy->GetGlobalTransform(FRigElementKey(Space, ERigElementType::Bone));
+	}
+
+	switch (Mode)
+	{
+		case ERigUnitDebugTransformMode::Axes:
+		{
+			ExecuteContext.GetDrawInterface()->DrawAxes(WorldOffset, DrawTransform, Scale, Thickness);
+			break;
+		}
+		case ERigUnitDebugTransformMode::Point:
+		{
+			ExecuteContext.GetDrawInterface()->DrawPoint(WorldOffset, DrawTransform.GetLocation(), Scale, Color);
+			break;
+		}
+		case ERigUnitDebugTransformMode::Box:
+		{
+			DrawTransform.SetScale3D(DrawTransform.GetScale3D() * Scale);
+			ExecuteContext.GetDrawInterface()->DrawBox(WorldOffset, DrawTransform, Color, Thickness);
+			break;
+		}
+	}
+}
+
+FRigVMStructUpgradeInfo FRigUnit_DebugTransform::GetUpgradeInfo() const
+{
+	FRigVMFunction_DebugTransformMutableNoSpace NewNode;
+	NewNode.Transform = Transform;
+	NewNode.Mode = Mode;
+	NewNode.Color = Color;
+	NewNode.Thickness = Thickness;
+	NewNode.Scale = Scale;
+	NewNode.WorldOffset = WorldOffset;
+	NewNode.bEnabled = bEnabled;
+
+	FRigVMStructUpgradeInfo Info(*this, NewNode);
+	return Info;
+}
+
+FRigUnit_DebugTransformMutable_Execute()
+{
+	FRigUnit_DebugTransformMutableItemSpace::StaticExecute(
+		ExecuteContext, 
+		Transform,
+		Mode,
+		Color,
+		Thickness,
+		Scale,
+		FRigElementKey(Space, ERigElementType::Bone), 
+		WorldOffset, 
+		bEnabled);
+}
+
+FRigVMStructUpgradeInfo FRigUnit_DebugTransformMutable::GetUpgradeInfo() const
+{
+	FRigVMFunction_DebugTransformMutableNoSpace NewNode;
+	NewNode.Transform = Transform;
+	NewNode.Mode = Mode;
+	NewNode.Color = Color;
+	NewNode.Thickness = Thickness;
+	NewNode.Scale = Scale;
+	NewNode.WorldOffset = WorldOffset;
+	NewNode.bEnabled = bEnabled;
+
+	FRigVMStructUpgradeInfo Info(*this, NewNode);
+	return Info;
+}
+
+FRigVMStructUpgradeInfo FRigUnit_DebugTransformMutableItemSpace::GetUpgradeInfo() const
+{
+	FRigVMFunction_DebugTransformMutableNoSpace NewNode;
+	NewNode.Transform = Transform;
+	NewNode.Mode = Mode;
+	NewNode.Color = Color;
+	NewNode.Thickness = Thickness;
+	NewNode.Scale = Scale;
+	NewNode.WorldOffset = WorldOffset;
+	NewNode.bEnabled = bEnabled;
+
+	FRigVMStructUpgradeInfo Info(*this, NewNode);
+	return Info;
+}
+
+FRigUnit_DebugTransformMutableItemSpace_Execute()
+{
+    DECLARE_SCOPE_HIERARCHICAL_COUNTER_RIGUNIT()
+	if (ExecuteContext.GetDrawInterface() == nullptr || !bEnabled)
+	{
+		return;
+	}
+
+	FTransform DrawTransform = Transform;
+	if (Space.IsValid() && ExecuteContext.Hierarchy != nullptr)
+	{
+		DrawTransform = Transform * ExecuteContext.Hierarchy->GetGlobalTransform(Space);
+	}
+
+	switch (Mode)
+	{
+		case ERigUnitDebugTransformMode::Axes:
+		{
+			ExecuteContext.GetDrawInterface()->DrawAxes(WorldOffset, DrawTransform, Scale, Thickness);
+			break;
+		}
+		case ERigUnitDebugTransformMode::Point:
+		{
+			ExecuteContext.GetDrawInterface()->DrawPoint(WorldOffset, DrawTransform.GetLocation(), Scale, Color);
+			break;
+		}
+		case ERigUnitDebugTransformMode::Box:
+		{
+			DrawTransform.SetScale3D(DrawTransform.GetScale3D() * Scale);
+			ExecuteContext.GetDrawInterface()->DrawBox(WorldOffset, DrawTransform, Color, Thickness);
+			break;
+		}
+	}
+}
+
+FRigUnit_DebugTransformArrayMutable_Execute()
+{
+    DECLARE_SCOPE_HIERARCHICAL_COUNTER_RIGUNIT()
+	if (ExecuteContext.GetDrawInterface() == nullptr || !bEnabled)
+	{
+		return;
+	}
+
+	WorkData.DrawTransforms.SetNumUninitialized(Transforms.Num());
+	if (Space != NAME_None && ExecuteContext.Hierarchy)
+	{
+		FTransform SpaceTransform = ExecuteContext.Hierarchy->GetGlobalTransform(FRigElementKey(Space, ERigElementType::Bone));
+		for(int32 Index=0;Index<Transforms.Num();Index++)
+		{
+			WorkData.DrawTransforms[Index] = Transforms[Index] * SpaceTransform;
+		}
+	}
+	else
+	{
+		for(int32 Index=0;Index<Transforms.Num();Index++)
+		{
+			WorkData.DrawTransforms[Index] = Transforms[Index];
+		}
+	}
+
+	for(FTransform& DrawTransform : WorkData.DrawTransforms)
+	{
+		switch (Mode)
+		{
+			case ERigUnitDebugTransformMode::Axes:
+			{
+				ExecuteContext.GetDrawInterface()->DrawAxes(WorldOffset, DrawTransform, Scale, Thickness);
+				break;
+			}
+			case ERigUnitDebugTransformMode::Point:
+			{
+				ExecuteContext.GetDrawInterface()->DrawPoint(WorldOffset, DrawTransform.GetLocation(), Scale, Color);
+				break;
+			}
+			case ERigUnitDebugTransformMode::Box:
+			{
+				DrawTransform.SetScale3D(DrawTransform.GetScale3D() * Scale);
+				ExecuteContext.GetDrawInterface()->DrawBox(WorldOffset, DrawTransform, Color, Thickness);
+				break;
+			}
+		}
+	}
+}
+
+FRigVMStructUpgradeInfo FRigUnit_DebugTransformArrayMutable::GetUpgradeInfo() const
+{
+	FRigVMFunction_DebugTransformArrayMutableNoSpace NewNode;
+	NewNode.Transforms = Transforms;
+	NewNode.Mode = Mode;
+	NewNode.Color = Color;
+	NewNode.Thickness = Thickness;
+	NewNode.Scale = Scale;
+	NewNode.WorldOffset = WorldOffset;
+	NewNode.bEnabled = bEnabled;
+
+	FRigVMStructUpgradeInfo Info(*this, NewNode);
+	return Info;
+}
+
+FRigVMStructUpgradeInfo FRigUnit_DebugTransformArrayMutableItemSpace::GetUpgradeInfo() const
+{
+	FRigVMFunction_DebugTransformArrayMutableNoSpace NewNode;
+	NewNode.Transforms = Transforms;
+	NewNode.Mode = Mode;
+	NewNode.Color = Color;
+	NewNode.Thickness = Thickness;
+	NewNode.Scale = Scale;
+	NewNode.WorldOffset = WorldOffset;
+	NewNode.bEnabled = bEnabled;
+
+	FRigVMStructUpgradeInfo Info(*this, NewNode);
+	return Info;
+}
+
+FRigUnit_DebugTransformArrayMutableItemSpace_Execute()
+{
+	DECLARE_SCOPE_HIERARCHICAL_COUNTER_RIGUNIT()
+	if (ExecuteContext.GetDrawInterface() == nullptr || !bEnabled || Transforms.Num() == 0)
+	{
+		return;
+	}
+
+	for(const FTransform& Transform : Transforms)
+	{
+		FRigUnit_DebugTransformMutableItemSpace::StaticExecute(
+			ExecuteContext, 
+			Transform,
+			Mode,
+			Color,
+			Thickness,
+			Scale,
+			Space, 
+			WorldOffset, 
+			bEnabled);
+	}
+
+	if(ParentIndices.Num() == Transforms.Num())
+	{
+		for(int32 Index = 0; Index < ParentIndices.Num(); Index++)
+		{
+			if(Transforms.IsValidIndex(ParentIndices[Index]))
+			{
+				ExecuteContext.GetDrawInterface()->DrawLine(
+					WorldOffset,
+					Transforms[Index].GetTranslation(),
+					Transforms[ParentIndices[Index]].GetTranslation(),
+					Color,
+					Thickness
+				);
+			}
+		}
+	}
+}
+
diff --git a/Engine/Plugins/Animation/ControlRig/Source/ControlRig/Private/Units/Debug/RigUnit_DebugTransform.h b/Engine/Plugins/Animation/ControlRig/Source/ControlRig/Private/Units/Debug/RigUnit_DebugTransform.h
new file mode 100644
--- /dev/null
+++ b/Engine/Plugins/Animation/ControlRig/Source/ControlRig/Private/Units/Debug/RigUnit_DebugTransform.h
@@ -0,0 +1,260 @@
+// Copyright Epic Games, Inc. All Rights Reserved.
+
+#pragma once
+
+#include "RigUnit_DebugBase.h"
+#include "RigVMFunctions/Debug/RigVMFunction_DebugTransform.h"
+#include "RigUnit_DebugTransform.generated.h"
+
+USTRUCT(meta=(DisplayName="Draw Transform In Place", Deprecated = "4.25.0"))
+struct CONTROLRIG_API FRigUnit_DebugTransform : public FRigUnit_DebugBase
+{
+	GENERATED_BODY()
+
+	FRigUnit_DebugTransform()
+	{
+		Transform = WorldOffset = FTransform::Identity;
+		Mode = ERigUnitDebugTransformMode::Axes;
+		Color = FLinearColor::Red;
+		Thickness = 0.f;
+		Scale = 10.f;
+		bEnabled = true;
+	}
+
+	RIGVM_METHOD()
+	virtual void Execute() override;
+
+	UPROPERTY(meta = (Input, Output))
+	FTransform Transform;
+
+	UPROPERTY(meta = (Input))
+	ERigUnitDebugTransformMode Mode;
+
+	UPROPERTY(meta = (Input))
+	FLinearColor Color;
+
+	UPROPERTY(meta = (Input))
+	float Thickness;
+
+	UPROPERTY(meta = (Input))
+	float Scale;
+
+	UPROPERTY(meta = (Input))
+	FName Space;
+
+	UPROPERTY(meta = (Input))
+	FTransform WorldOffset;
+
+	UPROPERTY(meta = (Input, Constant))
+	bool bEnabled;
+
+	RIGVM_METHOD()
+	virtual FRigVMStructUpgradeInfo GetUpgradeInfo() const override;
+};
+
+USTRUCT(meta=(DisplayName="Draw Transform", Deprecated = "4.25"))
+struct CONTROLRIG_API FRigUnit_DebugTransformMutable : public FRigUnit_DebugBaseMutable
+{
+	GENERATED_BODY()
+
+	FRigUnit_DebugTransformMutable()
+	{
+		Transform = WorldOffset = FTransform::Identity;
+		Mode = ERigUnitDebugTransformMode::Axes;
+		Color = FLinearColor::Red;
+		Thickness = 0.f;
+		Scale = 10.f;
+		bEnabled = true;
+	}
+
+	RIGVM_METHOD()
+	virtual void Execute() override;
+
+	UPROPERTY(meta = (Input))
+	FTransform Transform;
+
+	UPROPERTY(meta = (Input))
+	ERigUnitDebugTransformMode Mode;
+
+	UPROPERTY(meta = (Input))
+	FLinearColor Color;
+
+	UPROPERTY(meta = (Input))
+	float Thickness;
+
+	UPROPERTY(meta = (Input))
+	float Scale;
+
+	UPROPERTY(meta = (Input))
+	FName Space;
+
+	UPROPERTY(meta = (Input))
+	FTransform WorldOffset;
+
+	UPROPERTY(meta = (Input, Constant))
+	bool bEnabled;
+
+	RIGVM_METHOD()
+	virtual FRigVMStructUpgradeInfo GetUpgradeInfo() const override;
+};
+
+/**
+ * Given a transform, will draw a point, axis, or a box in the viewport
+ */
+USTRUCT(meta=(DisplayName="Draw Transform", Deprecated = "5.2"))
+struct CONTROLRIG_API FRigUnit_DebugTransformMutableItemSpace : public FRigUnit_DebugBaseMutable
+{
+	GENERATED_BODY()
+
+	FRigUnit_DebugTransformMutableItemSpace()
+	{
+		Transform = WorldOffset = FTransform::Identity;
+		Mode = ERigUnitDebugTransformMode::Axes;
+		Color = FLinearColor::Red;
+		Thickness = 0.f;
+		Scale = 10.f;
+		Space = FRigElementKey(NAME_None, ERigElementType::Bone);
+		bEnabled = true;
+	}
+
+	RIGVM_METHOD()
+	virtual void Execute() override;
+
+	UPROPERTY(meta = (Input))
+	FTransform Transform;
+
+	UPROPERTY(meta = (Input))
+	ERigUnitDebugTransformMode Mode;
+
+	UPROPERTY(meta = (Input))
+	FLinearColor Color;
+
+	UPROPERTY(meta = (Input))
+	float Thickness;
+
+	UPROPERTY(meta = (Input))
+	float Scale;
+
+	UPROPERTY(meta = (Input))
+	FRigElementKey Space;
+
+	UPROPERTY(meta = (Input))
+	FTransform WorldOffset;
+
+	UPROPERTY(meta = (Input))
+	bool bEnabled;
+
+	RIGVM_METHOD()
+	virtual FRigVMStructUpgradeInfo GetUpgradeInfo() const override;
+};
+
+USTRUCT()
+struct CONTROLRIG_API FRigUnit_DebugTransformArrayMutable_WorkData
+{
+	GENERATED_BODY()
+		
+	UPROPERTY()
+	TArray<FTransform> DrawTransforms;
+};
+
+USTRUCT(meta=(DisplayName="Draw Transform Array", Deprecated = "4.25.0"))
+struct CONTROLRIG_API FRigUnit_DebugTransformArrayMutable : public FRigUnit_DebugBaseMutable
+{
+ 	GENERATED_BODY()
+
+	FRigUnit_DebugTransformArrayMutable()
+	{
+		Mode = ERigUnitDebugTransformMode::Axes;
+		Color = FLinearColor::Red;
+		Thickness = 0.f;
+		Scale = 10.f;
+		WorldOffset = FTransform::Identity;
+		bEnabled = true;
+	}
+
+	RIGVM_METHOD()
+	virtual void Execute() override;
+
+	UPROPERTY(meta = (Input))
+	TArray<FTransform> Transforms;
+
+	UPROPERTY(meta = (Input))
+	ERigUnitDebugTransformMode Mode;
+
+	UPROPERTY(meta = (Input))
+	FLinearColor Color;
+
+	UPROPERTY(meta = (Input))
+	float Thickness;
+
+	UPROPERTY(meta = (Input))
+	float Scale;
+
+	UPROPERTY(meta = (Input))
+	FName Space;
+
+	UPROPERTY(meta = (Input))
+	FTransform WorldOffset;
+
+	UPROPERTY(meta = (Input, Constant))
+	bool bEnabled;
+
+	UPROPERTY(transient)
+	FRigUnit_DebugTransformArrayMutable_WorkData WorkData;
+
+	RIGVM_METHOD()
+	virtual FRigVMStructUpgradeInfo GetUpgradeInfo() const override;
+};
+
+/**
+* Given a transform array, will draw a point, axis, or a box in the viewport
+*/
+USTRUCT(meta=(DisplayName="Draw Transform Array", Deprecated = "5.2"))
+struct CONTROLRIG_API FRigUnit_DebugTransformArrayMutableItemSpace : public FRigUnit_DebugBaseMutable
+{
+	GENERATED_BODY()
+
+	FRigUnit_DebugTransformArrayMutableItemSpace()
+	{
+		WorldOffset = FTransform::Identity;
+		Mode = ERigUnitDebugTransformMode::Axes;
+		Color = FLinearColor::White;
+		Thickness = 0.f;
+		Scale = 10.f;
+		Space = FRigElementKey(NAME_None, ERigElementType::Bone);
+		bEnabled = true;
+	}
+
+	RIGVM_METHOD()
+	virtual void Execute() override;
+
+	UPROPERTY(meta = (Input))
+	TArray<FTransform> Transforms;
+
+	UPROPERTY(meta = (Input))
+	TArray<int32> ParentIndices;
+
+	UPROPERTY(meta = (Input))
+	ERigUnitDebugTransformMode Mode;
+
+	UPROPERTY(meta = (Input))
+	FLinearColor Color;
+
+	UPROPERTY(meta = (Input))
+	float Thickness;
+
+	UPROPERTY(meta = (Input))
+	float Scale;
+
+	UPROPERTY(meta = (Input))
+	FRigElementKey Space;
+
+	UPROPERTY(meta = (Input))
+	FTransform WorldOffset;
+
+	UPROPERTY(meta = (Input))
+	bool bEnabled;
+
+	RIGVM_METHOD()
+	virtual FRigVMStructUpgradeInfo GetUpgradeInfo() const override;
+};
diff --git a/Engine/Plugins/Animation/ControlRig/Source/ControlRig/Private/Units/Debug/RigUnit_VisualDebug.cpp b/Engine/Plugins/Animation/ControlRig/Source/ControlRig/Private/Units/Debug/RigUnit_VisualDebug.cpp
new file mode 100644
--- /dev/null
+++ b/Engine/Plugins/Animation/ControlRig/Source/ControlRig/Private/Units/Debug/RigUnit_VisualDebug.cpp
@@ -0,0 +1,160 @@
+// Copyright Epic Games, Inc. All Rights Reserved.
+
+#include "Units/Debug/RigUnit_VisualDebug.h"
+#include "Units/RigUnitContext.h"
+#include "RigVMFunctions/Debug/RigVMFunction_VisualDebug.h"
+
+#include UE_INLINE_GENERATED_CPP_BY_NAME(RigUnit_VisualDebug)
+
+FRigUnit_VisualDebugVector_Execute()
+{
+	FRigUnit_VisualDebugVectorItemSpace::StaticExecute(ExecuteContext, Value, bEnabled, Mode, Color, Thickness, Scale, FRigElementKey(BoneSpace, ERigElementType::Bone));
+}
+
+FRigVMStructUpgradeInfo FRigUnit_VisualDebugVector::GetUpgradeInfo() const
+{
+	FRigVMFunction_VisualDebugVectorNoSpace NewNode;
+	NewNode.Value = Value;
+	NewNode.bEnabled = bEnabled;
+	NewNode.Mode = Mode;
+	NewNode.Color = Color;
+	NewNode.Thickness = Thickness;
+	NewNode.Scale = Scale;
+
+	FRigVMStructUpgradeInfo Info(*this, NewNode);
+	return Info;
+}
+
+FRigUnit_VisualDebugVectorItemSpace_Execute()
+{
+    DECLARE_SCOPE_HIERARCHICAL_COUNTER_RIGUNIT()
+
+	if (ExecuteContext.GetDrawInterface() == nullptr || !bEnabled)
+	{
+		return;
+	}
+
+	FTransform WorldOffset = FTransform::Identity;
+	if (Space.IsValid())
+	{
+		WorldOffset = ExecuteContext.Hierarchy->GetGlobalTransform(Space);
+	}
+
+	switch(Mode)
+	{
+		case ERigUnitVisualDebugPointMode::Point:
+		{
+			ExecuteContext.GetDrawInterface()->DrawPoint(WorldOffset, Value, Thickness, Color);
+			break;
+		}
+		case ERigUnitVisualDebugPointMode::Vector:
+		{
+			ExecuteContext.GetDrawInterface()->DrawLine(WorldOffset, FVector::ZeroVector, Value * Scale, Color, Thickness);
+			break;
+		}
+		default:
+		{
+			checkNoEntry();
+		}
+	}
+}
+
+FRigUnit_VisualDebugQuat_Execute()
+{
+	FRigUnit_VisualDebugQuatItemSpace::StaticExecute(ExecuteContext, Value, bEnabled, Thickness, Scale, FRigElementKey(BoneSpace, ERigElementType::Bone));
+}
+
+FRigVMStructUpgradeInfo FRigUnit_VisualDebugVectorItemSpace::GetUpgradeInfo() const
+{
+	FRigVMFunction_VisualDebugVectorNoSpace NewNode;
+	NewNode.Value = Value;
+	NewNode.bEnabled = bEnabled;
+	NewNode.Mode = Mode;
+	NewNode.Color = Color;
+	NewNode.Thickness = Thickness;
+	NewNode.Scale = Scale;
+
+	FRigVMStructUpgradeInfo Info(*this, NewNode);
+	return Info;
+}
+
+FRigVMStructUpgradeInfo FRigUnit_VisualDebugQuat::GetUpgradeInfo() const
+{
+	FRigVMFunction_VisualDebugQuatNoSpace NewNode;
+	NewNode.Value = Value;
+	NewNode.Thickness = Thickness;
+	NewNode.Scale = Scale;
+	NewNode.bEnabled = bEnabled;
+	
+	FRigVMStructUpgradeInfo Info(*this, NewNode);
+	return Info;
+}
+
+FRigUnit_VisualDebugQuatItemSpace_Execute()
+{
+    DECLARE_SCOPE_HIERARCHICAL_COUNTER_RIGUNIT()
+
+    FTransform Transform = FTransform::Identity;
+    Transform.SetRotation(Value);
+
+	FRigUnit_VisualDebugTransformItemSpace::StaticExecute(ExecuteContext, Transform, bEnabled, Thickness, Scale, Space);
+}
+
+FRigUnit_VisualDebugTransform_Execute()
+{
+	FRigUnit_VisualDebugTransformItemSpace::StaticExecute(ExecuteContext, Value, bEnabled, Thickness, Scale, FRigElementKey(BoneSpace, ERigElementType::Bone));
+}
+
+FRigVMStructUpgradeInfo FRigUnit_VisualDebugQuatItemSpace::GetUpgradeInfo() const
+{
+	FRigVMFunction_VisualDebugQuatNoSpace NewNode;
+	NewNode.Value = Value;
+	NewNode.Thickness = Thickness;
+	NewNode.Scale = Scale;
+	NewNode.bEnabled = bEnabled;
+	
+	FRigVMStructUpgradeInfo Info(*this, NewNode);
+	return Info;
+}
+
+FRigVMStructUpgradeInfo FRigUnit_VisualDebugTransform::GetUpgradeInfo() const
+{
+	FRigVMFunction_VisualDebugTransformNoSpace NewNode;
+	NewNode.Value = Value;
+	NewNode.Thickness = Thickness;
+	NewNode.Scale = Scale;
+	NewNode.bEnabled = bEnabled;
+
+	FRigVMStructUpgradeInfo Info(*this, NewNode);
+	return Info;
+}
+
+FRigVMStructUpgradeInfo FRigUnit_VisualDebugTransformItemSpace::GetUpgradeInfo() const
+{
+	FRigVMFunction_VisualDebugTransformNoSpace NewNode;
+	NewNode.Value = Value;
+	NewNode.Thickness = Thickness;
+	NewNode.Scale = Scale;
+	NewNode.bEnabled = bEnabled;
+
+	FRigVMStructUpgradeInfo Info(*this, NewNode);
+	return Info;
+}
+
+FRigUnit_VisualDebugTransformItemSpace_Execute()
+{
+    DECLARE_SCOPE_HIERARCHICAL_COUNTER_RIGUNIT()
+
+	if (ExecuteContext.GetDrawInterface() == nullptr || !bEnabled)
+	{
+		return;
+	}
+
+	FTransform WorldOffset = FTransform::Identity;
+	if (Space.IsValid())
+	{
+		WorldOffset = ExecuteContext.Hierarchy->GetGlobalTransform(Space);
+	}
+
+	ExecuteContext.GetDrawInterface()->DrawAxes(WorldOffset, Value, Scale, Thickness);
+}
diff --git a/Engine/Plugins/Animation/ControlRig/Source/ControlRig/Private/Units/Debug/RigUnit_VisualDebug.h b/Engine/Plugins/Animation/ControlRig/Source/ControlRig/Private/Units/Debug/RigUnit_VisualDebug.h
new file mode 100644
--- /dev/null
+++ b/Engine/Plugins/Animation/ControlRig/Source/ControlRig/Private/Units/Debug/RigUnit_VisualDebug.h
@@ -0,0 +1,263 @@
+// Copyright Epic Games, Inc. All Rights Reserved.
+
+#pragma once
+
+#include "RigUnit_DebugBase.h"
+#include "RigVMFunctions/Debug/RigVMFunction_VisualDebug.h"
+#include "RigUnit_VisualDebug.generated.h"
+
+USTRUCT(meta=(DisplayName = "Visual Debug Vector", TemplateName="VisualDebug", Keywords = "Draw,Point", Deprecated = "4.25", Varying))
+struct CONTROLRIG_API FRigUnit_VisualDebugVector : public FRigUnit_DebugBase
+{
+	GENERATED_BODY()
+
+	FRigUnit_VisualDebugVector()
+	{
+		Value = FVector::ZeroVector;
+		bEnabled = true;
+		Mode = ERigUnitVisualDebugPointMode::Point;
+		Color = FLinearColor::Red;
+		Thickness = 10.f;
+		Scale = 1.f;
+		BoneSpace = NAME_None;
+	}
+
+	RIGVM_METHOD()
+	virtual void Execute() override;
+
+	UPROPERTY(meta = (Input, Output))
+	FVector Value;
+
+	UPROPERTY(meta = (Input))
+	bool bEnabled;
+
+	UPROPERTY(meta = (Input, EditCondition = "bEnabled"))
+	ERigUnitVisualDebugPointMode Mode;
+
+	UPROPERTY(meta = (Input, EditCondition = "bEnabled"))
+	FLinearColor Color;
+
+	UPROPERTY(meta = (Input, EditCondition = "bEnabled"))
+	float Thickness;
+
+	UPROPERTY(meta = (Input, EditCondition = "bEnabled"))
+	float Scale;
+
+	UPROPERTY(meta = (Input, EditCondition = "bEnabled"))
+	FName BoneSpace;
+
+	RIGVM_METHOD()
+	virtual FRigVMStructUpgradeInfo GetUpgradeInfo() const override;
+};
+
+/**
+ * Debug draw parameters for a Point or Vector given a vector
+ */
+USTRUCT(meta=(DisplayName = "Visual Debug Vector", TemplateName="VisualDebug", Keywords = "Draw,Point", Deprecated = "5.2", Varying))
+struct CONTROLRIG_API FRigUnit_VisualDebugVectorItemSpace : public FRigUnit_DebugBase
+{
+	GENERATED_BODY()
+
+	FRigUnit_VisualDebugVectorItemSpace()
+	{
+		Value = FVector::ZeroVector;
+		bEnabled = true;
+		Mode = ERigUnitVisualDebugPointMode::Point;
+		Color = FLinearColor::Red;
+		Thickness = 10.f;
+		Scale = 1.f;
+		Space = FRigElementKey(NAME_None, ERigElementType::Bone);
+	}
+
+	virtual FRigElementKey DetermineSpaceForPin(const FString& InPinPath, void* InUserContext) const override
+	{
+		return Space;
+	}
+
+	RIGVM_METHOD()
+	virtual void Execute() override;
+
+	UPROPERTY(meta = (Input, Output))
+	FVector Value;
+
+	UPROPERTY(meta = (Input))
+	bool bEnabled;
+
+	UPROPERTY(meta = (Input, EditCondition = "bEnabled"))
+	ERigUnitVisualDebugPointMode Mode;
+
+	UPROPERTY(meta = (Input, EditCondition = "bEnabled"))
+	FLinearColor Color;
+
+	UPROPERTY(meta = (Input, EditCondition = "bEnabled"))
+	float Thickness;
+
+	UPROPERTY(meta = (Input, EditCondition = "bEnabled"))
+	float Scale;
+
+	UPROPERTY(meta = (Input, EditCondition = "bEnabled"))
+	FRigElementKey Space;
+
+	RIGVM_METHOD()
+	virtual FRigVMStructUpgradeInfo GetUpgradeInfo() const override;
+};
+
+USTRUCT(meta = (DisplayName = "Visual Debug Quat", TemplateName="VisualDebug", Keywords = "Draw,Rotation", Deprecated = "4.25", Varying))
+struct CONTROLRIG_API FRigUnit_VisualDebugQuat : public FRigUnit_DebugBase
+{
+	GENERATED_BODY()
+
+	FRigUnit_VisualDebugQuat()
+	{
+		Value = FQuat::Identity;
+		bEnabled = true;
+		Thickness = 0.f;
+		Scale = 10.f;
+		BoneSpace = NAME_None;
+	}
+
+	RIGVM_METHOD()
+	virtual void Execute() override;
+
+	UPROPERTY(meta = (Input, Output))
+	FQuat Value;
+
+	UPROPERTY(meta = (Input))
+	bool bEnabled;
+
+	UPROPERTY(meta = (Input, EditCondition = "bEnabled"))
+	float Thickness;
+
+	UPROPERTY(meta = (Input, EditCondition = "bEnabled"))
+	float Scale;
+
+	UPROPERTY(meta = (Input, EditCondition = "bEnabled"))
+	FName BoneSpace;
+
+	RIGVM_METHOD()
+	virtual FRigVMStructUpgradeInfo GetUpgradeInfo() const override;
+};
+
+/**
+ * Debug draw parameters for an Axis given a quaternion
+ */
+USTRUCT(meta = (DisplayName = "Visual Debug Quat", TemplateName="VisualDebug", Keywords = "Draw,Rotation", Deprecated = "5.2", Varying))
+struct CONTROLRIG_API FRigUnit_VisualDebugQuatItemSpace : public FRigUnit_DebugBase
+{
+	GENERATED_BODY()
+
+	FRigUnit_VisualDebugQuatItemSpace()
+	{
+		Value = FQuat::Identity;
+		bEnabled = true;
+		Thickness = 0.f;
+		Scale = 10.f;
+		Space = FRigElementKey(NAME_None, ERigElementType::Bone);
+	}
+
+	virtual FRigElementKey DetermineSpaceForPin(const FString& InPinPath, void* InUserContext) const override
+	{
+		return Space;
+	}
+
+	RIGVM_METHOD()
+	virtual void Execute() override;
+
+	UPROPERTY(meta = (Input, Output))
+	FQuat Value;
+
+	UPROPERTY(meta = (Input))
+	bool bEnabled;
+
+	UPROPERTY(meta = (Input, EditCondition = "bEnabled"))
+	float Thickness;
+
+	UPROPERTY(meta = (Input, EditCondition = "bEnabled"))
+	float Scale;
+
+	UPROPERTY(meta = (Input, EditCondition = "bEnabled"))
+	FRigElementKey Space;
+
+	RIGVM_METHOD()
+	virtual FRigVMStructUpgradeInfo GetUpgradeInfo() const override;
+};
+
+USTRUCT(meta=(DisplayName="Visual Debug Transform", TemplateName="VisualDebug", Keywords = "Draw,Axes", Deprecated = "4.25", Varying))
+struct CONTROLRIG_API FRigUnit_VisualDebugTransform : public FRigUnit_DebugBase
+{
+	GENERATED_BODY()
+
+	FRigUnit_VisualDebugTransform()
+	{
+		Value = FTransform::Identity;
+		bEnabled = true;
+		Thickness = 0.f;
+		Scale = 10.f;
+		BoneSpace = NAME_None;
+	}
+
+	RIGVM_METHOD()
+	virtual void Execute() override;
+
+	UPROPERTY(meta = (Input, Output))
+	FTransform Value;
+
+	UPROPERTY(meta = (Input))
+	bool bEnabled;
+
+	UPROPERTY(meta = (Input, EditCondition = "bEnabled"))
+	float Thickness;
+
+	UPROPERTY(meta = (Input, EditCondition = "bEnabled"))
+	float Scale;
+
+	UPROPERTY(meta = (Input, EditCondition = "bEnabled"))
+	FName BoneSpace;
+
+	RIGVM_METHOD()
+	virtual FRigVMStructUpgradeInfo GetUpgradeInfo() const override;
+};
+
+/**
+ * Debug draw parameters for an Axis given a transform
+ */
+USTRUCT(meta=(DisplayName="Visual Debug Transform", TemplateName="VisualDebug", Keywords = "Draw,Axes", Deprecated = "5.2", Varying))
+struct CONTROLRIG_API FRigUnit_VisualDebugTransformItemSpace : public FRigUnit_DebugBase
+{
+	GENERATED_BODY()
+
+	FRigUnit_VisualDebugTransformItemSpace()
+	{
+		Value = FTransform::Identity;
+		bEnabled = true;
+		Thickness = 0.f;
+		Scale = 10.f;
+		Space = FRigElementKey(NAME_None, ERigElementType::Bone);
+	}
+
+	virtual FRigElementKey DetermineSpaceForPin(const FString& InPinPath, void* InUserContext) const override
+	{
+		return Space;
+	}
+
+	RIGVM_METHOD()
+	virtual void Execute() override;
+
+	UPROPERTY(meta = (Input, Output))
+	FTransform Value;
+
+	UPROPERTY(meta = (Input))
+	bool bEnabled;
+
+	UPROPERTY(meta = (Input, EditCondition = "bEnabled"))
+	float Thickness;
+
+	UPROPERTY(meta = (Input, EditCondition = "bEnabled"))
+	float Scale;
+
+	UPROPERTY(meta = (Input, EditCondition = "bEnabled"))
+	FRigElementKey Space;
+
+	RIGVM_METHOD()
+	virtual FRigVMStructUpgradeInfo GetUpgradeInfo() const override;
+};
\ No newline at end of file
diff --git a/Engine/Source/Runtime/RigVM/Private/RigVMFunctions/Debug/RigVMFunction_DebugLine.cpp b/Engine/Source/Runtime/RigVM/Private/RigVMFunctions/Debug/RigVMFunction_DebugLine.cpp
--- a/Engine/Source/Runtime/RigVM/Private/RigVMFunctions/Debug/RigVMFunction_DebugLine.cpp
+++ b/Engine/Source/Runtime/RigVM/Private/RigVMFunctions/Debug/RigVMFunction_DebugLine.cpp
@@ -4,32 +4,6 @@
 
 #include UE_INLINE_GENERATED_CPP_BY_NAME(RigVMFunction_DebugLine)
 
-FRigVMFunction_DebugLine_Execute()
-{
-	FRigVMFunction_DebugLineNoSpace::StaticExecute(
-		ExecuteContext, 
-		A,
-		B,
-		Color,
-		Thickness,
-		WorldOffset, 
-		bEnabled);
-}
-
-FRigVMStructUpgradeInfo FRigVMFunction_DebugLine::GetUpgradeInfo() const
-{
-	FRigVMFunction_DebugLineNoSpace NewNode;
-	NewNode.A = A;
-	NewNode.B = B;
-	NewNode.Color = Color;
-	NewNode.Thickness = Thickness;
-	NewNode.WorldOffset = WorldOffset;
-	NewNode.bEnabled = bEnabled;
-
-	FRigVMStructUpgradeInfo Info(*this, NewNode);
-	return Info;
-}
-
 FRigVMFunction_DebugLineNoSpace_Execute()
 {
 	if (ExecuteContext.GetDrawInterface() == nullptr || !bEnabled)
diff --git a/Engine/Source/Runtime/RigVM/Private/RigVMFunctions/Debug/RigVMFunction_DebugLineStrip.cpp b/Engine/Source/Runtime/RigVM/Private/RigVMFunctions/Debug/RigVMFunction_DebugLineStrip.cpp
--- a/Engine/Source/Runtime/RigVM/Private/RigVMFunctions/Debug/RigVMFunction_DebugLineStrip.cpp
+++ b/Engine/Source/Runtime/RigVM/Private/RigVMFunctions/Debug/RigVMFunction_DebugLineStrip.cpp
@@ -4,30 +4,6 @@
 
 #include UE_INLINE_GENERATED_CPP_BY_NAME(RigVMFunction_DebugLineStrip)
 
-FRigVMFunction_DebugLineStrip_Execute()
-{
-	FRigVMFunction_DebugLineStripNoSpace::StaticExecute(
-		ExecuteContext, 
-		Points,
-		Color,
-		Thickness,
-		WorldOffset, 
-		bEnabled);
-}
-
-FRigVMStructUpgradeInfo FRigVMFunction_DebugLineStrip::GetUpgradeInfo() const
-{
-	FRigVMFunction_DebugLineStripNoSpace NewNode;
-	NewNode.Points = Points;
-	NewNode.Color = Color;
-	NewNode.Thickness = Thickness;
-	NewNode.WorldOffset = WorldOffset;
-	NewNode.bEnabled = bEnabled;
-
-	FRigVMStructUpgradeInfo Info(*this, NewNode);
-	return Info;
-}
-
 FRigVMFunction_DebugLineStripNoSpace_Execute()
 {
 	if (ExecuteContext.GetDrawInterface() == nullptr || !bEnabled)
diff --git a/Engine/Source/Runtime/RigVM/Private/RigVMFunctions/Debug/RigVMFunction_DebugTransform.cpp b/Engine/Source/Runtime/RigVM/Private/RigVMFunctions/Debug/RigVMFunction_DebugTransform.cpp
--- a/Engine/Source/Runtime/RigVM/Private/RigVMFunctions/Debug/RigVMFunction_DebugTransform.cpp
+++ b/Engine/Source/Runtime/RigVM/Private/RigVMFunctions/Debug/RigVMFunction_DebugTransform.cpp
@@ -5,77 +5,6 @@
 
 #include UE_INLINE_GENERATED_CPP_BY_NAME(RigVMFunction_DebugTransform)
 
-FRigVMFunction_DebugTransform_Execute()
-{
-	if (ExecuteContext.GetDrawInterface() == nullptr || !bEnabled)
-	{
-		return;
-	}
-
-	FTransform DrawTransform = Transform;
-
-	switch (Mode)
-	{
-		case ERigUnitDebugTransformMode::Axes:
-		{
-			ExecuteContext.GetDrawInterface()->DrawAxes(WorldOffset, DrawTransform, Scale, Thickness);
-			break;
-		}
-		case ERigUnitDebugTransformMode::Point:
-		{
-			ExecuteContext.GetDrawInterface()->DrawPoint(WorldOffset, DrawTransform.GetLocation(), Scale, Color);
-			break;
-		}
-		case ERigUnitDebugTransformMode::Box:
-		{
-			DrawTransform.SetScale3D(DrawTransform.GetScale3D() * Scale);
-			ExecuteContext.GetDrawInterface()->DrawBox(WorldOffset, DrawTransform, Color, Thickness);
-			break;
-		}
-	}
-}
-
-FRigVMStructUpgradeInfo FRigVMFunction_DebugTransform::GetUpgradeInfo() const
-{
-	FRigVMFunction_VisualDebugTransformNoSpace NewNode;
-	NewNode.Value = Transform;
-	NewNode.Scale = Scale;
-	NewNode.Thickness = Thickness;
-	NewNode.bEnabled = bEnabled;
-
-	FRigVMStructUpgradeInfo Info(*this, NewNode);
-	Info.AddRemappedPin(TEXT("Transform"), TEXT("Value"));
-	return Info;
-}
-
-FRigVMFunction_DebugTransformMutable_Execute()
-{
-	FRigVMFunction_DebugTransformMutableNoSpace::StaticExecute(
-		ExecuteContext, 
-		Transform,
-		Mode,
-		Color,
-		Thickness,
-		Scale,
-		WorldOffset, 
-		bEnabled);
-}
-
-FRigVMStructUpgradeInfo FRigVMFunction_DebugTransformMutable::GetUpgradeInfo() const
-{
-	FRigVMFunction_DebugTransformMutableNoSpace NewNode;
-	NewNode.Transform = Transform;
-	NewNode.Mode = Mode;
-	NewNode.Color = Color;
-	NewNode.Thickness = Thickness;
-	NewNode.Scale = Scale;
-	NewNode.WorldOffset = WorldOffset;
-	NewNode.bEnabled = bEnabled;
-
-	FRigVMStructUpgradeInfo Info(*this, NewNode);
-	return Info;
-}
-
 FRigVMFunction_DebugTransformMutableNoSpace_Execute()
 {
 	if (ExecuteContext.GetDrawInterface() == nullptr || !bEnabled)
@@ -106,58 +35,6 @@ FRigVMFunction_DebugTransformMutableNoSpace_Execute()
 	}
 }
 
-FRigVMFunction_DebugTransformArrayMutable_Execute()
-{
-	if (ExecuteContext.GetDrawInterface() == nullptr || !bEnabled)
-	{
-		return;
-	}
-
-	WorkData.DrawTransforms.SetNumUninitialized(Transforms.Num());
-	for(int32 Index=0;Index<Transforms.Num();Index++)
-	{
-		WorkData.DrawTransforms[Index] = Transforms[Index];
-	}
-
-	for(FTransform& DrawTransform : WorkData.DrawTransforms)
-	{
-		switch (Mode)
-		{
-			case ERigUnitDebugTransformMode::Axes:
-			{
-				ExecuteContext.GetDrawInterface()->DrawAxes(WorldOffset, DrawTransform, Scale, Thickness);
-				break;
-			}
-			case ERigUnitDebugTransformMode::Point:
-			{
-				ExecuteContext.GetDrawInterface()->DrawPoint(WorldOffset, DrawTransform.GetLocation(), Scale, Color);
-				break;
-			}
-			case ERigUnitDebugTransformMode::Box:
-			{
-				DrawTransform.SetScale3D(DrawTransform.GetScale3D() * Scale);
-				ExecuteContext.GetDrawInterface()->DrawBox(WorldOffset, DrawTransform, Color, Thickness);
-				break;
-			}
-		}
-	}
-}
-
-FRigVMStructUpgradeInfo FRigVMFunction_DebugTransformArrayMutable::GetUpgradeInfo() const
-{
-	FRigVMFunction_DebugTransformArrayMutableNoSpace NewNode;
-	NewNode.Transforms = Transforms;
-	NewNode.Mode = Mode;
-	NewNode.Color = Color;
-	NewNode.Thickness = Thickness;
-	NewNode.Scale = Scale;
-	NewNode.WorldOffset = WorldOffset;
-	NewNode.bEnabled = bEnabled;
-
-	FRigVMStructUpgradeInfo Info(*this, NewNode);
-	return Info;
-}
-
 FRigVMFunction_DebugTransformArrayMutableNoSpace_Execute()
 {
 	if (ExecuteContext.GetDrawInterface() == nullptr || !bEnabled || Transforms.Num() == 0)
diff --git a/Engine/Source/Runtime/RigVM/Public/RigVMFunctions/Debug/RigVMFunction_DebugLine.h b/Engine/Source/Runtime/RigVM/Public/RigVMFunctions/Debug/RigVMFunction_DebugLine.h
--- a/Engine/Source/Runtime/RigVM/Public/RigVMFunctions/Debug/RigVMFunction_DebugLine.h
+++ b/Engine/Source/Runtime/RigVM/Public/RigVMFunctions/Debug/RigVMFunction_DebugLine.h
@@ -5,48 +5,6 @@
 #include "RigVMFunction_DebugBase.h"
 #include "RigVMFunction_DebugLine.generated.h"
 
-USTRUCT(meta=(DisplayName="Draw Line", Deprecated = "4.25"))
-struct RIGVM_API FRigVMFunction_DebugLine : public FRigVMFunction_DebugBaseMutable
-{
-	GENERATED_BODY()
-
-	FRigVMFunction_DebugLine()
-	{
-		A = B = FVector::ZeroVector;
-		Color = FLinearColor::Red;
-		Thickness = 0.f;
-		WorldOffset = FTransform::Identity;
-		bEnabled = true;
-	}
-
-	RIGVM_METHOD()
-	virtual void Execute() override;
-
-	UPROPERTY(meta = (Input))
-	FVector A;
-
-	UPROPERTY(meta = (Input))
-	FVector B;
-
-	UPROPERTY(meta = (Input))
-	FLinearColor Color;
-
-	UPROPERTY(meta = (Input))
-	float Thickness;
-
-	UPROPERTY(meta = (Input))
-	FName Space;
-
-	UPROPERTY(meta = (Input))
-	FTransform WorldOffset;
-
-	UPROPERTY(meta = (Input, Constant))
-	bool bEnabled;
-
-	RIGVM_METHOD()
-	virtual FRigVMStructUpgradeInfo GetUpgradeInfo() const override;
-};
-
 /**
  * Draws a line in the viewport given a start and end vector
  */
diff --git a/Engine/Source/Runtime/RigVM/Public/RigVMFunctions/Debug/RigVMFunction_DebugLineStrip.h b/Engine/Source/Runtime/RigVM/Public/RigVMFunctions/Debug/RigVMFunction_DebugLineStrip.h
--- a/Engine/Source/Runtime/RigVM/Public/RigVMFunctions/Debug/RigVMFunction_DebugLineStrip.h
+++ b/Engine/Source/Runtime/RigVM/Public/RigVMFunctions/Debug/RigVMFunction_DebugLineStrip.h
@@ -5,44 +5,6 @@
 #include "RigVMFunction_DebugBase.h"
 #include "RigVMFunction_DebugLineStrip.generated.h"
 
-USTRUCT(meta=(DisplayName="Draw Line Strip", Deprecated = "4.25"))
-struct RIGVM_API FRigVMFunction_DebugLineStrip : public FRigVMFunction_DebugBaseMutable
-{
-	GENERATED_BODY()
-
-	FRigVMFunction_DebugLineStrip()
-	{
-		Color = FLinearColor::Red;
-		Thickness = 0.f;
-		WorldOffset = FTransform::Identity;
-		bEnabled = true;
-	}
-
-	RIGVM_METHOD()
-	virtual void Execute() override;
-
-	UPROPERTY(meta = (Input))
-	TArray<FVector> Points;
-
-	UPROPERTY(meta = (Input))
-	FLinearColor Color;
-
-	UPROPERTY(meta = (Input))
-	float Thickness;
-
-	UPROPERTY(meta = (Input))
-	FName Space;
-
-	UPROPERTY(meta = (Input))
-	FTransform WorldOffset;
-	
-	UPROPERTY(meta = (Input, Constant))
-	bool bEnabled;
-
-	RIGVM_METHOD()
-	virtual FRigVMStructUpgradeInfo GetUpgradeInfo() const override;
-};
-
 /**
  * Draws a line strip in the viewport given any number of points
  */
diff --git a/Engine/Source/Runtime/RigVM/Public/RigVMFunctions/Debug/RigVMFunction_DebugTransform.h b/Engine/Source/Runtime/RigVM/Public/RigVMFunctions/Debug/RigVMFunction_DebugTransform.h
--- a/Engine/Source/Runtime/RigVM/Public/RigVMFunctions/Debug/RigVMFunction_DebugTransform.h
+++ b/Engine/Source/Runtime/RigVM/Public/RigVMFunctions/Debug/RigVMFunction_DebugTransform.h
@@ -21,98 +21,6 @@ enum class ERigUnitDebugTransformMode : uint8
 	Max UMETA(Hidden),
 };
 
-USTRUCT(meta=(DisplayName="Draw Transform In Place", Deprecated = "4.25.0"))
-struct RIGVM_API FRigVMFunction_DebugTransform : public FRigVMFunction_DebugBase
-{
-	GENERATED_BODY()
-
-	FRigVMFunction_DebugTransform()
-	{
-		Transform = WorldOffset = FTransform::Identity;
-		Mode = ERigUnitDebugTransformMode::Axes;
-		Color = FLinearColor::Red;
-		Thickness = 0.f;
-		Scale = 10.f;
-		bEnabled = true;
-	}
-
-	RIGVM_METHOD()
-	virtual void Execute() override;
-
-	UPROPERTY(meta = (Input, Output))
-	FTransform Transform;
-
-	UPROPERTY(meta = (Input))
-	ERigUnitDebugTransformMode Mode;
-
-	UPROPERTY(meta = (Input))
-	FLinearColor Color;
-
-	UPROPERTY(meta = (Input))
-	float Thickness;
-
-	UPROPERTY(meta = (Input))
-	float Scale;
-
-	UPROPERTY(meta = (Input))
-	FName Space;
-
-	UPROPERTY(meta = (Input))
-	FTransform WorldOffset;
-
-	UPROPERTY(meta = (Input, Constant))
-	bool bEnabled;
-
-	RIGVM_METHOD()
-	virtual FRigVMStructUpgradeInfo GetUpgradeInfo() const override;
-};
-
-USTRUCT(meta=(DisplayName="Draw Transform", Deprecated = "4.25"))
-struct RIGVM_API FRigVMFunction_DebugTransformMutable : public FRigVMFunction_DebugBaseMutable
-{
-	GENERATED_BODY()
-
-	FRigVMFunction_DebugTransformMutable()
-	{
-		Transform = WorldOffset = FTransform::Identity;
-		Mode = ERigUnitDebugTransformMode::Axes;
-		Color = FLinearColor::Red;
-		Thickness = 0.f;
-		Scale = 10.f;
-		bEnabled = true;
-	}
-
-	RIGVM_METHOD()
-	virtual void Execute() override;
-
-	UPROPERTY(meta = (Input))
-	FTransform Transform;
-
-	UPROPERTY(meta = (Input))
-	ERigUnitDebugTransformMode Mode;
-
-	UPROPERTY(meta = (Input))
-	FLinearColor Color;
-
-	UPROPERTY(meta = (Input))
-	float Thickness;
-
-	UPROPERTY(meta = (Input))
-	float Scale;
-
-	UPROPERTY(meta = (Input))
-	FName Space;
-
-	UPROPERTY(meta = (Input))
-	FTransform WorldOffset;
-
-	UPROPERTY(meta = (Input, Constant))
-	bool bEnabled;
-
-	RIGVM_METHOD()
-	virtual FRigVMStructUpgradeInfo GetUpgradeInfo() const override;
-};
-
 /**
  * Given a transform, will draw a point, axis, or a box in the viewport
  */
@@ -165,55 +73,6 @@ struct RIGVM_API FRigVMFunction_DebugTransformArrayMutable_WorkData
 	TArray<FTransform> DrawTransforms;
 };
 
-USTRUCT(meta=(DisplayName="Draw Transform Array", Deprecated = "4.25.0"))
-struct RIGVM_API FRigVMFunction_DebugTransformArrayMutable : public FRigVMFunction_DebugBaseMutable
-{
- 	GENERATED_BODY()
-
-	FRigVMFunction_DebugTransformArrayMutable()
-	{
-		Mode = ERigUnitDebugTransformMode::Axes;
-		Color = FLinearColor::Red;
-		Thickness = 0.f;
-		Scale = 10.f;
-		WorldOffset = FTransform::Identity;
-		bEnabled = true;
-	}
-
-	RIGVM_METHOD()
-	virtual void Execute() override;
-
-	UPROPERTY(meta = (Input))
-	TArray<FTransform> Transforms;
-
-	UPROPERTY(meta = (Input))
-	ERigUnitDebugTransformMode Mode;
-
-	UPROPERTY(meta = (Input))
-	FLinearColor Color;
-
-	UPROPERTY(meta = (Input))
-	float Thickness;
-
-	UPROPERTY(meta = (Input))
-	float Scale;
-
-	UPROPERTY(meta = (Input))
-	FName Space;
-
-	UPROPERTY(meta = (Input))
-	FTransform WorldOffset;
-
-	UPROPERTY(meta = (Input, Constant))
-	bool bEnabled;
-
-	UPROPERTY(transient)
-	FRigVMFunction_DebugTransformArrayMutable_WorkData WorkData;
-
-	RIGVM_METHOD()
-	virtual FRigVMStructUpgradeInfo GetUpgradeInfo() const override;
-};
-
 /**
 * Given a transform array, will draw a point, axis, or a box in the viewport
 */
