Sequencer: Fixed control rig tracks not being rebound when the object they animate changes

diff --git a/Engine/Plugins/Animation/ControlRig/Source/ControlRig/Private/Sequencer/MovieSceneControlRigParameterTemplate.cpp b/Engine/Plugins/Animation/ControlRig/Source/ControlRig/Private/Sequencer/MovieSceneControlRigParameterTemplate.cpp
--- a/Engine/Plugins/Animation/ControlRig/Source/ControlRig/Private/Sequencer/MovieSceneControlRigParameterTemplate.cpp
+++ b/Engine/Plugins/Animation/ControlRig/Source/ControlRig/Private/Sequencer/MovieSceneControlRigParameterTemplate.cpp
@@ -1129,16 +1129,17 @@ struct FControlRigParameterExecutionToken : IMovieSceneExecutionToken
 			const UMovieSceneSequence* Sequence = Player.State.FindSequence(Operand.SequenceID);
 			TArrayView<TWeakObjectPtr<>> BoundObjects = Player.FindBoundObjects(Operand);
 
-			if (Sequence && BoundObjects.Num() > 0 && BoundObjects[0].Get())
+			UObject* BoundObject = BoundObjects.Num() > 0 ? BoundObjects[0].Get() : nullptr;
+			if (Sequence && BoundObject)
 			{
 				if (!ControlRig->GetObjectBinding())
 				{
 					ControlRig->SetObjectBinding(MakeShared<FControlRigObjectBinding>());
 				}
 
-				if (!ControlRig->GetObjectBinding()->GetBoundObject())
+				if (ControlRig->GetObjectBinding()->GetBoundObject() != BoundObject)
 				{
-					ControlRig->GetObjectBinding()->BindToObject(BoundObjects[0].Get());
+					ControlRig->GetObjectBinding()->BindToObject(BoundObject);
 					TArray<FName> SelectedControls = ControlRig->CurrentControlSelection();
 					ControlRig->Initialize();
 					if (ControlRig->IsA<UFKControlRig>())
@@ -1182,7 +1183,7 @@ struct FControlRigParameterExecutionToken : IMovieSceneExecutionToken
 					}
 					else
 					{
-						ControlRig = GetControlRig(Section, BoundObjects[0].Get());
+						ControlRig = GetControlRig(Section, BoundObject);
 					}
 				}
 			}		
