Control Rig: Fix build break in ControlRigBlueprint

diff --git a/Engine/Plugins/Animation/ControlRig/Source/ControlRigDeveloper/Private/ControlRigBlueprint.cpp b/Engine/Plugins/Animation/ControlRig/Source/ControlRigDeveloper/Private/ControlRigBlueprint.cpp
--- a/Engine/Plugins/Animation/ControlRig/Source/ControlRigDeveloper/Private/ControlRigBlueprint.cpp
+++ b/Engine/Plugins/Animation/ControlRig/Source/ControlRigDeveloper/Private/ControlRigBlueprint.cpp
@@ -1328,20 +1328,7 @@ void UControlRigBlueprint::RefreshAllModels(EControlRigBlueprintLoadType InLoadT
 		for (URigVMGraph* ModelGraph : ModelGraphs)
 		{
 			URigVMController* Controller = GetOrCreateController(ModelGraph);
-			// at this stage, allow all links to be reattached,
-			// RecomputeAllTemplateFilteredPermutations() later should break any invalid links
 			Controller->ReattachLinksToPinObjects(true /* follow redirectors */, nullptr, true, true);
-
-			if(bIsPostLoad)
-			{
-				for(URigVMNode* ModelNode : ModelGraph->GetNodes())
-				{
-					if (URigVMLibraryNode* LibraryNode = ModelNode->GetTypedOuter<URigVMLibraryNode>())
-					{
-						Controller->UpdateLibraryTemplate(LibraryNode, false);
-					}
-				}
-			}
 		}
 		return;
 	}
