Control Rig: Fix warning with RigUnit_DebugBase conflicting with existing type

diff --git a/Engine/Plugins/Animation/ControlRig/Config/BaseControlRig.ini b/Engine/Plugins/Animation/ControlRig/Config/BaseControlRig.ini
--- a/Engine/Plugins/Animation/ControlRig/Config/BaseControlRig.ini
+++ b/Engine/Plugins/Animation/ControlRig/Config/BaseControlRig.ini
@@ -418,8 +418,6 @@
 +StructRedirects=(OldName="/Script/ControlRig.RigUnit_TimeOffsetVector",NewName="/Script/RigVM.RigVMFunction_TimeOffsetVector")
 +StructRedirects=(OldName="/Script/ControlRig.RigUnit_TimeOffsetTransform",NewName="/Script/RigVM.RigVMFunction_TimeOffsetTransform")
 +StructRedirects=(OldName="/Script/ControlRig.RigUnit_VerletIntegrateVector",NewName="/Script/RigVM.RigVMFunction_VerletIntegrateVector")
-+StructRedirects=(OldName="/Script/ControlRig.RigUnit_DebugBase",NewName="/Script/RigVM.RigVMFunction_DebugBase")
-+StructRedirects=(OldName="/Script/ControlRig.RigUnit_DebugBaseMutable",NewName="/Script/RigVM.RigVMFunction_DebugBaseMutable")
 +StructRedirects=(OldName="/Script/ControlRig.RigUnit_DebugPoint",NewName="/Script/RigVM.RigVMFunction_DebugPoint")
 +StructRedirects=(OldName="/Script/ControlRig.RigUnit_DebugPointMutable",NewName="/Script/RigVM.RigVMFunction_DebugPointMutable")
 +StructRedirects=(OldName="/Script/ControlRig.RigUnit_AnimBase",NewName="/Script/RigVM.RigVMFunction_AnimBase")
