Control Rig: Fix cannot connect reroute output to another pin (source pin is not an output)

diff --git a/Engine/Plugins/Animation/ControlRig/Source/ControlRigDeveloper/Private/Graph/ControlRigGraphNode.cpp b/Engine/Plugins/Animation/ControlRig/Source/ControlRigDeveloper/Private/Graph/ControlRigGraphNode.cpp
--- a/Engine/Plugins/Animation/ControlRig/Source/ControlRigDeveloper/Private/Graph/ControlRigGraphNode.cpp
+++ b/Engine/Plugins/Animation/ControlRig/Source/ControlRigDeveloper/Private/Graph/ControlRigGraphNode.cpp
@@ -1662,14 +1662,26 @@ bool UControlRigGraphNode::ShouldDrawNodeAsControlPointOnly(int32& OutInputPinIn
 	if (URigVMRerouteNode* Reroute = Cast<URigVMRerouteNode>(GetModelNode()))
 	{
 		if (!Reroute->GetShowsAsFullNode())
-	{
+		{
 			if (Pins.Num() >= 2)
 			{
 				OutInputPinIndex = 0;
-				OutOutputPinIndex = 1;
+				OutOutputPinIndex = Pins.Num()/2;
+
+				if (Pins[OutOutputPinIndex]->Direction != EGPD_Output || Pins[OutOutputPinIndex]->ParentPin != nullptr)
+				{
+					for (int32 i=0; i<Pins.Num(); ++i)
+					{
+						if (Pins[i]->Direction == EGPD_Output && Pins[i]->ParentPin == nullptr)
+						{
+							OutOutputPinIndex = i;
+							break;
+						}
+					}
+				}
 				return true;
 			}
-	}
+		}
 	}
 	return false;
 }
