Sequencer: Space Switching, need to recalc the channel index when setting tangents in a space switch since it can change do the the number of controls changing during evaluation. Also added the delegate for control rig initialization, which is needed since control rigs can now dynamically add controls on the fly during evaluation, to the track so it will work in runtime and if there are nested sequences.

diff --git a/Engine/Plugins/Animation/ControlRig/Source/ControlRig/Private/Sequencer/MovieSceneControlRigParameterTrack.cpp b/Engine/Plugins/Animation/ControlRig/Source/ControlRig/Private/Sequencer/MovieSceneControlRigParameterTrack.cpp
--- a/Engine/Plugins/Animation/ControlRig/Source/ControlRig/Private/Sequencer/MovieSceneControlRigParameterTrack.cpp
+++ b/Engine/Plugins/Animation/ControlRig/Source/ControlRig/Private/Sequencer/MovieSceneControlRigParameterTrack.cpp
@@ -29,6 +29,15 @@ UMovieSceneControlRigParameterTrack::UMovieSceneControlRigParameterTrack(const F
 	SupportedBlendTypes.Add(EMovieSceneBlendType::Absolute);
 
 }
+
+UMovieSceneControlRigParameterTrack::~UMovieSceneControlRigParameterTrack()
+{
+	if (IsValid(ControlRig))
+	{
+		ControlRig->OnInitialized_AnyThread().RemoveAll(this);
+	}
+}
+
 FMovieSceneEvalTemplatePtr UMovieSceneControlRigParameterTrack::CreateTemplateForSection(const UMovieSceneSection& InSection) const
 {
 	return FMovieSceneControlRigParameterTemplate(*CastChecked<UMovieSceneControlRigParameterSection>(&InSection), *this);
@@ -425,6 +434,32 @@ void UMovieSceneControlRigParameterTrack::PostLoad()
 	{		
 		ReconstructControlRig();
 	}
+
+	if (IsValid(ControlRig))
+	{
+		ControlRig->OnInitialized_AnyThread().AddUObject(this, &UMovieSceneControlRigParameterTrack::HandleOnInitialized);
+	}
+}
+
+
+void UMovieSceneControlRigParameterTrack::HandleOnInitialized(URigVMHost* Subject, const FName& InEventName)
+{
+	if (IsValid(ControlRig))
+	{
+		TArray<FRigControlElement*> SortedControls;
+		ControlRig->GetControlsInOrder(SortedControls);
+		for (UMovieSceneSection* BaseSection : GetAllSections())
+		{
+			if (UMovieSceneControlRigParameterSection* Section = Cast<UMovieSceneControlRigParameterSection>(BaseSection))
+			{
+				if (Section->IsDifferentThanLastControlsUsedToReconstruct(SortedControls))
+				{
+					//Section->ReconstructChannelProxy();
+					Section->RecreateWithThisControlRig(ControlRig, Section->GetBlendType() == EMovieSceneBlendType::Absolute);
+				}
+			}
+		}
+	}
 }
 
 #if WITH_EDITORONLY_DATA
diff --git a/Engine/Plugins/Animation/ControlRig/Source/ControlRig/Public/Sequencer/MovieSceneControlRigParameterTrack.h b/Engine/Plugins/Animation/ControlRig/Source/ControlRig/Public/Sequencer/MovieSceneControlRigParameterTrack.h
--- a/Engine/Plugins/Animation/ControlRig/Source/ControlRig/Public/Sequencer/MovieSceneControlRigParameterTrack.h
+++ b/Engine/Plugins/Animation/ControlRig/Source/ControlRig/Public/Sequencer/MovieSceneControlRigParameterTrack.h
@@ -27,6 +27,7 @@ class UMovieSceneControlRigParameterTrack
 	GENERATED_UCLASS_BODY()
 
 public:
+	virtual ~UMovieSceneControlRigParameterTrack();
 
 	// UMovieSceneTrack interface
 	virtual FMovieSceneEvalTemplatePtr CreateTemplateForSection(const UMovieSceneSection& InSection) const override;
@@ -54,6 +55,9 @@ public:
 	virtual TArray<FFBXNodeAndChannels>* GetNodeAndChannelMappings(UMovieSceneSection* InSection)  override;
 	virtual void GetSelectedNodes(TArray<FName>& OutSelectedNodes) override;
 
+	//UControlRig Delegates
+	void HandleOnInitialized(URigVMHost* Subject, const FName& InEventName);
+
 #if WITH_EDITOR
 	void HandlePackageDone(const FEndLoadPackageContext& Context);
 	// control Rigs are ready only after its package is fully end-loaded
diff --git a/Engine/Plugins/Animation/ControlRig/Source/ControlRigEditor/Private/Sequencer/ControlRigParameterTrackEditor.cpp b/Engine/Plugins/Animation/ControlRig/Source/ControlRigEditor/Private/Sequencer/ControlRigParameterTrackEditor.cpp
--- a/Engine/Plugins/Animation/ControlRig/Source/ControlRigEditor/Private/Sequencer/ControlRigParameterTrackEditor.cpp
+++ b/Engine/Plugins/Animation/ControlRig/Source/ControlRigEditor/Private/Sequencer/ControlRigParameterTrackEditor.cpp
@@ -2998,29 +2998,6 @@ void FControlRigParameterTrackEditor::HandleOnInitialized(URigVMHost* Subject, c
 			GetSequencer()->NotifyMovieSceneDataChanged(EMovieSceneDataChangeType::RefreshTree);
 		}
 	}
-	//also do to new procedural rigs, we may be creating controls dynamically, so here we need to check to see if we need to reconstruct channels
-	UMovieScene* MovieScene = GetSequencer()->GetFocusedMovieSceneSequence()->GetMovieScene();
-	const TArray<FMovieSceneBinding>& Bindings = MovieScene->GetBindings();
-	for (const FMovieSceneBinding& Binding : Bindings)
-	{
-		UMovieSceneControlRigParameterTrack* Track = Cast<UMovieSceneControlRigParameterTrack>(MovieScene->FindTrack(UMovieSceneControlRigParameterTrack::StaticClass(), Binding.GetObjectGuid(), NAME_None));
-		if (Track && Track->GetControlRig() == ControlRig)
-		{
-			TArray<FRigControlElement*> SortedControls;
-			ControlRig->GetControlsInOrder(SortedControls);
-			for (UMovieSceneSection* BaseSection : Track->GetAllSections())
-			{
-				if (UMovieSceneControlRigParameterSection* Section = Cast<UMovieSceneControlRigParameterSection>(BaseSection))
-				{
-					if (Section->IsDifferentThanLastControlsUsedToReconstruct(SortedControls))
-					{
-						Section->ReconstructChannelProxy();
-						Section->MarkAsChanged();
-					}
-				}
-			}
-		}
-	}
 }
 
 void FControlRigParameterTrackEditor::HandleControlModified(UControlRig* ControlRig, FRigControlElement* ControlElement, const FRigControlModifiedContext& Context)
diff --git a/Engine/Plugins/Animation/ControlRig/Source/ControlRigEditor/Private/Sequencer/ControlRigSpaceChannelEditors.cpp b/Engine/Plugins/Animation/ControlRig/Source/ControlRigEditor/Private/Sequencer/ControlRigSpaceChannelEditors.cpp
--- a/Engine/Plugins/Animation/ControlRig/Source/ControlRigEditor/Private/Sequencer/ControlRigSpaceChannelEditors.cpp
+++ b/Engine/Plugins/Animation/ControlRig/Source/ControlRigEditor/Private/Sequencer/ControlRigSpaceChannelEditors.cpp
@@ -60,6 +60,7 @@
 #include "Fonts/FontMeasure.h"
 #include "CurveEditorSettings.h"
 #include "TimeToPixel.h"
+#include "Templates/Tuple.h"
 
 #define LOCTEXT_NAMESPACE "ControlRigEditMode"
 
@@ -254,7 +255,26 @@ FSpaceChannelAndSection FControlRigSpaceChannelHelpers::FindSpaceChannelAndSecti
 	return SpaceChannelAndSection;
 }
 
+static TTuple<FRigControlElement*, FChannelMapInfo*, int32, int32> GetControlAndChannelInfos(UControlRig* ControlRig, UMovieSceneControlRigParameterSection* ControlRigSection, FName ControlName)
+{
+	FRigControlElement* ControlElement = nullptr;
+	FChannelMapInfo* pChannelIndex = nullptr;
+	int32 NumChannels = 0;
+	int32 ChannelIndex = 0;
+	Tie(ControlElement, pChannelIndex) = FControlRigSpaceChannelHelpers::GetControlAndChannelInfo(ControlRig, ControlRigSection, ControlName);
 
+	if (pChannelIndex && ControlElement)
+	{
+		// get the number of float channels to treat
+		NumChannels = FControlRigSpaceChannelHelpers::GetNumFloatChannels(ControlElement->Settings.ControlType);
+		if (NumChannels > 0)
+		{
+			ChannelIndex = pChannelIndex->ChannelIndex;
+		}
+	}
+
+	return TTuple<FRigControlElement*, FChannelMapInfo*, int32, int32>(ControlElement, pChannelIndex, NumChannels,ChannelIndex);
+}
 
 /*
 *  Situations to handle\
@@ -358,25 +378,24 @@ FKeyHandle FControlRigSpaceChannelHelpers::SequencerKeyControlRigSpaceChannel(UC
 
 		// store tangents to keep the animation when switching 
 		TArray<FMovieSceneTangentData> Tangents;
-		int32 ChannelIndex = 0, NumChannels = 0;
 		UMovieSceneControlRigParameterSection* ControlRigSection = Cast<UMovieSceneControlRigParameterSection>(SectionToKey);
 		if (ControlRigSection)
 		{
 			FChannelMapInfo* pChannelIndex = nullptr;
 			FRigControlElement* ControlElement = nullptr;
-			Tie(ControlElement, pChannelIndex) = GetControlAndChannelInfo(ControlRig, ControlRigSection, ControlKey.Name);
+			int32 NumChannels = 0;
+			int32 ChannelIndex = 0;
+			Tie(ControlElement, pChannelIndex,NumChannels, ChannelIndex) = GetControlAndChannelInfos(ControlRig, ControlRigSection, ControlKey.Name);
 
-			if (pChannelIndex && ControlElement)
+			if (pChannelIndex && ControlElement && NumChannels > 0)
 			{
-				// get the number of float channels to treat
-				NumChannels = GetNumFloatChannels(ControlElement->Settings.ControlType);
-				if (NumChannels > 0)
-				{
-					ChannelIndex = pChannelIndex->ChannelIndex;
-					EvaluateTangentAtThisTime<FMovieSceneFloatChannel>(ChannelIndex, NumChannels,ControlRigSection, Time, Tangents);
-				}
+				EvaluateTangentAtThisTime<FMovieSceneFloatChannel>(ChannelIndex, NumChannels,ControlRigSection, Time, Tangents);
 			}
 		}
+		else
+		{
+			return Handle;
+		}
 		
 		// if current Value not the same as the previous add new space key
 		if (PreviousValue != Value)
@@ -435,17 +454,23 @@ FKeyHandle FControlRigSpaceChannelHelpers::SequencerKeyControlRigSpaceChannel(UC
 			Context.LocalTime = TickResolution.AsSeconds(FFrameTime(Time - 1));
 			ControlRig->SetControlGlobalTransform(ControlKey.Name, ControlWorldTransforms[0], true, Context);
 
-			// need to update the tangents here to keep the arriving animation
-			if (ControlRigSection  && NumChannels > 0)
+			//need to do this after eval
+			FChannelMapInfo* pChannelIndex = nullptr;
+			FRigControlElement* ControlElement = nullptr;
+			int32 NumChannels = 0;
+			int32 ChannelIndex = 0;
+			Tie(ControlElement, pChannelIndex, NumChannels, ChannelIndex) = GetControlAndChannelInfos(ControlRig, ControlRigSection, ControlKey.Name);
+
+			if (pChannelIndex && ControlElement && NumChannels > 0)
 			{
-				SetTangentsAtThisTime<FMovieSceneFloatChannel>(ChannelIndex, NumChannels, ControlRigSection,GlobalTime.GetFrame(), Tangents);
+				// need to update the tangents here to keep the arriving animation
+				SetTangentsAtThisTime<FMovieSceneFloatChannel>(ChannelIndex, NumChannels, ControlRigSection, GlobalTime.GetFrame(), Tangents);
 			}
 		}
 
 		// effectively switch to new space
 		URigHierarchy::TElementDependencyMap Dependencies = ControlRig->GetHierarchy()->GetDependenciesForVM(ControlRig->GetVM());
 		ControlRig->GetHierarchy()->SwitchToParent(ControlKey, SpaceKey, false, true, Dependencies, nullptr);
-
 		// add new keys in the new space context
 		int32 FramesIndex = 0;
 		for (const FFrameNumber& Frame : Frames)
@@ -460,8 +485,15 @@ FKeyHandle FControlRigSpaceChannelHelpers::SequencerKeyControlRigSpaceChannel(UC
 			Context.LocalTime = TickResolution.AsSeconds(FFrameTime(Frame));
 			ControlRig->SetControlGlobalTransform(ControlKey.Name, ControlWorldTransforms[FramesIndex], true, Context);
 
+			//need to do this after eval
+			FChannelMapInfo* pChannelIndex = nullptr;
+			FRigControlElement* ControlElement = nullptr;
+			int32 NumChannels = 0;
+			int32 ChannelIndex = 0;
+			Tie(ControlElement, pChannelIndex, NumChannels, ChannelIndex) = GetControlAndChannelInfos(ControlRig, ControlRigSection, ControlKey.Name);
+
 			// need to update the tangents here to keep the leaving animation
-			if (bSetPreviousKey && FramesIndex == 0)
+			if (NumChannels > 0 && bSetPreviousKey && FramesIndex == 0)
 			{
 				if (ControlRigSection && NumChannels > 0)
 				{
