Control Rig: Fix ensure when modifying public functions

diff --git a/Engine/Plugins/Animation/ControlRig/Source/ControlRigEditor/Private/Editor/SControlRigFunctionBulkEditWidget.cpp b/Engine/Plugins/Animation/ControlRig/Source/ControlRigEditor/Private/Editor/SControlRigFunctionBulkEditWidget.cpp
--- a/Engine/Plugins/Animation/ControlRig/Source/ControlRigEditor/Private/Editor/SControlRigFunctionBulkEditWidget.cpp
+++ b/Engine/Plugins/Animation/ControlRig/Source/ControlRigEditor/Private/Editor/SControlRigFunctionBulkEditWidget.cpp
@@ -81,7 +81,10 @@ void SControlRigFunctionBulkEditWidget::LoadAffectedAssets()
 			const FString PackageName = InReference.ToSoftObjectPath().GetLongPackageName();
 			const FText Message = FText::FromString(FString::Printf(TEXT("Loading '%s' ..."), *PackageName));
 
-			SlowTask.TotalAmountOfWork = (float)InCount;
+			if (InIndex == 0)
+			{
+				SlowTask.TotalAmountOfWork += (float)InCount;
+			}
 			SlowTask.EnterProgressFrame(1.f, Message);
         });
 		
