Sequencer: Control Rig: When checking to see if the binding has changed in the control rig eval, need to handle the fact that we bind to components but an actor may be passed in.

diff --git a/Engine/Plugins/Animation/ControlRig/Source/ControlRig/Private/Sequencer/MovieSceneControlRigParameterTemplate.cpp b/Engine/Plugins/Animation/ControlRig/Source/ControlRig/Private/Sequencer/MovieSceneControlRigParameterTemplate.cpp
--- a/Engine/Plugins/Animation/ControlRig/Source/ControlRig/Private/Sequencer/MovieSceneControlRigParameterTemplate.cpp
+++ b/Engine/Plugins/Animation/ControlRig/Source/ControlRig/Private/Sequencer/MovieSceneControlRigParameterTemplate.cpp
@@ -1115,6 +1115,35 @@ struct FControlRigParameterExecutionToken : IMovieSceneExecutionToken
 	FControlRigParameterExecutionToken(const FControlRigParameterExecutionToken&) = delete;
 	FControlRigParameterExecutionToken& operator=(const FControlRigParameterExecutionToken&) = delete;
 
+	UObject* GetBindableObject(UObject* InObject) const
+	{
+		// If we are binding to an actor, find the first skeletal mesh component
+		if (AActor* Actor = Cast<AActor>(InObject))
+		{
+			if (UControlRigComponent* ControlRigComponent = Actor->FindComponentByClass<UControlRigComponent>())
+			{
+				return ControlRigComponent;
+			}
+			else if (USkeletalMeshComponent* SkeletalMeshComponent = Actor->FindComponentByClass<USkeletalMeshComponent>())
+			{
+				return SkeletalMeshComponent;
+			}
+		}
+		else if (UControlRigComponent* ControlRigComponent = Cast<UControlRigComponent>(InObject))
+		{
+			return ControlRigComponent;
+		}
+		else if (USkeletalMeshComponent* SkeletalMeshComponent = Cast<USkeletalMeshComponent>(InObject))
+		{
+			return SkeletalMeshComponent;
+		}
+		else if (USkeleton* Skeleton = Cast<USkeleton>(InObject))
+		{
+			return Skeleton;
+		}
+		return nullptr;
+	}
+
 	virtual void Execute(const FMovieSceneContext& Context, const FMovieSceneEvaluationOperand& Operand, FPersistentEvaluationData& PersistentData, IMovieScenePlayer& Player)
 	{
 		MOVIESCENE_DETAILED_SCOPE_CYCLE_COUNTER(MovieSceneEval_ControlRigParameterTrack_TokenExecute)
@@ -1137,7 +1166,7 @@ struct FControlRigParameterExecutionToken : IMovieSceneExecutionToken
 					ControlRig->SetObjectBinding(MakeShared<FControlRigObjectBinding>());
 				}
 
-				if (ControlRig->GetObjectBinding()->GetBoundObject() != BoundObject)
+				if (ControlRig->GetObjectBinding()->GetBoundObject() != GetBindableObject(BoundObject))
 				{
 					ControlRig->GetObjectBinding()->BindToObject(BoundObject);
 					TArray<FName> SelectedControls = ControlRig->CurrentControlSelection();
