Sequencer: Space Switching, followup from last submit for this, need to move some code to BeginDestroy and make sure HandleOnIinitalized is set up.

diff --git a/Engine/Plugins/Animation/ControlRig/Source/ControlRig/Private/Sequencer/MovieSceneControlRigParameterTrack.cpp b/Engine/Plugins/Animation/ControlRig/Source/ControlRig/Private/Sequencer/MovieSceneControlRigParameterTrack.cpp
--- a/Engine/Plugins/Animation/ControlRig/Source/ControlRig/Private/Sequencer/MovieSceneControlRigParameterTrack.cpp
+++ b/Engine/Plugins/Animation/ControlRig/Source/ControlRig/Private/Sequencer/MovieSceneControlRigParameterTrack.cpp
@@ -30,8 +30,9 @@ UMovieSceneControlRigParameterTrack::UMovieSceneControlRigParameterTrack(const F
 
 }
 
-UMovieSceneControlRigParameterTrack::~UMovieSceneControlRigParameterTrack()
+void UMovieSceneControlRigParameterTrack::BeginDestroy()
 {
+	Super::BeginDestroy();
 	if (IsValid(ControlRig))
 	{
 		ControlRig->OnInitialized_AnyThread().RemoveAll(this);
@@ -191,13 +192,26 @@ FText UMovieSceneControlRigParameterTrack::GetDefaultDisplayName() const
 
 UMovieSceneSection* UMovieSceneControlRigParameterTrack::CreateControlRigSection(FFrameNumber StartTime, UControlRig* InControlRig, bool bInOwnsControlRig)
 {
+	if (InControlRig == nullptr)
+	{
+		return nullptr;
+	}
 	if (!bInOwnsControlRig)
 	{
 		InControlRig->Rename(nullptr, this);
 	}
+
+	if (IsValid(ControlRig))
+	{
+		ControlRig->OnInitialized_AnyThread().RemoveAll(this);
+	}
+
 	ControlRig = InControlRig;
-	UMovieSceneControlRigParameterSection*  NewSection = Cast<UMovieSceneControlRigParameterSection>(CreateNewSection());
-	
+
+	ControlRig->OnInitialized_AnyThread().AddUObject(this, &UMovieSceneControlRigParameterTrack::HandleOnInitialized);
+
+	UMovieSceneControlRigParameterSection* NewSection = Cast<UMovieSceneControlRigParameterSection>(CreateNewSection());
+
 	UMovieScene* OuterMovieScene = GetTypedOuter<UMovieScene>();
 	NewSection->SetRange(TRange<FFrameNumber>::All());
 
@@ -454,7 +468,6 @@ void UMovieSceneControlRigParameterTrack::HandleOnInitialized(URigVMHost* Subjec
 			{
 				if (Section->IsDifferentThanLastControlsUsedToReconstruct(SortedControls))
 				{
-					//Section->ReconstructChannelProxy();
 					Section->RecreateWithThisControlRig(ControlRig, Section->GetBlendType() == EMovieSceneBlendType::Absolute);
 				}
 			}
@@ -535,10 +548,21 @@ void UMovieSceneControlRigParameterTrack::RenameParameterName(const FName& OldPa
 
 void UMovieSceneControlRigParameterTrack::ReplaceControlRig(UControlRig* NewControlRig, bool RecreateChannels)
 {
+	if (IsValid(ControlRig))
+	{
+		ControlRig->OnInitialized_AnyThread().RemoveAll(this);
+	}
+
 	ControlRig = NewControlRig;
-	if (ControlRig->GetOuter() != this)
+
+	if (IsValid(ControlRig))
 	{
-		ControlRig->Rename(nullptr, this);
+		ControlRig->OnInitialized_AnyThread().AddUObject(this, &UMovieSceneControlRigParameterTrack::HandleOnInitialized);
+
+		if (ControlRig->GetOuter() != this)
+		{
+			ControlRig->Rename(nullptr, this);
+		}
 	}
 	for (UMovieSceneSection* Section : Sections)
 	{
@@ -552,7 +576,7 @@ void UMovieSceneControlRigParameterTrack::ReplaceControlRig(UControlRig* NewCont
 			{
 				CRSection->SetControlRig(NewControlRig);
 			}
-		}	
+		}
 	}
 }
 
diff --git a/Engine/Plugins/Animation/ControlRig/Source/ControlRig/Public/Sequencer/MovieSceneControlRigParameterTrack.h b/Engine/Plugins/Animation/ControlRig/Source/ControlRig/Public/Sequencer/MovieSceneControlRigParameterTrack.h
--- a/Engine/Plugins/Animation/ControlRig/Source/ControlRig/Public/Sequencer/MovieSceneControlRigParameterTrack.h
+++ b/Engine/Plugins/Animation/ControlRig/Source/ControlRig/Public/Sequencer/MovieSceneControlRigParameterTrack.h
@@ -27,8 +27,6 @@ class UMovieSceneControlRigParameterTrack
 	GENERATED_UCLASS_BODY()
 
 public:
-	virtual ~UMovieSceneControlRigParameterTrack();
-
 	// UMovieSceneTrack interface
 	virtual FMovieSceneEvalTemplatePtr CreateTemplateForSection(const UMovieSceneSection& InSection) const override;
 	virtual bool SupportsType(TSubclassOf<UMovieSceneSection> SectionClass) const override;
@@ -43,6 +41,7 @@ public:
 	virtual FName GetTrackName() const override { return TrackName; }
 	// UObject
 	virtual void PostLoad() override;
+	virtual void BeginDestroy() override;
 #if WITH_EDITORONLY_DATA
 	static void DeclareConstructClasses(TArray<FTopLevelAssetPath>& OutConstructClasses, const UClass* SpecificSubclass);
 #endif
