Control Rig: Fixed a IsRooted Assertion caused by detail wrapper objects not getting cleared after BP compile.

diff --git a/Engine/Plugins/Animation/ControlRig/Source/ControlRigEditor/Private/Editor/ControlRigEditor.cpp b/Engine/Plugins/Animation/ControlRig/Source/ControlRigEditor/Private/Editor/ControlRigEditor.cpp
--- a/Engine/Plugins/Animation/ControlRig/Source/ControlRigEditor/Private/Editor/ControlRigEditor.cpp
+++ b/Engine/Plugins/Animation/ControlRig/Source/ControlRigEditor/Private/Editor/ControlRigEditor.cpp
@@ -2001,6 +2001,14 @@ void FControlRigEditor::RefreshDetailView()
 	{
 		SetDetailViewForLocalVariable();	
 	}
+	else
+	{
+		// detail view is showing other stuff; could be a BP variable for example
+		// in this case wrapper objects are not in use, yet still rooted
+		// and preventing their outer objects from getting GCed after a Compile()
+		// so let's take the chance to manually clear them here.
+		ClearDetailsViewWrapperObjects();
+	}
 }
 
 bool FControlRigEditor::DetailViewShowsAnyRigElement() const
@@ -2093,6 +2101,20 @@ void FControlRigEditor::ClearDetailObject(bool bChangeUISelectionState)
 		return;
 	}
 
+	ClearDetailsViewWrapperObjects();
+	
+	Inspector->GetPropertyView()->SetObjects(TArray<UObject*>(), true); // clear property view synchronously
+	Inspector->ShowDetailsForObjects(TArray<UObject*>());
+	Inspector->ShowSingleStruct(TSharedPtr<FStructOnScope>());
+
+	if (bChangeUISelectionState)
+	{
+		SetUISelectionState(FBlueprintEditor::SelectionState_Graph);
+	}
+}
+
+void FControlRigEditor::ClearDetailsViewWrapperObjects()
+{
 	for(const TStrongObjectPtr<UDetailsViewWrapperObject>& WrapperObjectPtr : WrapperObjects)
 	{
 		if(WrapperObjectPtr.IsValid())
@@ -2104,15 +2126,6 @@ void FControlRigEditor::ClearDetailObject(bool bChangeUISelectionState)
 		}
 	}
 	WrapperObjects.Reset();
-	
-	Inspector->GetPropertyView()->SetObjects(TArray<UObject*>(), true); // clear property view synchronously
-	Inspector->ShowDetailsForObjects(TArray<UObject*>());
-	Inspector->ShowSingleStruct(TSharedPtr<FStructOnScope>());
-
-	if (bChangeUISelectionState)
-	{
-		SetUISelectionState(FBlueprintEditor::SelectionState_Graph);
-	}
 }
 
 
diff --git a/Engine/Plugins/Animation/ControlRig/Source/ControlRigEditor/Private/Editor/ControlRigEditor.h b/Engine/Plugins/Animation/ControlRig/Source/ControlRigEditor/Private/Editor/ControlRigEditor.h
--- a/Engine/Plugins/Animation/ControlRig/Source/ControlRigEditor/Private/Editor/ControlRigEditor.h
+++ b/Engine/Plugins/Animation/ControlRig/Source/ControlRigEditor/Private/Editor/ControlRigEditor.h
@@ -167,6 +167,10 @@ public:
 	bool DetailViewShowsRigElement(FRigElementKey InKey) const;
 
 	void ClearDetailObject(bool bChangeUISelectionState = true);
+private:
+
+	void ClearDetailsViewWrapperObjects();
+public:
 
 	/** Get the persona toolkit */
 	TSharedRef<IPersonaToolkit> GetPersonaToolkit() const { return PersonaToolkit.ToSharedRef(); }
