Control Rig: Fix default values not showing correctly on reroutes when duplicating the node

diff --git a/Engine/Plugins/Animation/ControlRig/Source/ControlRigDeveloper/Private/Graph/ControlRigGraph.cpp b/Engine/Plugins/Animation/ControlRig/Source/ControlRigDeveloper/Private/Graph/ControlRigGraph.cpp
--- a/Engine/Plugins/Animation/ControlRig/Source/ControlRigDeveloper/Private/Graph/ControlRigGraph.cpp
+++ b/Engine/Plugins/Animation/ControlRig/Source/ControlRigDeveloper/Private/Graph/ControlRigGraph.cpp
@@ -494,7 +494,7 @@ void UControlRigGraph::HandleModifiedEvent(ERigVMGraphNotifType InNotifType, URi
 
 					ModelNodePathToEdNode.Add(ModelNode->GetFName(), NewNode);
 				}
-				else if (URigVMRerouteNode* RerouteModelNode = Cast<URigVMRerouteNode>(ModelNode))
+				else // struct, library, parameter + variable
 				{
 					UControlRigGraphNode* NewNode = NewObject<UControlRigGraphNode>(this, GetControlRigGraphSchema()->GetGraphNodeClass(), ModelNode->GetFName());
 					AddNode(NewNode, false, false);
@@ -503,36 +503,24 @@ void UControlRigGraph::HandleModifiedEvent(ERigVMGraphNotifType InNotifType, URi
 					NewNode->CreateNewGuid();
 					NewNode->PostPlacedNewNode();
 					NewNode->AllocateDefaultPins();
+					NewNode->PostReconstructNode();
 
 					NewNode->NodePosX = ModelNode->GetPosition().X;
 					NewNode->NodePosY = ModelNode->GetPosition().Y;
-
-					NewNode->SetFlags(RF_Transactional);
-
-					if (URigVMPin* ModelPin = ModelNode->FindPin("Value"))
+					if (URigVMRerouteNode* RerouteModelNode = Cast<URigVMRerouteNode>(ModelNode))
 					{
-						if (UEdGraphPin* ValuePin = NewNode->FindPin(ModelPin->GetPinPath()))
+						if (URigVMPin* ModelPin = ModelNode->FindPin("Value"))
 						{
-							NewNode->SetColorFromModel(GetSchema()->GetPinTypeColor(ValuePin->PinType));
+							if (UEdGraphPin* ValuePin = NewNode->FindPin(ModelPin->GetPinPath()))
+							{
+								NewNode->SetColorFromModel(GetSchema()->GetPinTypeColor(ValuePin->PinType));
+							}
 						}
 					}
-
-					ModelNodePathToEdNode.Add(ModelNode->GetFName(), NewNode);
-				}
-				else // struct, library, parameter + variable
-				{
-					UControlRigGraphNode* NewNode = NewObject<UControlRigGraphNode>(this, GetControlRigGraphSchema()->GetGraphNodeClass(), ModelNode->GetFName());
-					AddNode(NewNode, false, false);
-
-					NewNode->ModelNodePath = ModelNode->GetNodePath();
-					NewNode->CreateNewGuid();
-					NewNode->PostPlacedNewNode();
-					NewNode->AllocateDefaultPins();
-					NewNode->PostReconstructNode();
-
-					NewNode->NodePosX = ModelNode->GetPosition().X;
-					NewNode->NodePosY = ModelNode->GetPosition().Y;
-					NewNode->SetColorFromModel(ModelNode->GetNodeColor());
+					else
+					{
+						NewNode->SetColorFromModel(ModelNode->GetNodeColor());
+					}
 					NewNode->SetFlags(RF_Transactional);
 
 					ModelNodePathToEdNode.Add(ModelNode->GetFName(), NewNode);
