Control Rig: Fix override material from the control rig shape library not being applied to all material slots in a static mesh

diff --git a/Engine/Plugins/Animation/ControlRig/Source/ControlRig/Private/ControlRigControlActor.cpp b/Engine/Plugins/Animation/ControlRig/Source/ControlRig/Private/ControlRigControlActor.cpp
--- a/Engine/Plugins/Animation/ControlRig/Source/ControlRig/Private/ControlRigControlActor.cpp
+++ b/Engine/Plugins/Animation/ControlRig/Source/ControlRig/Private/ControlRigControlActor.cpp
@@ -242,7 +242,10 @@ void AControlRigControlActor::Refresh()
 						Component->bCastDynamicShadow = bCastShadows;
 
 						UMaterialInstanceDynamic* MaterialInstance = UMaterialInstanceDynamic::Create(BaseMaterial, Component);
-						Component->SetMaterial(0, MaterialInstance);
+						for (int32 i=0; i<Component->GetNumMaterials(); ++i)
+						{
+							Component->SetMaterial(i, MaterialInstance);
+						}
 
 						ControlNames.Add(ControlElement->GetName());
 						ShapeTransforms.Add(ShapeDef->Transform * Hierarchy->GetControlShapeTransform(ControlElement, ERigTransformType::CurrentLocal));
diff --git a/Engine/Plugins/Animation/ControlRig/Source/ControlRig/Private/ControlRigGizmoActor.cpp b/Engine/Plugins/Animation/ControlRig/Source/ControlRig/Private/ControlRigGizmoActor.cpp
--- a/Engine/Plugins/Animation/ControlRig/Source/ControlRig/Private/ControlRigGizmoActor.cpp
+++ b/Engine/Plugins/Animation/ControlRig/Source/ControlRig/Private/ControlRigGizmoActor.cpp
@@ -265,7 +265,11 @@ namespace FControlRigShapeHelper
 				ShapeActor->ColorParameterName = CreationParam.ColorParameterName;
 				UMaterialInstanceDynamic* MaterialInstance = UMaterialInstanceDynamic::Create(CreationParam.Material.Get(), ShapeActor);
 				MaterialInstance->SetVectorParameterValue(CreationParam.ColorParameterName, FVector(CreationParam.Color));
-				MeshComponent->SetMaterial(0, MaterialInstance);
+
+				for (int32 i=0; i<MeshComponent->GetNumMaterials(); ++i)
+				{
+					MeshComponent->SetMaterial(i, MaterialInstance);
+				}
 			}
 			return ShapeActor;
 		}
diff --git a/Engine/Plugins/Animation/ControlRig/Source/ControlRig/Public/ControlRigGizmoLibrary.h b/Engine/Plugins/Animation/ControlRig/Source/ControlRig/Public/ControlRigGizmoLibrary.h
--- a/Engine/Plugins/Animation/ControlRig/Source/ControlRig/Public/ControlRigGizmoLibrary.h
+++ b/Engine/Plugins/Animation/ControlRig/Source/ControlRig/Public/ControlRigGizmoLibrary.h
@@ -48,7 +48,7 @@ public:
 	UPROPERTY(EditAnywhere, Category = "ShapeLibrary")
 	FControlRigShapeDefinition DefaultShape;
 
-	UPROPERTY(EditAnywhere, Category = "ShapeLibrary")
+	UPROPERTY(EditAnywhere, Category = "ShapeLibrary", meta = (DisplayName = "Override Material"))
 	TSoftObjectPtr<UMaterial> DefaultMaterial;
 
 	UPROPERTY(EditAnywhere, Category = "ShapeLibrary")
diff --git a/Engine/Plugins/Animation/ControlRig/Source/ControlRigEditor/Private/ControlRigThumbnailRenderer.cpp b/Engine/Plugins/Animation/ControlRig/Source/ControlRigEditor/Private/ControlRigThumbnailRenderer.cpp
--- a/Engine/Plugins/Animation/ControlRig/Source/ControlRigEditor/Private/ControlRigThumbnailRenderer.cpp
+++ b/Engine/Plugins/Animation/ControlRig/Source/ControlRigEditor/Private/ControlRigThumbnailRenderer.cpp
@@ -151,7 +151,11 @@ void UControlRigThumbnailRenderer::AddAdditionalPreviewSceneContent(UObject* Obj
 
 							UMaterialInstanceDynamic* MaterialInstance = UMaterialInstanceDynamic::Create(DefaultMaterial, ShapeActor);
 							MaterialInstance->SetVectorParameterValue(ShapeDef->Library.Get()->MaterialColorParameter,FVector(ControlElement->Settings.ShapeColor));
-							ShapeActor->GetStaticMeshComponent()->SetMaterial(0, MaterialInstance);
+							UStaticMeshComponent* MeshComponent = ShapeActor->GetStaticMeshComponent();
+							for (int32 i=0; i<MeshComponent->GetNumMaterials(); ++i)
+							{
+								MeshComponent->SetMaterial(i, MaterialInstance);
+							}
 
 							ShapeActors.Add(ControlElement->GetName(), ShapeActor);
 
