Control Rig: Ensure to reattach pins only once during postload

diff --git a/Engine/Plugins/Animation/ControlRig/Source/ControlRigDeveloper/Private/ControlRigBlueprint.cpp b/Engine/Plugins/Animation/ControlRig/Source/ControlRigDeveloper/Private/ControlRigBlueprint.cpp
--- a/Engine/Plugins/Animation/ControlRig/Source/ControlRigDeveloper/Private/ControlRigBlueprint.cpp
+++ b/Engine/Plugins/Animation/ControlRig/Source/ControlRigDeveloper/Private/ControlRigBlueprint.cpp
@@ -1323,12 +1323,15 @@ void UControlRigBlueprint::RefreshAllModels(EControlRigBlueprintLoadType InLoadT
 	// avoid any compute if the current structure hashes match with the serialized ones
 	if(RigVMClient.GetStructureHash() == RigVMClient.GetSerializedStructureHash())
 	{
-		TArray<URigVMGraph*> ModelGraphs = RigVMClient.GetAllModels(true, true);
-		Algo::Reverse(ModelGraphs);
-		for (URigVMGraph* ModelGraph : ModelGraphs)
+		if(bIsPostLoad)
 		{
-			URigVMController* Controller = GetOrCreateController(ModelGraph);
-			Controller->ReattachLinksToPinObjects(true /* follow redirectors */, nullptr, true, true);
+			TArray<URigVMGraph*> ModelGraphs = RigVMClient.GetAllModels(true, true);
+			Algo::Reverse(ModelGraphs);
+			for (URigVMGraph* ModelGraph : ModelGraphs)
+			{	
+				URigVMController* Controller = GetOrCreateController(ModelGraph);
+				Controller->ReattachLinksToPinObjects(true /* follow redirectors */, nullptr, true, true);
+			}
 		}
 		return;
 	}
