Control Rig: nullptr checks for a item array nodes

diff --git a/Engine/Plugins/Animation/ControlRig/Source/ControlRig/Private/Units/Execution/RigUnit_Collection.cpp b/Engine/Plugins/Animation/ControlRig/Source/ControlRig/Private/Units/Execution/RigUnit_Collection.cpp
--- a/Engine/Plugins/Animation/ControlRig/Source/ControlRig/Private/Units/Execution/RigUnit_Collection.cpp
+++ b/Engine/Plugins/Animation/ControlRig/Source/ControlRig/Private/Units/Execution/RigUnit_Collection.cpp
@@ -40,6 +40,12 @@ FRigVMStructUpgradeInfo FRigUnit_CollectionChain::GetUpgradeInfo() const
 
 FRigUnit_CollectionChainArray_Execute()
 {
+	if(ExecuteContext.Hierarchy == nullptr)
+	{
+		Items.Reset();
+		return;
+	}
+	
 	uint32 Hash = FRigUnit_CollectionChain_Hash + ExecuteContext.Hierarchy->GetTopologyVersion() * 17;
 	Hash = HashCombine(Hash, GetTypeHash(FirstItem));
 	Hash = HashCombine(Hash, GetTypeHash(LastItem));
@@ -88,6 +94,12 @@ FRigUnit_CollectionNameSearchArray_Execute()
 {
 	DECLARE_SCOPE_HIERARCHICAL_COUNTER_RIGUNIT()
 
+	if(ExecuteContext.Hierarchy == nullptr)
+	{
+		Items.Reset();
+		return;
+	}
+
 	uint32 Hash = FRigUnit_CollectionNameSearch_Hash + ExecuteContext.Hierarchy->GetTopologyVersion() * 17;
 	Hash = HashCombine(Hash, GetTypeHash(PartialName));
 	Hash = HashCombine(Hash, (int32)TypeToSearch * 8);
@@ -128,6 +140,12 @@ FRigUnit_CollectionChildrenArray_Execute()
 {
 	DECLARE_SCOPE_HIERARCHICAL_COUNTER_RIGUNIT()
 
+	if(ExecuteContext.Hierarchy == nullptr)
+	{
+		Items.Reset();
+		return;
+	}
+	
 	uint32 Hash = FRigUnit_CollectionChildren_Hash + ExecuteContext.Hierarchy->GetTopologyVersion() * 17;
 	Hash = HashCombine(Hash, GetTypeHash(Parent));
 	Hash = HashCombine(Hash, bRecursive ? 2 : 0);
@@ -159,6 +177,12 @@ FRigUnit_CollectionGetAll_Execute()
 {
 	DECLARE_SCOPE_HIERARCHICAL_COUNTER_RIGUNIT()
 
+	if(ExecuteContext.Hierarchy == nullptr)
+	{
+		Items.Reset();
+		return;
+	}
+	
 	uint32 Hash = FRigUnit_CollectionGetAll_Hash + ExecuteContext.Hierarchy->GetTopologyVersion() * 17;
 	Hash = HashCombine(Hash, (int32)TypeToSearch * 8);
 
@@ -250,6 +274,12 @@ FRigUnit_CollectionReplaceItemsArray_Execute()
 {
 	DECLARE_SCOPE_HIERARCHICAL_COUNTER_RIGUNIT()
 
+	if(ExecuteContext.Hierarchy == nullptr)
+	{
+		Result.Reset();
+		return;
+	}
+	
 	uint32 Hash = FRigUnit_CollectionReplaceItems_Hash + ExecuteContext.Hierarchy->GetTopologyVersion() * 17;
 	Hash = HashCombine(Hash, GetTypeHash(Items));
 	Hash = HashCombine(Hash, 12 * GetTypeHash(Old));
