ControlRig: Deal with assets containing nullptr UEdGraphPin subpins

diff --git a/Engine/Plugins/Animation/ControlRig/Source/ControlRigDeveloper/Private/Graph/ControlRigGraphNode.cpp b/Engine/Plugins/Animation/ControlRig/Source/ControlRigDeveloper/Private/Graph/ControlRigGraphNode.cpp
--- a/Engine/Plugins/Animation/ControlRig/Source/ControlRigDeveloper/Private/Graph/ControlRigGraphNode.cpp
+++ b/Engine/Plugins/Animation/ControlRig/Source/ControlRigDeveloper/Private/Graph/ControlRigGraphNode.cpp
@@ -361,6 +361,10 @@ void UControlRigGraphNode::DestroyPinList(TArray<UEdGraphPin*>& InPins)
 	// Throw away the original pins
 	for (UEdGraphPin* Pin : InPins)
 	{
+		// previously serialized graphs potentially contain
+		// nullptr in the subpins
+		Pin->SubPins.Remove(nullptr);
+		
 		Pin->BreakAllPinLinks(bNotify);
 		Pin->SubPins.Remove(nullptr);
 		UEdGraphNode::DestroyPin(Pin);
