Control Rig: Fix ControlRigComponenent mapped elements reference creating orphan pin in the AnimBlueprint

diff --git a/Engine/Plugins/Animation/ControlRig/Source/ControlRig/Private/ControlRigComponent.cpp b/Engine/Plugins/Animation/ControlRig/Source/ControlRig/Private/ControlRigComponent.cpp
--- a/Engine/Plugins/Animation/ControlRig/Source/ControlRig/Private/ControlRigComponent.cpp
+++ b/Engine/Plugins/Animation/ControlRig/Source/ControlRig/Private/ControlRigComponent.cpp
@@ -94,21 +94,6 @@ void UControlRigComponent::PostEditChangeProperty(FPropertyChangedEvent& Propert
 void UControlRigComponent::PostLoad()
 {
 	Super::PostLoad();
-
-	for (FControlRigComponentMappedElement& MappedComponent : UserDefinedElements)
-	{
-		if (MappedComponent.ComponentReference_DEPRECATED.OtherActor.IsValid())
-		{
-			MappedComponent.SoftComponentReference.OtherActor = MappedComponent.ComponentReference_DEPRECATED.OtherActor.Get();
-		}
-	}
-	for (FControlRigComponentMappedElement& MappedComponent : MappedElements)
-	{
-		if (MappedComponent.ComponentReference_DEPRECATED.OtherActor.IsValid())
-		{
-			MappedComponent.SoftComponentReference.OtherActor = MappedComponent.ComponentReference_DEPRECATED.OtherActor.Get();
-		}
-	}
 }
 #endif
 
@@ -615,8 +600,8 @@ void UControlRigComponent::AddMappedComponents(TArray<FControlRigComponentMapped
 		USceneComponent* Component = ComponentToMap.Component;
 
 		FControlRigComponentMappedElement ElementToMap;
-		ElementToMap.SoftComponentReference.OtherActor = Component->GetOwner() != GetOwner() ? Component->GetOwner() : nullptr;
-		ElementToMap.SoftComponentReference.ComponentProperty = GetComponentNameWithinActor(Component);
+		ElementToMap.ComponentReference.OtherActor = Component->GetOwner() != GetOwner() ? Component->GetOwner() : nullptr;
+		ElementToMap.ComponentReference.ComponentProperty = GetComponentNameWithinActor(Component);
 
 		ElementToMap.ElementName = ComponentToMap.ElementName;
 		ElementToMap.ElementType = ComponentToMap.ElementType;
@@ -722,8 +707,8 @@ void UControlRigComponent::AddMappedSkeletalMesh(USkeletalMeshComponent* Skeleta
 		}
 
 		FControlRigComponentMappedElement ElementToMap;
-		ElementToMap.SoftComponentReference.OtherActor = SkeletalMeshComponent->GetOwner() != GetOwner() ? SkeletalMeshComponent->GetOwner() : nullptr;
-		ElementToMap.SoftComponentReference.ComponentProperty = GetComponentNameWithinActor(SkeletalMeshComponent);
+		ElementToMap.ComponentReference.OtherActor = SkeletalMeshComponent->GetOwner() != GetOwner() ? SkeletalMeshComponent->GetOwner() : nullptr;
+		ElementToMap.ComponentReference.ComponentProperty = GetComponentNameWithinActor(SkeletalMeshComponent);
 
 		ElementToMap.ElementName = BoneToMap.Source;
 		ElementToMap.ElementType = ERigElementType::Bone;
@@ -741,8 +726,8 @@ void UControlRigComponent::AddMappedSkeletalMesh(USkeletalMeshComponent* Skeleta
 		}
 
 		FControlRigComponentMappedElement ElementToMap;
-		ElementToMap.SoftComponentReference.OtherActor = SkeletalMeshComponent->GetOwner() != GetOwner() ? SkeletalMeshComponent->GetOwner() : nullptr;
-		ElementToMap.SoftComponentReference.ComponentProperty = GetComponentNameWithinActor(SkeletalMeshComponent);
+		ElementToMap.ComponentReference.OtherActor = SkeletalMeshComponent->GetOwner() != GetOwner() ? SkeletalMeshComponent->GetOwner() : nullptr;
+		ElementToMap.ComponentReference.ComponentProperty = GetComponentNameWithinActor(SkeletalMeshComponent);
 
 		ElementToMap.ElementName = CurveToMap.Source;
 		ElementToMap.ElementType = ERigElementType::Curve;
@@ -1323,17 +1308,17 @@ void UControlRigComponent::ValidateMappingData()
 			MappedElement.ElementIndex = INDEX_NONE;
 			MappedElement.SubIndex = INDEX_NONE;
 
-			AActor* MappedOwner = !MappedElement.SoftComponentReference.OtherActor.IsValid() ? GetOwner() : MappedElement.SoftComponentReference.OtherActor.Get();
-			MappedElement.SceneComponent = Cast<USceneComponent>(MappedElement.SoftComponentReference.GetComponent(MappedOwner));
+			AActor* MappedOwner = !MappedElement.ComponentReference.OtherActor.IsValid() ? GetOwner() : MappedElement.ComponentReference.OtherActor.Get();
+			MappedElement.SceneComponent = Cast<USceneComponent>(MappedElement.ComponentReference.GetComponent(MappedOwner));
 
 			// try again with the path to the component
 			if (MappedElement.SceneComponent == nullptr)
 			{
 				FSoftComponentReference TempReference;
-				TempReference.PathToComponent = MappedElement.SoftComponentReference.ComponentProperty.ToString();
+				TempReference.PathToComponent = MappedElement.ComponentReference.ComponentProperty.ToString();
 				if (USceneComponent* TempSceneComponent = Cast<USceneComponent>(TempReference.GetComponent(MappedOwner)))
 				{
-					MappedElement.SoftComponentReference = TempReference;
+					MappedElement.ComponentReference = TempReference;
 					MappedElement.SceneComponent = TempSceneComponent;
 				}
 			}
@@ -1346,7 +1331,7 @@ void UControlRigComponent::ValidateMappingData()
 			}
 
 			// cache the scene component also in the override component to avoid on further relying on names
-			MappedElement.SoftComponentReference.OverrideComponent = MappedElement.SceneComponent;
+			MappedElement.ComponentReference.OverrideComponent = MappedElement.SceneComponent;
 
 			if (MappedElement.Direction == EControlRigComponentMapDirection::Output && MappedElement.Weight <= SMALL_NUMBER)
 			{
diff --git a/Engine/Plugins/Animation/ControlRig/Source/ControlRig/Public/ControlRigComponent.h b/Engine/Plugins/Animation/ControlRig/Source/ControlRig/Public/ControlRigComponent.h
--- a/Engine/Plugins/Animation/ControlRig/Source/ControlRig/Public/ControlRigComponent.h
+++ b/Engine/Plugins/Animation/ControlRig/Source/ControlRig/Public/ControlRigComponent.h
@@ -53,7 +53,7 @@ struct FControlRigComponentMappedElement
 	GENERATED_BODY()
 
 		FControlRigComponentMappedElement()
-		: SoftComponentReference(FSoftComponentReference())
+		: ComponentReference(FSoftComponentReference())
 		, TransformIndex(INDEX_NONE)
 		, TransformName(NAME_None)
 		, ElementType(ERigElementType::Bone)
@@ -68,14 +68,9 @@ struct FControlRigComponentMappedElement
 	{
 	}
 
-#if WITH_EDITORONLY_DATA
 	// The component to map to the Control Rig
-	UPROPERTY(meta=(DeprecatedProperty, DeprecationMessage = "ComponentReference has been deprecated, use SoftComponentReference instead."))
-	FComponentReference ComponentReference_DEPRECATED;
-#endif
-
 	UPROPERTY(EditAnywhere, BlueprintReadWrite, Category = "Mapping")
-	FSoftComponentReference SoftComponentReference;
+	FSoftComponentReference ComponentReference;
 
 	// An optional index that can be used with components
 	// with multiple transforms (for example the InstancedStaticMeshComponent)
diff --git a/Engine/Source/Runtime/Engine/Classes/Engine/EngineTypes.h b/Engine/Source/Runtime/Engine/Classes/Engine/EngineTypes.h
--- a/Engine/Source/Runtime/Engine/Classes/Engine/EngineTypes.h
+++ b/Engine/Source/Runtime/Engine/Classes/Engine/EngineTypes.h
@@ -3215,6 +3215,17 @@ struct ENGINE_API FSoftComponentReference : public FBaseComponentReference
 	{
 		return (OtherActor == Other.OtherActor) && (FBaseComponentReference::operator==(Other));
 	}
+
+	bool SerializeFromMismatchedTag(const FPropertyTag& Tag, FStructuredArchive::FSlot Slot);
+};
+
+template<>
+struct TStructOpsTypeTraits<FSoftComponentReference> : public TStructOpsTypeTraitsBase2<FSoftComponentReference>
+{
+	enum
+	{
+		WithStructuredSerializeFromMismatchedTag = true,
+	};
 };
 
 
diff --git a/Engine/Source/Runtime/Engine/Classes/Kismet/KismetSystemLibrary.h b/Engine/Source/Runtime/Engine/Classes/Kismet/KismetSystemLibrary.h
--- a/Engine/Source/Runtime/Engine/Classes/Kismet/KismetSystemLibrary.h
+++ b/Engine/Source/Runtime/Engine/Classes/Kismet/KismetSystemLibrary.h
@@ -347,6 +347,9 @@ class ENGINE_API UKismetSystemLibrary : public UBlueprintFunctionLibrary
 	UFUNCTION(BlueprintPure, meta = (BlueprintInternalUseOnly = "true"), Category = "Utilities")
 	static TSoftClassPtr<UObject> Conv_ClassToSoftClassReference(const TSubclassOf<UObject>& Class);
 
+	UFUNCTION(BlueprintPure, meta = (BlueprintAutocast), Category = "Utilities")
+	static FSoftComponentReference Conv_ComponentReferenceToSoftComponentReference(const FComponentReference& ComponentReference);
+
 	DECLARE_DYNAMIC_DELEGATE_OneParam(FOnAssetLoaded, class UObject*, Loaded);
 
 	UFUNCTION(BlueprintCallable, meta = (Latent, LatentInfo = "LatentInfo", WorldContext = "WorldContextObject", BlueprintInternalUseOnly = "true"), Category = "Utilities")
diff --git a/Engine/Source/Runtime/Engine/Private/EngineTypes.cpp b/Engine/Source/Runtime/Engine/Private/EngineTypes.cpp
--- a/Engine/Source/Runtime/Engine/Private/EngineTypes.cpp
+++ b/Engine/Source/Runtime/Engine/Private/EngineTypes.cpp
@@ -179,6 +179,25 @@ UActorComponent* FSoftComponentReference::GetComponent(AActor* OwningActor) cons
 	return ExtractComponent(SearchActor);
 }
 
+bool FSoftComponentReference::SerializeFromMismatchedTag(const FPropertyTag& Tag, FStructuredArchive::FSlot Slot)
+{
+	static const FName ComponentReferenceContextName("ComponentReference");
+	if (Tag.Type == NAME_StructProperty && Tag.StructName == ComponentReferenceContextName)
+	{
+		FComponentReference Reference;
+		FComponentReference::StaticStruct()->SerializeItem(Slot, &Reference, nullptr);
+		if (Reference.OtherActor.IsValid())
+		{
+			OtherActor = Reference.OtherActor.Get();
+			ComponentProperty = Reference.ComponentProperty;
+			PathToComponent = Reference.PathToComponent;
+		}
+		return true;
+	}
+
+	return false;
+}
+
 const TCHAR* LexToString(const EWorldType::Type Value)
 {
 	switch (Value)
diff --git a/Engine/Source/Runtime/Engine/Private/KismetSystemLibrary.cpp b/Engine/Source/Runtime/Engine/Private/KismetSystemLibrary.cpp
--- a/Engine/Source/Runtime/Engine/Private/KismetSystemLibrary.cpp
+++ b/Engine/Source/Runtime/Engine/Private/KismetSystemLibrary.cpp
@@ -1393,6 +1393,18 @@ TSoftClassPtr<UObject> UKismetSystemLibrary::Conv_ClassToSoftClassReference(cons
 	return TSoftClassPtr<UObject>(*Class);
 }
 
+FSoftComponentReference UKismetSystemLibrary::Conv_ComponentReferenceToSoftComponentReference(const FComponentReference& ComponentReference)
+{
+	FSoftComponentReference SoftComponentReference;
+	SoftComponentReference.ComponentProperty = ComponentReference.ComponentProperty;
+	SoftComponentReference.PathToComponent = ComponentReference.PathToComponent;
+	if (ComponentReference.OtherActor.IsValid())
+	{
+		SoftComponentReference.OtherActor = ComponentReference.OtherActor.Get();
+	}
+	return SoftComponentReference;
+}
+
 void UKismetSystemLibrary::SetTextPropertyByName(UObject* Object, FName PropertyName, const FText& Value)
 {
 	if(Object != nullptr)
