Control Rig: Fix crash when during destroy of ControlRigComponent

diff --git a/Engine/Plugins/Animation/ControlRig/Source/ControlRig/Private/ControlRigComponent.cpp b/Engine/Plugins/Animation/ControlRig/Source/ControlRig/Private/ControlRigComponent.cpp
--- a/Engine/Plugins/Animation/ControlRig/Source/ControlRig/Private/ControlRigComponent.cpp
+++ b/Engine/Plugins/Animation/ControlRig/Source/ControlRig/Private/ControlRigComponent.cpp
@@ -258,14 +258,14 @@ FBoxSphereBounds UControlRigComponent::CalcBounds(const FTransform& LocalToWorld
 		const FTransform Transform = GetComponentToWorld();
 
 		// Get bounding box for bones
-		URigHierarchy* Hierarchy = ControlRig->GetHierarchy();
-		check(Hierarchy);
-
-		Hierarchy->ForEach<FRigTransformElement>([&BBox, Transform, Hierarchy](FRigTransformElement* TransformElement) -> bool
+		if (URigHierarchy* Hierarchy = ControlRig->GetHierarchy())
 		{
-			BBox += Transform.TransformPosition(Hierarchy->GetTransform(TransformElement, ERigTransformType::CurrentGlobal).GetLocation());
-			return true;
-        });
+			Hierarchy->ForEach<FRigTransformElement>([&BBox, Transform, Hierarchy](FRigTransformElement* TransformElement) -> bool
+			{
+				BBox += Transform.TransformPosition(Hierarchy->GetTransform(TransformElement, ERigTransformType::CurrentGlobal).GetLocation());
+				return true;
+			});
+		}
 	}
 
 	if (BBox.IsValid)
